Fixed crash by removing null nodes on PostLoad in EdGraphs.
[UE-18998] - CLONE - Crash in OceanGlobe_UMG when adding new Event Dispatcher

diff --git a/Engine/Source/Runtime/Engine/Classes/EdGraph/EdGraph.h b/Engine/Source/Runtime/Engine/Classes/EdGraph/EdGraph.h
--- a/Engine/Source/Runtime/Engine/Classes/EdGraph/EdGraph.h
+++ b/Engine/Source/Runtime/Engine/Classes/EdGraph/EdGraph.h
@@ -106,6 +106,7 @@ public:
 #if WITH_EDITORONLY_DATA
 	// Begin UObject interface
 	virtual void PostInitProperties() override;
+	virtual void PostLoad() override;
 	// End UObject interface
 #endif
 
diff --git a/Engine/Source/Runtime/Engine/Private/EdGraph/EdGraph.cpp b/Engine/Source/Runtime/Engine/Private/EdGraph/EdGraph.cpp
--- a/Engine/Source/Runtime/Engine/Private/EdGraph/EdGraph.cpp
+++ b/Engine/Source/Runtime/Engine/Private/EdGraph/EdGraph.cpp
@@ -95,6 +95,21 @@ void UEdGraph::PostInitProperties()
 		GraphGuid = FGuid::NewGuid();
 	}
 }
+
+void UEdGraph::PostLoad()
+{
+	Super::PostLoad();
+
+	// Strip out null nodes (likely from missing node classes) as they will cause crashes 
+	for (int32 i = Nodes.Num() - 1; i >= 0; i--)
+	{
+		if (Nodes[i] == nullptr)
+		{
+			Nodes.RemoveAt(i);
+			UE_LOG(LogBlueprint, Warning, TEXT("Found NULL Node in EdGraph Nodes array. A node type may have been deleted without creating an ActiveClassRedictor to K2Node_DeadClass."));
+		}
+	}
+}
 #endif
 
 const UEdGraphSchema* UEdGraph::GetSchema() const
