Removing programmer-only comments from documentation.

diff --git a/Engine/Source/Runtime/Engine/Classes/GameFramework/Actor.h b/Engine/Source/Runtime/Engine/Classes/GameFramework/Actor.h
--- a/Engine/Source/Runtime/Engine/Classes/GameFramework/Actor.h
+++ b/Engine/Source/Runtime/Engine/Classes/GameFramework/Actor.h
@@ -1402,7 +1402,7 @@ public:
 		return false;
 	}
 
-	/**
+	/*~
 	 * Returns location of the RootComponent 
 	 * this is a template for no other reason than to delay compilation until USceneComponent is defined
 	 */ 
@@ -1412,7 +1412,7 @@ public:
 		return (RootComponent != nullptr) ? RootComponent->GetComponentLocation() : FVector(0.f,0.f,0.f);
 	}
 
-	/**
+	/*~
 	 * Returns rotation of the RootComponent 
 	 * this is a template for no other reason than to delay compilation until USceneComponent is defined
 	 */ 
@@ -1422,7 +1422,7 @@ public:
 		return (RootComponent != nullptr) ? RootComponent->GetComponentRotation() : FRotator(0.f,0.f,0.f);
 	}
 
-	/**
+	/*~
 	 * Returns scale of the RootComponent 
 	 * this is a template for no other reason than to delay compilation until USceneComponent is defined
 	 */ 
@@ -1432,7 +1432,7 @@ public:
 		return (RootComponent != nullptr) ? RootComponent->GetComponentScale() : FVector(1.f,1.f,1.f);
 	}
 
-	/**
+	/*~
 	 * Returns quaternion of the RootComponent
 	 * this is a template for no other reason than to delay compilation until USceneComponent is defined
 	 */ 
