Reverting changed in 2643256.
Turns out it'll require a bit more thought.

diff --git a/Engine/Source/Runtime/Engine/Private/Materials/MaterialInstance.cpp b/Engine/Source/Runtime/Engine/Private/Materials/MaterialInstance.cpp
--- a/Engine/Source/Runtime/Engine/Private/Materials/MaterialInstance.cpp
+++ b/Engine/Source/Runtime/Engine/Private/Materials/MaterialInstance.cpp
@@ -1081,7 +1081,7 @@ void UMaterialInstance::CopyMaterialInstanceParameters(UMaterialInterface* Mater
 
 FMaterialResource* UMaterialInstance::GetMaterialResource(ERHIFeatureLevel::Type InFeatureLevel, EMaterialQualityLevel::Type QualityLevel)
 {
-	//check(!IsInActualRenderingThread());
+	check(!IsInActualRenderingThread());
 
 	if (QualityLevel == EMaterialQualityLevel::Num)
 	{
diff --git a/Engine/Source/Runtime/Engine/Private/Materials/MaterialInterface.cpp b/Engine/Source/Runtime/Engine/Private/Materials/MaterialInterface.cpp
--- a/Engine/Source/Runtime/Engine/Private/Materials/MaterialInterface.cpp
+++ b/Engine/Source/Runtime/Engine/Private/Materials/MaterialInterface.cpp
@@ -53,7 +53,7 @@ FMaterialRelevance UMaterialInterface::GetRelevance_Internal(const UMaterial* Ma
 {
 	if(Material)
 	{
-		const FMaterialResource* MaterialResource = GetMaterialResource(InFeatureLevel);
+		const FMaterialResource* MaterialResource = Material->GetMaterialResource(InFeatureLevel);
 		const bool bIsTranslucent = IsTranslucentBlendMode((EBlendMode)GetBlendMode());
 
 		EMaterialShadingModel ShadingModel = GetShadingModel();
