Remove per scene deletion buffers. This fixes out of order deletion.

diff --git a/Engine/Source/Runtime/Engine/Classes/Engine/SkeletalMesh.h b/Engine/Source/Runtime/Engine/Classes/Engine/SkeletalMesh.h
--- a/Engine/Source/Runtime/Engine/Classes/Engine/SkeletalMesh.h
+++ b/Engine/Source/Runtime/Engine/Classes/Engine/SkeletalMesh.h
@@ -737,7 +737,6 @@ public:
 	virtual void GetAssetRegistryTagMetadata(TMap<FName, FAssetRegistryTagMetadata>& OutMetadata) const override;
 #endif // WITH_EDITOR
 	virtual void BeginDestroy() override;
-	virtual void FinishDestroy() override;
 	virtual bool IsReadyForFinishDestroy() override;
 	virtual void PreSave() override;
 	virtual void Serialize(FArchive& Ar) override;
diff --git a/Engine/Source/Runtime/Engine/Private/Components/DestructibleComponent.cpp b/Engine/Source/Runtime/Engine/Private/Components/DestructibleComponent.cpp
--- a/Engine/Source/Runtime/Engine/Private/Components/DestructibleComponent.cpp
+++ b/Engine/Source/Runtime/Engine/Private/Components/DestructibleComponent.cpp
@@ -365,7 +365,7 @@ void UDestructibleComponent::CreatePhysicsState()
 	{
 		// Lock and flush deferred command handler here to stop any currently pending deletions from affecting new actors.
 		SCOPED_SCENE_WRITE_LOCK(PScene);
-		PhysScene->DeferredCommandHandler.Flush();
+		GPhysCommandHandler->Flush();
 
 		// Create an APEX NxDestructibleActor from the Destructible asset and actor descriptor, has to be in a locked scope as it affects simulation state
 		ApexDestructibleActor = static_cast<NxDestructibleActor*>(TheDestructibleMesh->ApexDestructibleAsset->createApexActor(*ActorParams, *ApexScene));
@@ -414,7 +414,7 @@ void UDestructibleComponent::DestroyPhysicsState()
 		{
 			if (FPhysScene * PhysScene = World->GetPhysicsScene())
 			{
-				PhysScene->DeferredCommandHandler.DeferredRelease(ApexDestructibleActor);
+				GPhysCommandHandler->DeferredRelease(ApexDestructibleActor);
 			}
 		}
 		
diff --git a/Engine/Source/Runtime/Engine/Private/LevelTick.cpp b/Engine/Source/Runtime/Engine/Private/LevelTick.cpp
--- a/Engine/Source/Runtime/Engine/Private/LevelTick.cpp
+++ b/Engine/Source/Runtime/Engine/Private/LevelTick.cpp
@@ -1168,7 +1168,7 @@ void UWorld::Tick( ELevelTick TickType, float DeltaSeconds )
 		}
 		if ( PhysicsScene != NULL )
 		{
-			PhysicsScene->DeferredCommandHandler.Flush();
+			GPhysCommandHandler->Flush();
 		}
 		{
 			SCOPE_CYCLE_COUNTER(STAT_TG_PreCloth);
diff --git a/Engine/Source/Runtime/Engine/Private/PhysicsEngine/PhysScene.cpp b/Engine/Source/Runtime/Engine/Private/PhysicsEngine/PhysScene.cpp
--- a/Engine/Source/Runtime/Engine/Private/PhysicsEngine/PhysScene.cpp
+++ b/Engine/Source/Runtime/Engine/Private/PhysicsEngine/PhysScene.cpp
@@ -1493,7 +1493,7 @@ void FPhysScene::TermPhysScene(uint32 SceneType)
 		GPhysCommandHandler->DeferredDeleteSimEventCallback(SimEventCallback[SceneType]);
 
 		// Commands may have accumulated as the scene is terminated - flush any commands for this scene.
-		DeferredCommandHandler.Flush();
+		GPhysCommandHandler->Flush();
 
 		// Remove from the map
 		GPhysXSceneMap.Remove(PhysXSceneIndex[SceneType]);
diff --git a/Engine/Source/Runtime/Engine/Private/SkeletalMesh.cpp b/Engine/Source/Runtime/Engine/Private/SkeletalMesh.cpp
--- a/Engine/Source/Runtime/Engine/Private/SkeletalMesh.cpp
+++ b/Engine/Source/Runtime/Engine/Private/SkeletalMesh.cpp
@@ -2401,14 +2401,6 @@ void USkeletalMesh::BeginDestroy()
 {
 	Super::BeginDestroy();
 
-	// Release the mesh's render resources.
-	ReleaseResources();
-
-
-}
-
-void USkeletalMesh::FinishDestroy()
-{
 #if WITH_APEX_CLOTHING
 	// release clothing assets
 	for (FClothingAssetData& Data : ClothingAssets)
@@ -2421,9 +2413,11 @@ void USkeletalMesh::FinishDestroy()
 	}
 #endif // #if WITH_APEX_CLOTHING
 
-	Super::FinishDestroy();
-}
+	// Release the mesh's render resources.
+	ReleaseResources();
+
 
+}
 
 bool USkeletalMesh::IsReadyForFinishDestroy()
 {
diff --git a/Engine/Source/Runtime/Engine/Private/SkeletalMeshComponentPhysics.cpp b/Engine/Source/Runtime/Engine/Private/SkeletalMeshComponentPhysics.cpp
--- a/Engine/Source/Runtime/Engine/Private/SkeletalMeshComponentPhysics.cpp
+++ b/Engine/Source/Runtime/Engine/Private/SkeletalMeshComponentPhysics.cpp
@@ -62,7 +62,7 @@ void FClothingActor::Clear(bool bReleaseResource)
 {
 	if(bReleaseResource)
 	{
-		PhysScene->DeferredCommandHandler.DeferredRelease(ApexClothingActor);
+		GPhysCommandHandler->DeferredRelease(ApexClothingActor);
 	}
 
 	ParentClothingAsset = NULL;
diff --git a/Engine/Source/Runtime/Engine/Public/PhysicsPublic.h b/Engine/Source/Runtime/Engine/Public/PhysicsPublic.h
--- a/Engine/Source/Runtime/Engine/Public/PhysicsPublic.h
+++ b/Engine/Source/Runtime/Engine/Public/PhysicsPublic.h
@@ -580,9 +580,6 @@ public:
 	void AddPendingDamageEvent(class UDestructibleComponent* DestructibleComponent, const NxApexDamageEventReportData& DamageEvent);
 #endif
 
-	/** DeferredCommandHandler - this is mainly used for scene specific APEX calls that need to be deferred */
-	FPhysCommandHandler DeferredCommandHandler;
-
 private:
 	/** Initialize a scene of the given type.  Must only be called once for each scene type. */
 	void InitPhysScene(uint32 SceneType);
