UE-18444 Making an audio component's attenuation settings accessible from blueprint callable function

diff --git a/Engine/Source/Runtime/Engine/Classes/Components/AudioComponent.h b/Engine/Source/Runtime/Engine/Classes/Components/AudioComponent.h
--- a/Engine/Source/Runtime/Engine/Classes/Components/AudioComponent.h
+++ b/Engine/Source/Runtime/Engine/Classes/Components/AudioComponent.h
@@ -286,6 +286,9 @@ public:
 	/** Returns a pointer to the attenuation settings to be used (if any) for this audio component dependent on the SoundAttenuation asset or overrides set. */
 	const FAttenuationSettings* GetAttenuationSettingsToApply() const;
 
+	UFUNCTION(BlueprintCallable, Category = "Audio|Components|Audio", meta = (DisplayName = "Get Attenuation Settings To Apply"))
+	bool BP_GetAttenuationSettingsToApply(FAttenuationSettings& OutAttenuationSettings);
+
 	/** Collects the various attenuation shapes that may be applied to the sound played by the audio component for visualization in the editor or via the in game debug visualization. */
 	void CollectAttenuationShapesForVisualization(TMultiMap<EAttenuationShape::Type, FAttenuationSettings::AttenuationShapeDetails>& ShapeDetailsMap) const;
 
diff --git a/Engine/Source/Runtime/Engine/Private/Components/AudioComponent.cpp b/Engine/Source/Runtime/Engine/Private/Components/AudioComponent.cpp
--- a/Engine/Source/Runtime/Engine/Private/Components/AudioComponent.cpp
+++ b/Engine/Source/Runtime/Engine/Private/Components/AudioComponent.cpp
@@ -365,6 +365,16 @@ const FAttenuationSettings* UAudioComponent::GetAttenuationSettingsToApply() con
 	return nullptr;
 }
 
+bool UAudioComponent::BP_GetAttenuationSettingsToApply(FAttenuationSettings& OutAttenuationSettings)
+{
+	if (const FAttenuationSettings* Settings = GetAttenuationSettingsToApply())
+	{
+		OutAttenuationSettings = *Settings;
+		return true;
+	}
+	return false;
+}
+
 void UAudioComponent::CollectAttenuationShapesForVisualization(TMultiMap<EAttenuationShape::Type, FAttenuationSettings::AttenuationShapeDetails>& ShapeDetailsMap) const
 {
 	const FAttenuationSettings *AttenuationSettingsToApply = GetAttenuationSettingsToApply();
