Merge 2651680 from UE4-Orion

Add console commands demoscrub and demopause

diff --git a/Engine/Source/Runtime/Engine/Classes/Engine/World.h b/Engine/Source/Runtime/Engine/Classes/Engine/World.h
--- a/Engine/Source/Runtime/Engine/Classes/Engine/World.h
+++ b/Engine/Source/Runtime/Engine/Classes/Engine/World.h
@@ -2725,6 +2725,12 @@ private:
 	/** Utility function to handle Exec/Console Commands related to stopping demo playback */
 	bool HandleDemoStopCommand( const TCHAR* Cmd, FOutputDevice& Ar, UWorld* InWorld );
 
+	/** Utility function to handle Exec/Console Command for scrubbing to a specific time */
+	bool HandleDemoScrubCommand(const TCHAR* Cmd, FOutputDevice& Ar, UWorld* InWorld);
+
+	/** Utility function to handle Exec/Console Command for pausing and unpausing a replay */
+	bool HandleDemoPauseCommand(const TCHAR* Cmd, FOutputDevice& Ar, UWorld* InWorld);
+
 public:
 
 	// Destroys the current demo net driver
diff --git a/Engine/Source/Runtime/Engine/Private/World.cpp b/Engine/Source/Runtime/Engine/Private/World.cpp
--- a/Engine/Source/Runtime/Engine/Private/World.cpp
+++ b/Engine/Source/Runtime/Engine/Private/World.cpp
@@ -81,6 +81,7 @@
 #include "Engine/GameInstance.h"
 #include "UObject/UObjectThreadContext.h"
 #include "Engine/CoreSettings.h"
+#include "NetworkReplayStreaming.h"
 
 DEFINE_LOG_CATEGORY_STATIC(LogWorld, Log, All);
 DEFINE_LOG_CATEGORY(LogSpawn);
@@ -2656,6 +2657,51 @@ bool UWorld::AllowLevelLoadRequests()
 	return true;
 }
 
+bool UWorld::HandleDemoScrubCommand(const TCHAR* Cmd, FOutputDevice& Ar, UWorld* InWorld)
+{
+	FString TimeString;
+	if (!FParse::Token(Cmd, TimeString, 0))
+	{
+		Ar.Log(TEXT("You must specify a time"));
+	}
+	else if (DemoNetDriver != nullptr && DemoNetDriver->ReplayStreamer.IsValid() && DemoNetDriver->ServerConnection != nullptr && DemoNetDriver->ServerConnection->OwningActor != nullptr)
+	{
+		APlayerController* PlayerController = Cast<APlayerController>(DemoNetDriver->ServerConnection->OwningActor);
+		if (PlayerController != nullptr)
+		{
+			GetWorldSettings()->Pauser = PlayerController->PlayerState;
+			const uint32 Time = FCString::Atoi(*TimeString);
+			DemoNetDriver->ReplayStreamer->GotoTimeInMS(Time * 1000, FOnCheckpointReadyDelegate::CreateUObject(DemoNetDriver, &UDemoNetDriver::CheckpointReady));
+		}
+	}
+	return true;
+}
+
+bool UWorld::HandleDemoPauseCommand(const TCHAR* Cmd, FOutputDevice& Ar, UWorld* InWorld)
+{
+	FString TimeString;
+
+	AWorldSettings* WorldSettings = GetWorldSettings();
+	check(WorldSettings != nullptr);
+
+	if (WorldSettings->Pauser == nullptr)
+	{
+		if (DemoNetDriver != nullptr && DemoNetDriver->ServerConnection != nullptr && DemoNetDriver->ServerConnection->OwningActor != nullptr)
+		{
+			APlayerController* PlayerController = Cast<APlayerController>(DemoNetDriver->ServerConnection->OwningActor);
+			if (PlayerController != nullptr)
+			{
+				WorldSettings->Pauser = PlayerController->PlayerState;
+			}
+		}
+	}
+	else
+	{
+		WorldSettings->Pauser = nullptr;
+	}
+	return true;
+}
+
 bool UWorld::Exec( UWorld* InWorld, const TCHAR* Cmd, FOutputDevice& Ar )
 {
 	if( FParse::Command( &Cmd, TEXT("TRACETAG") ) )
@@ -2682,6 +2728,14 @@ bool UWorld::Exec( UWorld* InWorld, const TCHAR* Cmd, FOutputDevice& Ar )
 	{		
 		return HandleDemoStopCommand( Cmd, Ar, InWorld );
 	}
+	else if (FParse::Command(&Cmd, TEXT("DEMOSCRUB")))
+	{
+		return HandleDemoScrubCommand(Cmd, Ar, InWorld);
+	}
+	else if (FParse::Command(&Cmd, TEXT("DEMOPAUSE")))
+	{
+		return HandleDemoPauseCommand(Cmd, Ar, InWorld);
+	}
 	else if( ExecPhysCommands( Cmd, &Ar, InWorld ) )
 	{
 		return HandleLogActorCountsCommand( Cmd, Ar, InWorld );
