Integrate CL 2641147 and CL 2641326. This adds virtual functions and does not change existing functionality: allows games to override the APlayerState replacement/duplication/restore code in game modes when a player disconnects.

diff --git a/Engine/Source/Runtime/Engine/Classes/GameFramework/PlayerState.h b/Engine/Source/Runtime/Engine/Classes/GameFramework/PlayerState.h
--- a/Engine/Source/Runtime/Engine/Classes/GameFramework/PlayerState.h
+++ b/Engine/Source/Runtime/Engine/Classes/GameFramework/PlayerState.h
@@ -182,6 +182,12 @@ public:
 	/** Copy properties which need to be saved in inactive PlayerState */
 	virtual void CopyProperties(class APlayerState* PlayerState);
 
+	/** Called on the server when the owning player has disconnected, by default this method destroys this player state */
+	virtual void OnDeactivated();
+
+	/** Called on the server when the owning player has reconnected and this player state is added to the active players array */
+	virtual void OnReactivated();
+
 	/** called by seamless travel when initializing a player on the other side - copy properties to the new PlayerState that should persist */
 	virtual void SeamlessTravelTo(class APlayerState* NewPlayerState);
 
diff --git a/Engine/Source/Runtime/Engine/Private/GameMode.cpp b/Engine/Source/Runtime/Engine/Private/GameMode.cpp
--- a/Engine/Source/Runtime/Engine/Private/GameMode.cpp
+++ b/Engine/Source/Runtime/Engine/Private/GameMode.cpp
@@ -1550,7 +1550,7 @@ void AGameMode::AddInactivePlayer(APlayerState* PlayerState, APlayerController*
 		}
 	}
 
-	PlayerState->Destroy();
+	PlayerState->OnDeactivated();
 }
 
 
@@ -1592,6 +1592,7 @@ bool AGameMode::FindInactivePlayer(APlayerController* PC)
 			// in UnregisterPlayerWithSession()
 			OldPlayerState->SetUniqueId(NULL);
 			OldPlayerState->Destroy();
+			PC->PlayerState->OnReactivated();
 			return true;
 		}
 	}
diff --git a/Engine/Source/Runtime/Engine/Private/PlayerState.cpp b/Engine/Source/Runtime/Engine/Private/PlayerState.cpp
--- a/Engine/Source/Runtime/Engine/Private/PlayerState.cpp
+++ b/Engine/Source/Runtime/Engine/Private/PlayerState.cpp
@@ -96,6 +96,17 @@ void APlayerState::CopyProperties(APlayerState* PlayerState)
 	PlayerState->SavedNetworkAddress = SavedNetworkAddress;
 }
 
+void APlayerState::OnDeactivated()
+{
+	// By default we duplicate the inactive player state and destroy the old one
+	Destroy();
+}
+
+void APlayerState::OnReactivated()
+{
+	// Stub
+}
+
 void APlayerState::PostInitializeComponents()
 {
 	Super::PostInitializeComponents();
