Attempt to fix shadowed local variable buried in macro expansion by renaming a variable. CIS fix.

diff --git a/Engine/Source/Runtime/Engine/Private/Components/SplineMeshComponent.cpp b/Engine/Source/Runtime/Engine/Private/Components/SplineMeshComponent.cpp
--- a/Engine/Source/Runtime/Engine/Private/Components/SplineMeshComponent.cpp
+++ b/Engine/Source/Runtime/Engine/Private/Components/SplineMeshComponent.cpp
@@ -119,18 +119,18 @@ void FSplineMeshSceneProxy::InitResources( USplineMeshComponent* InComponent, in
 			VET_PackedNormal
 			);
 		
-		FColorVertexBuffer* LODColorVertexBuffer = &RenderData->ColorVertexBuffer;
+		FColorVertexBuffer* ColorVertexBufferLOD = &RenderData->ColorVertexBuffer;
 		// Override the buffer if it has been changed
 		if ( OverrideColorVertexBuffer != NULL )
 		{
-			LODColorVertexBuffer = OverrideColorVertexBuffer;
+			ColorVertexBufferLOD = OverrideColorVertexBuffer;
 		}
-		if ( LODColorVertexBuffer->GetNumVertices() > 0 )
+		if (ColorVertexBufferLOD->GetNumVertices() > 0)
 		{
 			Data.ColorComponent = FVertexStreamComponent(
-				LODColorVertexBuffer,
+				ColorVertexBufferLOD,
 				0,	// Struct offset to color
-				LODColorVertexBuffer->GetStride(),
+				ColorVertexBufferLOD->GetStride(),
 				VET_Color
 				);
 		}
