Another attempted fix for shadowed local variables in macro expansion.

diff --git a/Engine/Source/Runtime/Engine/Private/Components/SplineMeshComponent.cpp b/Engine/Source/Runtime/Engine/Private/Components/SplineMeshComponent.cpp
--- a/Engine/Source/Runtime/Engine/Private/Components/SplineMeshComponent.cpp
+++ b/Engine/Source/Runtime/Engine/Private/Components/SplineMeshComponent.cpp
@@ -98,51 +98,51 @@ void FSplineMeshSceneProxy::InitResources( USplineMeshComponent* InComponent, in
 		UStaticMesh*, Parent, InComponent->StaticMesh,
 		FColorVertexBuffer*, OverrideColorVertexBuffer, InOverrideColorVertexBuffer,
 		{
-		FLocalVertexFactory::DataType Data;
+		FLocalVertexFactory::DataType VertexFactoryData;
 
-		Data.PositionComponent = FVertexStreamComponent(
+		VertexFactoryData.PositionComponent = FVertexStreamComponent(
 			&RenderData->PositionVertexBuffer,
 			STRUCT_OFFSET(FPositionVertex, Position),
 			RenderData->PositionVertexBuffer.GetStride(),
 			VET_Float3
 			);
-		Data.TangentBasisComponents[0] = FVertexStreamComponent(
+		VertexFactoryData.TangentBasisComponents[0] = FVertexStreamComponent(
 			&RenderData->VertexBuffer,
 			STRUCT_OFFSET(FStaticMeshFullVertex, TangentX),
 			RenderData->VertexBuffer.GetStride(),
 			VET_PackedNormal
 			);
-		Data.TangentBasisComponents[1] = FVertexStreamComponent(
+		VertexFactoryData.TangentBasisComponents[1] = FVertexStreamComponent(
 			&RenderData->VertexBuffer,
 			STRUCT_OFFSET(FStaticMeshFullVertex, TangentZ),
 			RenderData->VertexBuffer.GetStride(),
 			VET_PackedNormal
 			);
 		
-		FColorVertexBuffer* ColorVertexBufferLOD = &RenderData->ColorVertexBuffer;
+		FColorVertexBuffer* LODColorVertexBuffer = &RenderData->ColorVertexBuffer;
 		// Override the buffer if it has been changed
 		if ( OverrideColorVertexBuffer != NULL )
 		{
-			ColorVertexBufferLOD = OverrideColorVertexBuffer;
+			LODColorVertexBuffer = OverrideColorVertexBuffer;
 		}
-		if (ColorVertexBufferLOD->GetNumVertices() > 0)
+		if ( LODColorVertexBuffer->GetNumVertices() > 0 )
 		{
-			Data.ColorComponent = FVertexStreamComponent(
-				ColorVertexBufferLOD,
+			VertexFactoryData.ColorComponent = FVertexStreamComponent(
+				LODColorVertexBuffer,
 				0,	// Struct offset to color
-				ColorVertexBufferLOD->GetStride(),
+				LODColorVertexBuffer->GetStride(),
 				VET_Color
 				);
 		}
 
-		Data.TextureCoordinates.Empty();
+		VertexFactoryData.TextureCoordinates.Empty();
 
 		if (!RenderData->VertexBuffer.GetUseFullPrecisionUVs())
 		{
 			int32 UVIndex;
 			for (UVIndex = 0; UVIndex < (int32)RenderData->VertexBuffer.GetNumTexCoords() - 1; UVIndex += 2)
 			{
-				Data.TextureCoordinates.Add(FVertexStreamComponent(
+				VertexFactoryData.TextureCoordinates.Add(FVertexStreamComponent(
 					&RenderData->VertexBuffer,
 					STRUCT_OFFSET(TStaticMeshFullVertexFloat16UVs<MAX_STATIC_TEXCOORDS>, UVs) + sizeof(FVector2DHalf) * UVIndex,
 					RenderData->VertexBuffer.GetStride(),
@@ -152,7 +152,7 @@ void FSplineMeshSceneProxy::InitResources( USplineMeshComponent* InComponent, in
 			// possible last UV channel if we have an odd number
 			if (UVIndex < (int32)RenderData->VertexBuffer.GetNumTexCoords())
 			{
-				Data.TextureCoordinates.Add(FVertexStreamComponent(
+				VertexFactoryData.TextureCoordinates.Add(FVertexStreamComponent(
 					&RenderData->VertexBuffer,
 					STRUCT_OFFSET(TStaticMeshFullVertexFloat16UVs<MAX_STATIC_TEXCOORDS>, UVs) + sizeof(FVector2DHalf) * UVIndex,
 					RenderData->VertexBuffer.GetStride(),
@@ -162,7 +162,7 @@ void FSplineMeshSceneProxy::InitResources( USplineMeshComponent* InComponent, in
 
 			if (Parent->LightMapCoordinateIndex >= 0 && (uint32)Parent->LightMapCoordinateIndex < RenderData->VertexBuffer.GetNumTexCoords())
 			{
-				Data.LightMapCoordinateComponent = FVertexStreamComponent(
+				VertexFactoryData.LightMapCoordinateComponent = FVertexStreamComponent(
 					&RenderData->VertexBuffer,
 					STRUCT_OFFSET(TStaticMeshFullVertexFloat16UVs<MAX_STATIC_TEXCOORDS>, UVs) + sizeof(FVector2DHalf) * Parent->LightMapCoordinateIndex,
 					RenderData->VertexBuffer.GetStride(),
@@ -175,7 +175,7 @@ void FSplineMeshSceneProxy::InitResources( USplineMeshComponent* InComponent, in
 			int32 UVIndex;
 			for (UVIndex = 0; UVIndex < (int32)RenderData->VertexBuffer.GetNumTexCoords() - 1; UVIndex += 2)
 			{
-				Data.TextureCoordinates.Add(FVertexStreamComponent(
+				VertexFactoryData.TextureCoordinates.Add(FVertexStreamComponent(
 					&RenderData->VertexBuffer,
 					STRUCT_OFFSET(TStaticMeshFullVertexFloat32UVs<MAX_STATIC_TEXCOORDS>, UVs) + sizeof(FVector2D) * UVIndex,
 					RenderData->VertexBuffer.GetStride(),
@@ -185,7 +185,7 @@ void FSplineMeshSceneProxy::InitResources( USplineMeshComponent* InComponent, in
 			// possible last UV channel if we have an odd number
 			if (UVIndex < (int32)RenderData->VertexBuffer.GetNumTexCoords())
 			{
-				Data.TextureCoordinates.Add(FVertexStreamComponent(
+				VertexFactoryData.TextureCoordinates.Add(FVertexStreamComponent(
 					&RenderData->VertexBuffer,
 					STRUCT_OFFSET(TStaticMeshFullVertexFloat32UVs<MAX_STATIC_TEXCOORDS>, UVs) + sizeof(FVector2D) * UVIndex,
 					RenderData->VertexBuffer.GetStride(),
@@ -195,7 +195,7 @@ void FSplineMeshSceneProxy::InitResources( USplineMeshComponent* InComponent, in
 
 			if (Parent->LightMapCoordinateIndex >= 0 && (uint32)Parent->LightMapCoordinateIndex < RenderData->VertexBuffer.GetNumTexCoords())
 			{
-				Data.LightMapCoordinateComponent = FVertexStreamComponent(
+				VertexFactoryData.LightMapCoordinateComponent = FVertexStreamComponent(
 					&RenderData->VertexBuffer,
 					STRUCT_OFFSET(TStaticMeshFullVertexFloat32UVs<MAX_STATIC_TEXCOORDS>, UVs) + sizeof(FVector2D) * Parent->LightMapCoordinateIndex,
 					RenderData->VertexBuffer.GetStride(),
@@ -204,7 +204,7 @@ void FSplineMeshSceneProxy::InitResources( USplineMeshComponent* InComponent, in
 			}
 		}
 
-		VertexFactory->SetData(Data);
+		VertexFactory->SetData(VertexFactoryData);
 
 		VertexFactory->InitResource();
 	});
