Third attempted CIS warning fix, though I don't see how UVIndex can possibly shadow itself if in a different scope.

diff --git a/Engine/Source/Runtime/Engine/Private/Components/SplineMeshComponent.cpp b/Engine/Source/Runtime/Engine/Private/Components/SplineMeshComponent.cpp
--- a/Engine/Source/Runtime/Engine/Private/Components/SplineMeshComponent.cpp
+++ b/Engine/Source/Runtime/Engine/Private/Components/SplineMeshComponent.cpp
@@ -98,21 +98,21 @@ void FSplineMeshSceneProxy::InitResources( USplineMeshComponent* InComponent, in
 		UStaticMesh*, Parent, InComponent->StaticMesh,
 		FColorVertexBuffer*, OverrideColorVertexBuffer, InOverrideColorVertexBuffer,
 		{
-		FLocalVertexFactory::DataType VertexFactoryData;
+		FLocalVertexFactory::DataType Data;
 
-		VertexFactoryData.PositionComponent = FVertexStreamComponent(
+		Data.PositionComponent = FVertexStreamComponent(
 			&RenderData->PositionVertexBuffer,
 			STRUCT_OFFSET(FPositionVertex, Position),
 			RenderData->PositionVertexBuffer.GetStride(),
 			VET_Float3
 			);
-		VertexFactoryData.TangentBasisComponents[0] = FVertexStreamComponent(
+		Data.TangentBasisComponents[0] = FVertexStreamComponent(
 			&RenderData->VertexBuffer,
 			STRUCT_OFFSET(FStaticMeshFullVertex, TangentX),
 			RenderData->VertexBuffer.GetStride(),
 			VET_PackedNormal
 			);
-		VertexFactoryData.TangentBasisComponents[1] = FVertexStreamComponent(
+		Data.TangentBasisComponents[1] = FVertexStreamComponent(
 			&RenderData->VertexBuffer,
 			STRUCT_OFFSET(FStaticMeshFullVertex, TangentZ),
 			RenderData->VertexBuffer.GetStride(),
@@ -127,7 +127,7 @@ void FSplineMeshSceneProxy::InitResources( USplineMeshComponent* InComponent, in
 		}
 		if ( LODColorVertexBuffer->GetNumVertices() > 0 )
 		{
-			VertexFactoryData.ColorComponent = FVertexStreamComponent(
+			Data.ColorComponent = FVertexStreamComponent(
 				LODColorVertexBuffer,
 				0,	// Struct offset to color
 				LODColorVertexBuffer->GetStride(),
@@ -135,14 +135,14 @@ void FSplineMeshSceneProxy::InitResources( USplineMeshComponent* InComponent, in
 				);
 		}
 
-		VertexFactoryData.TextureCoordinates.Empty();
+		Data.TextureCoordinates.Empty();
 
 		if (!RenderData->VertexBuffer.GetUseFullPrecisionUVs())
 		{
 			int32 UVIndex;
 			for (UVIndex = 0; UVIndex < (int32)RenderData->VertexBuffer.GetNumTexCoords() - 1; UVIndex += 2)
 			{
-				VertexFactoryData.TextureCoordinates.Add(FVertexStreamComponent(
+				Data.TextureCoordinates.Add(FVertexStreamComponent(
 					&RenderData->VertexBuffer,
 					STRUCT_OFFSET(TStaticMeshFullVertexFloat16UVs<MAX_STATIC_TEXCOORDS>, UVs) + sizeof(FVector2DHalf) * UVIndex,
 					RenderData->VertexBuffer.GetStride(),
@@ -152,7 +152,7 @@ void FSplineMeshSceneProxy::InitResources( USplineMeshComponent* InComponent, in
 			// possible last UV channel if we have an odd number
 			if (UVIndex < (int32)RenderData->VertexBuffer.GetNumTexCoords())
 			{
-				VertexFactoryData.TextureCoordinates.Add(FVertexStreamComponent(
+				Data.TextureCoordinates.Add(FVertexStreamComponent(
 					&RenderData->VertexBuffer,
 					STRUCT_OFFSET(TStaticMeshFullVertexFloat16UVs<MAX_STATIC_TEXCOORDS>, UVs) + sizeof(FVector2DHalf) * UVIndex,
 					RenderData->VertexBuffer.GetStride(),
@@ -162,7 +162,7 @@ void FSplineMeshSceneProxy::InitResources( USplineMeshComponent* InComponent, in
 
 			if (Parent->LightMapCoordinateIndex >= 0 && (uint32)Parent->LightMapCoordinateIndex < RenderData->VertexBuffer.GetNumTexCoords())
 			{
-				VertexFactoryData.LightMapCoordinateComponent = FVertexStreamComponent(
+				Data.LightMapCoordinateComponent = FVertexStreamComponent(
 					&RenderData->VertexBuffer,
 					STRUCT_OFFSET(TStaticMeshFullVertexFloat16UVs<MAX_STATIC_TEXCOORDS>, UVs) + sizeof(FVector2DHalf) * Parent->LightMapCoordinateIndex,
 					RenderData->VertexBuffer.GetStride(),
@@ -172,22 +172,22 @@ void FSplineMeshSceneProxy::InitResources( USplineMeshComponent* InComponent, in
 		}
 		else
 		{
-			int32 UVIndex;
-			for (UVIndex = 0; UVIndex < (int32)RenderData->VertexBuffer.GetNumTexCoords() - 1; UVIndex += 2)
+			int32 UVIndex2;
+			for (UVIndex2 = 0; UVIndex2 < (int32)RenderData->VertexBuffer.GetNumTexCoords() - 1; UVIndex2 += 2)
 			{
-				VertexFactoryData.TextureCoordinates.Add(FVertexStreamComponent(
+				Data.TextureCoordinates.Add(FVertexStreamComponent(
 					&RenderData->VertexBuffer,
-					STRUCT_OFFSET(TStaticMeshFullVertexFloat32UVs<MAX_STATIC_TEXCOORDS>, UVs) + sizeof(FVector2D) * UVIndex,
+					STRUCT_OFFSET(TStaticMeshFullVertexFloat32UVs<MAX_STATIC_TEXCOORDS>, UVs) + sizeof(FVector2D) * UVIndex2,
 					RenderData->VertexBuffer.GetStride(),
 					VET_Float4
 					));
 			}
 			// possible last UV channel if we have an odd number
-			if (UVIndex < (int32)RenderData->VertexBuffer.GetNumTexCoords())
+			if (UVIndex2 < (int32)RenderData->VertexBuffer.GetNumTexCoords())
 			{
-				VertexFactoryData.TextureCoordinates.Add(FVertexStreamComponent(
+				Data.TextureCoordinates.Add(FVertexStreamComponent(
 					&RenderData->VertexBuffer,
-					STRUCT_OFFSET(TStaticMeshFullVertexFloat32UVs<MAX_STATIC_TEXCOORDS>, UVs) + sizeof(FVector2D) * UVIndex,
+					STRUCT_OFFSET(TStaticMeshFullVertexFloat32UVs<MAX_STATIC_TEXCOORDS>, UVs) + sizeof(FVector2D) * UVIndex2,
 					RenderData->VertexBuffer.GetStride(),
 					VET_Float2
 					));
@@ -195,7 +195,7 @@ void FSplineMeshSceneProxy::InitResources( USplineMeshComponent* InComponent, in
 
 			if (Parent->LightMapCoordinateIndex >= 0 && (uint32)Parent->LightMapCoordinateIndex < RenderData->VertexBuffer.GetNumTexCoords())
 			{
-				VertexFactoryData.LightMapCoordinateComponent = FVertexStreamComponent(
+				Data.LightMapCoordinateComponent = FVertexStreamComponent(
 					&RenderData->VertexBuffer,
 					STRUCT_OFFSET(TStaticMeshFullVertexFloat32UVs<MAX_STATIC_TEXCOORDS>, UVs) + sizeof(FVector2D) * Parent->LightMapCoordinateIndex,
 					RenderData->VertexBuffer.GetStride(),
@@ -204,7 +204,7 @@ void FSplineMeshSceneProxy::InitResources( USplineMeshComponent* InComponent, in
 			}
 		}
 
-		VertexFactory->SetData(VertexFactoryData);
+		VertexFactory->SetData(Data);
 
 		VertexFactory->InitResource();
 	});
