Integration from UE4-Orion

2651743 - Add demospeed command which changes the playback speed of the current replay

2651796 - Fix issue with using the demoscrub command twice in a row before unpausing

diff --git a/Engine/Source/Runtime/Engine/Classes/Engine/World.h b/Engine/Source/Runtime/Engine/Classes/Engine/World.h
--- a/Engine/Source/Runtime/Engine/Classes/Engine/World.h
+++ b/Engine/Source/Runtime/Engine/Classes/Engine/World.h
@@ -2731,6 +2731,9 @@ private:
 	/** Utility function to handle Exec/Console Command for pausing and unpausing a replay */
 	bool HandleDemoPauseCommand(const TCHAR* Cmd, FOutputDevice& Ar, UWorld* InWorld);
 
+	/** Utility function to handle Exec/Console Command for setting the speed of a replay */
+	bool HandleDemoSpeedCommand(const TCHAR* Cmd, FOutputDevice& Ar, UWorld* InWorld);
+
 public:
 
 	// Destroys the current demo net driver
diff --git a/Engine/Source/Runtime/Engine/Private/World.cpp b/Engine/Source/Runtime/Engine/Private/World.cpp
--- a/Engine/Source/Runtime/Engine/Private/World.cpp
+++ b/Engine/Source/Runtime/Engine/Private/World.cpp
@@ -2671,7 +2671,7 @@ bool UWorld::HandleDemoScrubCommand(const TCHAR* Cmd, FOutputDevice& Ar, UWorld*
 		{
 			GetWorldSettings()->Pauser = PlayerController->PlayerState;
 			const uint32 Time = FCString::Atoi(*TimeString);
-			DemoNetDriver->ReplayStreamer->GotoTimeInMS(Time * 1000, FOnCheckpointReadyDelegate::CreateUObject(DemoNetDriver, &UDemoNetDriver::CheckpointReady));
+			DemoNetDriver->GotoTimeInSeconds(Time);
 		}
 	}
 	return true;
@@ -2702,6 +2702,26 @@ bool UWorld::HandleDemoPauseCommand(const TCHAR* Cmd, FOutputDevice& Ar, UWorld*
 	return true;
 }
 
+bool UWorld::HandleDemoSpeedCommand(const TCHAR* Cmd, FOutputDevice& Ar, UWorld* InWorld)
+{
+	FString TimeString;
+
+	AWorldSettings* WorldSettings = GetWorldSettings();
+	check(WorldSettings != nullptr);
+
+	FString SpeedString;
+	if (!FParse::Token(Cmd, SpeedString, 0))
+	{
+		Ar.Log(TEXT("You must specify a speed in the form of a float"));
+	}
+	else
+	{
+		const float Speed = FCString::Atof(*SpeedString);
+		WorldSettings->DemoPlayTimeDilation = Speed;
+	}
+	return true;
+}
+
 bool UWorld::Exec( UWorld* InWorld, const TCHAR* Cmd, FOutputDevice& Ar )
 {
 	if( FParse::Command( &Cmd, TEXT("TRACETAG") ) )
@@ -2736,6 +2756,10 @@ bool UWorld::Exec( UWorld* InWorld, const TCHAR* Cmd, FOutputDevice& Ar )
 	{
 		return HandleDemoPauseCommand(Cmd, Ar, InWorld);
 	}
+	else if (FParse::Command(&Cmd, TEXT("DEMOSPEED")))
+	{
+		return HandleDemoSpeedCommand(Cmd, Ar, InWorld);
+	}
 	else if( ExecPhysCommands( Cmd, &Ar, InWorld ) )
 	{
 		return HandleLogActorCountsCommand( Cmd, Ar, InWorld );
