"Proper" fix for the shadowed local variable now that I think I've grokked the macros properly.

diff --git a/Engine/Source/Runtime/Engine/Private/Components/SplineMeshComponent.cpp b/Engine/Source/Runtime/Engine/Private/Components/SplineMeshComponent.cpp
--- a/Engine/Source/Runtime/Engine/Private/Components/SplineMeshComponent.cpp
+++ b/Engine/Source/Runtime/Engine/Private/Components/SplineMeshComponent.cpp
@@ -96,7 +96,7 @@ void FSplineMeshSceneProxy::InitResources( USplineMeshComponent* InComponent, in
 		FSplineMeshVertexFactory*, VertexFactory, LODResources[InLODIndex].VertexFactory,
 		FStaticMeshLODResources*, RenderData, &InComponent->StaticMesh->RenderData->LODResources[InLODIndex],
 		UStaticMesh*, Parent, InComponent->StaticMesh,
-		FColorVertexBuffer*, OverrideColorVertexBuffer, InOverrideColorVertexBuffer,
+		FColorVertexBuffer*, OverridenColorVertexBuffer, InOverrideColorVertexBuffer,
 		{
 		FLocalVertexFactory::DataType Data;
 
@@ -121,9 +121,9 @@ void FSplineMeshSceneProxy::InitResources( USplineMeshComponent* InComponent, in
 		
 		FColorVertexBuffer* LODColorVertexBuffer = &RenderData->ColorVertexBuffer;
 		// Override the buffer if it has been changed
-		if ( OverrideColorVertexBuffer != NULL )
+		if ( OverridenColorVertexBuffer != NULL )
 		{
-			LODColorVertexBuffer = OverrideColorVertexBuffer;
+			LODColorVertexBuffer = OverridenColorVertexBuffer;
 		}
 		if ( LODColorVertexBuffer->GetNumVertices() > 0 )
 		{
@@ -172,22 +172,22 @@ void FSplineMeshSceneProxy::InitResources( USplineMeshComponent* InComponent, in
 		}
 		else
 		{
-			int32 UVIndex2;
-			for (UVIndex2 = 0; UVIndex2 < (int32)RenderData->VertexBuffer.GetNumTexCoords() - 1; UVIndex2 += 2)
+			int32 UVIndex;
+			for (UVIndex = 0; UVIndex < (int32)RenderData->VertexBuffer.GetNumTexCoords() - 1; UVIndex += 2)
 			{
 				Data.TextureCoordinates.Add(FVertexStreamComponent(
 					&RenderData->VertexBuffer,
-					STRUCT_OFFSET(TStaticMeshFullVertexFloat32UVs<MAX_STATIC_TEXCOORDS>, UVs) + sizeof(FVector2D) * UVIndex2,
+					STRUCT_OFFSET(TStaticMeshFullVertexFloat32UVs<MAX_STATIC_TEXCOORDS>, UVs) + sizeof(FVector2D) * UVIndex,
 					RenderData->VertexBuffer.GetStride(),
 					VET_Float4
 					));
 			}
 			// possible last UV channel if we have an odd number
-			if (UVIndex2 < (int32)RenderData->VertexBuffer.GetNumTexCoords())
+			if (UVIndex < (int32)RenderData->VertexBuffer.GetNumTexCoords())
 			{
 				Data.TextureCoordinates.Add(FVertexStreamComponent(
 					&RenderData->VertexBuffer,
-					STRUCT_OFFSET(TStaticMeshFullVertexFloat32UVs<MAX_STATIC_TEXCOORDS>, UVs) + sizeof(FVector2D) * UVIndex2,
+					STRUCT_OFFSET(TStaticMeshFullVertexFloat32UVs<MAX_STATIC_TEXCOORDS>, UVs) + sizeof(FVector2D) * UVIndex,
 					RenderData->VertexBuffer.GetStride(),
 					VET_Float2
 					));
