UE4 - fix HISMC as occluder

diff --git a/Engine/Source/Runtime/Engine/Classes/Components/PrimitiveComponent.h b/Engine/Source/Runtime/Engine/Classes/Components/PrimitiveComponent.h
--- a/Engine/Source/Runtime/Engine/Classes/Components/PrimitiveComponent.h
+++ b/Engine/Source/Runtime/Engine/Classes/Components/PrimitiveComponent.h
@@ -241,7 +241,7 @@ public:
 	 * This should generally be true for all objects, and let the renderer make decisions about whether to render objects in the depth only pass.
 	 * @todo - if any rendering features rely on a complete depth only pass, this variable needs to go away.
 	 */
-	UPROPERTY()
+	UPROPERTY(EditAnywhere, AdvancedDisplay, BlueprintReadOnly, Category=Rendering)
 	uint32 bUseAsOccluder:1;
 
 	/** If this is True, this component can be selected in the editor. */
diff --git a/Engine/Source/Runtime/Engine/Private/HierarchicalInstancedStaticMesh.cpp b/Engine/Source/Runtime/Engine/Private/HierarchicalInstancedStaticMesh.cpp
--- a/Engine/Source/Runtime/Engine/Private/HierarchicalInstancedStaticMesh.cpp
+++ b/Engine/Source/Runtime/Engine/Private/HierarchicalInstancedStaticMesh.cpp
@@ -1094,7 +1094,7 @@ void FHierarchicalStaticMeshSceneProxy::FillDynamicMeshElements(FMeshElementColl
 					MeshElement.bCanApplyViewModeOverrides = true;
 					MeshElement.bUseSelectionOutline = ElementParams.BatchRenderSelection[SelectionGroupIndex];
 					MeshElement.bUseWireframeSelectionColoring = ElementParams.BatchRenderSelection[SelectionGroupIndex];
-					MeshElement.bUseAsOccluder = false;
+					MeshElement.bUseAsOccluder = ShouldUseAsOccluder();
 
 					if (!bDidStats)
 					{
@@ -1540,6 +1540,7 @@ UHierarchicalInstancedStaticMeshComponent::UHierarchicalInstancedStaticMeshCompo
 	, AccumulatedNavigationDirtyArea(0)
 {
 	bCanEverAffectNavigation = true;
+	bUseAsOccluder = false;
 }
 
 UHierarchicalInstancedStaticMeshComponent::~UHierarchicalInstancedStaticMeshComponent()
