Fixed a bug with timers being set and paused in the same frame.

diff --git a/Engine/Source/Runtime/Engine/Private/TimerManager.cpp b/Engine/Source/Runtime/Engine/Private/TimerManager.cpp
--- a/Engine/Source/Runtime/Engine/Private/TimerManager.cpp
+++ b/Engine/Source/Runtime/Engine/Private/TimerManager.cpp
@@ -346,8 +346,11 @@ void FTimerManager::InternalPauseTimer( FTimerData const* TimerToPause, int32 Ti
 			FTimerData& NewTimer = PausedTimerList[NewIndex];
 			NewTimer.Status = ETimerStatus::Paused;
 
-			// Store time remaining in ExpireTime while paused
-			NewTimer.ExpireTime = NewTimer.ExpireTime - InternalTime;
+			// Store time remaining in ExpireTime while paused. Don't do this if the timer is in the pending list.
+			if (PreviousStatus != ETimerStatus::Pending)
+			{
+				NewTimer.ExpireTime = NewTimer.ExpireTime - InternalTime;
+			}
 		}
 
 		// Remove from previous TArray
