Changed UTextRenderComponent to respond to the culture being changed, rather than constantly tick

UE-19848 - TextRenderComponent should never tick

diff --git a/Engine/Source/Runtime/Engine/Classes/Components/TextRenderComponent.h b/Engine/Source/Runtime/Engine/Classes/Components/TextRenderComponent.h
--- a/Engine/Source/Runtime/Engine/Classes/Components/TextRenderComponent.h
+++ b/Engine/Source/Runtime/Engine/Classes/Components/TextRenderComponent.h
@@ -158,17 +158,9 @@ class ENGINE_API UTextRenderComponent : public UPrimitiveComponent
 	virtual FBoxSphereBounds CalcBounds(const FTransform& LocalToWorld) const override;
 	// End USceneComponent interface.
 
-	// Begin UActorComponent interface.
-	virtual void TickComponent(float DeltaTime, enum ELevelTick TickType, FActorComponentTickFunction* ThisTickFunction) override;
-	// End UActorComponent interface.
-
 	// Begin UObject interface.
 	virtual void PostLoad() override;
 	// End UObject interface.
-
-private:
-	/** The state of the text the last time it was updated (used to allow updates when the text is changed due to a culture update) */
-	FTextSnapshot TextLastUpdate;
 };
 
 
diff --git a/Engine/Source/Runtime/Engine/Private/Components/TextRenderComponent.cpp b/Engine/Source/Runtime/Engine/Private/Components/TextRenderComponent.cpp
--- a/Engine/Source/Runtime/Engine/Private/Components/TextRenderComponent.cpp
+++ b/Engine/Source/Runtime/Engine/Private/Components/TextRenderComponent.cpp
@@ -687,6 +687,38 @@ bool  FTextRenderSceneProxy::BuildStringMesh( TArray<FDynamicMeshVertex>& OutVer
 
 // ------------------------------------------------------
 
+/** Watches for culture changes and updates all live UTextRenderComponent components */
+class FTextRenderComponentCultureChangedFixUp
+{
+public:
+	FTextRenderComponentCultureChangedFixUp()
+		: ImplPtr(MakeShareable(new FImpl()))
+	{
+		ImplPtr->Register();
+	}
+
+private:
+	struct FImpl : public TSharedFromThis<FImpl>
+	{
+		void Register()
+		{
+			FTextLocalizationManager::Get().OnTextRevisionChangedEvent.AddSP(this, &FImpl::HandleLocalizedTextChanged);
+		}
+
+		void HandleLocalizedTextChanged()
+		{
+			for (UTextRenderComponent* TextRenderComponent : TObjectRange<UTextRenderComponent>())
+			{
+				TextRenderComponent->MarkRenderStateDirty();
+			}
+		}
+	};
+
+	TSharedPtr<FImpl> ImplPtr;
+};
+
+// ------------------------------------------------------
+
 UTextRenderComponent::UTextRenderComponent(const FObjectInitializer& ObjectInitializer)
 	: Super(ObjectInitializer)
 {
@@ -703,11 +735,16 @@ UTextRenderComponent::UTextRenderComponent(const FObjectInitializer& ObjectIniti
 	};
 	static FConstructorStatics ConstructorStatics;
 
-	PrimaryComponentTick.bCanEverTick = true;
-	bTickInEditor = true;
+	{
+		// Static used to watch for culture changes and update all live UTextRenderComponent components
+		// In this constructor so that it has a known initialization order, and is only created when we need it
+		static FTextRenderComponentCultureChangedFixUp TextRenderComponentCultureChangedFixUp;
+	}
+
+	PrimaryComponentTick.bCanEverTick = false;
+	bTickInEditor = false;
 
 	Text = LOCTEXT("DefaultText", "Text");
-	TextLastUpdate = FTextSnapshot(Text);
 
 	Font = ConstructorStatics.Font.Get();
 	TextMaterial = ConstructorStatics.TextMaterial.Get();
@@ -866,7 +903,6 @@ void UTextRenderComponent::SetText(const FText& Value)
 void UTextRenderComponent::K2_SetText(const FText& Value)
 {
 	Text = Value;
-	TextLastUpdate = FTextSnapshot(Text); // update this immediately as we're calling MarkRenderStateDirty ourselves
 	MarkRenderStateDirty();
 }
 
@@ -939,22 +975,6 @@ FVector UTextRenderComponent::GetTextWorldSize() const
 	return Bounds.GetBox().GetSize();
 }
 
-void UTextRenderComponent::TickComponent(float DeltaTime, enum ELevelTick TickType, FActorComponentTickFunction* ThisTickFunction)
-{
-	if (!TextLastUpdate.IdenticalTo(Text))
-	{
-		// The pointer used by the bound text has changed, however the text may still be the same - check that now
-		if (!TextLastUpdate.IsDisplayStringEqualTo(Text))
-		{
-			// The source text has changed, so we need to update our render data
-			MarkRenderStateDirty();	
-		}
-
-		// Update this even if the text is lexically identical, as it will update the pointer compared by IdenticalTo for the next Tick
-		TextLastUpdate = FTextSnapshot(Text);
-	}
-}
-
 void UTextRenderComponent::PostLoad()
 {
 	// Try and fix up assets created before the vertical alignment fix was implemented. Because we didn't flag that
