Removed net connection PingAck code.

NOTE: This breaks net compatibility.

diff --git a/Engine/Source/Runtime/Engine/Classes/Engine/ChildConnection.h b/Engine/Source/Runtime/Engine/Classes/Engine/ChildConnection.h
--- a/Engine/Source/Runtime/Engine/Classes/Engine/ChildConnection.h
+++ b/Engine/Source/Runtime/Engine/Classes/Engine/ChildConnection.h
@@ -50,7 +50,7 @@ public:
 		Parent->AssertValid();
 	}
 
-	virtual void SendAck( int32 PacketId, bool FirstTime=1, bool bHavePingAckData=0, uint32 PingAckData=0 ) override
+	virtual void SendAck(int32 PacketId, bool FirstTime=1) override
 	{
 	}
 
diff --git a/Engine/Source/Runtime/Engine/Classes/Engine/NetConnection.h b/Engine/Source/Runtime/Engine/Classes/Engine/NetConnection.h
--- a/Engine/Source/Runtime/Engine/Classes/Engine/NetConnection.h
+++ b/Engine/Source/Runtime/Engine/Classes/Engine/NetConnection.h
@@ -122,12 +122,6 @@ struct DelayedPacket
 };
 #endif
 
-// The interval between ack packets, which is used to decide which acks are used for ping validation checks (must be power of two)
-#define PING_ACK_CHECKSUM_PACKET_INTERVAL 16
-
-// Used by the client, for setting a minimum delay between PingAck's
-//	(optionally, 'PING_ACK_CHECKSUM_PACKET_INTERVAL' can be tweaked, so that the interval checks take a similar amount of time as this delay)
-#define PING_ACK_CHECKSUM_DELAY 0.5
 
 UCLASS(customConstructor, Abstract, MinimalAPI, transient, config=Engine)
 class UNetConnection : public UPlayer
@@ -276,9 +270,6 @@ public:
 	int32			OutPacketId;			// Most recently sent packet.
 	int32 			OutAckPacketId;			// Most recently acked outgoing packet.
 
-	uint32			PingChecksumAckCache[MAX_PACKETID / PING_ACK_CHECKSUM_PACKET_INTERVAL];	// Caches packet data on the server, for verifying pings
-	float			LastPingChecksumAck;													// The time of the most ping checksum generated on the client
-	int32			LastPingChecksumAckPacketId;											// The PacketId of the last PingAck, on the server
 	bool			bLastHasServerFrameTime;
 
 	// Channel table.
@@ -433,7 +424,7 @@ public:
 	ENGINE_API virtual void AssertValid();
 
 	/** Send an acknowledgment. */
-	ENGINE_API virtual void SendAck( int32 PacketId, bool FirstTime=1, bool bHasPingChecksum=false, uint32 PingChecksum=0 );
+	ENGINE_API virtual void SendAck( int32 PacketId, bool FirstTime=1);
 
 	/**
 	 * flushes any pending data, bundling it into a packet and sending it via LowLevelSend()
diff --git a/Engine/Source/Runtime/Engine/Private/NetConnection.cpp b/Engine/Source/Runtime/Engine/Private/NetConnection.cpp
--- a/Engine/Source/Runtime/Engine/Private/NetConnection.cpp
+++ b/Engine/Source/Runtime/Engine/Private/NetConnection.cpp
@@ -65,8 +65,6 @@ UNetConnection::UNetConnection(const FObjectInitializer& ObjectInitializer)
 ,	InPacketId			( -1 )
 ,	OutPacketId			( 0 ) // must be initialized as OutAckPacketId + 1 so loss of first packet can be detected
 ,	OutAckPacketId		( -1 )
-,	LastPingChecksumAck	( 0.f )
-,	LastPingChecksumAckPacketId( -1 )
 ,	bLastHasServerFrameTime( false )
 ,	ClientWorldPackageName( NAME_None )
 {
@@ -661,27 +659,6 @@ void UNetConnection::ReceivedNak( int32 NakPacketId )
 	}
 }
 
-/** 
- * Generates a 32 bit value from some input data blob
- * @BunchData	Data blob used to produce 32 bit value
- * @NumBits		Number of bits in data blob
- * @PacketId	PacketId of the packet used to produce this ack data
-*/
-static uint32 CalcPingChecksum( const uint8* BunchData, const int32 NumBits, const int32 PacketId )
-{
-	// Simply walk the bunch data, based upon OutPacketId, to get 'good enough' random data for verification
-	const int32	BunchBytesFloor		= NumBits / 8;
-	const int32	PingChecksumAckIdx	= ( PacketId % MAX_PACKETID ) / PING_ACK_CHECKSUM_PACKET_INTERVAL;
-
-	return	( BunchData[( PingChecksumAckIdx + 0 ) % BunchBytesFloor] << 24 ) + ( BunchData[( PingChecksumAckIdx + 1 ) % BunchBytesFloor] << 16 ) +
-		( BunchData[( PingChecksumAckIdx + 2 ) % BunchBytesFloor] << 8 ) + ( BunchData[( PingChecksumAckIdx + 3 ) % BunchBytesFloor] );
-}
-
-static bool IsAckChecksumPacket( const int32 PacketId )
-{
-	return ( PacketId % PING_ACK_CHECKSUM_PACKET_INTERVAL ) == 0;
-}
-
 void UNetConnection::ReceivedPacket( FBitReader& Reader )
 {
 	AssertValid();
@@ -724,9 +701,6 @@ void UNetConnection::ReceivedPacket( FBitReader& Reader )
 		return;
 	}
 
-	// Detect packets on the client, which should trigger PingAck verification
-	bool bGotPingChecksum			= false;
-	uint32 OutPingChecksum			= 0;
 
 	bool bSkipAck = false;
 
@@ -780,21 +754,6 @@ void UNetConnection::ReceivedPacket( FBitReader& Reader )
 			}
 #endif
 
-			// Detect ack packets on the server, containing PingChecksum
-			const bool bCanHavePingChecksum	= Driver->IsServer() && IsAckChecksumPacket( AckPacketId );
-			bool bHasPingChecksum			= false;
-			uint32 InPingChecksum			= 0;
-
-			if ( bCanHavePingChecksum )
-			{
-				bHasPingChecksum = !!Reader.ReadBit();
-
-				if ( bHasPingChecksum )
-				{
-					Reader.Serialize( &InPingChecksum, sizeof( uint32 ) );
-				}
-			}
-
 			// Resend any old reliable packets that the receiver hasn't acknowledged.
 			if( AckPacketId>OutAckPacketId )
 			{
@@ -836,13 +795,7 @@ void UNetConnection::ReceivedPacket( FBitReader& Reader )
 
 				if (PlayerController != NULL)
 				{
-					// Verify PingAck's, and pass notification up
-					const bool bServerVerified = bCanHavePingChecksum && bHasPingChecksum && PingChecksumAckCache[( AckPacketId % MAX_PACKETID ) / PING_ACK_CHECKSUM_PACKET_INTERVAL] == InPingChecksum;
-
-					if ( bServerVerified || !Driver->IsServer() )
-					{
-						PlayerController->UpdatePing(NewLag);
-					}
+					PlayerController->UpdatePing(NewLag);
 				}
 			}
 
@@ -853,6 +806,7 @@ void UNetConnection::ReceivedPacket( FBitReader& Reader )
 
 			// Forward the ack to the channel.
 			UE_LOG(LogNetTraffic, Verbose, TEXT("   Received ack %i (%.1f)"), AckPacketId, (Reader.GetPosBits()-StartPos)/8.f );
+
 			for( int32 i=OpenChannels.Num()-1; i>=0; i-- )
 			{
 				UChannel* Channel = OpenChannels[i];
@@ -991,15 +945,6 @@ void UNetConnection::ReceivedPacket( FBitReader& Reader )
 				return;
 			}
 
-			const bool bCanHavePingChecksum = !Driver->IsServer() && IsAckChecksumPacket(PacketId);
-
-			// Handle grabbing of PingChecksum, for client PingAck verification (need a minimum of 32bits data, for OutPingChecksum)
-			if ( bCanHavePingChecksum && !bGotPingChecksum && BunchDataBits >= 32 && FMath::Abs( Driver->Time - LastPingChecksumAck ) >= PING_ACK_CHECKSUM_DELAY )
-			{
-				OutPingChecksum = CalcPingChecksum( Bunch.GetData(), Bunch.GetNumBits(), PacketId );
-				bGotPingChecksum = true;
-				LastPingChecksumAck = Driver->Time;
-			}
 
 			// Receiving data.
 			UChannel* Channel = Channels[Bunch.ChIndex];
@@ -1095,7 +1040,7 @@ void UNetConnection::ReceivedPacket( FBitReader& Reader )
 	// Acknowledge the packet.
 	if ( !bSkipAck )
 	{
-		SendAck( PacketId, true, bGotPingChecksum, OutPingChecksum );
+		SendAck(PacketId, true);
 	}
 }
 
@@ -1198,7 +1143,7 @@ void UNetConnection::PurgeAcks()
 }
 
 
-void UNetConnection::SendAck( int32 AckPacketId, bool FirstTime/*=1*/, bool bHasPingChecksum/*=0*/, uint32 PingChecksum/*=0*/ )
+void UNetConnection::SendAck(int32 AckPacketId, bool FirstTime/*=1*/)
 {
 	ValidateSendBuffer();
 
@@ -1230,18 +1175,6 @@ void UNetConnection::SendAck( int32 AckPacketId, bool FirstTime/*=1*/, bool bHas
 		AckData.WriteBit( 0 );
 #endif
 
-		const bool bCanHavePingChecksum = !Driver->IsServer() && IsAckChecksumPacket( AckPacketId );
-
-		if ( bCanHavePingChecksum )
-		{
-			AckData.WriteBit( bHasPingChecksum );
-
-			if ( bHasPingChecksum )
-			{
-				AckData.Serialize( &PingChecksum, sizeof( uint32 ) );
-			}
-		}
-
 		NETWORK_PROFILER( GNetworkProfiler.TrackSendAck( AckData.GetNumBits() ) );
 
 		WriteBitsToSendBuffer( AckData.GetData(), AckData.GetNumBits(), nullptr, 0, EWriteBitsDataType::Ack );
@@ -1332,19 +1265,6 @@ int32 UNetConnection::SendRawBunch( FOutBunch& Bunch, bool InAllowMerge )
 		Driver->NetGUIDOutBytes += (Header.GetNumBits() + Bunch.GetNumBits()) >> 3;
 	}
 
-	// Verified client ping tracking - caches some semi-random bytes of the packet, for ping validation
-	if ( Driver->IsServer() && Bunch.PacketId != LastPingChecksumAckPacketId && IsAckChecksumPacket( Bunch.PacketId ) && Bunch.GetNumBits() >= 32 )
-	{
-		const uint32 PingChecksum = CalcPingChecksum( Bunch.GetData(), Bunch.GetNumBits(), Bunch.PacketId );
-
-		const int32 PingChecksumAckIdx = ( Bunch.PacketId % MAX_PACKETID ) / PING_ACK_CHECKSUM_PACKET_INTERVAL;
-
-		PingChecksumAckCache[PingChecksumAckIdx] = PingChecksum;
-
-		// Multiple bunches get written per-packet, but PingAck only uses the first bunch, so don't process more bunches this PacketId
-		LastPingChecksumAckPacketId = Bunch.PacketId;
-	}
-
 	return Bunch.PacketId;
 }
 
