SplineComponent: Add RemoveSplinePoint and AddSplinePointAtIndex functions

diff --git a/Engine/Source/Runtime/Engine/Classes/Components/SplineComponent.h b/Engine/Source/Runtime/Engine/Classes/Components/SplineComponent.h
--- a/Engine/Source/Runtime/Engine/Classes/Components/SplineComponent.h
+++ b/Engine/Source/Runtime/Engine/Classes/Components/SplineComponent.h
@@ -173,6 +173,14 @@ public:
 	UFUNCTION(BlueprintCallable, Category = Spline)
 	void AddSplinePoint(const FVector& Position, ESplineCoordinateSpace::Type CoordinateSpace);
 
+	/** Adds a point to the spline at the specified index*/
+	UFUNCTION(BlueprintCallable, Category = Spline)
+	void AddSplinePointAtIndex(const FVector& Position, int32 Index, ESplineCoordinateSpace::Type CoordinateSpace);
+
+	/** Removes point at specified index from the spline */
+	UFUNCTION(BlueprintCallable, Category = Spline)
+	void RemoveSplinePoint(const int32 Index);
+
 	/** Adds a world space point to the spline */
 	UFUNCTION(BlueprintCallable, Category = Spline, meta = (DeprecatedFunction, DeprecationMessage = "Please use AddSplinePoint, specifying SplineCoordinateSpace::World"))
 	void AddSplineWorldPoint(const FVector& Position);
diff --git a/Engine/Source/Runtime/Engine/Private/Components/SplineComponent.cpp b/Engine/Source/Runtime/Engine/Private/Components/SplineComponent.cpp
--- a/Engine/Source/Runtime/Engine/Private/Components/SplineComponent.cpp
+++ b/Engine/Source/Runtime/Engine/Private/Components/SplineComponent.cpp
@@ -462,6 +462,46 @@ void USplineComponent::AddSplinePoint(const FVector& Position, ESplineCoordinate
 	UpdateSpline();
 }
 
+void USplineComponent::AddSplinePointAtIndex(const FVector& Position, int32 Index, ESplineCoordinateSpace::Type CoordinateSpace)
+{
+	const FVector TransformedPosition = (CoordinateSpace == ESplineCoordinateSpace::World) ?
+		ComponentToWorld.InverseTransformPosition(Position) : Position;
+
+	const float InKey = static_cast<float>(Index);
+
+	if (((Index >= 0) && 
+		(Index < SplineInfo.Points.Num())) &&
+		(Index < SplineRotInfo.Points.Num()) &&
+		(Index < SplineScaleInfo.Points.Num()))
+	{
+		SplineInfo.Points.Insert(FInterpCurvePoint<FVector>(InKey, TransformedPosition, FVector::ZeroVector, FVector::ZeroVector, CIM_CurveAuto), Index);
+		SplineRotInfo.Points.Insert(FInterpCurvePoint<FQuat>(InKey, FQuat::Identity, FQuat::Identity, FQuat::Identity, CIM_CurveAuto), Index);
+		SplineScaleInfo.Points.Insert(FInterpCurvePoint<FVector>(InKey, FVector(1.0f), FVector::ZeroVector, FVector::ZeroVector, CIM_CurveAuto), Index);
+	}
+
+	UpdateSpline();
+}
+
+void USplineComponent::RemoveSplinePoint(const int32 Index)
+{
+	int32 Count = 1;
+
+	if ((Index >= 0) &&
+		(SplineInfo.Points.Num() >= 0) &&
+		(SplineRotInfo.Points.Num() >= 0) &&
+		(SplineScaleInfo.Points.Num() >= 0) &&
+		(Index + Count <= SplineInfo.Points.Num()) &&
+		(Index + Count <= SplineRotInfo.Points.Num()) &&
+		(Index + Count <= SplineScaleInfo.Points.Num())
+		)
+	{
+		SplineInfo.Points.RemoveAt(Index, Count, false);
+		SplineRotInfo.Points.RemoveAt(Index, Count, false);
+		SplineScaleInfo.Points.RemoveAt(Index, Count, false);
+	}
+
+	UpdateSpline();
+}
 
 void USplineComponent::SetSplinePoints(const TArray<FVector>& Points, ESplineCoordinateSpace::Type CoordinateSpace)
 {
