Forcing navmesh debug drawing update when navmesh generation is finished #UE4

Merged over from Orion branch

diff --git a/Engine/Source/Runtime/Engine/Classes/AI/Navigation/NavMeshRenderingComponent.h b/Engine/Source/Runtime/Engine/Classes/AI/Navigation/NavMeshRenderingComponent.h
--- a/Engine/Source/Runtime/Engine/Classes/AI/Navigation/NavMeshRenderingComponent.h
+++ b/Engine/Source/Runtime/Engine/Classes/AI/Navigation/NavMeshRenderingComponent.h
@@ -30,6 +30,8 @@ public:
 	virtual FBoxSphereBounds CalcBounds(const FTransform &LocalToWorld) const override;
 	// End USceneComponent Interface
 
+	void ForceUpdate() { bForceUpdate = true; }
+	bool IsForcingUpdate() const { return bForceUpdate; }
 	void GatherData(FNavMeshSceneProxyData& DebugDrawData) const;
 
 	static bool IsNavigationShowFlagSet(const UWorld* World);
@@ -39,5 +41,6 @@ protected:
 
 protected:
 	uint32 bCollectNavigationData : 1;
+	uint32 bForceUpdate : 1;
 	FTimerHandle TimerHandle;
 };
diff --git a/Engine/Source/Runtime/Engine/Classes/AI/Navigation/RecastNavMesh.h b/Engine/Source/Runtime/Engine/Classes/AI/Navigation/RecastNavMesh.h
--- a/Engine/Source/Runtime/Engine/Classes/AI/Navigation/RecastNavMesh.h
+++ b/Engine/Source/Runtime/Engine/Classes/AI/Navigation/RecastNavMesh.h
@@ -778,7 +778,7 @@ public:
 	void UpdateDrawing();
 
 	/** Creates a task to be executed on GameThread calling UpdateDrawing */
-	void RequestDrawingUpdate();
+	void RequestDrawingUpdate(bool bForce = false);
 
 	/** Invalidates active paths that go through changed tiles  */
 	void InvalidateAffectedPaths(const TArray<uint32>& ChangedTiles);
diff --git a/Engine/Source/Runtime/Engine/Private/AI/Navigation/RecastNavMesh.cpp b/Engine/Source/Runtime/Engine/Private/AI/Navigation/RecastNavMesh.cpp
--- a/Engine/Source/Runtime/Engine/Private/AI/Navigation/RecastNavMesh.cpp
+++ b/Engine/Source/Runtime/Engine/Private/AI/Navigation/RecastNavMesh.cpp
@@ -458,7 +458,8 @@ void ARecastNavMesh::UpdateNavMeshDrawing()
 		}
 	}
 #else // DO_NAVMESH_DEBUG_DRAWING_PER_TILE
-	if (RenderingComp != NULL && RenderingComp->bVisible && UNavMeshRenderingComponent::IsNavigationShowFlagSet(GetWorld()))
+	UNavMeshRenderingComponent* NavMeshRenderComp = Cast<UNavMeshRenderingComponent>(RenderingComp);
+	if (NavMeshRenderComp != nullptr && NavMeshRenderComp->bVisible && (NavMeshRenderComp->IsForcingUpdate() || UNavMeshRenderingComponent::IsNavigationShowFlagSet(GetWorld())))
 	{
 		RenderingComp->MarkRenderStateDirty();
 	}
@@ -1399,11 +1400,16 @@ void ARecastNavMesh::GetDebugGeometry(FRecastDebugGeometry& OutGeometry, int32 T
 	}
 }
 
-void ARecastNavMesh::RequestDrawingUpdate()
+void ARecastNavMesh::RequestDrawingUpdate(bool bForce)
 {
 #if !UE_BUILD_SHIPPING
-	if (UNavMeshRenderingComponent::IsNavigationShowFlagSet(GetWorld()))
+	if (bForce || UNavMeshRenderingComponent::IsNavigationShowFlagSet(GetWorld()))
 	{
+		if (bForce && Cast<UNavMeshRenderingComponent>(RenderingComp))
+		{
+			Cast<UNavMeshRenderingComponent>(RenderingComp)->ForceUpdate();
+		}
+
 		DECLARE_CYCLE_STAT(TEXT("FSimpleDelegateGraphTask.Requesting navmesh redraw"),
 		STAT_FSimpleDelegateGraphTask_RequestingNavmeshRedraw,
 			STATGROUP_TaskGraphTasks);
@@ -1576,6 +1582,9 @@ void ARecastNavMesh::OnNavMeshGenerationFinished()
 				}
 			}
 		}
+
+		// force navmesh drawing update
+		RequestDrawingUpdate(/*bForce=*/true);		
 #endif// WITH_EDITOR
 
 		if (World->GetNavigationSystem())
diff --git a/Engine/Source/Runtime/Engine/Private/Components/NavMeshRenderingComponent.cpp b/Engine/Source/Runtime/Engine/Private/Components/NavMeshRenderingComponent.cpp
--- a/Engine/Source/Runtime/Engine/Private/Components/NavMeshRenderingComponent.cpp
+++ b/Engine/Source/Runtime/Engine/Private/Components/NavMeshRenderingComponent.cpp
@@ -101,10 +101,11 @@ void UNavMeshRenderingComponent::TimerFunction()
 	}
 #endif // WITH_EDITOR
 
-	const bool bShowNavigation = IsNavigationShowFlagSet(World);
+	const bool bShowNavigation = bForceUpdate || IsNavigationShowFlagSet(World);
 
 	if (bShowNavigation != !!bCollectNavigationData && bShowNavigation == true)
 	{
+		bForceUpdate = false;
 		bCollectNavigationData = bShowNavigation;
 		MarkRenderStateDirty();
 	}
