Curve Editor: Set the curve name so that the tooltips will have the curve name, instead of "None".

diff --git a/Engine/Source/Runtime/Engine/Private/CurveFloat.cpp b/Engine/Source/Runtime/Engine/Private/CurveFloat.cpp
--- a/Engine/Source/Runtime/Engine/Private/CurveFloat.cpp
+++ b/Engine/Source/Runtime/Engine/Private/CurveFloat.cpp
@@ -55,7 +55,7 @@ TArray<FRichCurveEditInfoConst> UCurveFloat::GetCurves() const
 TArray<FRichCurveEditInfo> UCurveFloat::GetCurves()
 {
 	TArray<FRichCurveEditInfo> Curves;
-	Curves.Add(FRichCurveEditInfo(&FloatCurve));
+	Curves.Add(FRichCurveEditInfo(&FloatCurve, FName(*GetName())));
 	return Curves;
 }
 
