Fixed: NavigationSystem.DiscardNavigationDataChunks is being called even for perfectly valid data #jira UE-19658
Navigation will no longer trigger full navmesh rebuild whenewer map is loaded

diff --git a/Engine/Source/Runtime/Engine/Private/AI/Navigation/NavigationSystem.cpp b/Engine/Source/Runtime/Engine/Private/AI/Navigation/NavigationSystem.cpp
--- a/Engine/Source/Runtime/Engine/Private/AI/Navigation/NavigationSystem.cpp
+++ b/Engine/Source/Runtime/Engine/Private/AI/Navigation/NavigationSystem.cpp
@@ -514,6 +514,22 @@ bool UNavigationSystem::ConditionalPopulateNavOctree()
 			}
 		}
 	}
+	
+	// Add all found elements to octree, this will not add new dirty areas to navigation
+	if (PendingOctreeUpdates.Num())
+	{
+		TArray<FNavigationDirtyArea> SavedDirtyAreas; 
+		Exchange(SavedDirtyAreas, DirtyAreas);
+
+		for (TSet<FNavigationDirtyElement>::TIterator It(PendingOctreeUpdates); It; ++It)
+		{
+			AddElementToNavOctree(*It);
+		}
+		PendingOctreeUpdates.Empty(32);
+
+		// Discard all navigation updates caused by octree construction
+		Exchange(SavedDirtyAreas, DirtyAreas);
+	}
 
 	return bSupportRebuilding;
 }
@@ -595,6 +611,10 @@ void UNavigationSystem::OnWorldInitDone(FNavigationSystemRunMode Mode)
 	}
 	else
 	{
+		// Discard all bounds updates that was submitted during world initialization, 
+		// to avoid navigation rebuild right after map is loaded
+		PendingNavBoundsUpdates.Empty();
+		
 		// gather navigable bounds
 		GatherNavigationBounds();
 
@@ -2206,9 +2226,8 @@ void UNavigationSystem::InitializeForWorld(UWorld* World, FNavigationSystemRunMo
 			NavSys = CreateNavigationSystem(World);
 		}
 
-		// Remove old/stale chunk data from all levels, when navigation auto-update is enabled
-		// In case navigation system will be created chunks will be regenerated anyway
-		if (Mode == FNavigationSystemRunMode::EditorMode && bNavigationAutoUpdateEnabled == true)
+		// Remove old/stale chunk data from all sub-levels when navigation system is disabled
+		if (NavSys == nullptr && Mode == FNavigationSystemRunMode::EditorMode)
 		{
 			DiscardNavigationDataChunks(World);
 		}
