Fix preview world causing audio device unpause

Merging 2654979

//depot/UE4-UT/...

to //depot/UE4/...

diff --git a/Engine/Source/Runtime/Engine/Private/GameEngine.cpp b/Engine/Source/Runtime/Engine/Private/GameEngine.cpp
--- a/Engine/Source/Runtime/Engine/Private/GameEngine.cpp
+++ b/Engine/Source/Runtime/Engine/Private/GameEngine.cpp
@@ -1105,10 +1105,13 @@ void UGameEngine::Tick( float DeltaSeconds, bool bIdleMode )
 			Context.World()->UpdateLevelStreaming();
 		}
 
-		// Update Audio. This needs to occur after rendering as the rendering code updates the listener position.
-		if (FAudioDevice* AudioDevice = Context.World()->GetAudioDevice())
+		if (Context.WorldType != EWorldType::Preview)
 		{
-			AudioDevice->Update(!Context.World()->IsPaused());
+			// Update Audio. This needs to occur after rendering as the rendering code updates the listener position.
+			if (FAudioDevice* AudioDevice = Context.World()->GetAudioDevice())
+			{
+				AudioDevice->Update(!Context.World()->IsPaused());
+			}
 		}
 
 		if( GIsClient )
