#ue4 PS4Tracker updates:
  -PS4Tracker now implements IViewExtension and ownership of it is shared between game thread and render thread.
  -All device polling is now done all at once in PreRenderViewFamily_RenderThread. GetTrackingData now returns cached results most appropriate for the calling thread.
  -New support for positional tracking on Windows builds through the SCE hmd_client API.
  -Now handles device to Unreal coordinate conversion

diff --git a/Engine/Source/Runtime/Engine/Classes/Engine/Engine.h b/Engine/Source/Runtime/Engine/Classes/Engine/Engine.h
--- a/Engine/Source/Runtime/Engine/Classes/Engine/Engine.h
+++ b/Engine/Source/Runtime/Engine/Classes/Engine/Engine.h
@@ -1505,7 +1505,10 @@ public:
 
 	/** Reference to the HMD device that is attached, if any */
 	TSharedPtr< class IHeadMountedDisplay, ESPMode::ThreadSafe > HMDDevice;
-	
+
+	/** Extensions that can modify view parameters on the render thread. */
+	TArray<TSharedPtr<class ISceneViewExtension, ESPMode::ThreadSafe> > ViewExtensions;
+
 	/** Reference to the Motion Control devices that are attached, if any */
 	TArray < class IMotionController*> MotionControllerDevices;
 
diff --git a/Engine/Source/Runtime/Engine/Private/GameViewportClient.cpp b/Engine/Source/Runtime/Engine/Private/GameViewportClient.cpp
--- a/Engine/Source/Runtime/Engine/Private/GameViewportClient.cpp
+++ b/Engine/Source/Runtime/Engine/Private/GameViewportClient.cpp
@@ -827,6 +827,11 @@ void UGameViewportClient::Draw(FViewport* InViewport, FCanvas* SceneCanvas)
 		EngineShowFlags)
 		.SetRealtimeUpdate(true));
 
+	if (GEngine->ViewExtensions.Num())
+	{
+		ViewFamily.ViewExtensions.Append(GEngine->ViewExtensions.GetData(), GEngine->ViewExtensions.Num());
+	}
+
 	// Allow HMD to modify the view later, just before rendering
 	if (GEngine->HMDDevice.IsValid() && GEngine->IsStereoscopic3D(InViewport))
 	{
diff --git a/Engine/Source/Runtime/Engine/Private/UnrealEngine.cpp b/Engine/Source/Runtime/Engine/Private/UnrealEngine.cpp
--- a/Engine/Source/Runtime/Engine/Private/UnrealEngine.cpp
+++ b/Engine/Source/Runtime/Engine/Private/UnrealEngine.cpp
@@ -1030,10 +1030,16 @@ void UEngine::PreExit()
 	{
 		auto SavedStereo = StereoRenderingDevice;
 		auto SavedHMD = HMDDevice;
+		auto SavedViewExtentions = ViewExtensions;
 		{
 			FSuspendRenderingThread Suspend(false);
-	StereoRenderingDevice.Reset();
-	HMDDevice.Reset();
+			StereoRenderingDevice.Reset();
+			HMDDevice.Reset();
+			for (auto& ViewExt : ViewExtensions)
+			{
+				ViewExt.Reset();
+			}
+			ViewExtensions.Empty();
 		}
 		// shutdown will occur here.
 	}
