Missing files for HLOD build error

diff --git a/Engine/Source/Editor/UnrealEd/Public/HierarchicalLODVolume.h b/Engine/Source/Editor/UnrealEd/Public/HierarchicalLODVolume.h
--- a/Engine/Source/Editor/UnrealEd/Public/HierarchicalLODVolume.h
+++ b/Engine/Source/Editor/UnrealEd/Public/HierarchicalLODVolume.h
@@ -6,7 +6,7 @@
 #include "HierarchicalLODVolume.generated.h"
 
 /** An invisible volume used to manually define/create a HLOD cluster. */
-UCLASS()
+UCLASS(MinimalAPI)
 class AHierarchicalLODVolume : public AVolume
 {
 	GENERATED_UCLASS_BODY()
diff --git a/Engine/Source/Runtime/Engine/Classes/Engine/LODActor.h b/Engine/Source/Runtime/Engine/Classes/Engine/LODActor.h
--- a/Engine/Source/Runtime/Engine/Classes/Engine/LODActor.h
+++ b/Engine/Source/Runtime/Engine/Classes/Engine/LODActor.h
@@ -94,7 +94,6 @@ public:
 	void UpdateSubActorLODParents();
 #endif // WITH_EDITOR
 	
-protected:
 	// Begin UObject interface.
 	virtual FString GetDetailedInfoInternal() const override;
 	virtual FBox GetComponentsBoundingBox(bool bNonColliding = false) const override;	
