Fix for OR-1761 - More safety in exporting between emitters of varying LODs.

--------
Integrated using branch UE4-To-UE4-Orion (reversed) of change#2653247 by [REDACTED_PERSON] on 2015/08/12 14:54:51.

diff --git a/Engine/Source/Runtime/Engine/Private/Particles/ParticleComponents.cpp b/Engine/Source/Runtime/Engine/Private/Particles/ParticleComponents.cpp
--- a/Engine/Source/Runtime/Engine/Private/Particles/ParticleComponents.cpp
+++ b/Engine/Source/Runtime/Engine/Private/Particles/ParticleComponents.cpp
@@ -2119,7 +2119,8 @@ void UParticleSystem::SetupSoloing()
 				if (Emitter != NULL)
 				{
 					FLODSoloTrack& SoloTrack = SoloTracking[EmitterIdx];
-					for (int32 LODIdx = 0; LODIdx < Emitter->LODLevels.Num(); LODIdx++)
+					int32 MaxLOD = FMath::Min(SoloTrack.SoloEnableSetting.Num(), Emitter->LODLevels.Num());
+					for (int32 LODIdx = 0; LODIdx < MaxLOD; LODIdx++)
 					{
 						UParticleLODLevel* LODLevel = Emitter->LODLevels[LODIdx];
 						check(LODLevel);
