Updated StartFPSChart to now report an average for Render thread time and Game thread time in the HTML output report.

diff --git a/Engine/Source/Runtime/Engine/Private/ChartCreation.cpp b/Engine/Source/Runtime/Engine/Private/ChartCreation.cpp
--- a/Engine/Source/Runtime/Engine/Private/ChartCreation.cpp
+++ b/Engine/Source/Runtime/Engine/Private/ChartCreation.cpp
@@ -939,6 +939,16 @@ void UEngine::DumpFPSChartToHTML( float TotalTime, float DeltaTime, int32 NumFra
 		// Get settings info
 		const Scalability::FQualityLevels& Quality = GEngine->GetGameUserSettings()->ScalabilityQuality;
 
+		// Sum up FrameTimes and GameTimes
+		float TotalFrameTime = 0;
+		float TotalGameTime = 0;
+
+		for (int32 i = 0; i < GFrameTimes.Num(); i++)
+		{
+			TotalFrameTime += GRenderThreadFrameTimes[i];
+			TotalGameTime += GGameThreadFrameTimes[i];
+		}
+
 		// Update non- bucket stats.
 		FPSChartRow = FPSChartRow.Replace( TEXT("TOKEN_MAPNAME"),		    *FString::Printf(TEXT("%s"), *InMapName ), ESearchCase::CaseSensitive );
 		FPSChartRow = FPSChartRow.Replace( TEXT("TOKEN_CHANGELIST"),		*FString::Printf(TEXT("%i"), GetChangeListNumberForPerfTesting() ), ESearchCase::CaseSensitive );
@@ -962,6 +972,8 @@ void UEngine::DumpFPSChartToHTML( float TotalTime, float DeltaTime, int32 NumFra
 		FPSChartRow = FPSChartRow.Replace( TEXT("TOKEN_TIME"),				*FString::Printf(TEXT("%4.2f"), DeltaTime), ESearchCase::CaseSensitive );
 		FPSChartRow = FPSChartRow.Replace( TEXT("TOKEN_FRAMECOUNT"),		*FString::Printf(TEXT("%i"), NumFrames), ESearchCase::CaseSensitive );
 		FPSChartRow = FPSChartRow.Replace( TEXT("TOKEN_AVG_GPUTIME"),		*FString::Printf(TEXT("%4.2f ms"), float((GTotalGPUTime / NumFrames)*1000.0) ), ESearchCase::CaseSensitive );
+		FPSChartRow = FPSChartRow.Replace(TEXT("TOKEN_AVG_RENDTIME"), *FString::Printf(TEXT("%4.2f ms"), float((TotalFrameTime / NumFrames)*1000.0)), ESearchCase::CaseSensitive);
+		FPSChartRow = FPSChartRow.Replace(TEXT("TOKEN_AVG_GAMETIME"), *FString::Printf(TEXT("%4.2f ms"), float((TotalGameTime / NumFrames)*1000.0)), ESearchCase::CaseSensitive);
 
 		FPSChartRow = FPSChartRow.Replace( TEXT("TOKEN_BOUND_GAME_THREAD_PERCENT"),		*FString::Printf(TEXT("%4.2f"), (float(GNumFramesBound_GameThread)/float(NumFrames))*100.0f ), ESearchCase::CaseSensitive );
 		FPSChartRow = FPSChartRow.Replace( TEXT("TOKEN_BOUND_RENDER_THREAD_PERCENT"),		*FString::Printf(TEXT("%4.2f"), (float(GNumFramesBound_RenderThread)/float(NumFrames))*100.0f ), ESearchCase::CaseSensitive );
