Add command to turn off ragdoll physics

diff --git a/Engine/Source/Runtime/Engine/Private/PhysicsEngine/PhysAnim.cpp b/Engine/Source/Runtime/Engine/Private/PhysicsEngine/PhysAnim.cpp
--- a/Engine/Source/Runtime/Engine/Private/PhysicsEngine/PhysAnim.cpp
+++ b/Engine/Source/Runtime/Engine/Private/PhysicsEngine/PhysAnim.cpp
@@ -546,7 +546,7 @@ void USkeletalMeshComponent::UpdateKinematicBonesToAnim(const TArray<FTransform>
 				FBodyInstance* BodyInst = Bodies[i];
 				check(BodyInst);
 
-				if (bTeleport || (BodyInst->IsValidBodyInstance() && !BodyInst->IsInstanceSimulatingPhysics()))
+				if (BodyInst->IsValidBodyInstance() && (bTeleport || !BodyInst->IsInstanceSimulatingPhysics()))
 				{
 					const int32 BoneIndex = BodyInst->InstanceBoneIndex;
 
diff --git a/Engine/Source/Runtime/Engine/Private/SkeletalMeshComponentPhysics.cpp b/Engine/Source/Runtime/Engine/Private/SkeletalMeshComponentPhysics.cpp
--- a/Engine/Source/Runtime/Engine/Private/SkeletalMeshComponentPhysics.cpp
+++ b/Engine/Source/Runtime/Engine/Private/SkeletalMeshComponentPhysics.cpp
@@ -707,6 +707,8 @@ void USkeletalMeshComponent::SetPhysMaterialOverride(UPhysicalMaterial* NewPhysM
 
 DEFINE_STAT(STAT_InitArticulated);
 
+TAutoConsoleVariable<int32> CVarEnableRagdollPhysics(TEXT("p.RagdollPhysics"), 1, TEXT("If 1, ragdoll physics will be used. Otherwise just root body is simulated"));
+
 void USkeletalMeshComponent::InitArticulated(FPhysScene* PhysScene)
 {
 	SCOPE_CYCLE_COUNTER(STAT_InitArticulated);
@@ -805,6 +807,10 @@ void USkeletalMeshComponent::InitArticulated(FPhysScene* PhysScene)
 			else
 			{
 				BodyInst->DOFMode = EDOFMode::None;
+				if (!CVarEnableRagdollPhysics.GetValueOnGameThread())
+				{
+					continue;
+				}
 			}
 
 #if WITH_PHYSX
