- fixed issue with copying to itself

Merging

//depot/UE4-Orion/Engine/Source/Runtime/Engine/...

to //depot/UE4/Engine/Source/Runtime/Engine/...

diff --git a/Engine/Source/Runtime/Engine/Private/Animation/AnimCurveTypes.cpp b/Engine/Source/Runtime/Engine/Private/Animation/AnimCurveTypes.cpp
--- a/Engine/Source/Runtime/Engine/Private/Animation/AnimCurveTypes.cpp
+++ b/Engine/Source/Runtime/Engine/Private/Animation/AnimCurveTypes.cpp
@@ -460,13 +460,37 @@ void FBlendedCurve::Blend(const FBlendedCurve& A, const FBlendedCurve& B, float
 		for(int32 CurveId=0; CurveId<A.Elements.Num(); ++CurveId)
 		{
 			Elements[CurveId].Value = FMath::Lerp(A.Elements[CurveId].Value, B.Elements[CurveId].Value, Alpha); 
-			Elements[CurveId].Flags |= (A.Elements[CurveId].Flags | B.Elements[CurveId].Flags);
+			Elements[CurveId].Flags |= (B.Elements[CurveId].Flags);
 		}
 	}
 
 	ValidateCurve(this);
 }
 
+void FBlendedCurve::BlendWith(const FBlendedCurve& Other, float Alpha)
+{
+	check(Num()==Other.Num());
+	if(FMath::Abs(Alpha) <= ZERO_ANIMWEIGHT_THRESH)
+	{
+		return;
+	}
+	else if(FMath::Abs(Alpha - 1.0f) <= ZERO_ANIMWEIGHT_THRESH)
+	{
+		// if blend is all the way for child2, then just copy its bone atoms
+		Override(Other);
+	}
+	else
+	{
+		for(int32 CurveId=0; CurveId<Elements.Num(); ++CurveId)
+		{
+			Elements[CurveId].Value = FMath::Lerp(Elements[CurveId].Value, Other.Elements[CurveId].Value, Alpha);
+			Elements[CurveId].Flags |= (Other.Elements[CurveId].Flags);
+		}
+	}
+	 
+	ValidateCurve(this);
+}
+
 void FBlendedCurve::ConvertToAdditive(const FBlendedCurve& BaseCurve)
 {
 	check(bInitialized);
diff --git a/Engine/Source/Runtime/Engine/Private/Components/SkeletalMeshComponent.cpp b/Engine/Source/Runtime/Engine/Private/Components/SkeletalMeshComponent.cpp
--- a/Engine/Source/Runtime/Engine/Private/Components/SkeletalMeshComponent.cpp
+++ b/Engine/Source/Runtime/Engine/Private/Components/SkeletalMeshComponent.cpp
@@ -1265,7 +1265,7 @@ void USkeletalMeshComponent::PostAnimEvaluation(FAnimationEvaluationContext& Eva
 		FAnimationRuntime::LerpBoneTransforms(GetEditableSpaceBases(), CachedSpaceBases, Alpha, RequiredBones);
 
 		// interpolate curve
-		EvaluatedCurve.Blend(EvaluatedCurve, CachedCurve, Alpha);
+		EvaluatedCurve.BlendWith(CachedCurve, Alpha);
 	}
 
 	bNeedToFlipSpaceBaseBuffers = true;
diff --git a/Engine/Source/Runtime/Engine/Public/Animation/AnimCurveTypes.h b/Engine/Source/Runtime/Engine/Public/Animation/AnimCurveTypes.h
--- a/Engine/Source/Runtime/Engine/Public/Animation/AnimCurveTypes.h
+++ b/Engine/Source/Runtime/Engine/Public/Animation/AnimCurveTypes.h
@@ -269,7 +269,8 @@ struct ENGINE_API FBlendedCurve
 	void InitFrom(const class USkeleton* Skeleton);
 	void InitFrom(const FBlendedCurve& InCurveToInitFrom)
 	{
-		if (&InCurveToInitFrom != this)
+		// make sure this doesn't happen
+		if (ensure(&InCurveToInitFrom != this))
 		{
 			UIDList.Reset();
 			UIDList.Append(InCurveToInitFrom.UIDList);
@@ -286,6 +287,10 @@ struct ENGINE_API FBlendedCurve
 	 * Blend (A, B) using Alpha, same as Lerp
 	 */
 	void Blend(const FBlendedCurve& A, const FBlendedCurve& B, float Alpha);
+	/**
+	 * Blend with Other using Alpha, same as Lerp 
+	 */
+	void BlendWith(const FBlendedCurve& Other, float Alpha);
 	/**
 	 * Convert current curves to Additive (this - BaseCurve) if same found
 	 */
