adding stats for cloth count and number of cloth vertices simulated.

iterating over the clothing actors on all the skeletal mesh components.

diff --git a/Engine/Source/Runtime/Engine/Private/PhysicsEngine/PhysLevel.cpp b/Engine/Source/Runtime/Engine/Private/PhysicsEngine/PhysLevel.cpp
--- a/Engine/Source/Runtime/Engine/Private/PhysicsEngine/PhysLevel.cpp
+++ b/Engine/Source/Runtime/Engine/Private/PhysicsEngine/PhysLevel.cpp
@@ -35,6 +35,8 @@ DEFINE_STAT(STAT_NumActiveKinematicBodies);
 DEFINE_STAT(STAT_NumMobileBodies);
 DEFINE_STAT(STAT_NumStaticBodies);
 DEFINE_STAT(STAT_NumShapes);
+DEFINE_STAT(STAT_NumCloths);
+DEFINE_STAT(STAT_NumClothVerts);
 
 DEFINE_STAT(STAT_NumBroadphaseAddsAsync);
 DEFINE_STAT(STAT_NumBroadphaseRemovesAsync);
diff --git a/Engine/Source/Runtime/Engine/Private/PhysicsEngine/PhysScene.cpp b/Engine/Source/Runtime/Engine/Private/PhysicsEngine/PhysScene.cpp
--- a/Engine/Source/Runtime/Engine/Private/PhysicsEngine/PhysScene.cpp
+++ b/Engine/Source/Runtime/Engine/Private/PhysicsEngine/PhysScene.cpp
@@ -563,6 +563,29 @@ void FinishSceneStat(uint32 Scene)
 	}
 }
 
+#if WITH_APEX
+void GatherApexStats(const UWorld* World, NxApexScene* ApexScene)
+{
+	SCOPED_APEX_SCENE_READ_LOCK(ApexScene);
+	int32 NumVerts = 0;
+	int32 NumCloths = 0;
+	for (TObjectIterator<USkeletalMeshComponent> Itr; Itr; ++Itr)
+	{
+		if (Itr->GetWorld() != World) { continue; }
+		for (const FClothingActor& ClothingActor : Itr->ClothingActors)
+		{
+			if (ClothingActor.ApexClothingActor && ClothingActor.ApexClothingActor->getActivePhysicalLod() == 1)
+			{
+				NumCloths++;
+				NumVerts += ClothingActor.ApexClothingActor->getNumSimulationVertices();
+			}
+		}
+	}
+	SET_DWORD_STAT(STAT_NumCloths, NumCloths);        // number of recently simulated apex cloths.
+	SET_DWORD_STAT(STAT_NumClothVerts, NumVerts);	  // number of recently simulated apex vertices.
+}
+#endif
+
 
 /** Exposes ticking of physics-engine scene outside Engine. */
 void FPhysScene::TickPhysScene(uint32 SceneType, FGraphEventRef& InOutCompletionEvent)
@@ -678,6 +701,9 @@ void FPhysScene::TickPhysScene(uint32 SceneType, FGraphEventRef& InOutCompletion
 			ApexScene->simulate(AveragedFrameTime[SceneType], true, Task);
 			Task->removeReference();
 			bTaskOutstanding = true;
+
+			GatherApexStats(this->OwningWorld,ApexScene);
+
 #endif
 		}
 	}
diff --git a/Engine/Source/Runtime/Engine/Public/PhysicsPublic.h b/Engine/Source/Runtime/Engine/Public/PhysicsPublic.h
--- a/Engine/Source/Runtime/Engine/Public/PhysicsPublic.h
+++ b/Engine/Source/Runtime/Engine/Public/PhysicsPublic.h
@@ -45,6 +45,8 @@ DECLARE_DWORD_COUNTER_STAT_EXTERN(TEXT("Active Kinematic Bodies"), STAT_NumActiv
 DECLARE_DWORD_COUNTER_STAT_EXTERN(TEXT("Mobile Bodies"), STAT_NumMobileBodies, STATGROUP_Physics, );
 DECLARE_DWORD_COUNTER_STAT_EXTERN(TEXT("Static Bodies"), STAT_NumStaticBodies, STATGROUP_Physics, );
 DECLARE_DWORD_COUNTER_STAT_EXTERN(TEXT("Shapes"), STAT_NumShapes, STATGROUP_Physics, );
+DECLARE_DWORD_COUNTER_STAT_EXTERN(TEXT("Cloths"), STAT_NumCloths, STATGROUP_Physics, );
+DECLARE_DWORD_COUNTER_STAT_EXTERN(TEXT("ClothVerts"), STAT_NumClothVerts, STATGROUP_Physics, );
 
 DECLARE_DWORD_COUNTER_STAT_EXTERN(TEXT("(ASync) Broadphase Adds"), STAT_NumBroadphaseAddsAsync, STATGROUP_Physics, );
 DECLARE_DWORD_COUNTER_STAT_EXTERN(TEXT("(ASync) Broadphase Removes"), STAT_NumBroadphaseRemovesAsync, STATGROUP_Physics, );
