Expose flag for switching between physx dispatcher and UE4 task graph

diff --git a/Engine/Source/Runtime/Engine/Private/PhysicsEngine/PhysScene.cpp b/Engine/Source/Runtime/Engine/Private/PhysicsEngine/PhysScene.cpp
--- a/Engine/Source/Runtime/Engine/Private/PhysicsEngine/PhysScene.cpp
+++ b/Engine/Source/Runtime/Engine/Private/PhysicsEngine/PhysScene.cpp
@@ -98,6 +98,14 @@ FPhysScene::FPhysScene()
 #endif
 	PhysxUserData = FPhysxUserData(this);
 
+	
+	int64 NumPhysxDispatcher = 0;
+	FParse::Value(FCommandLine::Get(), TEXT("physxDispatcher="), NumPhysxDispatcher);
+	if(NumPhysxDispatcher == 0 && FParse::Param(FCommandLine::Get(), TEXT("physxDispatcher")))
+	{
+		NumPhysxDispatcher = 4;	//by default give physx 4 threads
+	}
+	
 	// Create dispatcher for tasks
 	if (PhysSingleThreadedMode())
 	{
@@ -105,7 +113,14 @@ FPhysScene::FPhysScene()
 	}
 	else
 	{
-		CPUDispatcher = new FPhysXCPUDispatcher();
+		if(NumPhysxDispatcher)
+		{
+			CPUDispatcher = PxDefaultCpuDispatcherCreate(NumPhysxDispatcher);
+		}else
+		{
+			CPUDispatcher = new FPhysXCPUDispatcher();
+		}
+		
 	}
 
 #endif	//#if WITH_PHYSX
