Change 2655594 by [REDACTED_PERSON][REDACTED_HANDLE] on 2015/08/13 20:46:09

Engine: Prevented a possible access of uninitialized memory in a ribbon emitter

diff --git a/Engine/Source/Runtime/Engine/Private/Particles/ParticleTrail2EmitterInstance.cpp b/Engine/Source/Runtime/Engine/Private/Particles/ParticleTrail2EmitterInstance.cpp
--- a/Engine/Source/Runtime/Engine/Private/Particles/ParticleTrail2EmitterInstance.cpp
+++ b/Engine/Source/Runtime/Engine/Private/Particles/ParticleTrail2EmitterInstance.cpp
@@ -2265,8 +2265,10 @@ bool FParticleRibbonEmitterInstance::ResolveSourcePoint(int32 InTrailIdx,
 					bool bEncounteredNaNError = false;
 
 					// Grab the particle
-					FBaseParticle* SourceParticle = (SourceIndices[InTrailIdx] != -1) ? SourceEmitter->GetParticleDirect(SourceIndices[InTrailIdx]) : NULL;
-					if (SourceParticle != NULL)
+					
+					const int32 SourceEmitterParticleIndex = SourceIndices[InTrailIdx];
+					FBaseParticle* SourceParticle = ((SourceEmitterParticleIndex >= 0) && (SourceEmitterParticleIndex < SourceEmitter->ActiveParticles)) ? SourceEmitter->GetParticleDirect(SourceEmitterParticleIndex) : nullptr;
+					if (SourceParticle != nullptr)
 					{
 						const FVector WorldOrigin = SourceEmitter->SimulationToWorld.GetOrigin();
 						UParticleSystemComponent* Comp = SourceEmitter->Component;
