Engine: Added a stat for the timer manager heap size
[INTEGRATE] Change 2656757 by [REDACTED_PERSON][REDACTED_HANDLE] on 2015/08/14 16:57:17

diff --git a/Engine/Source/Runtime/Engine/Private/TimerManager.cpp b/Engine/Source/Runtime/Engine/Private/TimerManager.cpp
--- a/Engine/Source/Runtime/Engine/Private/TimerManager.cpp
+++ b/Engine/Source/Runtime/Engine/Private/TimerManager.cpp
@@ -407,11 +407,15 @@ void FTimerManager::InternalUnPauseTimer(int32 PausedTimerIdx)
 // Public members
 // ---------------------------------
 
+DECLARE_DWORD_COUNTER_STAT(TEXT("TimerManager Heap Size"),STAT_NumHeapEntries,STATGROUP_Game);
+
 void FTimerManager::Tick(float DeltaTime)
 {
 	// @todo, might need to handle long-running case
 	// (e.g. every X seconds, renormalize to InternalTime = 0)
 
+	INC_DWORD_STAT_BY(STAT_NumHeapEntries, ActiveTimerHeap.Num());
+
 	InternalTime += DeltaTime;
 
 	while (ActiveTimerHeap.Num() > 0)
