Engine: Added a stat around ProcessLatentActions
[INTEGRATE] Change 2656756 by [REDACTED_PERSON][REDACTED_HANDLE] on 2015/08/14 16:57:14

diff --git a/Engine/Source/Runtime/Engine/Private/LatentActionManager.cpp b/Engine/Source/Runtime/Engine/Private/LatentActionManager.cpp
--- a/Engine/Source/Runtime/Engine/Private/LatentActionManager.cpp
+++ b/Engine/Source/Runtime/Engine/Private/LatentActionManager.cpp
@@ -44,8 +44,13 @@ void FLatentActionManager::RemoveActionsForObject(TWeakObjectPtr<UObject> InObje
 	}
 }
 
+
+DECLARE_CYCLE_STAT(TEXT("Blueprint Latent Actions"), STAT_TickLatentActions, STATGROUP_Game);
+
 void FLatentActionManager::ProcessLatentActions(UObject* InObject, float DeltaTime)
 {
+	SCOPE_CYCLE_COUNTER(STAT_TickLatentActions);
+
 	for (FActionsForObject::TIterator It(ActionsToRemoveMap); It; ++It)
 	{
 		auto ActionList = GetActionListForObject(It.Key());
