Blueprints: Reducing some allocations due to latent actions
[INTEGRATE] Change 2657117 by [REDACTED_PERSON][REDACTED_HANDLE] on 2015/08/14 22:30:52

diff --git a/Engine/Source/Runtime/Engine/Classes/Engine/LatentActionManager.h b/Engine/Source/Runtime/Engine/Classes/Engine/LatentActionManager.h
--- a/Engine/Source/Runtime/Engine/Classes/Engine/LatentActionManager.h
+++ b/Engine/Source/Runtime/Engine/Classes/Engine/LatentActionManager.h
@@ -123,8 +123,9 @@ public:
 	/** Resets the list of objects we have processed the latent action list for.	 */	
 	void BeginFrame()
 	{
-		ProcessedThisFrame.Empty();
+		ProcessedThisFrame.Reset();
 	}
+
 #if WITH_EDITOR
 	/** 
 	 * Builds a set of the UUIDs of pending latent actions on a specific object.
@@ -135,7 +136,6 @@ public:
 	 */
 	void GetActiveUUIDs(UObject* InObject, TSet<int32>& UUIDList) const;
 
-	
 	/** 
 	 * Gets the description string of a pending latent action with the specified UUID for a given object, or the empty string if it's an invalid UUID
 	 *
diff --git a/Engine/Source/Runtime/Engine/Private/LatentActionManager.cpp b/Engine/Source/Runtime/Engine/Private/LatentActionManager.cpp
--- a/Engine/Source/Runtime/Engine/Private/LatentActionManager.cpp
+++ b/Engine/Source/Runtime/Engine/Private/LatentActionManager.cpp
@@ -69,7 +69,7 @@ void FLatentActionManager::ProcessLatentActions(UObject* InObject, float DeltaTi
 			}
 		}
 	}
-	ActionsToRemoveMap.Empty();
+	ActionsToRemoveMap.Reset();
 
 	//@TODO: K2: Very inefficient code right now
 	if (InObject != NULL)
@@ -115,7 +115,7 @@ void FLatentActionManager::ProcessLatentActions(UObject* InObject, float DeltaTi
 					Action->NotifyObjectDestroyed();
 					delete Action;
 				}
-				ObjectActionList.Empty();
+				ObjectActionList.Reset();
 			}
 
 			// Remove the entry if there are no pending actions remaining for this object (or if the object was NULLed and cleaned up)
@@ -215,6 +215,7 @@ FString FLatentActionManager::GetDescription(UObject* InObject, int32 UUID) cons
 	}
 	return Description;
 }
+
 void FLatentActionManager::GetActiveUUIDs(UObject* InObject, TSet<int32>& UUIDList) const
 {
 	check(InObject);
