Engine: Changed the timer manager to be grow-only, avoiding reallocs due to expiring timers, etc...
[INTEGRATE] Change 2657116 by [REDACTED_PERSON][REDACTED_HANDLE] on 2015/08/14 22:28:12

diff --git a/Engine/Source/Runtime/Engine/Private/TimerManager.cpp b/Engine/Source/Runtime/Engine/Private/TimerManager.cpp
--- a/Engine/Source/Runtime/Engine/Private/TimerManager.cpp
+++ b/Engine/Source/Runtime/Engine/Private/TimerManager.cpp
@@ -208,15 +208,15 @@ void FTimerManager::InternalClearTimer(int32 TimerIdx, ETimerStatus TimerStatus)
 	switch (TimerStatus)
 	{
 		case ETimerStatus::Pending:
-			PendingTimerList.RemoveAtSwap(TimerIdx);
+			PendingTimerList.RemoveAtSwap(TimerIdx, /*bAllowShrinking=*/ false);
 			break;
 
 		case ETimerStatus::Active:
-			ActiveTimerHeap.HeapRemoveAt(TimerIdx);
+			ActiveTimerHeap.HeapRemoveAt(TimerIdx, /*bAllowShrinking=*/ false);
 			break;
 
 		case ETimerStatus::Paused:
-			PausedTimerList.RemoveAtSwap(TimerIdx);
+			PausedTimerList.RemoveAtSwap(TimerIdx, /*bAllowShrinking=*/ false);
 			break;
 
 		case ETimerStatus::Executing:
@@ -244,7 +244,7 @@ void FTimerManager::InternalClearAllTimers(void const* Object)
 			{
 				// remove this item
 				// this will break the heap property, but we will re-heapify afterward
-				ActiveTimerHeap.RemoveAtSwap(Idx--);
+				ActiveTimerHeap.RemoveAtSwap(Idx--, /*bAllowShrinking=*/ false);
 			}
 		}
 		if (OldActiveHeapSize != ActiveTimerHeap.Num())
@@ -258,7 +258,7 @@ void FTimerManager::InternalClearAllTimers(void const* Object)
 		{
 			if (PausedTimerList[Idx].TimerDelegate.IsBoundToObject(Object))
 			{
-				PausedTimerList.RemoveAtSwap(Idx--);
+				PausedTimerList.RemoveAtSwap(Idx--, /*bAllowShrinking=*/ false);
 			}
 		}
 
@@ -267,7 +267,7 @@ void FTimerManager::InternalClearAllTimers(void const* Object)
 		{
 			if (PendingTimerList[Idx].TimerDelegate.IsBoundToObject(Object))
 			{
-				PendingTimerList.RemoveAtSwap(Idx--);
+				PendingTimerList.RemoveAtSwap(Idx--, /*bAllowShrinking=*/ false);
 			}
 		}
 
@@ -357,11 +357,11 @@ void FTimerManager::InternalPauseTimer( FTimerData const* TimerToPause, int32 Ti
 		switch( PreviousStatus )
 		{
 			case ETimerStatus::Active:
-				ActiveTimerHeap.HeapRemoveAt(TimerIdx);
+				ActiveTimerHeap.HeapRemoveAt(TimerIdx, /*bAllowShrinking=*/ false);
 				break;
 
 			case ETimerStatus::Pending:
-				PendingTimerList.RemoveAtSwap(TimerIdx);
+				PendingTimerList.RemoveAtSwap(TimerIdx, /*bAllowShrinking=*/ false);
 				break;
 
 			case ETimerStatus::Executing:
@@ -399,7 +399,7 @@ void FTimerManager::InternalUnPauseTimer(int32 PausedTimerIdx)
 		}
 
 		// remove from paused list
-		PausedTimerList.RemoveAtSwap(PausedTimerIdx);
+		PausedTimerList.RemoveAtSwap(PausedTimerIdx, /*bAllowShrinking=*/ false);
 	}
 }
 
@@ -426,7 +426,7 @@ void FTimerManager::Tick(float DeltaTime)
 			// Timer has expired! Fire the delegate, then handle potential looping.
 
 			// Remove it from the heap and store it while we're executing
-			ActiveTimerHeap.HeapPop(CurrentlyExecutingTimer);
+			ActiveTimerHeap.HeapPop(CurrentlyExecutingTimer, /*bAllowShrinking=*/ false);
 			CurrentlyExecutingTimer.Status = ETimerStatus::Executing;
 
 			// Determine how many times the timer may have elapsed (e.g. for large DeltaTime on a short looping timer)
@@ -482,7 +482,7 @@ void FTimerManager::Tick(float DeltaTime)
 			TimerToActivate.Status = ETimerStatus::Active;
 			ActiveTimerHeap.HeapPush( TimerToActivate );
 		}
-		PendingTimerList.Empty();
+		PendingTimerList.Reset();
 	}
 }
 
