UE4 - fix assert on threadsafe add

diff --git a/Engine/Source/Runtime/Engine/Private/TickTaskManager.cpp b/Engine/Source/Runtime/Engine/Private/TickTaskManager.cpp
--- a/Engine/Source/Runtime/Engine/Private/TickTaskManager.cpp
+++ b/Engine/Source/Runtime/Engine/Private/TickTaskManager.cpp
@@ -90,7 +90,7 @@ public:
 	{
 		checkSlow(Count >= 0);
 		const int32 OldNum = FPlatformAtomics::InterlockedAdd(&this->ArrayNum, Count);
-		check(this->ArrayNum + Count < this->ArrayMax);
+		check(OldNum + Count <= this->ArrayMax);
 		return OldNum;
 	}
 
