Making sound attenuation boundaries match the attenuation curves at min/max.

Attenuation algorithms assume that volume is 1.0 when the distance is less than the min distance and 0 when the distance is greater than the max distance. This not necessarily the case with the new custom curve feature. Now the value of attenuation will be whatever the algorithm/curve is at the min/max boundary. This won't produce a huge change for existing curves but allow for the the fact that attenuation curves may not necessarily be the canonical 1.0 -> 0.0 curve.

diff --git a/Engine/Source/Runtime/Engine/Private/SoundAttenuation.cpp b/Engine/Source/Runtime/Engine/Private/SoundAttenuation.cpp
--- a/Engine/Source/Runtime/Engine/Private/SoundAttenuation.cpp
+++ b/Engine/Source/Runtime/Engine/Private/SoundAttenuation.cpp
@@ -79,50 +79,64 @@ float FAttenuationSettings::GetMaxDimension() const
  */
 float FAttenuationSettings::AttenuationEval(const float Distance, const float Falloff) const
 {
-	if( Distance > Falloff )
-	{
-		return 0.0f;
-	}
-	// UsedMinRadius is the point at which to start attenuating
-	else if( Distance > 0.f )
-	{
-		switch(DistanceAlgorithm)
-		{
-			case ATTENUATION_Linear:
+	// Clamp the input distance between 0.0f and Falloff. If the Distance
+	// is actually less than the min value, it will use the min-value of the algorithm/curve
+	// rather than assume it's 1.0 (i.e. it could be 0.0 for an inverse curve). Similarly, if the distance
+	// is greater than the falloff value, it'll use the algorithm/curve value evaluated at Falloff distance,
+	// which could be 1.0 (and not 0.0f).
 
-				return ( 1.0f - ( Distance / Falloff ) );
+	float DistanceCopy = FMath::Clamp(Distance, 0.0f, Falloff);
 
-			case ATTENUATION_Logarithmic:
+	float Result = 0.0f;
+	switch(DistanceAlgorithm)
+	{
+		case ATTENUATION_Linear:
 
-				return FMath::Min( 0.5f * -FMath::Loge( Distance / Falloff ), 1.0f );
+			Result = (1.0f - (DistanceCopy / Falloff));
+			break;
 
-			case ATTENUATION_Inverse:
+		case ATTENUATION_Logarithmic:
 
-				return FMath::Min( 0.02f / ( Distance / Falloff ) , 1.0f );
+			Result = 0.5f * -FMath::Loge(DistanceCopy / Falloff);
+			break;
 
-			case ATTENUATION_LogReverse:
+		case ATTENUATION_Inverse:
 
-				return FMath::Max(1.0f + 0.5f * FMath::Loge(1.0f - (Distance / Falloff)), 0.0f);
+			Result = 0.02f / (DistanceCopy / Falloff);
+			break;
 
-			case ATTENUATION_NaturalSound:
-			{
-				check( dBAttenuationAtMax <= 0.0f );
-				return FMath::Pow( 10.0f, ( ( Distance / Falloff ) * dBAttenuationAtMax ) / 20.0f );
-			}
+		case ATTENUATION_LogReverse:
 
-			case ATTENUATION_Custom:
+			Result = 1.0f + 0.5f * FMath::Loge(1.0f - (DistanceCopy / Falloff));
+			break;
 
-				return CustomAttenuationCurve.GetRichCurveConst()->Eval(Distance/Falloff);
+		case ATTENUATION_NaturalSound:
+		{
+			check( dBAttenuationAtMax <= 0.0f );
+			Result = FMath::Pow(10.0f, ((DistanceCopy / Falloff) * dBAttenuationAtMax) / 20.0f);
+			break;
 		}
+
+		case ATTENUATION_Custom:
+
+			Result = CustomAttenuationCurve.GetRichCurveConst()->Eval(DistanceCopy / Falloff);
+			break;
+
+		default:
+			checkf(false, TEXT("Uknown attenuation distance algorithm!"))
+			break;
 	}
-	return 1.0f;
+
+	// Make sure the output is clamped between 0.0 and 1.0f. Some of the algorithms above can
+	// result in bad values at the edges.
+	return FMath::Clamp(Result, 0.0f, 1.0f);
 }
 
 float FAttenuationSettings::AttenuationEvalBox(const FTransform& SoundTransform, const FVector ListenerLocation) const
 {
 	const float DistanceSq = ComputeSquaredDistanceFromBoxToPoint(-AttenuationShapeExtents, AttenuationShapeExtents,SoundTransform.InverseTransformPositionNoScale(ListenerLocation));
 	if (DistanceSq < FalloffDistance * FalloffDistance)
-	{
+	{ 
 		return AttenuationEval(FMath::Sqrt(DistanceSq), FalloffDistance);
 	}
 
