Adding volume-weight voice prioritization

- sound prioritization is based on the effective volume (computed after attenuation) scaled with an arbitrary user-defined priority (default to 1.0). This will provide a better metric as to what sounds are important and which ones are not. Voices which are inaudible will be the first to be "stolen"
- "bAlwaysPlay" doubles the set volume-weighted prioritization, we may want to revisit what the intention of this feature/parameter is
- fixing bShouldRemainActiveIfDropped feature
- increasing the time that a sound will remain active if it's attenuated to 0.0. This will fix looping sounds starting over from the beginning as you leave, then come back within range of an attenuation curve

diff --git a/Engine/Source/Runtime/Engine/Classes/Components/AudioComponent.h b/Engine/Source/Runtime/Engine/Classes/Components/AudioComponent.h
--- a/Engine/Source/Runtime/Engine/Classes/Components/AudioComponent.h
+++ b/Engine/Source/Runtime/Engine/Classes/Components/AudioComponent.h
@@ -121,6 +121,10 @@ class ENGINE_API UAudioComponent : public USceneComponent
 	UPROPERTY(EditAnywhere, BlueprintReadWrite, Category=Sound)
 	float VolumeMultiplier;
 
+	/** A priority value that scales against final ouput volume and is used to determine a sound's play priority. */
+	UPROPERTY(EditAnywhere, BlueprintReadWrite, Category = Sound)
+	float VolumeWeightedPriorityScale;
+
 	/** A pitch multiplier to apply to sounds generated by this component */
 	UPROPERTY(EditAnywhere, BlueprintReadWrite, Category=Sound)
 	float PitchMultiplier;
diff --git a/Engine/Source/Runtime/Engine/Private/ActiveSound.cpp b/Engine/Source/Runtime/Engine/Private/ActiveSound.cpp
--- a/Engine/Source/Runtime/Engine/Private/ActiveSound.cpp
+++ b/Engine/Source/Runtime/Engine/Private/ActiveSound.cpp
@@ -44,6 +44,7 @@ FActiveSound::FActiveSound()
 	, PitchMultiplier(1.f)
 	, HighFrequencyGainMultiplier(1.f)
 	, SubtitlePriority(0.f)
+	, VolumeWeightedPriorityScale(1.0f)
 	, OcclusionCheckInterval(0.f)
 	, LastOcclusionCheckTime(0.f)
 	, LastLocation(FVector::ZeroVector)
@@ -174,6 +175,7 @@ void FActiveSound::UpdateWaveInstances( FAudioDevice* AudioDevice, TArray<FWaveI
 	// final value that is correct
 	UpdateAdjustVolumeMultiplier(DeltaTime);
 	ParseParams.VolumeMultiplier = VolumeMultiplier * Sound->GetVolumeMultiplier() * CurrentAdjustVolumeMultiplier * AudioDevice->TransientMasterVolume * (bOccluded ? 0.5f : 1.0f);
+	ParseParams.VolumeWeightedPriorityScale = VolumeWeightedPriorityScale;
 	ParseParams.Pitch *= PitchMultiplier * Sound->GetPitchMultiplier();
 	ParseParams.HighFrequencyGain *= HighFrequencyGainMultiplier;
 
diff --git a/Engine/Source/Runtime/Engine/Private/Audio.cpp b/Engine/Source/Runtime/Engine/Private/Audio.cpp
--- a/Engine/Source/Runtime/Engine/Private/Audio.cpp
+++ b/Engine/Source/Runtime/Engine/Private/Audio.cpp
@@ -458,6 +458,11 @@ float FWaveInstance::GetActualVolume() const
 	return Volume * VolumeMultiplier;
 }
 
+float FWaveInstance::GetVolumeWeightedPriority() const
+{
+	return GetActualVolume() * VolumeWeightedPriorityScale;
+}
+
 bool FWaveInstance::IsStreaming() const
 {
 	return FPlatformProperties::SupportsAudioStreaming() && WaveData != nullptr && WaveData->IsStreaming();
diff --git a/Engine/Source/Runtime/Engine/Private/AudioDevice.cpp b/Engine/Source/Runtime/Engine/Private/AudioDevice.cpp
--- a/Engine/Source/Runtime/Engine/Private/AudioDevice.cpp
+++ b/Engine/Source/Runtime/Engine/Private/AudioDevice.cpp
@@ -16,6 +16,13 @@
 	FAudioDevice implementation.
 -----------------------------------------------------------------------------*/
 
+/**
+* Number of ticks an inaudible source remains alive before being stopped
+*/
+#define AUDIOSOURCE_TICK_LONGEVITY		600
+
+
+
 FAudioDevice::FAudioDevice()
 	: CommonAudioPool(nullptr)
 	, CommonAudioPoolFreeBytes(0)
@@ -1745,7 +1752,7 @@ int32 FAudioDevice::GetSortedActiveWaveInstances(TArray<FWaveInstance*>& WaveIns
 	{
 		FORCEINLINE bool operator()( const FWaveInstance& A, const FWaveInstance& B) const
 		{
-			return A.PlayPriority < B.PlayPriority;
+			return A.GetVolumeWeightedPriority() < B.GetVolumeWeightedPriority();
 		}
 	};
 
@@ -1756,7 +1763,7 @@ int32 FAudioDevice::GetSortedActiveWaveInstances(TArray<FWaveInstance*>& WaveIns
 	int32 FirstActiveIndex = FMath::Max( WaveInstances.Num() - MaxChannels, 0 );
 	for( ; FirstActiveIndex < WaveInstances.Num(); FirstActiveIndex++ )
 	{
-		if( WaveInstances[ FirstActiveIndex ]->PlayPriority > KINDA_SMALL_NUMBER )
+		if (WaveInstances[FirstActiveIndex]->GetVolumeWeightedPriority() > KINDA_SMALL_NUMBER)
 		{
 			break;
 		}
@@ -1777,7 +1784,8 @@ void FAudioDevice::StopSources( TArray<FWaveInstance*>& WaveInstances, int32 Fir
 			Source->LastUpdate = CurrentTick;
 
 			// If they are still audible, mark them as such
-			if( WaveInstance->PlayPriority > KINDA_SMALL_NUMBER )
+			float VolumeWeightedPriority = WaveInstance->GetVolumeWeightedPriority();
+			if (VolumeWeightedPriority > KINDA_SMALL_NUMBER)
 			{
 				Source->LastHeardUpdate = CurrentTick;
 			}
@@ -1798,16 +1806,22 @@ void FAudioDevice::StopSources( TArray<FWaveInstance*>& WaveInstances, int32 Fir
 				INC_DWORD_STAT( STAT_OggWaveInstances );
 			}
 #endif
-			// Source was not high enough priority to play this tick
-			if( Source->LastUpdate != CurrentTick )
+			// Source was not one of the active sounds this tick so needs to be stopped
+			if (Source->LastUpdate != CurrentTick)
 			{
 				Source->Stop();
 			}
 			// Source has been inaudible for several ticks
-			else if( Source->LastHeardUpdate + AUDIOSOURCE_TICK_LONGEVITY < CurrentTick )
+			else if (Source->LastHeardUpdate + AUDIOSOURCE_TICK_LONGEVITY < CurrentTick)
 			{
 				Source->Stop();
 			}
+			// Need to update the source still so that it gets any volume settings applied to
+			// otherwise the source may play at a very quiet volume and not actually set to 0.0
+			else
+			{
+				Source->Update();
+			}
 		}
 	}
 
diff --git a/Engine/Source/Runtime/Engine/Private/Components/AudioComponent.cpp b/Engine/Source/Runtime/Engine/Private/Components/AudioComponent.cpp
--- a/Engine/Source/Runtime/Engine/Private/Components/AudioComponent.cpp
+++ b/Engine/Source/Runtime/Engine/Private/Components/AudioComponent.cpp
@@ -21,6 +21,7 @@ UAudioComponent::UAudioComponent(const FObjectInitializer& ObjectInitializer)
 	bVisualizeComponent = true;
 #endif
 	VolumeMultiplier = 1.f;
+	VolumeWeightedPriorityScale = 1.f;
 	PitchMultiplier = 1.f;
 	VolumeModulationMin = 1.f;
 	VolumeModulationMax = 1.f;
@@ -126,7 +127,10 @@ void UAudioComponent::PlayInternal(const float StartTime, const float FadeInDura
 		// If this is an auto destroy component we need to prevent it from being auto-destroyed since we're really just restarting it
 		bool bCurrentAutoDestroy = bAutoDestroy;
 		bAutoDestroy = false;
-		Stop();
+		if (!bShouldRemainActiveIfDropped)
+		{
+			Stop();
+		}
 		bAutoDestroy = bCurrentAutoDestroy;
 	}
 
@@ -141,6 +145,7 @@ void UAudioComponent::PlayInternal(const float StartTime, const float FadeInDura
 			NewActiveSound.SoundClassOverride = SoundClassOverride;
 
 			NewActiveSound.VolumeMultiplier = (VolumeModulationMax + ((VolumeModulationMin - VolumeModulationMax) * FMath::SRand())) * VolumeMultiplier;
+			NewActiveSound.VolumeWeightedPriorityScale = VolumeWeightedPriorityScale;
 			NewActiveSound.PitchMultiplier = (PitchModulationMax + ((PitchModulationMin - PitchModulationMax) * FMath::SRand())) * PitchMultiplier;
 			NewActiveSound.HighFrequencyGainMultiplier = HighFrequencyGainMultiplier;
 
diff --git a/Engine/Source/Runtime/Engine/Private/SoundWave.cpp b/Engine/Source/Runtime/Engine/Private/SoundWave.cpp
--- a/Engine/Source/Runtime/Engine/Private/SoundWave.cpp
+++ b/Engine/Source/Runtime/Engine/Private/SoundWave.cpp
@@ -652,7 +652,15 @@ void USoundWave::Parse( FAudioDevice* AudioDevice, const UPTRINT NodeWaveInstanc
 			bAlwaysPlay = ActiveSound.bAlwaysPlay;
 		}
 
-		WaveInstance->PlayPriority = WaveInstance->Volume + ( bAlwaysPlay ? 1.0f : 0.0f ) + WaveInstance->RadioFilterVolume;
+		// This is a first-guess at priority, this will later change according to VolumeWeightedPriorityScale
+		WaveInstance->PlayPriority = WaveInstance->Volume + (bAlwaysPlay ? 1.0f : 0.0f) + WaveInstance->RadioFilterVolume;
+
+		// If set to bAlwaysPlay, double the current sound's priority scale. This will still result in a possible 0-priority output if the sound has 0 actual volume
+		WaveInstance->VolumeWeightedPriorityScale = ParseParams.VolumeWeightedPriorityScale;
+		if (bAlwaysPlay)
+		{
+			WaveInstance->VolumeWeightedPriorityScale *= 2.0f;
+		}
 		WaveInstance->Location = ParseParams.Transform.GetTranslation();
 		WaveInstance->bIsStarted = true;
 		WaveInstance->bAlreadyNotifiedHook = false;
diff --git a/Engine/Source/Runtime/Engine/Public/ActiveSound.h b/Engine/Source/Runtime/Engine/Public/ActiveSound.h
--- a/Engine/Source/Runtime/Engine/Public/ActiveSound.h
+++ b/Engine/Source/Runtime/Engine/Public/ActiveSound.h
@@ -26,6 +26,8 @@ struct FSoundParseParameters
 	float Volume;
 	float VolumeMultiplier;
 
+	float VolumeWeightedPriorityScale;
+
 	float Pitch;
 	float HighFrequencyGain;
 
@@ -147,6 +149,7 @@ public:
 	float HighFrequencyGainMultiplier;
 
 	float SubtitlePriority;
+	float VolumeWeightedPriorityScale;
 
 	/** Frequency with which to check for occlusion from its closest listener */
 	float OcclusionCheckInterval;
diff --git a/Engine/Source/Runtime/Engine/Public/Audio.h b/Engine/Source/Runtime/Engine/Public/Audio.h
--- a/Engine/Source/Runtime/Engine/Public/Audio.h
+++ b/Engine/Source/Runtime/Engine/Public/Audio.h
@@ -16,10 +16,6 @@ ENGINE_API DECLARE_LOG_CATEGORY_EXTERN(LogAudio, Warning, All);
  */
 #define MAX_AUDIOCHANNELS				64
 
-/** 
- * Number of ticks an inaudible source remains alive before being stopped
- */
-#define AUDIOSOURCE_TICK_LONGEVITY		60
 
 /** 
  * Length of sound in seconds to be considered as looping forever
@@ -190,8 +186,10 @@ struct ENGINE_API FWaveInstance
 	float				Volume;
 	/** Current volume multiplier - used to zero the volume without stopping the source */
 	float				VolumeMultiplier;
-	/** Current priority */
+	/** Current priority (deprecated) */
 	float				PlayPriority;
+	/** A priority value that scales with volume (post all gain stages) and is used to determin voice playback priority. */
+	float				VolumeWeightedPriorityScale;
 	/** Voice center channel volume */
 	float				VoiceCenterChannelVolume;
 	/** Volume of the radio filter effect */
@@ -283,6 +281,9 @@ struct ENGINE_API FWaveInstance
 	/** Returns the actual volume the wave instance will play at */
 	float GetActualVolume() const;
 
+	/** Returns the weighted priority of the wave instance. */
+	float GetVolumeWeightedPriority() const;
+
 	/**
 	 * Checks whether wave is streaming and streaming is supported
 	 */
