setting cloth lod  github #1424
The suggested change is sound.

at first glance the code looks like the previous code tests a value and doesn't set if it is already at this value. However, within apex, the update of the parameter doesn't happen right away.
Digging into the perforce history, it looks like the entire routine was added on one shot, so it was possibly just programmer habit to add this sort of check.
revision: //depot/UE4/Engine/Source/Runtime/Engine/Private/SkeletalMeshComponentPhysics.cpp#68
old changelist: 1743701
The solution is to remove a check for the currently active lod and then only update theforce lod if they are different. Trouble is that the propogation from one internal variable to the
The rational as mentioned in a provided code comment for the solution should use the term "deferred update" rather than "async". I doubt its a bug manifested as a result of threading. its more likely cloth simulation being both set ON and then OFF within a single frame.

diff --git a/Engine/Source/Runtime/Engine/Private/SkeletalMeshComponentPhysics.cpp b/Engine/Source/Runtime/Engine/Private/SkeletalMeshComponentPhysics.cpp
--- a/Engine/Source/Runtime/Engine/Private/SkeletalMeshComponentPhysics.cpp
+++ b/Engine/Source/Runtime/Engine/Private/SkeletalMeshComponentPhysics.cpp
@@ -2239,6 +2239,7 @@ void USkeletalMeshComponent::SetClothingLOD(int32 LODIndex)
 
 			int32 NumClothLODs = Actor.ParentClothingAsset->getNumGraphicalLodLevels();
 
+			// NOTE: forcePhysicalLod() is an deferred operation.   getActivePhysicalLod() may return previous state.
 			// decide whether should enable or disable
 			if (!IsMappedClothLOD || (LODIndex >= NumClothLODs))
 			{
@@ -2247,12 +2248,8 @@ void USkeletalMeshComponent::SetClothingLOD(int32 LODIndex)
 			}
 			else
 			{	
-				int32 CurPhysLOD = Actor.ApexClothingActor->getActivePhysicalLod();
-				// if disabled, enables clothing simulation
-				if(CurPhysLOD == 0)
-				{
-					Actor.ApexClothingActor->forcePhysicalLod(1);
-				}
+				// enable clothing simulation
+				Actor.ApexClothingActor->forcePhysicalLod(1);
 			}
 		}
 	}
