Velocity alignment for the sprite renderer

diff --git a/Engine/Source/Runtime/Engine/Classes/Engine/NiagaraSpriteRendererProperties.h b/Engine/Source/Runtime/Engine/Classes/Engine/NiagaraSpriteRendererProperties.h
--- a/Engine/Source/Runtime/Engine/Classes/Engine/NiagaraSpriteRendererProperties.h
+++ b/Engine/Source/Runtime/Engine/Classes/Engine/NiagaraSpriteRendererProperties.h
@@ -18,4 +18,7 @@ public:
 
 	UPROPERTY(EditAnywhere, Category = "Sprite Rendering")
 	FVector2D SubImageInfo;
+
+	UPROPERTY(EditAnywhere, Category = "Sprite Rendering")
+	bool bBVelocityAligned;
 };
diff --git a/Engine/Source/Runtime/Engine/Private/NiagaraEffectRenderer.cpp b/Engine/Source/Runtime/Engine/Private/NiagaraEffectRenderer.cpp
--- a/Engine/Source/Runtime/Engine/Private/NiagaraEffectRenderer.cpp
+++ b/Engine/Source/Runtime/Engine/Private/NiagaraEffectRenderer.cpp
@@ -161,13 +161,22 @@ void NiagaraEffectRendererSprites::GetDynamicMeshElements(const TArray<const FSc
 				FParticleSpriteUniformParameters PerViewUniformParameters;// = UniformParameters;
 				PerViewUniformParameters.AxisLockRight = FVector4(0.0f, 0.0f, 0.0f, 0.0f);
 				PerViewUniformParameters.AxisLockUp = FVector4(0.0f, 0.0f, 0.0f, 0.0f);
-				PerViewUniformParameters.RotationScale = 1.0f;
 				PerViewUniformParameters.RotationBias = 0.0f;
+				PerViewUniformParameters.RotationScale = 1.0f;
 				PerViewUniformParameters.TangentSelector = FVector4(0.0f, 0.0f, 0.0f, 1.0f);
+
+
 				PerViewUniformParameters.InvDeltaSeconds = 30.0f;
 				if (Properties)
 				{
 					PerViewUniformParameters.SubImageSize = FVector4(Properties->SubImageInfo.X, Properties->SubImageInfo.Y, 1.0f / Properties->SubImageInfo.X, 1.0f / Properties->SubImageInfo.Y);
+
+					if (Properties->bBVelocityAligned)
+					{
+						// velocity aligned
+						PerViewUniformParameters.RotationScale = 0.0f;
+						PerViewUniformParameters.TangentSelector = FVector4(0.0f, 1.0f, 0.0f, 0.0f);
+					}
 				}
 				PerViewUniformParameters.NormalsType = 0.0f;
 				PerViewUniformParameters.NormalsSphereCenter = FVector4(0.0f, 0.0f, 0.0f, 1.0f);
