#ue4 - Simulated root motion uses FScopedMovementUpdate for cheaper movement.

--------

diff --git a/Engine/Source/Runtime/Engine/Private/Components/CharacterMovementComponent.cpp b/Engine/Source/Runtime/Engine/Private/Components/CharacterMovementComponent.cpp
--- a/Engine/Source/Runtime/Engine/Private/Components/CharacterMovementComponent.cpp
+++ b/Engine/Source/Runtime/Engine/Private/Components/CharacterMovementComponent.cpp
@@ -1122,6 +1122,8 @@ void UCharacterMovementComponent::SimulateRootMotion(float DeltaSeconds, const F
 {
 	if( CharacterOwner && CharacterOwner->GetMesh() && (DeltaSeconds > 0.f) )
 	{
+		FScopedMovementUpdate ScopedMovementUpdate(UpdatedComponent, bEnableScopedMovementUpdates ? EScopedUpdate::DeferredUpdates : EScopedUpdate::ImmediateUpdates);
+
 		// Convert Local Space Root Motion to world space. Do it right before used by physics to make sure we use up to date transforms, as translation is relative to rotation.
 		const FTransform WorldSpaceRootMotionTransform = CharacterOwner->GetMesh()->ConvertLocalRootMotionToWorld(LocalRootMotionTransform);
 
