Addressed all mis-parametrized RemoveAtSwap occurances in FTimerManager #UE4
Fixed infinite-looping when remowing timers from ActiveTimerHeap #UE4
[INTEGRATE] Change 2658494 and 2658463

diff --git a/Engine/Source/Runtime/Engine/Private/TimerManager.cpp b/Engine/Source/Runtime/Engine/Private/TimerManager.cpp
--- a/Engine/Source/Runtime/Engine/Private/TimerManager.cpp
+++ b/Engine/Source/Runtime/Engine/Private/TimerManager.cpp
@@ -208,7 +208,7 @@ void FTimerManager::InternalClearTimer(int32 TimerIdx, ETimerStatus TimerStatus)
 	switch (TimerStatus)
 	{
 		case ETimerStatus::Pending:
-			PendingTimerList.RemoveAtSwap(TimerIdx, /*bAllowShrinking=*/ false);
+			PendingTimerList.RemoveAtSwap(TimerIdx, 1, /*bAllowShrinking=*/ false);
 			break;
 
 		case ETimerStatus::Active:
@@ -216,7 +216,7 @@ void FTimerManager::InternalClearTimer(int32 TimerIdx, ETimerStatus TimerStatus)
 			break;
 
 		case ETimerStatus::Paused:
-			PausedTimerList.RemoveAtSwap(TimerIdx, /*bAllowShrinking=*/ false);
+			PausedTimerList.RemoveAtSwap(TimerIdx, 1, /*bAllowShrinking=*/ false);
 			break;
 
 		case ETimerStatus::Executing:
@@ -244,7 +244,7 @@ void FTimerManager::InternalClearAllTimers(void const* Object)
 			{
 				// remove this item
 				// this will break the heap property, but we will re-heapify afterward
-				ActiveTimerHeap.RemoveAtSwap(Idx--, /*bAllowShrinking=*/ false);
+				ActiveTimerHeap.RemoveAtSwap(Idx--, 1, /*bAllowShrinking=*/ false);
 			}
 		}
 		if (OldActiveHeapSize != ActiveTimerHeap.Num())
@@ -258,7 +258,7 @@ void FTimerManager::InternalClearAllTimers(void const* Object)
 		{
 			if (PausedTimerList[Idx].TimerDelegate.IsBoundToObject(Object))
 			{
-				PausedTimerList.RemoveAtSwap(Idx--, /*bAllowShrinking=*/ false);
+				PausedTimerList.RemoveAtSwap(Idx--, 1, /*bAllowShrinking=*/ false);
 			}
 		}
 
@@ -267,7 +267,7 @@ void FTimerManager::InternalClearAllTimers(void const* Object)
 		{
 			if (PendingTimerList[Idx].TimerDelegate.IsBoundToObject(Object))
 			{
-				PendingTimerList.RemoveAtSwap(Idx--, /*bAllowShrinking=*/ false);
+				PendingTimerList.RemoveAtSwap(Idx--, 1, /*bAllowShrinking=*/ false);
 			}
 		}
 
@@ -361,7 +361,7 @@ void FTimerManager::InternalPauseTimer( FTimerData const* TimerToPause, int32 Ti
 				break;
 
 			case ETimerStatus::Pending:
-				PendingTimerList.RemoveAtSwap(TimerIdx, /*bAllowShrinking=*/ false);
+				PendingTimerList.RemoveAtSwap(TimerIdx, 1, /*bAllowShrinking=*/ false);
 				break;
 
 			case ETimerStatus::Executing:
@@ -399,7 +399,7 @@ void FTimerManager::InternalUnPauseTimer(int32 PausedTimerIdx)
 		}
 
 		// remove from paused list
-		PausedTimerList.RemoveAtSwap(PausedTimerIdx, /*bAllowShrinking=*/ false);
+		PausedTimerList.RemoveAtSwap(PausedTimerIdx, 1, /*bAllowShrinking=*/ false);
 	}
 }
 
