Fix non-unity compile

diff --git a/Engine/Source/Runtime/Engine/Private/NetworkDriver.cpp b/Engine/Source/Runtime/Engine/Private/NetworkDriver.cpp
--- a/Engine/Source/Runtime/Engine/Private/NetworkDriver.cpp
+++ b/Engine/Source/Runtime/Engine/Private/NetworkDriver.cpp
@@ -18,6 +18,7 @@
 #include "Engine/PackageMapClient.h"
 #include "GameFramework/PlayerState.h"
 #include "GameFramework/GameMode.h"
+#include "PerfCountersHelpers.h"
 
 #if USE_SERVER_PERF_COUNTERS
 #include "PerfCountersModule.h"
