Make sure FStringAssetReferences are fixed up for PIE client instances. Fixes UE-20108.

diff --git a/Engine/Source/Runtime/Engine/Private/World.cpp b/Engine/Source/Runtime/Engine/Private/World.cpp
--- a/Engine/Source/Runtime/Engine/Private/World.cpp
+++ b/Engine/Source/Runtime/Engine/Private/World.cpp
@@ -2295,6 +2295,7 @@ UWorld* UWorld::DuplicateWorldForPIE(const FString& PackageName, UWorld* OwningW
 
 	// Set up string asset reference fixups
 	TArray<FString> PackageNamesBeingDuplicatedForPIE;
+	PackageNamesBeingDuplicatedForPIE.Add(PrefixedLevelName);
 	if ( OwningWorld )
 	{
 		const FString PlayWorldMapName = ConvertToPIEPackageName(OwningWorld->GetOutermost()->GetName(), WorldContext.PIEInstance);
@@ -2306,7 +2307,7 @@ UWorld* UWorld::DuplicateWorldForPIE(const FString& PackageName, UWorld* OwningW
 			if (StreamingLevel)
 			{
 				const FString StreamingLevelPIEName = UWorld::ConvertToPIEPackageName(StreamingLevel->GetWorldAssetPackageName(), WorldContext.PIEInstance);
-				PackageNamesBeingDuplicatedForPIE.Add(StreamingLevelPIEName);
+				PackageNamesBeingDuplicatedForPIE.AddUnique(StreamingLevelPIEName);
 			}
 		}
 	}
