Curve Editor: Fix shift and scale for all key handles.

diff --git a/Engine/Source/Runtime/Engine/Private/CurveBase.cpp b/Engine/Source/Runtime/Engine/Private/CurveBase.cpp
--- a/Engine/Source/Runtime/Engine/Private/CurveBase.cpp
+++ b/Engine/Source/Runtime/Engine/Private/CurveBase.cpp
@@ -489,6 +489,12 @@ float FRichCurve::GetKeyValue(FKeyHandle KeyHandle) const
 void FRichCurve::ShiftCurve(float DeltaTime)
 {
 	TSet<FKeyHandle> KeyHandles;
+	for (auto It = KeyHandlesToIndices.CreateIterator(); It; ++It)
+	{
+		FKeyHandle& KeyHandle = It.Key();
+		KeyHandles.Add(KeyHandle);
+	}
+
 	ShiftCurve(DeltaTime, KeyHandles);
 }
 
@@ -507,6 +513,12 @@ void FRichCurve::ShiftCurve(float DeltaTime, TSet<FKeyHandle>& KeyHandles)
 void FRichCurve::ScaleCurve(float ScaleOrigin, float ScaleFactor)
 {
 	TSet<FKeyHandle> KeyHandles;
+	for (auto It = KeyHandlesToIndices.CreateIterator(); It; ++It)
+	{
+		FKeyHandle& KeyHandle = It.Key();
+		KeyHandles.Add(KeyHandle);
+	}
+
 	ScaleCurve(ScaleOrigin, ScaleFactor, KeyHandles);
 }
 
@@ -1391,6 +1403,12 @@ float FIntegralCurve::GetKeyTime(FKeyHandle KeyHandle) const
 void FIntegralCurve::ShiftCurve(float DeltaTime)
 {
 	TSet<FKeyHandle> KeyHandles;
+	for (auto It = KeyHandlesToIndices.CreateIterator(); It; ++It)
+	{
+		FKeyHandle& KeyHandle = It.Key();
+		KeyHandles.Add(KeyHandle);
+	}
+
 	ShiftCurve(DeltaTime, KeyHandles);
 }
 
@@ -1409,6 +1427,12 @@ void FIntegralCurve::ShiftCurve(float DeltaTime, TSet<FKeyHandle>& KeyHandles)
 void FIntegralCurve::ScaleCurve(float ScaleOrigin, float ScaleFactor)
 {
 	TSet<FKeyHandle> KeyHandles;
+	for (auto It = KeyHandlesToIndices.CreateIterator(); It; ++It)
+	{
+		FKeyHandle& KeyHandle = It.Key();
+		KeyHandles.Add(KeyHandle);
+	}
+
 	ScaleCurve(ScaleOrigin, ScaleFactor, KeyHandles);
 }
 
