Fix crash on shutdown if preview world is open

Merging 2660103

//depot/UE4-UT/...

to //depot/UE4/...

diff --git a/Engine/Source/Runtime/Engine/Private/GameEngine.cpp b/Engine/Source/Runtime/Engine/Private/GameEngine.cpp
--- a/Engine/Source/Runtime/Engine/Private/GameEngine.cpp
+++ b/Engine/Source/Runtime/Engine/Private/GameEngine.cpp
@@ -541,7 +541,10 @@ void UGameEngine::PreExit()
 				ActorIt->RouteEndPlay(EEndPlayReason::Quit);
 			}
 
-			World->GetGameInstance()->Shutdown();
+			if (World->GetGameInstance() != nullptr)
+			{
+				World->GetGameInstance()->Shutdown();
+			}
 
 			World->FlushLevelStreaming(EFlushLevelStreamingType::Visibility);
 			World->CleanupWorld();
