Optimizations to BSP building when duplicating Brushes, and further optimizations to updating model components.  Added a slow task dialog which is shown if the BSP generation takes more than three seconds.

diff --git a/Engine/Source/Editor/UnrealEd/Private/EditorServer.cpp b/Engine/Source/Editor/UnrealEd/Private/EditorServer.cpp
--- a/Engine/Source/Editor/UnrealEd/Private/EditorServer.cpp
+++ b/Engine/Source/Editor/UnrealEd/Private/EditorServer.cpp
@@ -1544,6 +1544,11 @@ void UEditorEngine::RebuildLevel(ULevel& Level)
 		return;
 	}
 
+	FScopedSlowTask SlowTask(2);
+	SlowTask.MakeDialogDelayed(3.0f);
+
+	SlowTask.EnterProgressFrame(1);
+
 	// Note: most of the following code was taken from UEditorEngine::csgRebuild()
 	FinishAllSnaps();
 	FBSPOps::GFastRebuild = 1;
@@ -1563,6 +1568,7 @@ void UEditorEngine::RebuildLevel(ULevel& Level)
 
 	FBSPOps::GFastRebuild = 1;
 
+	SlowTask.EnterProgressFrame(1);
 	Level.UpdateModelComponents();
 
 	RebuildStaticNavigableGeometry(&Level);
@@ -1600,7 +1606,7 @@ void UEditorEngine::RebuildModelFromBrushes(UModel* Model, bool bSelectedBrushes
 	Model->Verts.Empty(NumVerts + NumVerts / 8);
 	Model->Vectors.Empty(NumVectors + NumVectors / 8);
 	Model->Surfs.Empty(NumSurfs + NumSurfs / 8);
-		
+
 	// Limit the brushes used to the level the model is for
 	ULevel* Level = Model->GetTypedOuter<ULevel>();
 	if ( !Level )
@@ -1612,38 +1618,39 @@ void UEditorEngine::RebuildModelFromBrushes(UModel* Model, bool bSelectedBrushes
 	}
 	check( Level );
 
-	// Compose all structural brushes and portals.
-	for( auto It(Level->Actors.CreateConstIterator()); It; ++It )
-	{	
+	// Build list of all static brushes, first structural brushes and portals
+	TArray<ABrush*> StaticBrushes;
+	for (auto It(Level->Actors.CreateConstIterator()); It; ++It)
+	{
 		ABrush* Brush = Cast<ABrush>(*It);
 		if ((Brush && (Brush->IsStaticBrush() || bTreatMovableBrushesAsStatic) && !FActorEditorUtils::IsABuilderBrush(Brush)) &&
 			(!bSelectedBrushesOnly || Brush->IsSelected()) &&
 			(!(Brush->PolyFlags & PF_Semisolid) || (Brush->BrushType != Brush_Add) || (Brush->PolyFlags & PF_Portal)))
 		{
+			StaticBrushes.Add(Brush);
+
 			// Treat portals as solids for cutting.
 			if (Brush->PolyFlags & PF_Portal)
 			{
 				Brush->PolyFlags = (Brush->PolyFlags & ~PF_Semisolid) | PF_NotSolid;
 			}
-			Brush->Modify();
-			bspBrushCSG(Brush, Model, Brush->PolyFlags, (EBrushType)Brush->BrushType, CSG_None, false, true, false, false);
 		}
 	}
 
-	// Compose all detail brushes.
-	for( auto It(Level->Actors.CreateConstIterator()); It; ++It )
+	// Next append all detail brushes
+	for (auto It(Level->Actors.CreateConstIterator()); It; ++It)
 	{
 		ABrush* Brush = Cast<ABrush>(*It);
 		if (Brush && Brush->IsStaticBrush() && !FActorEditorUtils::IsABuilderBrush(Brush) &&
 			(!bSelectedBrushesOnly || Brush->IsSelected()) &&
 			(Brush->PolyFlags & PF_Semisolid) && !(Brush->PolyFlags & PF_Portal) && (Brush->BrushType == Brush_Add))
 		{
-			Brush->Modify();
-			bspBrushCSG(Brush, Model, Brush->PolyFlags, (EBrushType)Brush->BrushType, CSG_None, false, true, false, false);
+			StaticBrushes.Add(Brush);
 		}
 	}
 
-	// Rebuild dynamic brush BSP's (if they weren't handled earlier)
+	// Build list of dynamic brushes
+	TArray<ABrush*> DynamicBrushes;
 	if (!bTreatMovableBrushesAsStatic)
 	{
 		for (auto It(Level->Actors.CreateConstIterator()); It; ++It)
@@ -1652,16 +1659,34 @@ void UEditorEngine::RebuildModelFromBrushes(UModel* Model, bool bSelectedBrushes
 			if (DynamicBrush && DynamicBrush->Brush && !DynamicBrush->IsStaticBrush() &&
 				(!bSelectedBrushesOnly || DynamicBrush->IsSelected()))
 			{
-				BspPoints = MakeUnique<FBspPointsGrid>(50.0f, THRESH_POINTS_ARE_SAME);
-				BspVectors = MakeUnique<FBspPointsGrid>(1 / 16.0f, FMath::Max(THRESH_NORMALS_ARE_SAME, THRESH_VECTORS_ARE_NEAR));
-				FBspPointsGrid::GBspPoints = BspPoints.Get();
-				FBspPointsGrid::GBspVectors = BspVectors.Get();
-
-				FBSPOps::csgPrepMovingBrush(DynamicBrush);
+				DynamicBrushes.Add(DynamicBrush);
 			}
 		}
 	}
 
+	FScopedSlowTask SlowTask(StaticBrushes.Num() + DynamicBrushes.Num());
+	SlowTask.MakeDialogDelayed(3.0f);
+
+	// Compose all static brushes
+	for (ABrush* Brush : StaticBrushes)
+	{
+		SlowTask.EnterProgressFrame(1);
+		Brush->Modify();
+		bspBrushCSG(Brush, Model, Brush->PolyFlags, (EBrushType)Brush->BrushType, CSG_None, false, true, false, false);
+	}
+
+	// Rebuild dynamic brush BSP's (if they weren't handled earlier)
+	for (ABrush* DynamicBrush : DynamicBrushes)
+	{
+		SlowTask.EnterProgressFrame(1);
+		BspPoints = MakeUnique<FBspPointsGrid>(50.0f, THRESH_POINTS_ARE_SAME);
+		BspVectors = MakeUnique<FBspPointsGrid>(1 / 16.0f, FMath::Max(THRESH_NORMALS_ARE_SAME, THRESH_VECTORS_ARE_NEAR));
+		FBspPointsGrid::GBspPoints = BspPoints.Get();
+		FBspPointsGrid::GBspVectors = BspVectors.Get();
+
+		FBSPOps::csgPrepMovingBrush(DynamicBrush);
+	}
+
 	FBspPointsGrid::GBspPoints = nullptr;
 	FBspPointsGrid::GBspVectors = nullptr;
 }
@@ -1724,12 +1749,19 @@ void UEditorEngine::RebuildAlteredBSP()
 		}
 
 		// Rebuild the levels
-		for (int32 LevelIdx = 0; LevelIdx < LevelsToRebuild.Num(); ++LevelIdx)
 		{
-			TWeakObjectPtr< ULevel > LevelToRebuild = LevelsToRebuild[LevelIdx];
-				if (LevelToRebuild.IsValid())
+			FScopedSlowTask SlowTask(LevelsToRebuild.Num(), NSLOCTEXT("EditorServer", "RebuildingBSP", "Rebuilding BSP..."));
+			SlowTask.MakeDialogDelayed(3.0f);
+
+			for (int32 LevelIdx = 0; LevelIdx < LevelsToRebuild.Num(); ++LevelIdx)
 			{
-				RebuildLevel(*LevelToRebuild.Get());
+				SlowTask.EnterProgressFrame(1.0f);
+
+				TWeakObjectPtr< ULevel > LevelToRebuild = LevelsToRebuild[LevelIdx];
+				if (LevelToRebuild.IsValid())
+				{
+					RebuildLevel(*LevelToRebuild.Get());
+				}
 			}
 		}
 
@@ -3318,7 +3350,9 @@ void UEditorEngine::PasteSelectedActorsFromClipboard( UWorld* InWorld, const FTe
 		const FScopedTransaction Transaction( TransDescription );
 
 		GEditor->SelectNone( true, false );
+		ABrush::SetSuppressBSPRegeneration(true);
 		edactPasteSelected( InWorld, false, false, true );
+		ABrush::SetSuppressBSPRegeneration(false);
 
 		if( PasteTo != PT_OriginalLocation )
 		{
diff --git a/Engine/Source/Editor/UnrealEd/Private/LevelEditorViewport.cpp b/Engine/Source/Editor/UnrealEd/Private/LevelEditorViewport.cpp
--- a/Engine/Source/Editor/UnrealEd/Private/LevelEditorViewport.cpp
+++ b/Engine/Source/Editor/UnrealEd/Private/LevelEditorViewport.cpp
@@ -2298,7 +2298,9 @@ bool FLevelEditorViewportClient::InputWidgetDelta(FViewport* Viewport, EAxisList
 						{
 							// Widget hasn't been dragged since ALT+LMB went down.
 							bDuplicateOnNextDrag = false;
+							ABrush::SetSuppressBSPRegeneration(true);
 							GEditor->edactDuplicateSelected(GetWorld()->GetCurrentLevel(), false);
+							ABrush::SetSuppressBSPRegeneration(false);
 						}
 					}
 				}
diff --git a/Engine/Source/Editor/UnrealEd/Private/UnrealEdSrv.cpp b/Engine/Source/Editor/UnrealEd/Private/UnrealEdSrv.cpp
--- a/Engine/Source/Editor/UnrealEd/Private/UnrealEdSrv.cpp
+++ b/Engine/Source/Editor/UnrealEd/Private/UnrealEdSrv.cpp
@@ -2772,7 +2772,9 @@ bool UUnrealEdEngine::Exec_Actor( UWorld* InWorld, const TCHAR* Str, FOutputDevi
 			const FScopedTransaction Transaction( NSLOCTEXT("UnrealEd", "DuplicateActors", "Duplicate Actors") );
 
 			// duplicate selected
+			ABrush::SetSuppressBSPRegeneration(true);
 			edactDuplicateSelected(InWorld->GetCurrentLevel(), GetDefault<ULevelEditorViewportSettings>()->GridEnabled);
+			ABrush::SetSuppressBSPRegeneration(false);
 
 			// Find out if any of the selected actors will change the BSP.
 			// and only then rebuild BSP as this is expensive.
diff --git a/Engine/Source/Runtime/Engine/Classes/Engine/Brush.h b/Engine/Source/Runtime/Engine/Classes/Engine/Brush.h
--- a/Engine/Source/Runtime/Engine/Classes/Engine/Brush.h
+++ b/Engine/Source/Runtime/Engine/Classes/Engine/Brush.h
@@ -173,6 +173,8 @@ public:
 	virtual void InitPosRotScale();
 	virtual void CopyPosRotScaleFrom( ABrush* Other );
 
+	static void SetSuppressBSPRegeneration(bool bSuppress) { bSuppressBSPRegeneration = bSuppress; }
+
 private:
 
 	/** An array to keep track of all the levels that need rebuilding. This is checked via NeedsRebuild() in the editor tick and triggers a csg rebuild. */
@@ -181,6 +183,9 @@ private:
 	/** Delegate called when PostRegisterAllComponents is called for a Brush */
 	static FOnBrushRegistered OnBrushRegistered;
 
+	/** Global bool to suppress automatic BSP regeneration */
+	static bool bSuppressBSPRegeneration;
+
 public:
 
 	/**
diff --git a/Engine/Source/Runtime/Engine/Private/Brush.cpp b/Engine/Source/Runtime/Engine/Private/Brush.cpp
--- a/Engine/Source/Runtime/Engine/Private/Brush.cpp
+++ b/Engine/Source/Runtime/Engine/Private/Brush.cpp
@@ -20,6 +20,9 @@ ABrush::FOnBrushRegistered ABrush::OnBrushRegistered;
 
 /** An array to keep track of all the levels that need rebuilding. This is checked via NeedsRebuild() in the editor tick and triggers a csg rebuild. */
 TArray< TWeakObjectPtr< ULevel > > ABrush::LevelsToRebuild;
+
+/** Whether BSP regeneration should be suppressed or not */
+bool ABrush::bSuppressBSPRegeneration = false;
 #endif
 
 ABrush::ABrush(const FObjectInitializer& ObjectInitializer)
@@ -59,11 +62,10 @@ void ABrush::PostEditChangeProperty(FPropertyChangedEvent& PropertyChangedEvent)
 		Brush->BuildBound();
 	}
 
-	if(IsStaticBrush() && PropertyChangedEvent.ChangeType != EPropertyChangeType::Interactive && GUndo)
+	if(!bSuppressBSPRegeneration && IsStaticBrush() && PropertyChangedEvent.ChangeType != EPropertyChangeType::Interactive && GUndo)
 	{
 		// BSP can only be rebuilt during a transaction
 		GEditor->RebuildAlteredBSP();
-
 	}
 
 	bool bIsBuilderBrush = FActorEditorUtils::IsABuilderBrush( this );
diff --git a/Engine/Source/Runtime/Engine/Private/Level.cpp b/Engine/Source/Runtime/Engine/Private/Level.cpp
--- a/Engine/Source/Runtime/Engine/Private/Level.cpp
+++ b/Engine/Source/Runtime/Engine/Private/Level.cpp
@@ -842,6 +842,11 @@ void ULevel::MarkLevelComponentsRenderStateDirty()
 
 void ULevel::CreateModelComponents()
 {
+	FScopedSlowTask SlowTask(10);
+	SlowTask.MakeDialogDelayed(3.0f);
+
+	SlowTask.EnterProgressFrame(4);
+
 	// Update the model vertices and edges.
 	Model->UpdateVertices();
 
@@ -850,66 +855,94 @@ void ULevel::CreateModelComponents()
 	// Clear the model index buffers.
 	Model->MaterialIndexBuffers.Empty();
 
-	TMap< FModelComponentKey, TArray<uint16> > ModelComponentMap;
-
-	// Sort the nodes by zone, grid cell and masked poly flags.
-	for(int32 NodeIndex = 0;NodeIndex < Model->Nodes.Num();NodeIndex++)
+	struct FNodeIndices
 	{
-		FBspNode& Node = Model->Nodes[NodeIndex];
-		FBspSurf& Surf = Model->Surfs[Node.iSurf];
+		TArray<uint16> Nodes;
+		TSet<uint16> UniqueNodes;
+
+		FNodeIndices()
+		{
+			Nodes.Reserve(16);
+			UniqueNodes.Reserve(16);
+		}
 
-		if(Node.NumVertices > 0)
+		void AddUnique(uint16 Index)
 		{
-			// Calculate the bounding box of this node.
-			FBox NodeBounds(0);
-			for(int32 VertexIndex = 0;VertexIndex < Node.NumVertices;VertexIndex++)
+			if (!UniqueNodes.Contains(Index))
 			{
-				NodeBounds += Model->Points[Model->Verts[Node.iVertPool + VertexIndex].pVertex];
+				Nodes.Add(Index);
+				UniqueNodes.Add(Index);
 			}
+		}
+	};
+
+	TMap< FModelComponentKey, FNodeIndices > ModelComponentMap;
+
+	{
+		FScopedSlowTask InnerTask(Model->Nodes.Num());
+		InnerTask.MakeDialogDelayed(3.0f);
+
+		// Sort the nodes by zone, grid cell and masked poly flags.
+		for (int32 NodeIndex = 0; NodeIndex < Model->Nodes.Num(); NodeIndex++)
+		{
+			InnerTask.EnterProgressFrame(1);
+
+			FBspNode& Node = Model->Nodes[NodeIndex];
+			FBspSurf& Surf = Model->Surfs[Node.iSurf];
+
+			if (Node.NumVertices > 0)
+			{
+				// Calculate the bounding box of this node.
+				FBox NodeBounds(0);
+				for (int32 VertexIndex = 0; VertexIndex < Node.NumVertices; VertexIndex++)
+				{
+					NodeBounds += Model->Points[Model->Verts[Node.iVertPool + VertexIndex].pVertex];
+				}
 
-			// Create a sort key for this node using the grid cell containing the center of the node's bounding box.
+				// Create a sort key for this node using the grid cell containing the center of the node's bounding box.
 #define MODEL_GRID_SIZE_XY	2048.0f
 #define MODEL_GRID_SIZE_Z	4096.0f
-			FModelComponentKey Key;
-			check( OwningWorld );
-			if (OwningWorld->GetWorldSettings()->bMinimizeBSPSections)
-			{
-				Key.X				= 0;
-				Key.Y				= 0;
-				Key.Z				= 0;
-			}
-			else
-			{
-				Key.X				= FMath::FloorToInt(NodeBounds.GetCenter().X / MODEL_GRID_SIZE_XY);
-				Key.Y				= FMath::FloorToInt(NodeBounds.GetCenter().Y / MODEL_GRID_SIZE_XY);
-				Key.Z				= FMath::FloorToInt(NodeBounds.GetCenter().Z / MODEL_GRID_SIZE_Z);
-			}
+				FModelComponentKey Key;
+				check(OwningWorld);
+				if (OwningWorld->GetWorldSettings()->bMinimizeBSPSections)
+				{
+					Key.X = 0;
+					Key.Y = 0;
+					Key.Z = 0;
+				}
+				else
+				{
+					Key.X = FMath::FloorToInt(NodeBounds.GetCenter().X / MODEL_GRID_SIZE_XY);
+					Key.Y = FMath::FloorToInt(NodeBounds.GetCenter().Y / MODEL_GRID_SIZE_XY);
+					Key.Z = FMath::FloorToInt(NodeBounds.GetCenter().Z / MODEL_GRID_SIZE_Z);
+				}
 
-			Key.MaskedPolyFlags = 0;
+				Key.MaskedPolyFlags = 0;
 
-			// Find an existing node list for the grid cell.
-			TArray<uint16>* ComponentNodes = ModelComponentMap.Find(Key);
-			if(!ComponentNodes)
+				// Find an existing node list for the grid cell.
+				FNodeIndices* ComponentNodes = ModelComponentMap.Find(Key);
+				if (!ComponentNodes)
+				{
+					// This is the first node we found in this grid cell, create a new node list for the grid cell.
+					ComponentNodes = &ModelComponentMap.Add(Key);
+				}
+
+				// Add the node to the grid cell's node list.
+				ComponentNodes->AddUnique(NodeIndex);
+			}
+			else
 			{
-				// This is the first node we found in this grid cell, create a new node list for the grid cell.
-				ComponentNodes = &ModelComponentMap.Add(Key,TArray<uint16>());
+				// Put it in component 0 until a rebuild occurs.
+				Node.ComponentIndex = 0;
 			}
-
-			// Add the node to the grid cell's node list.
-			ComponentNodes->AddUnique(NodeIndex);
-		}
-		else
-		{
-			// Put it in component 0 until a rebuild occurs.
-			Node.ComponentIndex = 0;
 		}
 	}
 
 	// Create a UModelComponent for each grid cell's node list.
-	for(TMap< FModelComponentKey, TArray<uint16> >::TConstIterator It(ModelComponentMap);It;++It)
+	for (TMap< FModelComponentKey, FNodeIndices >::TConstIterator It(ModelComponentMap); It; ++It)
 	{
 		const FModelComponentKey&		Key		= It.Key();
-		const TArray<uint16>&			Nodes	= It.Value();	
+		const TArray<uint16>&			Nodes	= It.Value().Nodes;
 
 		for(int32 NodeIndex = 0;NodeIndex < Nodes.Num();NodeIndex++)
 		{
@@ -942,6 +975,8 @@ void ULevel::CreateModelComponents()
 	Model->NumIncompleteNodeGroups = 0;
 	Model->CachedMappings.Empty();
 
+	SlowTask.EnterProgressFrame(4);
+
 	// Work only needed if we actually have BSP in the level.
 	if( ModelComponents.Num() )
 	{
@@ -1044,6 +1079,8 @@ void ULevel::CreateModelComponents()
 	}
 	Model->UpdateVertices();
 
+	SlowTask.EnterProgressFrame(2);
+
 	for (int32 UpdateCompIdx = 0; UpdateCompIdx < ModelComponents.Num(); UpdateCompIdx++)
 	{
 		UModelComponent* ModelComp = ModelComponents[UpdateCompIdx];
diff --git a/Engine/Source/Runtime/Engine/Private/ModelLight.cpp b/Engine/Source/Runtime/Engine/Private/ModelLight.cpp
--- a/Engine/Source/Runtime/Engine/Private/ModelLight.cpp
+++ b/Engine/Source/Runtime/Engine/Private/ModelLight.cpp
@@ -703,9 +703,14 @@ void FPlaneMap::AddPlane(const FPlane& Plane, int32 Index)
 void UModel::GroupAllNodes(const ULevel* Level, const TArray<ULightComponentBase*>& Lights)
 {
 #if WITH_EDITOR
+	FScopedSlowTask SlowTask(10);
+	SlowTask.MakeDialogDelayed(3.0f);
+
 	// cache the level
 	LightingLevel = Level;
 
+	SlowTask.EnterProgressFrame(1);
+
 	// gather all the lights for each component
 	TMap<int32, TArray<ULightComponent*> > ComponentRelevantLights;
 	for (int32 ComponentIndex = 0; ComponentIndex < Level->ModelComponents.Num(); ComponentIndex++)
@@ -765,6 +770,8 @@ void UModel::GroupAllNodes(const ULevel* Level, const TArray<ULightComponentBase
 	// The FPlaneMap buckets together planes with components within a specified granular range.
 	FPlaneMap PlaneMap(1/16.0f, 50.0f, GLightmassDebugOptions.CoplanarTolerance);
 
+	SlowTask.EnterProgressFrame(1);
+
 	for (int32 NodeIndex = 0; NodeIndex < Nodes.Num(); NodeIndex++)
 	{
 		const FBspNode& Node = Nodes[NodeIndex];
@@ -776,150 +783,159 @@ void UModel::GroupAllNodes(const ULevel* Level, const TArray<ULightComponentBase
 		}
 	}
 
+	SlowTask.EnterProgressFrame(8);
+
 	// Every item in the PlaneMap now contains a list of indices of nodes with similar planes.
 	// Now we can do a O(n^2) check to see if any pairs of nodes have planes within the allowed threshold, to be added to the same group.
-	for (const auto& Pair : PlaneMap.GetMap())
 	{
-		const auto& PlaneMapItem = Pair.Value;
+		FScopedSlowTask InnerTask(PlaneMap.GetMap().Num());
+		InnerTask.MakeDialogDelayed(3.0f);
 
-		const int32 NumMapNodes = PlaneMapItem.Num();
-		if (NumMapNodes > 1)
+		for (const auto& Pair : PlaneMap.GetMap())
 		{
-			for (int32 MapIndex1 = 0; MapIndex1 < NumMapNodes - 1; MapIndex1++)
+			InnerTask.EnterProgressFrame(1);
+
+			const auto& PlaneMapItem = Pair.Value;
+
+			const int32 NumMapNodes = PlaneMapItem.Num();
+			if (NumMapNodes > 1)
 			{
-				for (int32 MapIndex2 = MapIndex1 + 1; MapIndex2 < NumMapNodes; MapIndex2++)
+				for (int32 MapIndex1 = 0; MapIndex1 < NumMapNodes - 1; MapIndex1++)
 				{
-					const int32 NodeIndex1 = PlaneMapItem[MapIndex1];
-					const int32 NodeIndex2 = PlaneMapItem[MapIndex2];
-					const FBspNode& Node1 = Nodes[NodeIndex1];
-					const FBspNode& Node2 = Nodes[NodeIndex2];
-					const FBspSurf& Surf1 = Surfs[Node1.iSurf];
-					const FBspSurf& Surf2 = Surfs[Node2.iSurf];
-
-					// if I've already been parented, I don't need to reparent
-					if (ParentNodes[NodeIndex1] != nullptr && ParentNodes[NodeIndex2] != nullptr && ParentNodes[NodeIndex1] == ParentNodes[NodeIndex2])
+					for (int32 MapIndex2 = MapIndex1 + 1; MapIndex2 < NumMapNodes; MapIndex2++)
 					{
-						continue;
-					}
+						const int32 NodeIndex1 = PlaneMapItem[MapIndex1];
+						const int32 NodeIndex2 = PlaneMapItem[MapIndex2];
+						const FBspNode& Node1 = Nodes[NodeIndex1];
+						const FBspNode& Node2 = Nodes[NodeIndex2];
+						const FBspSurf& Surf1 = Surfs[Node1.iSurf];
+						const FBspSurf& Surf2 = Surfs[Node2.iSurf];
+
+						// if I've already been parented, I don't need to reparent
+						if (ParentNodes[NodeIndex1] != nullptr && ParentNodes[NodeIndex2] != nullptr && ParentNodes[NodeIndex1] == ParentNodes[NodeIndex2])
+						{
+							continue;
+						}
 
-					// variable to see check if the 2 nodes are conodes
-					bool bNodesAreConodes = false;
+						// variable to see check if the 2 nodes are conodes
+						bool bNodesAreConodes = false;
 
-					// if we have a tolerance, then join based on coplanar adjacency
-					if (GLightmassDebugOptions.bGatherBSPSurfacesAcrossComponents)
-					{
-						// are these two nodes conodes?
-						if (Surf1.LightMapScale == Surf2.LightMapScale)
+						// if we have a tolerance, then join based on coplanar adjacency
+						if (GLightmassDebugOptions.bGatherBSPSurfacesAcrossComponents)
 						{
-							if (LightmassSettingsEquality[Surf1.iLightmassIndex * NumLightmassSettings + Surf2.iLightmassIndex])
+							// are these two nodes conodes?
+							if (Surf1.LightMapScale == Surf2.LightMapScale)
 							{
-								if (Surf1.Plane.Equals(Surf2.Plane, GLightmassDebugOptions.CoplanarTolerance))
+								if (LightmassSettingsEquality[Surf1.iLightmassIndex * NumLightmassSettings + Surf2.iLightmassIndex])
 								{
-									// they are coplanar, have the same lightmap res and Lightmass settings, 
-									// now we need to check for adjacency which we check for by looking for a shared vertex
-									// This is O(n^2) but since there are often only 3 or 4 verts in a poly, this will iterate on average only about 16 times.
-									// I doubt it would be any more efficient to use a TSet to check for duplicated indices in this case.
-									FVert* VertPool1 = &Verts[Node1.iVertPool];
-									for (int32 A = 0; A < Node1.NumVertices && !bNodesAreConodes; A++)
+									if (Surf1.Plane.Equals(Surf2.Plane, GLightmassDebugOptions.CoplanarTolerance))
 									{
-										FVert* VertPool2 = &Verts[Node2.iVertPool];
-										for (int32 B = 0; B < Node2.NumVertices && !bNodesAreConodes; B++)
+										// they are coplanar, have the same lightmap res and Lightmass settings, 
+										// now we need to check for adjacency which we check for by looking for a shared vertex
+										// This is O(n^2) but since there are often only 3 or 4 verts in a poly, this will iterate on average only about 16 times.
+										// I doubt it would be any more efficient to use a TSet to check for duplicated indices in this case.
+										FVert* VertPool1 = &Verts[Node1.iVertPool];
+										for (int32 A = 0; A < Node1.NumVertices && !bNodesAreConodes; A++)
 										{
-											// if they share a vertex location, they are adjacent (this won't detect adjacency via T-joints)
-											if (VertPool1->pVertex == VertPool2->pVertex)
+											FVert* VertPool2 = &Verts[Node2.iVertPool];
+											for (int32 B = 0; B < Node2.NumVertices && !bNodesAreConodes; B++)
 											{
-												bNodesAreConodes = true;
+												// if they share a vertex location, they are adjacent (this won't detect adjacency via T-joints)
+												if (VertPool1->pVertex == VertPool2->pVertex)
+												{
+													bNodesAreConodes = true;
+												}
+												VertPool2++;
 											}
-											VertPool2++;
+											VertPool1++;
 										}
-										VertPool1++;
 									}
 								}
 							}
 						}
-					}
-					// if coplanar tolerance is < 0, then we join nodes together based on being in the same ModelComponent
-					// and from the same surface
-					else
-					{
-						if (Node1.iSurf == Node2.iSurf && Node1.ComponentIndex == Node2.ComponentIndex)
+						// if coplanar tolerance is < 0, then we join nodes together based on being in the same ModelComponent
+						// and from the same surface
+						else
 						{
-							bNodesAreConodes = true;
+							if (Node1.iSurf == Node2.iSurf && Node1.ComponentIndex == Node2.ComponentIndex)
+							{
+								bNodesAreConodes = true;
+							}
 						}
-					}
 
-					// are Node1 and Node2 conodes - if so, join into a group
-					if (bNodesAreConodes)
-					{
-						// okay, these two nodes are conodes, so we need to stick them together into some pot of nodes
-						// look to see if either one are already in a group
-						FNodeGroup* NodeGroup = NULL;
-						// if both are already in different groups, we need to combine the groups
-						if (ParentNodes[NodeIndex1] != NULL && ParentNodes[NodeIndex2] != NULL)
+						// are Node1 and Node2 conodes - if so, join into a group
+						if (bNodesAreConodes)
 						{
-							NodeGroup = ParentNodes[NodeIndex1];
-
-							// merge 2 into 1
-							FNodeGroup* NodeGroup2 = ParentNodes[NodeIndex2];
-							for (int32 NodeIndex = 0; NodeIndex < NodeGroup2->Nodes.Num(); NodeIndex++)
-							{
-								NodeGroup->Nodes.Add(NodeGroup2->Nodes[NodeIndex]);
-							}
-							for (int32 LightIndex = 0; LightIndex < NodeGroup2->RelevantLights.Num(); LightIndex++)
+							// okay, these two nodes are conodes, so we need to stick them together into some pot of nodes
+							// look to see if either one are already in a group
+							FNodeGroup* NodeGroup = NULL;
+							// if both are already in different groups, we need to combine the groups
+							if (ParentNodes[NodeIndex1] != NULL && ParentNodes[NodeIndex2] != NULL)
 							{
-								NodeGroup->RelevantLights.AddUnique(NodeGroup2->RelevantLights[LightIndex]);
-							}
+								NodeGroup = ParentNodes[NodeIndex1];
 
-							// replace all the users of NodeGroup2 with NodeGroup
-							for (int32 GroupIndex = 0; GroupIndex < ParentNodes.Num(); GroupIndex++)
-							{
-								if (ParentNodes[GroupIndex] == NodeGroup2)
+								// merge 2 into 1
+								FNodeGroup* NodeGroup2 = ParentNodes[NodeIndex2];
+								for (int32 NodeIndex = 0; NodeIndex < NodeGroup2->Nodes.Num(); NodeIndex++)
 								{
-									ParentNodes[GroupIndex] = NodeGroup;
+									NodeGroup->Nodes.Add(NodeGroup2->Nodes[NodeIndex]);
+								}
+								for (int32 LightIndex = 0; LightIndex < NodeGroup2->RelevantLights.Num(); LightIndex++)
+								{
+									NodeGroup->RelevantLights.AddUnique(NodeGroup2->RelevantLights[LightIndex]);
 								}
-							}
-
-							// the key for the nodegroup is the 0th node (could just be a set now)
-							NodeGroups.Remove(NodeGroup2->Nodes[0]);
 
-							// free the now useless nodegroup
-							delete NodeGroup2;
-						}
-						else if (ParentNodes[NodeIndex1] != NULL)
-						{
-							NodeGroup = ParentNodes[NodeIndex1];
-						}
-						else if (ParentNodes[NodeIndex2] != NULL)
-						{
-							NodeGroup = ParentNodes[NodeIndex2];
-						}
-						// otherwise, make a new group and put them both in it
-						else
-						{
-							NodeGroup = NodeGroups.Add(NodeIndex1, new FNodeGroup());
-						}
+								// replace all the users of NodeGroup2 with NodeGroup
+								for (int32 GroupIndex = 0; GroupIndex < ParentNodes.Num(); GroupIndex++)
+								{
+									if (ParentNodes[GroupIndex] == NodeGroup2)
+									{
+										ParentNodes[GroupIndex] = NodeGroup;
+									}
+								}
 
-						// apply both these nodes to the NodeGroup
-						for (int32 WhichNode = 0; WhichNode < 2; WhichNode++)
-						{
-							// operator on each node in this loop
-							int32 NodeIndex = WhichNode ? NodeIndex2 : NodeIndex1;
+								// the key for the nodegroup is the 0th node (could just be a set now)
+								NodeGroups.Remove(NodeGroup2->Nodes[0]);
 
-							// track what group the node went into
-							ParentNodes[NodeIndex] = NodeGroup;
+								// free the now useless nodegroup
+								delete NodeGroup2;
+							}
+							else if (ParentNodes[NodeIndex1] != NULL)
+							{
+								NodeGroup = ParentNodes[NodeIndex1];
+							}
+							else if (ParentNodes[NodeIndex2] != NULL)
+							{
+								NodeGroup = ParentNodes[NodeIndex2];
+							}
+							// otherwise, make a new group and put them both in it
+							else
+							{
+								NodeGroup = NodeGroups.Add(NodeIndex1, new FNodeGroup());
+							}
 
-							// is this node already not yet in the group
-							if (NodeGroup->Nodes.Find(NodeIndex) == INDEX_NONE)
+							// apply both these nodes to the NodeGroup
+							for (int32 WhichNode = 0; WhichNode < 2; WhichNode++)
 							{
-								// add it to the group
-								NodeGroup->Nodes.Add(NodeIndex);
+								// operator on each node in this loop
+								int32 NodeIndex = WhichNode ? NodeIndex2 : NodeIndex1;
+
+								// track what group the node went into
+								ParentNodes[NodeIndex] = NodeGroup;
 
-								// add the relevant lights to the nodegroup
-								TArray<ULightComponent*>* RelevantLights = ComponentRelevantLights.Find(Nodes[NodeIndex].ComponentIndex);
-								check(RelevantLights);
-								for (int32 LightIndex = 0; LightIndex < RelevantLights->Num(); LightIndex++)
+								// is this node already not yet in the group
+								if (NodeGroup->Nodes.Find(NodeIndex) == INDEX_NONE)
 								{
-									NodeGroup->RelevantLights.AddUnique((*RelevantLights)[LightIndex]);
+									// add it to the group
+									NodeGroup->Nodes.Add(NodeIndex);
+
+									// add the relevant lights to the nodegroup
+									TArray<ULightComponent*>* RelevantLights = ComponentRelevantLights.Find(Nodes[NodeIndex].ComponentIndex);
+									check(RelevantLights);
+									for (int32 LightIndex = 0; LightIndex < RelevantLights->Num(); LightIndex++)
+									{
+										NodeGroup->RelevantLights.AddUnique((*RelevantLights)[LightIndex]);
+									}
 								}
 							}
 						}
