Made SceneComponent notify NavigationSystem about editor-time changes to relative location/rotation/scale #UE4

diff --git a/Engine/Source/Runtime/Engine/Private/Components/SceneComponent.cpp b/Engine/Source/Runtime/Engine/Private/Components/SceneComponent.cpp
--- a/Engine/Source/Runtime/Engine/Private/Components/SceneComponent.cpp
+++ b/Engine/Source/Runtime/Engine/Private/Components/SceneComponent.cpp
@@ -264,13 +264,24 @@ static void UpdateAttachedMobility(USceneComponent* ComponentThatChanged)
 
 void USceneComponent::PostEditChangeProperty(FPropertyChangedEvent& PropertyChangedEvent)
 {
+	const static FName LocationName("RelativeLocation");
+	const static FName RotationName("RelativeRotation");
+	const static FName ScaleName("RelativeScale3D");
+
+	const FName PropertyName = PropertyChangedEvent.Property ? PropertyChangedEvent.Property->GetFName() : FName();
+
 	// Note: This must be called before UActorComponent::PostEditChangeChainProperty is called because this component will be reset when UActorComponent reruns construction scripts 
-	if(PropertyChangedEvent.Property && PropertyChangedEvent.Property->GetFName() == SceneComponentStatics::MobilityName)
+	if (PropertyName == SceneComponentStatics::MobilityName)
 	{
 		UpdateAttachedMobility(this);
 	}
 
 	Super::PostEditChangeProperty(PropertyChangedEvent);
+
+	if (PropertyName == LocationName || PropertyName == RotationName || PropertyName == ScaleName)
+	{
+		UNavigationSystem::UpdateNavOctree(this);
+	}
 }
 
 void USceneComponent::PostEditChangeChainProperty(FPropertyChangedChainEvent& PropertyChangedEvent)
