UE4 - fix lambda compile bug

Merging

//depot/UE4-Orion/Engine/Source/Runtime/Engine/Private/LevelTick.cpp

to //depot/UE4/Engine/Source/Runtime/Engine/Private/LevelTick.cpp

diff --git a/Engine/Source/Runtime/Engine/Private/LevelTick.cpp b/Engine/Source/Runtime/Engine/Private/LevelTick.cpp
--- a/Engine/Source/Runtime/Engine/Private/LevelTick.cpp
+++ b/Engine/Source/Runtime/Engine/Private/LevelTick.cpp
@@ -852,7 +852,27 @@ void UWorld::SendAllEndOfFrameUpdates()
 	}
 	else
 	{
+#if 1
+		QUICK_SCOPE_CYCLE_COUNTER(STAT_PostTickComponentUpdate_ForcedGameThread);
+		for (TArray<TWeakObjectPtr<UActorComponent> >::TIterator It(LocalComponentsThatNeedEndOfFrameUpdate_OnGameThread); It; ++It)
+		{
+			UActorComponent* Component = It->Get();
+			if (Component)
+			{
+				if ( !Component->IsPendingKill() && Component->IsRegistered() && !Component->IsTemplate())
+				{
+					FScopeCycleCounterUObject ComponentScope(Component);
+					FScopeCycleCounterUObject AdditionalScope(STATS ? Component->AdditionalStatObject() : NULL);
+					Component->DoDeferredRenderUpdates_Concurrent();
+				}
+				check(Component->GetMarkedForEndOfFrameUpdateState() == EComponentMarkedForEndOfFrameUpdateState::MarkedForGameThread);
+				FMarkComponentEndOfFrameUpdateState::Set(Component, EComponentMarkedForEndOfFrameUpdateState::Unmarked);
+			}
+		}
+		LocalComponentsThatNeedEndOfFrameUpdate_OnGameThread.Reset();
+#else
 		GTWork();
+#endif
 		ParallelFor(ComponentsThatNeedEndOfFrameUpdate.Num(), ParallelWork);
 	}
 	ComponentsThatNeedEndOfFrameUpdate.Reset();
