#ue4 - Added "bTickBeforeOwner" to MovementComponent to control whether a tick prereq is set for the owner.

- Default true since in most cases a ticking Pawn would need the latest location from the component update, but disabling this can get back some perf.

--------

diff --git a/Engine/Source/Runtime/Engine/Classes/GameFramework/MovementComponent.h b/Engine/Source/Runtime/Engine/Classes/GameFramework/MovementComponent.h
--- a/Engine/Source/Runtime/Engine/Classes/GameFramework/MovementComponent.h
+++ b/Engine/Source/Runtime/Engine/Classes/GameFramework/MovementComponent.h
@@ -139,6 +139,14 @@ public:
 	UPROPERTY(EditAnywhere, BlueprintReadOnly, Category=MovementComponent)
 	uint32 bAutoUpdateTickRegistration:1;
 
+	/**
+	 * If true, after registration we will add a tick dependency to tick before our owner (if we can both tick).
+	 * This is important when our tick causes an update in the owner's position, so that when the owner ticks it uses the most recent position without lag.
+	 * Disabling this can improve performance if both objects tick but the order of ticks doesn't matter.
+	 */
+	UPROPERTY(EditAnywhere, BlueprintReadOnly, Category=MovementComponent)
+	uint32 bTickBeforeOwner:1;
+
 	/** If true, registers the owner's Root component as the UpdatedComponent if there is not one currently assigned. */
 	UPROPERTY(EditAnywhere, BlueprintReadOnly, Category=MovementComponent)
 	uint32 bAutoRegisterUpdatedComponent:1;
diff --git a/Engine/Source/Runtime/Engine/Private/Components/MovementComponent.cpp b/Engine/Source/Runtime/Engine/Private/Components/MovementComponent.cpp
--- a/Engine/Source/Runtime/Engine/Private/Components/MovementComponent.cpp
+++ b/Engine/Source/Runtime/Engine/Private/Components/MovementComponent.cpp
@@ -25,6 +25,7 @@ UMovementComponent::UMovementComponent(const FObjectInitializer& ObjectInitializ
 
 	bUpdateOnlyIfRendered = false;
 	bAutoUpdateTickRegistration = true;
+	bTickBeforeOwner = true;
 	bAutoRegisterUpdatedComponent = true;
 
 	PlaneConstraintNormal = FVector::ZeroVector;
@@ -141,7 +142,7 @@ void UMovementComponent::RegisterComponentTickFunctions(bool bRegister)
 
 	// If the owner ticks, make sure we tick first
 	AActor* Owner = GetOwner();
-	if (bRegister && PrimaryComponentTick.bCanEverTick && Owner && Owner->CanEverTick())
+	if (bTickBeforeOwner && bRegister && PrimaryComponentTick.bCanEverTick && Owner && Owner->CanEverTick())
 	{
 		Owner->PrimaryActorTick.AddPrerequisite(this, PrimaryComponentTick);
 	}
diff --git a/Engine/Source/Runtime/Engine/Private/Controller.cpp b/Engine/Source/Runtime/Engine/Private/Controller.cpp
--- a/Engine/Source/Runtime/Engine/Private/Controller.cpp
+++ b/Engine/Source/Runtime/Engine/Private/Controller.cpp
@@ -327,10 +327,8 @@ void AController::RemovePawnTickDependency(APawn* InOldPawn)
 		{
 			PawnMovement->PrimaryComponentTick.RemovePrerequisite(this, this->PrimaryActorTick);
 		}
-		else
-		{
-			InOldPawn->PrimaryActorTick.RemovePrerequisite(this, this->PrimaryActorTick);
-		}
+		
+		InOldPawn->PrimaryActorTick.RemovePrerequisite(this, this->PrimaryActorTick);
 	}
 }
 
@@ -339,12 +337,20 @@ void AController::AddPawnTickDependency(APawn* NewPawn)
 {
 	if (NewPawn != NULL)
 	{
+		bool bNeedsPawnPrereq = true;
 		UPawnMovementComponent* PawnMovement = NewPawn->GetMovementComponent();
-		if (PawnMovement)
+		if (PawnMovement && PawnMovement->PrimaryComponentTick.bCanEverTick)
 		{
 			PawnMovement->PrimaryComponentTick.AddPrerequisite(this, this->PrimaryActorTick);
+
+			// Don't need a prereq on the pawn if the movement component already sets up a prereq.
+			if (PawnMovement->bTickBeforeOwner || NewPawn->PrimaryActorTick.GetPrerequisites().Contains(FTickPrerequisite(PawnMovement, PawnMovement->PrimaryComponentTick)))
+			{
+				bNeedsPawnPrereq = false;
+			}
 		}
-		else
+		
+		if (bNeedsPawnPrereq)
 		{
 			NewPawn->PrimaryActorTick.AddPrerequisite(this, this->PrimaryActorTick);
 		}
