Audio: Added another check to validate memory management assumptions
[INTEGRATE] Change 2660466 by [REDACTED_PERSON][REDACTED_HANDLE] on 2015/08/18 23:02:09

diff --git a/Engine/Source/Runtime/Engine/Private/AudioDecompress.cpp b/Engine/Source/Runtime/Engine/Private/AudioDecompress.cpp
--- a/Engine/Source/Runtime/Engine/Private/AudioDecompress.cpp
+++ b/Engine/Source/Runtime/Engine/Private/AudioDecompress.cpp
@@ -332,6 +332,7 @@ void FAsyncAudioDecompressWorker::DoWork()
 			else
 			{
 				Wave->RawPCMDataSize = QualityInfo.SampleDataSize;
+				check(Wave->RawPCMData == nullptr);
 				Wave->RawPCMData = ( uint8* )FMemory::Malloc( Wave->RawPCMDataSize );
 
 				// Decompress all the sample data into preallocated memory
