Crash fix for UWorld::SendAllEndOfFrameUpdates

A TFunctionRef was being used to store some temporary lambdas, but TFunctionRef only stores a reference to the lambda (it does not keep it alive), which meant that the lambda was immediately being destroyed, and the TFunctionRef was left pointing at invalid memory. This caused a crash in non-debug builds.

diff --git a/Engine/Source/Runtime/Engine/Private/LevelTick.cpp b/Engine/Source/Runtime/Engine/Private/LevelTick.cpp
--- a/Engine/Source/Runtime/Engine/Private/LevelTick.cpp
+++ b/Engine/Source/Runtime/Engine/Private/LevelTick.cpp
@@ -808,7 +808,7 @@ void UWorld::SendAllEndOfFrameUpdates()
 	TArray<TWeakObjectPtr<class UActorComponent> >& LocalComponentsThatNeedEndOfFrameUpdate = ComponentsThatNeedEndOfFrameUpdate;
 	TArray<TWeakObjectPtr<class UActorComponent> >& LocalComponentsThatNeedEndOfFrameUpdate_OnGameThread = ComponentsThatNeedEndOfFrameUpdate_OnGameThread;
 
-	TFunctionRef<void(int32)> ParallelWork = 
+	auto ParallelWork = 
 		[&LocalComponentsThatNeedEndOfFrameUpdate](int32 Index)
 		{
 			UActorComponent* NextComponent = LocalComponentsThatNeedEndOfFrameUpdate[Index].Get();
@@ -824,7 +824,7 @@ void UWorld::SendAllEndOfFrameUpdates()
 				FMarkComponentEndOfFrameUpdateState::Set(NextComponent, EComponentMarkedForEndOfFrameUpdateState::Unmarked);
 			}
 		};
-	TFunctionRef<void()> GTWork = 
+	auto GTWork = 
 		[&LocalComponentsThatNeedEndOfFrameUpdate_OnGameThread]()
 		{
 			QUICK_SCOPE_CYCLE_COUNTER(STAT_PostTickComponentUpdate_ForcedGameThread);
@@ -852,27 +852,7 @@ void UWorld::SendAllEndOfFrameUpdates()
 	}
 	else
 	{
-#if 1
-		QUICK_SCOPE_CYCLE_COUNTER(STAT_PostTickComponentUpdate_ForcedGameThread);
-		for (TArray<TWeakObjectPtr<UActorComponent> >::TIterator It(LocalComponentsThatNeedEndOfFrameUpdate_OnGameThread); It; ++It)
-		{
-			UActorComponent* Component = It->Get();
-			if (Component)
-			{
-				if ( !Component->IsPendingKill() && Component->IsRegistered() && !Component->IsTemplate())
-				{
-					FScopeCycleCounterUObject ComponentScope(Component);
-					FScopeCycleCounterUObject AdditionalScope(STATS ? Component->AdditionalStatObject() : NULL);
-					Component->DoDeferredRenderUpdates_Concurrent();
-				}
-				check(Component->GetMarkedForEndOfFrameUpdateState() == EComponentMarkedForEndOfFrameUpdateState::MarkedForGameThread);
-				FMarkComponentEndOfFrameUpdateState::Set(Component, EComponentMarkedForEndOfFrameUpdateState::Unmarked);
-			}
-		}
-		LocalComponentsThatNeedEndOfFrameUpdate_OnGameThread.Reset();
-#else
 		GTWork();
-#endif
 		ParallelFor(ComponentsThatNeedEndOfFrameUpdate.Num(), ParallelWork);
 	}
 	ComponentsThatNeedEndOfFrameUpdate.Reset();
