[Merging using UE4-Orion-To-UE4] - Fix for anim relevancy nodes giving strange values when no relevancy can be determined

diff --git a/Engine/Source/Runtime/Engine/Private/Animation/AnimInstance.cpp b/Engine/Source/Runtime/Engine/Private/Animation/AnimInstance.cpp
--- a/Engine/Source/Runtime/Engine/Private/Animation/AnimInstance.cpp
+++ b/Engine/Source/Runtime/Engine/Private/Animation/AnimInstance.cpp
@@ -2705,7 +2705,7 @@ float UAnimInstance::GetInstanceAssetPlayerTimeFromEndFraction(int32 AssetPlayer
 		}
 	}
 
-	return 0.0f;
+	return 1.0f;
 }
 
 float UAnimInstance::GetInstanceAssetPlayerTimeFromEnd(int32 AssetPlayerIndex)
@@ -2719,7 +2719,7 @@ float UAnimInstance::GetInstanceAssetPlayerTimeFromEnd(int32 AssetPlayerIndex)
 		}
 	}
 
-	return 0.0f;
+	return MAX_flt;
 }
 
 float UAnimInstance::GetInstanceStateWeight(int32 MachineIndex, int32 StateIndex)
@@ -2784,7 +2784,7 @@ float UAnimInstance::GetRelevantAnimTimeRemaining(int32 MachineIndex, int32 Stat
 		}
 	}
 
-	return 0.0f;
+	return MAX_flt;
 }
 
 float UAnimInstance::GetRelevantAnimTimeRemainingFraction(int32 MachineIndex, int32 StateIndex)
@@ -2801,7 +2801,7 @@ float UAnimInstance::GetRelevantAnimTimeRemainingFraction(int32 MachineIndex, in
 		}
 	}
 
-	return 0.0f;
+	return 1.0f;
 }
 
 float UAnimInstance::GetRelevantAnimLength(int32 MachineIndex, int32 StateIndex)
