Ensured that Engine/Runtime modules don't depend on SlateReflector in a shipping build

Changed everything using SlateReflector to list it as a dynamic dependency. Nothing should ever need to link directly to SlateReflector as it just provides an interface for spawning the various debug UIs, such as the WidgetReflector.

Went through and made sure that the Engine/Runtime modules that use SlateReflector aren't using it in a shipping build. Also removed the testing code from AppFramework, and made sure that you can't try and spawn the test suite in a shipping build.

diff --git a/Engine/Source/Editor/LevelEditor/LevelEditor.Build.cs b/Engine/Source/Editor/LevelEditor/LevelEditor.Build.cs
--- a/Engine/Source/Editor/LevelEditor/LevelEditor.Build.cs
+++ b/Engine/Source/Editor/LevelEditor/LevelEditor.Build.cs
@@ -14,6 +14,7 @@ public class LevelEditor : ModuleRules
                 "PlacementMode",
 				"ReferenceViewer",
 				"SizeMap",
+				"SlateReflector",
                 "IntroTutorials",
                 "AppFramework"
 			}
@@ -37,7 +38,6 @@ public class LevelEditor : ModuleRules
                 "InputCore",
 				"Slate",
 				"SlateCore",
-				"SlateReflector",
                 "EditorStyle",
 				"Engine",
 				"MessageLog",
@@ -72,6 +72,7 @@ public class LevelEditor : ModuleRules
 				"DeviceManager",
 				"SettingsEditor",
 				"SessionFrontend",
+				"SlateReflector",
 				"AutomationWindow",
 				"Layers",
                 "WorldBrowser",
diff --git a/Engine/Source/Programs/CrashReporter/CrashReportClient/CrashReportClient.Build.cs b/Engine/Source/Programs/CrashReporter/CrashReportClient/CrashReportClient.Build.cs
--- a/Engine/Source/Programs/CrashReporter/CrashReportClient/CrashReportClient.Build.cs
+++ b/Engine/Source/Programs/CrashReporter/CrashReportClient/CrashReportClient.Build.cs
@@ -37,11 +37,22 @@ public class CrashReportClient : ModuleRules
 				{
 					"Slate",
 					"SlateCore",
-					"SlateReflector",
 					"StandaloneRenderer",
 					"MessageLog",
 				}
 			);
+
+			PrivateIncludePathModuleNames.AddRange(
+				new string[] {
+					"SlateReflector",
+				}
+			);
+
+			DynamicallyLoadedModuleNames.AddRange(
+				new string[] {
+					"SlateReflector",
+				}
+			);
 		}
 
 		PrivateIncludePaths.Add("Runtime/Launch/Private");		// For LaunchEngineLoop.cpp include
diff --git a/Engine/Source/Programs/SlateViewer/SlateViewer.Build.cs b/Engine/Source/Programs/SlateViewer/SlateViewer.Build.cs
--- a/Engine/Source/Programs/SlateViewer/SlateViewer.Build.cs
+++ b/Engine/Source/Programs/SlateViewer/SlateViewer.Build.cs
@@ -15,13 +15,24 @@ public class SlateViewer : ModuleRules
 				"Projects",
 				"Slate",
 				"SlateCore",
-				"SlateReflector",
 				"StandaloneRenderer",
 				"SourceCodeAccess",
 				"WebBrowser",
 			}
 		);
 
+		PrivateIncludePathModuleNames.AddRange(
+			new string[] {
+				"SlateReflector",
+			}
+		);
+
+		DynamicallyLoadedModuleNames.AddRange(
+			new string[] {
+				"SlateReflector",
+			}
+		);
+
 		if (Target.Platform == UnrealTargetPlatform.Mac)
 		{
 			PrivateDependencyModuleNames.Add("XCodeSourceCodeAccess");
diff --git a/Engine/Source/Programs/UnrealFrontend/UnrealFrontend.Build.cs b/Engine/Source/Programs/UnrealFrontend/UnrealFrontend.Build.cs
--- a/Engine/Source/Programs/UnrealFrontend/UnrealFrontend.Build.cs
+++ b/Engine/Source/Programs/UnrealFrontend/UnrealFrontend.Build.cs
@@ -32,7 +32,6 @@ public class UnrealFrontend : ModuleRules
 				"SessionServices",
 				"Slate",
 				"SlateCore",
-				"SlateReflector",
 				"SourceCodeAccess",
 				"StandaloneRenderer",
 				"TargetDeviceServices",
@@ -40,6 +39,18 @@ public class UnrealFrontend : ModuleRules
 			}
 		);
 
+		PrivateIncludePathModuleNames.AddRange(
+			new string[] {
+				"SlateReflector",
+			}
+		);
+
+		DynamicallyLoadedModuleNames.AddRange(
+			new string[] {
+				"SlateReflector",
+			}
+		);
+
 		if (Target.Platform == UnrealTargetPlatform.Mac)
 		{
 			PrivateDependencyModuleNames.Add("XCodeSourceCodeAccess");
diff --git a/Engine/Source/Runtime/AppFramework/AppFramework.Build.cs b/Engine/Source/Runtime/AppFramework/AppFramework.Build.cs
--- a/Engine/Source/Runtime/AppFramework/AppFramework.Build.cs
+++ b/Engine/Source/Runtime/AppFramework/AppFramework.Build.cs
@@ -18,7 +18,6 @@ public class AppFramework : ModuleRules
 		PrivateDependencyModuleNames.AddRange(
 			new string[] {
 				"InputCore",
-				"SlateReflector",
 			}
 		);
 
@@ -34,5 +33,20 @@ public class AppFramework : ModuleRules
                 "Runtime/AppFramework/Private/Widgets/Workflow",
 			}
 		);
+
+		if (Target.Configuration != UnrealTargetConfiguration.Shipping)
+		{
+			PrivateIncludePathModuleNames.AddRange(
+				new string[] {
+					"SlateReflector",
+				}
+			);
+
+			DynamicallyLoadedModuleNames.AddRange(
+				new string[] {
+					"SlateReflector",
+				}
+			);
+		}
 	}
 }
diff --git a/Engine/Source/Runtime/AppFramework/Private/Framework/Testing/SLayoutExample.cpp b/Engine/Source/Runtime/AppFramework/Private/Framework/Testing/SLayoutExample.cpp
--- a/Engine/Source/Runtime/AppFramework/Private/Framework/Testing/SLayoutExample.cpp
+++ b/Engine/Source/Runtime/AppFramework/Private/Framework/Testing/SLayoutExample.cpp
@@ -3,6 +3,7 @@
 #include "AppFrameworkPrivatePCH.h"
 #include "SLayoutExample.h"
 
+#if !UE_BUILD_SHIPPING
 
 #define LOCTEXT_NAMESPACE "ExampleLayoutTest"
 
@@ -381,7 +382,7 @@ public:
 
 
 
-APPFRAMEWORK_API TSharedRef<SWidget> MakeLayoutExample()
+TSharedRef<SWidget> MakeLayoutExample()
 {
 	extern TOptional<FSlateRenderTransform> GetTestRenderTransform();
 	extern FVector2D GetTestRenderTransformPivot();
@@ -391,4 +392,6 @@ APPFRAMEWORK_API TSharedRef<SWidget> MakeLayoutExample()
 		.RenderTransformPivot_Static(&GetTestRenderTransformPivot);
 }
 
-#undef LOCTEXT_NAMESPACE
\ No newline at end of file
+#undef LOCTEXT_NAMESPACE
+
+#endif // #if !UE_BUILD_SHIPPING
diff --git a/Engine/Source/Runtime/AppFramework/Private/Framework/Testing/SLayoutExample.h b/Engine/Source/Runtime/AppFramework/Private/Framework/Testing/SLayoutExample.h
--- a/Engine/Source/Runtime/AppFramework/Private/Framework/Testing/SLayoutExample.h
+++ b/Engine/Source/Runtime/AppFramework/Private/Framework/Testing/SLayoutExample.h
@@ -2,5 +2,8 @@
 
 #pragma once
 
+#if !UE_BUILD_SHIPPING
 
-APPFRAMEWORK_API TSharedRef<SWidget> MakeLayoutExample();
+TSharedRef<SWidget> MakeLayoutExample();
+
+#endif // #if !UE_BUILD_SHIPPING
diff --git a/Engine/Source/Runtime/AppFramework/Private/Framework/Testing/STableViewTesting.cpp b/Engine/Source/Runtime/AppFramework/Private/Framework/Testing/STableViewTesting.cpp
--- a/Engine/Source/Runtime/AppFramework/Private/Framework/Testing/STableViewTesting.cpp
+++ b/Engine/Source/Runtime/AppFramework/Private/Framework/Testing/STableViewTesting.cpp
@@ -2,6 +2,9 @@
 
 #include "AppFrameworkPrivatePCH.h"
 #include "STableViewTesting.h"
+
+#if !UE_BUILD_SHIPPING
+
 #include "SNumericEntryBox.h"
 
 #define LOCTEXT_NAMESPACE "STableViewTesting"
@@ -1278,4 +1281,6 @@ TSharedRef<SWidget> MakeTableViewTesting()
 }
 
 
-#undef LOCTEXT_NAMESPACE
\ No newline at end of file
+#undef LOCTEXT_NAMESPACE
+
+#endif // #if !UE_BUILD_SHIPPING
diff --git a/Engine/Source/Runtime/AppFramework/Private/Framework/Testing/STableViewTesting.h b/Engine/Source/Runtime/AppFramework/Private/Framework/Testing/STableViewTesting.h
--- a/Engine/Source/Runtime/AppFramework/Private/Framework/Testing/STableViewTesting.h
+++ b/Engine/Source/Runtime/AppFramework/Private/Framework/Testing/STableViewTesting.h
@@ -2,5 +2,8 @@
 
 #pragma once
 
+#if !UE_BUILD_SHIPPING
 
-APPFRAMEWORK_API TSharedRef<SWidget> MakeTableViewTesting();
+TSharedRef<SWidget> MakeTableViewTesting();
+
+#endif // #if !UE_BUILD_SHIPPING
diff --git a/Engine/Source/Runtime/AppFramework/Private/Framework/Testing/SUserWidgetTest.cpp b/Engine/Source/Runtime/AppFramework/Private/Framework/Testing/SUserWidgetTest.cpp
--- a/Engine/Source/Runtime/AppFramework/Private/Framework/Testing/SUserWidgetTest.cpp
+++ b/Engine/Source/Runtime/AppFramework/Private/Framework/Testing/SUserWidgetTest.cpp
@@ -3,6 +3,7 @@
 #include "AppFrameworkPrivatePCH.h"
 #include "SUserWidgetTest.h"
 
+#if !UE_BUILD_SHIPPING
 
 class SUserWidgetExampleImpl
 	: public SUserWidgetExample
@@ -29,3 +30,5 @@ TSharedRef<SUserWidgetExample> SUserWidgetExample::New()
 {
 	return MakeShareable(new SUserWidgetExampleImpl()); 
 }
+
+#endif // #if !UE_BUILD_SHIPPING
diff --git a/Engine/Source/Runtime/AppFramework/Private/Framework/Testing/SUserWidgetTest.h b/Engine/Source/Runtime/AppFramework/Private/Framework/Testing/SUserWidgetTest.h
--- a/Engine/Source/Runtime/AppFramework/Private/Framework/Testing/SUserWidgetTest.h
+++ b/Engine/Source/Runtime/AppFramework/Private/Framework/Testing/SUserWidgetTest.h
@@ -2,6 +2,7 @@
 
 #pragma once
 
+#if !UE_BUILD_SHIPPING
 
 class SUserWidgetExample
 	: public SUserWidget
@@ -18,3 +19,5 @@ class SUserWidgetExample
 
 	virtual void DoStuff() = 0;
 };
+
+#endif // #if !UE_BUILD_SHIPPING
diff --git a/Engine/Source/Runtime/AppFramework/Private/Framework/Testing/SWidgetGallery.cpp b/Engine/Source/Runtime/AppFramework/Private/Framework/Testing/SWidgetGallery.cpp
--- a/Engine/Source/Runtime/AppFramework/Private/Framework/Testing/SWidgetGallery.cpp
+++ b/Engine/Source/Runtime/AppFramework/Private/Framework/Testing/SWidgetGallery.cpp
@@ -1,9 +1,12 @@
 // Copyright 1998-2015 Epic Games, Inc. All Rights Reserved.
 
 #include "AppFrameworkPrivatePCH.h"
+#include "SWidgetGallery.h"
+
+#if !UE_BUILD_SHIPPING
+
 #include "SSearchBox.h"
 #include "SDockTab.h"
-#include "SWidgetGallery.h"
 #include "TestStyle.h"
 
 #include "SUserWidgetTest.h"
@@ -1174,3 +1177,5 @@ TSharedRef<SWidget> MakeWidgetGallery()
 END_SLATE_FUNCTION_BUILD_OPTIMIZATION
 
 #undef LOCTEXT_NAMESPACE
+
+#endif // #if !UE_BUILD_SHIPPING
diff --git a/Engine/Source/Runtime/AppFramework/Private/Framework/Testing/SWidgetGallery.h b/Engine/Source/Runtime/AppFramework/Private/Framework/Testing/SWidgetGallery.h
--- a/Engine/Source/Runtime/AppFramework/Private/Framework/Testing/SWidgetGallery.h
+++ b/Engine/Source/Runtime/AppFramework/Private/Framework/Testing/SWidgetGallery.h
@@ -2,5 +2,9 @@
 
 #pragma once
 
+#if !UE_BUILD_SHIPPING
+
 /** @return a new Drag and Drop test widget */
 TSharedRef<SWidget> MakeWidgetGallery();
+
+#endif // #if !UE_BUILD_SHIPPING
diff --git a/Engine/Source/Runtime/AppFramework/Private/Framework/Testing/TestStyle.cpp b/Engine/Source/Runtime/AppFramework/Private/Framework/Testing/TestStyle.cpp
--- a/Engine/Source/Runtime/AppFramework/Private/Framework/Testing/TestStyle.cpp
+++ b/Engine/Source/Runtime/AppFramework/Private/Framework/Testing/TestStyle.cpp
@@ -3,6 +3,7 @@
 #include "AppFrameworkPrivatePCH.h"
 #include "TestStyle.h"
 
+#if !UE_BUILD_SHIPPING
 
 TSharedPtr< ISlateStyle > FTestStyle::Instance = nullptr;
 
@@ -396,3 +397,5 @@ TSharedRef< ISlateStyle > FTestStyle::Create()
 
 	return Style;
 }
+
+#endif // #if !UE_BUILD_SHIPPING
diff --git a/Engine/Source/Runtime/AppFramework/Private/Framework/Testing/TestStyle.h b/Engine/Source/Runtime/AppFramework/Private/Framework/Testing/TestStyle.h
--- a/Engine/Source/Runtime/AppFramework/Private/Framework/Testing/TestStyle.h
+++ b/Engine/Source/Runtime/AppFramework/Private/Framework/Testing/TestStyle.h
@@ -1,6 +1,8 @@
 // Copyright 1998-2015 Epic Games, Inc. All Rights Reserved.
 #pragma once
 
+#if !UE_BUILD_SHIPPING
+
 /**
  * Test slate style
  */
@@ -27,3 +29,5 @@ private:
 	/** Singleton instances of this style. */
 	static TSharedPtr< class ISlateStyle > Instance;
 };
+
+#endif // #if !UE_BUILD_SHIPPING
diff --git a/Engine/Source/Runtime/AppFramework/Private/Widgets/Testing/SPerfSuite.cpp b/Engine/Source/Runtime/AppFramework/Private/Widgets/Testing/SPerfSuite.cpp
--- a/Engine/Source/Runtime/AppFramework/Private/Widgets/Testing/SPerfSuite.cpp
+++ b/Engine/Source/Runtime/AppFramework/Private/Widgets/Testing/SPerfSuite.cpp
@@ -2,6 +2,9 @@
 
 #include "AppFrameworkPrivatePCH.h"
 #include "SPerfSuite.h"
+
+#if !UE_BUILD_SHIPPING
+
 #include "STableViewTesting.h"
 #include "ISlateReflectorModule.h"
 
@@ -27,3 +30,5 @@ void SummonPerfTestSuite()
 		]
 	);
 }
+
+#endif // #if !UE_BUILD_SHIPPING
diff --git a/Engine/Source/Runtime/AppFramework/Private/Widgets/Testing/STestSuite.cpp b/Engine/Source/Runtime/AppFramework/Private/Widgets/Testing/STestSuite.cpp
--- a/Engine/Source/Runtime/AppFramework/Private/Widgets/Testing/STestSuite.cpp
+++ b/Engine/Source/Runtime/AppFramework/Private/Widgets/Testing/STestSuite.cpp
@@ -1,13 +1,16 @@
 // Copyright 1998-2015 Epic Games, Inc. All Rights Reserved.
 
 #include "AppFrameworkPrivatePCH.h"
+#include "STestSuite.h"
+
+#if !UE_BUILD_SHIPPING
+
 #include "STableViewTesting.h"
 #include "SLayoutExample.h"
 #include "SWidgetGallery.h"
 #include "TestStyle.h"
 #include "RichTextLayoutMarshaller.h"
 #include "SyntaxHighlighterTextLayoutMarshaller.h"
-#include "STestSuite.h"
 #include "SScissorRectBox.h"
 #include "TransformCalculus3D.h"
 #include "SlateRenderTransform.h"
@@ -5758,3 +5761,5 @@ void MakeSplitterTest()
 END_SLATE_FUNCTION_BUILD_OPTIMIZATION
 
 #undef LOCTEXT_NAMESPACE
+
+#endif // #if !UE_BUILD_SHIPPING
diff --git a/Engine/Source/Runtime/AppFramework/Public/Widgets/Testing/SPerfSuite.h b/Engine/Source/Runtime/AppFramework/Public/Widgets/Testing/SPerfSuite.h
--- a/Engine/Source/Runtime/AppFramework/Public/Widgets/Testing/SPerfSuite.h
+++ b/Engine/Source/Runtime/AppFramework/Public/Widgets/Testing/SPerfSuite.h
@@ -2,4 +2,8 @@
 
 #pragma once
 
-void SummonPerfTestSuite();
\ No newline at end of file
+#if !UE_BUILD_SHIPPING
+
+APPFRAMEWORK_API void SummonPerfTestSuite();
+
+#endif // #if !UE_BUILD_SHIPPING
diff --git a/Engine/Source/Runtime/AppFramework/Public/Widgets/Testing/STestSuite.h b/Engine/Source/Runtime/AppFramework/Public/Widgets/Testing/STestSuite.h
--- a/Engine/Source/Runtime/AppFramework/Public/Widgets/Testing/STestSuite.h
+++ b/Engine/Source/Runtime/AppFramework/Public/Widgets/Testing/STestSuite.h
@@ -2,6 +2,9 @@
 
 #pragma once
 
+#if !UE_BUILD_SHIPPING
 
 APPFRAMEWORK_API void RestoreSlateTestSuite();
 APPFRAMEWORK_API void MakeSplitterTest();
+
+#endif // #if !UE_BUILD_SHIPPING
diff --git a/Engine/Source/Runtime/Engine/Engine.Build.cs b/Engine/Source/Runtime/Engine/Engine.Build.cs
--- a/Engine/Source/Runtime/Engine/Engine.Build.cs
+++ b/Engine/Source/Runtime/Engine/Engine.Build.cs
@@ -38,6 +38,21 @@ public class Engine : ModuleRules
 			PrivateIncludePathModuleNames.AddRange(new string[] { "TaskGraph" });
 		}
 
+		if (Target.Configuration != UnrealTargetConfiguration.Shipping)
+		{
+			PrivateIncludePathModuleNames.AddRange(
+				new string[] {
+					"SlateReflector",
+				}
+			);
+
+			DynamicallyLoadedModuleNames.AddRange(
+				new string[] {
+					"SlateReflector",
+				}
+			);
+		}
+
 		PublicDependencyModuleNames.AddRange(
 			new string[] {
 				"Core",
@@ -65,7 +80,6 @@ public class Engine : ModuleRules
                 "AppFramework",
 				"Networking",
 				"Sockets",
-				"SlateReflector",
 				"Landscape",
                 "UMG",
 				"Projects",
diff --git a/Engine/Source/Runtime/Engine/Private/CheatManager.cpp b/Engine/Source/Runtime/Engine/Private/CheatManager.cpp
--- a/Engine/Source/Runtime/Engine/Private/CheatManager.cpp
+++ b/Engine/Source/Runtime/Engine/Private/CheatManager.cpp
@@ -3,15 +3,19 @@
 #include "EnginePrivate.h"
 #include "UnrealNetwork.h"
 #include "SlateBasics.h"
-#include "ISlateReflectorModule.h"
 #include "NavDataGenerator.h"
 #include "OnlineSubsystemUtils.h"
 #include "VisualLogger/VisualLogger.h"
 #include "GameFramework/Character.h"
 
+#if !UE_BUILD_SHIPPING
+#include "ISlateReflectorModule.h"
+#endif // #if !UE_BUILD_SHIPPING
+
 #if WITH_EDITOR
 #include "UnrealEd.h"
 #endif
+
 #include "GameFramework/CheatManager.h"
 #include "Components/CapsuleComponent.h"
 #include "Engine/DebugCameraController.h"
@@ -922,20 +926,26 @@ void UCheatManager::TestCollisionDistance()
 
 void UCheatManager::WidgetReflector()
 {
+#if !UE_BUILD_SHIPPING
 	static const FName SlateReflectorModuleName("SlateReflector");
 	FModuleManager::LoadModuleChecked<ISlateReflectorModule>(SlateReflectorModuleName).DisplayWidgetReflector();
+#endif // #if !UE_BUILD_SHIPPING
 }
 
 void UCheatManager::TextureAtlasVisualizer()
 {
+#if !UE_BUILD_SHIPPING
 	static const FName SlateReflectorModuleName("SlateReflector");
 	FModuleManager::LoadModuleChecked<ISlateReflectorModule>(SlateReflectorModuleName).DisplayTextureAtlasVisualizer();
+#endif // #if !UE_BUILD_SHIPPING
 }
 
 void UCheatManager::FontAtlasVisualizer()
 {
+#if !UE_BUILD_SHIPPING
 	static const FName SlateReflectorModuleName("SlateReflector");
 	FModuleManager::LoadModuleChecked<ISlateReflectorModule>(SlateReflectorModuleName).DisplayFontAtlasVisualizer();
+#endif // #if !UE_BUILD_SHIPPING
 }
 
 void UCheatManager::RebuildNavigation()
diff --git a/Engine/Source/Runtime/Engine/Private/UnrealEngine.cpp b/Engine/Source/Runtime/Engine/Private/UnrealEngine.cpp
--- a/Engine/Source/Runtime/Engine/Private/UnrealEngine.cpp
+++ b/Engine/Source/Runtime/Engine/Private/UnrealEngine.cpp
@@ -778,12 +778,13 @@ void UEngine::Init(IEngineLoop* InEngineLoop)
 		FSlateApplication& CurrentSlateApp = FSlateApplication::Get();
 		CurrentSlateApp.InitializeSound( TSharedRef<FSlateSoundDevice>( new FSlateSoundDevice() ) );
 
-
+#if !UE_BUILD_SHIPPING
 		// Create test windows (if we were asked to do that)
 		if( FParse::Param( FCommandLine::Get(), TEXT("SlateDebug") ) )
 		{
 			RestoreSlateTestSuite();
 		}
+#endif // #if !UE_BUILD_SHIPPING
 	}
 
 	// Assign thumbnail compressor/decompressor
diff --git a/Engine/Source/Runtime/Toolbox/Private/ToolboxModule.cpp b/Engine/Source/Runtime/Toolbox/Private/ToolboxModule.cpp
--- a/Engine/Source/Runtime/Toolbox/Private/ToolboxModule.cpp
+++ b/Engine/Source/Runtime/Toolbox/Private/ToolboxModule.cpp
@@ -11,7 +11,10 @@
 #include "SNotificationList.h"
 #include "NotificationManager.h"
 #include "STestSuite.h"
+
+#if !UE_BUILD_SHIPPING
 #include "ISlateReflectorModule.h"
+#endif // #if !UE_BUILD_SHIPPING
 
 
 static bool bTabsRegistered = false;
@@ -94,15 +97,19 @@ public:
 
 	FReply OnDisplayTextureAtlases()
 	{
+#if !UE_BUILD_SHIPPING
 		static const FName SlateReflectorModuleName("SlateReflector");
 		FModuleManager::LoadModuleChecked<ISlateReflectorModule>(SlateReflectorModuleName).DisplayTextureAtlasVisualizer();
+#endif // #if !UE_BUILD_SHIPPING
 		return FReply::Handled();
 	}
 
 	FReply OnDisplayFontAtlases()
 	{
+#if !UE_BUILD_SHIPPING
 		static const FName SlateReflectorModuleName("SlateReflector");
 		FModuleManager::LoadModuleChecked<ISlateReflectorModule>(SlateReflectorModuleName).DisplayFontAtlasVisualizer();
+#endif // #if !UE_BUILD_SHIPPING
 		return FReply::Handled();
 	}
 
@@ -114,7 +121,9 @@ public:
 
 	FReply OnTestSuiteClicked()
 	{
+#if !UE_BUILD_SHIPPING
 		RestoreSlateTestSuite();
+#endif
 		return FReply::Handled();
 	}
 };
diff --git a/Engine/Source/Runtime/Toolbox/Toolbox.Build.cs b/Engine/Source/Runtime/Toolbox/Toolbox.Build.cs
--- a/Engine/Source/Runtime/Toolbox/Toolbox.Build.cs
+++ b/Engine/Source/Runtime/Toolbox/Toolbox.Build.cs
@@ -20,7 +20,6 @@ public class Toolbox : ModuleRules
                 "InputCore",
 				"Slate",
 				"SlateCore",
-				"SlateReflector",
                 "EditorStyle",
 				"DesktopPlatform",
 				"AppFramework"
@@ -34,5 +33,20 @@ public class Toolbox : ModuleRules
 				"ModuleUI"
 			}
 		);
+
+		if (Target.Configuration != UnrealTargetConfiguration.Shipping)
+		{
+			PrivateIncludePathModuleNames.AddRange(
+				new string[] {
+					"SlateReflector",
+				}
+			);
+
+			DynamicallyLoadedModuleNames.AddRange(
+				new string[] {
+					"SlateReflector",
+				}
+			);
+		}
 	}
 }
\ No newline at end of file
