Fix decals and lights from 'swimming' behind head movement in VR.

diff --git a/Engine/Source/Runtime/Engine/Private/SceneView.cpp b/Engine/Source/Runtime/Engine/Private/SceneView.cpp
--- a/Engine/Source/Runtime/Engine/Private/SceneView.cpp
+++ b/Engine/Source/Runtime/Engine/Private/SceneView.cpp
@@ -575,10 +575,10 @@ void FSceneView::UpdateViewMatrix()
 		FPlane(0,	0,	0,	1));
  
 	ViewMatrices.PreViewTranslation = -ViewMatrices.ViewOrigin;
-	FMatrix TranslatedViewMatrix = FTranslationMatrix(-ViewMatrices.PreViewTranslation) * ViewMatrices.ViewMatrix;
+	ViewMatrices.TranslatedViewMatrix = FTranslationMatrix(-ViewMatrices.PreViewTranslation) * ViewMatrices.ViewMatrix;
 	
 	// Compute a transform from view origin centered world-space to clip space.
-	ViewMatrices.TranslatedViewProjectionMatrix = TranslatedViewMatrix * ViewMatrices.ProjMatrix;
+	ViewMatrices.TranslatedViewProjectionMatrix = ViewMatrices.TranslatedViewMatrix * ViewMatrices.ProjMatrix;
 	ViewMatrices.InvTranslatedViewProjectionMatrix = ViewMatrices.TranslatedViewProjectionMatrix.Inverse();
 
 	ViewProjectionMatrix = ViewMatrices.GetViewProjMatrix();
