[UE-19299] Minor optimization to component data cache matching.

diff --git a/Engine/Source/Runtime/Engine/Private/ComponentInstanceDataCache.cpp b/Engine/Source/Runtime/Engine/Private/ComponentInstanceDataCache.cpp
--- a/Engine/Source/Runtime/Engine/Private/ComponentInstanceDataCache.cpp
+++ b/Engine/Source/Runtime/Engine/Private/ComponentInstanceDataCache.cpp
@@ -95,11 +95,10 @@ bool FActorComponentInstanceData::MatchesComponent(const UActorComponent* Compon
 			{
 				for (const UActorComponent* BlueprintCreatedComponent : ComponentOwner->BlueprintCreatedComponents)
 				{
-					if (BlueprintCreatedComponent != nullptr)
+					if (BlueprintCreatedComponent != nullptr && BlueprintCreatedComponent->CreationMethod == SourceComponentCreationMethod)
 					{
 						const UObject* BlueprintComponentTemplate = BlueprintCreatedComponent->GetArchetype();
 						if (   (BlueprintComponentTemplate == SourceComponentTemplate || (GIsReinstancing && BlueprintComponentTemplate->GetFName() == SourceComponentTemplate->GetFName()))
-							&& (BlueprintCreatedComponent->CreationMethod == SourceComponentCreationMethod)
 							&& (++FoundSerializedComponentsOfType == SourceComponentTypeSerializedIndex))
 						{
 							bMatches = (BlueprintCreatedComponent == Component);
