Merging

//depot/UE4-Orion/Engine/Source/Runtime/Engine/Private/Animation/BlendSpaceBase.cpp

to //depot/UE4/Engine/Source/Runtime/Engine/Private/Animation/BlendSpaceBase.cpp

diff --git a/Engine/Source/Runtime/Engine/Private/Animation/BlendSpaceBase.cpp b/Engine/Source/Runtime/Engine/Private/Animation/BlendSpaceBase.cpp
--- a/Engine/Source/Runtime/Engine/Private/Animation/BlendSpaceBase.cpp
+++ b/Engine/Source/Runtime/Engine/Private/Animation/BlendSpaceBase.cpp
@@ -379,7 +379,7 @@ bool UBlendSpaceBase::InterpolateWeightOfSampleData(float DeltaTime, const TArra
 	{
 		// Now need to modify old sample, so copy it
 		FBlendSampleData OldSample = *OldIt;
-		bool bNewTargetFound = false;
+		bool bTargetSampleExists = false;
 
 		if (OldSample.PerBoneBlendData.Num()!=PerBoneBlend.Num())
 		{
@@ -401,25 +401,39 @@ bool UBlendSpaceBase::InterpolateWeightOfSampleData(float DeltaTime, const TArra
 				// now interpolate the per bone weights
 				for (int32 Iter = 0; Iter<InterpData.PerBoneBlendData.Num(); ++Iter)
 				{
-					InterpData.PerBoneBlendData[Iter] = FMath::FInterpConstantTo(OldSample.PerBoneBlendData[Iter], NewSample.TotalWeight, DeltaTime, PerBoneBlend[Iter].InterpolationSpeedPerSec);
+					if (PerBoneBlend[Iter].InterpolationSpeedPerSec > 0.f)
+					{
+						InterpData.PerBoneBlendData[Iter] = FMath::FInterpConstantTo(OldSample.PerBoneBlendData[Iter], NewSample.TotalWeight, DeltaTime, PerBoneBlend[Iter].InterpolationSpeedPerSec);
+					}
+					else
+					{
+						InterpData.PerBoneBlendData[Iter] = NewSample.TotalWeight;
+					}
 				}
 
 				FinalSampleDataList.Add(InterpData);
 				TotalFinalWeight += InterpData.GetWeight();
-				bNewTargetFound = true;
+				bTargetSampleExists = true;
 				break;
 			}
 		}
 
 		// if new target isn't found, interpolate to 0.f, this is gone
-		if (bNewTargetFound == false)
+		if (bTargetSampleExists == false)
 		{
 			FBlendSampleData InterpData = OldSample;
 			InterpData.TotalWeight = FMath::FInterpConstantTo(OldSample.TotalWeight, 0.f, DeltaTime, TargetWeightInterpolationSpeedPerSec);
 			// now interpolate the per bone weights
 			for (int32 Iter = 0; Iter<InterpData.PerBoneBlendData.Num(); ++Iter)
 			{
-				InterpData.PerBoneBlendData[Iter] = FMath::FInterpConstantTo(OldSample.PerBoneBlendData[Iter], 0.f, DeltaTime, PerBoneBlend[Iter].InterpolationSpeedPerSec);
+				if (PerBoneBlend[Iter].InterpolationSpeedPerSec > 0.f)
+				{
+					InterpData.PerBoneBlendData[Iter] = FMath::FInterpConstantTo(OldSample.PerBoneBlendData[Iter], 0.f, DeltaTime, PerBoneBlend[Iter].InterpolationSpeedPerSec);
+				}
+				else
+				{
+					InterpData.PerBoneBlendData[Iter] = 0.f;
+				}
 			}
 
 			// add it if it's not zero
@@ -436,7 +450,7 @@ bool UBlendSpaceBase::InterpolateWeightOfSampleData(float DeltaTime, const TArra
 	{
 		// Now need to modify old sample, so copy it
 		FBlendSampleData OldSample = *OldIt;
-		bool bNewTargetFound = false;
+		bool bOldSampleExists = false;
 
 		if (OldSample.PerBoneBlendData.Num()!=PerBoneBlend.Num())
 		{
@@ -449,20 +463,27 @@ bool UBlendSpaceBase::InterpolateWeightOfSampleData(float DeltaTime, const TArra
 			const FBlendSampleData& NewSample = *NewIt;
 			if (NewSample.SampleDataIndex == OldSample.SampleDataIndex)
 			{
-				bNewTargetFound = true;
+				bOldSampleExists = true;
 				break;
 			}
 		}
 
 		// add those new samples
-		if (bNewTargetFound == false)
+		if (bOldSampleExists == false)
 		{
 			FBlendSampleData InterpData = OldSample;
 			InterpData.TotalWeight = FMath::FInterpConstantTo(0.f, OldSample.TotalWeight, DeltaTime, TargetWeightInterpolationSpeedPerSec);
 			// now interpolate the per bone weights
 			for (int32 Iter = 0; Iter<InterpData.PerBoneBlendData.Num(); ++Iter)
 			{
-				InterpData.PerBoneBlendData[Iter] = FMath::FInterpConstantTo(0.f, OldSample.PerBoneBlendData[Iter], DeltaTime, PerBoneBlend[Iter].InterpolationSpeedPerSec);
+				if (PerBoneBlend[Iter].InterpolationSpeedPerSec > 0.f)
+				{
+					InterpData.PerBoneBlendData[Iter] = FMath::FInterpConstantTo(0.f, OldSample.PerBoneBlendData[Iter], DeltaTime, PerBoneBlend[Iter].InterpolationSpeedPerSec);
+				}
+				else
+				{
+					InterpData.PerBoneBlendData[Iter] = OldSample.PerBoneBlendData[Iter];
+				}
 			}
 			FinalSampleDataList.Add(InterpData);
 			TotalFinalWeight += InterpData.GetWeight();
