Explicitely says the texture coordinate id in the TextureCoordinate material expression's name

diff --git a/Engine/Source/Runtime/Engine/Private/Materials/MaterialExpressions.cpp b/Engine/Source/Runtime/Engine/Private/Materials/MaterialExpressions.cpp
--- a/Engine/Source/Runtime/Engine/Private/Materials/MaterialExpressions.cpp
+++ b/Engine/Source/Runtime/Engine/Private/Materials/MaterialExpressions.cpp
@@ -2596,7 +2596,7 @@ int32 UMaterialExpressionTextureCoordinate::Compile(class FMaterialCompiler* Com
 
 void UMaterialExpressionTextureCoordinate::GetCaption(TArray<FString>& OutCaptions) const
 {
-	OutCaptions.Add(TEXT("TexCoord"));
+	OutCaptions.Add(FString::Printf(TEXT("TexCoord[%i]"), CoordinateIndex));
 }
 
 
