Update the comment on who should be using FEngineAnalytics (hint: ONLY EPIC!)

diff --git a/Engine/Source/Runtime/Engine/Public/EngineAnalytics.h b/Engine/Source/Runtime/Engine/Public/EngineAnalytics.h
--- a/Engine/Source/Runtime/Engine/Public/EngineAnalytics.h
+++ b/Engine/Source/Runtime/Engine/Public/EngineAnalytics.h
@@ -10,6 +10,10 @@ class IAnalyticsProvider;
  * For licensee builds, it will be NULL by default unless they provide their own
  * configuration.
  * 
+ * WARNING: This is an analytics provider instance that is created whenever UE4 is launched. 
+ * It is intended ONLY for use by Epic Games. This is NOT intended for games to send 
+ * game-specific telemetry. Create your own provider instance for your game and configure
+ * it independently.
  */
 class FEngineAnalytics : FNoncopyable
 {
