OR-5962 Adding new feature: global pitch shift

- New gameplay static function that will set a global pitch value
- An optional paramter can be given which will interpolate the current global pitch value to the new target value over time
- Non-UI sounds will be pitched
- This is a PC-only implementation, other platforms will follow once this feature is confirmed to be sufficient

diff --git a/Engine/Source/Runtime/Engine/Classes/Kismet/GameplayStatics.h b/Engine/Source/Runtime/Engine/Classes/Kismet/GameplayStatics.h
--- a/Engine/Source/Runtime/Engine/Classes/Kismet/GameplayStatics.h
+++ b/Engine/Source/Runtime/Engine/Classes/Kismet/GameplayStatics.h
@@ -303,6 +303,18 @@ class ENGINE_API UGameplayStatics : public UBlueprintFunctionLibrary
 	UFUNCTION(BlueprintCallable, Category = "Audio", meta = (WorldContext = "WorldContextObject"))
 	static bool AreAnyListenersWithinRange(UObject* WorldContextObject, FVector Location, float MaximumRange);
 	
+
+	/**
+	* Sets a global pitch modulation scalar that will apply to all non-UI sounds
+	*
+	* * Fire and Forget.
+	* * Not Replicated.
+	* @param PitchModulation - A pitch modulation value to globally set.
+	* @param TimeSec - A time value to linearly interpolate the global modulation pitch over from it's current value.
+	*/
+	UFUNCTION(BlueprintCallable, BlueprintCosmetic, Category = "Audio", meta = (WorldContext = "WorldContextObject"))
+	static void SetGlobalPitchModulation(UObject* WorldContextObject, float PitchModulation, float TimeSec);
+
 	/**
 	 * Plays a sound directly with no attenuation, perfect for UI sounds.
 	 *
diff --git a/Engine/Source/Runtime/Engine/Private/Audio.cpp b/Engine/Source/Runtime/Engine/Private/Audio.cpp
--- a/Engine/Source/Runtime/Engine/Private/Audio.cpp
+++ b/Engine/Source/Runtime/Engine/Private/Audio.cpp
@@ -21,6 +21,7 @@
 
 DEFINE_LOG_CATEGORY(LogAudio);
 
+DEFINE_LOG_CATEGORY(LogAudioDebug);
 
 /** Audio stats */
 
diff --git a/Engine/Source/Runtime/Engine/Private/AudioDevice.cpp b/Engine/Source/Runtime/Engine/Private/AudioDevice.cpp
--- a/Engine/Source/Runtime/Engine/Private/AudioDevice.cpp
+++ b/Engine/Source/Runtime/Engine/Private/AudioDevice.cpp
@@ -12,6 +12,53 @@
 #include "Sound/SoundWave.h"
 #include "IAudioExtensionPlugin.h"
 
+/*-----------------------------------------------------------------------------
+FDynamicParameter implementation.
+-----------------------------------------------------------------------------*/
+
+FDynamicParameter::FDynamicParameter(float Value)
+	: CurrValue(Value)
+	, StartValue(Value)
+	, DeltaValue(0.0f)
+	, CurrTimeSec(0.0f)
+	, DurationSec(0.0f)
+	, LastTime(0.0f)
+{}
+
+void FDynamicParameter::Set(float Value, float InDuration)
+{
+	if (InDuration > 0.0f)
+	{
+		DeltaValue = Value - CurrValue;
+		StartValue = CurrValue;
+		DurationSec = InDuration;
+		CurrTimeSec = 0.0f;
+	}
+	else
+	{
+		CurrValue = Value;
+	}
+}
+
+void FDynamicParameter::Update(float DeltaTime)
+{
+	if (DurationSec > 0.0f)
+	{
+		float TimeFraction = CurrTimeSec / DurationSec;
+		if (TimeFraction < 1.0f)
+		{
+			CurrValue = DeltaValue * TimeFraction + StartValue;
+		}
+		else
+		{
+			CurrValue = StartValue + DeltaValue;
+			DurationSec = 0.0f;
+		}
+		CurrTimeSec += DeltaTime;
+	}
+}
+
+
 /*-----------------------------------------------------------------------------
 	FAudioDevice implementation.
 -----------------------------------------------------------------------------*/
@@ -33,7 +80,8 @@ FAudioDevice::FAudioDevice()
 	, TestAudioComponent(nullptr)
 	, DebugState(DEBUGSTATE_None)
 	, TransientMasterVolume(1.0f)
-	, LastUpdateTime(0.0f)
+	, GlobalPitchScale(1.0f)
+	, LastUpdateTime(FPlatformTime::Seconds())
 	, NextResourceID(1)
 	, BaseSoundMix(nullptr)
 	, DefaultBaseSoundMix(nullptr)
@@ -1924,6 +1972,15 @@ void FAudioDevice::Update( bool bGameTicking )
 {
 	SCOPE_CYCLE_COUNTER( STAT_AudioUpdateTime );
 
+	double CurrTime = FPlatformTime::Seconds();
+	double DeltaTime = CurrTime - LastUpdateTime;
+	LastUpdateTime = CurrTime;
+
+	if (bGameTicking)
+	{
+		GlobalPitchScale.Update(DeltaTime);
+	}
+
 	// Start a new frame
 	CurrentTick++;
 
diff --git a/Engine/Source/Runtime/Engine/Private/GameplayStatics.cpp b/Engine/Source/Runtime/Engine/Private/GameplayStatics.cpp
--- a/Engine/Source/Runtime/Engine/Private/GameplayStatics.cpp
+++ b/Engine/Source/Runtime/Engine/Private/GameplayStatics.cpp
@@ -762,6 +762,26 @@ bool UGameplayStatics::AreAnyListenersWithinRange(UObject* WorldContextObject, F
 	return false;
 }
 
+void UGameplayStatics::SetGlobalPitchModulation(UObject* WorldContextObject, float PitchModulation, float TimeSec)
+{
+	if (!GEngine || !GEngine->UseSound())
+	{
+		return;
+	}
+
+	UWorld* ThisWorld = GEngine->GetWorldFromContextObject(WorldContextObject);
+	if (!ThisWorld || !ThisWorld->bAllowAudioPlayback || ThisWorld->GetNetMode() == NM_DedicatedServer)
+	{
+		return;
+	}
+
+	if (FAudioDevice* AudioDevice = ThisWorld->GetAudioDevice())
+	{
+		AudioDevice->GlobalPitchScale.Set(PitchModulation, TimeSec);
+	}
+}
+
+
 void UGameplayStatics::PlaySound2D(UObject* WorldContextObject, class USoundBase* Sound, float VolumeMultiplier, float PitchMultiplier, float StartTime)
 {
 	if (!Sound || !GEngine || !GEngine->UseSound())
diff --git a/Engine/Source/Runtime/Engine/Public/Audio.h b/Engine/Source/Runtime/Engine/Public/Audio.h
--- a/Engine/Source/Runtime/Engine/Public/Audio.h
+++ b/Engine/Source/Runtime/Engine/Public/Audio.h
@@ -11,6 +11,9 @@
 
 ENGINE_API DECLARE_LOG_CATEGORY_EXTERN(LogAudio, Warning, All);
 
+// Special log category used for temporary programmer debugging code of audio
+ENGINE_API DECLARE_LOG_CATEGORY_EXTERN(LogAudioDebug, Display, All);
+
 /** 
  * Maximum number of channels that can be set using the ini setting
  */
diff --git a/Engine/Source/Runtime/Engine/Public/AudioDevice.h b/Engine/Source/Runtime/Engine/Public/AudioDevice.h
--- a/Engine/Source/Runtime/Engine/Public/AudioDevice.h
+++ b/Engine/Source/Runtime/Engine/Public/AudioDevice.h
@@ -59,6 +59,29 @@ namespace ESortedActiveWaveGetType
 	};
 }
 
+/** Simple class that wraps the math involved with interpolating a parameter over time based on audio device update time. */
+class FDynamicParameter
+{
+public:
+	FDynamicParameter(float Value);
+
+	void Set(float Value, float InDuration);
+	void Update(float DeltaTime);
+	float GetValue() const
+	{
+		return CurrValue;
+	}
+
+private:
+	float CurrValue;
+	float StartValue;
+	float DeltaValue;
+	float CurrTimeSec;
+	float DurationSec;
+	float LastTime;
+};
+
+
 /** 
  * Defines the properties of the listener
  */
@@ -756,8 +779,11 @@ public:
 	/** transient master volume multiplier that can be modified at runtime without affecting user settings automatically reset to 1.0 on level change */
 	float TransientMasterVolume;
 
+	/** Global dynamic pitch scale parameter */
+	FDynamicParameter GlobalPitchScale;
+
 	/** Timestamp of the last update */
-	float LastUpdateTime;
+	double LastUpdateTime;
 
 	/** Next resource ID to assign out to a wave/buffer */
 	int32 NextResourceID;
diff --git a/Engine/Source/Runtime/Windows/XAudio2/Private/XAudio2Source.cpp b/Engine/Source/Runtime/Windows/XAudio2/Private/XAudio2Source.cpp
--- a/Engine/Source/Runtime/Windows/XAudio2/Private/XAudio2Source.cpp
+++ b/Engine/Source/Runtime/Windows/XAudio2/Private/XAudio2Source.cpp
@@ -1235,7 +1235,16 @@ void FXAudio2SoundSource::Update( void )
 		return;
 	}
 
-	const float Pitch = FMath::Clamp<float>( WaveInstance->Pitch, MIN_PITCH, MAX_PITCH );
+	float Pitch = WaveInstance->Pitch;
+
+	// Don't apply global pitch scale to UI sounds
+	if (!WaveInstance->bIsUISound)
+	{
+		Pitch *= AudioDevice->GlobalPitchScale.GetValue();
+	}
+
+	Pitch = FMath::Clamp<float>(Pitch, MIN_PITCH, MAX_PITCH);
+
 	AudioDevice->ValidateAPICall( TEXT( "SetFrequencyRatio" ), 
 		Source->SetFrequencyRatio( Pitch ) );
 
