Moving private latent commands into public area for use by other modules

diff --git a/Engine/Source/Runtime/Engine/Private/Tests/AutomationCommon.cpp b/Engine/Source/Runtime/Engine/Private/Tests/AutomationCommon.cpp
--- a/Engine/Source/Runtime/Engine/Private/Tests/AutomationCommon.cpp
+++ b/Engine/Source/Runtime/Engine/Private/Tests/AutomationCommon.cpp
@@ -5,6 +5,10 @@
 #include "AutomationCommon.h"
 #include "ImageUtils.h"
 
+#include "Matinee/MatineeActor.h"
+
+DEFINE_LOG_CATEGORY_STATIC(LogEngineAutomationLatentCommand, Log, All);
+
 ///////////////////////////////////////////////////////////////////////
 // Common Latent commands
 
@@ -59,4 +63,125 @@ bool FRequestExitCommand::Update()
 {
 	GIsRequestingExit = true;
 	return true;
+}
+
+
+///////////////////////////////////////////////////////////////////////
+// Common Latent commands which are used across test type. I.e. Engine, Network, etc...
+
+
+bool FPlayMatineeLatentCommand::Update()
+{
+	if (MatineeActor)
+	{
+		UE_LOG(LogEngineAutomationLatentCommand, Log, TEXT("Triggering the matinee named: '%s'"), *MatineeActor->GetName())
+
+			//force this matinee to not be looping so it doesn't infinitely loop
+			MatineeActor->bLooping = false;
+		MatineeActor->Play();
+	}
+	return true;
+}
+
+
+bool FWaitForMatineeToCompleteLatentCommand::Update()
+{
+	bool bTestComplete = true;
+	if (MatineeActor)
+	{
+		bTestComplete = !MatineeActor->bIsPlaying;
+	}
+	return bTestComplete;
+}
+
+
+bool FExecStringLatentCommand::Update()
+{
+	UE_LOG(LogEngineAutomationLatentCommand, Log, TEXT("Executing the console command: '%s'"), *ExecCommand);
+
+	if (GEngine->GameViewport)
+	{
+		GEngine->GameViewport->Exec(NULL, *ExecCommand, *GLog);
+	}
+	else
+	{
+		GEngine->Exec(NULL, *ExecCommand);
+	}
+	return true;
+}
+
+
+bool FEngineWaitLatentCommand::Update()
+{
+	float NewTime = FPlatformTime::Seconds();
+	if (NewTime - StartTime >= Duration)
+	{
+		return true;
+	}
+	return false;
+}
+
+
+bool FEnqueuePerformanceCaptureCommands::Update()
+{
+	//for every matinee actor in the level
+	for (TObjectIterator<AMatineeActor> It; It; ++It)
+	{
+		AMatineeActor* MatineeActor = *It;
+		if (MatineeActor && MatineeActor->GetName().Contains(TEXT("Automation")))
+		{
+			//add latent action to execute this matinee
+			ADD_LATENT_AUTOMATION_COMMAND(FPlayMatineeLatentCommand(MatineeActor));
+
+			//add action to wait until matinee is complete
+			ADD_LATENT_AUTOMATION_COMMAND(FWaitForMatineeToCompleteLatentCommand(MatineeActor));
+		}
+	}
+
+	return true;
+}
+
+
+bool FMatineePerformanceCaptureCommand::Update()
+{
+	//for every matinee actor in the level
+	for (TObjectIterator<AMatineeActor> It; It; ++It)
+	{
+		AMatineeActor* MatineeActor = *It;
+		FString MatineeFOOName = MatineeActor->GetName();
+		if (MatineeActor->GetName().Equals(MatineeName, ESearchCase::IgnoreCase))
+		{
+
+
+			//add latent action to execute this matinee
+			ADD_LATENT_AUTOMATION_COMMAND(FPlayMatineeLatentCommand(MatineeActor));
+
+			//Run the Stat FPS Chart command
+			ADD_LATENT_AUTOMATION_COMMAND(FExecWorldStringLatentCommand(TEXT("StartFPSChart")));
+
+			//add action to wait until matinee is complete
+			ADD_LATENT_AUTOMATION_COMMAND(FWaitForMatineeToCompleteLatentCommand(MatineeActor));
+
+			//Stop the Stat FPS Chart command
+			ADD_LATENT_AUTOMATION_COMMAND(FExecWorldStringLatentCommand(TEXT("StopFPSChart")));
+		}
+		else
+		{
+			UE_LOG(LogEngineAutomationLatentCommand, Log, TEXT("'%s' is not the matinee name that is being searched for."), *MatineeActor->GetName())
+		}
+	}
+
+	return true;
+
+}
+
+
+bool FExecWorldStringLatentCommand::Update()
+{
+	check(GEngine->GetWorldContexts().Num() == 1);
+	check(GEngine->GetWorldContexts()[0].WorldType == EWorldType::Game);
+
+	UE_LOG(LogEngineAutomationLatentCommand, Log, TEXT("Running Exec Command. '%s'"), *ExecCommand);
+	GEngine->Exec(GEngine->GetWorldContexts()[0].World(), *ExecCommand);
+	return true;
 }
\ No newline at end of file
diff --git a/Engine/Source/Runtime/Engine/Private/Tests/AutomationTestCommon.cpp b/Engine/Source/Runtime/Engine/Private/Tests/AutomationTestCommon.cpp
--- a/Engine/Source/Runtime/Engine/Private/Tests/AutomationTestCommon.cpp
+++ b/Engine/Source/Runtime/Engine/Private/Tests/AutomationTestCommon.cpp
@@ -1,127 +1,4 @@
 // Copyright 1998-2015 Epic Games, Inc. All Rights Reserved.
 
 #include "EnginePrivate.h"
-#include "Matinee/MatineeActor.h"
 #include "AutomationTestCommon.h"
-
-DEFINE_LOG_CATEGORY_STATIC(LogEngineAutomationLatentCommand, Log, All);
-
-///////////////////////////////////////////////////////////////////////
-// Common Latent commands which are used across test type. I.e. Engine, Network, etc...
-
-
-bool FPlayMatineeLatentCommand::Update() 
-{
-	if (MatineeActor)
-	{
-		UE_LOG(LogEngineAutomationLatentCommand, Log, TEXT("Triggering the matinee named: '%s'"), *MatineeActor->GetName())
-		
-		//force this matinee to not be looping so it doesn't infinitely loop
-		MatineeActor->bLooping = false;
-		MatineeActor->Play();
-	}
-	return true;
-}
-
-
-bool FWaitForMatineeToCompleteLatentCommand::Update()
-{
-	bool bTestComplete = true;
-	if (MatineeActor)
-	{
-		bTestComplete = !MatineeActor->bIsPlaying;
-	}
-	return bTestComplete;
-}
-
-
-bool FExecStringLatentCommand::Update()
-{
-	UE_LOG(LogEngineAutomationLatentCommand, Log, TEXT("Executing the console command: '%s'"), *ExecCommand);
-
-	if (GEngine->GameViewport)
-	{
-		GEngine->GameViewport->Exec( NULL, *ExecCommand, *GLog);
-	}
-	else
-	{
-		GEngine->Exec( NULL, *ExecCommand);
-	}
-	return true;
-}
-
-
-bool FEngineWaitLatentCommand::Update()
-{
-	float NewTime = FPlatformTime::Seconds();
-	if (NewTime - StartTime >= Duration)
-	{
-		return true;
-	}
-	return false;
-}
-
-
-bool FEnqueuePerformanceCaptureCommands::Update()
-{
-	//for every matinee actor in the level
-	for (TObjectIterator<AMatineeActor> It; It; ++It)
-	{
-		AMatineeActor* MatineeActor = *It;
-		if(MatineeActor && MatineeActor->GetName().Contains(TEXT("Automation")) )
-		{
-			//add latent action to execute this matinee
-			ADD_LATENT_AUTOMATION_COMMAND(FPlayMatineeLatentCommand(MatineeActor));
-
-			//add action to wait until matinee is complete
-			ADD_LATENT_AUTOMATION_COMMAND(FWaitForMatineeToCompleteLatentCommand(MatineeActor));
-		}
-	}
-
-	return true;
-}
-
-
-bool FMatineePerformanceCaptureCommand::Update()
-{
-	//for every matinee actor in the level
-	for (TObjectIterator<AMatineeActor> It; It; ++It)
-	{
-		AMatineeActor* MatineeActor = *It;
-		FString MatineeFOOName = MatineeActor->GetName();
-		if (MatineeActor->GetName().Equals(MatineeName,ESearchCase::IgnoreCase))
-		{	
-
-
-			//add latent action to execute this matinee
-			ADD_LATENT_AUTOMATION_COMMAND(FPlayMatineeLatentCommand(MatineeActor));
-
-			//Run the Stat FPS Chart command
-			ADD_LATENT_AUTOMATION_COMMAND(FExecWorldStringLatentCommand(TEXT("StartFPSChart")));
-
-			//add action to wait until matinee is complete
-			ADD_LATENT_AUTOMATION_COMMAND(FWaitForMatineeToCompleteLatentCommand(MatineeActor));
-
-			//Stop the Stat FPS Chart command
-			ADD_LATENT_AUTOMATION_COMMAND(FExecWorldStringLatentCommand(TEXT("StopFPSChart")));
-		}
-		else
-		{
-			UE_LOG(LogEngineAutomationLatentCommand, Log, TEXT("'%s' is not the matinee name that is being searched for."), *MatineeActor->GetName())
-		}
-	}
-
-	return true;
-
-}
-
-
-bool FExecWorldStringLatentCommand::Update()
-{
-	check(GEngine->GetWorldContexts().Num() == 1);
-	check(GEngine->GetWorldContexts()[0].WorldType == EWorldType::Game);
-
-	UE_LOG(LogEngineAutomationLatentCommand, Log, TEXT("Running Exec Command. '%s'"), *ExecCommand);
-	GEngine->Exec(GEngine->GetWorldContexts()[0].World(), *ExecCommand);
-	return true;
-}
\ No newline at end of file
diff --git a/Engine/Source/Runtime/Engine/Private/Tests/AutomationTestCommon.h b/Engine/Source/Runtime/Engine/Private/Tests/AutomationTestCommon.h
--- a/Engine/Source/Runtime/Engine/Private/Tests/AutomationTestCommon.h
+++ b/Engine/Source/Runtime/Engine/Private/Tests/AutomationTestCommon.h
@@ -8,43 +8,3 @@
 ///////////////////////////////////////////////////////////////////////
 // Common Latent commands which are used across test type. I.e. Engine, Network, etc...
 
-
-/**
- * Force a matinee to not loop and request that it play
- */
-DEFINE_LATENT_AUTOMATION_COMMAND_ONE_PARAMETER(FPlayMatineeLatentCommand, AMatineeActor*, MatineeActor);
-
-
-/**
- * Wait for a particular matinee actor to finish playing
- */
-DEFINE_LATENT_AUTOMATION_COMMAND_ONE_PARAMETER(FWaitForMatineeToCompleteLatentCommand, AMatineeActor*, MatineeActor);
-
-
-/**
- * Execute command string
- */
-DEFINE_LATENT_AUTOMATION_COMMAND_ONE_PARAMETER(FExecStringLatentCommand, FString, ExecCommand);
-
-
-/**
- * Wait for the given amount of time
- */
-DEFINE_LATENT_AUTOMATION_COMMAND_ONE_PARAMETER(FEngineWaitLatentCommand, float, Duration);
-
-
-/**
- * Enqueue performance capture commands after a map has been loaded
- */
-DEFINE_LATENT_AUTOMATION_COMMAND(FEnqueuePerformanceCaptureCommands);
-
-
-/**
-* Run FPS chart command using for the actual duration of the matinee.
-*/
-DEFINE_LATENT_AUTOMATION_COMMAND_ONE_PARAMETER(FMatineePerformanceCaptureCommand, FString, MatineeName);
-
-/**
-* Latent command to run an exec command that also requires a UWorld.
-*/
-DEFINE_LATENT_AUTOMATION_COMMAND_ONE_PARAMETER(FExecWorldStringLatentCommand, FString, ExecCommand);
\ No newline at end of file
diff --git a/Engine/Source/Runtime/Engine/Private/Tests/NetworkAutomationTests.cpp b/Engine/Source/Runtime/Engine/Private/Tests/NetworkAutomationTests.cpp
--- a/Engine/Source/Runtime/Engine/Private/Tests/NetworkAutomationTests.cpp
+++ b/Engine/Source/Runtime/Engine/Private/Tests/NetworkAutomationTests.cpp
@@ -2,7 +2,7 @@
 
 
 #include "EnginePrivate.h"
-#include "AutomationTestCommon.h"
+#include "AutomationCommon.h"
 
 DEFINE_LOG_CATEGORY_STATIC(LogNetworkAutomationTests, Log, All);
 
diff --git a/Engine/Source/Runtime/Engine/Private/Tests/OSSTests.cpp b/Engine/Source/Runtime/Engine/Private/Tests/OSSTests.cpp
--- a/Engine/Source/Runtime/Engine/Private/Tests/OSSTests.cpp
+++ b/Engine/Source/Runtime/Engine/Private/Tests/OSSTests.cpp
@@ -4,7 +4,7 @@
 #if WITH_EDITOR
 #include "UnrealEd.h"
 
-#include "AutomationTestCommon.h"
+#include "AutomationCommon.h"
 
 DEFINE_LOG_CATEGORY_STATIC(LogHackAutomationTests, Log, All);
 
diff --git a/Engine/Source/Runtime/Engine/Public/Tests/AutomationCommon.h b/Engine/Source/Runtime/Engine/Public/Tests/AutomationCommon.h
--- a/Engine/Source/Runtime/Engine/Public/Tests/AutomationCommon.h
+++ b/Engine/Source/Runtime/Engine/Public/Tests/AutomationCommon.h
@@ -115,3 +115,42 @@ DEFINE_LATENT_AUTOMATION_COMMAND_ONE_PARAMETER(FLoadGameMapCommand, FString, Map
 DEFINE_ENGINE_LATENT_AUTOMATION_COMMAND( FRequestExitCommand );
 
 
+/**
+* Force a matinee to not loop and request that it play
+*/
+DEFINE_LATENT_AUTOMATION_COMMAND_ONE_PARAMETER(FPlayMatineeLatentCommand, AMatineeActor*, MatineeActor);
+
+
+/**
+* Wait for a particular matinee actor to finish playing
+*/
+DEFINE_LATENT_AUTOMATION_COMMAND_ONE_PARAMETER(FWaitForMatineeToCompleteLatentCommand, AMatineeActor*, MatineeActor);
+
+
+/**
+* Execute command string
+*/
+DEFINE_LATENT_AUTOMATION_COMMAND_ONE_PARAMETER(FExecStringLatentCommand, FString, ExecCommand);
+
+
+/**
+* Wait for the given amount of time
+*/
+DEFINE_LATENT_AUTOMATION_COMMAND_ONE_PARAMETER(FEngineWaitLatentCommand, float, Duration);
+
+
+/**
+* Enqueue performance capture commands after a map has been loaded
+*/
+DEFINE_ENGINE_LATENT_AUTOMATION_COMMAND(FEnqueuePerformanceCaptureCommands);
+
+
+/**
+* Run FPS chart command using for the actual duration of the matinee.
+*/
+DEFINE_LATENT_AUTOMATION_COMMAND_ONE_PARAMETER(FMatineePerformanceCaptureCommand, FString, MatineeName);
+
+/**
+* Latent command to run an exec command that also requires a UWorld.
+*/
+DEFINE_LATENT_AUTOMATION_COMMAND_ONE_PARAMETER(FExecWorldStringLatentCommand, FString, ExecCommand);
