Engine: Improved server log message for a failed RPC when the actor's level isn't ready on a client
[INTEGRATE] Change 2661203 by [REDACTED_PERSON][REDACTED_HANDLE] on 2015/08/19 13:11:24

diff --git a/Engine/Source/Runtime/Engine/Private/NetworkDriver.cpp b/Engine/Source/Runtime/Engine/Private/NetworkDriver.cpp
--- a/Engine/Source/Runtime/Engine/Private/NetworkDriver.cpp
+++ b/Engine/Source/Runtime/Engine/Private/NetworkDriver.cpp
@@ -806,7 +806,7 @@ void UNetDriver::InternalProcessRemoteFunction
 			}
 			else
 			{
-				UE_LOG(LogNet, Verbose, TEXT("Can't send function '%s' on '%s': Client hasn't loaded the level for this Actor"), *Function->GetName(), *Actor->GetName());
+				UE_LOG(LogNet, Verbose, TEXT("Can't send function '%s' on actor '%s' because client hasn't loaded the level '%s' containing it"), *Function->GetName(), *Actor->GetName(), *Actor->GetLevel()->GetName());
 				return;
 			}
 		}
