Added accessor for reading all neighbors of navmesh poly with their edge data. Allows custom implementations of navmesh queries and will be required for preparing flow field data, since amount of data exceeds pathfinding's node pool size.

#ue4 (merged from dev branch)

diff --git a/Engine/Source/Runtime/Engine/Classes/AI/Navigation/RecastNavMesh.h b/Engine/Source/Runtime/Engine/Classes/AI/Navigation/RecastNavMesh.h
--- a/Engine/Source/Runtime/Engine/Classes/AI/Navigation/RecastNavMesh.h
+++ b/Engine/Source/Runtime/Engine/Classes/AI/Navigation/RecastNavMesh.h
@@ -880,6 +880,9 @@ public:
 	bool GetPolyFlags(NavNodeRef PolyID, uint16& PolyFlags, uint16& AreaFlags) const;
 	bool GetPolyFlags(NavNodeRef PolyID, FNavMeshNodeFlags& Flags) const;
 
+	/** Finds all polys connected with specified one */
+	bool GetPolyNeighbors(NavNodeRef PolyID, TArray<FNavigationPortalEdge>& Neighbors) const;
+
 	/** Finds closest point constrained to given poly */
 	bool GetClosestPointOnPoly(NavNodeRef PolyID, const FVector& TestPt, FVector& PointOnPoly) const;
 
diff --git a/Engine/Source/Runtime/Engine/Private/AI/Navigation/PImplRecastNavMesh.cpp b/Engine/Source/Runtime/Engine/Private/AI/Navigation/PImplRecastNavMesh.cpp
--- a/Engine/Source/Runtime/Engine/Private/AI/Navigation/PImplRecastNavMesh.cpp
+++ b/Engine/Source/Runtime/Engine/Private/AI/Navigation/PImplRecastNavMesh.cpp
@@ -1708,6 +1708,47 @@ bool FPImplRecastNavMesh::GetPolyData(NavNodeRef PolyID, uint16& Flags, uint8& A
 	return false;
 }
 
+bool FPImplRecastNavMesh::GetPolyNeighbors(NavNodeRef PolyID, TArray<FNavigationPortalEdge>& Neighbors) const
+{
+	if (DetourNavMesh)
+	{
+		dtPolyRef PolyRef = (dtPolyRef)PolyID;
+		dtPoly const* Poly = 0;
+		dtMeshTile const* Tile = 0;
+
+		dtStatus Status = DetourNavMesh->getTileAndPolyByRef(PolyRef, &Tile, &Poly);
+		if (dtStatusSucceed(Status))
+		{
+			INITIALIZE_NAVQUERY_SIMPLE(NavQuery, RECAST_MAX_SEARCH_NODES);
+
+			float RcLeft[3], RcRight[3];
+			uint8 DummyType1, DummyType2;
+
+			uint32 LinkIdx = Poly->firstLink;
+			while (LinkIdx != DT_NULL_LINK)
+			{
+				const dtLink& Link = DetourNavMesh->getLink(Tile, LinkIdx);
+				LinkIdx = Link.next;
+				
+				Status = NavQuery.getPortalPoints(PolyRef, Link.ref, RcLeft, RcRight, DummyType1, DummyType2);
+				if (dtStatusSucceed(Status))
+				{
+					FNavigationPortalEdge NeiData;
+					NeiData.ToRef = Link.ref;
+					NeiData.Left = Recast2UnrealPoint(RcLeft);
+					NeiData.Right = Recast2UnrealPoint(RcRight);
+
+					Neighbors.Add(NeiData);
+				}
+			}
+
+			return true;
+		}
+	}
+
+	return false;
+}
+
 bool FPImplRecastNavMesh::GetPolyTileIndex(NavNodeRef PolyID, uint32& PolyIndex, uint32& TileIndex) const
 {
 	if (DetourNavMesh && PolyID)
diff --git a/Engine/Source/Runtime/Engine/Private/AI/Navigation/RecastNavMesh.cpp b/Engine/Source/Runtime/Engine/Private/AI/Navigation/RecastNavMesh.cpp
--- a/Engine/Source/Runtime/Engine/Private/AI/Navigation/RecastNavMesh.cpp
+++ b/Engine/Source/Runtime/Engine/Private/AI/Navigation/RecastNavMesh.cpp
@@ -1325,6 +1325,17 @@ bool ARecastNavMesh::GetPolyFlags(NavNodeRef PolyID, FNavMeshNodeFlags& Flags) c
 	return bFound;
 }
 
+bool ARecastNavMesh::GetPolyNeighbors(NavNodeRef PolyID, TArray<FNavigationPortalEdge>& Neighbors) const
+{
+	bool bFound = false;
+	if (RecastNavMeshImpl)
+	{
+		bFound = RecastNavMeshImpl->GetPolyNeighbors(PolyID, Neighbors);
+	}
+
+	return bFound;
+}
+
 bool ARecastNavMesh::GetClosestPointOnPoly(NavNodeRef PolyID, const FVector& TestPt, FVector& PointOnPoly) const
 {
 	bool bFound = false;
diff --git a/Engine/Source/Runtime/Engine/Public/AI/Navigation/PImplRecastNavMesh.h b/Engine/Source/Runtime/Engine/Public/AI/Navigation/PImplRecastNavMesh.h
--- a/Engine/Source/Runtime/Engine/Public/AI/Navigation/PImplRecastNavMesh.h
+++ b/Engine/Source/Runtime/Engine/Public/AI/Navigation/PImplRecastNavMesh.h
@@ -173,6 +173,8 @@ public:
 	bool GetPolyData(NavNodeRef PolyID, uint16& Flags, uint8& AreaType) const;
 	/** Retrieves area ID for the specified polygon. */
 	uint32 GetPolyAreaID(NavNodeRef PolyID) const;
+	/** Finds all polys connected with specified one */
+	bool GetPolyNeighbors(NavNodeRef PolyID, TArray<FNavigationPortalEdge>& Neighbors) const;
 	/** Finds closest point constrained to given poly */
 	bool GetClosestPointOnPoly(NavNodeRef PolyID, const FVector& TestPt, FVector& PointOnPoly) const;
 	/** Decode poly ID into tile index and poly index */
