Adding 'debug stackoverflow' console command

diff --git a/Engine/Source/Runtime/Engine/Private/UnrealEngine.cpp b/Engine/Source/Runtime/Engine/Private/UnrealEngine.cpp
--- a/Engine/Source/Runtime/Engine/Private/UnrealEngine.cpp
+++ b/Engine/Source/Runtime/Engine/Private/UnrealEngine.cpp
@@ -4611,6 +4611,18 @@ bool UEngine::HandleMemCommand( const TCHAR* Cmd, FOutputDevice& Ar )
 // debug flag to allocate memory every frame, to trigger an OOM condition
 static bool GDebugAllocMemEveryFrame = false;
 
+/** Helper function to cause a stack overflow crash */
+FORCENOINLINE static void StackOverflowFunction(int32* DummyArg)
+{
+	int32 StackArray[8196];
+	FMemory::Memset(StackArray, 0, sizeof(StackArray));
+	if (StackArray[0] == 0)
+	{
+		UE_LOG(LogEngine, VeryVerbose, TEXT("StackOverflowFunction(%d)"), DummyArg ? DummyArg[0] : 0);
+		StackOverflowFunction(StackArray);
+	}
+}
+
 bool UEngine::HandleDebugCommand( const TCHAR* Cmd, FOutputDevice& Ar )
 {
 	if( FParse::Command(&Cmd,TEXT("RENDERCRASH")) )
@@ -4875,6 +4887,12 @@ bool UEngine::HandleDebugCommand( const TCHAR* Cmd, FOutputDevice& Ar )
 		GDebugAllocMemEveryFrame = true;
 		return true;
 	}
+	else if (FParse::Command(&Cmd, TEXT("STACKOVERFLOW")))
+	{
+		Ar.Log(TEXT("Inifnite recursion to cause stack overflow"));
+		StackOverflowFunction(nullptr);
+		return true;
+	}
 
 	return false;
 }
