UE4 - Added a GUID to the temp directory for the shader compiler.

diff --git a/Engine/Source/Runtime/Core/Private/GenericPlatform/GenericPlatformProcess.cpp b/Engine/Source/Runtime/Core/Private/GenericPlatform/GenericPlatformProcess.cpp
--- a/Engine/Source/Runtime/Core/Private/GenericPlatform/GenericPlatformProcess.cpp
+++ b/Engine/Source/Runtime/Core/Private/GenericPlatform/GenericPlatformProcess.cpp
@@ -123,7 +123,7 @@ void FGenericPlatformProcess::SetShaderDir(const TCHAR*Where)
  */
 const FString FGenericPlatformProcess::ShaderWorkingDir()
 {
-	return (FPaths::GameIntermediateDir() / TEXT("Shaders/WorkingDirectory/"));
+	return (FPaths::GameIntermediateDir() / TEXT("Shaders/tmp/"));
 }
 
 /**
@@ -133,8 +133,10 @@ void FGenericPlatformProcess::CleanShaderWorkingDir()
 {
 	// Path to the working directory where files are written for multi-threaded compilation
 	FString ShaderWorkingDirectory = ShaderWorkingDir();
-
 	IFileManager::Get().DeleteDirectory(*ShaderWorkingDirectory, false, true);
+
+	FString LegacyShaderWorkingDirectory = FPaths::GameIntermediateDir() / TEXT("Shaders/WorkingDirectory/");
+	IFileManager::Get().DeleteDirectory(*LegacyShaderWorkingDirectory, false, true);
 }
 
 const TCHAR* FGenericPlatformProcess::ExecutableName(bool bRemoveExtension)
diff --git a/Engine/Source/Runtime/Engine/Private/ShaderCompiler/ShaderCompiler.cpp b/Engine/Source/Runtime/Engine/Private/ShaderCompiler/ShaderCompiler.cpp
--- a/Engine/Source/Runtime/Engine/Private/ShaderCompiler/ShaderCompiler.cpp
+++ b/Engine/Source/Runtime/Engine/Private/ShaderCompiler/ShaderCompiler.cpp
@@ -844,7 +844,16 @@ FShaderCompilingManager::FShaderCompilingManager() :
 	// Use a working directory unique to this game, process and thread so that it will not conflict 
 	// With processes from other games, processes from the same game or threads in this same process.
 	// Use IFileManager to do path conversion to properly handle sandbox paths (outside of standard paths in particular).
-	ShaderBaseWorkingDirectory = FPlatformProcess::ShaderWorkingDir() / FString::FromInt(ProcessId) + TEXT("/");
+	//ShaderBaseWorkingDirectory = FPlatformProcess::ShaderWorkingDir() / FString::FromInt(ProcessId) + TEXT("/");
+
+	{
+		FGuid Guid;
+		Guid = FGuid::NewGuid();
+		FString LegacyShaderWorkingDirectory = FPaths::GameIntermediateDir() / TEXT("Shaders/WorkingDirectory/")  / FString::FromInt(ProcessId) + TEXT("/");
+		ShaderBaseWorkingDirectory = FPlatformProcess::ShaderWorkingDir() / *Guid.ToString(EGuidFormats::Digits) + TEXT("/");
+		UE_LOG(LogShaderCompilers, Log, TEXT("Guid format shader working directory is %d characters bigger than the processId version (%s)."), ShaderBaseWorkingDirectory.Len() - LegacyShaderWorkingDirectory.Len(), *LegacyShaderWorkingDirectory );
+	}
+
 	if (!IFileManager::Get().DeleteDirectory(*ShaderBaseWorkingDirectory, false, true))
 	{
 		UE_LOG(LogShaderCompilers, Fatal, TEXT("Could not delete the shader compiler working directory '%s'."), *ShaderBaseWorkingDirectory);
