Optimize FAudioDevice::FindActiveSound by eliminating large numbers of Weak Object Ptr comparisons
Privatize FActiveSound::AudioComponent. Use Get/SetAudioComponent instead.

diff --git a/Engine/Source/Editor/GraphEditor/Private/SoundCueGraphConnectionDrawingPolicy.cpp b/Engine/Source/Editor/GraphEditor/Private/SoundCueGraphConnectionDrawingPolicy.cpp
--- a/Engine/Source/Editor/GraphEditor/Private/SoundCueGraphConnectionDrawingPolicy.cpp
+++ b/Engine/Source/Editor/GraphEditor/Private/SoundCueGraphConnectionDrawingPolicy.cpp
@@ -58,7 +58,7 @@ void FSoundCueGraphConnectionDrawingPolicy::BuildAudioFlowRoadmap()
 
 			for (int32 WaveIndex = WaveInstances.Num() - 1; WaveIndex >= 0 ; --WaveIndex)
 			{
-				UAudioComponent* WaveInstanceAudioComponent = WaveInstances[WaveIndex]->ActiveSound->AudioComponent.Get();
+				UAudioComponent* WaveInstanceAudioComponent = WaveInstances[WaveIndex]->ActiveSound->GetAudioComponent();
 				if (WaveInstanceAudioComponent != PreviewAudioComponent)
 				{
 					WaveInstances.RemoveAtSwap(WaveIndex);
diff --git a/Engine/Source/Runtime/Engine/Private/ActiveSound.cpp b/Engine/Source/Runtime/Engine/Private/ActiveSound.cpp
--- a/Engine/Source/Runtime/Engine/Private/ActiveSound.cpp
+++ b/Engine/Source/Runtime/Engine/Private/ActiveSound.cpp
@@ -12,6 +12,7 @@ FActiveSound::FActiveSound()
 	: Sound(NULL)
 	, World(NULL)
 	, AudioComponent(NULL)
+	, AudioComponentIndex(0)
 	, SoundClassOverride(NULL)
 	, bOccluded(false)
 	, bAllowSpatialization(true)
diff --git a/Engine/Source/Runtime/Engine/Private/Audio.cpp b/Engine/Source/Runtime/Engine/Private/Audio.cpp
--- a/Engine/Source/Runtime/Engine/Private/Audio.cpp
+++ b/Engine/Source/Runtime/Engine/Private/Audio.cpp
@@ -175,9 +175,10 @@ FString FSoundSource::Describe(bool bUseLongName)
 	AActor* SoundOwner = NULL;
 	
 	// TODO - Audio Threading. This won't work cross thread.
-	if (WaveInstance->ActiveSound && WaveInstance->ActiveSound->AudioComponent.IsValid())
+	UAudioComponent* AudioComponent = (WaveInstance->ActiveSound ? WaveInstance->ActiveSound->GetAudioComponent() : nullptr);
+	if (AudioComponent)
 	{
-		SoundOwner = WaveInstance->ActiveSound->AudioComponent->GetOwner();
+		SoundOwner = AudioComponent->GetOwner();
 	}
 
 	return FString::Printf(TEXT("Wave: %s, Volume: %6.2f, Owner: %s"), 
diff --git a/Engine/Source/Runtime/Engine/Private/AudioDevice.cpp b/Engine/Source/Runtime/Engine/Private/AudioDevice.cpp
--- a/Engine/Source/Runtime/Engine/Private/AudioDevice.cpp
+++ b/Engine/Source/Runtime/Engine/Private/AudioDevice.cpp
@@ -359,7 +359,8 @@ bool FAudioDevice::HandleListWavesCommand( const TCHAR* Cmd, FOutputDevice& Ar )
 	{
 		FWaveInstance* WaveInstance = WaveInstances[ InstanceIndex ];
 		FSoundSource* Source = WaveInstanceSourceMap.FindRef( WaveInstance );
-		AActor* SoundOwner = WaveInstance->ActiveSound->AudioComponent.IsValid() ? WaveInstance->ActiveSound->AudioComponent->GetOwner() : NULL;
+		UAudioComponent* AudioComponent = WaveInstance->ActiveSound->GetAudioComponent();
+		AActor* SoundOwner = AudioComponent ? AudioComponent->GetOwner() : nullptr;
 		Ar.Logf( TEXT( "%4i.    %s %6.2f %6.2f  %s   %s"), InstanceIndex, Source ? TEXT( "Yes" ) : TEXT( " No" ), WaveInstance->ActiveSound->PlaybackTime, WaveInstance->GetActualVolume(), *WaveInstance->WaveData->GetPathName(), SoundOwner ? *SoundOwner->GetName() : TEXT("None") );
 	}
 
@@ -601,7 +602,7 @@ bool FAudioDevice::HandleListAudioComponentsCommand( const TCHAR* Cmd, FOutputDe
 	for( int32 ASIndex = 0; ASIndex < ActiveSounds.Num(); ASIndex++ )
 	{
 		const FActiveSound* ActiveSound = ActiveSounds[ASIndex];
-		UAudioComponent* AComp = ActiveSound->AudioComponent.Get();
+		UAudioComponent* AComp = ActiveSound->GetAudioComponent();
 		if( AComp )
 		{
 			Ar.Logf( TEXT( "    0x%p: %4d - %s, %s, %s, %s" ),
@@ -1774,7 +1775,7 @@ int32 FAudioDevice::GetSortedActiveWaveInstances(TArray<FWaveInstance*>& WaveIns
 		{
 #if !(UE_BUILD_SHIPPING || UE_BUILD_TEST)
 			if ( !ensureMsgf(ActiveSound->Sound->IsValidLowLevel(), TEXT("ActiveSound with INVALID sound. AudioComponent=%s. DebugOriginalSoundName=%s"),
-				 ActiveSound->AudioComponent.IsValid() ? *ActiveSound->AudioComponent->GetPathName() : TEXT("NO COMPONENT"),
+				 ActiveSound->GetAudioComponent() ? *ActiveSound->GetAudioComponent()->GetPathName() : TEXT("NO COMPONENT"),
 				 *ActiveSound->DebugOriginalSoundName.ToString() ))
 			{
 				// Sound was not valid, stop playing it.
@@ -1791,7 +1792,7 @@ int32 FAudioDevice::GetSortedActiveWaveInstances(TArray<FWaveInstance*>& WaveIns
 						ActiveSound->PlaybackTime, 
 						Duration, 
 						*ActiveSound->Sound->GetName(), 
-						(ActiveSound->AudioComponent.IsValid() ? *ActiveSound->AudioComponent->GetName() : TEXT("NO COMPONENT")));
+						(ActiveSound->GetAudioComponent() ? *ActiveSound->GetAudioComponent()->GetName() : TEXT("NO COMPONENT")));
 					ActiveSound->Stop(this);
 				}
 				else
@@ -2148,9 +2149,9 @@ void FAudioDevice::AddNewActiveSound( const FActiveSound& NewActiveSound )
 		else
 		{
 			// TODO - Audio Threading. This call would be a task back to game thread
-			if (NewActiveSound.AudioComponent.IsValid())
+			if (UAudioComponent* AudioComponent = NewActiveSound.GetAudioComponent())
 			{
-				NewActiveSound.AudioComponent->PlaybackCompleted(true);
+				AudioComponent->PlaybackCompleted(true);
 			}
 			//UE_LOG(LogAudio, Verbose, TEXT( "   %g: MaxConcurrentPlayCount AudioComponent : '%s' with Sound: '%s' Max: %d   Curr: %d " ), NewActiveSound.World ? NewActiveSound.World->GetAudioTimeSeconds() : 0.0f, *GetFullName(), Sound ? *Sound->GetName() : TEXT( "NULL" ), Sound->MaxConcurrentPlayCount, Sound->CurrentPlayCount );
 			return;
@@ -2160,8 +2161,10 @@ void FAudioDevice::AddNewActiveSound( const FActiveSound& NewActiveSound )
 
 	++NewActiveSound.Sound->CurrentPlayCount;
 
+	UAudioComponent* AudioComponent = NewActiveSound.GetAudioComponent();
+
 #if !(UE_BUILD_SHIPPING || UE_BUILD_TEST)
-	UE_LOG(LogAudio, VeryVerbose, TEXT("New ActiveSound %s Comp: %s Loc: %s"), *NewActiveSound.Sound->GetName(), (NewActiveSound.AudioComponent.IsValid() ? *NewActiveSound.AudioComponent->GetFullName() : TEXT("No AudioComponent")), *NewActiveSound.Transform.GetTranslation().ToString() );
+	UE_LOG(LogAudio, VeryVerbose, TEXT("New ActiveSound %s Comp: %s Loc: %s"), *NewActiveSound.Sound->GetName(), (AudioComponent ? *AudioComponent->GetFullName() : TEXT("No AudioComponent")), *NewActiveSound.Transform.GetTranslation().ToString() );
 #endif // !(UE_BUILD_SHIPPING || UE_BUILD_TEST)
 
 	FActiveSound* ActiveSound = new FActiveSound(NewActiveSound);
@@ -2169,14 +2172,14 @@ void FAudioDevice::AddNewActiveSound( const FActiveSound& NewActiveSound )
 	if (ActiveSound->Sound)
 	{
 		if (!ensureMsgf(ActiveSound->Sound->IsValidLowLevel(), TEXT("AddNewActiveSound with INVALID sound. AudioComponent=%s"),
-			ActiveSound->AudioComponent.IsValid() ? *ActiveSound->AudioComponent->GetPathName() : TEXT("NO COMPONENT") ))
+			AudioComponent ? *AudioComponent->GetPathName() : TEXT("NO COMPONENT") ))
 		{
 			static FName InvalidSoundName(TEXT("INVALID_Sound"));
 			ActiveSound->DebugOriginalSoundName = InvalidSoundName;
 		}
 		else if (!ensureMsgf(ActiveSound->Sound->GetFName() != NAME_None, TEXT("AddNewActiveSound with DESTROYED sound %s. AudioComponent=%s. IsPendingKill=%d. BeginDestroy=%d"),
 			*ActiveSound->Sound->GetPathName(),
-			ActiveSound->AudioComponent.IsValid() ? *ActiveSound->AudioComponent->GetPathName() : TEXT("NO COMPONENT"),
+			AudioComponent ? *AudioComponent->GetPathName() : TEXT("NO COMPONENT"),
 			(int32)ActiveSound->Sound->IsPendingKill(),
 			(int32)ActiveSound->Sound->HasAnyFlags(RF_BeginDestroyed)))
 		{
@@ -2190,6 +2193,10 @@ void FAudioDevice::AddNewActiveSound( const FActiveSound& NewActiveSound )
 	}
 #endif
 	ActiveSounds.Add(ActiveSound);
+	if (AudioComponent)
+	{
+		AudioComponentToActiveSoundMap.Add((UPTRINT)AudioComponent, ActiveSound);
+	}
 }
 
 void FAudioDevice::StopActiveSound( UAudioComponent* AudioComponent )
@@ -2205,27 +2212,27 @@ void FAudioDevice::StopActiveSound( UAudioComponent* AudioComponent )
 
 FActiveSound* FAudioDevice::FindActiveSound( UAudioComponent* AudioComponent )
 {
-	// find the active sound corresponding to this audiocomponent
-	for( int32 Index = 0; Index < ActiveSounds.Num(); ++Index )
+	// find the active sound corresponding to this audio component
+	if (FActiveSound** ActiveSoundPtr = AudioComponentToActiveSoundMap.Find((UPTRINT)AudioComponent))
 	{
-		FActiveSound* ActiveSound = ActiveSounds[Index];
-		if (ActiveSound->AudioComponent.IsValid(true) && ActiveSound->AudioComponent.Get() == AudioComponent)
-		{
-			return ActiveSound;
-		}
+		return *ActiveSoundPtr;
 	}
 
-	return NULL;
+	return nullptr;
 }
 
 void FAudioDevice::RemoveActiveSound(FActiveSound* ActiveSound)
 {
-	// find the active sound corresponding to this audiocomponent
 	for( int32 Index = 0; Index < ActiveSounds.Num(); ++Index )
 	{
 		if (ActiveSound == ActiveSounds[Index])
 		{
+			const UPTRINT AudioComponentIndex = ActiveSound->GetAudioComponentIndex();
 			ActiveSounds.RemoveAt(Index);
+			if (AudioComponentIndex > 0)
+			{
+				AudioComponentToActiveSoundMap.Remove(AudioComponentIndex);
+			}
 			delete ActiveSound;
 			break;
 		}
@@ -2343,10 +2350,18 @@ void FAudioDevice::Flush( UWorld* WorldToFlush, bool bClearActivatedReverb )
 		}
 		else
 		{
-			if( WorldToFlush == nullptr || ActiveSound->World == nullptr || ActiveSound->World == WorldToFlush )
+			if (WorldToFlush == nullptr)
 			{
 				ActiveSound->Stop(this);
 			}
+			else
+			{
+				UWorld* ActiveSoundWorld = ActiveSound->World.Get();
+				if (ActiveSoundWorld == nullptr || ActiveSoundWorld == WorldToFlush)
+				{
+					ActiveSound->Stop(this);
+				}
+			}
 		}
 	}
 
@@ -2629,9 +2644,9 @@ void FAudioDevice::StopSoundsUsingResource(USoundWave* SoundWave, TArray<UAudioC
 			FWaveInstance* WaveInstance = WaveInstanceIt.Value();
 			if (WaveInstance->WaveData == SoundWave)
 			{
-				if (ActiveSound->AudioComponent.IsValid())
+				if (UAudioComponent* AudioComponent = ActiveSound->GetAudioComponent())
 				{
-					StoppedComponents.Add(ActiveSound->AudioComponent.Get());
+					StoppedComponents.Add(AudioComponent);
 				}
 				ActiveSound->Stop(this);
 				bStoppedSounds = true;
diff --git a/Engine/Source/Runtime/Engine/Private/Components/AudioComponent.cpp b/Engine/Source/Runtime/Engine/Private/Components/AudioComponent.cpp
--- a/Engine/Source/Runtime/Engine/Private/Components/AudioComponent.cpp
+++ b/Engine/Source/Runtime/Engine/Private/Components/AudioComponent.cpp
@@ -139,7 +139,7 @@ void UAudioComponent::PlayInternal(const float StartTime, const float FadeInDura
 		if (FAudioDevice* AudioDevice = GetAudioDevice())
 		{
 			FActiveSound NewActiveSound;
-			NewActiveSound.AudioComponent = this;
+			NewActiveSound.SetAudioComponent(this);
 			NewActiveSound.World = GetWorld();
 			NewActiveSound.Sound = Sound;
 			NewActiveSound.SoundClassOverride = SoundClassOverride;
diff --git a/Engine/Source/Runtime/Engine/Private/DialogueWave.cpp b/Engine/Source/Runtime/Engine/Private/DialogueWave.cpp
--- a/Engine/Source/Runtime/Engine/Private/DialogueWave.cpp
+++ b/Engine/Source/Runtime/Engine/Private/DialogueWave.cpp
@@ -466,10 +466,11 @@ void UDialogueSoundWaveProxy::Parse( class FAudioDevice* AudioDevice, const UPTR
 	if (ActiveSound.bHandleSubtitles && Subtitles.Num() > 0)
 	{
 		// TODO - Audio Threading. This would need to be a call back to the main thread.
-		if (ActiveSound.AudioComponent.IsValid() && ActiveSound.AudioComponent->OnQueueSubtitles.IsBound())
+		UAudioComponent* AudioComponent = ActiveSound.GetAudioComponent();
+		if (AudioComponent && AudioComponent->OnQueueSubtitles.IsBound())
 		{
 			// intercept the subtitles if the delegate is set
-			ActiveSound.AudioComponent->OnQueueSubtitles.ExecuteIfBound( Subtitles, GetDuration() );
+			AudioComponent->OnQueueSubtitles.ExecuteIfBound( Subtitles, GetDuration() );
 		}
 		else
 		{
diff --git a/Engine/Source/Runtime/Engine/Private/SoundWave.cpp b/Engine/Source/Runtime/Engine/Private/SoundWave.cpp
--- a/Engine/Source/Runtime/Engine/Private/SoundWave.cpp
+++ b/Engine/Source/Runtime/Engine/Private/SoundWave.cpp
@@ -524,13 +524,13 @@ FWaveInstance* USoundWave::HandleStart( FActiveSound& ActiveSound, const UPTRINT
 	// Add in the subtitle if they exist
 	if (ActiveSound.bHandleSubtitles && Subtitles.Num() > 0)
 	{
-		if (UAudioComponent* AudioComponent = ActiveSound.AudioComponent.Get())
+		if (UAudioComponent* AudioComponent = ActiveSound.GetAudioComponent())
 		{
 			// TODO - Audio Threading. This would need to be a call back to the main thread.
 			if (AudioComponent->OnQueueSubtitles.IsBound())
 			{
 				// intercept the subtitles if the delegate is set
-				ActiveSound.AudioComponent->OnQueueSubtitles.ExecuteIfBound( Subtitles, Duration );
+				AudioComponent->OnQueueSubtitles.ExecuteIfBound( Subtitles, Duration );
 			}
 			else if( ActiveSound.World.IsValid() )
 			{
@@ -587,7 +587,7 @@ void USoundWave::Parse( FAudioDevice* AudioDevice, const UPTRINT NodeWaveInstanc
 	{
 		WaveInstance->bIsFinished = false;
 #if !NO_LOGGING
-		if (!ActiveSound.bWarnedAboutOrphanedLooping && !ActiveSound.AudioComponent.IsValid())
+		if (!ActiveSound.bWarnedAboutOrphanedLooping && ActiveSound.GetAudioComponent() == nullptr)
 		{
 			UE_LOG(LogAudio, Warning, TEXT("Detected orphaned looping sound '%s'."), *ActiveSound.Sound->GetName());
 			ActiveSound.bWarnedAboutOrphanedLooping = true;
@@ -700,11 +700,11 @@ void USoundWave::Parse( FAudioDevice* AudioDevice, const UPTRINT NodeWaveInstanc
 					SoundWarningInfo += FString::Printf(TEXT(" SoundCue: %s"), *ActiveSound.Sound->GetName());
 				}
 
-				if (ActiveSound.AudioComponent.IsValid())
+				if (UAudioComponent* AudioComponent = ActiveSound.GetAudioComponent())
 				{
 					// TODO - Audio Threading. This log would have to be a task back to game thread
-					AActor* SoundOwner = ActiveSound.AudioComponent->GetOwner();
-					UE_LOG(LogAudio, Warning, TEXT( "%s Actor: %s AudioComponent: %s" ), *SoundWarningInfo, (SoundOwner ? *SoundOwner->GetName() : TEXT("None")), *ActiveSound.AudioComponent->GetName() );
+					AActor* SoundOwner = AudioComponent->GetOwner();
+					UE_LOG(LogAudio, Warning, TEXT( "%s Actor: %s AudioComponent: %s" ), *SoundWarningInfo, (SoundOwner ? *SoundOwner->GetName() : TEXT("None")), *AudioComponent->GetName() );
 				}
 				else
 				{
diff --git a/Engine/Source/Runtime/Engine/Private/UnrealEngine.cpp b/Engine/Source/Runtime/Engine/Private/UnrealEngine.cpp
--- a/Engine/Source/Runtime/Engine/Private/UnrealEngine.cpp
+++ b/Engine/Source/Runtime/Engine/Private/UnrealEngine.cpp
@@ -12123,7 +12123,8 @@ int32 UEngine::RenderStatSoundWaves(UWorld* World, FViewport* Viewport, FCanvas*
 
 			ActiveSounds.Add(WaveInstance->ActiveSound);
 
-			AActor* SoundOwner = WaveInstance->ActiveSound->AudioComponent.IsValid() ? WaveInstance->ActiveSound->AudioComponent->GetOwner() : NULL;
+			UAudioComponent* AudioComponent = WaveInstance->ActiveSound->GetAudioComponent();
+			AActor* SoundOwner = AudioComponent ? AudioComponent->GetOwner() : nullptr;
 			USoundClass* SoundClass = WaveInstance->SoundClass;
 
 			FString TheString = *FString::Printf(TEXT("%4i.    %6.2f  %s   Owner: %s   SoundClass: %s"),
diff --git a/Engine/Source/Runtime/Engine/Public/ActiveSound.h b/Engine/Source/Runtime/Engine/Public/ActiveSound.h
--- a/Engine/Source/Runtime/Engine/Public/ActiveSound.h
+++ b/Engine/Source/Runtime/Engine/Public/ActiveSound.h
@@ -75,7 +75,20 @@ public:
 
 	class USoundBase* Sound;
 	TWeakObjectPtr<class UWorld> World;
+
+private:
 	TWeakObjectPtr<class UAudioComponent> AudioComponent;
+	UPTRINT AudioComponentIndex;
+
+public:
+
+	UAudioComponent* GetAudioComponent() const { return AudioComponent.Get(); }
+	UPTRINT GetAudioComponentIndex() const { return AudioComponentIndex; }
+	void SetAudioComponent(UAudioComponent* Component)
+	{
+		AudioComponent = Component;
+		AudioComponentIndex = (UPTRINT)Component;
+	}
 
 	/** Optional SoundClass to override Sound */
 	USoundClass* SoundClassOverride;
diff --git a/Engine/Source/Runtime/Engine/Public/AudioDevice.h b/Engine/Source/Runtime/Engine/Public/AudioDevice.h
--- a/Engine/Source/Runtime/Engine/Public/AudioDevice.h
+++ b/Engine/Source/Runtime/Engine/Public/AudioDevice.h
@@ -845,6 +845,8 @@ private:
 
 	TArray<struct FActiveSound*> ActiveSounds;
 
+	TMap<UPTRINT, struct FActiveSound*> AudioComponentToActiveSoundMap;
+
 	/** List of passive SoundMixes active last frame */
 	TArray<class USoundMix*> PrevPassiveSoundMixModifiers;
 };
diff --git a/Engine/Source/Runtime/Windows/XAudio2/Private/XAudio2Source.cpp b/Engine/Source/Runtime/Windows/XAudio2/Private/XAudio2Source.cpp
--- a/Engine/Source/Runtime/Windows/XAudio2/Private/XAudio2Source.cpp
+++ b/Engine/Source/Runtime/Windows/XAudio2/Private/XAudio2Source.cpp
@@ -441,9 +441,9 @@ bool FXAudio2SoundSource::Init(FWaveInstance* InWaveInstance)
 	{
 		// Find matching buffer.
 		FAudioDevice* BestAudioDevice = nullptr;
-		if (InWaveInstance->ActiveSound->AudioComponent.IsValid())
+		if (UAudioComponent* AudioComponent = InWaveInstance->ActiveSound->GetAudioComponent())
 		{
-			BestAudioDevice = InWaveInstance->ActiveSound->AudioComponent->GetAudioDevice();
+			BestAudioDevice = AudioComponent->GetAudioDevice();
 		}
 		else
 		{
@@ -1175,9 +1175,10 @@ FString FXAudio2SoundSource::Describe_Internal(bool bUseLongName, bool bIncludeC
 	AActor* SoundOwner = NULL;
 
 	// TODO - Audio Threading. This won't work cross thread.
-	if (WaveInstance->ActiveSound && WaveInstance->ActiveSound->AudioComponent.IsValid())
+	UAudioComponent* AudioComponent = (WaveInstance->ActiveSound ?  WaveInstance->ActiveSound->GetAudioComponent() : nullptr);
+	if (AudioComponent)
 	{
-		SoundOwner = WaveInstance->ActiveSound->AudioComponent->GetOwner();
+		SoundOwner = AudioComponent->GetOwner();
 	}
 
 	FString SpatializedVolumeInfo;
