Fixes from non editor builds
Merge from 2660881 and 2660967

diff --git a/Engine/Source/Runtime/Engine/Private/Animation/AnimSequence.cpp b/Engine/Source/Runtime/Engine/Private/Animation/AnimSequence.cpp
--- a/Engine/Source/Runtime/Engine/Private/Animation/AnimSequence.cpp
+++ b/Engine/Source/Runtime/Engine/Private/Animation/AnimSequence.cpp
@@ -330,6 +330,7 @@ bool UAnimSequence::IsValidToPlay() const
 	// make sure sequence length is valid and raw animation data exists, and compressed
 	return ( SequenceLength > 0.f && RawAnimationData.Num() > 0 && CompressedTrackOffsets.Num() > 0 );
 }
+#endif
 
 void UAnimSequence::SortSyncMarkers()
 {
@@ -337,8 +338,6 @@ void UAnimSequence::SortSyncMarkers()
 	RefreshSyncMarkerDataFromAuthored();
 }
 
-#endif
-
 void UAnimSequence::PreSave()
 {
 #if WITH_EDITOR
@@ -3478,30 +3477,6 @@ void UAnimSequence::RefreshTrackMapFromAnimTrackNames()
 	}
 }
 
-void UAnimSequence::RefreshSyncMarkerDataFromAuthored()
-{
-#if WITH_EDITOR
-	MarkerDataUpdateCounter++;
-#endif
-
-	if (AuthoredSyncMarkers.Num() > 0)
-	{
-		UniqueMarkerNames.Reset();
-		UniqueMarkerNames.Reserve(AuthoredSyncMarkers.Num());
-
-		const FAnimSyncMarker* PreviousMarker = nullptr;
-		for (const FAnimSyncMarker& Marker : AuthoredSyncMarkers)
-		{
-			UniqueMarkerNames.AddUnique(Marker.MarkerName);
-			PreviousMarker = &Marker;
-		}
-	}
-	else
-	{
-		UniqueMarkerNames.Empty();
-	}
-}
-
 uint8* UAnimSequence::FindSyncMarkerPropertyData(int32 SyncMarkerIndex, UArrayProperty*& ArrayProperty)
 {
 	ArrayProperty = NULL;
@@ -3628,16 +3603,16 @@ bool UAnimSequence::CreateAnimation(UAnimSequence * Sequence)
 	return false;
 }
 
+#endif
+
 void UAnimSequence::RefreshCacheData()
 {
 	SortSyncMarkers();
-	RefreshSyncMarkerDataFromAuthored();
-
+#if WITH_EDITOR
 	for (int32 TrackIndex = 0; TrackIndex < AnimNotifyTracks.Num(); ++TrackIndex)
 	{
 		AnimNotifyTracks[TrackIndex].SyncMarkers.Empty();
 	}
-
 	for (FAnimSyncMarker& SyncMarker : AuthoredSyncMarkers)
 	{
 		const int32 TrackIndex = SyncMarker.TrackIndex;
@@ -3653,10 +3628,34 @@ void UAnimSequence::RefreshCacheData()
 			SyncMarker.TrackIndex = 0;
 		}
 	}
-
+#endif
 	Super::RefreshCacheData();
 }
 
+void UAnimSequence::RefreshSyncMarkerDataFromAuthored()
+{
+#if WITH_EDITOR
+	MarkerDataUpdateCounter++;
+#endif
+
+	if (AuthoredSyncMarkers.Num() > 0)
+	{
+		UniqueMarkerNames.Reset();
+		UniqueMarkerNames.Reserve(AuthoredSyncMarkers.Num());
+
+		const FAnimSyncMarker* PreviousMarker = nullptr;
+		for (const FAnimSyncMarker& Marker : AuthoredSyncMarkers)
+		{
+			UniqueMarkerNames.AddUnique(Marker.MarkerName);
+			PreviousMarker = &Marker;
+		}
+	}
+	else
+	{
+		UniqueMarkerNames.Empty();
+	}
+}
+
 bool IsMarkerValid(const FAnimSyncMarker* Marker, bool bLooping, const TArray<FName>& ValidMarkerNames)
 {
 	return (Marker == nullptr && !bLooping) || (Marker && ValidMarkerNames.Contains(Marker->MarkerName));
@@ -4105,8 +4104,6 @@ void UAnimSequence::GetMarkerIndicesForPosition(const FMarkerSyncAnimPosition& S
 		}
 	}
 }
-
-#endif
 /*-----------------------------------------------------------------------------
 	AnimNotify& subclasses
 -----------------------------------------------------------------------------*/
