#ANIM: Fix additive blendspace/aimoffset retargeting

Merging

//depot/UE4-Orion/Engine/Source/Runtime/Engine/Private/Animation/BlendSpaceBase.cpp

to //depot/UE4/Engine/Source/Runtime/Engine/Private/Animation/BlendSpaceBase.cpp

diff --git a/Engine/Source/Runtime/Engine/Private/Animation/BlendSpaceBase.cpp b/Engine/Source/Runtime/Engine/Private/Animation/BlendSpaceBase.cpp
--- a/Engine/Source/Runtime/Engine/Private/Animation/BlendSpaceBase.cpp
+++ b/Engine/Source/Runtime/Engine/Private/Animation/BlendSpaceBase.cpp
@@ -1203,11 +1203,16 @@ void UBlendSpaceBase::GetAnimationPose(TArray<FBlendSampleData>& BlendSampleData
 #if WITH_EDITOR
 bool UBlendSpaceBase::GetAllAnimationSequencesReferred(TArray<UAnimSequence*>& AnimationSequences)
 {
- 	for (auto Iter = SampleData.CreateConstIterator(); Iter; ++Iter)
- 	{
- 		// saves all samples in the AnimSequences
- 		AnimationSequences.AddUnique((*Iter).Animation);
- 	}
+	for (auto Iter = SampleData.CreateConstIterator(); Iter; ++Iter)
+	{
+		// saves all samples in the AnimSequences
+		AnimationSequences.AddUnique((*Iter).Animation);
+	}
+
+	if (PreviewBasePose)
+	{
+		AnimationSequences.AddUnique(PreviewBasePose);
+	}
  
  	return (AnimationSequences.Num() > 0);
 }
@@ -1230,6 +1235,16 @@ void UBlendSpaceBase::ReplaceReferredAnimations(const TMap<UAnimSequence*, UAnim
 			}
 		}
 	}
+
+	if (PreviewBasePose)
+	{
+		UAnimSequence* const* ReplacementAsset = (UAnimSequence*const*)ReplacementMap.Find(PreviewBasePose);
+		if(ReplacementAsset)
+		{
+			PreviewBasePose = *ReplacementAsset;
+		}
+	}
+	
 	SampleData = NewSamples;
 }
 
