Fix PS4 Warnings
Merged from 2661005 (Thomas.Sarkanen)

diff --git a/Engine/Source/Runtime/Engine/Private/Animation/AnimInstance.cpp b/Engine/Source/Runtime/Engine/Private/Animation/AnimInstance.cpp
--- a/Engine/Source/Runtime/Engine/Private/Animation/AnimInstance.cpp
+++ b/Engine/Source/Runtime/Engine/Private/Animation/AnimInstance.cpp
@@ -2706,11 +2706,11 @@ bool UAnimInstance::GetTimeToClosestMarker(FName SyncGroup, FName MarkerName, fl
 
 	if (SyncGroups.IsValidIndex(SyncGroupIndex))
 	{
-		const FAnimGroupInstance& SyncGroup = SyncGroups[SyncGroupIndex];
-		if (SyncGroup.ActivePlayers.IsValidIndex(SyncGroup.GroupLeaderIndex))
+		const FAnimGroupInstance& SyncGroupInstance = SyncGroups[SyncGroupIndex];
+		if (SyncGroupInstance.ActivePlayers.IsValidIndex(SyncGroupInstance.GroupLeaderIndex))
 		{
-			const FMarkerSyncAnimPosition& EndPosition = SyncGroup.MarkerTickContext.GetMarkerSyncEndPosition();
-			const FAnimTickRecord& Leader = SyncGroup.ActivePlayers[SyncGroup.GroupLeaderIndex];
+			const FMarkerSyncAnimPosition& EndPosition = SyncGroupInstance.MarkerTickContext.GetMarkerSyncEndPosition();
+			const FAnimTickRecord& Leader = SyncGroupInstance.ActivePlayers[SyncGroupInstance.GroupLeaderIndex];
 			if (EndPosition.PreviousMarkerName == MarkerName)
 			{
 				OutMarkerTime = Leader.MarkerTickRecord.PreviousMarker.TimeToMarker;
@@ -2733,8 +2733,8 @@ bool UAnimInstance::HasMarkerBeenHitThisFrame(FName SyncGroup, FName MarkerName)
 
 	if (SyncGroups.IsValidIndex(SyncGroupIndex))
 	{
-		const FAnimGroupInstance& SyncGroup = SyncGroups[SyncGroupIndex];
-		return SyncGroup.MarkerTickContext.MarkersPassedThisTick.ContainsByPredicate([&MarkerName](const FPassedMarker& PassedMarker) -> bool
+		const FAnimGroupInstance& SyncGroupInstance = SyncGroups[SyncGroupIndex];
+		return SyncGroupInstance.MarkerTickContext.MarkersPassedThisTick.ContainsByPredicate([&MarkerName](const FPassedMarker& PassedMarker) -> bool
 		{
 			return PassedMarker.PassedMarkerName == MarkerName;
 		});
