nicer printout

diff --git a/Engine/Source/Runtime/Engine/Private/ContentStreaming.cpp b/Engine/Source/Runtime/Engine/Private/ContentStreaming.cpp
--- a/Engine/Source/Runtime/Engine/Private/ContentStreaming.cpp
+++ b/Engine/Source/Runtime/Engine/Private/ContentStreaming.cpp
@@ -2162,7 +2162,7 @@ FStreamingManagerTexture::FStreamingManagerTexture()
 		UE_LOG(LogContentStreaming, Log, TEXT("Textures will NEVER stream out!"));
 	}
 
-	UE_LOG(LogContentStreaming,Log,TEXT("Texture pool size is %fMB"),GTexturePoolSize/1024.f/1024.f);
+	UE_LOG(LogContentStreaming,Log,TEXT("Texture pool size is %.2f MB"),GTexturePoolSize/1024.f/1024.f);
 
 	// Convert from MByte to byte.
 	MinEvictSize *= 1024 * 1024;
