fixed UE-20190 Give access to SceneDepth in worldPositionOffset

diff --git a/Engine/Shaders/MaterialTemplate.usf b/Engine/Shaders/MaterialTemplate.usf
--- a/Engine/Shaders/MaterialTemplate.usf
+++ b/Engine/Shaders/MaterialTemplate.usf
@@ -334,6 +334,18 @@ float3 GetTranslatedWorldPosition_NoMaterialOffsets(FMaterialPixelParameters Par
 	return Parameters.WorldPosition_NoOffsets_CamRelative;
 }
 
+// using this function allows to write the same code for VS and PS
+float4 GetScreenPosition(FMaterialVertexParameters Parameters)
+{
+	return mul(float4(Parameters.WorldPosition, 1.0f), View.TranslatedWorldToClip);
+}
+
+// using this function allows to write the same code for VS and PS
+float4 GetScreenPosition(FMaterialPixelParameters Parameters)
+{
+	return Parameters.ScreenPosition;
+}
+
 #if DECAL_RENDERSTAGE && NUM_MATERIAL_TEXCOORDS
 	/*
 	 * Material node DecalMipmapLevel's code designed to avoid the 2x2 pixels artefacts on the edges around where the decal
diff --git a/Engine/Shaders/ShaderVersion.usf b/Engine/Shaders/ShaderVersion.usf
--- a/Engine/Shaders/ShaderVersion.usf
+++ b/Engine/Shaders/ShaderVersion.usf
@@ -2,4 +2,4 @@
 // Each time the console command is executed it generates a new GUID. As this file is included
 // in common.usf (which should be included in any shader) it allows to invalidate the shader DDC.
 // 
-// GUID = 1E7DB6E145A42CAE76519AA01490523B
\ No newline at end of file
+// GUID = AD9819114E9A382DD5C6E2B82A720CC3
\ No newline at end of file
diff --git a/Engine/Source/Runtime/Engine/Private/Materials/HLSLMaterialTranslator.h b/Engine/Source/Runtime/Engine/Private/Materials/HLSLMaterialTranslator.h
--- a/Engine/Source/Runtime/Engine/Private/Materials/HLSLMaterialTranslator.h
+++ b/Engine/Source/Runtime/Engine/Private/Materials/HLSLMaterialTranslator.h
@@ -1943,12 +1943,12 @@ protected:
 
 	virtual int32 ScreenPosition() override
 	{
-		if (ShaderFrequency != SF_Pixel && ShaderFrequency != SF_Compute)
+		if (ShaderFrequency != SF_Pixel && ShaderFrequency != SF_Compute && ShaderFrequency != SF_Vertex)
 		{
-			return NonPixelShaderExpressionError();
+			return Errorf(TEXT("Invalid node used in hull/domain shader input!"));
 		}
 
-		return AddCodeChunk(MCT_Float2,TEXT("ScreenAlignedPosition(Parameters.ScreenPosition).xy"));		
+		return AddCodeChunk(MCT_Float2,TEXT("ScreenAlignedPosition(GetScreenPosition(Parameters)).xy"));		
 	}
 
 	virtual int32 ParticleMacroUV() override 
@@ -2494,11 +2494,11 @@ protected:
 
 	virtual int32 PixelDepth() override
 	{
-		if (ShaderFrequency != SF_Pixel && ShaderFrequency != SF_Compute)
+		if (ShaderFrequency != SF_Pixel && ShaderFrequency != SF_Compute && ShaderFrequency != SF_Vertex)
 		{
-			return NonPixelShaderExpressionError();
+			return Errorf(TEXT("Invalid node used in hull/domain shader input!"));
 		}
-		return AddInlinedCodeChunk(MCT_Float, TEXT("Parameters.ScreenPosition.w"));		
+		return AddInlinedCodeChunk(MCT_Float, TEXT("GetScreenPosition(Parameters).w"));		
 	}
 
 	/** Calculate screen aligned UV coodinates from an offset fraction or texture coordinate */
@@ -2506,11 +2506,11 @@ protected:
 	{
 		if(bUseOffset)
 		{
-			return AddCodeChunk(MCT_Float2, TEXT("CalcScreenUVFromOffsetFraction(Parameters.ScreenPosition, %s)"), *GetParameterCode(Offset));
+			return AddCodeChunk(MCT_Float2, TEXT("CalcScreenUVFromOffsetFraction(GetScreenPosition(Parameters), %s)"), *GetParameterCode(Offset));
 		}
 		else
 		{
-			FString DefaultScreenAligned(TEXT("MaterialFloat2(ScreenAlignedPosition(Parameters.ScreenPosition).xy)"));
+			FString DefaultScreenAligned(TEXT("MaterialFloat2(ScreenAlignedPosition(GetScreenPosition(Parameters)).xy)"));
 			FString CodeString = (UV != INDEX_NONE) ? CoerceParameter(UV,MCT_Float2) : DefaultScreenAligned;
 			return AddInlinedCodeChunk(MCT_Float2, *CodeString );
 		}
@@ -2518,11 +2518,6 @@ protected:
 
 	virtual int32 SceneDepth(int32 Offset, int32 UV, bool bUseOffset) override
 	{
-		if (ShaderFrequency != SF_Pixel && ShaderFrequency != SF_Compute)
-		{
-			return NonPixelShaderExpressionError();
-		}
-
 		if (Offset == INDEX_NONE && bUseOffset)
 		{
 			return INDEX_NONE;
@@ -2548,7 +2543,7 @@ protected:
 			return INDEX_NONE;
 		}
 
-		if (ShaderFrequency != SF_Pixel)
+		if (ShaderFrequency != SF_Pixel && ShaderFrequency != SF_Vertex)
 		{
 			// we can relax this later if needed
 			return NonPixelShaderExpressionError();
@@ -2558,7 +2553,7 @@ protected:
 
 		UseSceneTextureId(SceneTextureId, true);
 
-		FString DefaultScreenAligned(TEXT("MaterialFloat2(ScreenAlignedPosition(Parameters.ScreenPosition).xy)"));
+		FString DefaultScreenAligned(TEXT("MaterialFloat2(ScreenAlignedPosition(GetScreenPosition(Parameters)).xy)"));
 		FString TexCoordCode((UV != INDEX_NONE) ? CoerceParameter(UV, MCT_Float2) : DefaultScreenAligned);
 
 		return AddCodeChunk(
