Fix time to marker api function on anim instance
Merged from 2662860

diff --git a/Engine/Source/Runtime/Engine/Classes/Animation/AnimInstance.h b/Engine/Source/Runtime/Engine/Classes/Animation/AnimInstance.h
--- a/Engine/Source/Runtime/Engine/Classes/Animation/AnimInstance.h
+++ b/Engine/Source/Runtime/Engine/Classes/Animation/AnimInstance.h
@@ -855,6 +855,11 @@ public:
 	UPROPERTY(transient)
 	TArray<FAnimNotifyEvent> ActiveAnimNotifyState;
 
+protected:
+	int32 GetSyncGroupReadIndex() const { return 1 - SyncGroupWriteIndex; }
+	int32 GetSyncGroupWriteIndex() const { return SyncGroupWriteIndex; }
+	void TickSyncGroupWriteIndex() { SyncGroupWriteIndex = GetSyncGroupReadIndex(); }
+
 private:
 	/** Curve Values that are added to trigger in event**/
 	TMap<FName, float>	EventCurves;
@@ -862,7 +867,7 @@ private:
 	TArray<FName> MaterialParamatersToClear;
 
 	// Current sync group buffer index
-	int32 ActiveSyncGroupsArray;
+	int32 SyncGroupWriteIndex;
 
 	//This frames marker sync data
 	FMarkerTickContext MarkerTickContext;
diff --git a/Engine/Source/Runtime/Engine/Private/Animation/AnimInstance.cpp b/Engine/Source/Runtime/Engine/Private/Animation/AnimInstance.cpp
--- a/Engine/Source/Runtime/Engine/Private/Animation/AnimInstance.cpp
+++ b/Engine/Source/Runtime/Engine/Private/Animation/AnimInstance.cpp
@@ -79,7 +79,7 @@ UAnimInstance::UAnimInstance(const FObjectInitializer& ObjectInitializer)
 	RootNode = NULL;
 	RootMotionMode = ERootMotionMode::RootMotionFromMontagesOnly;
 	SlotNodeInitializationCounter = INDEX_NONE;
-	ActiveSyncGroupsArray = 0;
+	SyncGroupWriteIndex = 0;
 }
 
 void UAnimInstance::MakeSequenceTickRecord(FAnimTickRecord& TickRecord, class UAnimSequenceBase* Sequence, bool bLooping, float PlayRate, float FinalBlendWeight, float& CurrentTime) const
@@ -129,7 +129,7 @@ FAnimTickRecord& UAnimInstance::CreateUninitializedTickRecord(int32 GroupIndex,
 	OutSyncGroupPtr = NULL;
 	if (GroupIndex >= 0)
 	{
-		TArray<FAnimGroupInstance>& SyncGroups = SyncGroupArrays[ActiveSyncGroupsArray];
+		TArray<FAnimGroupInstance>& SyncGroups = SyncGroupArrays[GetSyncGroupWriteIndex()];
 		while (SyncGroups.Num() <= GroupIndex)
 		{
 			new (SyncGroups) FAnimGroupInstance();
@@ -138,7 +138,7 @@ FAnimTickRecord& UAnimInstance::CreateUninitializedTickRecord(int32 GroupIndex,
 	}
 
 	// Create the record
-	FAnimTickRecord* TickRecord = new ((OutSyncGroupPtr != NULL) ? OutSyncGroupPtr->ActivePlayers : UngroupedActivePlayerArrays[ActiveSyncGroupsArray]) FAnimTickRecord();
+	FAnimTickRecord* TickRecord = new ((OutSyncGroupPtr != NULL) ? OutSyncGroupPtr->ActivePlayers : UngroupedActivePlayerArrays[GetSyncGroupWriteIndex()]) FAnimTickRecord();
 	return *TickRecord;
 }
 
@@ -163,14 +163,13 @@ void UAnimInstance::InitTickRecord(const TArray<FAnimTickRecord>& TickRecords, F
 
 void UAnimInstance::InitTickRecordFromLastFrame(int32 GroupIndex, FAnimTickRecord& NewTickRecord) const
 {
-	const int32 PreviousBufferIndex = 1 - ActiveSyncGroupsArray;
 	if (GroupIndex == INDEX_NONE)
 	{
-		InitTickRecord(UngroupedActivePlayerArrays[PreviousBufferIndex], NewTickRecord);
+		InitTickRecord(UngroupedActivePlayerArrays[GetSyncGroupReadIndex()], NewTickRecord);
 	}
 	else
 	{
-		const TArray<FAnimGroupInstance>& PreviousSyncGroups = SyncGroupArrays[PreviousBufferIndex];
+		const TArray<FAnimGroupInstance>& PreviousSyncGroups = SyncGroupArrays[GetSyncGroupReadIndex()];
 		if (PreviousSyncGroups.IsValidIndex(GroupIndex))
 		{
 			const TArray<FAnimTickRecord>& TickRecords = PreviousSyncGroups[GroupIndex].ActivePlayers;
@@ -374,12 +373,10 @@ void UAnimInstance::UpdateAnimation(float DeltaSeconds)
 	}
 #endif
 
-	ActiveSyncGroupsArray = 1 - ActiveSyncGroupsArray;
+	TArray<FAnimGroupInstance>& SyncGroups = SyncGroupArrays[GetSyncGroupWriteIndex()];
+	TArray<FAnimTickRecord>& UngroupedActivePlayers = UngroupedActivePlayerArrays[GetSyncGroupWriteIndex()];
 
-	TArray<FAnimGroupInstance>& SyncGroups = SyncGroupArrays[ActiveSyncGroupsArray];
-	TArray<FAnimTickRecord>& UngroupedActivePlayers = UngroupedActivePlayerArrays[ActiveSyncGroupsArray];
-
-	const TArray<FAnimGroupInstance>& PreviousSyncGroups = SyncGroupArrays[1 - ActiveSyncGroupsArray];
+	const TArray<FAnimGroupInstance>& PreviousSyncGroups = SyncGroupArrays[GetSyncGroupReadIndex()];
 
 	{
 		QUICK_SCOPE_CYCLE_COUNTER(STAT_UAnimInstance_UpdateAnimation_MiscClear);
@@ -415,6 +412,8 @@ void UAnimInstance::UpdateAnimation(float DeltaSeconds)
 		RootNode->Update(UpdateContext);
 	}
 
+	TickSyncGroupWriteIndex();
+
 	{
 		QUICK_SCOPE_CYCLE_COUNTER(STAT_TickAssetPlayerInstances);
 
@@ -854,8 +853,8 @@ void UAnimInstance::DisplayDebug(class UCanvas* Canvas, const FDebugDisplayInfo&
 		FIndenter AnimIndent(Indent);
 
 		//Display Sync Groups
-		const TArray<FAnimGroupInstance>& SyncGroups = SyncGroupArrays[ActiveSyncGroupsArray];
-		const TArray<FAnimTickRecord>& UngroupedActivePlayers = UngroupedActivePlayerArrays[ActiveSyncGroupsArray];
+		const TArray<FAnimGroupInstance>& SyncGroups = SyncGroupArrays[GetSyncGroupReadIndex()];
+		const TArray<FAnimTickRecord>& UngroupedActivePlayers = UngroupedActivePlayerArrays[GetSyncGroupReadIndex()];
 
 		Heading = FString::Printf(TEXT("SyncGroups: %i"), SyncGroups.Num());
 		YL = Canvas->DrawText(RenderFont, Heading, Indent, YPos, 1.f, 1.f, RenderInfo);
@@ -2702,7 +2701,7 @@ void UAnimInstance::UnlockAIResources(bool bUnlockMovement, bool UnlockAILogic)
 bool UAnimInstance::GetTimeToClosestMarker(FName SyncGroup, FName MarkerName, float& OutMarkerTime) const
 {
 	const int32 SyncGroupIndex = GetSyncGroupIndexFromName(SyncGroup);
-	const TArray<FAnimGroupInstance>& SyncGroups = SyncGroupArrays[ActiveSyncGroupsArray];
+	const TArray<FAnimGroupInstance>& SyncGroups = SyncGroupArrays[GetSyncGroupReadIndex()];
 
 	if (SyncGroups.IsValidIndex(SyncGroupIndex))
 	{
@@ -2729,7 +2728,7 @@ bool UAnimInstance::GetTimeToClosestMarker(FName SyncGroup, FName MarkerName, fl
 bool UAnimInstance::HasMarkerBeenHitThisFrame(FName SyncGroup, FName MarkerName) const
 {
 	const int32 SyncGroupIndex = GetSyncGroupIndexFromName(SyncGroup);
-	const TArray<FAnimGroupInstance>& SyncGroups = SyncGroupArrays[ActiveSyncGroupsArray];
+	const TArray<FAnimGroupInstance>& SyncGroups = SyncGroupArrays[GetSyncGroupReadIndex()];
 
 	if (SyncGroups.IsValidIndex(SyncGroupIndex))
 	{
