Make AAmbientSound ENGINE_API by licensee request

diff --git a/Engine/Source/Runtime/Engine/Classes/Sound/AmbientSound.h b/Engine/Source/Runtime/Engine/Classes/Sound/AmbientSound.h
--- a/Engine/Source/Runtime/Engine/Classes/Sound/AmbientSound.h
+++ b/Engine/Source/Runtime/Engine/Classes/Sound/AmbientSound.h
@@ -4,8 +4,8 @@
 #include "AmbientSound.generated.h"
 
 /** A sound actor that can be placed in a level */
-UCLASS(AutoExpandCategories=Audio, ClassGroup=Sounds, MinimalAPI, hideCategories(Collision, Input, Game), showCategories=("Input|MouseInput", "Input|TouchInput", "Game|Damage"), ComponentWrapperClass)
-class AAmbientSound : public AActor
+UCLASS(AutoExpandCategories=Audio, ClassGroup=Sounds, hideCategories(Collision, Input, Game), showCategories=("Input|MouseInput", "Input|TouchInput", "Game|Damage"), ComponentWrapperClass)
+class ENGINE_API AAmbientSound : public AActor
 {
 	GENERATED_UCLASS_BODY()
 
@@ -16,7 +16,7 @@ private_subobject:
 	class UAudioComponent* AudioComponent;
 public:
 	
-	ENGINE_API FString GetInternalSoundCueName();
+	FString GetInternalSoundCueName();
 
 	// Begin AActor interface.
 #if WITH_EDITOR
@@ -41,7 +41,7 @@ public:
 
 public:
 	/** Returns AudioComponent subobject **/
-	ENGINE_API class UAudioComponent* GetAudioComponent() const;
+	class UAudioComponent* GetAudioComponent() const;
 };
 
 
