Moved all occurrences of UNavigationSystem::InitializeForWorld taking place before world's InitializeActorsForPlay back to guarantee actors are initialized when navigation system tries to reason about them #UE4

This basically fixed navigation bounds gathering.
Integrated over from Orion branch.
UE-19707

diff --git a/Engine/Source/Runtime/Engine/Private/AI/Navigation/NavigationSystem.cpp b/Engine/Source/Runtime/Engine/Private/AI/Navigation/NavigationSystem.cpp
--- a/Engine/Source/Runtime/Engine/Private/AI/Navigation/NavigationSystem.cpp
+++ b/Engine/Source/Runtime/Engine/Private/AI/Navigation/NavigationSystem.cpp
@@ -1612,14 +1612,24 @@ FBox UNavigationSystem::GetWorldBounds() const
 
 	NavigableWorldBounds = FBox(0);
 
-	if (GetWorld() != NULL && bWholeWorldNavigable == true)
+	if (GetWorld() != nullptr)
 	{
-		// @TODO - super slow! Need to ask tech guys where I can get this from
-		for( FActorIterator It(GetWorld()); It; ++It )
+		if (bWholeWorldNavigable == false)
 		{
-			if (IsNavigationRelevant(*It))
+			for (const FNavigationBounds& Bounds : RegisteredNavBounds)
 			{
-				NavigableWorldBounds += (*It)->GetComponentsBoundingBox();
+				NavigableWorldBounds += Bounds.AreaBox;
+			}
+		}
+		else
+		{
+			// @TODO - super slow! Need to ask tech guys where I can get this from
+			for (FActorIterator It(GetWorld()); It; ++It)
+			{
+				if (IsNavigationRelevant(*It))
+				{
+					NavigableWorldBounds += (*It)->GetComponentsBoundingBox();
+				}
 			}
 		}
 	}
diff --git a/Engine/Source/Runtime/Engine/Private/GameInstance.cpp b/Engine/Source/Runtime/Engine/Private/GameInstance.cpp
--- a/Engine/Source/Runtime/Engine/Private/GameInstance.cpp
+++ b/Engine/Source/Runtime/Engine/Private/GameInstance.cpp
@@ -242,10 +242,11 @@ bool UGameInstance::StartPIEGameInstance(ULocalPlayer* LocalPlayer, bool bInSimu
 		// Make sure "always loaded" sub-levels are fully loaded
 		PlayWorld->FlushLevelStreaming(EFlushLevelStreamingType::Visibility);
 
-		UNavigationSystem::InitializeForWorld(PlayWorld, LocalPlayers.Num() > 0 ? FNavigationSystemRunMode::PIEMode : FNavigationSystemRunMode::SimulationMode);
 		PlayWorld->CreateAISystem();
 
 		PlayWorld->InitializeActorsForPlay(URL);
+		// calling it after InitializeActorsForPlay has been called to have all potential bounding boxed initialized
+		UNavigationSystem::InitializeForWorld(PlayWorld, LocalPlayers.Num() > 0 ? FNavigationSystemRunMode::PIEMode : FNavigationSystemRunMode::SimulationMode);
 
 		// @todo, just use WorldContext.GamePlayer[0]?
 		if (LocalPlayer)
diff --git a/Engine/Source/Runtime/Engine/Private/UnrealEngine.cpp b/Engine/Source/Runtime/Engine/Private/UnrealEngine.cpp
--- a/Engine/Source/Runtime/Engine/Private/UnrealEngine.cpp
+++ b/Engine/Source/Runtime/Engine/Private/UnrealEngine.cpp
@@ -9741,14 +9741,15 @@ bool UEngine::LoadMap( FWorldContext& WorldContext, FURL URL, class UPendingNetG
 		WorldContext.World()->FlushLevelStreaming(EFlushLevelStreamingType::Visibility);
 	}
 	
-	UNavigationSystem::InitializeForWorld(WorldContext.World(), FNavigationSystemRunMode::GameMode);
-	
 	// Note that AI system will be created only if ai-system-creation conditions are met
 	WorldContext.World()->CreateAISystem();
 
 	// Initialize gameplay for the level.
 	WorldContext.World()->InitializeActorsForPlay(URL);		
 
+	// calling it after InitializeActorsForPlay has been called to have all potential bounding boxed initialized
+	UNavigationSystem::InitializeForWorld(WorldContext.World(), FNavigationSystemRunMode::GameMode);
+
 	// Remember the URL. Put this before spawning player controllers so that
 	// a player controller can get the map name during initialization and
 	// have it be correct
diff --git a/Engine/Source/Runtime/Engine/Private/World.cpp b/Engine/Source/Runtime/Engine/Private/World.cpp
--- a/Engine/Source/Runtime/Engine/Private/World.cpp
+++ b/Engine/Source/Runtime/Engine/Private/World.cpp
@@ -4948,14 +4948,15 @@ UWorld* FSeamlessTravelHandler::Tick()
 			// Make sure "always loaded" sub-levels are fully loaded
 			LoadedWorld->FlushLevelStreaming(EFlushLevelStreamingType::Visibility);
 			
-			UNavigationSystem::InitializeForWorld(LoadedWorld, FNavigationSystemRunMode::GameMode);
-
 			// Note that AI system will be created only if ai-system-creation conditions are met
 			LoadedWorld->CreateAISystem();
 
 			// call initialize functions on everything that wasn't carried over from the old world
 			LoadedWorld->InitializeActorsForPlay(PendingTravelURL, false);
 
+			// calling it after InitializeActorsForPlay has been called to have all potential bounding boxed initialized
+			UNavigationSystem::InitializeForWorld(LoadedWorld, FNavigationSystemRunMode::GameMode);
+
 			// send loading complete notifications for all local players
 			for (FLocalPlayerIterator It(GEngine, LoadedWorld); It; ++It)
 			{
