Fix static lighting on blueprints during reconstruction

diff --git a/Engine/Source/Runtime/Engine/Private/ComponentInstanceDataCache.cpp b/Engine/Source/Runtime/Engine/Private/ComponentInstanceDataCache.cpp
--- a/Engine/Source/Runtime/Engine/Private/ComponentInstanceDataCache.cpp
+++ b/Engine/Source/Runtime/Engine/Private/ComponentInstanceDataCache.cpp
@@ -166,15 +166,10 @@ FComponentInstanceDataCache::FComponentInstanceDataCache(const AActor* Actor)
 		{
 			if (Component->IsCreatedByConstructionScript()) // Only cache data from 'created by construction script' components
 			{
-				// Also exclude instances that are based on the component CDO (i.e. not a unique archetype) - these are not going to be editable anyway.
-				const UObject* ComponentTemplate = Component->GetArchetype();
-				if (ComponentTemplate && ComponentTemplate != Component->GetClass()->GetDefaultObject())
+				FActorComponentInstanceData* ComponentInstanceData = Component->GetComponentInstanceData();
+				if (ComponentInstanceData)
 				{
-					FActorComponentInstanceData* ComponentInstanceData = Component->GetComponentInstanceData();
-					if (ComponentInstanceData)
-					{
-						ComponentsInstanceData.Add(ComponentInstanceData);
-					}
+					ComponentsInstanceData.Add(ComponentInstanceData);
 				}
 			}
 			else if (Component->CreationMethod == EComponentCreationMethod::Instance)
