GitHub 1293 : Fixed typo in property tooltip and comments

diff --git a/Engine/Source/Runtime/Engine/Classes/Engine/ActorChannel.h b/Engine/Source/Runtime/Engine/Classes/Engine/ActorChannel.h
--- a/Engine/Source/Runtime/Engine/Classes/Engine/ActorChannel.h
+++ b/Engine/Source/Runtime/Engine/Classes/Engine/ActorChannel.h
@@ -141,13 +141,13 @@ public:
 	/** Cleans up replicators and clears references to the actor class this channel was associated with.*/
 	void CleanupReplicators( const bool bKeepReplicators = false );
 
-	/** Writes the header for a content block of proeprties / RPCs for the given object (either the actor a subobject of the actor) */
+	/** Writes the header for a content block of properties / RPCs for the given object (either the actor a subobject of the actor) */
 	void BeginContentBlock( UObject* Obj, FOutBunch &Bunch );
 
 	/** Writes the header for a content block specifically for deleting sub-objects */
 	void BeginContentBlockForSubObjectDelete( FOutBunch & Bunch, FNetworkGUID & GuidToDelete );
 
-	/** Writes the footer for a content block of proeprties / RPCs for the given object (either the actor a subobject of the actor) */
+	/** Writes the footer for a content block of properties / RPCs for the given object (either the actor a subobject of the actor) */
 	void EndContentBlock( UObject *Obj, FOutBunch &Bunch, const FClassNetCache* ClassCache = nullptr );
 
 	/** Reads the header of the content block and instantiates the subobject if necessary */
diff --git a/Engine/Source/Runtime/Engine/Classes/Engine/StaticMeshActor.h b/Engine/Source/Runtime/Engine/Classes/Engine/StaticMeshActor.h
--- a/Engine/Source/Runtime/Engine/Classes/Engine/StaticMeshActor.h
+++ b/Engine/Source/Runtime/Engine/Classes/Engine/StaticMeshActor.h
@@ -25,7 +25,7 @@ public:
 	
 	virtual void BeginPlay() override;
 
-	/** This static mesh should replicate movement. Automatically sets the RemoteRole and bReplicateMovement flags. Meant to be edited on placed actors (those other two proeprties are not) */
+	/** This static mesh should replicate movement. Automatically sets the RemoteRole and bReplicateMovement flags. Meant to be edited on placed actors (those other two properties are not) */
 	UPROPERTY(Category=Actor, EditAnywhere, AdvancedDisplay)
 	bool bStaticMeshReplicateMovement;
 
