UE4 - Integrate texture streaming message

diff --git a/Engine/Source/Runtime/Engine/Private/ContentStreaming.cpp b/Engine/Source/Runtime/Engine/Private/ContentStreaming.cpp
--- a/Engine/Source/Runtime/Engine/Private/ContentStreaming.cpp
+++ b/Engine/Source/Runtime/Engine/Private/ContentStreaming.cpp
@@ -2100,6 +2100,7 @@ FStreamingManagerTexture::FStreamingManagerTexture()
 ,	TotalDynamicTextureHeuristicSize(0)
 ,	TotalLastRenderHeuristicSize(0)
 ,	TotalForcedHeuristicSize(0)
+,	MemoryOverBudget(0)
 ,	OriginalTexturePoolSize(0)
 ,	PreviousPoolSizeTimestamp(0.0)
 ,	PreviousPoolSizeSetting(-1)
@@ -3375,11 +3376,13 @@ void FStreamingManagerTexture::StreamTextures( bool bProcessEverything )
 		AvailableNow = GStreamMemoryTracker.CalcAvailableNow( Stats.ComputeAvailableMemorySize(), MemoryMargin );
 		AvailableLater = GStreamMemoryTracker.CalcAvailableLater( Stats.ComputeAvailableMemorySize(), MemoryMargin );
 
-		STAT( int64 NonStreamingUsage = Stats.AllocatedMemorySize - int64(ThreadStats.TotalResidentSize) - int64(ThreadStats.PendingStreamInSize) - int64(TempMemoryUsed) );
-		STAT( int64 MemoryBudget = Stats.TexturePoolSize - NonStreamingUsage - int64(MemoryMargin) );
-		STAT( int64 MemoryOverBudget = int64(ThreadStats.TotalRequiredSize) - MemoryBudget );
-		SET_MEMORY_STAT( STAT_StreamingOverBudget,	FMath::Max(MemoryOverBudget,0ll) );
-		SET_MEMORY_STAT( STAT_StreamingUnderBudget,	FMath::Max(-MemoryOverBudget,0ll) );
+		int64 NonStreamingUsage = Stats.AllocatedMemorySize - ThreadStats.TotalResidentSize - int64(ThreadStats.PendingStreamInSize) - TempMemoryUsed;
+		int64 MemoryBudget = Stats.TexturePoolSize - NonStreamingUsage - MemoryMargin;
+		int64 LocalMemoryOverBudget = FMath::Max(ThreadStats.TotalRequiredSize - MemoryBudget, 0LL);
+		MemoryOverBudget = FMath::Max(LocalMemoryOverBudget, 0LL);
+
+		SET_MEMORY_STAT( STAT_StreamingOverBudget,	MemoryOverBudget);
+		SET_MEMORY_STAT( STAT_StreamingUnderBudget,	FMath::Max(-LocalMemoryOverBudget,0ll) );
 		SET_MEMORY_STAT( STAT_NonStreamingTexturesSize,	FMath::Max<int64>(NonStreamingUsage, 0) );
 	}
 	else
@@ -3387,6 +3390,7 @@ void FStreamingManagerTexture::StreamTextures( bool bProcessEverything )
 		TempMemoryUsed = ThreadStats.TempStreamingSize;
 		AvailableNow = MAX_int64;
 		AvailableLater = MAX_int64;
+		MemoryOverBudget = 0;
 		SET_MEMORY_STAT( STAT_StreamingOverBudget, 0 );
 		SET_MEMORY_STAT( STAT_StreamingUnderBudget, 0 );
 	}
diff --git a/Engine/Source/Runtime/Engine/Private/UnrealEngine.cpp b/Engine/Source/Runtime/Engine/Private/UnrealEngine.cpp
--- a/Engine/Source/Runtime/Engine/Private/UnrealEngine.cpp
+++ b/Engine/Source/Runtime/Engine/Private/UnrealEngine.cpp
@@ -7447,6 +7447,17 @@ void DrawStatsHUD( UWorld* World, FViewport* Viewport, FCanvas* Canvas, UCanvas*
 				MessageY += 20;
 			}
 			
+#if !(UE_BUILD_SHIPPING || UE_BUILD_TEST)
+			auto MemOver = IStreamingManager::Get().GetTextureStreamingManager().GetMemoryOverBudget();
+			if (MemOver > 0)
+			{
+				SmallTextItem.SetColor(FLinearColor::Red);
+				SmallTextItem.Text =  FText::FromString(FString::Printf(TEXT("TEXTURE STREAMING POOL OVER %0.2f MB"), (float)MemOver / 1024.0f / 1024.0f));
+				Canvas->DrawItem(SmallTextItem, FVector2D(MessageX, MessageY));
+				MessageY += 20;
+			}
+#endif
+
 			// check navmesh
 #if WITH_EDITOR
 			const bool bIsNavigationAutoUpdateEnabled = UNavigationSystem::GetIsNavigationAutoUpdateEnabled();
diff --git a/Engine/Source/Runtime/Engine/Public/ContentStreaming.h b/Engine/Source/Runtime/Engine/Public/ContentStreaming.h
--- a/Engine/Source/Runtime/Engine/Public/ContentStreaming.h
+++ b/Engine/Source/Runtime/Engine/Public/ContentStreaming.h
@@ -474,6 +474,8 @@ struct ITextureStreamingManager : public IStreamingManager
 
 	/** Returns true if this is a streaming texture that is managed by the streaming manager. */
 	virtual bool IsManagedStreamingTexture(const UTexture2D* Texture2D) = 0;
+
+	virtual int64 GetMemoryOverBudget() const = 0;
 };
 
 /**
@@ -897,6 +899,8 @@ struct FStreamingManagerTexture : public ITextureStreamingManager
 	 **/
 	virtual bool StreamOutTextureData( int64 RequiredMemorySize ) override;
 
+	virtual int64 GetMemoryOverBudget() const override { return MemoryOverBudget; }
+
 	/**
 	 * Allows the streaming manager to process exec commands.
 	 *
@@ -1321,6 +1325,7 @@ protected:
 	uint64 TotalLastRenderHeuristicSize;
 	/** Total number of bytes using ForcedIntoMemory heuristics, currently in memory. */
 	uint64 TotalForcedHeuristicSize;
+	int64 MemoryOverBudget;
 
 	/** Unmodified texture pool size, in bytes, as specified in the .ini file. */
 	int64 OriginalTexturePoolSize;
