Engine - Fixing several shadow variable warnings.

diff --git a/Engine/Source/Runtime/Engine/Private/AnimationAsset.cpp b/Engine/Source/Runtime/Engine/Private/AnimationAsset.cpp
--- a/Engine/Source/Runtime/Engine/Private/AnimationAsset.cpp
+++ b/Engine/Source/Runtime/Engine/Private/AnimationAsset.cpp
@@ -46,21 +46,21 @@ void FAnimGroupInstance::Finalize(const TArray<FName>& PreviousValidMarkers)
 		bCanUseMarkerSync = true;
 
 		//filter markers based on what exists in the other animations
-		for (int32 Idx = 0; Idx < ActivePlayers.Num(); ++Idx)
+		for ( int32 ActivePlayerIndex = 0; ActivePlayerIndex < ActivePlayers.Num(); ++ActivePlayerIndex )
 		{
-			if (Idx != GroupLeaderIndex)
+			if ( ActivePlayerIndex != GroupLeaderIndex )
 			{
-				FAnimTickRecord& Candidate = ActivePlayers[Idx];
-				TArray<FName>* MarkerNames = Candidate.SourceAsset->GetUniqueMarkerNames();
-				if (MarkerNames) // Let anims with no markers set use length scaling sync
+				FAnimTickRecord& Candidate = ActivePlayers[ActivePlayerIndex];
+				TArray<FName>* PlayerMarkerNames = Candidate.SourceAsset->GetUniqueMarkerNames();
+				if ( PlayerMarkerNames ) // Let anims with no markers set use length scaling sync
 				{
 					Candidate.bCanUseMarkerSync = true;
-					for (int Idx = ValidMarkers.Num() - 1; Idx >= 0; --Idx)
+					for ( int32 ValidMarkerIndex = ValidMarkers.Num() - 1; ValidMarkerIndex >= 0; --ValidMarkerIndex )
 					{
-						FName& MarkerName = ValidMarkers[Idx];
-						if (!MarkerNames->Contains(MarkerName))
+						FName& MarkerName = ValidMarkers[ValidMarkerIndex];
+						if ( !PlayerMarkerNames->Contains(MarkerName) )
 						{
-							ValidMarkers.RemoveAtSwap(Idx, 1, false);
+							ValidMarkers.RemoveAtSwap(ValidMarkerIndex, 1, false);
 						}
 					}
 					if (ValidMarkers.Num() == 0) //No common markers between all anims, no marker syncing
