Engine - Low level logging a new line every frame is bad.

diff --git a/Engine/Source/Runtime/Engine/Private/Animation/BlendSpaceBase.cpp b/Engine/Source/Runtime/Engine/Private/Animation/BlendSpaceBase.cpp
--- a/Engine/Source/Runtime/Engine/Private/Animation/BlendSpaceBase.cpp
+++ b/Engine/Source/Runtime/Engine/Private/Animation/BlendSpaceBase.cpp
@@ -294,7 +294,7 @@ void UBlendSpaceBase::TickAssetPlayerInstance(FAnimTickRecord& Instance, class U
 			{
 				// advance current time - blend spaces hold normalized time as when dealing with changing anim length it would be possible to go backwards
 				UE_LOG(LogAnimation, Verbose, TEXT("BlendSpace(%s) - BlendInput(%s) : AnimLength(%0.5f) "), *GetName(), *BlendInput.ToString(), NewAnimLength);
-				FPlatformMisc::LowLevelOutputDebugString(TEXT("\n"));
+				
 				if (bCanDoMarkerSync)
 				{
 					const int32 HighestWeightIndex = GetHighestWeightSample(SampleDataList);
@@ -1214,7 +1214,7 @@ bool UBlendSpaceBase::GetAllAnimationSequencesReferred(TArray<UAnimSequence*>& A
 		AnimationSequences.AddUnique(PreviewBasePose);
 	}
  
- 	return (AnimationSequences.Num() > 0);
+	return (AnimationSequences.Num() > 0);
 }
 
 void UBlendSpaceBase::ReplaceReferredAnimations(const TMap<UAnimSequence*, UAnimSequence*>& ReplacementMap)
