Skeletal Mesh Component Joint/Constraint Interface for named bones
added ability to break a named constraint on  a skeltal mesh from blueprint.
added ability to set twist and swing limits.  done in a single function for all 3, where limit values of 0 and over 180 are used to specify locked or free if not limited by a value between these.  Hence avoids unnecessary mode setting in api.

Features requested for VR demo, so adding to engine branch instead of orion branch.
features were tested in a blueprint and did work as described.

diff --git a/Engine/Source/Runtime/Engine/Classes/Components/SkeletalMeshComponent.h b/Engine/Source/Runtime/Engine/Classes/Components/SkeletalMeshComponent.h
--- a/Engine/Source/Runtime/Engine/Classes/Components/SkeletalMeshComponent.h
+++ b/Engine/Source/Runtime/Engine/Classes/Components/SkeletalMeshComponent.h
@@ -1120,8 +1120,28 @@ public:
 	 * @param	HitLocation	location of the hit
 	 * @param	InBoneName	Name of bone to break constraint for
 	 */
+	UFUNCTION(BlueprintCallable, Category = "Components|SkeletalMesh", meta = (Keywords = "Constraint"))
 	void BreakConstraint(FVector Impulse, FVector HitLocation, FName InBoneName);
-	
+
+	/** Sets the Angular Motion Ranges for a named bone
+	*  @param InBoneName  Name of bone to adjust constraint ranges for
+	*  @param Swing1LimitAngle	 Size of limit in degrees, 0 means locked, 180 means free
+	*  @param TwistLimitAngle	 Size of limit in degrees, 0 means locked, 180 means free
+	*  @param Swing2LimitAngle	 Size of limit in degrees, 0 means locked, 180 means free
+	*/
+	UFUNCTION(BlueprintCallable, Category = "Components|SkeletalMesh")
+	void  SetAngularLimits(FName InBoneName,float Swing1LimitAngle, float TwistLimitAngle, float Swing2LimitAngle);
+
+	/** Gets the current Angular state for a named bone constraint 
+	*  @param InBoneName  Name of bone to get constraint ranges for
+	*  @param Swing1Angle current angular state of the constraint
+	*  @param TwistAngle  current angular state of the constraint
+	*  @param Swing2Angle current angular state of the constraint
+	*/
+	UFUNCTION(BlueprintCallable, Category = "Components|SkeletalMesh")
+	void GetCurrentJointAngles(FName InBoneName,float& Swing1Angle, float& TwistAngle, float& Swing2Angle) ;
+
+
 	/** iterates through all bodies in our PhysicsAsset and returns the location of the closest bone associated
 	 * with a body that has collision enabled.
 	 * @param TestLocation - location to check against
diff --git a/Engine/Source/Runtime/Engine/Private/SkeletalMeshComponentPhysics.cpp b/Engine/Source/Runtime/Engine/Private/SkeletalMeshComponentPhysics.cpp
--- a/Engine/Source/Runtime/Engine/Private/SkeletalMeshComponentPhysics.cpp
+++ b/Engine/Source/Runtime/Engine/Private/SkeletalMeshComponentPhysics.cpp
@@ -1612,6 +1612,55 @@ void USkeletalMeshComponent::BreakConstraint(FVector Impulse, FVector HitLocatio
 }
 
 
+void USkeletalMeshComponent::SetAngularLimits(FName InBoneName, float Swing1LimitAngle, float TwistLimitAngle, float Swing2LimitAngle)
+{
+	int32 ConstraintIndex = FindConstraintIndex(InBoneName);
+	if (ConstraintIndex == INDEX_NONE || ConstraintIndex >= Constraints.Num())
+	{
+		return;
+	}
+
+	FConstraintInstance* Constraint = Constraints[ConstraintIndex];
+	// If already broken, our job has already been done. Bail!
+	if (Constraint->IsTerminated())
+	{
+		return;
+	}
+
+	UPhysicsAsset * const PhysicsAsset = GetPhysicsAsset();
+
+	// Figure out if Body is fixed or not
+	FBodyInstance* Body = GetBodyInstance(Constraint->JointName);
+
+	if (Body != NULL && Body->IsInstanceSimulatingPhysics())
+	{
+		// Unfix body so it can be broken.
+		Body->SetInstanceSimulatePhysics(true);
+	}
+
+	// update limits
+	Constraint->SetAngularSwing1Limit(Swing1LimitAngle == 0 ? ACM_Locked : (Swing1LimitAngle >= 180) ? ACM_Free : ACM_Limited, Swing1LimitAngle);
+	Constraint->SetAngularTwistLimit(TwistLimitAngle == 0 ? ACM_Locked : (TwistLimitAngle >= 180) ? ACM_Free : ACM_Limited, TwistLimitAngle);
+	Constraint->SetAngularSwing2Limit(Swing2LimitAngle == 0 ? ACM_Locked : (Swing2LimitAngle >= 180) ? ACM_Free : ACM_Limited, Swing2LimitAngle);
+}
+
+
+void USkeletalMeshComponent::GetCurrentJointAngles(FName InBoneName, float &Swing1Angle, float &TwistAngle, float &Swing2Angle)
+{
+	int32 ConstraintIndex = FindConstraintIndex(InBoneName);
+	if (ConstraintIndex == INDEX_NONE || ConstraintIndex >= Constraints.Num())
+	{
+		return;
+	}
+
+	FConstraintInstance* Constraint = Constraints[ConstraintIndex];
+	
+	Swing1Angle = FMath::RadiansToDegrees(Constraint->GetCurrentSwing1());
+	Swing2Angle = FMath::RadiansToDegrees(Constraint->GetCurrentSwing2());
+	TwistAngle = FMath::RadiansToDegrees(Constraint->GetCurrentTwist());
+}
+
+
 void USkeletalMeshComponent::SetPhysicsAsset(UPhysicsAsset* InPhysicsAsset, bool bForceReInit)
 {
 	// If this is different from what we have now, or we should have an instance but for whatever reason it failed last time, teardown/recreate now.
