Animation: Prevented some unnecessary allocations in UAnimInstance::UpdateCurves
[INTEGRATE] Change 2663695 by [REDACTED_PERSON][REDACTED_HANDLE] on 2015/08/20 19:40:09

diff --git a/Engine/Source/Runtime/Engine/Private/Animation/AnimInstance.cpp b/Engine/Source/Runtime/Engine/Private/Animation/AnimInstance.cpp
--- a/Engine/Source/Runtime/Engine/Private/Animation/AnimInstance.cpp
+++ b/Engine/Source/Runtime/Engine/Private/Animation/AnimInstance.cpp
@@ -290,7 +290,7 @@ void UAnimInstance::UninitializeAnimation()
 {
 	StopAllMontages(0.f);
 
-	USkeletalMeshComponent * SkelMeshComp = GetSkelMeshComponent();
+	USkeletalMeshComponent* SkelMeshComp = GetSkelMeshComponent();
 	if (SkelMeshComp)
 	{
 		// Tick currently active AnimNotifyState
@@ -1058,7 +1058,7 @@ void UAnimInstance::DisplayDebug(class UCanvas* Canvas, const FDebugDisplayInfo&
 
 void UAnimInstance::RecalcRequiredBones()
 {
-	USkeletalMeshComponent * SkelMeshComp = GetSkelMeshComponent();
+	USkeletalMeshComponent* SkelMeshComp = GetSkelMeshComponent();
 	check( SkelMeshComp )
 
 	if( SkelMeshComp->SkeletalMesh && SkelMeshComp->SkeletalMesh->Skeleton )
@@ -1222,7 +1222,7 @@ void UAnimInstance::AddCurveValue(const USkeleton::AnimCurveUID Uid, float Value
 void UAnimInstance::UpdateCurves(const FBlendedCurve& InCurve)
 {
 	//Track material params we set last time round so we can clear them if they aren't set again.
-	MaterialParamatersToClear.Empty();
+	MaterialParamatersToClear.Reset();
 	for(auto Iter = MaterialParameterCurves.CreateConstIterator(); Iter; ++Iter)
 	{
 		if(Iter.Value() > 0.0f)
@@ -1231,9 +1231,9 @@ void UAnimInstance::UpdateCurves(const FBlendedCurve& InCurve)
 		}
 	}
 
-	EventCurves.Empty();
-	MorphTargetCurves.Empty();
-	MaterialParameterCurves.Empty();
+	EventCurves.Reset();
+	MorphTargetCurves.Reset();
+	MaterialParameterCurves.Reset();
 
 	for(int32 CurveId=0; CurveId<InCurve.UIDList.Num(); ++CurveId)
 	{
@@ -1253,7 +1253,7 @@ void UAnimInstance::UpdateCurves(const FBlendedCurve& InCurve)
 void UAnimInstance::TriggerAnimNotifies(float DeltaSeconds)
 {
 	SCOPE_CYCLE_COUNTER(STAT_AnimTriggerAnimNotifies);
-	USkeletalMeshComponent * SkelMeshComp = GetSkelMeshComponent();
+	USkeletalMeshComponent* SkelMeshComp = GetSkelMeshComponent();
 
 	// Array that will replace the 'ActiveAnimNotifyState' at the end of this function.
 	TArray<FAnimNotifyEvent> NewActiveAnimNotifyState;
@@ -1369,7 +1369,7 @@ void UAnimInstance::GetSlotWeight(FName const & SlotNodeName, float& out_SlotNod
 	// first get all the montage instance weight this slot node has
 	for (int32 Index = 0; Index < MontageInstances.Num(); Index++)
 	{
-		FAnimMontageInstance const * const MontageInstance = MontageInstances[Index];
+		FAnimMontageInstance const* const MontageInstance = MontageInstances[Index];
 		if (MontageInstance && MontageInstance->IsValid() && MontageInstance->Montage->IsValidSlot(SlotNodeName))
 		{
 			NodeTotalWeight += MontageInstance->Weight;
@@ -1393,7 +1393,7 @@ void UAnimInstance::GetSlotWeight(FName const & SlotNodeName, float& out_SlotNod
 		// Re-normalize instance weights.
 		for (int32 Index = 0; Index < MontageInstances.Num(); Index++)
 		{
-			FAnimMontageInstance * MontageInstance = MontageInstances[Index];
+			FAnimMontageInstance* MontageInstance = MontageInstances[Index];
 			if (MontageInstance && MontageInstance->IsValid() && MontageInstance->Montage->IsValidSlot(SlotNodeName))
 			{
 				MontageInstance->Weight /= NodeTotalWeight;
@@ -1441,7 +1441,7 @@ void UAnimInstance::SlotEvaluatePose(FName SlotNodeName, const FCompactPose& Sou
 	{
 		if (EvalState.Montage->IsValidSlot(SlotNodeName))
 		{
-			FAnimTrack const * const AnimTrack = EvalState.Montage->GetAnimationData(SlotNodeName);
+			FAnimTrack const* const AnimTrack = EvalState.Montage->GetAnimationData(SlotNodeName);
 
 			// Find out additive type for pose.
 			EAdditiveAnimationType const AdditiveAnimType = AnimTrack->IsAdditive() 
@@ -1753,7 +1753,7 @@ void UAnimInstance::Montage_Advance(float DeltaSeconds)
 	for (int32 InstanceIndex = 0; InstanceIndex<MontageInstances.Num(); InstanceIndex++)
 	{
 		// should never be NULL
-		FAnimMontageInstance * MontageInstance = MontageInstances[InstanceIndex];
+		FAnimMontageInstance* MontageInstance = MontageInstances[InstanceIndex];
 		ensure(MontageInstance);
 		if (MontageInstance)
 		{
@@ -1785,7 +1785,7 @@ void UAnimInstance::Montage_Advance(float DeltaSeconds)
 #if DO_CHECK && WITH_EDITORONLY_DATA && 0
 			else
 			{
-				FAnimMontageInstance * AnimMontageInstance = MontageInstances(I);
+				FAnimMontageInstance* AnimMontageInstance = MontageInstances(I);
 				// print blending time and weight and montage name
 				UE_LOG(LogAnimation, Warning, TEXT("%d. Montage (%s), DesiredWeight(%0.2f), CurrentWeight(%0.2f), BlendingTime(%0.2f)"), 
 					I+1, *AnimMontageInstance->Montage->GetName(), AnimMontageInstance->DesiredWeight, AnimMontageInstance->Weight,  
@@ -1883,7 +1883,7 @@ float UAnimInstance::PlaySlotAnimation(UAnimSequenceBase* Asset, FName SlotNodeN
 	}
 
 	// now play
-	UAnimMontage * NewMontage = NewObject<UAnimMontage>();
+	UAnimMontage* NewMontage = NewObject<UAnimMontage>();
 	NewMontage->SetSkeleton(AssetSkeleton);
 
 	// add new track
@@ -1938,7 +1938,7 @@ UAnimMontage* UAnimInstance::PlaySlotAnimationAsDynamicMontage(UAnimSequenceBase
 	}
 
 	// now play
-	UAnimMontage * NewMontage = NewObject<UAnimMontage>();
+	UAnimMontage* NewMontage = NewObject<UAnimMontage>();
 	NewMontage->SetSkeleton(AssetSkeleton);
 
 	// add new track
@@ -1979,17 +1979,17 @@ void UAnimInstance::StopSlotAnimation(float InBlendOutTime, FName SlotNodeName)
 		for (int32 InstanceIndex = 0; InstanceIndex < MontageInstances.Num(); InstanceIndex++)
 		{
 			// check if this is playing
-			FAnimMontageInstance * MontageInstance = MontageInstances[InstanceIndex];
+			FAnimMontageInstance* MontageInstance = MontageInstances[InstanceIndex];
 			// make sure what is active right now is transient that we created by request
 			if (MontageInstance && MontageInstance->IsActive() && MontageInstance->IsPlaying())
 			{
-				UAnimMontage * CurMontage = MontageInstance->Montage;
+				UAnimMontage* CurMontage = MontageInstance->Montage;
 				if (CurMontage && CurMontage->GetOuter() == GetTransientPackage())
 				{
 					// Check each track, in practice there should only be one on these
 					for (int32 SlotTrackIndex = 0; SlotTrackIndex < CurMontage->SlotAnimTracks.Num(); SlotTrackIndex++)
 					{
-						const FSlotAnimationTrack * AnimTrack = &CurMontage->SlotAnimTracks[SlotTrackIndex];
+						const FSlotAnimationTrack* AnimTrack = &CurMontage->SlotAnimTracks[SlotTrackIndex];
 						if (AnimTrack && AnimTrack->SlotName == SlotNodeName)
 						{
 							// Found it
@@ -2013,11 +2013,11 @@ bool UAnimInstance::IsPlayingSlotAnimation(UAnimSequenceBase* Asset, FName SlotN
 	for (int32 InstanceIndex = 0; InstanceIndex < MontageInstances.Num(); InstanceIndex++)
 	{
 		// check if this is playing
-		FAnimMontageInstance * MontageInstance = MontageInstances[InstanceIndex];
+		FAnimMontageInstance* MontageInstance = MontageInstances[InstanceIndex];
 		// make sure what is active right now is transient that we created by request
 		if (MontageInstance && MontageInstance->IsActive() && MontageInstance->IsPlaying())
 		{
-			UAnimMontage * CurMontage = MontageInstance->Montage;
+			UAnimMontage* CurMontage = MontageInstance->Montage;
 			if (CurMontage && CurMontage->GetOuter() == GetTransientPackage())
 			{
 				const FAnimTrack* AnimTrack = CurMontage->GetAnimationData(SlotNodeName);
@@ -2034,7 +2034,7 @@ bool UAnimInstance::IsPlayingSlotAnimation(UAnimSequenceBase* Asset, FName SlotN
 }
 
 /** Play a Montage. Returns Length of Montage in seconds. Returns 0.f if failed to play. */
-float UAnimInstance::Montage_Play(UAnimMontage * MontageToPlay, float InPlayRate/*= 1.f*/)
+float UAnimInstance::Montage_Play(UAnimMontage* MontageToPlay, float InPlayRate/*= 1.f*/)
 {
 	if (MontageToPlay && (MontageToPlay->SequenceLength > 0.f) && MontageToPlay->HasValidSlotSetup())
 	{
@@ -2054,7 +2054,7 @@ float UAnimInstance::Montage_Play(UAnimMontage * MontageToPlay, float InPlayRate
 				}
 			}
 
-			FAnimMontageInstance * NewInstance = new FAnimMontageInstance(this);
+			FAnimMontageInstance* NewInstance = new FAnimMontageInstance(this);
 			check(NewInstance);
 
 			NewInstance->Initialize(MontageToPlay);
@@ -2082,11 +2082,11 @@ float UAnimInstance::Montage_Play(UAnimMontage * MontageToPlay, float InPlayRate
 	return 0.f;
 }
 
-void UAnimInstance::Montage_Stop(float InBlendOutTime, UAnimMontage * Montage)
+void UAnimInstance::Montage_Stop(float InBlendOutTime, UAnimMontage* Montage)
 {
 	if (Montage)
 	{
-		FAnimMontageInstance * MontageInstance = GetActiveInstanceForMontage(*Montage);
+		FAnimMontageInstance* MontageInstance = GetActiveInstanceForMontage(*Montage);
 		if (MontageInstance)
 		{
 			MontageInstance->Stop(InBlendOutTime);
@@ -2097,7 +2097,7 @@ void UAnimInstance::Montage_Stop(float InBlendOutTime, UAnimMontage * Montage)
 		// If no Montage reference, do it on all active ones.
 		for (int32 InstanceIndex = 0; InstanceIndex < MontageInstances.Num(); InstanceIndex++)
 		{
-			FAnimMontageInstance * MontageInstance = MontageInstances[InstanceIndex];
+			FAnimMontageInstance* MontageInstance = MontageInstances[InstanceIndex];
 			if (MontageInstance && MontageInstance->IsActive())
 			{
 				MontageInstance->Stop(InBlendOutTime);
@@ -2106,11 +2106,11 @@ void UAnimInstance::Montage_Stop(float InBlendOutTime, UAnimMontage * Montage)
 	}
 }
 
-void UAnimInstance::Montage_Pause(UAnimMontage * Montage)
+void UAnimInstance::Montage_Pause(UAnimMontage* Montage)
 {
 	if (Montage)
 	{
-		FAnimMontageInstance * MontageInstance = GetActiveInstanceForMontage(*Montage);
+		FAnimMontageInstance* MontageInstance = GetActiveInstanceForMontage(*Montage);
 		if (MontageInstance)
 		{
 			MontageInstance->Pause();
@@ -2121,7 +2121,7 @@ void UAnimInstance::Montage_Pause(UAnimMontage * Montage)
 		// If no Montage reference, do it on all active ones.
 		for (int32 InstanceIndex = 0; InstanceIndex < MontageInstances.Num(); InstanceIndex++)
 		{
-			FAnimMontageInstance * MontageInstance = MontageInstances[InstanceIndex];
+			FAnimMontageInstance* MontageInstance = MontageInstances[InstanceIndex];
 			if (MontageInstance && MontageInstance->IsActive())
 			{
 				MontageInstance->Pause();
@@ -2130,11 +2130,11 @@ void UAnimInstance::Montage_Pause(UAnimMontage * Montage)
 	}
 }
 
-void UAnimInstance::Montage_JumpToSection(FName SectionName, UAnimMontage * Montage)
+void UAnimInstance::Montage_JumpToSection(FName SectionName, UAnimMontage* Montage)
 {
 	if (Montage)
 	{
-		FAnimMontageInstance * MontageInstance = GetActiveInstanceForMontage(*Montage);
+		FAnimMontageInstance* MontageInstance = GetActiveInstanceForMontage(*Montage);
 		if (MontageInstance)
 		{
 			bool const bEndOfSection = (MontageInstance->GetPlayRate() < 0.f);
@@ -2146,7 +2146,7 @@ void UAnimInstance::Montage_JumpToSection(FName SectionName, UAnimMontage * Mont
 		// If no Montage reference, do it on all active ones.
 		for (int32 InstanceIndex = 0; InstanceIndex < MontageInstances.Num(); InstanceIndex++)
 		{
-			FAnimMontageInstance * MontageInstance = MontageInstances[InstanceIndex];
+			FAnimMontageInstance* MontageInstance = MontageInstances[InstanceIndex];
 			if (MontageInstance && MontageInstance->IsActive())
 			{
 				bool const bEndOfSection = (MontageInstance->GetPlayRate() < 0.f);
@@ -2156,11 +2156,11 @@ void UAnimInstance::Montage_JumpToSection(FName SectionName, UAnimMontage * Mont
 	}
 }
 
-void UAnimInstance::Montage_JumpToSectionsEnd(FName SectionName, UAnimMontage * Montage)
+void UAnimInstance::Montage_JumpToSectionsEnd(FName SectionName, UAnimMontage* Montage)
 {
 	if (Montage)
 	{
-		FAnimMontageInstance * MontageInstance = GetActiveInstanceForMontage(*Montage);
+		FAnimMontageInstance* MontageInstance = GetActiveInstanceForMontage(*Montage);
 		if (MontageInstance)
 		{
 			bool const bEndOfSection = (MontageInstance->GetPlayRate() >= 0.f);
@@ -2172,7 +2172,7 @@ void UAnimInstance::Montage_JumpToSectionsEnd(FName SectionName, UAnimMontage *
 		// If no Montage reference, do it on all active ones.
 		for (int32 InstanceIndex = 0; InstanceIndex < MontageInstances.Num(); InstanceIndex++)
 		{
-			FAnimMontageInstance * MontageInstance = MontageInstances[InstanceIndex];
+			FAnimMontageInstance* MontageInstance = MontageInstances[InstanceIndex];
 			if (MontageInstance && MontageInstance->IsActive())
 			{
 				bool const bEndOfSection = (MontageInstance->GetPlayRate() >= 0.f);
@@ -2182,11 +2182,11 @@ void UAnimInstance::Montage_JumpToSectionsEnd(FName SectionName, UAnimMontage *
 	}
 }
 
-void UAnimInstance::Montage_SetNextSection(FName SectionNameToChange, FName NextSection, UAnimMontage * Montage)
+void UAnimInstance::Montage_SetNextSection(FName SectionNameToChange, FName NextSection, UAnimMontage* Montage)
 {
 	if (Montage)
 	{
-		FAnimMontageInstance * MontageInstance = GetActiveInstanceForMontage(*Montage);
+		FAnimMontageInstance* MontageInstance = GetActiveInstanceForMontage(*Montage);
 		if (MontageInstance)
 		{
 			MontageInstance->SetNextSectionName(SectionNameToChange, NextSection);
@@ -2199,7 +2199,7 @@ void UAnimInstance::Montage_SetNextSection(FName SectionNameToChange, FName Next
 		// If no Montage reference, do it on all active ones.
 		for (int32 InstanceIndex = 0; InstanceIndex < MontageInstances.Num(); InstanceIndex++)
 		{
-			FAnimMontageInstance * MontageInstance = MontageInstances[InstanceIndex];
+			FAnimMontageInstance* MontageInstance = MontageInstances[InstanceIndex];
 			if (MontageInstance && MontageInstance->IsActive())
 			{
 				MontageInstance->SetNextSectionName(SectionNameToChange, NextSection);
@@ -2214,11 +2214,11 @@ void UAnimInstance::Montage_SetNextSection(FName SectionNameToChange, FName Next
 	}
 }
 
-void UAnimInstance::Montage_SetPlayRate(UAnimMontage * Montage, float NewPlayRate)
+void UAnimInstance::Montage_SetPlayRate(UAnimMontage* Montage, float NewPlayRate)
 {
 	if (Montage)
 	{
-		FAnimMontageInstance * MontageInstance = GetActiveInstanceForMontage(*Montage);
+		FAnimMontageInstance* MontageInstance = GetActiveInstanceForMontage(*Montage);
 		if (MontageInstance)
 		{
 			MontageInstance->SetPlayRate(NewPlayRate);
@@ -2229,7 +2229,7 @@ void UAnimInstance::Montage_SetPlayRate(UAnimMontage * Montage, float NewPlayRat
 		// If no Montage reference, do it on all active ones.
 		for (int32 InstanceIndex = 0; InstanceIndex < MontageInstances.Num(); InstanceIndex++)
 		{
-			FAnimMontageInstance * MontageInstance = MontageInstances[InstanceIndex];
+			FAnimMontageInstance* MontageInstance = MontageInstances[InstanceIndex];
 			if (MontageInstance && MontageInstance->IsActive())
 			{
 				MontageInstance->SetPlayRate(NewPlayRate);
@@ -2238,11 +2238,11 @@ void UAnimInstance::Montage_SetPlayRate(UAnimMontage * Montage, float NewPlayRat
 	}
 }
 
-bool UAnimInstance::Montage_IsActive(UAnimMontage * Montage)
+bool UAnimInstance::Montage_IsActive(UAnimMontage* Montage)
 {
 	if (Montage)
 	{
-		FAnimMontageInstance * MontageInstance = GetActiveInstanceForMontage(*Montage);
+		FAnimMontageInstance* MontageInstance = GetActiveInstanceForMontage(*Montage);
 		if (MontageInstance)
 		{
 			return true;
@@ -2253,7 +2253,7 @@ bool UAnimInstance::Montage_IsActive(UAnimMontage * Montage)
 		// If no Montage reference, return true if there is any active montage.
 		for (int32 InstanceIndex = 0; InstanceIndex < MontageInstances.Num(); InstanceIndex++)
 		{
-			FAnimMontageInstance * MontageInstance = MontageInstances[InstanceIndex];
+			FAnimMontageInstance* MontageInstance = MontageInstances[InstanceIndex];
 			if (MontageInstance && MontageInstance->IsActive())
 			{
 				return true;
@@ -2264,11 +2264,11 @@ bool UAnimInstance::Montage_IsActive(UAnimMontage * Montage)
 	return false;
 }
 
-bool UAnimInstance::Montage_IsPlaying(UAnimMontage * Montage)
+bool UAnimInstance::Montage_IsPlaying(UAnimMontage* Montage)
 {
 	if (Montage)
 	{
-		FAnimMontageInstance * MontageInstance = GetActiveInstanceForMontage(*Montage);
+		FAnimMontageInstance* MontageInstance = GetActiveInstanceForMontage(*Montage);
 		if (MontageInstance)
 		{
 			return MontageInstance->IsPlaying();
@@ -2279,7 +2279,7 @@ bool UAnimInstance::Montage_IsPlaying(UAnimMontage * Montage)
 		// If no Montage reference, return true if there is any active playing montage.
 		for (int32 InstanceIndex = 0; InstanceIndex < MontageInstances.Num(); InstanceIndex++)
 		{
-			FAnimMontageInstance * MontageInstance = MontageInstances[InstanceIndex];
+			FAnimMontageInstance* MontageInstance = MontageInstances[InstanceIndex];
 			if (MontageInstance && MontageInstance->IsActive() && MontageInstance->IsPlaying())
 			{
 				return true;
@@ -2294,7 +2294,7 @@ FName UAnimInstance::Montage_GetCurrentSection(UAnimMontage* Montage)
 {
 	if (Montage)
 	{
-		FAnimMontageInstance * MontageInstance = GetActiveInstanceForMontage(*Montage);
+		FAnimMontageInstance* MontageInstance = GetActiveInstanceForMontage(*Montage);
 		if (MontageInstance)
 		{
 			return MontageInstance->GetCurrentSection();
@@ -2305,7 +2305,7 @@ FName UAnimInstance::Montage_GetCurrentSection(UAnimMontage* Montage)
 		// If no Montage reference, get first active one.
 		for (int32 InstanceIndex = 0; InstanceIndex < MontageInstances.Num(); InstanceIndex++)
 		{
-			FAnimMontageInstance * MontageInstance = MontageInstances[InstanceIndex];
+			FAnimMontageInstance* MontageInstance = MontageInstances[InstanceIndex];
 			if (MontageInstance && MontageInstance->IsActive())
 			{
 				return MontageInstance->GetCurrentSection();
@@ -2320,7 +2320,7 @@ void UAnimInstance::Montage_SetEndDelegate(FOnMontageEnded & InOnMontageEnded, U
 {
 	if (Montage)
 	{
-		FAnimMontageInstance * MontageInstance = GetActiveInstanceForMontage(*Montage);
+		FAnimMontageInstance* MontageInstance = GetActiveInstanceForMontage(*Montage);
 		if (MontageInstance)
 		{
 			MontageInstance->OnMontageEnded = InOnMontageEnded;
@@ -2331,7 +2331,7 @@ void UAnimInstance::Montage_SetEndDelegate(FOnMontageEnded & InOnMontageEnded, U
 		// If no Montage reference, do it on all active ones.
 		for (int32 InstanceIndex = 0; InstanceIndex < MontageInstances.Num(); InstanceIndex++)
 		{
-			FAnimMontageInstance * MontageInstance = MontageInstances[InstanceIndex];
+			FAnimMontageInstance* MontageInstance = MontageInstances[InstanceIndex];
 			if (MontageInstance && MontageInstance->IsActive())
 			{
 				MontageInstance->OnMontageEnded = InOnMontageEnded;
@@ -2344,7 +2344,7 @@ void UAnimInstance::Montage_SetBlendingOutDelegate(FOnMontageBlendingOutStarted
 {
 	if (Montage)
 	{
-		FAnimMontageInstance * MontageInstance = GetActiveInstanceForMontage(*Montage);
+		FAnimMontageInstance* MontageInstance = GetActiveInstanceForMontage(*Montage);
 		if (MontageInstance)
 		{
 			MontageInstance->OnMontageBlendingOutStarted = InOnMontageBlendingOut;
@@ -2355,7 +2355,7 @@ void UAnimInstance::Montage_SetBlendingOutDelegate(FOnMontageBlendingOutStarted
 		// If no Montage reference, do it on all active ones.
 		for (int32 InstanceIndex = 0; InstanceIndex < MontageInstances.Num(); InstanceIndex++)
 		{
-			FAnimMontageInstance * MontageInstance = MontageInstances[InstanceIndex];
+			FAnimMontageInstance* MontageInstance = MontageInstances[InstanceIndex];
 			if (MontageInstance && MontageInstance->IsActive())
 			{
 				MontageInstance->OnMontageBlendingOutStarted = InOnMontageBlendingOut;
@@ -2364,11 +2364,11 @@ void UAnimInstance::Montage_SetBlendingOutDelegate(FOnMontageBlendingOutStarted
 	}
 }
 
-FOnMontageBlendingOutStarted * UAnimInstance::Montage_GetBlendingOutDelegate(UAnimMontage * Montage)
+FOnMontageBlendingOutStarted* UAnimInstance::Montage_GetBlendingOutDelegate(UAnimMontage* Montage)
 {
 	if (Montage)
 	{
-		FAnimMontageInstance * MontageInstance = GetActiveInstanceForMontage(*Montage);
+		FAnimMontageInstance* MontageInstance = GetActiveInstanceForMontage(*Montage);
 		if (MontageInstance)
 		{
 			return &MontageInstance->OnMontageBlendingOutStarted;
@@ -2379,7 +2379,7 @@ FOnMontageBlendingOutStarted * UAnimInstance::Montage_GetBlendingOutDelegate(UAn
 		// If no Montage reference, use first active one found.
 		for (int32 InstanceIndex = 0; InstanceIndex < MontageInstances.Num(); InstanceIndex++)
 		{
-			FAnimMontageInstance * MontageInstance = MontageInstances[InstanceIndex];
+			FAnimMontageInstance* MontageInstance = MontageInstances[InstanceIndex];
 			if (MontageInstance && MontageInstance->IsActive())
 			{
 				return &MontageInstance->OnMontageBlendingOutStarted;
@@ -2394,7 +2394,7 @@ void UAnimInstance::Montage_SetPosition(UAnimMontage* Montage, float NewPosition
 {
 	if (Montage)
 	{
-		FAnimMontageInstance * MontageInstance = GetActiveInstanceForMontage(*Montage);
+		FAnimMontageInstance* MontageInstance = GetActiveInstanceForMontage(*Montage);
 		if (MontageInstance)
 		{
 			MontageInstance->SetPosition(NewPosition);
@@ -2405,7 +2405,7 @@ void UAnimInstance::Montage_SetPosition(UAnimMontage* Montage, float NewPosition
 		// If no Montage reference, do it on all active ones.
 		for (int32 InstanceIndex = 0; InstanceIndex < MontageInstances.Num(); InstanceIndex++)
 		{
-			FAnimMontageInstance * MontageInstance = MontageInstances[InstanceIndex];
+			FAnimMontageInstance* MontageInstance = MontageInstances[InstanceIndex];
 			if (MontageInstance && MontageInstance->IsActive())
 			{
 				MontageInstance->SetPosition(NewPosition);
@@ -2418,7 +2418,7 @@ float UAnimInstance::Montage_GetPosition(UAnimMontage* Montage)
 {
 	if (Montage)
 	{
-		FAnimMontageInstance * MontageInstance = GetActiveInstanceForMontage(*Montage);
+		FAnimMontageInstance* MontageInstance = GetActiveInstanceForMontage(*Montage);
 		if (MontageInstance)
 		{
 			return MontageInstance->GetPosition();
@@ -2429,7 +2429,7 @@ float UAnimInstance::Montage_GetPosition(UAnimMontage* Montage)
 		// If no Montage reference, use first active one found.
 		for (int32 InstanceIndex = 0; InstanceIndex < MontageInstances.Num(); InstanceIndex++)
 		{
-			FAnimMontageInstance * MontageInstance = MontageInstances[InstanceIndex];
+			FAnimMontageInstance* MontageInstance = MontageInstances[InstanceIndex];
 			if (MontageInstance && MontageInstance->IsActive())
 			{
 				return MontageInstance->GetPosition();
@@ -2444,7 +2444,7 @@ bool UAnimInstance::Montage_GetIsStopped(UAnimMontage* Montage)
 {
 	if (Montage)
 	{
-		FAnimMontageInstance * MontageInstance = GetActiveInstanceForMontage(*Montage);
+		FAnimMontageInstance* MontageInstance = GetActiveInstanceForMontage(*Montage);
 		return (!MontageInstance); // Not active == Stopped.
 	}
 	return true;
@@ -2454,7 +2454,7 @@ float UAnimInstance::Montage_GetBlendTime(UAnimMontage* Montage)
 {
 	if (Montage)
 	{
-		FAnimMontageInstance * MontageInstance = GetActiveInstanceForMontage(*Montage);
+		FAnimMontageInstance* MontageInstance = GetActiveInstanceForMontage(*Montage);
 		if (MontageInstance)
 		{
 			return MontageInstance->BlendTime;
@@ -2465,7 +2465,7 @@ float UAnimInstance::Montage_GetBlendTime(UAnimMontage* Montage)
 		// If no Montage reference, use first active one found.
 		for (int32 InstanceIndex = 0; InstanceIndex < MontageInstances.Num(); InstanceIndex++)
 		{
-			FAnimMontageInstance * MontageInstance = MontageInstances[InstanceIndex];
+			FAnimMontageInstance* MontageInstance = MontageInstances[InstanceIndex];
 			if (MontageInstance && MontageInstance->IsActive())
 			{
 				return MontageInstance->BlendTime;
@@ -2480,7 +2480,7 @@ float UAnimInstance::Montage_GetPlayRate(UAnimMontage* Montage)
 {
 	if (Montage)
 	{
-		FAnimMontageInstance * MontageInstance = GetActiveInstanceForMontage(*Montage);
+		FAnimMontageInstance* MontageInstance = GetActiveInstanceForMontage(*Montage);
 		if (MontageInstance)
 		{
 			return MontageInstance->GetPlayRate();
@@ -2491,7 +2491,7 @@ float UAnimInstance::Montage_GetPlayRate(UAnimMontage* Montage)
 		// If no Montage reference, use first active one found.
 		for (int32 InstanceIndex = 0; InstanceIndex < MontageInstances.Num(); InstanceIndex++)
 		{
-			FAnimMontageInstance * MontageInstance = MontageInstances[InstanceIndex];
+			FAnimMontageInstance* MontageInstance = MontageInstances[InstanceIndex];
 			if (MontageInstance && MontageInstance->IsActive())
 			{
 				return MontageInstance->GetPlayRate();
@@ -2502,11 +2502,11 @@ float UAnimInstance::Montage_GetPlayRate(UAnimMontage* Montage)
 	return 0.f;
 }
 
-int32 UAnimInstance::Montage_GetNextSectionID(UAnimMontage const * const Montage, int32 const & CurrentSectionID) const
+int32 UAnimInstance::Montage_GetNextSectionID(UAnimMontage const* const Montage, int32 const & CurrentSectionID) const
 {
 	if (Montage)
 	{
-		FAnimMontageInstance * MontageInstance = GetActiveInstanceForMontage(*Montage);
+		FAnimMontageInstance* MontageInstance = GetActiveInstanceForMontage(*Montage);
 		if (MontageInstance)
 		{
 			return (CurrentSectionID < MontageInstance->NextSections.Num()) ? MontageInstance->NextSections[CurrentSectionID] : INDEX_NONE;
@@ -2517,7 +2517,7 @@ int32 UAnimInstance::Montage_GetNextSectionID(UAnimMontage const * const Montage
 		// If no Montage reference, use first active one found.
 		for (int32 InstanceIndex = 0; InstanceIndex < MontageInstances.Num(); InstanceIndex++)
 		{
-			FAnimMontageInstance * MontageInstance = MontageInstances[InstanceIndex];
+			FAnimMontageInstance* MontageInstance = MontageInstances[InstanceIndex];
 			if (MontageInstance && MontageInstance->IsActive())
 			{
 				return (CurrentSectionID < MontageInstance->NextSections.Num()) ? MontageInstance->NextSections[CurrentSectionID] : INDEX_NONE;
@@ -2528,13 +2528,13 @@ int32 UAnimInstance::Montage_GetNextSectionID(UAnimMontage const * const Montage
 	return INDEX_NONE;
 }
 
-UAnimMontage * UAnimInstance::GetCurrentActiveMontage()
+UAnimMontage* UAnimInstance::GetCurrentActiveMontage()
 {
 	// Start from end, as most recent instances are added at the end of the queue.
 	int32 const NumInstances = MontageInstances.Num();
 	for (int32 InstanceIndex = NumInstances - 1; InstanceIndex >= 0; InstanceIndex--)
 	{
-		FAnimMontageInstance * MontageInstance = MontageInstances[InstanceIndex];
+		FAnimMontageInstance* MontageInstance = MontageInstances[InstanceIndex];
 		if (MontageInstance && MontageInstance->IsActive())
 		{
 			return MontageInstance->Montage;
@@ -2550,7 +2550,7 @@ FAnimMontageInstance* UAnimInstance::GetActiveMontageInstance()
 	int32 const NumInstances = MontageInstances.Num();
 	for (int32 InstanceIndex = NumInstances - 1; InstanceIndex >= 0; InstanceIndex--)
 	{
-		FAnimMontageInstance * MontageInstance = MontageInstances[InstanceIndex];
+		FAnimMontageInstance* MontageInstance = MontageInstances[InstanceIndex];
 		if (MontageInstance && MontageInstance->IsActive())
 		{
 			return MontageInstance;
@@ -2572,7 +2572,7 @@ void UAnimInstance::StopAllMontagesByGroupName(FName InGroupName, float BlendOut
 {
 	for (int32 InstanceIndex = MontageInstances.Num() - 1; InstanceIndex >= 0; InstanceIndex--)
 	{
-		FAnimMontageInstance * MontageInstance = MontageInstances[InstanceIndex];
+		FAnimMontageInstance* MontageInstance = MontageInstances[InstanceIndex];
 		if (MontageInstance && MontageInstance->IsActive() && MontageInstance->Montage && (MontageInstance->Montage->GetGroupName() == InGroupName))
 		{
 			MontageInstances[InstanceIndex]->Stop(BlendOutTime, true);
@@ -2582,7 +2582,7 @@ void UAnimInstance::StopAllMontagesByGroupName(FName InGroupName, float BlendOut
 
 void UAnimInstance::OnMontageInstanceStopped(FAnimMontageInstance & StoppedMontageInstance)
 {
-	UAnimMontage * MontageStopped = StoppedMontageInstance.Montage;
+	UAnimMontage* MontageStopped = StoppedMontageInstance.Montage;
 	ensure(MontageStopped != NULL);
 
 	// Remove instance for Active List.
@@ -2599,13 +2599,13 @@ void UAnimInstance::OnMontageInstanceStopped(FAnimMontageInstance & StoppedMonta
 	}
 }
 
-FAnimMontageInstance * UAnimInstance::GetActiveInstanceForMontage(UAnimMontage const & Montage) const
+FAnimMontageInstance* UAnimInstance::GetActiveInstanceForMontage(UAnimMontage const & Montage) const
 {
-	FAnimMontageInstance * const * FoundInstancePtr = ActiveMontagesMap.Find(&Montage);
+	FAnimMontageInstance* const* FoundInstancePtr = ActiveMontagesMap.Find(&Montage);
 	return FoundInstancePtr ? *FoundInstancePtr : NULL;
 }
 
-FAnimMontageInstance * UAnimInstance::GetRootMotionMontageInstance() const
+FAnimMontageInstance* UAnimInstance::GetRootMotionMontageInstance() const
 {
 	return RootMotionMontageInstance;
 }
@@ -2630,7 +2630,7 @@ FRootMotionMovementParams UAnimInstance::ConsumeExtractedRootMotion(float Alpha)
 
 void UAnimInstance::SetMorphTarget(FName MorphTargetName, float Value)
 {
-	USkeletalMeshComponent * Component = GetOwningComponent();
+	USkeletalMeshComponent* Component = GetOwningComponent();
 	if (Component)
 	{
 		Component->SetMorphTarget(MorphTargetName, Value);
@@ -2639,7 +2639,7 @@ void UAnimInstance::SetMorphTarget(FName MorphTargetName, float Value)
 
 void UAnimInstance::ClearMorphTargets()
 {
-	USkeletalMeshComponent * Component = GetOwningComponent();
+	USkeletalMeshComponent* Component = GetOwningComponent();
 	if (Component)
 	{
 		Component->ClearMorphTargets();
@@ -2743,7 +2743,7 @@ bool UAnimInstance::HasMarkerBeenHitThisFrame(FName SyncGroup, FName MarkerName)
 
 void UAnimInstance::UpdateMontageEvaluationData()
 {
-	MontageEvaluationData.Empty(MontageInstances.Num());
+	MontageEvaluationData.Reset(MontageInstances.Num());
 	for (FAnimMontageInstance* MontageInstance : MontageInstances)
 	{
 		if (MontageInstance->Montage && MontageInstance->Weight > ZERO_ANIMWEIGHT_THRESH)
