Animation: Prevented unnecessary allocations in FillSpaceBases
[INTEGRATE] Change 2663696 by [REDACTED_PERSON][REDACTED_HANDLE] on 2015/08/20 19:40:14

diff --git a/Engine/Source/Runtime/Engine/Private/Components/PoseableMeshComponent.cpp b/Engine/Source/Runtime/Engine/Private/Components/PoseableMeshComponent.cpp
--- a/Engine/Source/Runtime/Engine/Private/Components/PoseableMeshComponent.cpp
+++ b/Engine/Source/Runtime/Engine/Private/Components/PoseableMeshComponent.cpp
@@ -84,9 +84,9 @@ void UPoseableMeshComponent::FillSpaceBases()
 
 	const int32 NumBones = LocalAtoms.Num();
 
-#if (UE_BUILD_DEBUG || UE_BUILD_DEVELOPMENT)
+#if DO_GUARD_SLOW
 	/** Keep track of which bones have been processed for fast look up */
-	TArray<uint8> BoneProcessed;
+	TArray<uint8, TInlineAllocator<256>> BoneProcessed;
 	BoneProcessed.AddZeroed(NumBones);
 #endif
 	// Build in 3 passes.
@@ -94,7 +94,7 @@ void UPoseableMeshComponent::FillSpaceBases()
 	FTransform* SpaceBasesData = GetEditableSpaceBases().GetData();
 	
 	GetEditableSpaceBases()[0] = LocalAtoms[0];
-#if (UE_BUILD_DEBUG || UE_BUILD_DEVELOPMENT)
+#if DO_GUARD_SLOW
 	BoneProcessed[0] = 1;
 #endif
 
@@ -102,7 +102,7 @@ void UPoseableMeshComponent::FillSpaceBases()
 	{
 		FPlatformMisc::Prefetch(SpaceBasesData + BoneIndex);
 
-#if (UE_BUILD_DEBUG || UE_BUILD_DEVELOPMENT)
+#if DO_GUARD_SLOW
 		// Mark bone as processed
 		BoneProcessed[BoneIndex] = 1;
 #endif
@@ -110,7 +110,7 @@ void UPoseableMeshComponent::FillSpaceBases()
 		const int32 ParentIndex = SkeletalMesh->RefSkeleton.GetParentIndex(BoneIndex);
 		FPlatformMisc::Prefetch(SpaceBasesData + ParentIndex);
 
-#if (UE_BUILD_DEBUG || UE_BUILD_DEVELOPMENT)
+#if DO_GUARD_SLOW
 		// Check the precondition that Parents occur before Children in the RequiredBones array.
 		checkSlow(BoneProcessed[ParentIndex] == 1);
 #endif
diff --git a/Engine/Source/Runtime/Engine/Private/Components/SkeletalMeshComponent.cpp b/Engine/Source/Runtime/Engine/Private/Components/SkeletalMeshComponent.cpp
--- a/Engine/Source/Runtime/Engine/Private/Components/SkeletalMeshComponent.cpp
+++ b/Engine/Source/Runtime/Engine/Private/Components/SkeletalMeshComponent.cpp
@@ -702,9 +702,9 @@ void USkeletalMeshComponent::FillSpaceBases(const USkeletalMesh* InSkeletalMesh,
 
 	const int32 NumBones = SourceAtoms.Num();
 
-#if (UE_BUILD_DEBUG || UE_BUILD_DEVELOPMENT)
+#if DO_GUARD_SLOW
 	/** Keep track of which bones have been processed for fast look up */
-	TArray<uint8> BoneProcessed;
+	TArray<uint8, TInlineAllocator<256>> BoneProcessed;
 	BoneProcessed.AddZeroed(NumBones);
 #endif
 
@@ -716,7 +716,7 @@ void USkeletalMeshComponent::FillSpaceBases(const USkeletalMesh* InSkeletalMesh,
 		check(FillSpaceBasesRequiredBones[0] == 0);
 		DestSpaceBases[0] = SourceAtoms[0];
 
-#if (UE_BUILD_DEBUG || UE_BUILD_DEVELOPMENT)
+#if DO_GUARD_SLOW
 		// Mark bone as processed
 		BoneProcessed[0] = 1;
 #endif
@@ -729,7 +729,7 @@ void USkeletalMeshComponent::FillSpaceBases(const USkeletalMesh* InSkeletalMesh,
 
 		FPlatformMisc::Prefetch(SpaceBase);
 
-#if (UE_BUILD_DEBUG || UE_BUILD_DEVELOPMENT)
+#if DO_GUARD_SLOW
 		// Mark bone as processed
 		BoneProcessed[BoneIndex] = 1;
 #endif
@@ -738,7 +738,7 @@ void USkeletalMeshComponent::FillSpaceBases(const USkeletalMesh* InSkeletalMesh,
 		FTransform* ParentSpaceBase = SpaceBasesData + ParentIndex;
 		FPlatformMisc::Prefetch(ParentSpaceBase);
 
-#if (UE_BUILD_DEBUG || UE_BUILD_DEVELOPMENT)
+#if DO_GUARD_SLOW
 		// Check the precondition that Parents occur before Children in the RequiredBones array.
 		checkSlow(BoneProcessed[ParentIndex] == 1);
 #endif
