Added KeyState accessor function to UPlayerInput.

--------
Integrated using branch UE4-To-UE4-Orion (reversed) of change#2663781 by [REDACTED_PERSON] on 2015/08/20 20:44:50.

diff --git a/Engine/Source/Runtime/Engine/Classes/GameFramework/PlayerInput.h b/Engine/Source/Runtime/Engine/Classes/GameFramework/PlayerInput.h
--- a/Engine/Source/Runtime/Engine/Classes/GameFramework/PlayerInput.h
+++ b/Engine/Source/Runtime/Engine/Classes/GameFramework/PlayerInput.h
@@ -446,6 +446,9 @@ public:
 	 */
 	void DisplayDebug(class UCanvas* Canvas, const FDebugDisplayInfo& DebugDisplay, float& YL, float& YPos);
 
+	/** @return key state of the InKey */
+	FKeyState* GetKeyState(FKey InKey) { return KeyStateMap.Find(InKey); }
+
 	/** @return true if InKey is currently held */
 	bool IsPressed( FKey InKey ) const;
 	
