UE4 - Compile fix

diff --git a/Engine/Source/Runtime/Engine/Private/ContentStreaming.cpp b/Engine/Source/Runtime/Engine/Private/ContentStreaming.cpp
--- a/Engine/Source/Runtime/Engine/Private/ContentStreaming.cpp
+++ b/Engine/Source/Runtime/Engine/Private/ContentStreaming.cpp
@@ -3376,10 +3376,10 @@ void FStreamingManagerTexture::StreamTextures( bool bProcessEverything )
 		AvailableNow = GStreamMemoryTracker.CalcAvailableNow( Stats.ComputeAvailableMemorySize(), MemoryMargin );
 		AvailableLater = GStreamMemoryTracker.CalcAvailableLater( Stats.ComputeAvailableMemorySize(), MemoryMargin );
 
-		int64 NonStreamingUsage = Stats.AllocatedMemorySize - ThreadStats.TotalResidentSize - int64(ThreadStats.PendingStreamInSize) - TempMemoryUsed;
-		int64 MemoryBudget = Stats.TexturePoolSize - NonStreamingUsage - MemoryMargin;
-		int64 LocalMemoryOverBudget = FMath::Max(ThreadStats.TotalRequiredSize - MemoryBudget, 0LL);
-		MemoryOverBudget = FMath::Max(LocalMemoryOverBudget, 0LL);
+		STAT( int64 NonStreamingUsage = Stats.AllocatedMemorySize - ThreadStats.TotalResidentSize - int64(ThreadStats.PendingStreamInSize) - TempMemoryUsed );
+		STAT( int64 MemoryBudget = Stats.TexturePoolSize - NonStreamingUsage - MemoryMargin );
+		STAT( int64 LocalMemoryOverBudget = FMath::Max(ThreadStats.TotalRequiredSize - MemoryBudget, 0LL) );
+		STAT( MemoryOverBudget = FMath::Max(LocalMemoryOverBudget, 0LL) );
 
 		SET_MEMORY_STAT( STAT_StreamingOverBudget,	MemoryOverBudget);
 		SET_MEMORY_STAT( STAT_StreamingUnderBudget,	FMath::Max(-LocalMemoryOverBudget,0ll) );
