Properly handling 'package async load failed' case in level streaming

diff --git a/Engine/Source/Runtime/Engine/Private/LevelStreaming.cpp b/Engine/Source/Runtime/Engine/Private/LevelStreaming.cpp
--- a/Engine/Source/Runtime/Engine/Private/LevelStreaming.cpp
+++ b/Engine/Source/Runtime/Engine/Private/LevelStreaming.cpp
@@ -644,6 +644,10 @@ void ULevelStreaming::AsyncLevelLoadComplete(const FName& InPackageName, UPackag
 	else
 	{
 		UE_LOG(LogLevelStreaming, Warning, TEXT("Failed to load package '%s'"), *InPackageName.ToString() );
+		
+		bFailedToLoad = true;
+ 		bHasLoadRequestPending = false;
+ 		bShouldBeLoaded = false;
 	}
 
 	// Clean up the world type list and owning world list now that PostLoad has occurred
