Integrate 2664490 from Orion branch

Fix crash in AnimInstance debug info if the RootNode is null

diff --git a/Engine/Source/Runtime/Engine/Private/Animation/AnimInstance.cpp b/Engine/Source/Runtime/Engine/Private/Animation/AnimInstance.cpp
--- a/Engine/Source/Runtime/Engine/Private/Animation/AnimInstance.cpp
+++ b/Engine/Source/Runtime/Engine/Private/Animation/AnimInstance.cpp
@@ -982,7 +982,7 @@ void UAnimInstance::DisplayDebug(class UCanvas* Canvas, const FDebugDisplayInfo&
 		}
 	}
 
-	if (bShowGraph)
+	if (bShowGraph && RootNode != nullptr)
 	{
 		Canvas->SetLinearDrawColor(TextYellow);
 
