Fix for assert when blending between looping and non looping animations
Merged from 2662451

diff --git a/Engine/Source/Runtime/Engine/Private/Animation/AnimSequence.cpp b/Engine/Source/Runtime/Engine/Private/Animation/AnimSequence.cpp
--- a/Engine/Source/Runtime/Engine/Private/Animation/AnimSequence.cpp
+++ b/Engine/Source/Runtime/Engine/Private/Animation/AnimSequence.cpp
@@ -3878,7 +3878,7 @@ void UAnimSequence::AdvanceMarkerPhaseAsFollower(const FMarkerTickContext& Conte
 		{
 			check(NextMarker.MarkerIndex == -1);
 		}
-		else
+		else if (NextMarker.MarkerIndex != -1)
 		{
 			check(AuthoredSyncMarkers[NextMarker.MarkerIndex].MarkerName == End.NextMarkerName);
 		}
@@ -3922,7 +3922,7 @@ void UAnimSequence::AdvanceMarkerPhaseAsFollower(const FMarkerTickContext& Conte
 		{
 			check(PreviousMarker.MarkerIndex == -1);
 		}
-		else
+		else if (PreviousMarker.MarkerIndex != -1)
 		{
 			check(AuthoredSyncMarkers[PreviousMarker.MarkerIndex].MarkerName == End.PreviousMarkerName);
 		}
