[UE-20189] Fix a potential undo/redo crash.

diff --git a/Engine/Source/Runtime/Engine/Private/Actor.cpp b/Engine/Source/Runtime/Engine/Private/Actor.cpp
--- a/Engine/Source/Runtime/Engine/Private/Actor.cpp
+++ b/Engine/Source/Runtime/Engine/Private/Actor.cpp
@@ -200,7 +200,7 @@ void AActor::ResetOwnedComponents()
 		for (auto Component : OwnedComponents)
 		{
 			// Only consider native components
-			if (Component->CreationMethod == EComponentCreationMethod::Native)
+			if (Component && Component->CreationMethod == EComponentCreationMethod::Native)
 			{
 				// Find the property or properties that previously referenced the natively-constructed component.
 				TArray<UObjectProperty*> Properties;
