Fixed crash on platforms that does not support texture streaming

diff --git a/Engine/Source/Runtime/Engine/Private/UnrealEngine.cpp b/Engine/Source/Runtime/Engine/Private/UnrealEngine.cpp
--- a/Engine/Source/Runtime/Engine/Private/UnrealEngine.cpp
+++ b/Engine/Source/Runtime/Engine/Private/UnrealEngine.cpp
@@ -7448,13 +7448,16 @@ void DrawStatsHUD( UWorld* World, FViewport* Viewport, FCanvas* Canvas, UCanvas*
 			}
 			
 #if !(UE_BUILD_SHIPPING || UE_BUILD_TEST)
-			auto MemOver = IStreamingManager::Get().GetTextureStreamingManager().GetMemoryOverBudget();
-			if (MemOver > 0)
+			if (FPlatformProperties::SupportsTextureStreaming())
 			{
-				SmallTextItem.SetColor(FLinearColor::Red);
-				SmallTextItem.Text =  FText::FromString(FString::Printf(TEXT("TEXTURE STREAMING POOL OVER %0.2f MB"), (float)MemOver / 1024.0f / 1024.0f));
-				Canvas->DrawItem(SmallTextItem, FVector2D(MessageX, MessageY));
-				MessageY += 20;
+				auto MemOver = IStreamingManager::Get().GetTextureStreamingManager().GetMemoryOverBudget();
+				if (MemOver > 0)
+				{
+					SmallTextItem.SetColor(FLinearColor::Red);
+					SmallTextItem.Text = FText::FromString(FString::Printf(TEXT("TEXTURE STREAMING POOL OVER %0.2f MB"), (float)MemOver / 1024.0f / 1024.0f));
+					Canvas->DrawItem(SmallTextItem, FVector2D(MessageX, MessageY));
+					MessageY += 20;
+				}
 			}
 #endif
 
