Comment fix in physics constraint initialization.

diff --git a/Engine/Source/Runtime/Engine/Private/PhysicsEngine/ConstraintInstance.cpp b/Engine/Source/Runtime/Engine/Private/PhysicsEngine/ConstraintInstance.cpp
--- a/Engine/Source/Runtime/Engine/Private/PhysicsEngine/ConstraintInstance.cpp
+++ b/Engine/Source/Runtime/Engine/Private/PhysicsEngine/ConstraintInstance.cpp
@@ -274,7 +274,7 @@ FConstraintInstance::FConstraintInstance()
 	AngularBreakThreshold = 500.0f;
 
 	ProjectionLinearTolerance = 5.f; // Linear projection when error > 5 cm
-	ProjectionAngularTolerance = 180.f;// Angular projection when error > 10 degrees
+	ProjectionAngularTolerance = 180.f; // Angular projection when error > 180 degrees
 }
 
 void FConstraintInstance::SetDisableCollision(bool InDisableCollision)
