Merging CL 2665753:

RVO avoidance now takes the height of agents into account when culling obstacles.

Previously a fixed separation distance was used for cylinder checking. Now this is derived from agent heights, ensuring that varied expected height deltas are properly taken into account.

diff --git a/Engine/Source/Runtime/Engine/Classes/AI/Navigation/AvoidanceManager.h b/Engine/Source/Runtime/Engine/Classes/AI/Navigation/AvoidanceManager.h
--- a/Engine/Source/Runtime/Engine/Classes/AI/Navigation/AvoidanceManager.h
+++ b/Engine/Source/Runtime/Engine/Classes/AI/Navigation/AvoidanceManager.h
@@ -27,7 +27,7 @@ struct FNavAvoidanceData
 	float Radius;
 
 	/** Height (object is treated as a cylinder) */
-	float Height;
+	float HalfHeight;
 
 	/** Weight for RVO (set by user) */
 	float Weight;
@@ -51,7 +51,7 @@ struct FNavAvoidanceData
 	FNavAvoidanceData(UAvoidanceManager* Manager, IRVOAvoidanceInterface* AvoidanceComp);
 
 	/** Init function for internal use to guard against data changes not being reflected in blueprint-accessible creation functions */
-	void Init(UAvoidanceManager* Avoidance, const FVector& InCenter, float InRadius, float InHeight,
+	void Init(UAvoidanceManager* Avoidance, const FVector& InCenter, float InRadius, float InHalfHeight,
 		const FVector& InVelocity, float InWeight = 0.5f,
 		int32 InGroupMask = 1, int32 InGroupsToAvoid = 0xffffffff, int32 InGroupsToIgnore = 0,
 		float InTestRadius2D = 500.0f);
@@ -97,9 +97,13 @@ class ENGINE_API UAvoidanceManager : public UObject, public FSelfRegisteringExec
 	UPROPERTY(EditAnywhere, Category="Avoidance", config, meta=(ClampMin = "0.0"))
 	float ArtificialRadiusExpansion;
 
-	/** Test against obstacles within given height difference from moving agent */
+	/** Deprecated - use HeightCheckMargin, generally a much smaller value. */
+	UPROPERTY()
+	float TestHeightDifference_DEPRECATED;
+
+	/** Allowable height margin between obstacles and agents. This is over and above the difference in agent heights. */
 	UPROPERTY(EditAnywhere, Category="Avoidance", config, meta=(ClampMin = "0.0"))
-	float TestHeightDifference;
+	float HeightCheckMargin;
 
 	/** Get the number of avoidance objects currently in the manager. */
 	UFUNCTION(BlueprintCallable, Category="AI")
diff --git a/Engine/Source/Runtime/Engine/Private/AI/Navigation/AvoidanceManager.cpp b/Engine/Source/Runtime/Engine/Private/AI/Navigation/AvoidanceManager.cpp
--- a/Engine/Source/Runtime/Engine/Private/AI/Navigation/AvoidanceManager.cpp
+++ b/Engine/Source/Runtime/Engine/Private/AI/Navigation/AvoidanceManager.cpp
@@ -22,7 +22,7 @@ FNavAvoidanceData::FNavAvoidanceData(UAvoidanceManager* Manager, IRVOAvoidanceIn
 		);
 }
 
-void FNavAvoidanceData::Init(UAvoidanceManager* Avoidance, const FVector& InCenter, float InRadius, float InHeight,
+void FNavAvoidanceData::Init(UAvoidanceManager* Avoidance, const FVector& InCenter, float InRadius, float InHalfHeight,
 							 const FVector& InVelocity, float InWeight,
 							 int32 InGroupMask, int32 InGroupsToAvoid, int32 InGroupsToIgnore,
 							 float InTestRadius2D)
@@ -30,7 +30,7 @@ void FNavAvoidanceData::Init(UAvoidanceManager* Avoidance, const FVector& InCent
 	Center = InCenter;
 	Velocity = InVelocity;
 	Radius = InRadius * Avoidance->ArtificialRadiusExpansion;
-	Height = InHeight;
+	HalfHeight = InHalfHeight;
 	Weight = FMath::Clamp<float>(InWeight, 0.0f, 1.0f);
 	GroupMask = InGroupMask;
 	GroupsToAvoid = InGroupsToAvoid;
@@ -54,8 +54,9 @@ UAvoidanceManager::UAvoidanceManager(const FObjectInitializer& ObjectInitializer
 	LockTimeAfterClean = 0.01f;
 	DeltaTimeToPredict = 0.5f;
 	ArtificialRadiusExpansion = 1.5f;
-	TestHeightDifference = 500.0f;
+	TestHeightDifference_DEPRECATED = 500.0f;
 	bRequestedUpdateTimer = false;
+	HeightCheckMargin = 10.0f;
 
 #if !(UE_BUILD_SHIPPING || UE_BUILD_TEST)
 	bDebugAll = false;
@@ -332,7 +333,7 @@ FVector UAvoidanceManager::GetAvoidanceVelocity_Internal(const FNavAvoidanceData
 			continue;
 		}
 
-		if (FMath::Abs(OtherObject.Center.Z - inAvoidanceData.Center.Z) > TestHeightDifference)
+		if (FMath::Abs(OtherObject.Center.Z - inAvoidanceData.Center.Z) > OtherObject.HalfHeight + inAvoidanceData.HalfHeight + HeightCheckMargin)
 		{
 			continue;
 		}
