Remove the r.Shaders.AvoidFlowControl CVar introduced in 2662625 - setting this by default was causing problems on various platforms and it isn't sufficiently useful to retain. Instead just force the Mac OpenGL shader platform to retain flow-control statements and let everything else use the default as before 2662625.

diff --git a/Engine/Source/Developer/ShaderFormatOpenGL/Private/OpenGLShaderCompiler.cpp b/Engine/Source/Developer/ShaderFormatOpenGL/Private/OpenGLShaderCompiler.cpp
--- a/Engine/Source/Developer/ShaderFormatOpenGL/Private/OpenGLShaderCompiler.cpp
+++ b/Engine/Source/Developer/ShaderFormatOpenGL/Private/OpenGLShaderCompiler.cpp
@@ -1247,7 +1247,7 @@ void CompileShader_Windows_OGL(const FShaderCompilerInput& Input,FShaderCompiler
 	
 	const bool bDumpDebugInfo = (Input.DumpDebugInfoPath != TEXT("") && IFileManager::Get().DirectoryExists(*Input.DumpDebugInfoPath));
 
-	if(Input.Environment.CompilerFlags.Contains(CFLAG_AvoidFlowControl) && PlatformFlowControl == CFLAG_AvoidFlowControl)
+	if(Input.Environment.CompilerFlags.Contains(CFLAG_AvoidFlowControl) || PlatformFlowControl == CFLAG_AvoidFlowControl)
 	{
 		AdditionalDefines.SetDefine(TEXT("COMPILER_SUPPORTS_ATTRIBUTES"), (uint32)1);
 	}
diff --git a/Engine/Source/Runtime/Engine/Private/ShaderCompiler/ShaderCompiler.cpp b/Engine/Source/Runtime/Engine/Private/ShaderCompiler/ShaderCompiler.cpp
--- a/Engine/Source/Runtime/Engine/Private/ShaderCompiler/ShaderCompiler.cpp
+++ b/Engine/Source/Runtime/Engine/Private/ShaderCompiler/ShaderCompiler.cpp
@@ -53,13 +53,6 @@ static TAutoConsoleVariable<int32> CVarOptimizeShaders(
 	TEXT("Whether to optimize shaders.  When using graphical debuggers like Nsight it can be useful to disable this on startup."),
 	ECVF_ReadOnly);
 
-static TAutoConsoleVariable<int32> CVarAvoidFlowControl(
-	TEXT("r.Shaders.AvoidFlowControl"),
-	1,
-	TEXT("When enabled reduces flow-control in favour of longer shaders with greater register usage. When disabled shader flow-control is preserved for platform runtime to optimise which may be beneficial on register constrained GPUs.\n")
-	TEXT("Defaults to 1 (enabled)."),
-	ECVF_ReadOnly);
-
 #if !(UE_BUILD_SHIPPING || UE_BUILD_TEST)
 static TAutoConsoleVariable<FString> CVarD3DCompilerPath(TEXT("r.D3DCompilerPath"),
 	TEXT(""),	// default
@@ -1908,15 +1901,6 @@ void GlobalBeginCompileShader(
 			Input.Environment.CompilerFlags.Add(CFLAG_KeepDebugInfo);
 		}
 	}
-	
-	{
-		static const auto CVar = IConsoleManager::Get().FindConsoleVariable(TEXT("r.Shaders.AvoidFlowControl"));
-		
-		if (CVar->GetInt() != 0)
-		{
-			Input.Environment.CompilerFlags.Add(CFLAG_AvoidFlowControl);
-		}
-	}
 
 	Input.Environment.SetDefine(TEXT("HAS_INVERTED_Z_BUFFER"), ((int32)ERHIZBuffer::IsInverted != 0) ? 1 : 0);
 
