PR #1225: Added GetCurrentStateName function to AnimInstance (Contributed by [REDACTED_PERSON])

diff --git a/Engine/Source/Runtime/Engine/Classes/Animation/AnimInstance.h b/Engine/Source/Runtime/Engine/Classes/Animation/AnimInstance.h
--- a/Engine/Source/Runtime/Engine/Classes/Animation/AnimInstance.h
+++ b/Engine/Source/Runtime/Engine/Classes/Animation/AnimInstance.h
@@ -750,6 +750,10 @@ public:
 	UFUNCTION(BlueprintPure, Category="Animation", meta=(BlueprintInternalUseOnly = "true"))
 	float GetCurrentStateElapsedTime(int32 MachineIndex);
 
+	/** Returns the name of a currently active state in a state machine. */
+	UFUNCTION(BlueprintPure, Category="Animation", meta=(BlueprintInternalUseOnly = "true", AnimGetter = "true"))
+	FName GetCurrentStateName(int32 MachineIndex);
+
 	/** Sets a morph target to a certain weight. */
 	UFUNCTION(BlueprintCallable, Category="Animation")
 	void SetMorphTarget(FName MorphTargetName, float Value);
diff --git a/Engine/Source/Runtime/Engine/Classes/Animation/AnimNode_StateMachine.h b/Engine/Source/Runtime/Engine/Classes/Animation/AnimNode_StateMachine.h
--- a/Engine/Source/Runtime/Engine/Classes/Animation/AnimNode_StateMachine.h
+++ b/Engine/Source/Runtime/Engine/Classes/Animation/AnimNode_StateMachine.h
@@ -124,6 +124,8 @@ public:
 		return ElapsedTime;
 	}
 
+	FName GetCurrentStateName() const;
+
 #if WITH_EDITORONLY_DATA
 	bool IsTransitionActive(int32 TransIndex) const;
 #endif
diff --git a/Engine/Source/Runtime/Engine/Private/Animation/AnimInstance.cpp b/Engine/Source/Runtime/Engine/Private/Animation/AnimInstance.cpp
--- a/Engine/Source/Runtime/Engine/Private/Animation/AnimInstance.cpp
+++ b/Engine/Source/Runtime/Engine/Private/Animation/AnimInstance.cpp
@@ -1726,6 +1726,26 @@ float UAnimInstance::GetCurrentStateElapsedTime(int32 MachineIndex)
 	return 0.0f;
 }
 
+FName UAnimInstance::GetCurrentStateName(int32 MachineIndex)
+{
+	if (UAnimBlueprintGeneratedClass* AnimBlueprintClass = Cast<UAnimBlueprintGeneratedClass>((UObject*)GetClass()))
+	{
+		if ((MachineIndex >= 0) && (MachineIndex < AnimBlueprintClass->AnimNodeProperties.Num()))
+		{
+			const int32 InstancePropertyIndex = AnimBlueprintClass->AnimNodeProperties.Num() - 1 - MachineIndex; //@TODO: ANIMREFACTOR: Reverse indexing
+
+			UStructProperty* MachineInstanceProperty = AnimBlueprintClass->AnimNodeProperties[InstancePropertyIndex];
+			checkSlow(MachineInstanceProperty->Struct->IsChildOf(FAnimNode_StateMachine::StaticStruct()));
+
+			FAnimNode_StateMachine* MachineInstance = MachineInstanceProperty->ContainerPtrToValuePtr<FAnimNode_StateMachine>(this);
+
+			return MachineInstance->GetCurrentStateName();
+		}
+	}
+
+	return NAME_None;
+}
+
 void UAnimInstance::Montage_UpdateWeight(float DeltaSeconds)
 {
 	QUICK_SCOPE_CYCLE_COUNTER(STAT_UAnimInstance_Montage_UpdateWeight);
diff --git a/Engine/Source/Runtime/Engine/Private/Animation/AnimNode_StateMachine.cpp b/Engine/Source/Runtime/Engine/Private/Animation/AnimNode_StateMachine.cpp
--- a/Engine/Source/Runtime/Engine/Private/Animation/AnimNode_StateMachine.cpp
+++ b/Engine/Source/Runtime/Engine/Private/Animation/AnimNode_StateMachine.cpp
@@ -638,12 +638,12 @@ void FAnimNode_StateMachine::EvaluateTransitionStandardBlend(FPoseContext& Outpu
 	// Blend it in
 	const ScalarRegister VPreviousWeight(1.0f - Transition.Alpha);
 	const ScalarRegister VWeight(Transition.Alpha);
-	for (FCompactPoseBoneIndex BoneIndex : Output.Pose.ForEachBoneIndex())
-	{
-		Output.Pose[BoneIndex] = PreviouseStateResult.Pose[BoneIndex] * VPreviousWeight;
-		Output.Pose[BoneIndex].AccumulateWithShortestRotation(NextStateResult.Pose[BoneIndex], VWeight);
-	}
-
+		for(FCompactPoseBoneIndex BoneIndex : Output.Pose.ForEachBoneIndex())
+		{
+			Output.Pose[BoneIndex] = PreviouseStateResult.Pose[BoneIndex] * VPreviousWeight;
+			Output.Pose[BoneIndex].AccumulateWithShortestRotation(NextStateResult.Pose[BoneIndex], VWeight);
+		}
+	
 	// blend curve in
 	Output.Curve.Override(PreviouseStateResult.Curve, 1.0-Transition.Alpha);
 	Output.Curve.Accumulate(NextStateResult.Curve, Transition.Alpha);
@@ -891,3 +891,12 @@ bool FAnimNode_StateMachine::IsValidTransitionIndex(int32 TransitionIndex) const
 	return PRIVATE_MachineDescription->Transitions.IsValidIndex(TransitionIndex);
 }
 
+FName FAnimNode_StateMachine::GetCurrentStateName() const
+{
+	if(PRIVATE_MachineDescription->States.IsValidIndex(CurrentState))
+	{
+		return GetStateInfo().StateName;
+	}
+	return NAME_None;
+}
+
