Implements the TextureProperty material expression

The TextureProperty exposes a texture's property of your choice such as the texture's size or texel size.
Also exposes the TextureSample material expression's Tex input in materials (not only material functions)

diff --git a/Engine/Documentation/Source/Engine/Rendering/Materials/ExpressionReference/Textures/TextureExpressions.INT.udn b/Engine/Documentation/Source/Engine/Rendering/Materials/ExpressionReference/Textures/TextureExpressions.INT.udn
--- a/Engine/Documentation/Source/Engine/Rendering/Materials/ExpressionReference/Textures/TextureExpressions.INT.udn
+++ b/Engine/Documentation/Source/Engine/Rendering/Materials/ExpressionReference/Textures/TextureExpressions.INT.udn
@@ -89,6 +89,12 @@ For a better understanding of the DDX(UVs) and DDY(UVs) inputs, here is what you
 
 ![SampleTextureDerivative.png](SampleTextureDerivative.png)
 
+## TextureProperty
+
+The TextureProperty exposes a texture's property of your choice such as the texture's size or texel size.
+
+![TextureProperty.png](TextureProperty.png)
+
 [INCLUDE:Engine/Rendering/Materials/ExpressionReference/Particles#ParticleSubUV]
 
 
diff --git a/Engine/Documentation/Source/Shared/GraphNodes/Material/Material.INT.udn b/Engine/Documentation/Source/Shared/GraphNodes/Material/Material.INT.udn
--- a/Engine/Documentation/Source/Shared/GraphNodes/Material/Material.INT.udn
+++ b/Engine/Documentation/Source/Shared/GraphNodes/Material/Material.INT.udn
@@ -793,6 +793,14 @@ Engine/Rendering/Materials/ExpressionReference/Textures#texturesample
 [/VAR]
 [/EXCERPT:UMaterialExpressionTextureSample]
 
+[EXCERPT:UMaterialExpressionTextureProperty]
+The TextureProperty expression exposes a texture's properties such as the texture's size or the texel size. 
+
+[VAR:ToolTipFullLink]
+Engine/Rendering/Materials/ExpressionReference/Textures#textureproperty
+[/VAR]
+[/EXCERPT:UMaterialExpressionTextureProperty]
+
 [EXCERPT:UMaterialExpressionBlackBody]
 The Black Body expression simulates the effects of black body radiation within your Material. The user inputs a temperature (in Kelvin) and the resulting color and intensity can be used to drive Base Color and Emissive values to get a physically accurate result. 
 
diff --git a/Engine/Source/Editor/UnrealEd/Private/MaterialGraphNode.cpp b/Engine/Source/Editor/UnrealEd/Private/MaterialGraphNode.cpp
--- a/Engine/Source/Editor/UnrealEd/Private/MaterialGraphNode.cpp
+++ b/Engine/Source/Editor/UnrealEd/Private/MaterialGraphNode.cpp
@@ -25,6 +25,7 @@
 #include "Materials/MaterialExpressionTextureSample.h"
 #include "Materials/MaterialExpressionTextureSampleParameter.h"
 #include "Materials/MaterialExpressionTextureObject.h"
+#include "Materials/MaterialExpressionTextureProperty.h"
 #include "Materials/MaterialExpressionVectorParameter.h"
 #include "Materials/MaterialFunction.h"
 
@@ -176,7 +177,14 @@ FText UMaterialGraphNode::GetNodeTitle(ENodeTitleType::Type TitleType) const
 
 		if ( MaterialExpression->bShaderInputData && (MaterialExpression->bHidePreviewWindow || MaterialExpression->bCollapsed))
 		{
-			NodeTitle.AppendLine(LOCTEXT("InputData", "Input Data"));
+			if (MaterialExpression->IsA<UMaterialExpressionTextureProperty>())
+			{
+				NodeTitle.AppendLine(LOCTEXT("TextureProperty", "Texture Property"));
+			}
+			else
+			{
+				NodeTitle.AppendLine(LOCTEXT("InputData", "Input Data"));
+			}
 		}
 
 		if (bIsPreviewExpression)
diff --git a/Engine/Source/Runtime/Engine/Classes/Materials/MaterialExpressionTextureProperty.h b/Engine/Source/Runtime/Engine/Classes/Materials/MaterialExpressionTextureProperty.h
new file mode 100644
--- /dev/null
+++ b/Engine/Source/Runtime/Engine/Classes/Materials/MaterialExpressionTextureProperty.h
@@ -0,0 +1,42 @@
+// Copyright 1998-2015 Epic Games, Inc. All Rights Reserved.
+
+
+#pragma once
+#include "Materials/MaterialExpression.h"
+#include "MaterialExpressionTextureProperty.generated.h"
+
+
+/** Selects the texture property to output */
+UENUM()
+enum EMaterialExposedTextureProperty
+{
+	/* The texture's size. */
+	TMTM_TextureSize UMETA(DisplayName="Texture Size"),
+
+	/* The texture's texel size in the UV space (1 / Texture Size) */
+	TMTM_TexelSize UMETA(DisplayName="Texel Size"),
+	
+	TMTM_MAX,
+};
+
+UCLASS(collapsecategories, hidecategories=Object)
+class ENGINE_API UMaterialExpressionTextureProperty : public UMaterialExpression
+{
+	GENERATED_UCLASS_BODY()
+
+	UPROPERTY(meta = (RequiredInput = "false", ToolTip = "Defaults to 'Texture' if not specified"))
+	FExpressionInput TextureObject;
+	
+	/** Texture property to be accessed */
+	UPROPERTY(EditAnywhere, Category=UMaterialExpressionTextureProperty, meta=(DisplayName = "Texture Property"))
+	TEnumAsByte<EMaterialExposedTextureProperty> Property;
+	
+	// Begin UMaterialExpression Interface
+	virtual int32 Compile(class FMaterialCompiler* Compiler, int32 OutputIndex, int32 MultiplexIndex) override;
+	virtual void GetCaption(TArray<FString>& OutCaptions) const override;
+	
+#if WITH_EDITOR
+	virtual uint32 GetInputType(int32 InputIndex) override;
+#endif //WITH_EDITOR
+	// End UMaterialExpression Interface
+};
diff --git a/Engine/Source/Runtime/Engine/Classes/Materials/MaterialExpressionTextureSample.h b/Engine/Source/Runtime/Engine/Classes/Materials/MaterialExpressionTextureSample.h
--- a/Engine/Source/Runtime/Engine/Classes/Materials/MaterialExpressionTextureSample.h
+++ b/Engine/Source/Runtime/Engine/Classes/Materials/MaterialExpressionTextureSample.h
@@ -91,7 +91,6 @@ class ENGINE_API UMaterialExpressionTextureSample : public UMaterialExpressionTe
 	virtual bool MatchesSearchQuery( const TCHAR* SearchQuery ) override;
 #if WITH_EDITOR
 	virtual uint32 GetInputType(int32 InputIndex) override;
-	virtual void GetConnectorToolTip(int32 InputIndex, int32 OutputIndex, TArray<FString>& OutToolTip) override;
 #endif
 	// End UMaterialExpression Interface
 
diff --git a/Engine/Source/Runtime/Engine/Private/Materials/HLSLMaterialTranslator.h b/Engine/Source/Runtime/Engine/Private/Materials/HLSLMaterialTranslator.h
--- a/Engine/Source/Runtime/Engine/Private/Materials/HLSLMaterialTranslator.h
+++ b/Engine/Source/Runtime/Engine/Private/Materials/HLSLMaterialTranslator.h
@@ -2456,6 +2456,20 @@ protected:
 			);
 	}
 
+	virtual int32 TextureProperty(int32 TextureIndex, EMaterialExposedTextureProperty Property) override
+	{
+		EMaterialValueType TextureType = GetParameterType(TextureIndex);
+
+		if(TextureType != MCT_Texture2D)
+		{
+			return Errorf(TEXT("Texture size only available for Texture2D, not %s"),DescribeType(TextureType));
+		}
+		
+		auto TextureExpression = (FMaterialUniformExpressionTexture*) (*CurrentScopeChunks)[TextureIndex].UniformExpression.GetReference();
+
+		return AddUniformExpression(new FMaterialUniformExpressionTextureProperty(TextureExpression, Property), MCT_Float2, TEXT(""));
+	}
+
 	virtual int32 TextureDecalMipmapLevel(int32 TextureSizeInput) override
 	{
 		if (Material->GetMaterialDomain() != MD_DeferredDecal)
diff --git a/Engine/Source/Runtime/Engine/Private/Materials/MaterialExpressions.cpp b/Engine/Source/Runtime/Engine/Private/Materials/MaterialExpressions.cpp
--- a/Engine/Source/Runtime/Engine/Private/Materials/MaterialExpressions.cpp
+++ b/Engine/Source/Runtime/Engine/Private/Materials/MaterialExpressions.cpp
@@ -113,6 +113,7 @@
 #include "Materials/MaterialExpressionSubtract.h"
 #include "Materials/MaterialExpressionTextureBase.h"
 #include "Materials/MaterialExpressionTextureObject.h"
+#include "Materials/MaterialExpressionTextureProperty.h"
 #include "Materials/MaterialExpressionTextureSample.h"
 #include "Materials/MaterialExpressionParticleSubUV.h"
 #include "Materials/MaterialExpressionTextureSampleParameter.h"
@@ -1215,12 +1216,7 @@ const TArray<FExpressionInput*> UMaterialExpressionTextureSample::GetInputs()
 FExpressionInput* UMaterialExpressionTextureSample::GetInput(int32 InputIndex)
 {
 	IF_INPUT_RETURN(Coordinates);
-
-	// Only show the TextureObject input inside a material function, since that's the only place it is useful
-	if(GetOuter()->IsA(UMaterialFunction::StaticClass()))
-	{
-		IF_INPUT_RETURN(TextureObject);
-	}
+	IF_INPUT_RETURN(TextureObject);
 
 	if(MipValueMode == TMVM_Derivative)
 	{
@@ -1241,12 +1237,7 @@ FExpressionInput* UMaterialExpressionTextureSample::GetInput(int32 InputIndex)
 FString UMaterialExpressionTextureSample::GetInputName(int32 InputIndex) const
 {
 	IF_INPUT_RETURN(Coordinates, TEXT("Coordinates"));
-
-	// Only show the TextureObject input inside a material function, since that's the only place it is useful
-	if(GetOuter()->IsA(UMaterialFunction::StaticClass()))
-	{
-		IF_INPUT_RETURN(TextureObject, TEXT("TextureObject"));
-	}
+	IF_INPUT_RETURN(TextureObject, TEXT("TextureObject"));
 
 	if(MipValueMode == TMVM_MipLevel)
 	{
@@ -1416,11 +1407,8 @@ uint32 UMaterialExpressionTextureSample::GetInputType(int32 InputIndex)
 	IF_INPUT_RETURN(Coordinates, MCT_Float);
 
 	// Only show the TextureObject input inside a material function, since that's the only place it is useful
-	if(GetOuter()->IsA(UMaterialFunction::StaticClass()))
-	{
-		IF_INPUT_RETURN(TextureObject, MCT_Texture);
-	}
-
+	IF_INPUT_RETURN(TextureObject, MCT_Texture);
+	
 	if(MipValueMode == TMVM_MipLevel || MipValueMode == TMVM_MipBias)
 	{
 		IF_INPUT_RETURN(MipValue, MCT_Float);
@@ -1434,17 +1422,6 @@ uint32 UMaterialExpressionTextureSample::GetInputType(int32 InputIndex)
 	return MCT_Unknown;
 }
 #undef IF_INPUT_RETURN
-
-void UMaterialExpressionTextureSample::GetConnectorToolTip(int32 InputIndex, int32 OutputIndex, TArray<FString>& OutToolTip)
-{
-	if (InputIndex == 1 && !GetOuter()->IsA(UMaterialFunction::StaticClass()))
-	{
-		// If input pin 1 is omitted, increment the index so the correct tooltip is looked up
-		InputIndex++;
-	}
-
-	Super::GetConnectorToolTip(InputIndex, OutputIndex, OutToolTip);
-}
 #endif
 
 int32 UMaterialExpressionTextureSample::CompileMipValue0(class FMaterialCompiler* Compiler)
@@ -1761,6 +1738,73 @@ uint32 UMaterialExpressionTextureObject::GetOutputType(int32 OutputIndex)
 }
 #endif //WITH_EDITOR
 
+//
+//  UMaterialExpressionTextureProperty
+//
+UMaterialExpressionTextureProperty::UMaterialExpressionTextureProperty(const FObjectInitializer& ObjectInitializer)
+	: Super(ObjectInitializer)
+{
+	// Structure to hold one-time initialization
+	struct FConstructorStatics
+	{
+		FText NAME_Constants;
+		FConstructorStatics()
+			: NAME_Constants(LOCTEXT( "Texture", "Texture" ))
+		{
+		}
+	};
+	static FConstructorStatics ConstructorStatics;
+
+	Property = TMTM_TextureSize;
+	MenuCategories.Add(ConstructorStatics.NAME_Constants);
+	bShaderInputData = true;
+	bShowOutputNameOnPin = false;
+	
+	Outputs.Reset();
+	Outputs.Add(FExpressionOutput(TEXT("")));
+}
+
+int32 UMaterialExpressionTextureProperty::Compile(class FMaterialCompiler* Compiler, int32 OutputIndex, int32 MultiplexIndex)
+{	
+	if (!TextureObject.Expression)
+	{
+		return Compiler->Errorf(TEXT("TextureSample> Missing input texture"));
+	}
+
+	const int32 TextureCodeIndex = TextureObject.Compile(Compiler);
+	if (TextureCodeIndex == INDEX_NONE)
+	{
+		return INDEX_NONE;
+	}
+
+	return Compiler->TextureProperty(TextureCodeIndex, Property);
+}
+
+void UMaterialExpressionTextureProperty::GetCaption(TArray<FString>& OutCaptions) const
+{
+#if WITH_EDITOR
+	const UEnum* TexturePropertyEnum = FindObject<UEnum>(NULL, TEXT("Engine.EMaterialExposedTextureProperty"));
+	check(TexturePropertyEnum);
+
+	const FString PropertyDisplayName = TexturePropertyEnum->GetDisplayNameText(Property).ToString();
+#else
+	const FString PropertyDisplayName = TEXT("");
+#endif
+
+	OutCaptions.Add(PropertyDisplayName);
+}
+
+#if WITH_EDITOR
+// this define is only used for the following function
+#define IF_INPUT_RETURN(Item, Type) if(!InputIndex) return Type; --InputIndex
+uint32 UMaterialExpressionTextureProperty::GetInputType(int32 InputIndex)
+{
+	IF_INPUT_RETURN(TextureObject, MCT_Texture);
+	return MCT_Unknown;
+}
+#undef IF_INPUT_RETURN
+#endif
+
 //
 //  UMaterialExpressionTextureSampleParameter2D
 //
diff --git a/Engine/Source/Runtime/Engine/Private/Materials/MaterialUniformExpressions.cpp b/Engine/Source/Runtime/Engine/Private/Materials/MaterialUniformExpressions.cpp
--- a/Engine/Source/Runtime/Engine/Private/Materials/MaterialUniformExpressions.cpp
+++ b/Engine/Source/Runtime/Engine/Private/Materials/MaterialUniformExpressions.cpp
@@ -529,3 +529,4 @@ IMPLEMENT_MATERIALUNIFORMEXPRESSION_TYPE(FMaterialUniformExpressionCeil);
 IMPLEMENT_MATERIALUNIFORMEXPRESSION_TYPE(FMaterialUniformExpressionFrac);
 IMPLEMENT_MATERIALUNIFORMEXPRESSION_TYPE(FMaterialUniformExpressionFmod);
 IMPLEMENT_MATERIALUNIFORMEXPRESSION_TYPE(FMaterialUniformExpressionAbs);
+IMPLEMENT_MATERIALUNIFORMEXPRESSION_TYPE(FMaterialUniformExpressionTextureProperty);
diff --git a/Engine/Source/Runtime/Engine/Private/Materials/MaterialUniformExpressions.h b/Engine/Source/Runtime/Engine/Private/Materials/MaterialUniformExpressions.h
--- a/Engine/Source/Runtime/Engine/Private/Materials/MaterialUniformExpressions.h
+++ b/Engine/Source/Runtime/Engine/Private/Materials/MaterialUniformExpressions.h
@@ -7,6 +7,7 @@
 
 #include "EnginePrivate.h"
 #include "MaterialShared.h"
+#include "Materials/MaterialExpressionTextureProperty.h"
 
 /**
  */
@@ -1339,3 +1340,72 @@ public:
 private:
 	TRefCountPtr<FMaterialUniformExpression> X;
 };
+
+/**
+ */
+class FMaterialUniformExpressionTextureProperty: public FMaterialUniformExpression
+{
+	DECLARE_MATERIALUNIFORMEXPRESSION_TYPE(FMaterialUniformExpressionTextureProperty);
+public:
+	
+	FMaterialUniformExpressionTextureProperty() {}
+	FMaterialUniformExpressionTextureProperty(FMaterialUniformExpressionTexture* InTextureExpression, EMaterialExposedTextureProperty InTextureProperty)
+		: TextureExpression(InTextureExpression)
+		, TextureProperty(InTextureProperty)
+	{}
+
+	// FMaterialUniformExpression interface.
+	virtual void Serialize(FArchive& Ar) override
+	{
+		Ar << TextureExpression << TextureProperty;
+	}
+	virtual void GetNumberValue(const FMaterialRenderContext& Context,FLinearColor& OutValue) const override
+	{
+		const UTexture* Texture;
+
+		{
+			ESamplerSourceMode SamplerSource;
+			TextureExpression->GetTextureValue(Context, Context.Material, Texture, SamplerSource);
+		}
+
+		if (!Texture || !Texture->Resource)
+		{
+			return;
+		}
+	
+		if (TextureProperty == TMTM_TextureSize)
+		{
+			OutValue.R = Texture->Resource->GetSizeX();
+			OutValue.G = Texture->Resource->GetSizeY();
+		}
+		else if (TextureProperty == TMTM_TexelSize)
+		{
+			OutValue.R = 1.0f / float(Texture->Resource->GetSizeX());
+			OutValue.G = 1.0f / float(Texture->Resource->GetSizeY());
+		}
+		else
+		{
+			check(0);
+		}
+	}
+	virtual bool IsIdentical(const FMaterialUniformExpression* OtherExpression) const override
+	{
+		if (GetType() != OtherExpression->GetType())
+		{
+			return false;
+		}
+
+		auto OtherTexturePropertyExpression = (const FMaterialUniformExpressionTextureProperty*)OtherExpression;
+		
+		if (TextureProperty != OtherTexturePropertyExpression->TextureProperty)
+		{
+			return false;
+		}
+
+		return TextureExpression->IsIdentical(OtherTexturePropertyExpression->TextureExpression);
+	}
+	
+private:
+	TRefCountPtr<FMaterialUniformExpressionTexture> TextureExpression;
+	int8 TextureProperty;
+};
diff --git a/Engine/Source/Runtime/Engine/Public/MaterialCompiler.h b/Engine/Source/Runtime/Engine/Public/MaterialCompiler.h
--- a/Engine/Source/Runtime/Engine/Public/MaterialCompiler.h
+++ b/Engine/Source/Runtime/Engine/Public/MaterialCompiler.h
@@ -10,6 +10,7 @@
 #include "Materials/MaterialExpressionSpeedTree.h"
 #include "Materials/MaterialExpressionTextureSample.h"
 #include "Materials/MaterialExpressionWorldPosition.h"
+#include "Materials/MaterialExpressionTextureProperty.h"
 #include "Materials/MaterialExpressionViewProperty.h"
 
 /** 
@@ -101,6 +102,7 @@ public:
 
 	virtual int32 TextureCoordinate(uint32 CoordinateIndex, bool UnMirrorU, bool UnMirrorV) = 0;
 	virtual int32 TextureSample(int32 Texture,int32 Coordinate,enum EMaterialSamplerType SamplerType,int32 MipValue0Index=INDEX_NONE,int32 MipValue1Index=INDEX_NONE,ETextureMipValueMode MipValueMode=TMVM_None,ESamplerSourceMode SamplerSource=SSM_FromTextureAsset) = 0;
+	virtual int32 TextureProperty(int32 InTexture, EMaterialExposedTextureProperty Property) = 0;
 
 	virtual int32 TextureDecalMipmapLevel(int32 TextureSizeInput) = 0;
 	virtual int32 TextureDecalDerivative(bool bDDY) = 0;
@@ -262,6 +264,8 @@ public:
 
 	virtual int32 TextureSample(int32 InTexture,int32 Coordinate,enum EMaterialSamplerType SamplerType,int32 MipValue0Index,int32 MipValue1Index,ETextureMipValueMode MipValueMode,ESamplerSourceMode SamplerSource) override 
 		{ return Compiler->TextureSample(InTexture,Coordinate,SamplerType,MipValue0Index,MipValue1Index,MipValueMode,SamplerSource); }
+	virtual int32 TextureProperty(int32 InTexture, EMaterialExposedTextureProperty Property) override 
+		{ return Compiler->TextureProperty(InTexture, Property); }
 
 	virtual int32 TextureCoordinate(uint32 CoordinateIndex, bool UnMirrorU, bool UnMirrorV) override { return Compiler->TextureCoordinate(CoordinateIndex, UnMirrorU, UnMirrorV); }
 
