Removed unused FNavDataGenerator::OnWorldInitDone #UE4

github PR#1488 by slonopotamus

diff --git a/Engine/Source/Runtime/Engine/Public/AI/NavDataGenerator.h b/Engine/Source/Runtime/Engine/Public/AI/NavDataGenerator.h
--- a/Engine/Source/Runtime/Engine/Public/AI/NavDataGenerator.h
+++ b/Engine/Source/Runtime/Engine/Public/AI/NavDataGenerator.h
@@ -32,9 +32,6 @@ public:
 	/** determines whether this generator is performing navigation building actions at the moment*/
 	virtual bool IsBuildInProgress(bool bCheckDirtyToo = false) const { return false; }	
 
-	/** Called when world is loaded and all actors are initialized - to be used for post-load steps like data validating */
-	virtual void OnWorldInitDone(bool bAllowedToRebuild) {}	
-
 	/** Returns number of remaining tasks till build is complete
 	 */
 	virtual int32 GetNumRemaningBuildTasks() const { return 0; };
