PR #1487: Cleanup dead 'ChartDistanceFactor' code (Contributed by [REDACTED_PERSON])

diff --git a/Engine/Source/Runtime/Engine/Classes/Components/SkinnedMeshComponent.h b/Engine/Source/Runtime/Engine/Classes/Components/SkinnedMeshComponent.h
--- a/Engine/Source/Runtime/Engine/Classes/Components/SkinnedMeshComponent.h
+++ b/Engine/Source/Runtime/Engine/Classes/Components/SkinnedMeshComponent.h
@@ -337,10 +337,6 @@ private:
 	uint32 bForceMeshObjectUpdate:1;
 
 public:
-	/** If true, DistanceFactor for this SkinnedMeshComponent will be added to global chart. */
-	UPROPERTY()
-	uint32 bChartDistanceFactor:1;
-
 	/** Whether or not we can highlight selected sections - this should really only be done in the editor */
 	UPROPERTY(transient)
 	uint32 bCanHighlightSelectedSections:1;
diff --git a/Engine/Source/Runtime/Engine/Private/Character.cpp b/Engine/Source/Runtime/Engine/Private/Character.cpp
--- a/Engine/Source/Runtime/Engine/Private/Character.cpp
+++ b/Engine/Source/Runtime/Engine/Private/Character.cpp
@@ -90,7 +90,6 @@ ACharacter::ACharacter(const FObjectInitializer& ObjectInitializer)
 		Mesh->bCastDynamicShadow = true;
 		Mesh->bAffectDynamicIndirectLighting = true;
 		Mesh->PrimaryComponentTick.TickGroup = TG_PrePhysics;
-		Mesh->bChartDistanceFactor = true;
 		Mesh->AttachParent = CapsuleComponent;
 		static FName MeshCollisionProfileName(TEXT("CharacterMesh"));
 		Mesh->SetCollisionProfileName(MeshCollisionProfileName);
diff --git a/Engine/Source/Runtime/Engine/Private/Components/SkeletalMeshComponent.cpp b/Engine/Source/Runtime/Engine/Private/Components/SkeletalMeshComponent.cpp
--- a/Engine/Source/Runtime/Engine/Private/Components/SkeletalMeshComponent.cpp
+++ b/Engine/Source/Runtime/Engine/Private/Components/SkeletalMeshComponent.cpp
@@ -359,36 +359,6 @@ void USkeletalMeshComponent::InitializeComponent()
 	InitAnim(false);
 }
 
-#if CHART_DISTANCE_FACTORS
-static void AddDistanceFactorToChart(float DistanceFactor)
-{
-	if(DistanceFactor < SMALL_NUMBER)
-	{
-		return;
-	}
-
-	if(DistanceFactor >= GDistanceFactorDivision[NUM_DISTANCEFACTOR_BUCKETS-2])
-	{
-		GDistanceFactorChart[NUM_DISTANCEFACTOR_BUCKETS-1]++;
-	}
-	else if(DistanceFactor < GDistanceFactorDivision[0])
-	{
-		GDistanceFactorChart[0]++;
-	}
-	else
-	{
-		for(int32 i=1; i<NUM_DISTANCEFACTOR_BUCKETS-2; i++)
-		{
-			if(DistanceFactor < GDistanceFactorDivision[i])
-			{
-				GDistanceFactorChart[i]++;
-				break;
-			}
-		}
-	}
-}
-#endif // CHART_DISTANCE_FACTORS
-
 #if WITH_EDITOR
 void USkeletalMeshComponent::PostEditChangeProperty(FPropertyChangedEvent& PropertyChangedEvent)
 {
diff --git a/Engine/Source/Runtime/Engine/Private/Components/SkinnedMeshComponent.cpp b/Engine/Source/Runtime/Engine/Private/Components/SkinnedMeshComponent.cpp
--- a/Engine/Source/Runtime/Engine/Private/Components/SkinnedMeshComponent.cpp
+++ b/Engine/Source/Runtime/Engine/Private/Components/SkinnedMeshComponent.cpp
@@ -1857,14 +1857,6 @@ bool USkinnedMeshComponent::UpdateLODStatus()
 	if(MeshObject)
 	{
 		MaxDistanceFactor = MeshObject->MaxDistanceFactor;
-
-#if CHART_DISTANCE_FACTORS
-		// Only chart DistanceFactor if it was actually rendered recently
-		if(bChartDistanceFactor && ((LastRenderTime > GetWorld()->TimeSeconds - 1.0f) || bUpdateSkelWhenNotRendered))
-		{
-			AddDistanceFactorToChart(MaxDistanceFactor);
-		}
-#endif // CHART_DISTANCE_FACTORS
 	}
 
 	return bLODChanged;
diff --git a/Engine/Source/Runtime/Engine/Public/SkeletalMeshTypes.h b/Engine/Source/Runtime/Engine/Public/SkeletalMeshTypes.h
--- a/Engine/Source/Runtime/Engine/Public/SkeletalMeshTypes.h
+++ b/Engine/Source/Runtime/Engine/Public/SkeletalMeshTypes.h
@@ -19,10 +19,6 @@ typedef uint16 FBoneIndexType;
 
 #include "ReferenceSkeleton.h"
 
-// Define that controls showing chart of distance factors for skel meshes during entire run of the game on exit.
-#define CHART_DISTANCE_FACTORS 0
-
-
 class FRawStaticIndexBuffer16or32Interface;
 class FPrimitiveDrawInterface;
 
