Homogenizes the ViewProperty and TextureProperty together.

diff --git a/Engine/Source/Editor/UnrealEd/Private/MaterialGraphNode.cpp b/Engine/Source/Editor/UnrealEd/Private/MaterialGraphNode.cpp
--- a/Engine/Source/Editor/UnrealEd/Private/MaterialGraphNode.cpp
+++ b/Engine/Source/Editor/UnrealEd/Private/MaterialGraphNode.cpp
@@ -27,6 +27,7 @@
 #include "Materials/MaterialExpressionTextureObject.h"
 #include "Materials/MaterialExpressionTextureProperty.h"
 #include "Materials/MaterialExpressionVectorParameter.h"
+#include "Materials/MaterialExpressionViewProperty.h"
 #include "Materials/MaterialFunction.h"
 
 #include "MaterialEditorUtilities.h"
@@ -181,6 +182,10 @@ FText UMaterialGraphNode::GetNodeTitle(ENodeTitleType::Type TitleType) const
 			{
 				NodeTitle.AppendLine(LOCTEXT("TextureProperty", "Texture Property"));
 			}
+			else if (MaterialExpression->IsA<UMaterialExpressionViewProperty>())
+			{
+				NodeTitle.AppendLine(LOCTEXT("ViewProperty", "View Property"));
+			}
 			else
 			{
 				NodeTitle.AppendLine(LOCTEXT("InputData", "Input Data"));
diff --git a/Engine/Source/Runtime/Engine/Classes/Materials/MaterialExpressionTextureProperty.h b/Engine/Source/Runtime/Engine/Classes/Materials/MaterialExpressionTextureProperty.h
--- a/Engine/Source/Runtime/Engine/Classes/Materials/MaterialExpressionTextureProperty.h
+++ b/Engine/Source/Runtime/Engine/Classes/Materials/MaterialExpressionTextureProperty.h
@@ -24,7 +24,7 @@ class ENGINE_API UMaterialExpressionTextureProperty : public UMaterialExpression
 {
 	GENERATED_UCLASS_BODY()
 
-	UPROPERTY(meta = (RequiredInput = "false", ToolTip = "Defaults to 'Texture' if not specified"))
+	UPROPERTY(meta = (RequiredInput = "false", ToolTip = "Texture Object to access the property from."))
 	FExpressionInput TextureObject;
 	
 	/** Texture property to be accessed */
diff --git a/Engine/Source/Runtime/Engine/Classes/Materials/MaterialExpressionViewProperty.h b/Engine/Source/Runtime/Engine/Classes/Materials/MaterialExpressionViewProperty.h
--- a/Engine/Source/Runtime/Engine/Classes/Materials/MaterialExpressionViewProperty.h
+++ b/Engine/Source/Runtime/Engine/Classes/Materials/MaterialExpressionViewProperty.h
@@ -9,17 +9,17 @@ UENUM()
 enum EMaterialExposedViewProperty
 {	
 	/** Horizontal and vertical size of the view's buffer in pixels */
-	MEVP_BufferSize UMETA(DisplayName="RenderTargetSize"),
+	MEVP_BufferSize UMETA(DisplayName="Render Target Size"),
 	/** Horizontal and vertical field of view angles in radian */
-	MEVP_FieldOfView UMETA(DisplayName="FieldOfView"),
+	MEVP_FieldOfView UMETA(DisplayName="Field Of View"),
 	/** Tan(FieldOfView * 0.5) */
-	MEVP_TanHalfFieldOfView UMETA(DisplayName="TanHalfFieldOfView"),
+	MEVP_TanHalfFieldOfView UMETA(DisplayName="Tan(0.5 * Field Of View)"),
 	/** Horizontal and vertical size of the view in pixels */
-	MEVP_ViewSize UMETA(DisplayName="ViewSize"),
+	MEVP_ViewSize UMETA(DisplayName="View Size"),
 	/** Absolute world space view position (differs from the camera position in the shadow passes) */
-	MEVP_WorldSpaceViewPosition UMETA(DisplayName="ViewPosition (Absolute World Space)"),
+	MEVP_WorldSpaceViewPosition UMETA(DisplayName="View Position (Absolute World Space)"),
 	/** Absolute world space camera position */
-	MEVP_WorldSpaceCameraPosition UMETA(DisplayName="CameraPosition (Absolute World Space)"),
+	MEVP_WorldSpaceCameraPosition UMETA(DisplayName="Camera Position (Absolute World Space)"),
 
 	MEVP_MAX,
 };
diff --git a/Engine/Source/Runtime/Engine/Private/Materials/MaterialExpressions.cpp b/Engine/Source/Runtime/Engine/Private/Materials/MaterialExpressions.cpp
--- a/Engine/Source/Runtime/Engine/Private/Materials/MaterialExpressions.cpp
+++ b/Engine/Source/Runtime/Engine/Private/Materials/MaterialExpressions.cpp
@@ -4769,7 +4769,7 @@ void UMaterialExpressionViewProperty::GetCaption(TArray<FString>& OutCaptions) c
 	const FString PropertyDisplayName = TEXT("");
 #endif
 
-	OutCaptions.Add(FString(TEXT("View.")) + PropertyDisplayName);
+	OutCaptions.Add(PropertyDisplayName);
 }
 
 UMaterialExpressionViewSize::UMaterialExpressionViewSize(const FObjectInitializer& ObjectInitializer)
