Put bChartDistanceFactor back in but as deprecated to avoid compilation failures for projects that were (unnecessarily) setting it

diff --git a/Engine/Source/Runtime/Engine/Classes/Components/SkinnedMeshComponent.h b/Engine/Source/Runtime/Engine/Classes/Components/SkinnedMeshComponent.h
--- a/Engine/Source/Runtime/Engine/Classes/Components/SkinnedMeshComponent.h
+++ b/Engine/Source/Runtime/Engine/Classes/Components/SkinnedMeshComponent.h
@@ -337,6 +337,10 @@ private:
 	uint32 bForceMeshObjectUpdate:1;
 
 public:
+
+	DEPRECATED(4.10, "bChartDistanceFactor is no longer useful, please remove references to it")
+	uint32 bChartDistanceFactor:1;
+
 	/** Whether or not we can highlight selected sections - this should really only be done in the editor */
 	UPROPERTY(transient)
 	uint32 bCanHighlightSelectedSections:1;
