UE4 - Fixed dithered LOD transitions on HISMC meshes with materials without dithered LOD transitions.

diff --git a/Engine/Shaders/LocalVertexFactory.usf b/Engine/Shaders/LocalVertexFactory.usf
--- a/Engine/Shaders/LocalVertexFactory.usf
+++ b/Engine/Shaders/LocalVertexFactory.usf
@@ -11,8 +11,17 @@ float4x4 PreviousLocalToWorld;
 float4 LightMapCoordinateScaleBias;
 float4 ShadowMapCoordinateScaleBias;
 
-#ifndef USE_DITHERED_LOD_TRANSITION
-	#error "USE_DITHERED_LOD_TRANSITION should have been defined"
+
+#if USE_INSTANCING
+	#ifndef USE_DITHERED_LOD_TRANSITION_FOR_INSTANCED
+		#error "USE_DITHERED_LOD_TRANSITION_FOR_INSTANCED should have been defined"
+	#endif
+	#define USE_DITHERED_LOD_TRANSITION USE_DITHERED_LOD_TRANSITION_FOR_INSTANCED
+#else
+	#ifndef USE_DITHERED_LOD_TRANSITION_FROM_MATERIAL
+		#error "USE_DITHERED_LOD_TRANSITION_FROM_MATERIAL should have been defined"
+	#endif
+	#define USE_DITHERED_LOD_TRANSITION USE_DITHERED_LOD_TRANSITION_FROM_MATERIAL
 #endif
 
 #if USE_INSTANCING
diff --git a/Engine/Shaders/MaterialTemplate.usf b/Engine/Shaders/MaterialTemplate.usf
--- a/Engine/Shaders/MaterialTemplate.usf
+++ b/Engine/Shaders/MaterialTemplate.usf
@@ -25,7 +25,17 @@
 #include "PaniniProjection.usf"
 
 #ifndef USE_DITHERED_LOD_TRANSITION
-	#define USE_DITHERED_LOD_TRANSITION 0
+    #if USE_INSTANCING
+	    #ifndef USE_DITHERED_LOD_TRANSITION_FOR_INSTANCED
+		    #error "USE_DITHERED_LOD_TRANSITION_FOR_INSTANCED should have been defined"
+	    #endif
+	    #define USE_DITHERED_LOD_TRANSITION USE_DITHERED_LOD_TRANSITION_FOR_INSTANCED
+    #else
+	    #ifndef USE_DITHERED_LOD_TRANSITION_FROM_MATERIAL
+		    #error "USE_DITHERED_LOD_TRANSITION_FROM_MATERIAL should have been defined"
+	    #endif
+	    #define USE_DITHERED_LOD_TRANSITION USE_DITHERED_LOD_TRANSITION_FROM_MATERIAL
+    #endif
 #endif
 
 //Platforms that don't run the editor shouldn't need editor features in the shaders.
diff --git a/Engine/Source/Runtime/Engine/Private/InstancedStaticMesh.h b/Engine/Source/Runtime/Engine/Private/InstancedStaticMesh.h
--- a/Engine/Source/Runtime/Engine/Private/InstancedStaticMesh.h
+++ b/Engine/Source/Runtime/Engine/Private/InstancedStaticMesh.h
@@ -205,7 +205,7 @@ public:
 		OutEnvironment.SetDefine(TEXT("USE_INSTANCING"),TEXT("1"));
 		const bool bInstanced = RHISupportsInstancing(Platform);
 		OutEnvironment.SetDefine(TEXT("USE_INSTANCING_EMULATED"), bInstanced ? TEXT("0") : TEXT("1"));
-		OutEnvironment.SetDefine(TEXT("USE_DITHERED_LOD_TRANSITION"), ALLOW_DITHERED_LOD_FOR_INSTANCED_STATIC_MESHES ? TEXT("1") : TEXT("0"));
+		OutEnvironment.SetDefine(TEXT("USE_DITHERED_LOD_TRANSITION_FOR_INSTANCED"), ALLOW_DITHERED_LOD_FOR_INSTANCED_STATIC_MESHES ? TEXT("1") : TEXT("0"));
 		FLocalVertexFactory::ModifyCompilationEnvironment(Platform, Material, OutEnvironment);
 	}
 
diff --git a/Engine/Source/Runtime/Engine/Private/Materials/MaterialShared.cpp b/Engine/Source/Runtime/Engine/Private/Materials/MaterialShared.cpp
--- a/Engine/Source/Runtime/Engine/Private/Materials/MaterialShared.cpp
+++ b/Engine/Source/Runtime/Engine/Private/Materials/MaterialShared.cpp
@@ -1221,7 +1221,7 @@ void FMaterial::SetupMaterialEnvironment(
 		OutEnvironment.SetDefine(TEXT("MATERIALDECALRESPONSEMASK"), MaterialDecalResponseMask);
 	}
 
-	OutEnvironment.SetDefine(TEXT("USE_DITHERED_LOD_TRANSITION"), IsDitheredLODTransition() ? TEXT("1") : TEXT("0"));
+	OutEnvironment.SetDefine(TEXT("USE_DITHERED_LOD_TRANSITION_FROM_MATERIAL"), IsDitheredLODTransition() ? TEXT("1") : TEXT("0"));
 	OutEnvironment.SetDefine(TEXT("MATERIAL_TWOSIDED"), IsTwoSided() ? TEXT("1") : TEXT("0"));
 	OutEnvironment.SetDefine(TEXT("MATERIAL_TANGENTSPACENORMAL"), IsTangentSpaceNormal() ? TEXT("1") : TEXT("0"));
 	OutEnvironment.SetDefine(TEXT("GENERATE_SPHERICAL_PARTICLE_NORMALS"),ShouldGenerateSphericalParticleNormals() ? TEXT("1") : TEXT("0"));
