#ue4 - Inlined SceneComponent::GetAttachParent() and added GetAttachSocketName().

--------

diff --git a/Engine/Source/Runtime/Engine/Classes/Components/SceneComponent.h b/Engine/Source/Runtime/Engine/Classes/Components/SceneComponent.h
--- a/Engine/Source/Runtime/Engine/Classes/Components/SceneComponent.h
+++ b/Engine/Source/Runtime/Engine/Classes/Components/SceneComponent.h
@@ -523,9 +523,13 @@ public:
 	UFUNCTION(BlueprintCallable, Category="Physics")
 	virtual bool IsAnySimulatingPhysics() const;
 
-	/** Get a pointer to the USceneComponent we are attached to */
+	/** Get the SceneComponent we are attached to. */
 	UFUNCTION(BlueprintCallable, Category="Utilities|Transformation")
-	class USceneComponent* GetAttachParent() const;
+	USceneComponent* GetAttachParent() const;
+
+	/** Get the socket we are attached to. */
+	UFUNCTION(BlueprintCallable, Category="Utilities|Transformation")
+	FName GetAttachSocketName() const;
 
 	/** Gets all parent components up to and including the root component */
 	UFUNCTION(BlueprintCallable, Category="Components")
@@ -1057,6 +1061,16 @@ protected:
 //////////////////////////////////////////////////////////////////////////
 // USceneComponent inlines
 
+FORCEINLINE USceneComponent* USceneComponent::GetAttachParent() const
+{
+	return AttachParent;
+}
+
+FORCEINLINE FName USceneComponent::GetAttachSocketName() const
+{
+	return AttachSocketName;
+}
+
 FORCEINLINE_DEBUGGABLE void USceneComponent::ConditionalUpdateComponentToWorld()
 {
 	if (!bWorldToComponentUpdated)
diff --git a/Engine/Source/Runtime/Engine/Private/Components/SceneComponent.cpp b/Engine/Source/Runtime/Engine/Private/Components/SceneComponent.cpp
--- a/Engine/Source/Runtime/Engine/Private/Components/SceneComponent.cpp
+++ b/Engine/Source/Runtime/Engine/Private/Components/SceneComponent.cpp
@@ -1080,12 +1080,6 @@ FVector USceneComponent::K2_GetComponentScale() const
 	return GetComponentScale();
 }
 
-
-USceneComponent* USceneComponent::GetAttachParent() const
-{
-	return AttachParent;
-}
-
 void USceneComponent::GetParentComponents(TArray<class USceneComponent*>& Parents) const
 {
 	Parents.Empty();
