[UEDOC-2241] Adding suggested text to Sound Looping Node tooltip.

diff --git a/Engine/Source/Runtime/Engine/Classes/Sound/SoundNodeLooping.h b/Engine/Source/Runtime/Engine/Classes/Sound/SoundNodeLooping.h
--- a/Engine/Source/Runtime/Engine/Classes/Sound/SoundNodeLooping.h
+++ b/Engine/Source/Runtime/Engine/Classes/Sound/SoundNodeLooping.h
@@ -6,7 +6,10 @@
 #include "SoundNodeLooping.generated.h"
 
 /** 
- * Defines how a sound loops; either indefinitely, or for a set number of times
+ * Defines how a sound loops; either indefinitely, or for a set number of times.
+ * Note: The Looping node should only be used for logical or procedural looping such as introducing a delay.
+ * These sounds will not be played seamlessly. If you want a sound to loop seamlessly and indefinitely,
+ * use the Looping flag on the Wave Player node for that sound.
  */
 UCLASS(hidecategories=Object, editinlinenew, MinimalAPI, meta=( DisplayName="Looping" ))
 class USoundNodeLooping : public USoundNode
@@ -22,7 +25,7 @@ class USoundNodeLooping : public USoundNode
 	uint32 bLoopIndefinitely : 1;
 
 public:	
-	// Begin USoundNode interface. 
+	//~ Begin USoundNode interface. 
 	virtual bool NotifyWaveInstanceFinished( struct FWaveInstance* WaveInstance ) override;
 	virtual float MaxAudibleDistance( float CurrentMaxDistance ) override 
 	{ 
@@ -30,7 +33,7 @@ public:
 	}
 	virtual float GetDuration( void ) override;
 	virtual void ParseNodes( FAudioDevice* AudioDevice, const UPTRINT NodeWaveInstanceHash, FActiveSound& ActiveSound, const FSoundParseParameters& ParseParams, TArray<FWaveInstance*>& WaveInstances ) override;
-	// End USoundNode interface. 
+	//~ End USoundNode interface. 
 };
 
 
