#ue4 - USceneComponent::AttachTo() now returns a bool indicating if attachment is successful (or if already attached). OnRegister() now clears AttachParent and AttachSocketName if the AttachTo() call fails.

- Fixes AttachParent still being set when initialized in a constructor but when the attachment is invalid.

--------

diff --git a/Engine/Source/Runtime/Engine/Classes/Components/SceneComponent.h b/Engine/Source/Runtime/Engine/Classes/Components/SceneComponent.h
--- a/Engine/Source/Runtime/Engine/Classes/Components/SceneComponent.h
+++ b/Engine/Source/Runtime/Engine/Classes/Components/SceneComponent.h
@@ -551,15 +551,22 @@ public:
 	UFUNCTION(BlueprintCallable, Category="Components")
 	void GetChildrenComponents(bool bIncludeAllDescendants, TArray<USceneComponent*>& Children) const;
 
-	/** 
-	 *   Attach this component to another scene component, optionally at a named socket. It is valid to call this on components whether or not they have been Registered.
-	 *   @param bMaintainWorldTransform	If true, update the relative location/rotation of the component to keep its world position the same
+	/**
+	 * Attach this component to another scene component, optionally at a named socket. It is valid to call this on components whether or not they have been Registered.
+	 * @param  InParent				Parent to attach to.
+	 * @param  InSocketName			Optional socket to attach to on the parent.
+	 * @param  AttachType			How to handle transform when attaching (Keep relative offset, keep world position, etc).
+	 * @param  bWeldSimulatedBodies Whether to weld together simulated physics bodies.
+	 * @return True if attachment is successful (or already attached to requested parent/socket), false if attachment is rejected and there is no change in AttachParent.
 	 */
-	void AttachTo(class USceneComponent* InParent, FName InSocketName = NAME_None, EAttachLocation::Type AttachType = EAttachLocation::KeepRelativeOffset, bool bWeldSimulatedBodies = false);
+	 bool AttachTo(class USceneComponent* InParent, FName InSocketName = NAME_None, EAttachLocation::Type AttachType = EAttachLocation::KeepRelativeOffset, bool bWeldSimulatedBodies = false);
 
 	/**
-	*   Attach this component to another scene component, optionally at a named socket. It is valid to call this on components whether or not they have been Registered.
-	*   @param bMaintainWorldTransform	If true, update the relative location/rotation of the component to keep its world position the same
+	 * Attach this component to another scene component, optionally at a named socket. It is valid to call this on components whether or not they have been Registered.
+	 * @param  InParent				Parent to attach to.
+	 * @param  InSocketName			Optional socket to attach to on the parent.
+	 * @param  AttachType			How to handle transform when attaching (Keep relative offset, keep world position, etc).
+	 * @param  bWeldSimulatedBodies Whether to weld together simulated physics bodies.
 	*/
 	UFUNCTION(BlueprintCallable, Category = "Utilities|Transformation", meta = (DisplayName = "AttachTo", AttachType = "KeepRelativeOffset"))
 	void K2_AttachTo(class USceneComponent* InParent, FName InSocketName = NAME_None, EAttachLocation::Type AttachType = EAttachLocation::KeepRelativeOffset, bool bWeldSimulatedBodies = true);
diff --git a/Engine/Source/Runtime/Engine/Private/Components/SceneComponent.cpp b/Engine/Source/Runtime/Engine/Private/Components/SceneComponent.cpp
--- a/Engine/Source/Runtime/Engine/Private/Components/SceneComponent.cpp
+++ b/Engine/Source/Runtime/Engine/Private/Components/SceneComponent.cpp
@@ -374,7 +374,15 @@ void USceneComponent::OnRegister()
 	// If we need to perform a call to AttachTo, do that now
 	// At this point scene component still has no any state (rendering, physics),
 	// so this call will just add this component to an AttachChildren array of a the Parent component
-	AttachTo(AttachParent, AttachSocketName);
+	if (AttachParent)
+	{
+		if (AttachTo(AttachParent, AttachSocketName) == false)
+		{
+			// Failed to attach, we need to clear AttachParent so we don't think we're actually attached when we're not.
+			AttachParent = nullptr;
+			AttachSocketName = NAME_None;
+		}
+	}
 	
 	Super::OnRegister();
 
@@ -1147,21 +1155,21 @@ void USceneComponent::K2_AttachTo(class USceneComponent* InParent, FName InSocke
 	AttachTo(InParent, InSocketName, AttachLocationType, bWeldSimulatedBodies);
 }
 
-void USceneComponent::AttachTo(class USceneComponent* Parent, FName InSocketName, EAttachLocation::Type AttachType /*= EAttachLocation::KeepRelativeOffset */, bool bWeldSimulatedBodies /*= false*/)
+bool USceneComponent::AttachTo(class USceneComponent* Parent, FName InSocketName, EAttachLocation::Type AttachType /*= EAttachLocation::KeepRelativeOffset */, bool bWeldSimulatedBodies /*= false*/)
 {
 	if(Parent != NULL)
 	{
 		if (Parent == AttachParent && InSocketName == AttachSocketName && Parent->AttachChildren.Contains(this))
 		{
 			// already attached!
-			return;
+			return true;
 		}
 
 		if(Parent == this)
 		{
 			FMessageLog("PIE").Warning(FText::Format(LOCTEXT("AttachToSelfWarning", "AttachTo: '{0}' cannot be attached to itself. Aborting."), 
 				FText::FromString(GetPathName())));
-			return;
+			return false;
 		}
 
 		if(Parent->IsAttachedTo(this))
@@ -1169,13 +1177,13 @@ void USceneComponent::AttachTo(class USceneComponent* Parent, FName InSocketName
 			FMessageLog("PIE").Warning(FText::Format(LOCTEXT("AttachCycleWarning", "AttachTo: '{0}' already attached to '{1}', would form cycle. Aborting."), 
 				FText::FromString(Parent->GetPathName()), 
 				FText::FromString(GetPathName())));
-			return;
+			return false;
 		}
 
 		if(!Parent->CanAttachAsChild(this, InSocketName))
 		{
 			UE_LOG(LogSceneComponent, Warning, TEXT("AttachTo: '%s' will not allow '%s' to be attached as a child."), *Parent->GetPathName(), *GetPathName());
-			return;
+			return false;
 		}
 
 		// Don't allow components with static mobility to be attached to non-static parents (except during UCS)
@@ -1193,7 +1201,7 @@ void USceneComponent::AttachTo(class USceneComponent* Parent, FName InSocketName
 				FText::FromString(Parent->GetPathName()), 
 				FText::FromString(ExtraBlueprintInfo), 
 				FText::FromString(GetPathName())));
-			return;
+			return false;
 		}
 
 		// if our template type doesn't match
@@ -1207,7 +1215,7 @@ void USceneComponent::AttachTo(class USceneComponent* Parent, FName InSocketName
 			{
 				ensureMsgf(false, TEXT("Template Mismatch during attachment. Attaching template component to instanced component. Parent '%s' Self '%s'"), *Parent->GetName(), *GetName());
 			}
-			return;
+			return false;
 		}
 
 		// Don't call UpdateOverlaps() when detaching, since we are going to do it anyway after we reattach below.
@@ -1246,7 +1254,7 @@ void USceneComponent::AttachTo(class USceneComponent* Parent, FName InSocketName
 					RelativeScale3D = GetComponentScale();
 				}
 
-				return;
+				return false;
 			}
 		}
 
@@ -1367,7 +1375,11 @@ void USceneComponent::AttachTo(class USceneComponent* Parent, FName InSocketName
 		{
 			UpdateOverlaps();
 		}
+
+		return true;
 	}
+
+	return false;
 }
 
 void USceneComponent::SnapTo(class USceneComponent* Parent, FName InSocketName)
