fixed UE-19238 Fringe Saturation in Post Process Volume not functioning

feature was removed some time ago, this was left and should have been removed as well.

diff --git a/Engine/Source/Runtime/Engine/Private/SceneView.cpp b/Engine/Source/Runtime/Engine/Private/SceneView.cpp
--- a/Engine/Source/Runtime/Engine/Private/SceneView.cpp
+++ b/Engine/Source/Runtime/Engine/Private/SceneView.cpp
@@ -866,7 +866,6 @@ void FSceneView::OverridePostProcessSettings(const FPostProcessSettings& Src, fl
 
 		LERP_PP(SceneColorTint);
 		LERP_PP(SceneFringeIntensity);
-		LERP_PP(SceneFringeSaturation);
 		LERP_PP(BloomIntensity);
 		LERP_PP(BloomThreshold);
 		LERP_PP(Bloom1Tint);
diff --git a/Engine/Source/Runtime/Renderer/Private/PostProcess/PostProcessSceneColorFringe.cpp b/Engine/Source/Runtime/Renderer/Private/PostProcess/PostProcessSceneColorFringe.cpp
deleted file mode 100644
diff --git a/Engine/Source/Runtime/Renderer/Private/PostProcess/PostProcessSceneColorFringe.h b/Engine/Source/Runtime/Renderer/Private/PostProcess/PostProcessSceneColorFringe.h
deleted file mode 100644
