Engine: Added a method to get the squared distance between two actors (AActor::GetSquaredDistanceTo)
[INTEGRATE] Change 2667035 by [REDACTED_PERSON][REDACTED_HANDLE] on 2015/08/24 20:17:38

diff --git a/Engine/Source/Runtime/Engine/Classes/GameFramework/Actor.h b/Engine/Source/Runtime/Engine/Classes/GameFramework/Actor.h
--- a/Engine/Source/Runtime/Engine/Classes/GameFramework/Actor.h
+++ b/Engine/Source/Runtime/Engine/Classes/GameFramework/Actor.h
@@ -770,6 +770,10 @@ public:
 	UFUNCTION(BlueprintCallable, Category = "Utilities|Transformation")
 	float GetDistanceTo(const AActor* OtherActor) const;
 
+	/** Returns the squared distance from this Actor to OtherActor. */
+	UFUNCTION(BlueprintCallable, Category = "Utilities|Transformation")
+	float GetSquaredDistanceTo(const AActor* OtherActor) const;
+
 	/** Returns the distance from this Actor to OtherActor, ignoring Z. */
 	UFUNCTION(BlueprintCallable, Category = "Utilities|Transformation")
 	float GetHorizontalDistanceTo(const AActor* OtherActor) const;
diff --git a/Engine/Source/Runtime/Engine/Private/Actor.cpp b/Engine/Source/Runtime/Engine/Private/Actor.cpp
--- a/Engine/Source/Runtime/Engine/Private/Actor.cpp
+++ b/Engine/Source/Runtime/Engine/Private/Actor.cpp
@@ -3998,6 +3998,11 @@ float AActor::GetDistanceTo(const AActor* OtherActor) const
 	return OtherActor ? (GetActorLocation() - OtherActor->GetActorLocation()).Size() : 0.f;
 }
 
+float AActor::GetSquaredDistanceTo(const AActor* OtherActor) const
+{
+	return OtherActor ? (GetActorLocation() - OtherActor->GetActorLocation()).SizeSquared() : 0.f;
+}
+
 float AActor::GetHorizontalDistanceTo(const AActor* OtherActor) const
 {
 	return OtherActor ? (GetActorLocation() - OtherActor->GetActorLocation()).Size2D() : 0.f;
