Fix for infinite loop bug when blending looping and non looping animations
Merged from 2664352

diff --git a/Engine/Source/Runtime/Engine/Private/Animation/AnimSequence.cpp b/Engine/Source/Runtime/Engine/Private/Animation/AnimSequence.cpp
--- a/Engine/Source/Runtime/Engine/Private/Animation/AnimSequence.cpp
+++ b/Engine/Source/Runtime/Engine/Private/Animation/AnimSequence.cpp
@@ -3848,14 +3848,14 @@ void UAnimSequence::AdvanceMarkerPhaseAsFollower(const FMarkerTickContext& Conte
 		{
 			if (NextMarker.MarkerIndex == -1)
 			{
-				check(!bLooping); // shouldnt have an end of anim marker if looping
+				check(!bLooping || Context.GetMarkerSyncEndPosition().NextMarkerName == NAME_None); // shouldnt have an end of anim marker if looping
 				CurrentTime = FMath::Min(CurrentTime + DeltaRemaining, SequenceLength);
 				break;
 			}
 			else if (PassedMarkersIndex < Context.MarkersPassedThisTick.Num())
 			{
 				PreviousMarker.MarkerIndex = NextMarker.MarkerIndex;
-
+				checkSlow(NextMarker.MarkerIndex != -1);
 				const FPassedMarker& PassedMarker = Context.MarkersPassedThisTick[PassedMarkersIndex];
 				AdvanceMarkerForwards(NextMarker.MarkerIndex, PassedMarker.PassedMarkerName, bLooping, AuthoredSyncMarkers);
 				if (NextMarker.MarkerIndex == -1)
@@ -3868,17 +3868,18 @@ void UAnimSequence::AdvanceMarkerPhaseAsFollower(const FMarkerTickContext& Conte
 
 		const FMarkerSyncAnimPosition& End = Context.GetMarkerSyncEndPosition();
 		
-		if (Context.MarkersPassedThisTick.Num() > 0)
+		if (End.NextMarkerName == NAME_None)
 		{
-			AdvanceMarkerForwards(NextMarker.MarkerIndex, End.NextMarkerName, bLooping, AuthoredSyncMarkers);
+			NextMarker.MarkerIndex = -1;
 		}
 
-		//Validation
-		if (End.NextMarkerName == NAME_None)
+		if (NextMarker.MarkerIndex != -1 && Context.MarkersPassedThisTick.Num() > 0)
 		{
-			check(NextMarker.MarkerIndex == -1);
+			AdvanceMarkerForwards(NextMarker.MarkerIndex, End.NextMarkerName, bLooping, AuthoredSyncMarkers);
 		}
-		else if (NextMarker.MarkerIndex != -1)
+
+		//Validation
+		if (NextMarker.MarkerIndex != -1)
 		{
 			check(AuthoredSyncMarkers[NextMarker.MarkerIndex].MarkerName == End.NextMarkerName);
 		}
@@ -3892,14 +3893,14 @@ void UAnimSequence::AdvanceMarkerPhaseAsFollower(const FMarkerTickContext& Conte
 		{
 			if (PreviousMarker.MarkerIndex == -1)
 			{
-				check(!bLooping); // shouldnt have an end of anim marker if looping
+				check(!bLooping || Context.GetMarkerSyncEndPosition().PreviousMarkerName == NAME_None); // shouldnt have an end of anim marker if looping
 				CurrentTime = FMath::Max(CurrentTime + DeltaRemaining, 0.f);
 				break;
 			}
 			else if (PassedMarkersIndex < Context.MarkersPassedThisTick.Num())
 			{
 				NextMarker.MarkerIndex = PreviousMarker.MarkerIndex;
-
+				checkSlow(PreviousMarker.MarkerIndex != -1);
 				const FPassedMarker& PassedMarker = Context.MarkersPassedThisTick[PassedMarkersIndex];
 				AdvanceMarkerBackwards(PreviousMarker.MarkerIndex, PassedMarker.PassedMarkerName, bLooping, AuthoredSyncMarkers);
 				if (PreviousMarker.MarkerIndex == -1)
@@ -3912,17 +3913,18 @@ void UAnimSequence::AdvanceMarkerPhaseAsFollower(const FMarkerTickContext& Conte
 
 		const FMarkerSyncAnimPosition& End = Context.GetMarkerSyncEndPosition();
 
-		if (Context.MarkersPassedThisTick.Num() > 0)
+		if (PreviousMarker.MarkerIndex != -1 && Context.MarkersPassedThisTick.Num() > 0)
 		{
 			AdvanceMarkerBackwards(PreviousMarker.MarkerIndex, End.PreviousMarkerName, bLooping, AuthoredSyncMarkers);
 		}
 
-		//Validation
 		if (End.PreviousMarkerName == NAME_None)
 		{
-			check(PreviousMarker.MarkerIndex == -1);
+			PreviousMarker.MarkerIndex = -1;
 		}
-		else if (PreviousMarker.MarkerIndex != -1)
+
+		//Validation
+		if (PreviousMarker.MarkerIndex != -1)
 		{
 			check(AuthoredSyncMarkers[PreviousMarker.MarkerIndex].MarkerName == End.PreviousMarkerName);
 		}
