#ANIM:  Add Start Position for Sequence Player

Merging

//depot/UE4-Orion/Engine/Source/Runtime/Engine/...

to //depot/UE4/Engine/Source/Runtime/Engine/...

diff --git a/Engine/Source/Runtime/Engine/Classes/Animation/AnimNode_SequencePlayer.h b/Engine/Source/Runtime/Engine/Classes/Animation/AnimNode_SequencePlayer.h
--- a/Engine/Source/Runtime/Engine/Classes/Animation/AnimNode_SequencePlayer.h
+++ b/Engine/Source/Runtime/Engine/Classes/Animation/AnimNode_SequencePlayer.h
@@ -25,6 +25,11 @@ public:
 	UPROPERTY(EditAnywhere, BlueprintReadWrite, Category=Settings, meta=(PinHiddenByDefault))
 	mutable float PlayRate;
 
+	// The start up position, it only applies when reinitialized
+	// if you loop, it will still start from 0.f after finishing the round
+	UPROPERTY(EditAnywhere, BlueprintReadWrite, Category=Settings, meta=(PinHiddenByDefault))
+	mutable float StartPosition;
+
 	// The group index, assigned at compile time based on the editoronly GroupName (or INDEX_NONE if it is not part of any group)
 	UPROPERTY()
 	int32 GroupIndex;
@@ -37,6 +42,7 @@ public:
 		: Sequence(NULL)
 		, bLoopAnimation(true)
 		, PlayRate(1.0f)
+		, StartPosition(0.f)
 		, GroupIndex(INDEX_NONE)
 	{
 	}
diff --git a/Engine/Source/Runtime/Engine/Private/Animation/AnimNode_SequencePlayer.cpp b/Engine/Source/Runtime/Engine/Private/Animation/AnimNode_SequencePlayer.cpp
--- a/Engine/Source/Runtime/Engine/Private/Animation/AnimNode_SequencePlayer.cpp
+++ b/Engine/Source/Runtime/Engine/Private/Animation/AnimNode_SequencePlayer.cpp
@@ -9,11 +9,12 @@
 void FAnimNode_SequencePlayer::Initialize(const FAnimationInitializeContext& Context)
 {
 	EvaluateGraphExposedInputs.Execute(Context);
-	InternalTimeAccumulator = 0.0f;
-
+	InternalTimeAccumulator = StartPosition;
 	if (Sequence != NULL)
 	{
-		if ((PlayRate * Sequence->RateScale) < 0.0f)
+		InternalTimeAccumulator = FMath::Clamp(StartPosition, 0.f, Sequence->SequenceLength);
+
+		if (StartPosition == 0.f && (PlayRate * Sequence->RateScale) < 0.0f)
 		{
 			InternalTimeAccumulator = Sequence->SequenceLength;
 		}
