New LOD Coloration viewmode.
Enabled from the viewmode menu, or using "viewmode lodcoloration" command.
LOD colors are configured in BaseEngine.ini under LODColorationColors.
Unified landscape debug lod colors with the new view mode colors.

diff --git a/Engine/Config/BaseEngine.ini b/Engine/Config/BaseEngine.ini
--- a/Engine/Config/BaseEngine.ini
+++ b/Engine/Config/BaseEngine.ini
@@ -115,6 +115,14 @@ LightingOnlyBrightness=(R=0.3,G=0.3,B=0.3,A=1.0)
 +StationaryLightOverlapColors=(R=1.0,G=0.0,B=0.0,A=1.0)
 +StationaryLightOverlapColors=(R=1.0,G=0.0,B=0.5,A=1.0)
 +StationaryLightOverlapColors=(R=1.0,G=0.9,B=0.9,A=1.0)
++LODColorationColors=(R=1.0,G=1.0,B=1.0,A=1.0)
++LODColorationColors=(R=1.0,G=0.0,B=0.0,A=1.0)
++LODColorationColors=(R=0.0,G=1.0,B=0.0,A=1.0)
++LODColorationColors=(R=0.0,G=0.0,B=1.0,A=1.0)
++LODColorationColors=(R=1.0,G=1.0,B=0.0,A=1.0)
++LODColorationColors=(R=1.0,G=0.0,B=1.0,A=1.0)
++LODColorationColors=(R=0.0,G=1.0,B=1.0,A=1.0)
++LODColorationColors=(R=0.5,G=0.0,B=0.5,A=1.0)
 MaxPixelShaderAdditiveComplexityCount=2000
 MaxES2PixelShaderAdditiveComplexityCount=600
 bSubtitlesEnabled=True
diff --git a/Engine/Source/Editor/EditorStyle/Private/SlateEditorStyle.cpp b/Engine/Source/Editor/EditorStyle/Private/SlateEditorStyle.cpp
--- a/Engine/Source/Editor/EditorStyle/Private/SlateEditorStyle.cpp
+++ b/Engine/Source/Editor/EditorStyle/Private/SlateEditorStyle.cpp
@@ -4228,7 +4228,8 @@ void FSlateEditorStyle::FStyle::SetupLevelEditorStyle()
 		Set( "EditorViewport.ShaderComplexityMode", new IMAGE_BRUSH( "Icons/icon_ViewMode_Shadercomplexity_16px", Icon16x16 ) );
 		Set( "EditorViewport.StationaryLightOverlapMode", new IMAGE_BRUSH( "Icons/icon_ViewMode_StationaryLightOverlap_16px", Icon16x16 ) );
 		Set( "EditorViewport.LightmapDensityMode", new IMAGE_BRUSH( "Icons/icon_ViewMode_LightmapDensity_16px", Icon16x16 ) );
-		Set( "EditorViewport.VisualizeGBufferMode", new IMAGE_BRUSH( "Icons/icon_ViewMode_VisualisationGBuffer_16px", Icon16x16 ) );
+		Set( "EditorViewport.LODColorationMode", new IMAGE_BRUSH("Icons/icon_ViewMode_LODColoration_16px", Icon16x16) );
+		Set( "EditorViewport.VisualizeGBufferMode", new IMAGE_BRUSH("Icons/icon_ViewMode_VisualisationGBuffer_16px", Icon16x16) );
 		Set( "EditorViewport.ReflectionOverrideMode", new IMAGE_BRUSH( "Icons/icon_ViewMode_ReflectionOverride_16px", Icon16x16 ) );
 		Set( "EditorViewport.VisualizeBufferMode", new IMAGE_BRUSH( "Icons/icon_ViewMode_VisualisationGBuffer_16px", Icon16x16 ) );
 		Set( "EditorViewport.CollisionPawn", new IMAGE_BRUSH( "Icons/icon_ViewMode_CollsionPawn_16px", Icon16x16 ) );
diff --git a/Engine/Source/Editor/UnrealEd/Private/EditorViewportCommands.cpp b/Engine/Source/Editor/UnrealEd/Private/EditorViewportCommands.cpp
--- a/Engine/Source/Editor/UnrealEd/Private/EditorViewportCommands.cpp
+++ b/Engine/Source/Editor/UnrealEd/Private/EditorViewportCommands.cpp
@@ -25,7 +25,8 @@ void FEditorViewportCommands::RegisterCommands()
 	UI_COMMAND( ShaderComplexityMode, "Shader Complexity View Mode", "Renders the scene with shader complexity visualization", EUserInterfaceActionType::RadioButton, FInputChord( EModifierKey::Alt, EKeys::Eight ) );
 	UI_COMMAND( StationaryLightOverlapMode, "Stationary Light Overlap View Mode", "Visualizes overlap of stationary lights", EUserInterfaceActionType::RadioButton, FInputChord() );
 	UI_COMMAND( LightmapDensityMode, "Lightmap Density View Mode", "Renders the scene with lightmap density visualization", EUserInterfaceActionType::RadioButton, FInputChord( EModifierKey::Alt, EKeys::Zero ) );
-	UI_COMMAND( VisualizeBufferMode, "Buffer Visualization View Mode", "Renders a set of selected post process materials, which visualize various intermediate render buffers (material attributes)", EUserInterfaceActionType::RadioButton, FInputChord() );
+	UI_COMMAND( LODColorationMode, "LOD Coloration View Mode", "Renders the scene using LOD color visualization", EUserInterfaceActionType::RadioButton, FInputChord() );
+	UI_COMMAND( VisualizeBufferMode, "Buffer Visualization View Mode", "Renders a set of selected post process materials, which visualize various intermediate render buffers (material attributes)", EUserInterfaceActionType::RadioButton, FInputChord());
 	UI_COMMAND( ReflectionOverrideMode, "Reflections View Mode", "Renders the scene with reflections only", EUserInterfaceActionType::RadioButton, FInputChord() );
 	UI_COMMAND( CollisionPawn, "Player Collision", "Renders player collision visualization", EUserInterfaceActionType::RadioButton, FInputChord() );
 	UI_COMMAND( CollisionVisibility, "Visibility Collision", "Renders visibility collision visualization", EUserInterfaceActionType::RadioButton, FInputChord() );
diff --git a/Engine/Source/Editor/UnrealEd/Private/SEditorViewport.cpp b/Engine/Source/Editor/UnrealEd/Private/SEditorViewport.cpp
--- a/Engine/Source/Editor/UnrealEd/Private/SEditorViewport.cpp
+++ b/Engine/Source/Editor/UnrealEd/Private/SEditorViewport.cpp
@@ -352,6 +352,7 @@ void SEditorViewport::BindCommands()
 	MAP_VIEWMODE_ACTION( Commands.StationaryLightOverlapMode, VMI_StationaryLightOverlap );
 	MAP_VIEWMODE_ACTION( Commands.LightmapDensityMode, VMI_LightmapDensity );
 	MAP_VIEWMODE_ACTION( Commands.ReflectionOverrideMode, VMI_ReflectionOverride );
+	MAP_VIEWMODE_ACTION( Commands.LODColorationMode, VMI_LODColoration );
 	MAP_VIEWMODE_ACTION( Commands.VisualizeBufferMode, VMI_VisualizeBuffer );
 	MAP_VIEWMODE_ACTION( Commands.CollisionPawn, VMI_CollisionPawn);
 	MAP_VIEWMODE_ACTION( Commands.CollisionVisibility, VMI_CollisionVisibility);
diff --git a/Engine/Source/Editor/UnrealEd/Private/SEditorViewportViewMenu.cpp b/Engine/Source/Editor/UnrealEd/Private/SEditorViewportViewMenu.cpp
--- a/Engine/Source/Editor/UnrealEd/Private/SEditorViewportViewMenu.cpp
+++ b/Engine/Source/Editor/UnrealEd/Private/SEditorViewportViewMenu.cpp
@@ -94,6 +94,10 @@ FText SEditorViewportViewMenu::GetViewMenuLabel() const
 
 			case VMI_LitLightmapDensity:
 				break;
+
+			case VMI_LODColoration:
+				Label = LOCTEXT("ViewMenuTitle_LODColoration", "LOD Coloration");
+				break;
 		}
 	}
 
@@ -116,6 +120,7 @@ const FSlateBrush* SEditorViewportViewMenu::GetViewMenuLabelIcon() const
 		static FName LightOverlapIcon("EditorViewport.StationaryLightOverlapMode");
 		static FName LightmapDensityIcon("EditorViewport.LightmapDensityMode");
 		static FName ReflectionModeIcon("EditorViewport.ReflectionOverrideMode");
+		static FName LODColorationIcon("EditorViewport.LODColorationMode");
 		static FName VisualizeBufferIcon("EditorViewport.VisualizeBufferMode");
 		static FName CollisionPawnIcon("EditorViewport.CollisionPawn");
 		static FName CollisionVisibilityIcon("EditorViewport.CollisionVisibility");
@@ -180,6 +185,10 @@ const FSlateBrush* SEditorViewportViewMenu::GetViewMenuLabelIcon() const
 
 			case VMI_LitLightmapDensity:
 				break;
+
+			case VMI_LODColoration:
+				Icon = LODColorationIcon;
+				break;
 		}
 	}
 
@@ -209,6 +218,7 @@ TSharedRef<SWidget> SEditorViewportViewMenu::GenerateViewMenuContent() const
 				ViewMenuBuilder.AddMenuEntry( BaseViewportActions.StationaryLightOverlapMode, NAME_None, LOCTEXT("StationaryLightOverlapViewModeDisplayName", "Stationary Light Overlap") );
 				ViewMenuBuilder.AddMenuEntry( BaseViewportActions.LightmapDensityMode, NAME_None, LOCTEXT("LightmapDensityViewModeDisplayName", "Lightmap Density") );
 				ViewMenuBuilder.AddMenuEntry( BaseViewportActions.ReflectionOverrideMode, NAME_None, LOCTEXT("ReflectionOverrideViewModeDisplayName", "Reflections") );
+				ViewMenuBuilder.AddMenuEntry( BaseViewportActions.LODColorationMode, NAME_None, LOCTEXT("LODColorationViewModeDisplayName", "LOD Coloration") );
 			}
 			ViewMenuBuilder.EndSection();
 		}
diff --git a/Engine/Source/Editor/UnrealEd/Public/EditorViewportCommands.h b/Engine/Source/Editor/UnrealEd/Public/EditorViewportCommands.h
--- a/Engine/Source/Editor/UnrealEd/Public/EditorViewportCommands.h
+++ b/Engine/Source/Editor/UnrealEd/Public/EditorViewportCommands.h
@@ -74,6 +74,9 @@ public:
 	/** Changes the viewport to vertex densities mode */
 	TSharedPtr< FUICommandInfo > VertexDensitiesMode;
 
+	/** Changes the viewport to LOD index coloration mode */
+	TSharedPtr< FUICommandInfo > LODColorationMode;
+
 	/** Changes the viewport to visualize the buffer content */
 	TSharedPtr< FUICommandInfo > VisualizeBufferMode;
 
diff --git a/Engine/Source/Runtime/Engine/Classes/Engine/Engine.h b/Engine/Source/Runtime/Engine/Classes/Engine/Engine.h
--- a/Engine/Source/Runtime/Engine/Classes/Engine/Engine.h
+++ b/Engine/Source/Runtime/Engine/Classes/Engine/Engine.h
@@ -988,6 +988,10 @@ public:
 	UPROPERTY(globalconfig)
 	TArray<FLinearColor> StationaryLightOverlapColors;
 
+	/** The colors used to render LOD coloration. */
+	UPROPERTY(globalconfig)
+	TArray<FLinearColor> LODColorationColors;
+
 	/**
 	* Complexity limits for the various complexity view mode combinations.
 	* These limits are used to map instruction counts to ShaderComplexityColors.
diff --git a/Engine/Source/Runtime/Engine/Classes/Engine/EngineBaseTypes.h b/Engine/Source/Runtime/Engine/Classes/Engine/EngineBaseTypes.h
--- a/Engine/Source/Runtime/Engine/Classes/Engine/EngineBaseTypes.h
+++ b/Engine/Source/Runtime/Engine/Classes/Engine/EngineBaseTypes.h
@@ -728,7 +728,7 @@ enum EViewModeIndex
 	VMI_LitLightmapDensity = 10,
 	VMI_ReflectionOverride = 11,
 	VMI_VisualizeBuffer = 12,
-//	VMI_VoxelLighting = 13,
+	//	VMI_VoxelLighting = 13,
 
 	/** Colored according to stationary light overlap. */
 	VMI_StationaryLightOverlap = 14,
@@ -736,6 +736,8 @@ enum EViewModeIndex
 	VMI_CollisionPawn = 15, 
 	VMI_CollisionVisibility = 16, 
 	VMI_VertexDensities = 17,
+	/** Colored according to the current LOD index. */
+	VMI_LODColoration = 18,
 	VMI_Max UMETA(Hidden),
 
 	VMI_Unknown = 255 UMETA(Hidden),
diff --git a/Engine/Source/Runtime/Engine/Private/Particles/ParticleComponents.cpp b/Engine/Source/Runtime/Engine/Private/Particles/ParticleComponents.cpp
--- a/Engine/Source/Runtime/Engine/Private/Particles/ParticleComponents.cpp
+++ b/Engine/Source/Runtime/Engine/Private/Particles/ParticleComponents.cpp
@@ -3322,6 +3322,7 @@ void UParticleSystemComponent::UpdateDynamicData(FParticleSystemSceneProxy* Prox
 		[REDACTED_EMAIL]. Remove thisline  - it is used for debugging purposes...
 		Proxy->SetLastDynamicData(Proxy->GetDynamicData());
 		[REDACTED_EMAIL]. END
+		Proxy->SetVisualizeLODIndex(GetCurrentLODIndex());
 #endif // !(UE_BUILD_SHIPPING || UE_BUILD_TEST)
 		Proxy->UpdateData( ParticleDynamicData );
 	}
diff --git a/Engine/Source/Runtime/Engine/Private/Particles/ParticleSystemRender.cpp b/Engine/Source/Runtime/Engine/Private/Particles/ParticleSystemRender.cpp
--- a/Engine/Source/Runtime/Engine/Private/Particles/ParticleSystemRender.cpp
+++ b/Engine/Source/Runtime/Engine/Private/Particles/ParticleSystemRender.cpp
@@ -1110,6 +1110,10 @@ void FDynamicSpriteEmitterData::GetDynamicMeshElementsEmitter(const FParticleSys
 
 				Mesh.bCanApplyViewModeOverrides = true;
 				Mesh.bUseWireframeSelectionColoring = Proxy->IsSelected();
+			
+			#if !(UE_BUILD_SHIPPING || UE_BUILD_TEST)
+				Mesh.VisualizeLODIndex = (int8)Proxy->GetVisualizeLODIndex();
+			#endif
 
 				Collector.AddMesh(ViewIndex, Mesh);
 			}
@@ -1562,6 +1566,10 @@ void FDynamicMeshEmitterData::GetDynamicMeshElementsEmitter(const FParticleSyste
 
 					Mesh.bCanApplyViewModeOverrides = true;
 					Mesh.bUseWireframeSelectionColoring = Proxy->IsSelected();
+			
+				#if !(UE_BUILD_SHIPPING || UE_BUILD_TEST)
+					Mesh.VisualizeLODIndex = (int8)Proxy->GetVisualizeLODIndex();
+				#endif
 
 					Collector.AddMesh(ViewIndex, Mesh);
 				}
@@ -2345,6 +2353,10 @@ void FDynamicBeam2EmitterData::GetDynamicMeshElementsEmitter(const FParticleSyst
 
 		Mesh.bCanApplyViewModeOverrides = true;
 		Mesh.bUseWireframeSelectionColoring = Proxy->IsSelected();
+		
+	#if !(UE_BUILD_SHIPPING || UE_BUILD_TEST)
+		Mesh.VisualizeLODIndex = (int8)Proxy->GetVisualizeLODIndex();
+	#endif
 
 		Collector.AddMesh(ViewIndex, Mesh);
 
@@ -5068,6 +5080,10 @@ void FDynamicTrailsEmitterData::GetDynamicMeshElementsEmitter(const FParticleSys
 
 		Mesh.bCanApplyViewModeOverrides = true;
 		Mesh.bUseWireframeSelectionColoring = Proxy->IsSelected();
+		
+	#if !(UE_BUILD_SHIPPING || UE_BUILD_TEST)
+		Mesh.VisualizeLODIndex = (int8)Proxy->GetVisualizeLODIndex();
+	#endif
 
 		Collector.AddMesh(ViewIndex, Mesh);
 
@@ -6302,6 +6318,7 @@ FParticleSystemSceneProxy::FParticleSystemSceneProxy(const UParticleSystemCompon
 		GetSelectionColor(FLinearColor(1.0f, 0.0f, 0.0f, 1.0f),false,false)
 		)
 	, PendingLODDistance(0.0f)
+	, VisualizeLODIndex(Component->GetCurrentLODIndex())
 	, LastFramePreRendered(-1)
 	, FirstFreeMeshBatch(0)
 {
@@ -6473,8 +6490,6 @@ void FParticleSystemSceneProxy::UpdateData_RenderThread(FParticleDynamicData* Ne
 
 void FParticleSystemSceneProxy::DetermineLODDistance(const FSceneView* View, int32 FrameNumber)
 {
-	int32	LODIndex = -1;
-
 	if (LODMethod == PARTICLESYSTEMLODMETHOD_Automatic)
 	{
 		// Default to the highest LOD level
diff --git a/Engine/Source/Runtime/Engine/Private/PrimitiveDrawingUtils.cpp b/Engine/Source/Runtime/Engine/Private/PrimitiveDrawingUtils.cpp
--- a/Engine/Source/Runtime/Engine/Private/PrimitiveDrawingUtils.cpp
+++ b/Engine/Source/Runtime/Engine/Private/PrimitiveDrawingUtils.cpp
@@ -1270,7 +1270,8 @@ bool IsRichView(const FSceneViewFamily& ViewFamily)
 		ViewFamily.EngineShowFlags.MeshEdges ||
 		ViewFamily.EngineShowFlags.LightInfluences ||
 		ViewFamily.EngineShowFlags.Wireframe ||
-		ViewFamily.EngineShowFlags.LevelColoration)
+		ViewFamily.EngineShowFlags.LevelColoration ||
+		ViewFamily.EngineShowFlags.LODColoration)
 	{
 		return true;
 	}
@@ -1334,6 +1335,24 @@ void ApplyViewModeOverrides(
 			Collector.RegisterOneFrameMaterialProxy(WireframeMaterialInstance);
 		}
 	}
+	else if (EngineShowFlags.LODColoration)
+	{
+		if (!Mesh.IsTranslucent(FeatureLevel) && GEngine->LODColorationColors.Num()  > 0)
+		{
+			int32 lodColorationIndex =  FMath::Clamp((int32)Mesh.VisualizeLODIndex, 0, GEngine->LODColorationColors.Num() - 1);
+	
+			bool bLit = Mesh.MaterialRenderProxy->GetMaterial(FeatureLevel)->GetShadingModel() != MSM_Unlit;
+			const UMaterial* LODColorationMaterial = (bLit && EngineShowFlags.Lighting) ? GEngine->LevelColorationLitMaterial : GEngine->LevelColorationUnlitMaterial;
+
+			auto LODColorationMaterialInstance = new FColoredMaterialRenderProxy(
+				LODColorationMaterial->GetRenderProxy( Mesh.MaterialRenderProxy->IsSelected(), Mesh.MaterialRenderProxy->IsHovered() ),
+				GetSelectionColor(GEngine->LODColorationColors[lodColorationIndex], bSelected, Mesh.MaterialRenderProxy->IsHovered() )
+				);
+
+			Mesh.MaterialRenderProxy = LODColorationMaterialInstance;
+			Collector.RegisterOneFrameMaterialProxy(LODColorationMaterialInstance);
+		}
+	}
 	else if (EngineShowFlags.LightComplexity)
 	{
 		// Don't render unlit translucency when in 'light complexity' viewmode.
diff --git a/Engine/Source/Runtime/Engine/Private/ShowFlags.cpp b/Engine/Source/Runtime/Engine/Private/ShowFlags.cpp
--- a/Engine/Source/Runtime/Engine/Private/ShowFlags.cpp
+++ b/Engine/Source/Runtime/Engine/Private/ShowFlags.cpp
@@ -279,6 +279,9 @@ void ApplyViewMode(EViewModeIndex ViewModeIndex, bool bPerspective, FEngineShowF
 		case VMI_VertexDensities:
 			bPostProcessing = false;
 			break;
+		case VMI_LODColoration:
+			bPostProcessing = true;
+			break;
 	}
 
 	if(!bPerspective)
@@ -305,6 +308,7 @@ void ApplyViewMode(EViewModeIndex ViewModeIndex, bool bPerspective, FEngineShowF
 	EngineShowFlags.SetCollisionPawn(ViewModeIndex == VMI_CollisionPawn);
 	EngineShowFlags.SetCollisionVisibility(ViewModeIndex == VMI_CollisionVisibility);
 	EngineShowFlags.SetVertexDensities(ViewModeIndex == VMI_VertexDensities);
+	EngineShowFlags.SetLODColoration(ViewModeIndex == VMI_LODColoration);
 }
 
 void EngineShowFlagOverride(EShowFlagInitMode ShowFlagInitMode, EViewModeIndex ViewModeIndex, FEngineShowFlags& EngineShowFlags, FName CurrentBufferVisualizationMode, bool bIsSplitScreen)
@@ -356,7 +360,7 @@ void EngineShowFlagOverride(EShowFlagInitMode ShowFlagInitMode, EViewModeIndex V
 			EngineShowFlags.ScreenSpaceReflections = 0;
 		}
 	}
-	
+
 	{
 		static const auto ICVar = IConsoleManager::Get().FindTConsoleVariableDataInt(TEXT("r.RefractionQuality"));
 		if(ICVar->GetValueOnGameThread() <= 0)
@@ -396,7 +400,7 @@ void EngineShowFlagOverride(EShowFlagInitMode ShowFlagInitMode, EViewModeIndex V
 			EngineShowFlags.DynamicShadows = 0;
 		}
 
-		if(ViewModeIndex == VMI_BrushWireframe)
+		if( ViewModeIndex == VMI_BrushWireframe)
 		{
 			EngineShowFlags.Brushes = 1;
 		}
@@ -431,6 +435,14 @@ void EngineShowFlagOverride(EShowFlagInitMode ShowFlagInitMode, EViewModeIndex V
 		{
 			EngineShowFlags.Translucency = 0;
 		}
+
+		if( ViewModeIndex == VMI_LODColoration )
+		{
+			EngineShowFlags.Lighting = 1;		// Best currently otherwise the image becomes hard to read.
+			EngineShowFlags.Fog = 0;			// Removed fog to improve color readability.
+			EngineShowFlags.AtmosphericFog = 0;
+			EngineShowFlags.Translucency = 0;	// Translucent are off because there are no color override shader currently for translucency.
+		}
 	}
 
 	// disable AA in full screen GBuffer visualization
@@ -569,6 +581,10 @@ EViewModeIndex FindViewMode(const FEngineShowFlags& EngineShowFlags)
 	{
 		return VMI_VertexDensities;
 	}
+	else if (EngineShowFlags.LODColoration)
+	{
+		return VMI_LODColoration;
+	}
 
 	return EngineShowFlags.Lighting ? VMI_Lit : VMI_Unlit;
 }
@@ -594,6 +610,7 @@ const TCHAR* GetViewModeName(EViewModeIndex ViewModeIndex)
 		case VMI_CollisionPawn:				return TEXT("CollisionPawn");
 		case VMI_CollisionVisibility:		return TEXT("CollisionVis");
 		case VMI_VertexDensities:			return TEXT("VertexDensity");
+		case VMI_LODColoration:				return TEXT("LODColoration");
 	}
 	return TEXT("");
 }
diff --git a/Engine/Source/Runtime/Engine/Private/SkeletalMesh.cpp b/Engine/Source/Runtime/Engine/Private/SkeletalMesh.cpp
--- a/Engine/Source/Runtime/Engine/Private/SkeletalMesh.cpp
+++ b/Engine/Source/Runtime/Engine/Private/SkeletalMesh.cpp
@@ -4706,6 +4706,10 @@ void FSkeletalMeshSceneProxy::GetDynamicElementsSection(const TArray<const FScen
 			Mesh.bCanApplyViewModeOverrides = true;
 			Mesh.bUseWireframeSelectionColoring = bIsSelected;
 
+		#if !(UE_BUILD_SHIPPING || UE_BUILD_TEST)
+			Mesh.VisualizeLODIndex = LODIndex;
+		#endif
+
 			Collector.AddMesh(ViewIndex, Mesh);
 
 			const int32 NumVertices = Chunk.NumRigidVertices + Chunk.NumSoftVertices;
diff --git a/Engine/Source/Runtime/Engine/Private/StaticMeshRender.cpp b/Engine/Source/Runtime/Engine/Private/StaticMeshRender.cpp
--- a/Engine/Source/Runtime/Engine/Private/StaticMeshRender.cpp
+++ b/Engine/Source/Runtime/Engine/Private/StaticMeshRender.cpp
@@ -190,6 +190,9 @@ bool FStaticMeshSceneProxy::GetShadowMeshElement(int32 LODIndex, int32 BatchInde
 	OutMeshBatch.ReverseCulling = IsLocalToWorldDeterminantNegative();
 	OutMeshBatch.CastShadow = true;
 	OutMeshBatch.LODIndex = LODIndex;
+#if !(UE_BUILD_SHIPPING || UE_BUILD_TEST)
+	OutMeshBatch.VisualizeLODIndex = LODIndex;
+#endif
 	OutMeshBatch.LCI = &ProxyLODInfo;
 	if (ForcedLodModel > 0) 
 	{
@@ -250,6 +253,9 @@ bool FStaticMeshSceneProxy::GetMeshElement(int32 LODIndex, int32 BatchIndex, int
 		OutBatchElement.MinVertexIndex = Section.MinVertexIndex;
 		OutBatchElement.MaxVertexIndex = Section.MaxVertexIndex;
 		OutMeshBatch.LODIndex = LODIndex;
+#if !(UE_BUILD_SHIPPING || UE_BUILD_TEST)
+		OutMeshBatch.VisualizeLODIndex = LODIndex;
+#endif
 		OutMeshBatch.UseDynamicData = false;
 		OutMeshBatch.ReverseCulling = IsLocalToWorldDeterminantNegative();
 		OutMeshBatch.CastShadow = bCastShadow && Section.bCastShadow;
diff --git a/Engine/Source/Runtime/Engine/Public/MeshBatch.h b/Engine/Source/Runtime/Engine/Public/MeshBatch.h
--- a/Engine/Source/Runtime/Engine/Public/MeshBatch.h
+++ b/Engine/Source/Runtime/Engine/Public/MeshBatch.h
@@ -71,6 +71,9 @@ struct FMeshBatch
 	/** LOD index of the mesh, used for fading LOD transitions. */
 	int8 LODIndex;
 
+	/** Conceptual LOD index used for the LOD Coloration visualization. */
+	int8 VisualizeLODIndex;
+
 	uint32 UseDynamicData : 1;
 	uint32 ReverseCulling : 1;
 	uint32 bDisableBackfaceCulling : 1;
@@ -194,6 +197,7 @@ struct FMeshBatch
 	FMeshBatch()
 	:	DynamicVertexStride(0)
 	,	LODIndex(INDEX_NONE)
+	,	VisualizeLODIndex(INDEX_NONE)
 	,	UseDynamicData(false)
 	,	ReverseCulling(false)
 	,	bDisableBackfaceCulling(false)
diff --git a/Engine/Source/Runtime/Engine/Public/ParticleHelper.h b/Engine/Source/Runtime/Engine/Public/ParticleHelper.h
--- a/Engine/Source/Runtime/Engine/Public/ParticleHelper.h
+++ b/Engine/Source/Runtime/Engine/Public/ParticleHelper.h
@@ -2477,7 +2477,10 @@ public:
 	{
 		return MaterialRelevance;
 	}
-	float GetPendingLODDistance()		{	return PendingLODDistance;		}
+	float GetPendingLODDistance() const	{	return PendingLODDistance;		}
+	void SetVisualizeLODIndex(int32 InVisualizeLODIndex) { VisualizeLODIndex = InVisualizeLODIndex; }
+	int32  GetVisualizeLODIndex() const { return VisualizeLODIndex; }
+
 	inline const TUniformBuffer<FPrimitiveUniformShaderParameters>& GetWorldSpacePrimitiveUniformBuffer() const { return WorldSpacePrimitiveUniformBuffer; }
 
 	const FColoredMaterialRenderProxy* GetDeselectedWireframeMatInst() const	{	return &DeselectedWireframeMaterialInstance;	}
@@ -2515,6 +2518,7 @@ protected:
 
 	int32 LODMethod;
 	float PendingLODDistance;
+	int32 VisualizeLODIndex; // Only used in the LODColoration view mode.
 
 	int32 LastFramePreRendered;
 
diff --git a/Engine/Source/Runtime/Engine/Public/ShowFlags.h b/Engine/Source/Runtime/Engine/Public/ShowFlags.h
--- a/Engine/Source/Runtime/Engine/Public/ShowFlags.h
+++ b/Engine/Source/Runtime/Engine/Public/ShowFlags.h
@@ -292,6 +292,7 @@ private:
 		SetStationaryLightOverlap(false);
 		SetLightMapDensity(false);
 		SetVertexDensities(false);
+		SetLODColoration(false);
 		SetVisualizeLPV(false);
 		SetStreamingBounds(false);
 		SetConstraints(false);
diff --git a/Engine/Source/Runtime/Engine/Public/ShowFlagsValues.inl b/Engine/Source/Runtime/Engine/Public/ShowFlagsValues.inl
--- a/Engine/Source/Runtime/Engine/Public/ShowFlagsValues.inl
+++ b/Engine/Source/Runtime/Engine/Public/ShowFlagsValues.inl
@@ -290,6 +290,8 @@ SHOWFLAG_ALWAYS_ACCESSIBLE(VisualizeSenses, SFG_Advanced, LOCTEXT("VisualizeSens
 SHOWFLAG_FIXED_IN_SHIPPING(VisualizeBloom, 0, SFG_Visualize, LOCTEXT("VisualizeBloom", "Bloom"))
 /** Visualize Vertex Density */
 SHOWFLAG_ALWAYS_ACCESSIBLE(VertexDensities, SFG_Hidden, LOCTEXT("VisualizeVertexDensities", "Visualize Vertex Densities"))
+/** Visualize LOD Coloration */
+SHOWFLAG_FIXED_IN_SHIPPING(LODColoration, 0, SFG_Hidden, LOCTEXT("VisualizeLODColoration", "Visualize LOD Coloration"))
 
 #undef SHOWFLAG_ALWAYS_ACCESSIBLE
 #undef SHOWFLAG_FIXED_IN_SHIPPING
diff --git a/Engine/Source/Runtime/Landscape/Private/LandscapeRender.cpp b/Engine/Source/Runtime/Landscape/Private/LandscapeRender.cpp
--- a/Engine/Source/Runtime/Landscape/Private/LandscapeRender.cpp
+++ b/Engine/Source/Runtime/Landscape/Private/LandscapeRender.cpp
@@ -1386,6 +1386,10 @@ void FLandscapeComponentSceneProxy::GetDynamicMeshElements(const TArray<const FS
 					int32 SubSectionIdx = SubX + SubY*NumSubsections;
 					int32 CurrentLOD = CalcLODForSubsection(*View, SubX, SubY, CameraLocalPos);
 
+			#if !(UE_BUILD_SHIPPING || UE_BUILD_TEST)
+					// We simplify this by considering only the biggest LOD index for this mesh element.
+					Mesh.VisualizeLODIndex = (int8)FMath::Max((int32)Mesh.VisualizeLODIndex, CurrentLOD);
+			#endif
 					FMeshBatchElement& BatchElement = (SubX == 0 && SubY == 0) ? *Mesh.Elements.GetData() : *(new(Mesh.Elements) FMeshBatchElement);
 					BatchElement.PrimitiveUniformBufferResource = &GetUniformBuffer();
 					FLandscapeBatchElementParams& BatchElementParams = ParameterArray.ElementParams[SubSectionIdx];
@@ -1482,24 +1486,24 @@ void FLandscapeComponentSceneProxy::GetDynamicMeshElements(const TArray<const FS
 
 			case ELandscapeViewMode::LOD:
 			{
-				FLinearColor WireColors[LANDSCAPE_LOD_LEVELS];
-				WireColors[0] = FLinearColor(1, 1, 1);
-				WireColors[1] = FLinearColor(1, 0, 0);
-				WireColors[2] = FLinearColor(0, 1, 0);
-				WireColors[3] = FLinearColor(0, 0, 1);
-				WireColors[4] = FLinearColor(1, 1, 0);
-				WireColors[5] = FLinearColor(1, 0, 1);
-				WireColors[6] = FLinearColor(0, 1, 1);
-				WireColors[7] = FLinearColor(0.5f, 0, 0.5f);
-
 				for (int32 i = 0; i < MeshTools.Elements.Num(); i++)
 				{
 					FMeshBatch& LODMesh = Collector.AllocateMesh();
 					LODMesh = MeshTools;
 					LODMesh.Elements.Empty(1);
 					LODMesh.Elements.Add(MeshTools.Elements[i]);
-					int32 ColorIndex = ((FLandscapeBatchElementParams*)MeshTools.Elements[i].UserData)->CurrentLOD;
-					FLinearColor Color = ForcedLOD >= 0 ? WireColors[ColorIndex] : WireColors[ColorIndex] * 0.2f;
+
+					int32 ColorIndex = INDEX_NONE;
+					if (GEngine->LODColorationColors.Num() > 0)
+					{
+						ColorIndex = ((FLandscapeBatchElementParams*)MeshTools.Elements[i].UserData)->CurrentLOD;
+						ColorIndex = FMath::Clamp(ColorIndex, 0, GEngine->LODColorationColors.Num() - 1);
+						LODMesh.VisualizeLODIndex = ColorIndex;
+					}
+
+					const FLinearColor& LODColor = ColorIndex != INDEX_NONE ? GEngine->LODColorationColors[ColorIndex] : FLinearColor::Gray;
+					FLinearColor Color = ForcedLOD >= 0 ? LODColor : LODColor * 0.2f;
+
 					auto LODMaterialInstance = new FColoredMaterialRenderProxy(GEngine->LevelColorationUnlitMaterial->GetRenderProxy(false), Color);
 					LODMesh.MaterialRenderProxy = LODMaterialInstance;
 					Collector.RegisterOneFrameMaterialProxy(LODMaterialInstance);
