Removed unused DO_NAVMESH_DEBUG_DRAWING_PER_TILE code #UE4

github PR#1496 by slonopotamus
UE-20351

diff --git a/Engine/Source/Runtime/Engine/Classes/AI/Navigation/RecastNavMesh.h b/Engine/Source/Runtime/Engine/Classes/AI/Navigation/RecastNavMesh.h
--- a/Engine/Source/Runtime/Engine/Classes/AI/Navigation/RecastNavMesh.h
+++ b/Engine/Source/Runtime/Engine/Classes/AI/Navigation/RecastNavMesh.h
@@ -965,9 +965,6 @@ protected:
 	/** Returns query extent including adjustments for voxelization error compensation */
 	FVector GetModifiedQueryExtent(const FVector& QueryExtent) const;
 
-	// @todo docuement
-	UPrimitiveComponent* ConstructRenderingComponentImpl();
-
 	/** Spawns an ARecastNavMesh instance, and configures it if AgentProps != NULL */
 	static ARecastNavMesh* SpawnInstance(UNavigationSystem* NavSys, const FNavDataConfig* AgentProps = NULL);
 
diff --git a/Engine/Source/Runtime/Engine/Private/AI/Navigation/RecastNavMesh.cpp b/Engine/Source/Runtime/Engine/Private/AI/Navigation/RecastNavMesh.cpp
--- a/Engine/Source/Runtime/Engine/Private/AI/Navigation/RecastNavMesh.cpp
+++ b/Engine/Source/Runtime/Engine/Private/AI/Navigation/RecastNavMesh.cpp
@@ -134,8 +134,6 @@ void ARecastNavMesh::Serialize( FArchive& Ar )
 #include "RecastNavMeshGenerator.h"
 #include "DetourNavMeshQuery.h"
 
-#define DO_NAVMESH_DEBUG_DRAWING_PER_TILE 0
-
 //----------------------------------------------------------------------//
 // FRecastDebugGeometry
 //----------------------------------------------------------------------//
@@ -398,72 +396,17 @@ ANavigationData* ARecastNavMesh::CreateNavigationInstances(UNavigationSystem* Na
 
 UPrimitiveComponent* ARecastNavMesh::ConstructRenderingComponent() 
 {
-	return ConstructRenderingComponentImpl();
-}
-
-UPrimitiveComponent* ARecastNavMesh::ConstructRenderingComponentImpl() 
-{
-#if DO_NAVMESH_DEBUG_DRAWING_PER_TILE
-	const bool bIsGameThread = IsInGameThread();
-
-	if (RecastNavMeshImpl != NULL)
-	{
-		bool bComponentsRemoved = false;
-
-		for (int i = 0; i < TileRenderingComponents.Num(); ++i)
-		{
-			if (TileRenderingComponents[i] != NULL)
-			{
-				TileRenderingComponents[i]->UnregisterComponent();
-				Components.Remove(TileRenderingComponents[i]);
-				bComponentsRemoved = true;
-			}
-		}
-
-		TileRenderingComponents.Reset();
-		RecastNavMeshImpl->GenerateTileRenderingComponents(TileRenderingComponents);
-
-		//Components.Append(TileRenderingComponents);
-		for (int i = 0; i < TileRenderingComponents.Num(); ++i)
-		{
-			if (TileRenderingComponents[i] != NULL)
-			{
-				TileRenderingComponents[i]->RegisterComponent();
-			}
-		}
-
-		if (bComponentsRemoved || TileRenderingComponents.Num() > 0)
-		{
-			//RegisterAllComponents();
-			//MarkComponentsRenderStateDirty();
-		}
-	}
-
-	return NULL;
-#else
 	return NewObject<UNavMeshRenderingComponent>(this, TEXT("NavRenderingComp"), RF_Transient);
-#endif // DO_NAVMESH_DEBUG_DRAWING_PER_TILE
 }
 
 void ARecastNavMesh::UpdateNavMeshDrawing()
 {
 #if !UE_BUILD_SHIPPING
-#if DO_NAVMESH_DEBUG_DRAWING_PER_TILE
-	// @todo - we shouldn't be updating _every_ tile renderer, just the ones that have been changed
-	for (int i = 0; i < TileRenderingComponents.Num(); ++i)
-	{
-		if (TileRenderingComponents[i] != NULL && TileRenderingComponents[i]->bVisible)
-		{
-			TileRenderingComponents[i]->MarkRenderStateDirty();
-		}
-	}
-#else // DO_NAVMESH_DEBUG_DRAWING_PER_TILE
 	UNavMeshRenderingComponent* NavMeshRenderComp = Cast<UNavMeshRenderingComponent>(RenderingComp);
 	if (NavMeshRenderComp != nullptr && NavMeshRenderComp->bVisible && (NavMeshRenderComp->IsForcingUpdate() || UNavMeshRenderingComponent::IsNavigationShowFlagSet(GetWorld())))
 	{
 		RenderingComp->MarkRenderStateDirty();
 	}
-#endif // DO_NAVMESH_DEBUG_DRAWING_PER_TILE
 #endif // UE_BUILD_SHIPPING
 }
 
@@ -1434,9 +1377,6 @@ void ARecastNavMesh::RequestDrawingUpdate(bool bForce)
 
 void ARecastNavMesh::UpdateDrawing()
 {
-#if DO_NAVMESH_DEBUG_DRAWING_PER_TILE
-	ConstructRenderingComponentImpl();
-#endif
 	UpdateNavMeshDrawing();
 }
 
