Merging CLs 2667248 & 2667319:

Native anim instance support improvements

Added a few support functions to UAnimInstance to allow us to more easily query asset player states.
Also added optional TransitionName to native transition data so it can be displayed to the user whrn the binding does not supply a name of its own (i.e. non-UFUNCTIONs).

Reduced size of large function GetMachineDescription

Factored in GetStateMachineInstance & renamed to GetStateMachineInstanceDesc

diff --git a/Engine/Source/Runtime/Engine/Classes/Animation/AnimInstance.h b/Engine/Source/Runtime/Engine/Classes/Animation/AnimInstance.h
--- a/Engine/Source/Runtime/Engine/Classes/Animation/AnimInstance.h
+++ b/Engine/Source/Runtime/Engine/Classes/Animation/AnimInstance.h
@@ -707,6 +707,23 @@ protected:
 	template<class NodeType>
 	NodeType* GetNodeFromIndex(int32 NodeIdx);
 
+	/** 
+	 * Get the index of the specified instance asset player. Useful to pass to GetInstanceAssetPlayerLength (etc.).
+	 * Passing NAME_None to InstanceName will return the first (assumed only) player instance index found.
+	 */
+	int32 GetInstanceAssetPlayerIndex(FName MachineName, FName StateName, FName InstanceName = NAME_None);
+
+protected:
+
+	/** Gets the runtime instance of the specified state machine */
+	FAnimNode_StateMachine* GetStateMachineInstance(int32 MachineIndex);
+
+	/** Gets the runtime instance desc of the state machine specified by name */
+	const FBakedAnimationStateMachine* GetStateMachineInstanceDesc(FName MachineName);
+
+	/** Get the machine description for the specified instance. Does not rely on PRIVATE_MachineDescription being initialized */
+	const FBakedAnimationStateMachine* GetMachineDescription(UAnimBlueprintGeneratedClass* AnimBlueprintClass, FAnimNode_StateMachine* MachineInstance);
+
 	/** Gets the most relevant asset player in a specified state */
 	FAnimNode_AssetPlayerBase* GetRelevantAssetPlayerFromState(int32 MachineIndex, int32 StateIndex);
 
diff --git a/Engine/Source/Runtime/Engine/Private/Animation/AnimInstance.cpp b/Engine/Source/Runtime/Engine/Private/Animation/AnimInstance.cpp
--- a/Engine/Source/Runtime/Engine/Private/Animation/AnimInstance.cpp
+++ b/Engine/Source/Runtime/Engine/Private/Animation/AnimInstance.cpp
@@ -2980,6 +2980,33 @@ FAnimNode_StateMachine* UAnimInstance::GetStateMachineInstance(int32 MachineInde
 	return nullptr;
 }
 
+const FBakedAnimationStateMachine* UAnimInstance::GetStateMachineInstanceDesc(FName MachineName)
+{
+	if(UAnimBlueprintGeneratedClass* AnimBlueprintClass = Cast<UAnimBlueprintGeneratedClass>((UObject*)GetClass()))
+	{
+		for(int32 MachineIndex = 0; MachineIndex < AnimBlueprintClass->AnimNodeProperties.Num(); MachineIndex++)
+		{
+			UStructProperty* Property = AnimBlueprintClass->AnimNodeProperties[AnimBlueprintClass->AnimNodeProperties.Num() - 1 - MachineIndex];
+			if(Property && Property->Struct == FAnimNode_StateMachine::StaticStruct())
+			{
+				FAnimNode_StateMachine* StateMachine = Property->ContainerPtrToValuePtr<FAnimNode_StateMachine>(this);
+				if(StateMachine)
+				{
+					if(const FBakedAnimationStateMachine* MachineDescription = GetMachineDescription(AnimBlueprintClass, StateMachine))
+					{
+						if(MachineDescription->MachineName == MachineName)
+						{
+							return MachineDescription;
+						}
+					}
+				}
+			}
+		}
+	}
+
+	return nullptr;
+}
+
 template<class NodeType>
 NodeType* UAnimInstance::GetCheckedNodeFromIndex(int32 NodeIdx)
 {
@@ -3034,6 +3061,46 @@ NodeType* UAnimInstance::GetNodeFromIndex(int32 NodeIdx)
 	return NodePtr;
 }
 
+const FBakedAnimationStateMachine* UAnimInstance::GetMachineDescription(UAnimBlueprintGeneratedClass* AnimBlueprintClass, FAnimNode_StateMachine* MachineInstance)
+{
+	return AnimBlueprintClass->BakedStateMachines.IsValidIndex(MachineInstance->StateMachineIndexInClass) ? &(AnimBlueprintClass->BakedStateMachines[MachineInstance->StateMachineIndexInClass]) : nullptr;
+}
+
+int32 UAnimInstance::GetInstanceAssetPlayerIndex(FName MachineName, FName StateName, FName AssetName)
+{
+	if(UAnimBlueprintGeneratedClass* AnimBlueprintClass = Cast<UAnimBlueprintGeneratedClass>((UObject*)GetClass()))
+	{
+		if(const FBakedAnimationStateMachine* MachineDescription = GetStateMachineInstanceDesc(MachineName))
+		{
+			for(int32 StateIndex = 0; StateIndex < MachineDescription->States.Num(); StateIndex++)
+			{
+				const FBakedAnimationState& State = MachineDescription->States[StateIndex];
+				if(State.StateName == StateName)
+				{
+					for(int32 PlayerIndex = 0; PlayerIndex < State.PlayerNodeIndices.Num(); PlayerIndex++)
+					{
+						checkSlow(State.PlayerNodeIndices[PlayerIndex] < AnimBlueprintClass->AnimNodeProperties.Num());
+						UStructProperty* AssetPlayerProperty = AnimBlueprintClass->AnimNodeProperties[AnimBlueprintClass->AnimNodeProperties.Num() - 1 - State.PlayerNodeIndices[PlayerIndex]];
+						if(AssetPlayerProperty && AssetPlayerProperty->Struct->IsChildOf(FAnimNode_AssetPlayerBase::StaticStruct()))
+						{
+							FAnimNode_AssetPlayerBase* AssetPlayer = AssetPlayerProperty->ContainerPtrToValuePtr<FAnimNode_AssetPlayerBase>(this);
+							if(AssetPlayer)
+							{
+								if(AssetName == NAME_None || AssetPlayer->GetAnimAsset()->GetFName() == AssetName)
+								{
+									return State.PlayerNodeIndices[PlayerIndex];
+								}
+							}
+						}
+					}
+				}
+			}
+		}
+	}
+
+	return INDEX_NONE;
+}
+
 FAnimNode_AssetPlayerBase* UAnimInstance::GetRelevantAssetPlayerFromState(int32 MachineIndex, int32 StateIndex)
 {
 	FAnimNode_AssetPlayerBase* ResultPlayer = nullptr;
