Fix CIS

diff --git a/Engine/Source/Runtime/Engine/Classes/Animation/AnimInstance.h b/Engine/Source/Runtime/Engine/Classes/Animation/AnimInstance.h
--- a/Engine/Source/Runtime/Engine/Classes/Animation/AnimInstance.h
+++ b/Engine/Source/Runtime/Engine/Classes/Animation/AnimInstance.h
@@ -715,9 +715,6 @@ protected:
 
 protected:
 
-	/** Gets the runtime instance of the specified state machine */
-	FAnimNode_StateMachine* GetStateMachineInstance(int32 MachineIndex);
-
 	/** Gets the runtime instance desc of the state machine specified by name */
 	const FBakedAnimationStateMachine* GetStateMachineInstanceDesc(FName MachineName);
 
