fixed potential crash when using Blendables from blueprints

diff --git a/Engine/Source/Runtime/Engine/Private/SceneView.cpp b/Engine/Source/Runtime/Engine/Private/SceneView.cpp
--- a/Engine/Source/Runtime/Engine/Private/SceneView.cpp
+++ b/Engine/Source/Runtime/Engine/Private/SceneView.cpp
@@ -1046,7 +1046,7 @@ void FSceneView::OverridePostProcessSettings(const FPostProcessSettings& Src, fl
 				continue;
 			}
 
-			float LocalWeight = Src.WeightedBlendables.Array[i].Weight * Weight;
+			float LocalWeight = FMath::Min(1.0f, Src.WeightedBlendables.Array[i].Weight) * Weight;
 
 			if(LocalWeight > 0.0f)
 			{
