jira: UE-12328

When the resolution is not full screen, resolution will always be smaller then the primary desktop resolution, even if the resolution is specified on the commandline(-ResX= and/or -ResY=).

In the case the primary desktop size is smaller then the minimum resolution (1280x720), the system will create a resolution of a scale of 5/6 of the desktop size in both direction, width and height. This will prevent having a window of the desktop size which conflict with the taskbar. Note that the UI layout will not support well small resolution like this.

diff --git a/Engine/Source/Runtime/Core/Private/Misc/CoreMisc.cpp b/Engine/Source/Runtime/Core/Private/Misc/CoreMisc.cpp
--- a/Engine/Source/Runtime/Core/Private/Misc/CoreMisc.cpp
+++ b/Engine/Source/Runtime/Core/Private/Misc/CoreMisc.cpp
@@ -960,10 +960,21 @@ void GenerateConvenientWindowedResolutions(const FDisplayMetrics& InDisplayMetri
 		}
 	}
 	
-	// if no convenient resolutions have been found, add a minimum one (if it fits)
-	if (OutResolutions.Num() == 0 && InDisplayMetrics.PrimaryDisplayHeight > MinHeight && InDisplayMetrics.PrimaryDisplayWidth > MinWidth)
+	// if no convenient resolutions have been found, add a minimum one
+	if (OutResolutions.Num() == 0)
 	{
-		OutResolutions.Add(FIntPoint(MinWidth, MinHeight));
+		if (InDisplayMetrics.PrimaryDisplayHeight > MinHeight && InDisplayMetrics.PrimaryDisplayWidth > MinWidth)
+		{
+			//Add the minimum size if it fit
+			OutResolutions.Add(FIntPoint(MinWidth, MinHeight));
+		}
+		else
+		{
+			//Force a resolution even if its smaller then the minimum height and width to avoid a bigger window then the desktop
+			float TargetWidth = FMath::RoundToFloat(InDisplayMetrics.PrimaryDisplayWidth) * Scales[NumScales - 1];
+			float TargetHeight = FMath::RoundToFloat(InDisplayMetrics.PrimaryDisplayHeight) * Scales[NumScales - 1];
+			OutResolutions.Add(FIntPoint(TargetWidth, TargetHeight));
+		}
 	}
 }
 
diff --git a/Engine/Source/Runtime/Engine/Private/GameEngine.cpp b/Engine/Source/Runtime/Engine/Private/GameEngine.cpp
--- a/Engine/Source/Runtime/Engine/Private/GameEngine.cpp
+++ b/Engine/Source/Runtime/Engine/Private/GameEngine.cpp
@@ -199,14 +199,18 @@ void UGameEngine::ConditionallyOverrideSettings(int32& ResolutionX, int32& Resol
 	FParse::Value(FCommandLine::Get(), TEXT("ResX="), ResolutionX);
 	FParse::Value(FCommandLine::Get(), TEXT("ResY="), ResolutionY);
 
-	if (!IsRunningDedicatedServer() && ((ResolutionX <= 0) || (ResolutionY <= 0)))
+	// consume available desktop area
+	FDisplayMetrics DisplayMetrics;
+	FDisplayMetrics::GetDisplayMetrics(DisplayMetrics);
+
+	int32 DesktopResolutionX = DisplayMetrics.PrimaryDisplayWidth;
+	int32 DesktopResolutionY = DisplayMetrics.PrimaryDisplayHeight;
+
+	//Dont allow a resolution bigger then the desktop found a convenient one
+	if (!IsRunningDedicatedServer() && ((ResolutionX <= 0 || ResolutionX >= DesktopResolutionX) || (ResolutionY <= 0 || ResolutionY >= DesktopResolutionY)))
 	{
-		// consume available desktop area
-		FDisplayMetrics DisplayMetrics;
-		FDisplayMetrics::GetDisplayMetrics( DisplayMetrics );
-		
-		ResolutionX = DisplayMetrics.PrimaryDisplayWidth;
-		ResolutionY = DisplayMetrics.PrimaryDisplayHeight;
+		ResolutionX = DesktopResolutionX;
+		ResolutionY = DesktopResolutionY;
 
 		// If we're in windowed mode, attempt to choose a suitable starting resolution that is smaller than the desktop, with a matching aspect ratio
 		if (WindowMode == EWindowMode::Windowed)
@@ -239,10 +243,6 @@ void UGameEngine::ConditionallyOverrideSettings(int32& ResolutionX, int32& Resol
 	// Check the platform to see if we should override the user settings.
 	if (FPlatformProperties::HasFixedResolution())
 	{
-		// Always use the device's actual resolution that has been setup earlier
-		FDisplayMetrics DisplayMetrics;
-		FDisplayMetrics::GetDisplayMetrics( DisplayMetrics );
-
 		// We need to pass the resolution back out to GameUserSettings, or it will just override it again
 		ResolutionX = DisplayMetrics.PrimaryDisplayWorkAreaRect.Right - DisplayMetrics.PrimaryDisplayWorkAreaRect.Left;
 		ResolutionY = DisplayMetrics.PrimaryDisplayWorkAreaRect.Bottom - DisplayMetrics.PrimaryDisplayWorkAreaRect.Top;
