#ue4 - Particles: added "bAutoManageAttachment" option to ParticleSystemComponent such that it can automatically attach to a component when activated and detach when deactivated. During registration this will also detach the component from any AttachParent previously assigned.

--------

diff --git a/Engine/Source/Runtime/Engine/Classes/Particles/ParticleSystemComponent.h b/Engine/Source/Runtime/Engine/Classes/Particles/ParticleSystemComponent.h
--- a/Engine/Source/Runtime/Engine/Classes/Particles/ParticleSystemComponent.h
+++ b/Engine/Source/Runtime/Engine/Classes/Particles/ParticleSystemComponent.h
@@ -296,6 +296,21 @@ public:
 	/** This flag will be set the first time the PSysComp is activated... used to prevent auto activated PSysComps from calling InitParticles twice on level load */
 	uint32 bHasBeenActivated:1;
 
+	/**
+	 * True if we should automatically attach to AutoAttachParent when activated, and detach from our parent when completed.
+	 * This overrides any current attachment that may be present at the time of activation (deferring initial attachment until activation, if AutoAttachParent is null).
+	 * When enabled, detachment occurs regardless of whether AutoAttachParent is assigned, and the relative transform from the time of activation is restored.
+	 * @see AutoAttachParent, AutoAttachSocketName, AutoAttachLocationType
+	 */
+	UPROPERTY(EditAnywhere, BlueprintReadWrite, Category=Attachment)
+	uint32 bAutoManageAttachment:1;
+
+private:
+	/** Did we auto attach during activation? Used to determine if we should restore the relative transform during detachment. */
+	uint32 bDidAutoAttach:1;
+
+public:
+
 	/**
 	 *	Array holding name instance parameters for this ParticleSystemComponent.
 	 *	Parameters can be used in Cascade using DistributionFloat/VectorParticleParameters.
@@ -446,6 +461,46 @@ public:
 	UPROPERTY(BlueprintAssignable)
 	FOnSystemFinished OnSystemFinished;
 
+public:
+	/**
+	 * Component we automatically attach to when activated, if bAutoManageAttachment is true.
+	 * If null during registration, we assign the existing AttachParent and defer attachment until we activate.
+	 * @see bAutoManageAttachment
+	 */
+	UPROPERTY(VisibleInstanceOnly, BlueprintReadWrite, Category=Attachment, meta=(EditCondition="bAutoManageAttachment"))
+	TWeakObjectPtr<USceneComponent> AutoAttachParent;
+
+	/**
+	 * Socket we automatically attach to on the AutoAttachParent, if bAutoManageAttachment is true.
+	 * @see bAutoManageAttachment
+	 */
+	UPROPERTY(EditAnywhere, BlueprintReadWrite, Category=Attachment, meta=(EditCondition="bAutoManageAttachment"))
+	FName AutoAttachSocketName;
+
+	/**
+	 * Options for how we handle our location when we attach to the AutoAttachParent, if bAutoManageAttachment is true.
+	 * @see bAutoManageAttachment, EAttachLocation::Type
+	 */
+	UPROPERTY(EditAnywhere, BlueprintReadWrite, Category=Attachment, meta=(EditCondition="bAutoManageAttachment"))
+	TEnumAsByte<EAttachLocation::Type> AutoAttachLocationType;
+
+	/**
+	 * Set AutoAttachParent, AutoAttachSocketName, AutoAttachLocationType to the specified parameters. Does not change bAutoManageAttachment; that must be set separately.
+	 * @param  Parent			Component to attach to. 
+	 * @param  SocketName		Socket on Parent to attach to.
+	 * @param  LocationType		Option for how we handle our location when we attach to Parent.
+	 * @see bAutoManageAttachment, AutoAttachParent, AutoAttachSocketName, AutoAttachLocationType
+	 */
+	UFUNCTION(BlueprintCallable, Category="Effects|Components|ParticleSystem")
+	void SetAutoAttachParams(USceneComponent* Parent, FName SocketName = NAME_None, EAttachLocation::Type LocationType = EAttachLocation::KeepRelativeOffset);
+
+private:
+
+	/** Saved relative transform before auto attachement. Used during detachment to restore the transform if we had automatically attached. */
+	FVector SavedAutoAttachRelativeLocation;
+	FRotator SavedAutoAttachRelativeRotation;
+	FVector SavedAutoAttachRelativeScale3D;
+
 private:
 	/** Cached copy of the transform for async work */
 	FTransform AsyncComponentToWorld;
@@ -1047,4 +1102,13 @@ private:
 };
 
 
+//////////////////////////////////////////////////////////////////////////
+// ParticleSystemComponent inlines
+
+FORCEINLINE_DEBUGGABLE void UParticleSystemComponent::SetAutoAttachParams(USceneComponent* Parent, FName SocketName, EAttachLocation::Type LocationType)
+{
+	AutoAttachParent = Parent;
+	AutoAttachSocketName = SocketName;
+	AutoAttachLocationType = LocationType;
+}
 
diff --git a/Engine/Source/Runtime/Engine/Private/Particles/ParticleComponents.cpp b/Engine/Source/Runtime/Engine/Private/Particles/ParticleComponents.cpp
--- a/Engine/Source/Runtime/Engine/Private/Particles/ParticleComponents.cpp
+++ b/Engine/Source/Runtime/Engine/Private/Particles/ParticleComponents.cpp
@@ -2671,6 +2671,8 @@ UParticleSystemComponent::UParticleSystemComponent(const FObjectInitializer& Obj
 
 	// Disable receiving decals by default.
 	bReceivesDecals = false;
+
+	SavedAutoAttachRelativeScale3D = FVector(1.f, 1.f, 1.f);
 }
 
 #if WITH_EDITOR
@@ -2815,6 +2817,39 @@ void UParticleSystemComponent::OnRegister()
 		FXSystem = World->Scene->GetFXSystem();
 	}
 
+	if (bAutoManageAttachment && !IsActive())
+	{
+		// Detach from current parent, we are supposed to wait for activation.
+		if (AttachParent)
+		{
+			// If no auto attach parent override, use the current parent when we activate
+			if (!AutoAttachParent.IsValid())
+			{
+				AutoAttachParent = AttachParent;
+			}
+			// If no auto attach socket override, use current socket when we activate
+			if (AutoAttachSocketName == NAME_None)
+			{
+				AutoAttachSocketName = AttachSocketName;
+			}
+
+			// Prevent attachment before Super::OnRegister() tries to attach us, since we only attach when activated.
+			if (AttachParent->AttachChildren.Contains(this))
+			{
+				DetachFromParent(/*bMaintainWorldPosition=*/ false, /*bCallModify=*/ false);
+			}
+			else
+			{
+				AttachParent = nullptr;
+				AttachSocketName = NAME_None;
+			}
+
+			SavedAutoAttachRelativeLocation = RelativeLocation;
+			SavedAutoAttachRelativeRotation = RelativeRotation;
+			SavedAutoAttachRelativeScale3D = RelativeScale3D;
+		}
+	}
+
 	Super::OnRegister();
 
 	// If we were active before but are not now, activate us
@@ -3891,6 +3926,21 @@ void UParticleSystemComponent::FinalizeTickComponent()
 		{
 			DestroyComponent();
 		}
+		else if (bAutoManageAttachment)
+		{
+			if (AttachParent)
+			{
+				if (bDidAutoAttach)
+				{
+					// Restore relative transform from before attachment. Actual transform will be updated as part of DetachFromParent().
+					RelativeLocation = SavedAutoAttachRelativeLocation;
+					RelativeRotation = SavedAutoAttachRelativeRotation;
+					RelativeScale3D = SavedAutoAttachRelativeScale3D;
+				}
+				DetachFromParent(/*bMaintainWorldPosition=*/ false);
+				bDidAutoAttach = false;
+			}
+		}
 	}
 	bWasCompleted = bIsCompleted;
 
@@ -4324,6 +4374,23 @@ void UParticleSystemComponent::ActivateSystem(bool bFlagAsJustAttached)
 
 	if( GIsAllowingParticles && bDetailModeAllowsRendering )
 	{
+		// Auto attach if requested
+		bDidAutoAttach = false;
+		if (bAutoManageAttachment)
+		{
+			USceneComponent* NewParent = AutoAttachParent.Get();
+			if (NewParent)
+			{
+				SavedAutoAttachRelativeLocation = RelativeLocation;
+				SavedAutoAttachRelativeRotation = RelativeRotation;
+				SavedAutoAttachRelativeScale3D = RelativeScale3D;
+
+				AttachTo(NewParent, AutoAttachSocketName, AutoAttachLocationType);
+				bDidAutoAttach = true;
+				bFlagAsJustAttached = true;
+			}
+		}
+
 		if (bFlagAsJustAttached)
 		{
 			bJustRegistered = true;
