FPS Chart html needs frame % below 120 FPS

Merging 2668141

//depot/UE4-UT/...

to //depot/UE4/...

diff --git a/Engine/Source/Runtime/Engine/Private/ChartCreation.cpp b/Engine/Source/Runtime/Engine/Private/ChartCreation.cpp
--- a/Engine/Source/Runtime/Engine/Private/ChartCreation.cpp
+++ b/Engine/Source/Runtime/Engine/Private/ChartCreation.cpp
@@ -456,6 +456,8 @@ void UEngine::DumpFPSChartToLog( float TotalTime, float DeltaTime, int32 NumFram
 	float PctTimeAbove30 = 0;	// Keep track of percentage of time at 30+ FPS.
 	int32 NumFramesBelow60 = 0; // keep track of the number of frames below 60 FPS
 	float PctTimeAbove60 = 0;	// Keep track of percentage of time at 60+ FPS.
+	int32 NumFramesBelow120 = 0; // keep track of the number of frames below 120 FPS
+	float PctTimeAbove120 = 0;	// Keep track of percentage of time at 120+ FPS.
 
 	UE_LOG(LogChartCreation, Log, TEXT("--- Begin : FPS chart dump for level '%s'"), *InMapName );
 
@@ -491,29 +493,42 @@ void UEngine::DumpFPSChartToLog( float TotalTime, float DeltaTime, int32 NumFram
 			NumFramesBelow60 += GFPSChart[BucketIndex].Count;
 		}
 
+		// Keep track of time spent at 120+ FPS.
+		if (StartFPS >= 120)
+		{
+			PctTimeAbove120 += BucketTimePercentage;
+		}
+		else
+		{
+			NumFramesBelow120 += GFPSChart[BucketIndex].Count;
+		}
+
 		// Log bucket index, time and frame Percentage.
 		UE_LOG(LogChartCreation, Log, TEXT("Bucket: %2i - %2i  Time: %5.2f  Frame: %5.2f"), StartFPS, EndFPS, BucketTimePercentage, BucketFramePercentage );
 	}
 
-	UE_LOG(LogChartCreation, Log, TEXT("%i frames collected over %4.2f seconds, disregarding %4.2f seconds for a %4.2f FPS average, %4.2f percent of time spent > 30 FPS, %4.2f percent of time spent > 60 FPS"),
+	UE_LOG(LogChartCreation, Log, TEXT("%i frames collected over %4.2f seconds, disregarding %4.2f seconds for a %4.2f FPS average, %4.2f percent of time spent > 30 FPS, %4.2f percent of time spent > 60 FPS, %4.2f percent of time spent > 120 FPS"),
 		NumFrames, 
 		DeltaTime, 
 		FMath::Max<float>( 0, DeltaTime - TotalTime ), 
 		NumFrames / TotalTime,
 		PctTimeAbove30,
-		PctTimeAbove60);
+		PctTimeAbove60,
+		PctTimeAbove120);
 	UE_LOG(LogChartCreation, Log, TEXT("Average GPU frametime: %4.2f ms"), float((GTotalGPUTime / NumFrames)*1000.0));
-	UE_LOG(LogChartCreation, Log, TEXT("BoundGameThreadPct: %4.2f  BoundRenderThreadPct: %4.2f  BoundGPUPct: %4.2f PercentFrames30+: %f   PercentFrames60+: %f   BoundGameTime: %f  BoundRenderTime: %f  BoundGPUTime: %f  PctTimeAbove30: %f  PctTimeAbove60: %f ")
+	UE_LOG(LogChartCreation, Log, TEXT("BoundGameThreadPct: %4.2f  BoundRenderThreadPct: %4.2f  BoundGPUPct: %4.2f  PercentFrames30+: %f  PercentFrames60+: %f  PercentFrames120+: %f  BoundGameTime: %f  BoundRenderTime: %f  BoundGPUTime: %f  PctTimeAbove30: %f  PctTimeAbove60: %f  PctTimeAbove120: %f ")
 		, (float(GNumFramesBound_GameThread)/float(NumFrames))*100.0f
 		, (float(GNumFramesBound_RenderThread)/float(NumFrames))*100.0f
 		, (float(GNumFramesBound_GPU)/float(NumFrames))*100.0f
 		, float(NumFrames - NumFramesBelow30) / float(NumFrames)*100.0f
 		, float(NumFrames - NumFramesBelow60) / float(NumFrames)*100.0f
+		, float(NumFrames - NumFramesBelow120) / float(NumFrames)*100.0f
 		, (GTotalFramesBoundTime_GameThread / DeltaTime)*100.0f
 		, ((GTotalFramesBoundTime_RenderThread)/DeltaTime)*100.0f
 		, ((GTotalFramesBoundTime_GPU)/DeltaTime)*100.0f
 		, PctTimeAbove30
 		, PctTimeAbove60
+		, PctTimeAbove120
 		);
 
 	UE_LOG(LogChartCreation, Log, TEXT("--- End"));
@@ -723,6 +738,8 @@ void UEngine::DumpFPSChartToStatsLog( float TotalTime, float DeltaTime, int32 Nu
 		float PctTimeAbove30 = 0;	// Keep track of percentage of time at 30+ FPS.
 		int32 NumFramesBelow60 = 0; // keep track of the number of frames below 60 FPS
 		float PctTimeAbove60 = 0;	// Keep track of percentage of time at 60+ FPS.
+		int32 NumFramesBelow120 = 0; // keep track of the number of frames below 120 FPS
+		float PctTimeAbove120 = 0;	// Keep track of percentage of time at 120+ FPS.
 
 		// Iterate over all buckets, dumping percentages.
 		for( int32 BucketIndex=0; BucketIndex<ARRAY_COUNT(GFPSChart); BucketIndex++ )
@@ -755,28 +772,40 @@ void UEngine::DumpFPSChartToStatsLog( float TotalTime, float DeltaTime, int32 Nu
 				NumFramesBelow60 += GFPSChart[BucketIndex].Count;
 			}
 
+			// Keep track of time spent at 120+ FPS.
+			if (StartFPS >= 120)
+			{
+				PctTimeAbove120 += BucketTimePercentage;
+			}
+			else
+			{
+				NumFramesBelow120 += GFPSChart[BucketIndex].Count;
+			}
+
 			// Log bucket index, time and frame Percentage.
 			OutputFile->Logf(TEXT("Bucket: %2i - %2i  Time: %5.2f  Frame: %5.2f"), StartFPS, EndFPS, BucketTimePercentage, BucketFramePercentage);
 		}
 
-		OutputFile->Logf(TEXT("%i frames collected over %4.2f seconds, disregarding %4.2f seconds for a %4.2f FPS average, %4.2f percent of time spent > 30 FPS, %4.2f percent of time spent > 60 FPS"), 
+		OutputFile->Logf(TEXT("%i frames collected over %4.2f seconds, disregarding %4.2f seconds for a %4.2f FPS average, %4.2f percent of time spent > 30 FPS, %4.2f percent of time spent > 60 FPS, %4.2f percent of time spent > 120 FPS"), 
 			NumFrames, 
 			DeltaTime, 
 			FMath::Max<float>( 0, DeltaTime - TotalTime ), 
 			NumFrames / TotalTime,
-			PctTimeAbove30, PctTimeAbove60 );
+			PctTimeAbove30, PctTimeAbove60, PctTimeAbove120 );
 		OutputFile->Logf(TEXT("Average GPU frame time: %4.2f ms"), float((GTotalGPUTime / NumFrames)*1000.0));
-		OutputFile->Logf(TEXT("BoundGameThreadPct: %4.2f  BoundRenderThreadPct: %4.2f  BoundGPUPct: %4.2f PercentFrames30+: %f   PercentFrames60+: %f   BoundGameTime: %f  BoundRenderTime: %f  BoundGPUTime: %f  PctTimeAbove30: %f  PctTimeAbove60: %f ")
+		OutputFile->Logf(TEXT("BoundGameThreadPct: %4.2f  BoundRenderThreadPct: %4.2f  BoundGPUPct: %4.2f  PercentFrames30+: %f  PercentFrames60+: %f  PercentFrames120+: %f  BoundGameTime: %f  BoundRenderTime: %f  BoundGPUTime: %f  PctTimeAbove30: %f  PctTimeAbove60: %f  PctTimeAbove120: %f ")
 			, (float(GNumFramesBound_GameThread)/float(NumFrames))*100.0f
 			, (float(GNumFramesBound_RenderThread)/float(NumFrames))*100.0f
 			, (float(GNumFramesBound_GPU)/float(NumFrames))*100.0f
 			, float(NumFrames - NumFramesBelow30) / float(NumFrames)*100.0f
 			, float(NumFrames - NumFramesBelow60) / float(NumFrames)*100.0f
+			, float(NumFrames - NumFramesBelow120) / float(NumFrames)*100.0f
 			, (GTotalFramesBoundTime_GameThread / DeltaTime)*100.0f
 			, ((GTotalFramesBoundTime_RenderThread)/DeltaTime)*100.0f
 			, ((GTotalFramesBoundTime_GPU)/DeltaTime)*100.0f
 			, PctTimeAbove30
 			, PctTimeAbove60
+			, PctTimeAbove120
 			);
 
 		// Dump hitch data
@@ -861,6 +890,8 @@ void UEngine::DumpFPSChartToHTML( float TotalTime, float DeltaTime, int32 NumFra
 		float PctTimeAbove30 = 0;
 		// Keep track of percentage of time at 60+ FPS.
 		float PctTimeAbove60 = 0;
+		// Keep track of percentage of time at 120+ FPS.
+		float PctTimeAbove120 = 0;
 
 		// Iterate over all buckets, updating row 
 		for( int32 BucketIndex=0; BucketIndex<ARRAY_COUNT(GFPSChart); BucketIndex++ )
@@ -886,6 +917,12 @@ void UEngine::DumpFPSChartToHTML( float TotalTime, float DeltaTime, int32 NumFra
 				PctTimeAbove60 += BucketTimePercentage;
 			}
 
+			// Keep track of time spent at 120+ FPS.
+			if (StartFPS >= 120)
+			{
+				PctTimeAbove120 += BucketTimePercentage;
+			}
+
 			const FString SrcToken = FString::Printf(TEXT("TOKEN_%i_%i"), StartFPS, EndFPS);
 			const FString DstToken = FString::Printf( TEXT("%5.2f"), BucketTimePercentage );
 
@@ -966,9 +1003,10 @@ void UEngine::DumpFPSChartToHTML( float TotalTime, float DeltaTime, int32 NumFra
 		FPSChartRow = FPSChartRow.Replace( TEXT("TOKEN_SETTINGS_FX"),	    *FString::Printf(TEXT("%d"), Quality.EffectsQuality ), ESearchCase::CaseSensitive );
 
 		FPSChartRow = FPSChartRow.Replace( TEXT("TOKEN_AVG_FPS"),			*FString::Printf(TEXT("%4.2f"), NumFrames / TotalTime), ESearchCase::CaseSensitive );
-		FPSChartRow = FPSChartRow.Replace(TEXT("TOKEN_PCT_ABOVE_30"), *FString::Printf(TEXT("%4.2f"), PctTimeAbove30), ESearchCase::CaseSensitive);
-		FPSChartRow = FPSChartRow.Replace(TEXT("TOKEN_PCT_ABOVE_60"), *FString::Printf(TEXT("%4.2f"), PctTimeAbove60), ESearchCase::CaseSensitive);
-		FPSChartRow = FPSChartRow.Replace(TEXT("TOKEN_TIME_DISREGARDED"), *FString::Printf(TEXT("%4.2f"), FMath::Max<float>(0, DeltaTime - TotalTime)), ESearchCase::CaseSensitive);
+		FPSChartRow = FPSChartRow.Replace( TEXT("TOKEN_PCT_ABOVE_30"),		*FString::Printf(TEXT("%4.2f"), PctTimeAbove30), ESearchCase::CaseSensitive);
+		FPSChartRow = FPSChartRow.Replace( TEXT("TOKEN_PCT_ABOVE_60"),		*FString::Printf(TEXT("%4.2f"), PctTimeAbove60), ESearchCase::CaseSensitive);
+		FPSChartRow = FPSChartRow.Replace( TEXT("TOKEN_PCT_ABOVE_120"),		*FString::Printf(TEXT("%4.2f"), PctTimeAbove120), ESearchCase::CaseSensitive);
+		FPSChartRow = FPSChartRow.Replace( TEXT("TOKEN_TIME_DISREGARDED"),	*FString::Printf(TEXT("%4.2f"), FMath::Max<float>(0, DeltaTime - TotalTime)), ESearchCase::CaseSensitive);
 		FPSChartRow = FPSChartRow.Replace( TEXT("TOKEN_TIME"),				*FString::Printf(TEXT("%4.2f"), DeltaTime), ESearchCase::CaseSensitive );
 		FPSChartRow = FPSChartRow.Replace( TEXT("TOKEN_FRAMECOUNT"),		*FString::Printf(TEXT("%i"), NumFrames), ESearchCase::CaseSensitive );
 		FPSChartRow = FPSChartRow.Replace( TEXT("TOKEN_AVG_GPUTIME"),		*FString::Printf(TEXT("%4.2f ms"), float((GTotalGPUTime / NumFrames)*1000.0) ), ESearchCase::CaseSensitive );
