Fix for global pause of audio when game mode is already paused
#ue4
#android

diff --git a/Engine/Source/Runtime/Engine/Private/AudioDevice.cpp b/Engine/Source/Runtime/Engine/Private/AudioDevice.cpp
--- a/Engine/Source/Runtime/Engine/Private/AudioDevice.cpp
+++ b/Engine/Source/Runtime/Engine/Private/AudioDevice.cpp
@@ -1720,24 +1720,24 @@ void FAudioDevice::DestroyEffect( FSoundSource* Source )
 void FAudioDevice::HandlePause( bool bGameTicking, bool bGlobalPause )
 {
 	// Pause all sounds if transitioning to pause mode.
-	if( !bGameTicking && bGameWasTicking )
+	if( !bGameTicking && (bGameWasTicking || bGlobalPause) )
 	{
 		for( int32 i = 0; i < Sources.Num(); i++ )
 		{
 			FSoundSource* Source = Sources[ i ];
-			if( bGlobalPause || Source->IsGameOnly() )
+			if( !Source->IsPaused() && (bGlobalPause || Source->IsGameOnly()) )
 			{
 				Source->Pause();
 			}
 		}
 	}
 	// Unpause all sounds if transitioning back to game.
-	else if( bGameTicking && !bGameWasTicking )
+	else if( bGameTicking && (!bGameWasTicking || bGlobalPause) )
 	{
 		for( int32 i = 0; i < Sources.Num(); i++ )
 		{
 			FSoundSource* Source = Sources[ i ];
-			if( bGlobalPause || Source->IsGameOnly() )
+			if( Source->IsPaused() && (bGlobalPause || Source->IsGameOnly()) )
 			{
 				Source->Play();
 			}
