#ANIM: Fixed issue with crash when editing tangent of curves

https://jira.ol.epicgames.net/browse/UE-20399

Merging

//depot/UE4-Orion/Engine/Source/Runtime/Engine/Private/Animation/AnimCurveTypes.cpp

to //depot/UE4/Engine/Source/Runtime/Engine/Private/Animation/AnimCurveTypes.cpp

diff --git a/Engine/Source/Runtime/Engine/Private/Animation/AnimCurveTypes.cpp b/Engine/Source/Runtime/Engine/Private/Animation/AnimCurveTypes.cpp
--- a/Engine/Source/Runtime/Engine/Private/Animation/AnimCurveTypes.cpp
+++ b/Engine/Source/Runtime/Engine/Private/Animation/AnimCurveTypes.cpp
@@ -290,8 +290,10 @@ void FRawCurveTracks::Serialize(FArchive& Ar)
 		}
 	}
 #endif // WITH_EDITORONLY_DATA
-
-	SortFloatCurvesByUID();
+	if (Ar.IsLoading())
+	{
+		SortFloatCurvesByUID();
+	}
 }
 
 void FRawCurveTracks::UpdateLastObservedNames(FSmartNameMapping* NameMapping, ESupportedCurveType SupportedCurveType /*= FloatType*/)
