Fix warning about treating an integer as a boolean

diff --git a/Engine/Source/Runtime/Engine/Public/GeneratedCodeHelpers.h b/Engine/Source/Runtime/Engine/Public/GeneratedCodeHelpers.h
--- a/Engine/Source/Runtime/Engine/Public/GeneratedCodeHelpers.h
+++ b/Engine/Source/Runtime/Engine/Public/GeneratedCodeHelpers.h
@@ -120,7 +120,7 @@ public:
 	template<typename T>
 	static bool Array_RemoveItem(TArray<T>& TargetArray, const T& Item)
 	{
-		return TargetArray.Remove(Item);
+		return TargetArray.Remove(Item) != 0;
 	}
 
 	template<typename T>
