#ue4 - Removed the need for APawn to implement Tick().

- Moved RemoteViewPitch update to Pawn::PreReplication() from Tick().
- Also inlined APawn::GetController().

--------

diff --git a/Engine/Source/Runtime/Engine/Classes/GameFramework/Pawn.h b/Engine/Source/Runtime/Engine/Classes/GameFramework/Pawn.h
--- a/Engine/Source/Runtime/Engine/Classes/GameFramework/Pawn.h
+++ b/Engine/Source/Runtime/Engine/Classes/GameFramework/Pawn.h
@@ -38,6 +38,7 @@ public:
 	APawn(const FObjectInitializer& ObjectInitializer = FObjectInitializer::Get());
 
 	virtual void GetLifetimeReplicatedProps(TArray<FLifetimeProperty>& OutLifetimeProps) const override;
+	virtual void PreReplication( IRepChangedPropertyTracker & ChangedPropertyTracker ) override;
 
 	/** Return our PawnMovementComponent, if we have one. By default, returns the first PawnMovementComponent found. Native classes that create their own movement component should override this method for more efficiency. */
 	UFUNCTION(BlueprintCallable, meta=(Tooltip="Return our PawnMovementComponent, if we have one."), Category="Pawn")
@@ -194,7 +195,6 @@ public:
 
 	// Begin AActor Interface.
 	virtual FVector GetVelocity() const override;
-	virtual void Tick(float DeltaSeconds) override;
 	virtual void Reset() override;
 	virtual FString GetHumanReadableName() const override;
 	virtual bool ShouldTickIfViewportsOnly() const override;
@@ -463,3 +463,11 @@ public:
 	FVector K2_GetMovementInputVector() const;	
 };
 
+
+//////////////////////////////////////////////////////////////////////////
+// Pawn inlines
+
+FORCEINLINE AController* APawn::GetController() const
+{
+	return Controller;
+}
diff --git a/Engine/Source/Runtime/Engine/Private/Pawn.cpp b/Engine/Source/Runtime/Engine/Private/Pawn.cpp
--- a/Engine/Source/Runtime/Engine/Private/Pawn.cpp
+++ b/Engine/Source/Runtime/Engine/Private/Pawn.cpp
@@ -428,12 +428,6 @@ bool APawn::IsControlled() const
 	return(PC != NULL);
 }
 
-/** For K2 access to the controller. */
-AController* APawn::GetController() const
-{
-	return Controller;
-}
-
 FRotator APawn::GetControlRotation() const
 {
 	return Controller ? Controller->GetControlRotation() : FRotator::ZeroRotator;
@@ -700,16 +694,6 @@ void APawn::SetPlayerDefaults()
 {
 }
 
-void APawn::Tick( float DeltaSeconds )
-{
-	Super::Tick(DeltaSeconds);
-
-	if (Role == ROLE_Authority && GetController())
-	{
-		SetRemoteViewPitch(GetController()->GetControlRotation().Pitch);
-	}
-}
-
 void APawn::RecalculateBaseEyeHeight()
 {
 	BaseEyeHeight = GetDefault<APawn>(GetClass())->BaseEyeHeight;
@@ -1059,6 +1043,16 @@ void APawn::GetLifetimeReplicatedProps( TArray< FLifetimeProperty > & OutLifetim
 	DOREPLIFETIME_CONDITION( APawn, RemoteViewPitch, 	COND_SkipOwner );
 }
 
+void APawn::PreReplication( IRepChangedPropertyTracker & ChangedPropertyTracker )
+{
+	Super::PreReplication( ChangedPropertyTracker );
+
+	if (Role == ROLE_Authority && GetController())
+	{
+		SetRemoteViewPitch(GetController()->GetControlRotation().Pitch);
+	}
+}
+
 void APawn::MoveIgnoreActorAdd(AActor* ActorToIgnore)
 {
 	UPrimitiveComponent* RootPrimitiveComponent = Cast<UPrimitiveComponent>(GetRootComponent());
