Physics: Avoid unnecessary overhead in GatherApexStats when STATS=0
[INTEGRATE] Change 2668321 by [REDACTED_PERSON][REDACTED_HANDLE] on 2015/08/25 16:35:33

diff --git a/Engine/Source/Runtime/Engine/Private/PhysicsEngine/PhysScene.cpp b/Engine/Source/Runtime/Engine/Private/PhysicsEngine/PhysScene.cpp
--- a/Engine/Source/Runtime/Engine/Private/PhysicsEngine/PhysScene.cpp
+++ b/Engine/Source/Runtime/Engine/Private/PhysicsEngine/PhysScene.cpp
@@ -581,6 +581,7 @@ void FinishSceneStat(uint32 Scene)
 #if WITH_APEX
 void GatherApexStats(const UWorld* World, NxApexScene* ApexScene)
 {
+#if STATS
 	SCOPED_APEX_SCENE_READ_LOCK(ApexScene);
 	int32 NumVerts = 0;
 	int32 NumCloths = 0;
@@ -598,6 +599,7 @@ void GatherApexStats(const UWorld* World, NxApexScene* ApexScene)
 	}
 	SET_DWORD_STAT(STAT_NumCloths, NumCloths);        // number of recently simulated apex cloths.
 	SET_DWORD_STAT(STAT_NumClothVerts, NumVerts);	  // number of recently simulated apex vertices.
+#endif
 }
 #endif
 
