Animation: Potentially avoid a reallocation during BlendInPhysics
[INTEGRATE] Change 2668347 by [REDACTED_PERSON][REDACTED_HANDLE] on 2015/08/25 16:45:59

diff --git a/Engine/Source/Runtime/Engine/Private/PhysicsEngine/PhysAnim.cpp b/Engine/Source/Runtime/Engine/Private/PhysicsEngine/PhysAnim.cpp
--- a/Engine/Source/Runtime/Engine/Private/PhysicsEngine/PhysAnim.cpp
+++ b/Engine/Source/Runtime/Engine/Private/PhysicsEngine/PhysAnim.cpp
@@ -397,7 +397,8 @@ void USkeletalMeshComponent::BlendInPhysics(FTickFunction& ThisTickFunction)
 				AnimEvaluationContext.VertexAnims = ActiveVertexAnims;
 			}
 
-			AnimEvaluationContext.LocalAtoms = LocalAtoms;
+			AnimEvaluationContext.LocalAtoms.Reset(LocalAtoms.Num());
+			AnimEvaluationContext.LocalAtoms.Append(LocalAtoms);
 
 			ParallelAnimationEvaluationTask = TGraphTask<FParallelBlendPhysicsTask>::CreateTask().ConstructAndDispatchWhenReady(this);
 
@@ -409,7 +410,8 @@ void USkeletalMeshComponent::BlendInPhysics(FTickFunction& ThisTickFunction)
 			ParallelBlendPhysicsCompletionTask = TGraphTask<FParallelBlendPhysicsCompletionTask>::CreateTask(&Prerequistes).ConstructAndDispatchWhenReady(this);
 
 			ThisTickFunction.GetCompletionHandle()->DontCompleteUntil(ParallelBlendPhysicsCompletionTask);
-		}else
+		}
+		else
 		{
 			PerformBlendPhysicsBones(RequiredBones, LocalAtoms);
 			PostBlendPhysics();
