#ANIM: Fix the montage curve not getting through if negative

Merging

//depot/UE4-Orion/Engine/Source/Runtime/Engine/Private/Animation/AnimCurveTypes.cpp

to //depot/UE4/Engine/Source/Runtime/Engine/Private/Animation/AnimCurveTypes.cpp

diff --git a/Engine/Source/Runtime/Engine/Private/Animation/AnimCurveTypes.cpp b/Engine/Source/Runtime/Engine/Private/Animation/AnimCurveTypes.cpp
--- a/Engine/Source/Runtime/Engine/Private/Animation/AnimCurveTypes.cpp
+++ b/Engine/Source/Runtime/Engine/Private/Animation/AnimCurveTypes.cpp
@@ -601,7 +601,15 @@ void FBlendedCurve::Combine(const FBlendedCurve& CurveToCombine)
 
 	for(int32 CurveId=0; CurveId<CurveToCombine.Elements.Num(); ++CurveId)
 	{
-		Elements[CurveId].Value = FMath::Max(Elements[CurveId].Value, CurveToCombine.Elements[CurveId].Value);
+		// if target value is non zero, we accpet target's value
+		// originally this code was doing max, but that doesn't make sense since the values can be negative
+		// we could try to pick non-zero, but if target value is non-zero, I think we should accept that value 
+		// if source is non zero, it will be overriden
+		if (CurveToCombine.Elements[CurveId].Value != 0.f)
+		{
+			Elements[CurveId].Value = CurveToCombine.Elements[CurveId].Value; 
+		}
+
 		Elements[CurveId].Flags |= CurveToCombine.Elements[CurveId].Flags;
 	}
 
