Fixed crash when texture streaming is disabled on platforms that actually support it

diff --git a/Engine/Source/Runtime/Engine/Private/UnrealEngine.cpp b/Engine/Source/Runtime/Engine/Private/UnrealEngine.cpp
--- a/Engine/Source/Runtime/Engine/Private/UnrealEngine.cpp
+++ b/Engine/Source/Runtime/Engine/Private/UnrealEngine.cpp
@@ -7567,7 +7567,7 @@ void DrawStatsHUD( UWorld* World, FViewport* Viewport, FCanvas* Canvas, UCanvas*
 			}
 			
 #if !(UE_BUILD_SHIPPING || UE_BUILD_TEST)
-			if (FPlatformProperties::SupportsTextureStreaming())
+			if (FPlatformProperties::SupportsTextureStreaming() && IStreamingManager::Get().IsTextureStreamingEnabled())
 			{
 				auto MemOver = IStreamingManager::Get().GetTextureStreamingManager().GetMemoryOverBudget();
 				if (MemOver > 0)
