Fix for UE-20172

diff --git a/Engine/Source/Runtime/Engine/Private/BatchedElements.cpp b/Engine/Source/Runtime/Engine/Private/BatchedElements.cpp
--- a/Engine/Source/Runtime/Engine/Private/BatchedElements.cpp
+++ b/Engine/Source/Runtime/Engine/Private/BatchedElements.cpp
@@ -473,6 +473,9 @@ void FBatchedElements::PrepareShaders(
 	bool bEncodedHDR = Is32BppHDREncoded(View);
 
 	float GammaToUse = Gamma;
+
+	ESimpleElementBlendMode MaskedBlendMode = SE_BLEND_Opaque;
+
 	if(BlendMode >= SE_BLEND_RGBA_MASK_START && BlendMode <= SE_BLEND_RGBA_MASK_END)
 	{
 		/*
@@ -496,17 +499,19 @@ void FBatchedElements::PrepareShaders(
 		bool bAlphaOnly = bAlphaChannel && !bRedChannel && !bGreenChannel && !bBlueChannel;
 		uint32 NumChannelsOn = ( bRedChannel ? 1 : 0 ) + ( bGreenChannel ? 1 : 0 ) + ( bBlueChannel ? 1 : 0 );
 		GammaToUse = bAlphaOnly? 1.0f: Gamma;
-
+		
 		// If we are only to draw the alpha channel, make the Blend state opaque, to allow easy identification of the alpha values
 		if( bAlphaOnly )
 		{
-			SetBlendState(RHICmdList, SE_BLEND_Opaque, bEncodedHDR);
-
+			MaskedBlendMode = SE_BLEND_Opaque;
+			SetBlendState(RHICmdList, MaskedBlendMode, bEncodedHDR);
+			
 			R.W = G.W = B.W = 1.0f;
 		}
 		else
 		{
-			SetBlendState(RHICmdList, !bAlphaChannel ? SE_BLEND_Opaque : SE_BLEND_Translucent, bEncodedHDR); // If alpha channel is disabled, do not allow alpha blending
+			MaskedBlendMode = !bAlphaChannel ? SE_BLEND_Opaque : SE_BLEND_Translucent;  // If alpha channel is disabled, do not allow alpha blending
+			SetBlendState(RHICmdList, MaskedBlendMode, bEncodedHDR);
 
 			// Determine the red, green, blue and alpha components of their respective weights to enable that colours prominence
 			R.X = bRedChannel ? 1.0f : 0.0f;
@@ -657,7 +662,7 @@ void FBatchedElements::PrepareShaders(
 			else if(BlendMode >= SE_BLEND_RGBA_MASK_START && BlendMode <= SE_BLEND_RGBA_MASK_END)
 			{
 				TShaderMapRef<FSimpleElementColorChannelMaskPS> ColorChannelMaskPixelShader(GetGlobalShaderMap(FeatureLevel));
-				SetGlobalBoundShaderState(RHICmdList, FeatureLevel, ColorChannelMaskShaderState.GetBSS(bEncodedHDR, BlendMode), GSimpleElementVertexDeclaration.VertexDeclarationRHI,
+				SetGlobalBoundShaderState(RHICmdList, FeatureLevel, ColorChannelMaskShaderState.GetBSS(bEncodedHDR, MaskedBlendMode), GSimpleElementVertexDeclaration.VertexDeclarationRHI,
 					*VertexShader, *ColorChannelMaskPixelShader );
 			
 				ColorChannelMaskPixelShader->SetParameters(RHICmdList, Texture, ColorWeights, GammaToUse );
diff --git a/Engine/Source/Runtime/Engine/Public/BatchedElements.h b/Engine/Source/Runtime/Engine/Public/BatchedElements.h
--- a/Engine/Source/Runtime/Engine/Public/BatchedElements.h
+++ b/Engine/Source/Runtime/Engine/Public/BatchedElements.h
@@ -304,8 +304,15 @@ private:
 	public:
 		FGlobalBoundShaderState& GetBSS(bool bEncoded, ESimpleElementBlendMode BlendMode)
 		{
-			check((uint32)BlendMode < NumBSS);
-			return bEncoded ? EncodedBSS[BlendMode] : UnencodedBSS;
+			if (bEncoded)
+			{
+				check((uint32)BlendMode < NumBSS);
+				return EncodedBSS[BlendMode];
+			}
+			else
+			{
+				return UnencodedBSS;
+			}
 		}
 	};
 
