UE-16534 - Adjusting transform of a mesh with LODs causes LODs to flicker

diff --git a/Engine/Source/Runtime/Engine/Classes/Components/HierarchicalInstancedStaticMeshComponent.h b/Engine/Source/Runtime/Engine/Classes/Components/HierarchicalInstancedStaticMeshComponent.h
--- a/Engine/Source/Runtime/Engine/Classes/Components/HierarchicalInstancedStaticMeshComponent.h
+++ b/Engine/Source/Runtime/Engine/Classes/Components/HierarchicalInstancedStaticMeshComponent.h
@@ -96,10 +96,18 @@ class ENGINE_API UHierarchicalInstancedStaticMeshComponent : public UInstancedSt
 	UPROPERTY()
 	int32 NumBuiltInstances;
 
+	// Bounding box of any built instances (cached from the ClusterTree)
+	UPROPERTY()
+	FBox BuiltInstanceBounds;
+
 	// Bounding box of any unbuilt instances
 	UPROPERTY()
 	FBox UnbuiltInstanceBounds;
 
+	// Bounds of each individual unbuilt instance, used for LOD calculation
+	UPROPERTY()
+	TArray<FBox> UnbuiltInstanceBoundsList;
+
 	// The number of nodes in the occlusion layer
 	UPROPERTY()
 	int32 OcclusionLayerNumNodes;
@@ -170,6 +178,9 @@ protected:
 	/** Removes a single instance without extra work such as rebuilding the tree or marking render state dirty. */
 	void RemoveInstanceInternal(int32 InstanceIndex);
 
+	void UpdateInstanceTreeBoundsInternal(int32 InstanceIndex, const FBox& NewBounds);
+	static void UpdateInstanceTreeBoundsInternal_RenderThread(TArray<FClusterNode>& ClusterTree, int32 InstanceIndex, const FBox& NewBounds);
+
 	/** Gets and approximate number of verts for each LOD to generate heuristics **/
 	int32 GetVertsForLOD(int32 LODIndex);
 	/** Average number of instances per leaf **/
diff --git a/Engine/Source/Runtime/Engine/Private/HierarchicalInstancedStaticMesh.cpp b/Engine/Source/Runtime/Engine/Private/HierarchicalInstancedStaticMesh.cpp
--- a/Engine/Source/Runtime/Engine/Private/HierarchicalInstancedStaticMesh.cpp
+++ b/Engine/Source/Runtime/Engine/Private/HierarchicalInstancedStaticMesh.cpp
@@ -642,6 +642,7 @@ class FHierarchicalStaticMeshSceneProxy : public FInstancedStaticMeshSceneProxy
 	TSharedPtr<TArray<FClusterNode>, ESPMode::ThreadSafe> ClusterTreePtr;
 	const TArray<FClusterNode>& ClusterTree;
 
+	TArray<FBox> UnbuiltBounds;
 	int32 FirstUnbuiltIndex;
 	int32 LastUnbuiltIndex;
 
@@ -665,6 +666,7 @@ public:
 		: FInstancedStaticMeshSceneProxy(InComponent, InFeatureLevel)
 		, ClusterTreePtr(InComponent->ClusterTreePtr)
 		, ClusterTree(*InComponent->ClusterTreePtr)
+		, UnbuiltBounds(InComponent->UnbuiltInstanceBoundsList)
 		, FirstUnbuiltIndex(InComponent->NumBuiltInstances)
 		, LastUnbuiltIndex(InComponent->PerInstanceSMData.Num()+InComponent->RemovedInstances.Num()-1)
 		, bIsGrass(bInIsGrass)
@@ -680,6 +682,7 @@ public:
 		: FInstancedStaticMeshSceneProxy(InComponent, InFeatureLevel, Other)
 		, ClusterTreePtr(InComponent->ClusterTreePtr)
 		, ClusterTree(*InComponent->ClusterTreePtr)
+		, UnbuiltBounds(InComponent->UnbuiltInstanceBoundsList)
 		, FirstUnbuiltIndex(InComponent->NumBuiltInstances)
 		, LastUnbuiltIndex(InstancedRenderData.NumInstances+InComponent->RemovedInstances.Num()-1)
 		, bIsGrass(bInIsGrass)
@@ -945,6 +948,28 @@ static FORCEINLINE_DEBUGGABLE bool CullNode(const FFoliageCullInstanceParams& Pa
 	return false;
 }
 
+inline void CalcLOD(int32& InOutMinLOD, int32& InOutMaxLOD, const FVector& BoundMin, const FVector& BoundMax, const FVector& ViewOriginInLocalZero, const FVector& ViewOriginInLocalOne, const float LODPlanesMin[], const float LODPlanesMax[])
+{
+	if (InOutMinLOD != InOutMaxLOD)
+	{
+		const FVector Center = (BoundMax + BoundMin) * 0.5f;
+		const float DistCenterZero = FVector::Dist(Center, ViewOriginInLocalZero);
+		const float DistCenterOne = FVector::Dist(Center, ViewOriginInLocalOne);
+		const float HalfWidth = FVector::Dist(BoundMax, BoundMin) * 0.5f;
+		const float NearDot = FMath::Min(DistCenterZero, DistCenterOne) - HalfWidth;
+		const float FarDot = FMath::Max(DistCenterZero, DistCenterOne) + HalfWidth;
+
+		while (InOutMaxLOD > InOutMinLOD && NearDot > LODPlanesMax[InOutMinLOD])
+		{
+			InOutMinLOD++;
+		}
+		while (InOutMaxLOD > InOutMinLOD && FarDot < LODPlanesMin[InOutMaxLOD - 1])
+		{
+			InOutMaxLOD--;
+		}
+	}
+}
+
 template<bool TUseVector>
 void FHierarchicalStaticMeshSceneProxy::Traverse(const FFoliageCullInstanceParams& Params, int32 Index, int32 MinLOD, int32 MaxLOD, bool bFullyContained) const
 {
@@ -959,25 +984,12 @@ void FHierarchicalStaticMeshSceneProxy::Traverse(const FFoliageCullInstanceParam
 
 	if (MinLOD != MaxLOD)
 	{
-		FVector Center = (Node.BoundMax + Node.BoundMin) * 0.5f;
-		float DistCenterZero = FVector::Dist(Center, Params.ViewOriginInLocalZero);
-		float DistCenterOne = FVector::Dist(Center, Params.ViewOriginInLocalOne);
-		float HalfWidth = FVector::Dist(Node.BoundMax, Node.BoundMin) * 0.5f;
-		float NearDot = FMath::Min(DistCenterZero, DistCenterOne) - HalfWidth;
-		float FarDot = FMath::Max(DistCenterZero, DistCenterOne) + HalfWidth;
+		CalcLOD(MinLOD, MaxLOD, Node.BoundMin, Node.BoundMax, Params.ViewOriginInLocalZero, Params.ViewOriginInLocalOne, Params.LODPlanesMin, Params.LODPlanesMax);
 
-		while (MaxLOD > MinLOD &&  NearDot > Params.LODPlanesMax[MinLOD])
-		{
-			MinLOD++;
-		}
 		if (MinLOD >= Params.LODs)
 		{
 			return;
 		}
-		while (MaxLOD > MinLOD && FarDot < Params.LODPlanesMin[MaxLOD - 1])
-		{
-			MaxLOD--;
-		}
 	}
 	if (Index >= Params.FirstOcclusionNode && Index <= Params.LastOcclusionNode)
 	{
@@ -1453,14 +1465,92 @@ void FHierarchicalStaticMeshSceneProxy::GetDynamicMeshElements(const TArray<cons
 				// disable LOD blending for unbuilt instances as we haven't calculated the correct LOD.
 				ElementParams.bBlendLODs = false;
 
-				if (LastUnbuiltIndex - FirstUnbuiltIndex < 1000)
+				const int32 NumUnbuiltInstances = LastUnbuiltIndex - FirstUnbuiltIndex + 1;
+				if (NumUnbuiltInstances < 1000)
 				{
-					// less than 1000 instances, just add them all at LOD 0
-					InstanceParams.AddRun(0, 0, FirstUnbuiltIndex, LastUnbuiltIndex);
+					const int32 LODs = RenderData->LODResources.Num();
+
+					int32 Force = CVarForceLOD.GetValueOnRenderThread();
+					if (Force >= 0)
+					{
+						Force = FMath::Clamp(Force, 0, LODs - 1);
+						InstanceParams.AddRun(Force, Force, FirstUnbuiltIndex, LastUnbuiltIndex);
+					}
+					else
+					{
+						FMatrix WorldToLocal = GetLocalToWorld().Inverse();
+						FVector ViewOriginInLocalZero = WorldToLocal.TransformPosition(View->GetTemporalLODOrigin(0, bMultipleSections));
+						FVector ViewOriginInLocalOne  = WorldToLocal.TransformPosition(View->GetTemporalLODOrigin(1, bMultipleSections));
+						float LODPlanesMax[MAX_STATIC_MESH_LODS];
+						float LODPlanesMin[MAX_STATIC_MESH_LODS];
+
+						const bool bIsOrtho = !View->ViewMatrices.IsPerspectiveProjection();
+						const float MinSize = bIsOrtho ? 0.0f : CVarFoliageMinimumScreenSize.GetValueOnRenderThread();
+						const float LODScale = CVarFoliageLODDistanceScale.GetValueOnRenderThread();
+						const float LODRandom = CVarRandomLODRange.GetValueOnRenderThread();
+						const float MaxDrawDistanceScale = GetCachedScalabilityCVars().ViewDistanceScale;
+						const float SphereRadius = RenderData->Bounds.SphereRadius;
+
+						for (int32 LODIndex = 0; LODIndex < LODs; LODIndex++)
+						{
+							LODPlanesMax[LODIndex] = MIN_flt;
+							LODPlanesMin[LODIndex] = MAX_flt;
+						}
+						ElementParams.FinalCullDistance = MIN_flt;
+						for (int32 SampleIndex = 0; SampleIndex < 2; SampleIndex++)
+						{
+							// we aren't going to accurately figure out LOD planes for both samples and try all 4 possibilities, rather we just consider the worst case LOD range; the shaders will render correctly
+							float Fac = View->GetTemporalLODDistanceFactor(SampleIndex, bMultipleSections) * SphereRadius * SphereRadius * LODScale * LODScale;
+
+							float FinalCull = MAX_flt;
+
+							if (MinSize > 0.0)
+							{
+								FinalCull = FMath::Sqrt(Fac / MinSize);
+							}
+							if (UserData_AllInstances.EndCullDistance > 0.0f && UserData_AllInstances.EndCullDistance < FinalCull)
+							{
+								FinalCull = UserData_AllInstances.EndCullDistance;
+							}
+							FinalCull *= MaxDrawDistanceScale;
+							ElementParams.FinalCullDistance = FMath::Max(ElementParams.FinalCullDistance, FinalCull);
+
+							for (int32 LODIndex = 1; LODIndex < LODs; LODIndex++)
+							{
+								float Val = FMath::Min(FMath::Sqrt(Fac / RenderData->ScreenSize[LODIndex]), FinalCull);
+								LODPlanesMin[LODIndex - 1] = FMath::Min(LODPlanesMin[LODIndex - 1], Val - LODRandom);
+								LODPlanesMax[LODIndex - 1] = FMath::Max(LODPlanesMax[LODIndex - 1], Val);
+							}
+							LODPlanesMin[LODs - 1] = FMath::Min(LODPlanesMin[LODs - 1], FinalCull - LODRandom);
+							LODPlanesMax[LODs - 1] = FMath::Max(LODPlanesMax[LODs - 1], FinalCull);
+						}
+
+						// calculate runs
+						int32 MinLOD = 0;
+						int32 MaxLOD = LODs;
+						CalcLOD(MinLOD, MaxLOD, UnbuiltBounds[0].Min, UnbuiltBounds[0].Max, ViewOriginInLocalZero, ViewOriginInLocalOne, LODPlanesMin, LODPlanesMax);
+						int32 FirstIndexInRun = 0;
+						for (int32 Index = 1; Index < NumUnbuiltInstances; ++Index)
+						{
+							int32 TempMinLOD = 0;
+							int32 TempMaxLOD = LODs;
+							CalcLOD(TempMinLOD, TempMaxLOD, UnbuiltBounds[Index].Min, UnbuiltBounds[Index].Max, ViewOriginInLocalZero, ViewOriginInLocalOne, LODPlanesMin, LODPlanesMax);
+							if (TempMinLOD != MinLOD)
+							{
+								if (MinLOD < LODs)
+								{
+									InstanceParams.AddRun(MinLOD, MinLOD, FirstIndexInRun + FirstUnbuiltIndex, (Index - 1) + FirstUnbuiltIndex);
+								}
+								MinLOD = TempMinLOD;
+								FirstIndexInRun = Index;
+							}
+						}
+						InstanceParams.AddRun(MinLOD, MinLOD, FirstIndexInRun + FirstUnbuiltIndex, LastUnbuiltIndex);
+					}
 				}
 				else
 				{
-					// more than 1000, render them all at Max LOD
+					// more than 1000, render them all at lowest LOD (until we have an updated tree)
 					InstanceParams.AddRun(RenderData->LODResources.Num() - 1, RenderData->LODResources.Num() - 1, FirstUnbuiltIndex, LastUnbuiltIndex);
 				}
 				FillDynamicMeshElements(Collector, ElementParams, InstanceParams);
@@ -1503,23 +1593,9 @@ const TArray<FBoxSphereBounds>* FHierarchicalStaticMeshSceneProxy::GetOcclusionQ
 
 FBoxSphereBounds UHierarchicalInstancedStaticMeshComponent::CalcBounds(const FTransform& BoundTransform) const
 {
-	const TArray<FClusterNode>& ClusterTree = *ClusterTreePtr;
-	if (ClusterTree.Num() || UnbuiltInstanceBounds.IsValid)
+	if (BuiltInstanceBounds.IsValid || UnbuiltInstanceBounds.IsValid)
 	{
-		FBoxSphereBounds Result;
-		if (ClusterTree.Num())
-		{
-			Result = FBox(ClusterTree[0].BoundMin, ClusterTree[0].BoundMax);
-			// Include any unbuilt instances
-			if (UnbuiltInstanceBounds.IsValid)
-			{
-				Result = Result + UnbuiltInstanceBounds;
-			}
-		}
-		else
-		{
-			Result = UnbuiltInstanceBounds;
-		}
+		FBoxSphereBounds Result = BuiltInstanceBounds + UnbuiltInstanceBounds;
 
 		return Result.TransformBy(BoundTransform);
 	}
@@ -1587,6 +1663,10 @@ void UHierarchicalInstancedStaticMeshComponent::Serialize(FArchive& Ar)
 	}
 	TArray<FClusterNode>& ClusterTree = *ClusterTreePtr;
 	ClusterTree.BulkSerialize(Ar);
+	if (Ar.IsLoading() && !BuiltInstanceBounds.IsValid)
+	{
+		BuiltInstanceBounds = (ClusterTree.Num() > 0 ? FBox(ClusterTree[0].BoundMin, ClusterTree[0].BoundMax) : FBox(0));
+	}
 }
 
 SIZE_T UHierarchicalInstancedStaticMeshComponent::GetResourceSize( EResourceSizeMode::Type Mode )
@@ -1705,6 +1785,46 @@ bool UHierarchicalInstancedStaticMeshComponent::RemoveInstance(int32 InstanceInd
 	return true;
 }
 
+void UHierarchicalInstancedStaticMeshComponent::UpdateInstanceTreeBoundsInternal(int32 InstanceIndex, const FBox& NewBounds)
+{
+	ENQUEUE_UNIQUE_RENDER_COMMAND_THREEPARAMETER(
+		FUpdateInstanceTreeBoundsInternal,
+		TArray<FClusterNode>&, ClusterTree, *ClusterTreePtr,
+		const int32, InstanceIndex, InstanceIndex,
+		const FBox, NewBounds, NewBounds,
+		{
+			UpdateInstanceTreeBoundsInternal_RenderThread(ClusterTree, InstanceIndex, NewBounds);
+		});
+}
+
+void UHierarchicalInstancedStaticMeshComponent::UpdateInstanceTreeBoundsInternal_RenderThread(TArray<FClusterNode>& ClusterTree, int32 InstanceIndex, const FBox& NewBounds)
+{
+	checkSlow(IsInRenderingThread());
+
+	FClusterNode* Node = &ClusterTree[0];
+	checkSlow(InstanceIndex >= Node->FirstInstance && InstanceIndex <= Node->LastInstance);
+	while (1)
+	{
+		Node->BoundMin = Node->BoundMin.ComponentMin(NewBounds.Min);
+		Node->BoundMax = Node->BoundMax.ComponentMax(NewBounds.Max);
+
+		if (Node->FirstChild == INDEX_NONE)
+		{
+			break;
+		}
+
+		for (int32 i = Node->FirstChild; i <= Node->LastChild; ++i)
+		{
+			FClusterNode* ChildNode = &ClusterTree[i];
+			if (InstanceIndex >= ChildNode->FirstInstance && InstanceIndex <= ChildNode->LastInstance)
+			{
+				Node = ChildNode;
+				break;
+			}
+		}
+	}
+}
+
 bool UHierarchicalInstancedStaticMeshComponent::UpdateInstanceTransform(int32 InstanceIndex, const FTransform& NewInstanceTransform, bool bWorldSpace, bool bMarkRenderStateDirty)
 {
 	if (!PerInstanceSMData.IsValidIndex(InstanceIndex))
@@ -1718,20 +1838,50 @@ bool UHierarchicalInstancedStaticMeshComponent::UpdateInstanceTransform(int32 In
 		bConcurrentRemoval = true;
 	}
 
-	// Treat the old instance render data like a removal.
-	RemovedInstances.Add(InstanceReorderTable[InstanceIndex]);
+	const int32 OldReorderIndex = InstanceReorderTable[InstanceIndex];
+	const FMatrix OldTransform = PerInstanceSMData[InstanceIndex].Transform;
+	const FTransform NewLocalTransform = bWorldSpace ? NewInstanceTransform.GetRelativeTransform(ComponentToWorld) : NewInstanceTransform;
+	const FVector NewLocalLocation = NewLocalTransform.GetTranslation();
+
+	// if we are only updating rotation/scale we update the instance directly in the cluster tree
+	const bool bIsBuiltInstance = OldReorderIndex < NumBuiltInstances;
+	const bool bDoInPlaceUpdate = bIsBuiltInstance && NewLocalLocation.Equals(OldTransform.GetOrigin());
+
+	// If we're updating an instance in the tree and can't do it in-place, we have to remove it and re-add it
+	if (bIsBuiltInstance && !bDoInPlaceUpdate)
+	{
+		// Allocate a new instance render order ID, rendered last
+		const int32 NewReorderIndex = PerInstanceSMData.Num() + RemovedInstances.Num();
+		InstanceReorderTable[InstanceIndex] = NewReorderIndex;
 
-	// Allocate a new instance render order ID, rendered last
-	InstanceReorderTable[InstanceIndex] = PerInstanceSMData.Num()-1 + RemovedInstances.Num();
+		// Treat the old instance render data like a removal.
+		RemovedInstances.Add(OldReorderIndex);
+	}
 
 	bool Result = Super::UpdateInstanceTransform(InstanceIndex, NewInstanceTransform, bWorldSpace, bMarkRenderStateDirty);
 
 	if (StaticMesh)
 	{
-		UnbuiltInstanceBounds += StaticMesh->GetBounds().GetBox().TransformBy(NewInstanceTransform);
+		const FBox NewInstanceBounds = StaticMesh->GetBounds().GetBox().TransformBy(NewLocalTransform);
+
+		if (bDoInPlaceUpdate)
+		{
+			// If the new bounds are larger than the old ones, then expand the bounds on the tree to make sure culling works correctly
+			const FBox OldInstanceBounds = StaticMesh->GetBounds().GetBox().TransformBy(OldTransform);
+			if (!OldInstanceBounds.IsInside(NewInstanceBounds))
+			{
+				BuiltInstanceBounds += NewInstanceBounds;
+				UpdateInstanceTreeBoundsInternal(OldReorderIndex, NewInstanceBounds);
+			}
+		}
+		else
+		{
+			UnbuiltInstanceBounds += NewInstanceBounds;
+			UnbuiltInstanceBoundsList.Add(NewInstanceBounds);
+		}
 	}
 
-	if (!IsAsyncBuilding())
+	if (!bDoInPlaceUpdate && !IsAsyncBuilding())
 	{
 		BuildTreeAsync();
 	}
@@ -1758,7 +1908,9 @@ int32 UHierarchicalInstancedStaticMeshComponent::AddInstance(const FTransform& I
 
 	if (StaticMesh)
 	{
-		UnbuiltInstanceBounds += StaticMesh->GetBounds().GetBox().TransformBy(InstanceTransform);
+		const FBox NewInstanceBounds = StaticMesh->GetBounds().GetBox().TransformBy(InstanceTransform);
+		UnbuiltInstanceBounds += NewInstanceBounds;
+		UnbuiltInstanceBoundsList.Add(NewInstanceBounds);
 	}
 
 	return InstanceIndex;
@@ -1775,6 +1927,7 @@ void UHierarchicalInstancedStaticMeshComponent::ClearInstances()
 	NumBuiltInstances = 0;
 	SortedInstances.Empty();
 	UnbuiltInstanceBounds.Init();
+	UnbuiltInstanceBoundsList.Empty();
 
 	if (ProxySize)
 	{
@@ -1886,6 +2039,7 @@ void UHierarchicalInstancedStaticMeshComponent::BuildTree()
 		OcclusionLayerNumNodes = Builder->Result->OutOcclusionLayerNum;
 		UnbuiltInstanceBounds.Init();
 		RemovedInstances.Empty();
+		UnbuiltInstanceBoundsList.Empty();
 
 		ClusterTreePtr = MakeShareable(new TArray<FClusterNode>);
 		Exchange(*ClusterTreePtr, Builder->Result->Nodes);
@@ -1928,6 +2082,7 @@ void UHierarchicalInstancedStaticMeshComponent::AcceptPrebuiltTree(TArray<FClust
 	NumBuiltInstances = WriteOncePrebuiltInstanceBuffer.Num();
 	check(NumBuiltInstances);
 	UnbuiltInstanceBounds.Init();
+	UnbuiltInstanceBoundsList.Empty();
 	RemovedInstances.Empty();
 	ClusterTreePtr = MakeShareable(new TArray<FClusterNode>);
 	InstanceReorderTable.Empty();
@@ -1990,6 +2145,7 @@ void UHierarchicalInstancedStaticMeshComponent::BuildFlatTree(const TArray<int32
 
 		NumBuiltInstances = PerInstanceSMData.Num();
 		UnbuiltInstanceBounds.Init();
+		UnbuiltInstanceBoundsList.Empty();
 		RemovedInstances.Empty();
 
 		const FBox InstBox = StaticMesh->GetBounds().GetBox();
@@ -2069,47 +2225,49 @@ void UHierarchicalInstancedStaticMeshComponent::ApplyBuildTreeAsync(ENamedThread
 
 		UE_LOG(LogStaticMesh, Display, TEXT("Discarded foliage hierarchy of %d elements build due to concurrent removal (%.1fs)"), Builder->Result->InstanceReorderTable.Num(), float(FPlatformTime::Seconds() - StartTime));
 
-		// There were removes while we were building, it's too slow to fix up the result now, so build async again.
+		// There were removes or updates while we were building, it's too slow to fix up the result now, so build async again.
 		BuildTreeAsync();
 	}
 	else
 	{
 		NumBuiltInstances = Builder->Result->InstanceReorderTable.Num();
 
-	    if (NumBuiltInstances < PerInstanceSMData.Num())
-	    {
-		    // Add remap entries for unbuilt instances
+		if (NumBuiltInstances < PerInstanceSMData.Num())
+		{
+			// Add remap entries for unbuilt instances
 			Builder->Result->InstanceReorderTable.AddUninitialized(PerInstanceSMData.Num() - NumBuiltInstances);
-		    for (int32 Index = NumBuiltInstances; Index < PerInstanceSMData.Num(); Index++)
-		    {
+			for (int32 Index = NumBuiltInstances; Index < PerInstanceSMData.Num(); Index++)
+			{
 				Builder->Result->InstanceReorderTable[Index] = Index;
-		    }
-	    }
-	
+			}
+		}
+
 		ClusterTreePtr = MakeShareable(new TArray<FClusterNode>);
 		TArray<FClusterNode>& ClusterTree = *ClusterTreePtr;
-	    Exchange(ClusterTree, Builder->Result->Nodes);
+		Exchange(ClusterTree, Builder->Result->Nodes);
 		Exchange(InstanceReorderTable, Builder->Result->InstanceReorderTable);
 		Exchange(SortedInstances, Builder->Result->SortedInstances);
 		RemovedInstances.Empty();
 		OcclusionLayerNumNodes = Builder->Result->OutOcclusionLayerNum;
+		BuiltInstanceBounds = (ClusterTree.Num() > 0 ? FBox(ClusterTree[0].BoundMin, ClusterTree[0].BoundMax) : FBox(0));
 
+		//UE_LOG(LogStaticMesh, Display, TEXT("Built a foliage hierarchy with %d of %d elements in %.1fs."), NumBuiltInstances, PerInstanceSMData.Num(), float(FPlatformTime::Seconds() - StartTime));
 
-	    //UE_LOG(LogStaticMesh, Display, TEXT("Built a foliage hierarchy with %d of %d elements in %.1fs."), NumBuiltInstances, PerInstanceSMData.Num(), float(FPlatformTime::Seconds() - StartTime));
-	
-	    if (NumBuiltInstances < PerInstanceSMData.Num())
-	    {
-		    // There are new outstanding instances, build again!
-		    BuildTreeAsync();
-	    }
+		if (NumBuiltInstances < PerInstanceSMData.Num())
+		{
+			// There are new outstanding instances, build again!
+			UnbuiltInstanceBoundsList.RemoveAt(0, UnbuiltInstanceBoundsList.Num() - (PerInstanceSMData.Num() - NumBuiltInstances));
+			BuildTreeAsync();
+		}
 		else
 		{
 			UnbuiltInstanceBounds.Init();
+			UnbuiltInstanceBoundsList.Empty();
 			FlushAccumulatedNavigationUpdates();
 		}
 
 		ReleasePerInstanceRenderData();
-	    MarkRenderStateDirty();
+		MarkRenderStateDirty();
 
 		PostBuildStats();
 	}
