fully removed Fringe Saturation, clamped UI weight for Blendables

diff --git a/Engine/Source/Runtime/Engine/Classes/Engine/Scene.h b/Engine/Source/Runtime/Engine/Classes/Engine/Scene.h
--- a/Engine/Source/Runtime/Engine/Classes/Engine/Scene.h
+++ b/Engine/Source/Runtime/Engine/Classes/Engine/Scene.h
@@ -34,7 +34,7 @@ struct FWeightedBlendable
 	GENERATED_USTRUCT_BODY()
 
 	/** 0:no effect .. 1:full effect */
-	UPROPERTY(interp, BlueprintReadWrite, Category=FWeightedBlendable, meta=(UIMin = "0.0", UIMax = "1.0", Delta = "0.01"))
+	UPROPERTY(interp, BlueprintReadWrite, Category=FWeightedBlendable, meta=(ClampMin = "0.0", ClampMax = "1.0", Delta = "0.01"))
 	float Weight;
 
 	/** should be of the IBlendableInterface* type but UProperties cannot express that */
@@ -148,9 +148,6 @@ struct FPostProcessSettings
 	UPROPERTY(BlueprintReadWrite, Category=Overrides, meta=(PinHiddenByDefault))
 	uint32 bOverride_SceneFringeIntensity:1;
 
-	UPROPERTY(BlueprintReadWrite, Category=Overrides, meta=(PinHiddenByDefault))
-	uint32 bOverride_SceneFringeSaturation:1;
-
 	UPROPERTY(BlueprintReadWrite, Category=Overrides, meta=(PinHiddenByDefault))
 	uint32 bOverride_AmbientCubemapTint:1;
 
@@ -500,10 +497,6 @@ struct FPostProcessSettings
 	UPROPERTY(interp, BlueprintReadWrite, Category=SceneColor, meta=(UIMin = "0.0", UIMax = "5.0", editcondition = "bOverride_SceneFringeIntensity", DisplayName = "Fringe Intensity"))
 	float SceneFringeIntensity;
 
-	/** 0..1, Scene chromatic aberration / color fringe (camera imperfection) to simulate an artifact that happens in real-world lens, mostly visible in the image corners. */
-	UPROPERTY(interp, BlueprintReadWrite, Category=SceneColor, AdvancedDisplay, meta=(ClampMin = "0.0", UIMax = "1.0", editcondition = "bOverride_SceneFringeSaturation", DisplayName = "Fringe Saturation"))
-	float SceneFringeSaturation;
-
 	/** Multiplier for all bloom contributions >=0: off, 1(default), >1 brighter */
 	UPROPERTY(interp, BlueprintReadWrite, Category=Bloom, meta=(ClampMin = "0.0", UIMax = "8.0", editcondition = "bOverride_BloomIntensity", DisplayName = "Intensity"))
 	float BloomIntensity;
@@ -1044,7 +1037,6 @@ struct FPostProcessSettings
 
 		SceneColorTint = FLinearColor(1, 1, 1);
 		SceneFringeIntensity = 0.0f;
-		SceneFringeSaturation = 0.5f;
 		// next value might get overwritten by r.DefaultFeature.Bloom
 		BloomIntensity = 1.0f;
 		BloomThreshold = 1.0f;
