#merge, replays: Adding task system to replay net driver to simplify order of operations, and help with competing things happening at once

diff --git a/Engine/Source/Runtime/Engine/Classes/Engine/DemoNetDriver.h b/Engine/Source/Runtime/Engine/Classes/Engine/DemoNetDriver.h
--- a/Engine/Source/Runtime/Engine/Classes/Engine/DemoNetDriver.h
+++ b/Engine/Source/Runtime/Engine/Classes/Engine/DemoNetDriver.h
@@ -10,6 +10,26 @@ DECLARE_LOG_CATEGORY_EXTERN( LogDemo, Log, All );
 
 DECLARE_DELEGATE_OneParam(FOnGotoTimeDelegate, const bool);
 
+class UDemoNetDriver;
+
+class FQueuedReplayTask
+{
+public:
+	FQueuedReplayTask( UDemoNetDriver* InDriver ) : Driver( InDriver )
+	{
+	}
+
+	virtual ~FQueuedReplayTask()
+	{
+	}
+
+	virtual void	StartTask() = 0;
+	virtual bool	Tick() = 0;
+	virtual FString	GetName() = 0;
+
+	UDemoNetDriver* Driver;
+};
+
 /**
  * Simulated network driver for recording and playing back game sessions.
  */
@@ -62,23 +82,19 @@ class ENGINE_API UDemoNetDriver : public UNetDriver
 	double		MaxRecordTime;
 	int32		RecordCountSinceFlush;
 	float		TimeToSkip;
-	float		QueuedGotoTimeInSeconds;
-	uint32		InitialLiveDemoTime;
-	double		InitialLiveDemoTimeRealtime;
 
 	bool		bSavingCheckpoint;
 	double		LastCheckpointTime;
 
 	void		SaveCheckpoint();
 
-	FArchive*	GotoCheckpointArchive;
-	int64		GotoCheckpointSkipExtraTimeInMS;
+	void		LoadCheckpoint( FArchive* GotoCheckpointArchive, int64 GotoCheckpointSkipExtraTimeInMS );
 
-	void		LoadCheckpoint();
+	FOnGotoTimeDelegate OnGotoTimeDelegate;
 
 private:
 	bool		bIsFastForwarding;
-	bool		bIsLoadingCheckpoint;
+	bool		bIsFastForwardingForCheckpoint;
 	bool		bWasStartStreamingSuccessful;
 
 	TArray<FNetworkGUID> NonQueuedGUIDsForScrubbing;
@@ -86,7 +102,10 @@ private:
 	// If a channel is associated with Actor, adds the channel's GUID to the list of GUIDs excluded from queuing bunches during scrubbing.
 	void		AddNonQueuedActorForScrubbing(AActor* Actor);
 
-	FOnGotoTimeDelegate OnGotoTimeDelegate;
+	// Replay tasks
+	TArray< TSharedPtr< FQueuedReplayTask > >	QueuedReplayTasks;
+	TSharedPtr< FQueuedReplayTask >				ActiveReplayTask;
+	TSharedPtr< FQueuedReplayTask >				ActiveScrubReplayTask;
 
 public:
 
@@ -148,5 +167,10 @@ public:
 	void StopDemo();
 
 	void ReplayStreamingReady( bool bSuccess, bool bRecord );
-	void CheckpointReady( const bool bSuccess, const int64 SkipExtraTimeInMS );
+
+	void AddReplayTask( FQueuedReplayTask* NewTask );
+	bool IsAnyTaskPending();
+	void ClearReplayTasks();
+	bool ProcessReplayTasks();
+	bool IsNamedTaskInQueue( const FString& Name );
 };
diff --git a/Engine/Source/Runtime/Engine/Private/DemoNetDriver.cpp b/Engine/Source/Runtime/Engine/Private/DemoNetDriver.cpp
--- a/Engine/Source/Runtime/Engine/Private/DemoNetDriver.cpp
+++ b/Engine/Source/Runtime/Engine/Private/DemoNetDriver.cpp
@@ -30,7 +30,6 @@ static TAutoConsoleVariable<float> CVarDemoRecordHz( TEXT( "demo.RecordHz" ), 10
 static TAutoConsoleVariable<float> CVarDemoTimeDilation( TEXT( "demo.TimeDilation" ), -1.0f, TEXT( "Override time dilation during demo playback (-1 = don't override)" ) );
 static TAutoConsoleVariable<float> CVarDemoSkipTime( TEXT( "demo.SkipTime" ), 0, TEXT( "Skip fixed amount of network replay time (in seconds)" ) );
 static TAutoConsoleVariable<int32> CVarEnableCheckpoints( TEXT( "demo.EnableCheckpoints" ), 1, TEXT( "Whether or not checkpoints save on the server" ) );
-static TAutoConsoleVariable<int32> CVarGotoCheckpointIndex( TEXT( "demo.GotoCheckpointIndex" ), -2, TEXT( "For testing only, jump to a particular checkpoint" ) );
 static TAutoConsoleVariable<float> CVarGotoTimeInSeconds( TEXT( "demo.GotoTimeInSeconds" ), -1, TEXT( "For testing only, jump to a particular time" ) );
 static TAutoConsoleVariable<int32> CVarDemoFastForwardDestroyTearOffActors( TEXT( "demo.FastForwardDestroyTearOffActors" ), 1, TEXT( "If true, the driver will destroy any torn-off actors immediately while fast-forwarding a replay." ) );
 static TAutoConsoleVariable<int32> CVarDemoFastForwardSkipRepNotifies( TEXT( "demo.FastForwardSkipRepNotifies" ), 1, TEXT( "If true, the driver will optimize fast-forwarding by deferring calls to RepNotify functions until the fast-forward is complete. " ) );
@@ -40,6 +39,138 @@ static const int32 MAX_DEMO_READ_WRITE_BUFFER = 1024 * 2;
 
 #define DEMO_CHECKSUMS 0		// When setting this to 1, this will invalidate all demos, you will need to re-record and playback
 
+class FJumpToLiveReplayTask : public FQueuedReplayTask
+{
+public:
+	FJumpToLiveReplayTask( UDemoNetDriver* InDriver ) : FQueuedReplayTask( InDriver )
+	{
+		InitialTotalDemoTime	= Driver->ReplayStreamer->GetTotalDemoTime();
+		TaskStartTime			= FPlatformTime::Seconds();
+	}
+
+	virtual void StartTask()
+	{
+	}
+
+	virtual bool Tick() override
+	{
+		if ( !Driver->ReplayStreamer->IsLive() )
+		{
+			// The replay is no longer live, so don't try to jump to end
+			return true;
+		}
+
+		// Wait for the most recent live time
+		const bool bHasNewReplayTime = ( Driver->ReplayStreamer->GetTotalDemoTime() != InitialTotalDemoTime );
+
+		// If we haven't gotten a new time from the demo by now, assume it might not be live, and just jump to the end now so we don't hang forever
+		const bool bTimeExpired = ( FPlatformTime::Seconds() - TaskStartTime >= 15 );
+
+		if ( bHasNewReplayTime || bTimeExpired )
+		{
+			if ( bTimeExpired )
+			{
+				UE_LOG( LogDemo, Warning, TEXT( "FJumpToLiveReplayTask::Tick: Too much time since last live update." ) );
+			}
+
+			// We're ready to jump to the end now
+			Driver->JumpToEndOfLiveReplay();
+			return true;
+		}
+
+		// Waiting to get the latest update
+		return false;
+	}
+
+	virtual FString GetName() override
+	{
+		return TEXT( "FJumpToLiveReplayTask" );
+	}
+
+	uint32	InitialTotalDemoTime;		// Initial total demo time. This is used to wait until we get a more updated time so we jump to the most recent end time
+	double	TaskStartTime;				// This is the time the task started. If too much real-time passes, we'll just jump to the current end
+};
+
+class FGotoTimeInSecondsTask : public FQueuedReplayTask
+{
+public:
+	FGotoTimeInSecondsTask( UDemoNetDriver* InDriver, const float InTimeInSeconds ) : FQueuedReplayTask( InDriver ), TimeInSeconds( InTimeInSeconds ), GotoCheckpointArchive( nullptr ), GotoCheckpointSkipExtraTimeInMS( -1 )
+	{
+	}
+
+	virtual void StartTask() override
+	{		
+		OldTimeInSeconds		= Driver->DemoCurrentTime;	// Rember current time, so we can restore on failure
+		Driver->DemoCurrentTime	= TimeInSeconds;			// Also, update current time so HUD reflects desired scrub time now
+
+		// Tell the streamer to start going to this time
+		Driver->ReplayStreamer->GotoTimeInMS( TimeInSeconds * 1000, FOnCheckpointReadyDelegate::CreateRaw( this, &FGotoTimeInSecondsTask::CheckpointReady ) );
+
+		// Pause channels while we wait (so the world is paused while we wait for the new stream location to load)
+		Driver->PauseChannels( true );
+	}
+
+	virtual bool Tick() override
+	{
+		if ( GotoCheckpointSkipExtraTimeInMS == -2 )
+		{
+			// Detect failure case
+			return true;
+		}
+
+		if ( GotoCheckpointArchive != nullptr )
+		{
+			if ( GotoCheckpointSkipExtraTimeInMS > 0 && !Driver->ReplayStreamer->IsDataAvailable() )
+			{
+				// Wait for rest of stream before loading checkpoint
+				// We do this so we can load the checkpoint and fastforward the stream all at once
+				// We do this so that the OnReps don't stay queued up outside of this frame
+				return false;
+			}
+
+			// We're done
+			Driver->LoadCheckpoint( GotoCheckpointArchive, GotoCheckpointSkipExtraTimeInMS );
+			return true;
+		}
+
+		return false;
+	}
+
+	virtual FString GetName() override
+	{
+		return TEXT( "FGotoTimeInSecondsTask" );
+	}
+
+	void CheckpointReady( const bool bSuccess, const int64 SkipExtraTimeInMS )
+	{
+		check( GotoCheckpointArchive == NULL );
+		check( GotoCheckpointSkipExtraTimeInMS == -1 );
+
+		if ( !bSuccess )
+		{
+			UE_LOG( LogDemo, Warning, TEXT( "FGotoTimeInSecondsTask::CheckpointReady: Failed to go to checkpoint." ) );
+
+			// Restore old demo time
+			Driver->DemoCurrentTime = OldTimeInSeconds;
+
+			// Call delegate if any
+			Driver->OnGotoTimeDelegate.ExecuteIfBound( false );
+			Driver->OnGotoTimeDelegate.Unbind();
+
+			GotoCheckpointSkipExtraTimeInMS = -2;	// So tick can detect failure case
+			return;
+		}
+
+		GotoCheckpointArchive			= Driver->ReplayStreamer->GetCheckpointArchive();
+		GotoCheckpointSkipExtraTimeInMS = SkipExtraTimeInMS;
+	}
+
+	float				OldTimeInSeconds;		// So we can restore on failure
+	float				TimeInSeconds;
+	FArchive*			GotoCheckpointArchive;
+	int64				GotoCheckpointSkipExtraTimeInMS;
+};
+
 /*-----------------------------------------------------------------------------
 	UDemoNetDriver.
 -----------------------------------------------------------------------------*/
@@ -49,6 +180,83 @@ UDemoNetDriver::UDemoNetDriver(const FObjectInitializer& ObjectInitializer)
 {
 }
 
+void UDemoNetDriver::AddReplayTask( FQueuedReplayTask* NewTask )
+{
+	UE_LOG( LogDemo, Verbose, TEXT( "UDemoNetDriver::AddReplayTask. Name: %s" ), *NewTask->GetName() );
+
+	QueuedReplayTasks.Add( TSharedPtr< FQueuedReplayTask >( NewTask ) );
+
+	// Give this task a chance to immediately start if nothing else is happening
+	if ( !IsAnyTaskPending() )
+	{
+		ProcessReplayTasks();	
+	}
+}
+
+bool UDemoNetDriver::IsAnyTaskPending()
+{
+	return ( QueuedReplayTasks.Num() > 0 ) || ActiveReplayTask.IsValid();
+}
+
+void UDemoNetDriver::ClearReplayTasks()
+{
+	QueuedReplayTasks.Empty();
+
+	ActiveReplayTask = nullptr;
+}
+
+bool UDemoNetDriver::ProcessReplayTasks()
+{
+	while ( IsAnyTaskPending() )
+	{
+		if ( !ActiveReplayTask.IsValid() )
+		{
+			// If we don't have an active task, pull one off now
+			ActiveReplayTask = QueuedReplayTasks[0];
+			QueuedReplayTasks.RemoveAt( 0 );
+
+			UE_LOG( LogDemo, Verbose, TEXT( "UDemoNetDriver::ProcessReplayTasks. Name: %s" ), *ActiveReplayTask->GetName() );
+
+			// Start the task
+			ActiveReplayTask->StartTask();
+		}
+
+		// Tick the currently active task
+		if ( ActiveReplayTask.IsValid() )
+		{
+			if ( !ActiveReplayTask->Tick() )
+			{
+				// Task isn't done, we can return
+				return false;
+			}
+
+			// This task is now done
+			// Allow the next task to start processing (or return below if QueuedReplayTasks is empty)
+			ActiveReplayTask = nullptr;
+		}
+	}
+
+	return true;	// No tasks to process
+}
+
+bool UDemoNetDriver::IsNamedTaskInQueue( const FString& Name )
+{
+	if ( ActiveReplayTask.IsValid() && ActiveReplayTask->GetName() == Name )
+	{
+		return true;
+	}
+
+	for ( int32 i = 0; i < QueuedReplayTasks.Num(); i++ )
+	{
+		if ( QueuedReplayTasks[0]->GetName() == Name )
+		{
+			return true;
+		}
+	}
+
+	return false;
+}
+
 bool UDemoNetDriver::InitBase( bool bInitAsClient, FNetworkNotify* InNotify, const FURL& URL, bool bReuseAddressAndPort, FString& Error )
 {
 	if ( Super::InitBase( bInitAsClient, InNotify, URL, bReuseAddressAndPort, Error ) )
@@ -60,11 +268,7 @@ bool UDemoNetDriver::InitBase( bool bInitAsClient, FNetworkNotify* InNotify, con
 		bChannelsArePaused				= false;
 		TimeToSkip						= 0.0f;
 		bIsFastForwarding				= false;
-		GotoCheckpointSkipExtraTimeInMS = -1;
-		QueuedGotoTimeInSeconds			= -1.0f;
-		bIsLoadingCheckpoint			= false;
-		InitialLiveDemoTime				= 0;
-		InitialLiveDemoTimeRealtime		= 0;
+		bIsFastForwardingForCheckpoint	= false;
 		bWasStartStreamingSuccessful	= true;
 
 		ResetDemoState();
@@ -713,6 +917,7 @@ void UDemoNetDriver::StopDemo()
 	}
 
 	ReplayStreamer->StopStreaming();
+	ClearReplayTasks();
 
 	check( ClientConnections.Num() == 0 );
 	check( ServerConnection == NULL );
@@ -1288,11 +1493,13 @@ void UDemoNetDriver::SkipTime(const float InTimeToSkip)
 
 void UDemoNetDriver::GotoTimeInSeconds(const float TimeInSeconds, const FOnGotoTimeDelegate& InOnGotoTimeDelegate)
 {
-	if ( QueuedGotoTimeInSeconds < 0 )
+	if ( IsNamedTaskInQueue( TEXT( "FGotoTimeInSecondsTask" ) ) )
 	{
-		QueuedGotoTimeInSeconds = TimeInSeconds;
+		InOnGotoTimeDelegate.ExecuteIfBound( false );
+		return;		// Don't allow scrubbing if we already are
 	}
 
+	AddReplayTask( new FGotoTimeInSecondsTask( this, TimeInSeconds ) );
 	OnGotoTimeDelegate = InOnGotoTimeDelegate;
 }
 
@@ -1304,18 +1511,14 @@ void UDemoNetDriver::JumpToEndOfLiveReplay()
 
 	DemoTotalTime = (float)TotalDemoTimeInMS / 1000.0f;
 
-	// For now, use no buffer, try to get as close to the end as possible
-	// This will likely mean the world is paused at first, but this might be ok considering it will be as live as it can be
-	const uint32 BufferInMS = 0;//InitialLiveDemoTime != 0 ? 5 * 1000 : 0;
-
-	const uint32 JoinTimeInMS = ReplayStreamer->GetTotalDemoTime() - BufferInMS;
+	const uint32 BufferInMS = 5 * 1000;
 
-	OldDemoCurrentTime	= DemoCurrentTime;		// So we can restore on failure
-	DemoCurrentTime		= ( float )( JoinTimeInMS ) / 1000.0f;
+	const uint32 JoinTimeInMS = FMath::Max( (uint32)0, ReplayStreamer->GetTotalDemoTime() - BufferInMS );
 
-	ReplayStreamer->GotoTimeInMS( JoinTimeInMS, FOnCheckpointReadyDelegate::CreateUObject( this, &UDemoNetDriver::CheckpointReady ) );
-
-	InitialLiveDemoTime = 0;
+	if ( JoinTimeInMS > 0 )
+	{
+		GotoTimeInSeconds( (float)JoinTimeInMS / 1000.0f );
+	}
 }
 
 void UDemoNetDriver::AddUserToReplay( const FString& UserString )
@@ -1334,72 +1537,25 @@ void UDemoNetDriver::TickDemoPlayback( float DeltaSeconds )
 		return;
 	}
 
-	// See if we need to jump to the end (or close to it) of a live game we just joined
-	if ( InitialLiveDemoTime > 0 )
+	if ( CVarGotoTimeInSeconds.GetValueOnGameThread() >= 0.0f )
 	{
-		if ( !ReplayStreamer->IsLive() )
-		{
-			// The replay was finalized while we joined, so don't try to jump to end
-			InitialLiveDemoTime = 0;
-		}
-		else
-		{
-			// Wait for the most recent live time
-			const bool bHasNewReplayTime = ( ReplayStreamer->GetTotalDemoTime() != InitialLiveDemoTime );
-				
-			// If we haven't gotten a new time from the demo by now, assume it might not be live, and just jump to the end now so we don't hang forever
-			const bool bTimeExpired = ( FPlatformTime::Seconds() - InitialLiveDemoTimeRealtime >= 15 );
-
-			if ( bTimeExpired )
-			{
-				UE_LOG( LogDemo, Warning, TEXT( "UDemoNetDriver::TickDemoPlayback: Too much time since last live update." ) );
-			}
-
-			if ( !bHasNewReplayTime && !bTimeExpired )
-			{
-				return;
-			}
-
-			JumpToEndOfLiveReplay();
-		}
-	}
-
-	if ( GotoCheckpointArchive == NULL && !ReplayStreamer->IsLoadingCheckpoint() )
-	{
-		// Check checkpoint debug cvars
-		check( GotoCheckpointSkipExtraTimeInMS == -1 );
-
-		if ( CVarGotoCheckpointIndex.GetValueOnGameThread() >= -1 )
-		{
-			ReplayStreamer->GotoCheckpointIndex( CVarGotoCheckpointIndex.GetValueOnGameThread(), FOnCheckpointReadyDelegate::CreateUObject( this, &UDemoNetDriver::CheckpointReady ) );
-			CVarGotoCheckpointIndex.AsVariable()->Set( TEXT( "-2" ), ECVF_SetByConsole );
-		}
-
-		if ( CVarGotoTimeInSeconds.GetValueOnGameThread() >= 0.0f )
-		{
-			ReplayStreamer->GotoTimeInMS( CVarGotoTimeInSeconds.GetValueOnGameThread() * 1000, FOnCheckpointReadyDelegate::CreateUObject( this, &UDemoNetDriver::CheckpointReady ) );
-			CVarGotoTimeInSeconds.AsVariable()->Set( TEXT( "-1" ), ECVF_SetByConsole );
-		}
-
-		// Process any queued scrub time
-		if ( QueuedGotoTimeInSeconds >= 0.0f )
-		{
-			OldDemoCurrentTime		= DemoCurrentTime;			// So we can restore on failure
-			DemoCurrentTime			= QueuedGotoTimeInSeconds;
-
-			ReplayStreamer->GotoTimeInMS( QueuedGotoTimeInSeconds * 1000, FOnCheckpointReadyDelegate::CreateUObject( this, &UDemoNetDriver::CheckpointReady ) );
-
-			QueuedGotoTimeInSeconds = -1.0f;
-		}
+		GotoTimeInSeconds( CVarGotoTimeInSeconds.GetValueOnGameThread() );
+		CVarGotoTimeInSeconds.AsVariable()->Set( TEXT( "-1" ), ECVF_SetByConsole );
 	}
 
 	if ( CVarDemoSkipTime.GetValueOnGameThread() > 0 )
 	{
 		// Just overwrite existing value, cvar wins in this case
-		SkipTime( CVarDemoSkipTime.GetValueOnGameThread() );		
+		SkipTime( CVarDemoSkipTime.GetValueOnGameThread() );
 		CVarDemoSkipTime.AsVariable()->Set( TEXT( "0" ), ECVF_SetByConsole );
 	}
 
+	if ( !ProcessReplayTasks() )
+	{
+		// We're busy processing tasks, return
+		return;
+	}
+
 	// Update total demo time
 	const uint32 TotalDemoTimeInMS = ReplayStreamer->GetTotalDemoTime();
 
@@ -1409,7 +1565,7 @@ void UDemoNetDriver::TickDemoPlayback( float DeltaSeconds )
 	}
 
 	// See if we have time to skip
-	if ( TimeToSkip > 0.0f && !ReplayStreamer->IsLoadingCheckpoint() )
+	if ( TimeToSkip > 0.0f && !IsAnyTaskPending() )
 	{
 		const uint32 TimeInMSToCheck = FMath::Clamp( GetDemoCurrentTimeInMS() + (uint32)( TimeToSkip * 1000 ), (uint32)0, TotalDemoTimeInMS );
 
@@ -1417,34 +1573,13 @@ void UDemoNetDriver::TickDemoPlayback( float DeltaSeconds )
 
 		DemoCurrentTime += TimeToSkip;
 
-		bIsFastForwarding = true;
-
-		TimeToSkip = 0.0f;
-	}
-
-	// If we have a checkpoint but and we don't need to fine scrub, we can load it now
-	if ( GotoCheckpointArchive != NULL && GotoCheckpointSkipExtraTimeInMS == -1 )
-	{
-		LoadCheckpoint();
+		bIsFastForwarding	= true;
+		TimeToSkip			= 0.0f;
 	}
 
 	// Make sure there is data available to read
-	if ( !ReplayStreamer->IsDataAvailable() )
-	{
-		PauseChannels( true );
-		return;
-	}
-
-	// Check to see if we have a pending checkpoint to read from
-	if ( GotoCheckpointArchive != NULL )
-	{
-		LoadCheckpoint();
-	}
-
-	check( GotoCheckpointArchive == NULL );
-
 	// If we're at the end of the demo, just pause channels and return
-	if ( bDemoPlaybackDone )
+	if ( bDemoPlaybackDone || !ReplayStreamer->IsDataAvailable() )
 	{
 		PauseChannels( true );
 		return;
@@ -1478,7 +1613,7 @@ void UDemoNetDriver::TickDemoPlayback( float DeltaSeconds )
 
 		// We may have been fast-forwarding immediately after loading a checkpoint
 		// for fine-grained scrubbing. If so, at this point we are no longer loading a checkpoint.
-		bIsLoadingCheckpoint = false;
+		bIsFastForwardingForCheckpoint = false;
 
 		// Flush all pending RepNotifies that were built up during the fast-forward.
 		if ( ServerConnection != nullptr)
@@ -1560,8 +1695,6 @@ void UDemoNetDriver::ReplayStreamingReady( bool bSuccess, bool bRecord )
 
 	if ( !bRecord )
 	{
-		InitialLiveDemoTime = 0;
-
 		FString Error;
 		
 		const double StartTime = FPlatformTime::Seconds();
@@ -1582,40 +1715,17 @@ void UDemoNetDriver::ReplayStreamingReady( bool bSuccess, bool bRecord )
 			else
 			{
 				UE_LOG( LogDemo, Log, TEXT( "UDemoNetConnection::ReplayStreamingReady: Deferring checkpoint until next available time." ) );
-				InitialLiveDemoTime = ReplayStreamer->GetTotalDemoTime();
-				InitialLiveDemoTimeRealtime = FPlatformTime::Seconds();
+				AddReplayTask( new FJumpToLiveReplayTask( this ) );
 			}
 		}
 	}
 }
 
-void UDemoNetDriver::CheckpointReady( const bool bSuccess, const int64 SkipExtraTimeInMS )
-{
-	check( GotoCheckpointArchive == NULL );
-	check( GotoCheckpointSkipExtraTimeInMS == -1 );
-
-	if ( !bSuccess )
-	{
-		UE_LOG( LogDemo, Warning, TEXT( "UDemoNetConnection::CheckpointReady: Failed to go to checkpoint." ) );
-		DemoCurrentTime = OldDemoCurrentTime;
-		OnGotoTimeDelegate.ExecuteIfBound(false);
-		OnGotoTimeDelegate.Unbind();
-		return;
-	}
-
-	GotoCheckpointArchive			= ReplayStreamer->GetCheckpointArchive();
-	GotoCheckpointSkipExtraTimeInMS = SkipExtraTimeInMS;
-}
-
-void UDemoNetDriver::LoadCheckpoint()
+void UDemoNetDriver::LoadCheckpoint( FArchive* GotoCheckpointArchive, int64 GotoCheckpointSkipExtraTimeInMS )
 {
 	DECLARE_SCOPE_CYCLE_COUNTER(TEXT("LoadCheckpoint time"), STAT_ReplayCheckpointLoadTime, STATGROUP_Net);
 
-	if ( GotoCheckpointArchive == NULL )
-	{
-		UE_LOG( LogDemo, Warning, TEXT( "UDemoNetConnection::LoadCheckpoint: GotoCheckpointArchive == NULL." ) );
-		return;
-	}
+	check( GotoCheckpointArchive != NULL );
 
 	// Reset the never-queue GUID list, we'll rebuild it
 	NonQueuedGUIDsForScrubbing.Empty();
@@ -1752,17 +1862,15 @@ void UDemoNetDriver::LoadCheckpoint()
 		{
 			DemoCurrentTime += (float)GotoCheckpointSkipExtraTimeInMS / 1000;
 
-			bIsFastForwarding		= true;
-			bIsLoadingCheckpoint	= true;
+			bIsFastForwarding				= true;
+			bIsFastForwardingForCheckpoint	= true;
 		}
 		else
 		{
-			bIsLoadingCheckpoint	= false;
-			bIsFastForwarding		= false;
+			bIsFastForwardingForCheckpoint	= false;
+			bIsFastForwarding				= false;
 		}
 
-		GotoCheckpointArchive			= NULL;
-		GotoCheckpointSkipExtraTimeInMS = -1;
 		return;
 	}
 
@@ -1802,13 +1910,13 @@ void UDemoNetDriver::LoadCheckpoint()
 
 	if ( GotoCheckpointSkipExtraTimeInMS != -1 )
 	{
-		bIsLoadingCheckpoint = true;
-		bIsFastForwarding = true;
+		bIsFastForwardingForCheckpoint	= true;
+		bIsFastForwarding				= true;
 	}
 	else
 	{
-		bIsLoadingCheckpoint = false;
-		bIsFastForwarding = false;
+		bIsFastForwardingForCheckpoint	= false;
+		bIsFastForwarding				= false;
 	}
 
 	ReadDemoFrame( GotoCheckpointArchive );
@@ -1819,9 +1927,7 @@ void UDemoNetDriver::LoadCheckpoint()
 		DemoCurrentTime += (float)GotoCheckpointSkipExtraTimeInMS / 1000;
 	}
 
-	GotoCheckpointSkipExtraTimeInMS = -1;
-	GotoCheckpointArchive			= NULL;
-	bDemoPlaybackDone				= false;
+	bDemoPlaybackDone = false;
 
 	if ( SpectatorController && ViewTargetGUID.IsValid() )
 	{
@@ -1842,7 +1948,7 @@ bool UDemoNetDriver::ShouldQueueBunchesForActorGUID(FNetworkGUID InGUID) const
 	}
 
 	// While loading a checkpoint, queue most bunches so that we don't process them all on one frame.
-	if (bIsLoadingCheckpoint)
+	if ( bIsFastForwardingForCheckpoint )
 	{
 		return !NonQueuedGUIDsForScrubbing.Contains(InGUID);
 	}
