Engine: Fixed a typo in the abtest command help
[INTEGRATE] Change 2670403 by [REDACTED_PERSON][REDACTED_HANDLE] on 2015/08/26 20:25:33

diff --git a/Engine/Source/Runtime/Engine/Private/UnrealEngine.cpp b/Engine/Source/Runtime/Engine/Private/UnrealEngine.cpp
--- a/Engine/Source/Runtime/Engine/Private/UnrealEngine.cpp
+++ b/Engine/Source/Runtime/Engine/Private/UnrealEngine.cpp
@@ -1093,7 +1093,7 @@ static TAutoConsoleVariable<int32> CVarABTestCoolDown(
 static TAutoConsoleVariable<int32> CVarABTestMinFramesPerTrial(
 	TEXT("abtest.MinFramesPerTrial"),
 	5,
-	TEXT("The number of frames to run a given command before swithing; this is randomized."));
+	TEXT("The number of frames to run a given command before switching; this is randomized."));
 
 static TAutoConsoleVariable<int32> CVarABTestNumResamples(
 	TEXT("abtest.NumResamples"),
