Merge changes to workaround the absence of VertexBuffer SRVs with implicit format conversion on Metal & an alignment gotcha in the particle simulation uniform buffer definition.

diff --git a/Engine/Shaders/ParticleGPUSpriteVertexFactory.usf b/Engine/Shaders/ParticleGPUSpriteVertexFactory.usf
--- a/Engine/Shaders/ParticleGPUSpriteVertexFactory.usf
+++ b/Engine/Shaders/ParticleGPUSpriteVertexFactory.usf
@@ -17,7 +17,13 @@
 #define USE_PARTICLE_SIZE (FEATURE_LEVEL >= FEATURE_LEVEL_SM5 && NEEDS_PARTICLE_SIZE)
 
 /** Buffer containing particle indices. */
-Buffer<float2> ParticleIndices;
+#if COMPILER_METAL
+	// Metal does not have buffer SRVs, so data types have to match
+	#undef half2
+	Buffer<half2> ParticleIndices;
+#else
+	Buffer<float2> ParticleIndices;
+#endif
 
 /** Offset in to the particle indices buffer. */
 uint ParticleIndicesOffset;
diff --git a/Engine/Source/Runtime/Engine/Private/Particles/ParticleGpuSimulation.cpp b/Engine/Source/Runtime/Engine/Private/Particles/ParticleGpuSimulation.cpp
--- a/Engine/Source/Runtime/Engine/Private/Particles/ParticleGpuSimulation.cpp
+++ b/Engine/Source/Runtime/Engine/Private/Particles/ParticleGpuSimulation.cpp
@@ -763,10 +763,10 @@ FArchive& operator<<(FArchive& Ar, FParticlePerFrameSimulationShaderParameters&
 BEGIN_UNIFORM_BUFFER_STRUCT( FVectorFieldUniformParameters,)
 	DECLARE_UNIFORM_BUFFER_STRUCT_MEMBER_ARRAY( FMatrix, WorldToVolume, [MAX_VECTOR_FIELDS] )
 	DECLARE_UNIFORM_BUFFER_STRUCT_MEMBER_ARRAY( FMatrix, VolumeToWorld, [MAX_VECTOR_FIELDS] )
-	DECLARE_UNIFORM_BUFFER_STRUCT_MEMBER_ARRAY( FVector, VolumeSize, [MAX_VECTOR_FIELDS] )
-	DECLARE_UNIFORM_BUFFER_STRUCT_MEMBER_ARRAY( FVector2D, IntensityAndTightness, [MAX_VECTOR_FIELDS] )
 	DECLARE_UNIFORM_BUFFER_STRUCT_MEMBER( int32, Count )
 	DECLARE_UNIFORM_BUFFER_STRUCT_MEMBER_ARRAY( FVector, TilingAxes, [MAX_VECTOR_FIELDS] )
+	DECLARE_UNIFORM_BUFFER_STRUCT_MEMBER_ARRAY( FVector, VolumeSize, [MAX_VECTOR_FIELDS] )
+	DECLARE_UNIFORM_BUFFER_STRUCT_MEMBER_ARRAY( FVector2D, IntensityAndTightness, [MAX_VECTOR_FIELDS] )
 END_UNIFORM_BUFFER_STRUCT( FVectorFieldUniformParameters )
 
 IMPLEMENT_UNIFORM_BUFFER_STRUCT(FVectorFieldUniformParameters,TEXT("VectorFields"));
@@ -4406,6 +4406,9 @@ void FFXSystem::SimulateGPUParticles(
 		};
 		SetRenderTargets(RHICmdList, 4, InjectRenderTargets, FTextureRHIParamRef(), 0, NULL);
 		RHICmdList.SetViewport(0, 0, 0.0f, GParticleSimulationTextureSizeX, GParticleSimulationTextureSizeY, 1.0f);
+		RHICmdList.SetDepthStencilState(TStaticDepthStencilState<false, CF_Always>::GetRHI());
+		RHICmdList.SetRasterizerState(TStaticRasterizerState<FM_Solid, CM_None>::GetRHI());
+		RHICmdList.SetBlendState(TStaticBlendState<>::GetRHI());
 
 		// Inject particles.
 		InjectNewParticles(RHICmdList, FeatureLevel, NewParticles);
