#merge, replays: Only process one task per frame to simplify things + better sanity checks for setting bIsFastForwarding + convert skip time logic to task + cleanup fast forward finalization.

diff --git a/Engine/Source/Runtime/Engine/Classes/Engine/DemoNetDriver.h b/Engine/Source/Runtime/Engine/Classes/Engine/DemoNetDriver.h
--- a/Engine/Source/Runtime/Engine/Classes/Engine/DemoNetDriver.h
+++ b/Engine/Source/Runtime/Engine/Classes/Engine/DemoNetDriver.h
@@ -81,7 +81,6 @@ class ENGINE_API UDemoNetDriver : public UNetDriver
 	double		LastRecordAvgFlush;
 	double		MaxRecordTime;
 	int32		RecordCountSinceFlush;
-	float		TimeToSkip;
 
 	bool		bSavingCheckpoint;
 	double		LastCheckpointTime;
@@ -120,6 +119,7 @@ public:
 	virtual void ProcessRemoteFunction( class AActor* Actor, class UFunction* Function, void* Parameters, struct FOutParmRec* OutParms, struct FFrame* Stack, class UObject* SubObject = nullptr ) override;
 	virtual bool IsAvailable() const override { return true; }
 	void SkipTime(const float InTimeToSkip);
+	void SkipTimeInternal( const float SecondsToSkip, const bool InFastForward, const bool InIsForCheckpoint );
 	bool InitConnectInternal( FString& Error );
 	virtual bool ShouldClientDestroyTearOffActors() const override;
 	virtual bool ShouldSkipRepNotifies() const override;
@@ -149,6 +149,7 @@ public:
 	bool ConditionallyReadDemoFrame();
 	bool ReadDemoFrame( FArchive* Archive );
 	void TickDemoPlayback( float DeltaSeconds );
+	void FinalizeFastForward( const float StartTime );
 	void SpawnDemoRecSpectator( UNetConnection* Connection );
 	void ResetDemoState();
 	void JumpToEndOfLiveReplay();
diff --git a/Engine/Source/Runtime/Engine/Private/DemoNetDriver.cpp b/Engine/Source/Runtime/Engine/Private/DemoNetDriver.cpp
--- a/Engine/Source/Runtime/Engine/Private/DemoNetDriver.cpp
+++ b/Engine/Source/Runtime/Engine/Private/DemoNetDriver.cpp
@@ -100,6 +100,8 @@ public:
 
 	virtual void StartTask() override
 	{		
+		check( !Driver->IsFastForwarding() );
+
 		OldTimeInSeconds		= Driver->DemoCurrentTime;	// Rember current time, so we can restore on failure
 		Driver->DemoCurrentTime	= TimeInSeconds;			// Also, update current time so HUD reflects desired scrub time now
 
@@ -171,6 +173,38 @@ public:
 	int64				GotoCheckpointSkipExtraTimeInMS;
 };
 
+class FSkipTimeInSecondsTask : public FQueuedReplayTask
+{
+public:
+	FSkipTimeInSecondsTask( UDemoNetDriver* InDriver, const float InSecondsToSkip ) : FQueuedReplayTask( InDriver ), SecondsToSkip( InSecondsToSkip )
+	{
+	}
+
+	virtual void StartTask() override
+	{
+		check( !Driver->IsFastForwarding() );
+
+		const uint32 TimeInMSToCheck = FMath::Clamp( Driver->GetDemoCurrentTimeInMS() + ( uint32 )( SecondsToSkip * 1000 ), ( uint32 )0, Driver->ReplayStreamer->GetTotalDemoTime() );
+
+		Driver->ReplayStreamer->SetHighPriorityTimeRange( Driver->GetDemoCurrentTimeInMS(), TimeInMSToCheck );
+
+		Driver->SkipTimeInternal( SecondsToSkip, true, false );
+	}
+
+	virtual bool Tick() override
+	{
+		// The real work was done in StartTask, so we're done
+		return true;
+	}
+
+	virtual FString GetName() override
+	{
+		return TEXT( "FSkipTimeInSecondsTask" );
+	}
+
+	float SecondsToSkip;
+};
+
 /*-----------------------------------------------------------------------------
 	UDemoNetDriver.
 -----------------------------------------------------------------------------*/
@@ -207,33 +241,29 @@ void UDemoNetDriver::ClearReplayTasks()
 
 bool UDemoNetDriver::ProcessReplayTasks()
 {
-	while ( IsAnyTaskPending() )
+	if ( !ActiveReplayTask.IsValid() && QueuedReplayTasks.Num() > 0 )
 	{
-		if ( !ActiveReplayTask.IsValid() )
-		{
-			// If we don't have an active task, pull one off now
-			ActiveReplayTask = QueuedReplayTasks[0];
-			QueuedReplayTasks.RemoveAt( 0 );
+		// If we don't have an active task, pull one off now
+		ActiveReplayTask = QueuedReplayTasks[0];
+		QueuedReplayTasks.RemoveAt( 0 );
 
-			UE_LOG( LogDemo, Verbose, TEXT( "UDemoNetDriver::ProcessReplayTasks. Name: %s" ), *ActiveReplayTask->GetName() );
+		UE_LOG( LogDemo, Verbose, TEXT( "UDemoNetDriver::ProcessReplayTasks. Name: %s" ), *ActiveReplayTask->GetName() );
 
-			// Start the task
-			ActiveReplayTask->StartTask();
-		}
+		// Start the task
+		ActiveReplayTask->StartTask();
+	}
 
-		// Tick the currently active task
-		if ( ActiveReplayTask.IsValid() )
+	// Tick the currently active task
+	if ( ActiveReplayTask.IsValid() )
+	{
+		if ( !ActiveReplayTask->Tick() )
 		{
-			if ( !ActiveReplayTask->Tick() )
-			{
-				// Task isn't done, we can return
-				return false;
-			}
-
-			// This task is now done
-			// Allow the next task to start processing (or return below if QueuedReplayTasks is empty)
-			ActiveReplayTask = nullptr;
+			// Task isn't done, we can return
+			return false;
 		}
+
+		// This task is now done
+		ActiveReplayTask = nullptr;
 	}
 
 	return true;	// No tasks to process
@@ -266,7 +296,6 @@ bool UDemoNetDriver::InitBase( bool bInitAsClient, FNetworkNotify* InNotify, con
 		bIsRecordingDemoFrame			= false;
 		bDemoPlaybackDone				= false;
 		bChannelsArePaused				= false;
-		TimeToSkip						= 0.0f;
 		bIsFastForwarding				= false;
 		bIsFastForwardingForCheckpoint	= false;
 		bWasStartStreamingSuccessful	= true;
@@ -1488,7 +1517,23 @@ bool UDemoNetDriver::ReadDemoFrame( FArchive* Archive )
 
 void UDemoNetDriver::SkipTime(const float InTimeToSkip)
 {
-	TimeToSkip = InTimeToSkip;
+	if ( IsNamedTaskInQueue( TEXT( "FSkipTimeInSecondsTask" ) ) )
+	{
+		return;		// Don't allow time skipping if we already are
+	}
+
+	AddReplayTask( new FSkipTimeInSecondsTask( this, InTimeToSkip ) );
+}
+
+void UDemoNetDriver::SkipTimeInternal( const float SecondsToSkip, const bool InFastForward, const bool InIsForCheckpoint )
+{
+	check( !bIsFastForwarding );				// Can only do one of these at a time (use tasks to gate this)
+	check( !bIsFastForwardingForCheckpoint );	// Can only do one of these at a time (use tasks to gate this)
+
+	DemoCurrentTime += SecondsToSkip;
+
+	bIsFastForwarding				= InFastForward;
+	bIsFastForwardingForCheckpoint	= InIsForCheckpoint;
 }
 
 void UDemoNetDriver::GotoTimeInSeconds(const float TimeInSeconds, const FOnGotoTimeDelegate& InOnGotoTimeDelegate)
@@ -1557,24 +1602,9 @@ void UDemoNetDriver::TickDemoPlayback( float DeltaSeconds )
 	}
 
 	// Update total demo time
-	const uint32 TotalDemoTimeInMS = ReplayStreamer->GetTotalDemoTime();
-
-	if ( TotalDemoTimeInMS > 0 )
+	if ( ReplayStreamer->GetTotalDemoTime() > 0 )
 	{
-		DemoTotalTime = (float)TotalDemoTimeInMS / 1000.0f;
-	}
-
-	// See if we have time to skip
-	if ( TimeToSkip > 0.0f && !IsAnyTaskPending() )
-	{
-		const uint32 TimeInMSToCheck = FMath::Clamp( GetDemoCurrentTimeInMS() + (uint32)( TimeToSkip * 1000 ), (uint32)0, TotalDemoTimeInMS );
-
-		ReplayStreamer->SetHighPriorityTimeRange( GetDemoCurrentTimeInMS(), TimeInMSToCheck );
-
-		DemoCurrentTime += TimeToSkip;
-
-		bIsFastForwarding	= true;
-		TimeToSkip			= 0.0f;
+		DemoTotalTime = ( float )ReplayStreamer->GetTotalDemoTime() / 1000.0f;
 	}
 
 	// Make sure there is data available to read
@@ -1609,34 +1639,39 @@ void UDemoNetDriver::TickDemoPlayback( float DeltaSeconds )
 	// Finalize any fast forward stuff that needs to happen
 	if ( bIsFastForwarding )
 	{
-		bIsFastForwarding = false;
+		FinalizeFastForward( FastForwardStartSeconds );
+	}
+}
 
-		// We may have been fast-forwarding immediately after loading a checkpoint
-		// for fine-grained scrubbing. If so, at this point we are no longer loading a checkpoint.
-		bIsFastForwardingForCheckpoint = false;
+void UDemoNetDriver::FinalizeFastForward( const float StartTime )
+{
+	bIsFastForwarding = false;
+
+	// We may have been fast-forwarding immediately after loading a checkpoint
+	// for fine-grained scrubbing. If so, at this point we are no longer loading a checkpoint.
+	bIsFastForwardingForCheckpoint = false;
 
-		// Flush all pending RepNotifies that were built up during the fast-forward.
-		if ( ServerConnection != nullptr)
+	// Flush all pending RepNotifies that were built up during the fast-forward.
+	if ( ServerConnection != nullptr )
+	{
+		for ( auto& ChannelPair : ServerConnection->ActorChannels )
 		{
-			for ( auto& ChannelPair : ServerConnection->ActorChannels )
+			if ( ChannelPair.Value != NULL )
 			{
-				if ( ChannelPair.Value != NULL )
+				for ( auto& ReplicatorPair : ChannelPair.Value->ReplicationMap )
 				{
-					for (auto& ReplicatorPair : ChannelPair.Value->ReplicationMap)
-					{
-						ReplicatorPair.Value->CallRepNotifies(true);
-					}
+					ReplicatorPair.Value->CallRepNotifies( true );
 				}
 			}
 		}
+	}
 
-		const auto FastForwardTotalSeconds = FPlatformTime::Seconds() - FastForwardStartSeconds;
+	const auto FastForwardTotalSeconds = FPlatformTime::Seconds() - StartTime;
 
-		OnGotoTimeDelegate.ExecuteIfBound(true);
-		OnGotoTimeDelegate.Unbind();
+	OnGotoTimeDelegate.ExecuteIfBound( true );
+	OnGotoTimeDelegate.Unbind();
 
-		UE_LOG( LogDemo, Log, TEXT( "Fast forward took %.2f seconds." ), FastForwardTotalSeconds );
-	}
+	UE_LOG( LogDemo, Log, TEXT( "Fast forward took %.2f seconds." ), FastForwardTotalSeconds );
 }
 
 void UDemoNetDriver::SpawnDemoRecSpectator( UNetConnection* Connection )
@@ -1726,6 +1761,8 @@ void UDemoNetDriver::LoadCheckpoint( FArchive* GotoCheckpointArchive, int64 Goto
 	DECLARE_SCOPE_CYCLE_COUNTER(TEXT("LoadCheckpoint time"), STAT_ReplayCheckpointLoadTime, STATGROUP_Net);
 
 	check( GotoCheckpointArchive != NULL );
+	check( !bIsFastForwardingForCheckpoint );
+	check( !bIsFastForwarding );
 
 	// Reset the never-queue GUID list, we'll rebuild it
 	NonQueuedGUIDsForScrubbing.Empty();
@@ -1860,15 +1897,7 @@ void UDemoNetDriver::LoadCheckpoint( FArchive* GotoCheckpointArchive, int64 Goto
 
 		if ( GotoCheckpointSkipExtraTimeInMS != -1 )
 		{
-			DemoCurrentTime += (float)GotoCheckpointSkipExtraTimeInMS / 1000;
-
-			bIsFastForwarding				= true;
-			bIsFastForwardingForCheckpoint	= true;
-		}
-		else
-		{
-			bIsFastForwardingForCheckpoint	= false;
-			bIsFastForwarding				= false;
+			SkipTimeInternal( ( float )GotoCheckpointSkipExtraTimeInMS / 1000.0f, true, true );
 		}
 
 		return;
@@ -1910,23 +1939,12 @@ void UDemoNetDriver::LoadCheckpoint( FArchive* GotoCheckpointArchive, int64 Goto
 
 	if ( GotoCheckpointSkipExtraTimeInMS != -1 )
 	{
-		bIsFastForwardingForCheckpoint	= true;
-		bIsFastForwarding				= true;
-	}
-	else
-	{
-		bIsFastForwardingForCheckpoint	= false;
-		bIsFastForwarding				= false;
+		// If we need to skip more time for fine scrubbing, set that up now
+		SkipTimeInternal( ( float )GotoCheckpointSkipExtraTimeInMS / 1000.0f, true, true );
 	}
 
 	ReadDemoFrame( GotoCheckpointArchive );
 
-	// If we need to skip more time for fine scrubbing, set that up now
-	if ( GotoCheckpointSkipExtraTimeInMS != -1 )
-	{
-		DemoCurrentTime += (float)GotoCheckpointSkipExtraTimeInMS / 1000;
-	}
-
 	bDemoPlaybackDone = false;
 
 	if ( SpectatorController && ViewTargetGUID.IsValid() )
