Integrated github pull request #189: Auto-mipped RenderTarget2D

diff --git a/Engine/Source/Editor/TextureEditor/Private/TextureEditorToolkit.cpp b/Engine/Source/Editor/TextureEditor/Private/TextureEditorToolkit.cpp
--- a/Engine/Source/Editor/TextureEditor/Private/TextureEditorToolkit.cpp
+++ b/Engine/Source/Editor/TextureEditor/Private/TextureEditorToolkit.cpp
@@ -418,6 +418,10 @@ void FTextureEditorToolkit::PopulateQuickInfo( )
 	{
 		NumMips = TextureCube->GetNumMips();
 	}
+	else if (Texture2DRT)
+	{
+		NumMips = Texture2DRT->GetNumMips();
+	}
 
 	NumMipsText->SetText(FText::Format(NSLOCTEXT("TextureEditor", "QuickInfo_NumMips", "Number of Mips: {0}"), FText::AsNumber(NumMips)));
 
diff --git a/Engine/Source/Runtime/Engine/Classes/Engine/TextureRenderTarget2D.h b/Engine/Source/Runtime/Engine/Classes/Engine/TextureRenderTarget2D.h
--- a/Engine/Source/Runtime/Engine/Classes/Engine/TextureRenderTarget2D.h
+++ b/Engine/Source/Runtime/Engine/Classes/Engine/TextureRenderTarget2D.h
@@ -45,6 +45,10 @@ class UTextureRenderTarget2D : public UTextureRenderTarget
 	UPROPERTY(EditAnywhere, BlueprintReadOnly, Category=TextureRenderTarget2D, AssetRegistrySearchable)
 	uint32 bHDR:1;
 
+	/** Whether to support Mip maps for this render target texture */
+	UPROPERTY(EditAnywhere, BlueprintReadOnly, Category=TextureRenderTarget2D, AssetRegistrySearchable)
+	uint32 bAutoGenerateMips:1;
+
 	/** Normally the format is derived from bHDR, this allows code to set the format explicitly. */
 	UPROPERTY()
 	TEnumAsByte<enum EPixelFormat> OverrideFormat;
@@ -98,7 +102,7 @@ class UTextureRenderTarget2D : public UTextureRenderTarget
 
 	FORCEINLINE int32 GetNumMips() const
 	{
-		return 1;
+		return NumMips;
 	}
 
 	FORCEINLINE EPixelFormat GetFormat() const
@@ -112,6 +116,9 @@ class UTextureRenderTarget2D : public UTextureRenderTarget
 			return OverrideFormat;
 		}
 	}
+
+private:
+	int32	NumMips;
 };
 
 
diff --git a/Engine/Source/Runtime/Engine/Private/TextureRenderTarget2D.cpp b/Engine/Source/Runtime/Engine/Private/TextureRenderTarget2D.cpp
--- a/Engine/Source/Runtime/Engine/Private/TextureRenderTarget2D.cpp
+++ b/Engine/Source/Runtime/Engine/Private/TextureRenderTarget2D.cpp
@@ -15,6 +15,8 @@ UTextureRenderTarget2D::UTextureRenderTarget2D(const FObjectInitializer& ObjectI
 	: Super(ObjectInitializer)
 {
 	bHDR = true;
+	bAutoGenerateMips = false;
+	NumMips = 0;
 	ClearColor = FLinearColor(0.0f, 1.0f, 0.0f, 1.0f);
 	OverrideFormat = PF_Unknown;
 	bForceLinearGamma = true;
@@ -33,6 +35,16 @@ FTextureResource* UTextureRenderTarget2D::CreateResource()
 			OverrideFormat = PF_B8G8R8A8;
 		}
 	}
+
+	if (bAutoGenerateMips)
+	{
+		NumMips = FGenericPlatformMath::CeilToInt(FGenericPlatformMath::Log2(FGenericPlatformMath::Max(SizeX, SizeY)));
+	}
+	else
+	{
+		NumMips = 1;
+	}
+ 
 	FTextureRenderTarget2DResource* Result = new FTextureRenderTarget2DResource(this);
 	return Result;
 }
@@ -317,11 +329,16 @@ void FTextureRenderTarget2DResource::InitDynamicRHI()
 		uint32 TexCreateFlags = bSRGB ? TexCreate_SRGB : 0;
 		FRHIResourceCreateInfo CreateInfo = FRHIResourceCreateInfo(FClearValueBinding(ClearColor));
 
+		if (Owner->bAutoGenerateMips)
+		{
+			TexCreateFlags |= TexCreate_GenerateMipCapable;
+		}
+
 		RHICreateTargetableShaderResource2D(
 			Owner->SizeX, 
 			Owner->SizeY, 
 			Format,
-			1,
+			Owner->GetNumMips(),
 			TexCreateFlags,
 			TexCreate_RenderTargetable,
 			Owner->bNeedsTwoCopies,
@@ -381,6 +398,11 @@ void FTextureRenderTarget2DResource::UpdateDeferredResource( FRHICommandListImme
 		SetRenderTarget(RHICmdList, RenderTargetTextureRHI, FTextureRHIRef(), ESimpleRenderTargetMode::EClearColorExistingDepth);
 	}
  
+	if (Owner->bAutoGenerateMips)
+	{
+		RHICmdList.GenerateMips(RenderTargetTextureRHI);
+	}
+
  	// copy surface to the texture for use
 	RHICmdList.CopyToResolveTarget(RenderTargetTextureRHI, TextureRHI, true, FResolveParams());
 }
