Fix so that mirrored collision volumes are drawn in Player Collision view mode.

diff --git a/Engine/Source/Runtime/Engine/Classes/PhysicsEngine/ConvexElem.h b/Engine/Source/Runtime/Engine/Classes/PhysicsEngine/ConvexElem.h
--- a/Engine/Source/Runtime/Engine/Classes/PhysicsEngine/ConvexElem.h
+++ b/Engine/Source/Runtime/Engine/Classes/PhysicsEngine/ConvexElem.h
@@ -49,7 +49,7 @@ struct FKConvexElem : public FKShapeElem
 	
 	ENGINE_API void	DrawElemWire(class FPrimitiveDrawInterface* PDI, const FTransform& ElemTM, const float Scale, const FColor Color) const;
 
-	void AddCachedSolidConvexGeom(TArray<FDynamicMeshVertex>& VertexBuffer, TArray<int32>& IndexBuffer, const float Scale, const FColor VertexColor) const;
+	void AddCachedSolidConvexGeom(TArray<FDynamicMeshVertex>& VertexBuffer, TArray<int32>& IndexBuffer, const float Scale, const FColor VertexColor, const bool bIsMirrored) const;
 
 	/** Reset the hull to empty all arrays */
 	ENGINE_API void	Reset();
diff --git a/Engine/Source/Runtime/Engine/Private/PhysicsEngine/PhysDrawing.cpp b/Engine/Source/Runtime/Engine/Private/PhysicsEngine/PhysDrawing.cpp
--- a/Engine/Source/Runtime/Engine/Private/PhysicsEngine/PhysDrawing.cpp
+++ b/Engine/Source/Runtime/Engine/Private/PhysicsEngine/PhysDrawing.cpp
@@ -425,22 +425,23 @@ void FKConvexElem::DrawElemWire(FPrimitiveDrawInterface* PDI, const FTransform&
 #endif // WITH_PHYSX
 }
 
-void FKConvexElem::AddCachedSolidConvexGeom(TArray<FDynamicMeshVertex>& VertexBuffer, TArray<int32>& IndexBuffer, const float Scale, const FColor VertexColor) const
+void FKConvexElem::AddCachedSolidConvexGeom(TArray<FDynamicMeshVertex>& VertexBuffer, TArray<int32>& IndexBuffer, const float Scale, const FColor VertexColor, const bool bIsMirrored) const
 {
 #if WITH_PHYSX
-	if(ConvexMesh)
+	const PxConvexMesh* ConvexMeshToUse = bIsMirrored ? ConvexMeshNegX : ConvexMesh;
+	if(ConvexMeshToUse)
 	{
 		int32 StartVertOffset = VertexBuffer.Num();
 
 		// get PhysX data
-		const PxVec3* PVertices = ConvexMesh->getVertices();
-		const PxU8* PIndexBuffer = ConvexMesh->getIndexBuffer();
-		PxU32 NbPolygons = ConvexMesh->getNbPolygons();
+		const PxVec3* PVertices = ConvexMeshToUse->getVertices();
+		const PxU8* PIndexBuffer = ConvexMeshToUse->getIndexBuffer();
+		PxU32 NbPolygons = ConvexMeshToUse->getNbPolygons();
 
 		for(PxU32 i=0;i<NbPolygons;i++)
 		{
 			PxHullPolygon Data;
-			bool bStatus = ConvexMesh->getPolygonData(i, Data);
+			bool bStatus = ConvexMeshToUse->getPolygonData(i, Data);
 			check(bStatus);
 
 			const PxU8* indices = PIndexBuffer + Data.mIndexBase;
@@ -601,10 +602,11 @@ void FKAggregateGeom::GetAggGeom(const FTransform& Transform, const FColor Color
 				ThisGeom.RenderInfo->VertexBuffer = new FConvexCollisionVertexBuffer();
 				ThisGeom.RenderInfo->IndexBuffer = new FConvexCollisionIndexBuffer();
 
+				const bool bIsMirrored = (Scale3D.X * Scale3D.Y * Scale3D.Z < 0.0f);
 				for(int32 i=0; i<ConvexElems.Num(); i++)
 				{
 					// Get vertices/triangles from this hull.
-					ConvexElems[i].AddCachedSolidConvexGeom(ThisGeom.RenderInfo->VertexBuffer->Vertices, ThisGeom.RenderInfo->IndexBuffer->Indices, 1.0f, FColor::White);
+					ConvexElems[i].AddCachedSolidConvexGeom(ThisGeom.RenderInfo->VertexBuffer->Vertices, ThisGeom.RenderInfo->IndexBuffer->Indices, 1.0f, FColor::White, bIsMirrored);
 				}
 
 				// Only continue if we actually got some valid geometry
