Fixed crash on client level change when some replicated Actors have been torn off

Merging 2671248

//depot/UE4-UT/...

to //depot/UE4/...

diff --git a/Engine/Source/Runtime/Engine/Private/DataChannel.cpp b/Engine/Source/Runtime/Engine/Private/DataChannel.cpp
--- a/Engine/Source/Runtime/Engine/Private/DataChannel.cpp
+++ b/Engine/Source/Runtime/Engine/Private/DataChannel.cpp
@@ -1515,10 +1515,16 @@ bool UActorChannel::CleanUp( const bool bForDestroy )
 		{
 			if (Actor->bTearOff && !Connection->Driver->ShouldClientDestroyTearOffActors())
 			{
-				Actor->Role = ROLE_Authority;
-				Actor->SetReplicates(false);
-				bTornOff = true;
-				Actor->TornOff();
+				if (!bTornOff)
+				{
+					Actor->Role = ROLE_Authority;
+					Actor->SetReplicates(false);
+					bTornOff = true;
+					if (Actor->GetWorld() != NULL && !GIsRequestingExit)
+					{
+						Actor->TornOff();
+					}
+				}
 			}
 			else if (Dormant && !Actor->bTearOff)
 			{
