cleanup

diff --git a/Engine/Source/Runtime/Engine/Classes/Materials/MaterialInstance.h b/Engine/Source/Runtime/Engine/Classes/Materials/MaterialInstance.h
--- a/Engine/Source/Runtime/Engine/Classes/Materials/MaterialInstance.h
+++ b/Engine/Source/Runtime/Engine/Classes/Materials/MaterialInstance.h
@@ -422,8 +422,9 @@ protected:
 	/** 
 	 * Copies over material instance parameters from the base material given a material interface.
 	 * This is a slow operation that is needed for the editor.
+	 * @param Source silently ignores the case if 0
 	 */
-	ENGINE_API void CopyMaterialInstanceParameters(UMaterialInterface* MaterialInterface);
+	ENGINE_API void CopyMaterialInstanceParameters(UMaterialInterface* Source);
 
 	// to share code between PostLoad() and PostEditChangeProperty()
 	void PropagateDataToMaterialProxy();
diff --git a/Engine/Source/Runtime/Engine/Classes/Materials/MaterialInstanceDynamic.h b/Engine/Source/Runtime/Engine/Classes/Materials/MaterialInstanceDynamic.h
--- a/Engine/Source/Runtime/Engine/Classes/Materials/MaterialInstanceDynamic.h
+++ b/Engine/Source/Runtime/Engine/Classes/Materials/MaterialInstanceDynamic.h
@@ -37,17 +37,21 @@ class UMaterialInstanceDynamic : public UMaterialInstance
 	/**
 	 * Interpolates the scalar and vector parameters of this material instance based on two other material instances, and an alpha blending factor
 	 * The output is the object itself (this).
-	 * Supports the case MaterialInstanceA==this || MaterialInstanceB==this 
+	 * Supports the case SourceA==this || SourceB==this
+	 * Both material have to be from the same base material
+	 * @param SourceA value that is used for Alpha=0, silently ignores the case if 0
+	 * @param SourceB value that is used for Alpha=1, silently ignores the case if 0
+	 * @param Alpha usually in the range 0..1, values outside the range extrapolate
 	 */
 	UFUNCTION(BlueprintCallable, meta=(DisplayName = "InterpolateMaterialInstanceParameters"), Category="Rendering|Material")
-	void K2_InterpolateMaterialInstanceParams(UMaterialInstance* MaterialInstanceA, UMaterialInstance* MaterialInstanceB, float Alpha);
+	void K2_InterpolateMaterialInstanceParams(UMaterialInstance* SourceA, UMaterialInstance* SourceB, float Alpha);
 
 	/**
 	 * Copies over parameters given a material interface.
 	 * The output is the object itself (this).
 	 */
 	UFUNCTION(BlueprintCallable, meta=(DisplayName = "CopyMaterialInstanceParameters"), Category="Rendering|Material")
-	void K2_CopyMaterialInstanceParameters(UMaterialInterface* SourceMaterialToCopyFrom);
+	void K2_CopyMaterialInstanceParameters(UMaterialInterface* Source);
 
 	/**
 	 * Create a material instance dynamic parented to the specified material.
diff --git a/Engine/Source/Runtime/Engine/Private/Materials/MaterialInstance.cpp b/Engine/Source/Runtime/Engine/Private/Materials/MaterialInstance.cpp
--- a/Engine/Source/Runtime/Engine/Private/Materials/MaterialInstance.cpp
+++ b/Engine/Source/Runtime/Engine/Private/Materials/MaterialInstance.cpp
@@ -992,9 +992,9 @@ bool UMaterialInstance::IsDependent(UMaterialInterface* TestDependency)
 	}
 }
 
-void UMaterialInstance::CopyMaterialInstanceParameters(UMaterialInterface* MaterialInterface)
+void UMaterialInstance::CopyMaterialInstanceParameters(UMaterialInterface* Source)
 {
-	if(MaterialInterface)
+	if(Source)
 	{
 		UMaterial* Material = GetMaterial();
 
@@ -1012,7 +1012,7 @@ void UMaterialInstance::CopyMaterialInstanceParameters(UMaterialInterface* Mater
 			FName ParameterName = Names[i];
 			UFont* FontValue = NULL;
 			int32 FontPage;
-			if(MaterialInterface->GetFontParameterValue(ParameterName, FontValue, FontPage))
+			if(Source->GetFontParameterValue(ParameterName, FontValue, FontPage))
 			{
 				FFontParameterValue* ParameterValue = new(FontParameterValues) FFontParameterValue;
 				ParameterValue->ParameterName = ParameterName;
@@ -1031,7 +1031,7 @@ void UMaterialInstance::CopyMaterialInstanceParameters(UMaterialInterface* Mater
 		{
 			FName ParameterName = Names[i];
 			float ScalarValue = 1.0f;
-			if(MaterialInterface->GetScalarParameterValue(ParameterName, ScalarValue))
+			if(Source->GetScalarParameterValue(ParameterName, ScalarValue))
 			{
 				FScalarParameterValue* ParameterValue = new(ScalarParameterValues) FScalarParameterValue;
 				ParameterValue->ParameterName = ParameterName;
@@ -1048,7 +1048,7 @@ void UMaterialInstance::CopyMaterialInstanceParameters(UMaterialInterface* Mater
 		{
 			FName ParameterName = Names[i];
 			FLinearColor VectorValue;
-			if(MaterialInterface->GetVectorParameterValue(ParameterName, VectorValue))
+			if(Source->GetVectorParameterValue(ParameterName, VectorValue))
 			{
 				FVectorParameterValue* ParameterValue = new(VectorParameterValues) FVectorParameterValue;
 				ParameterValue->ParameterName = ParameterName;
@@ -1065,7 +1065,7 @@ void UMaterialInstance::CopyMaterialInstanceParameters(UMaterialInterface* Mater
 		{
 			FName ParameterName = Names[i];
 			UTexture* TextureValue = NULL;
-			if(MaterialInterface->GetTextureParameterValue(ParameterName, TextureValue))
+			if(Source->GetTextureParameterValue(ParameterName, TextureValue))
 			{
 				FTextureParameterValue* ParameterValue = new(TextureParameterValues) FTextureParameterValue;
 				ParameterValue->ParameterName = ParameterName;
diff --git a/Engine/Source/Runtime/Engine/Private/Materials/MaterialInstanceDynamic.cpp b/Engine/Source/Runtime/Engine/Private/Materials/MaterialInstanceDynamic.cpp
--- a/Engine/Source/Runtime/Engine/Private/Materials/MaterialInstanceDynamic.cpp
+++ b/Engine/Source/Runtime/Engine/Private/Materials/MaterialInstanceDynamic.cpp
@@ -133,28 +133,28 @@ FVectorParameterValue* GameThread_GetVectorParameterValue(UMaterialInstance* Mat
 	return 0;
 }
 
-void UMaterialInstanceDynamic::K2_InterpolateMaterialInstanceParams(UMaterialInstance* MaterialInstanceA, UMaterialInstance* MaterialInstanceB, float Alpha)
+void UMaterialInstanceDynamic::K2_InterpolateMaterialInstanceParams(UMaterialInstance* SourceA, UMaterialInstance* SourceB, float Alpha)
 {
-	if(MaterialInstanceA && MaterialInstanceB)
+	if(SourceA && SourceB)
 	{
-		UMaterial* BaseMaterialInstanceA = MaterialInstanceA->GetBaseMaterial();
-		UMaterial* BaseMaterialInstanceB = MaterialInstanceB->GetBaseMaterial();
+		UMaterial* BaseA = SourceA->GetBaseMaterial();
+		UMaterial* BaseB = SourceB->GetBaseMaterial();
 
-		if(BaseMaterialInstanceA == BaseMaterialInstanceB)
+		if(BaseA == BaseB)
 		{
 			// todo: can be optimized, at least we can reserve
 			TArray<FName> Names;
 
-			GameThread_FindAllScalarParameterNames(MaterialInstanceA, Names);
-			GameThread_FindAllScalarParameterNames(MaterialInstanceB, Names);
+			GameThread_FindAllScalarParameterNames(SourceA, Names);
+			GameThread_FindAllScalarParameterNames(SourceB, Names);
 
 			// Interpolate the scalar parameters common to both materials
 			for(int32 Idx = 0, Count = Names.Num(); Idx < Count; ++Idx)
 			{
 				FName Name = Names[Idx];
 
-				auto ParamValueA = GameThread_GetScalarParameterValue(MaterialInstanceA, Name);
-				auto ParamValueB = GameThread_GetScalarParameterValue(MaterialInstanceB, Name);
+				auto ParamValueA = GameThread_GetScalarParameterValue(SourceA, Name);
+				auto ParamValueB = GameThread_GetScalarParameterValue(SourceB, Name);
 
 				if(ParamValueA || ParamValueB)
 				{
@@ -162,7 +162,7 @@ void UMaterialInstanceDynamic::K2_InterpolateMaterialInstanceParams(UMaterialIns
 
 					if(!ParamValueA || !ParamValueB)
 					{
-						BaseMaterialInstanceA->GetScalarParameterValue(Name, Default);
+						BaseA->GetScalarParameterValue(Name, Default);
 					}
 
 					auto ValueA = ParamValueA ? ParamValueA->ParameterValue : Default;
@@ -174,16 +174,16 @@ void UMaterialInstanceDynamic::K2_InterpolateMaterialInstanceParams(UMaterialIns
 
 			// reused array to minimize further allocations
 			Names.Empty();
-			GameThread_FindAllVectorParameterNames(MaterialInstanceA, Names);
-			GameThread_FindAllVectorParameterNames(MaterialInstanceB, Names);
+			GameThread_FindAllVectorParameterNames(SourceA, Names);
+			GameThread_FindAllVectorParameterNames(SourceB, Names);
 
 			// Interpolate the vector parameters common to both
 			for(int32 Idx = 0, Count = Names.Num(); Idx < Count; ++Idx)
 			{
 				FName Name = Names[Idx];
 
-				auto ParamValueA = GameThread_GetVectorParameterValue(MaterialInstanceA, Name);
-				auto ParamValueB = GameThread_GetVectorParameterValue(MaterialInstanceB, Name);
+				auto ParamValueA = GameThread_GetVectorParameterValue(SourceA, Name);
+				auto ParamValueB = GameThread_GetVectorParameterValue(SourceB, Name);
 
 				if(ParamValueA || ParamValueB)
 				{
@@ -191,7 +191,7 @@ void UMaterialInstanceDynamic::K2_InterpolateMaterialInstanceParams(UMaterialIns
 
 					if(!ParamValueA || !ParamValueB)
 					{
-						BaseMaterialInstanceA->GetVectorParameterValue(Name, Default);
+						BaseA->GetVectorParameterValue(Name, Default);
 					}
 
 					auto ValueA = ParamValueA ? ParamValueA->ParameterValue : Default;
@@ -204,14 +204,15 @@ void UMaterialInstanceDynamic::K2_InterpolateMaterialInstanceParams(UMaterialIns
 		else
 		{
 			// to find bad usage of this method
-			ensure(BaseMaterialInstanceA == BaseMaterialInstanceB);
+			// Maybe we can log a content error instead
+			// ensure(BaseA == BaseB);
 		}
 	}
 }
 
-void UMaterialInstanceDynamic::K2_CopyMaterialInstanceParameters(UMaterialInterface* SourceMaterialToCopyFrom)
+void UMaterialInstanceDynamic::K2_CopyMaterialInstanceParameters(UMaterialInterface* Source)
 {
-	CopyMaterialInstanceParameters(SourceMaterialToCopyFrom);
+	CopyMaterialInstanceParameters(Source);
 }
 
 void UMaterialInstanceDynamic::CopyParameterOverrides(UMaterialInstance* MaterialInstance)
