Engine: Exposed GetVectorParameterByName and GetScalarParameterByName on UMaterialParameterCollection as dll exported methods
[INTEGRATE] Change 2671936 by [REDACTED_PERSON][REDACTED_HANDLE] on 2015/08/27 22:24:40

diff --git a/Engine/Source/Runtime/Engine/Classes/Materials/MaterialParameterCollection.h b/Engine/Source/Runtime/Engine/Classes/Materials/MaterialParameterCollection.h
--- a/Engine/Source/Runtime/Engine/Classes/Materials/MaterialParameterCollection.h
+++ b/Engine/Source/Runtime/Engine/Classes/Materials/MaterialParameterCollection.h
@@ -88,10 +88,10 @@ class UMaterialParameterCollection : public UObject
 	ENGINE_API void GetParameterNames(TArray<FName>& OutParameterNames, bool bVectorParameters) const;
 
 	/** Utility to find a scalar parameter struct given a parameter name.  Returns NULL if not found. */
-	const FCollectionScalarParameter* GetScalarParameterByName(FName ParameterName) const;
+	ENGINE_API const FCollectionScalarParameter* GetScalarParameterByName(FName ParameterName) const;
 
 	/** Utility to find a vector parameter struct given a parameter name.  Returns NULL if not found. */
-	const FCollectionVectorParameter* GetVectorParameterByName(FName ParameterName) const;
+	ENGINE_API const FCollectionVectorParameter* GetVectorParameterByName(FName ParameterName) const;
 
 	/** Accessor for the uniform buffer layout description. */
 	const FUniformBufferStruct& GetUniformBufferStruct() const
