#ue4 - Fix broken overlap check during swept movement. Other component was always being centered at the query location with the query rotation.

diff --git a/Engine/Source/Runtime/Engine/Private/Components/PrimitiveComponent.cpp b/Engine/Source/Runtime/Engine/Private/Components/PrimitiveComponent.cpp
--- a/Engine/Source/Runtime/Engine/Private/Components/PrimitiveComponent.cpp
+++ b/Engine/Source/Runtime/Engine/Private/Components/PrimitiveComponent.cpp
@@ -2183,7 +2183,7 @@ const TArray<FOverlapInfo>* UPrimitiveComponent::ConvertSweptOverlapsToCurrentOv
 							// Not handled yet. We could do it by checking every body explicitly and track each body index in the overlap test, but this seems like a rare need.
 							return nullptr;
 						}
-						else if (ComponentOverlapComponent(OtherPrimitive, EndLocation, EndRotationQuat, UnusedQueryParams))
+						else if (OtherPrimitive->ComponentOverlapComponent(this, EndLocation, EndRotationQuat, UnusedQueryParams))
 						{
 							OverlapsAtEndLocation.Add(OtherOverlap);
 						}
