Fix for marker based sync jumping issue when loops

Merging

//depot/UE4-Orion/Engine/Source/Runtime/Engine/Private/Animation/AnimSequence.cpp

to //depot/UE4-Orion/Engine/Source/Runtime/Engine/Private/Animation/AnimSequence.cpp

diff --git a/Engine/Source/Runtime/Engine/Private/Animation/AnimSequence.cpp b/Engine/Source/Runtime/Engine/Private/Animation/AnimSequence.cpp
--- a/Engine/Source/Runtime/Engine/Private/Animation/AnimSequence.cpp
+++ b/Engine/Source/Runtime/Engine/Private/Animation/AnimSequence.cpp
@@ -4070,6 +4070,7 @@ void UAnimSequence::GetMarkerIndicesForPosition(const FMarkerSyncAnimPosition& S
 	}
 
 	float DiffToCurrentTime = FLT_MAX;
+	const float CurrentInputTime  = OutCurrentTime;
 
 	for (int32 PrevMarkerIdx = 0; PrevMarkerIdx < AuthoredSyncMarkers.Num(); ++PrevMarkerIdx)
 	{
@@ -4092,15 +4093,17 @@ void UAnimSequence::GetMarkerIndicesForPosition(const FMarkerSyncAnimPosition& S
 					{
 						ThisCurrentTime -= SequenceLength;
 					}
-					float ThisDiff = FMath::Abs(ThisCurrentTime - OutCurrentTime);
+					float ThisDiff = FMath::Abs(ThisCurrentTime - CurrentInputTime);
 					if (ThisDiff < DiffToCurrentTime)
 					{
 						DiffToCurrentTime = ThisDiff;
 						OutPrevMarker.MarkerIndex = PrevMarkerIdx;
 						OutNextMarker.MarkerIndex = NextMarkerIdx;
 						OutCurrentTime = GetCurrentTimeFromMarkers(OutPrevMarker, OutNextMarker, SyncPosition.PositionBetweenMarkers);
-						break;
 					}
+
+					// this marker test is done, move onto next one
+					break;
 				}
 			}
 		}
