added preferred navigation data class to nav agent properties, support for multiple active navigation datas
#ue4

diff --git a/Engine/Source/Runtime/Engine/Classes/AI/Navigation/NavigationTypes.h b/Engine/Source/Runtime/Engine/Classes/AI/Navigation/NavigationTypes.h
--- a/Engine/Source/Runtime/Engine/Classes/AI/Navigation/NavigationTypes.h
+++ b/Engine/Source/Runtime/Engine/Classes/AI/Navigation/NavigationTypes.h
@@ -458,8 +458,12 @@ struct ENGINE_API FNavAgentProperties : public FMovementProperties
 	UPROPERTY(EditAnywhere, BlueprintReadWrite, Category=MovementProperties)
 	float NavWalkingSearchHeightScale;
 
+	/** Type of navigation data used by agent, null means "any" */
+	UPROPERTY(EditAnywhere, BlueprintReadWrite, Category = MovementProperties)
+	TSubclassOf<ANavigationData> PreferredNavData;
+
 	FNavAgentProperties(float Radius = -1.f, float Height = -1.f)
-		: AgentRadius(Radius), AgentHeight(Height), AgentStepHeight(-1), NavWalkingSearchHeightScale(0.5f)
+		: AgentRadius(Radius), AgentHeight(Height), AgentStepHeight(-1), NavWalkingSearchHeightScale(0.5f), PreferredNavData(nullptr)
 	{
 	}
 
@@ -472,7 +476,13 @@ struct ENGINE_API FNavAgentProperties : public FMovementProperties
 	{
 		return FGenericPlatformMath::Abs(AgentRadius - Other.AgentRadius) < Precision &&
 			FGenericPlatformMath::Abs(AgentHeight - Other.AgentHeight) < Precision &&
-			FGenericPlatformMath::Abs(AgentStepHeight - Other.AgentStepHeight) < Precision;
+			FGenericPlatformMath::Abs(AgentStepHeight - Other.AgentStepHeight) < Precision &&
+			IsNavDataMatching(Other);
+	}
+
+	bool IsNavDataMatching(const FNavAgentProperties& Other) const
+	{
+		return (PreferredNavData == nullptr || Other.PreferredNavData == nullptr || PreferredNavData == Other.PreferredNavData);
 	}
 
 	bool operator==(const FNavAgentProperties& Other) const
diff --git a/Engine/Source/Runtime/Engine/Private/AI/Navigation/NavigationSystem.cpp b/Engine/Source/Runtime/Engine/Private/AI/Navigation/NavigationSystem.cpp
--- a/Engine/Source/Runtime/Engine/Private/AI/Navigation/NavigationSystem.cpp
+++ b/Engine/Source/Runtime/Engine/Private/AI/Navigation/NavigationSystem.cpp
@@ -380,6 +380,13 @@ void UNavigationSystem::SetSupportedAgentsNavigationClass(int32 AgentIndex, TSub
 	check(SupportedAgents.IsValidIndex(AgentIndex));
 	SupportedAgents[AgentIndex].NavigationDataClass = NavigationDataClass;
 
+	// keep preferred navigation data class in sync with actual class
+	// this will be passed to navigation data actor and will be required
+	// for comparisons done in DoesSupportAgent calls
+	//
+	// "Any" navigation data preference is valid only for instanced agents
+	SupportedAgents[AgentIndex].PreferredNavData = NavigationDataClass;
+
 	if (NavigationDataClass != nullptr)
 	{
 		SupportedAgents[AgentIndex].NavigationDataClassName = FStringClassReference::GetOrCreateIDForClass(NavigationDataClass);
@@ -1450,6 +1457,12 @@ const ANavigationData* UNavigationSystem::GetNavDataForProps(const FNavAgentProp
 		for (TArray<FNavAgentProperties>::TConstIterator It(AgentPropertiesList); It; ++It)
 		{
 			const FNavAgentProperties& NavIt = *It;
+			const bool bNavClassMatch = NavIt.IsNavDataMatching(AgentProperties);
+			if (!bNavClassMatch)
+			{
+				continue;
+			}
+
 			ExcessRadius = NavIt.AgentRadius - AgentProperties.AgentRadius;
 			ExcessHeight = NavIt.AgentHeight - AgentHeight;
 
