compilation fix
#ue4

diff --git a/Engine/Source/Runtime/Engine/Classes/AI/Navigation/NavigationTypes.h b/Engine/Source/Runtime/Engine/Classes/AI/Navigation/NavigationTypes.h
--- a/Engine/Source/Runtime/Engine/Classes/AI/Navigation/NavigationTypes.h
+++ b/Engine/Source/Runtime/Engine/Classes/AI/Navigation/NavigationTypes.h
@@ -463,7 +463,7 @@ struct ENGINE_API FNavAgentProperties : public FMovementProperties
 	TSubclassOf<ANavigationData> PreferredNavData;
 
 	FNavAgentProperties(float Radius = -1.f, float Height = -1.f)
-		: AgentRadius(Radius), AgentHeight(Height), AgentStepHeight(-1), NavWalkingSearchHeightScale(0.5f), PreferredNavData(nullptr)
+		: AgentRadius(Radius), AgentHeight(Height), AgentStepHeight(-1), NavWalkingSearchHeightScale(0.5f)
 	{
 	}
 
@@ -480,10 +480,7 @@ struct ENGINE_API FNavAgentProperties : public FMovementProperties
 			IsNavDataMatching(Other);
 	}
 
-	bool IsNavDataMatching(const FNavAgentProperties& Other) const
-	{
-		return (PreferredNavData == nullptr || Other.PreferredNavData == nullptr || PreferredNavData == Other.PreferredNavData);
-	}
+	bool IsNavDataMatching(const FNavAgentProperties& Other) const;
 
 	bool operator==(const FNavAgentProperties& Other) const
 	{
diff --git a/Engine/Source/Runtime/Engine/Private/Components/NavMovementComponent.cpp b/Engine/Source/Runtime/Engine/Private/Components/NavMovementComponent.cpp
--- a/Engine/Source/Runtime/Engine/Private/Components/NavMovementComponent.cpp
+++ b/Engine/Source/Runtime/Engine/Private/Components/NavMovementComponent.cpp
@@ -16,6 +16,11 @@ void FNavAgentProperties::UpdateWithCollisionComponent(UShapeComponent* Collisio
 	AgentRadius = CollisionComponent->Bounds.SphereRadius;
 }
 
+bool FNavAgentProperties::IsNavDataMatching(const FNavAgentProperties& Other) const
+{
+	return (PreferredNavData == nullptr || Other.PreferredNavData == nullptr || PreferredNavData == Other.PreferredNavData);
+}
+
 //----------------------------------------------------------------------//
 // UMovementComponent
 //----------------------------------------------------------------------//
