#ENGINE - Integrating automation test flags refactor
- New filter flags in automationtest.h
- Compile time errors if proper flags are not specified
- Existing tests converted to new flags
- New "disabled" flag
- Filter dropdown in automation UI replaces Smoke Tests button
- Automation features are no longer compiled out in Test configuration at runtime (except in Shipping)
- Enabled automation tests to be triggered from in-game console in non-editor builds (automation controller is enabled)

NOTE: This will force licensees who use automation to update their tests.

diff --git a/Engine/Plugins/Experimental/UnitTestFramework/Source/UnitTestFramework/Private/UTFGameAutomationTests.cpp b/Engine/Plugins/Experimental/UnitTestFramework/Source/UnitTestFramework/Private/UTFGameAutomationTests.cpp
--- a/Engine/Plugins/Experimental/UnitTestFramework/Source/UnitTestFramework/Private/UTFGameAutomationTests.cpp
+++ b/Engine/Plugins/Experimental/UnitTestFramework/Source/UnitTestFramework/Private/UTFGameAutomationTests.cpp
@@ -61,7 +61,7 @@ namespace UTFGameAutomationTestsImpl
  * FUTFResetGameActorTest
  ******************************************************************************/
 
-IMPLEMENT_COMPLEX_AUTOMATION_TEST(FUTFResetGameActorTest, "UnitTestFramework.Game.ResetUnitTestActors", EAutomationTestFlags::ATF_Game | EAutomationTestFlags::ATF_Editor)
+IMPLEMENT_COMPLEX_AUTOMATION_TEST(FUTFResetGameActorTest, "UnitTestFramework.Game.ResetUnitTestActors", EAutomationTestFlags::ClientContext | EAutomationTestFlags::EditorContext | EAutomationTestFlags::StressFilter)
 
 //------------------------------------------------------------------------------
 void FUTFResetGameActorTest::GetTests(TArray<FString>& OutBeautifiedNames, TArray<FString>& OutTestCommands) const
@@ -94,7 +94,7 @@ bool FUTFResetGameActorTest::RunTest(const FString& ObjectPath)
  * FUTFGameActorTest
  ******************************************************************************/
 
-IMPLEMENT_COMPLEX_AUTOMATION_TEST(FUTFGameActorTest, "UnitTestFramework.Game.UnitTestActors", EAutomationTestFlags::ATF_Game | EAutomationTestFlags::ATF_Editor)
+IMPLEMENT_COMPLEX_AUTOMATION_TEST(FUTFGameActorTest, "UnitTestFramework.Game.UnitTestActors", EAutomationTestFlags::ClientContext | EAutomationTestFlags::EditorContext | EAutomationTestFlags::StressFilter)
 
 //------------------------------------------------------------------------------
 void FUTFGameActorTest::GetTests(TArray<FString>& OutBeautifiedNames, TArray<FString>& OutTestCommands) const
diff --git a/Engine/Plugins/Experimental/UnitTestFramework/Source/UnitTestFrameworkEditor/Private/UTFEditorAutomationTests.cpp b/Engine/Plugins/Experimental/UnitTestFramework/Source/UnitTestFrameworkEditor/Private/UTFEditorAutomationTests.cpp
--- a/Engine/Plugins/Experimental/UnitTestFramework/Source/UnitTestFrameworkEditor/Private/UTFEditorAutomationTests.cpp
+++ b/Engine/Plugins/Experimental/UnitTestFramework/Source/UnitTestFrameworkEditor/Private/UTFEditorAutomationTests.cpp
@@ -41,7 +41,7 @@ bool FUTFEditorAutomationTests::IsClassSuitableForInEditorTesting(UClass* Class)
  * FUTFEditorMapTest
  ******************************************************************************/
 
-IMPLEMENT_COMPLEX_AUTOMATION_TEST(FUTFEditorMapTest, "UnitTestFramework.Editor.RunUnitTestMaps", EAutomationTestFlags::ATF_Editor)
+IMPLEMENT_COMPLEX_AUTOMATION_TEST(FUTFEditorMapTest, "UnitTestFramework.Editor.RunUnitTestMaps", EAutomationTestFlags::EditorContext | EAutomationTestFlags::StressFilter)
 
 //------------------------------------------------------------------------------
 void FUTFEditorMapTest::GetTests(TArray<FString>& OutBeautifiedNames, TArray<FString>& OutTestCommands) const
@@ -153,7 +153,7 @@ bool FUTFEditorMapTest::RunTest(const FString& LevelAssetPath)
 * FUTFEditorActorTest
 ******************************************************************************/
 
-IMPLEMENT_COMPLEX_AUTOMATION_TEST(FUTFEditorActorTest, "UnitTestFramework.Editor.UnitTestActors", EAutomationTestFlags::ATF_Editor)
+IMPLEMENT_COMPLEX_AUTOMATION_TEST(FUTFEditorActorTest, "UnitTestFramework.Editor.UnitTestActors", EAutomationTestFlags::EditorContext | EAutomationTestFlags::StressFilter)
 
 //------------------------------------------------------------------------------
 void FUTFEditorActorTest::GetTests(TArray<FString>& OutBeautifiedNames, TArray<FString>& OutTestCommands) const
diff --git a/Engine/Plugins/Messaging/UdpMessaging/Source/UdpMessaging/Private/Transport/Tests/UdpMessageSegmenterTest.cpp b/Engine/Plugins/Messaging/UdpMessaging/Source/UdpMessaging/Private/Transport/Tests/UdpMessageSegmenterTest.cpp
--- a/Engine/Plugins/Messaging/UdpMessaging/Source/UdpMessaging/Private/Transport/Tests/UdpMessageSegmenterTest.cpp
+++ b/Engine/Plugins/Messaging/UdpMessaging/Source/UdpMessaging/Private/Transport/Tests/UdpMessageSegmenterTest.cpp
@@ -4,7 +4,7 @@
 #include "AutomationTest.h"
 
 
-IMPLEMENT_SIMPLE_AUTOMATION_TEST(FUdpMessageSegmenterTest, "System.Core.Messaging.Transports.Udp.UdpMessageSegmenter", EAutomationTestFlags::ATF_Editor)
+IMPLEMENT_SIMPLE_AUTOMATION_TEST(FUdpMessageSegmenterTest, "System.Core.Messaging.Transports.Udp.UdpMessageSegmenter", EAutomationTestFlags::EditorContext | EAutomationTestFlags::EngineFilter)
 
 
 bool FUdpMessageSegmenterTest::RunTest( const FString& Parameters )
diff --git a/Engine/Plugins/NetcodeUnitTest/NetcodeUnitTest/Source/NetcodeUnitTest/Private/Tests/AllUnitTestsTest.cpp b/Engine/Plugins/NetcodeUnitTest/NetcodeUnitTest/Source/NetcodeUnitTest/Private/Tests/AllUnitTestsTest.cpp
--- a/Engine/Plugins/NetcodeUnitTest/NetcodeUnitTest/Source/NetcodeUnitTest/Private/Tests/AllUnitTestsTest.cpp
+++ b/Engine/Plugins/NetcodeUnitTest/NetcodeUnitTest/Source/NetcodeUnitTest/Private/Tests/AllUnitTestsTest.cpp
@@ -50,7 +50,7 @@ bool FAutomationConsoleCommand::Update()
  * FAllUnitTestsTest
  */
 
-IMPLEMENT_COMPLEX_AUTOMATION_TEST(FAllUnitTestsTest, "System.Netcode Unit Test", EAutomationTestFlags::ATF_Game);
+IMPLEMENT_COMPLEX_AUTOMATION_TEST(FAllUnitTestsTest, "System.Netcode Unit Test", EAutomationTestFlags::ClientContext | EAutomationTestFlags::EngineFilter);
 
 void FAllUnitTestsTest::GetTests(TArray<FString>& OutBeautifiedNames, TArray<FString>& OutTestCommands) const
 {
diff --git a/Engine/Source/Developer/AITestSuite/Private/Tests/AITestsCommon.h b/Engine/Source/Developer/AITestSuite/Private/Tests/AITestsCommon.h
--- a/Engine/Source/Developer/AITestSuite/Private/Tests/AITestsCommon.h
+++ b/Engine/Source/Developer/AITestSuite/Private/Tests/AITestsCommon.h
@@ -86,7 +86,7 @@ DEFINE_LATENT_AUTOMATION_COMMAND_ONE_PARAMETER(FAITestCommand_TearDownTest, FAIT
 
 // @note that TestClass needs to derive from FAITestBase
 #define IMPLEMENT_AI_LATENT_TEST(TestClass, PrettyName) \
-	IMPLEMENT_SIMPLE_AUTOMATION_TEST(TestClass##_Runner, PrettyName, (EAutomationTestFlags::ATF_Game | EAutomationTestFlags::ATF_Editor)) \
+	IMPLEMENT_SIMPLE_AUTOMATION_TEST(TestClass##_Runner, PrettyName, (EAutomationTestFlags::ClientContext | EAutomationTestFlags::EditorContext | EAutomationTestFlags::EngineFilter)) \
 	bool TestClass##_Runner::RunTest(const FString& Parameters) \
 	{ \
 		/* spawn test instance. Setup should be done in test's constructor */ \
@@ -102,7 +102,7 @@ DEFINE_LATENT_AUTOMATION_COMMAND_ONE_PARAMETER(FAITestCommand_TearDownTest, FAIT
 	} 
 
 #define IMPLEMENT_AI_INSTANT_TEST(TestClass, PrettyName) \
-	IMPLEMENT_SIMPLE_AUTOMATION_TEST(TestClass##Runner, PrettyName, (EAutomationTestFlags::ATF_Game | EAutomationTestFlags::ATF_Editor)) \
+	IMPLEMENT_SIMPLE_AUTOMATION_TEST(TestClass##Runner, PrettyName, (EAutomationTestFlags::ClientContext | EAutomationTestFlags::EditorContext | EAutomationTestFlags::EngineFilter)) \
 	bool TestClass##Runner::RunTest(const FString& Parameters) \
 	{ \
 		/* spawn test instance. Setup should be done in test's constructor */ \
diff --git a/Engine/Source/Developer/AutomationController/Private/AutomationControllerManager.h b/Engine/Source/Developer/AutomationController/Private/AutomationControllerManager.h
--- a/Engine/Source/Developer/AutomationController/Private/AutomationControllerManager.h
+++ b/Engine/Source/Developer/AutomationController/Private/AutomationControllerManager.h
@@ -129,14 +129,10 @@ public:
 		return bDeveloperDirectoryIncluded;
 	}
 
-	virtual void SetVisualCommandletFilter(const bool bInVisualCommandletFilterOn) override
+	virtual void SetRequestedTestFlags(const uint32 InRequestedTestFlags) override
 	{
-		bVisualCommandletFilterOn = bInVisualCommandletFilterOn;
-	}
-
-	virtual bool IsVisualCommandletFilterOn(void) const override
-	{
-		return bVisualCommandletFilterOn;
+		RequestedTestFlags = InRequestedTestFlags;
+		RequestTests();
 	}
 
 	virtual const bool CheckTestResultsAvailable() const override
@@ -267,6 +263,9 @@ private:
 	/** Handles FAutomationWorkerRequestTestsReply messages. */
 	void HandleRequestTestsReplyMessage( const FAutomationWorkerRequestTestsReply& Message, const IMessageContextRef& Context );
 
+	/** Handles FAutomationWorkerRequestTestsReplyComplete messages. */
+	void HandleRequestTestsReplyCompleteMessage(const FAutomationWorkerRequestTestsReplyComplete& Message, const IMessageContextRef& Context);
+
 	/** Handles FAutomationWorkerRunTestsReply messages. */
 	void HandleRunTestsReplyMessage( const FAutomationWorkerRunTestsReply& Message, const IMessageContextRef& Context );
 
@@ -302,8 +301,8 @@ private:
 	/** Are tests results available */
 	bool bTestResultsAvailable;
 
-	/** Whether to use visual commandlet **/
-	bool bVisualCommandletFilterOn;
+	/** Which sets of tests to consider */
+	uint32 RequestedTestFlags;
 
 	/** Timer to keep track of the last time tests were updated */
 	double CheckTestTimer;
diff --git a/Engine/Source/Developer/AutomationController/Private/AutomationControllerManger.cpp b/Engine/Source/Developer/AutomationController/Private/AutomationControllerManger.cpp
--- a/Engine/Source/Developer/AutomationController/Private/AutomationControllerManger.cpp
+++ b/Engine/Source/Developer/AutomationController/Private/AutomationControllerManger.cpp
@@ -52,7 +52,7 @@ void FAutomationControllerManager::RequestTests()
 			ResetIntermediateTestData();
 
 			//issue tests on appropriate platforms
-			MessageEndpoint->Send(new FAutomationWorkerRequestTests(bDeveloperDirectoryIncluded, bVisualCommandletFilterOn), MessageAddress);
+			MessageEndpoint->Send(new FAutomationWorkerRequestTests(bDeveloperDirectoryIncluded, RequestedTestFlags), MessageAddress);
 		}
 	}
 }
@@ -294,6 +294,7 @@ void FAutomationControllerManager::Startup()
 		.Handling<FAutomationWorkerPong>(this, &FAutomationControllerManager::HandlePongMessage)
 		.Handling<FAutomationWorkerRequestNextNetworkCommand>(this, &FAutomationControllerManager::HandleRequestNextNetworkCommandMessage)
 		.Handling<FAutomationWorkerRequestTestsReply>(this, &FAutomationControllerManager::HandleRequestTestsReplyMessage)
+		.Handling<FAutomationWorkerRequestTestsReplyComplete>(this, &FAutomationControllerManager::HandleRequestTestsReplyCompleteMessage)
 		.Handling<FAutomationWorkerRunTestsReply>(this, &FAutomationControllerManager::HandleRunTestsReplyMessage)
 		.Handling<FAutomationWorkerScreenImage>(this, &FAutomationControllerManager::HandleReceivedScreenShot)
 		.Handling<FAutomationWorkerWorkerOffline>(this, &FAutomationControllerManager::HandleWorkerOfflineMessage);
@@ -306,7 +307,7 @@ void FAutomationControllerManager::Startup()
 	ClusterDistributionMask = 0;
 	ExecutionCount = 0;
 	bDeveloperDirectoryIncluded = false;
-	bVisualCommandletFilterOn = false;
+	RequestedTestFlags = EAutomationTestFlags::SmokeFilter | EAutomationTestFlags::EngineFilter | EAutomationTestFlags::ProductFilter | EAutomationTestFlags::PerfFilter;
 
 	NumOfTestsToReceive = 0;
 	NumTestPasses = 1;
@@ -628,17 +629,16 @@ void FAutomationControllerManager::HandleRequestNextNetworkCommandMessage( const
 }
 
 
-void FAutomationControllerManager::HandleRequestTestsReplyMessage( const FAutomationWorkerRequestTestsReply& Message, const IMessageContextRef& Context )
+void FAutomationControllerManager::HandleRequestTestsReplyMessage(const FAutomationWorkerRequestTestsReply& Message, const IMessageContextRef& Context)
 {
-	NumOfTestsToReceive = Message.TotalNumTests;
-
 	FAutomationTestInfo NewTest(Message.TestInfo);
 	TestInfo.Add(NewTest);
-	
-	if (TestInfo.Num() == NumOfTestsToReceive)
-	{
-		SetTestNames(Context->GetSender());
-	}
+}
+
+
+void FAutomationControllerManager::HandleRequestTestsReplyCompleteMessage(const FAutomationWorkerRequestTestsReplyComplete& Message, const IMessageContextRef& Context)
+{
+	SetTestNames(Context->GetSender());
 }
 
 
diff --git a/Engine/Source/Developer/AutomationController/Private/AutomationReport.cpp b/Engine/Source/Developer/AutomationController/Private/AutomationReport.cpp
--- a/Engine/Source/Developer/AutomationController/Private/AutomationReport.cpp
+++ b/Engine/Source/Developer/AutomationController/Private/AutomationReport.cpp
@@ -14,7 +14,7 @@ FAutomationReport::FAutomationReport(FAutomationTestInfo& InTestInfo, bool InIsP
 	, NumRecordsToKeep(0)
 {
 	// Enable smoke tests
-	if ( TestInfo.GetTestType() == EAutomationTestType::ATT_SmokeTest )
+	if ( TestInfo.GetTestFlags() == EAutomationTestFlags::SmokeFilter )
 	{
 		bEnabled = true;
 	}
@@ -197,17 +197,17 @@ bool FAutomationReport::IsSupported(const int32 ClusterIndex) const
 }
 
 
-uint8 FAutomationReport::GetTestType( ) const
+uint32 FAutomationReport::GetTestFlags( ) const
 {
-	return TestInfo.GetTestType();
+	return TestInfo.GetTestFlags();
 }
 
 
-void FAutomationReport::SetTestType( const uint8 InTestType )
+void FAutomationReport::SetTestFlags( const uint32 InTestFlags)
 {
-	TestInfo.AddTestType( InTestType );
+	TestInfo.AddTestFlags( InTestFlags );
 
-	if ( InTestType == EAutomationTestType::ATT_SmokeTest )
+	if ( InTestFlags == EAutomationTestFlags::SmokeFilter )
 	{
 		bEnabled = true;
 	}
@@ -220,7 +220,7 @@ const bool FAutomationReport::IsParent()
 
 const bool FAutomationReport::IsSmokeTest( )
 {
-	return GetTestType( ) & EAutomationTestType::ATT_SmokeTest ? true : false;
+	return GetTestFlags( ) & EAutomationTestFlags::SmokeFilter ? true : false;
 }
 
 bool FAutomationReport::SetFilter( TSharedPtr< AutomationFilterCollection > InFilter, const bool ParentPassedFilter )
@@ -737,7 +737,7 @@ TSharedPtr<IAutomationReport> FAutomationReport::EnsureReportExists(FAutomationT
 		else
 		{
 			// Create a parent node
-			FAutomationTestInfo ParentTestInfo( NameToMatch, "", InTestInfo.GetTestType(), InTestInfo.GetNumParticipantsRequired() ) ;
+			FAutomationTestInfo ParentTestInfo( NameToMatch, "", InTestInfo.GetTestFlags(), InTestInfo.GetNumParticipantsRequired() ) ;
 			MatchTest = MakeShareable(new FAutomationReport(ParentTestInfo, true));
 		}
 		//make new test
@@ -747,7 +747,7 @@ TSharedPtr<IAutomationReport> FAutomationReport::EnsureReportExists(FAutomationT
 	//mark this test as supported on a particular platform
 	MatchTest->SetSupport(ClusterIndex);
 
-	MatchTest->SetTestType( InTestInfo.GetTestType() );
+	MatchTest->SetTestFlags( InTestInfo.GetTestFlags() );
 	MatchTest->SetNumParticipantsRequired( MatchTest->GetNumParticipantsRequired() > InTestInfo.GetNumParticipantsRequired() ? MatchTest->GetNumParticipantsRequired() : InTestInfo.GetNumParticipantsRequired() );
 
 	TSharedPtr<IAutomationReport> FoundTest;
diff --git a/Engine/Source/Developer/AutomationController/Private/AutomationReport.h b/Engine/Source/Developer/AutomationController/Private/AutomationReport.h
--- a/Engine/Source/Developer/AutomationController/Private/AutomationReport.h
+++ b/Engine/Source/Developer/AutomationController/Private/AutomationReport.h
@@ -30,8 +30,8 @@ public:
 	virtual void SetEnabled(bool bShouldBeEnabled) override;
 	virtual void SetSupport(const int32 ClusterIndex) override;
 	virtual bool IsSupported(const int32 ClusterIndex) const override;
-	virtual void SetTestType( const uint8 TestType ) override;
-	virtual uint8 GetTestType( ) const override;
+	virtual void SetTestFlags(const uint32 TestFlags) override;
+	virtual uint32 GetTestFlags( ) const override;
 	virtual const bool IsParent( ) override;
 	virtual const bool IsSmokeTest( ) override;
 	virtual bool SetFilter( TSharedPtr< AutomationFilterCollection > InFilter, const bool ParentPassedFilter = false ) override;
diff --git a/Engine/Source/Developer/AutomationController/Public/Interfaces/IAutomationControllerManager.h b/Engine/Source/Developer/AutomationController/Public/Interfaces/IAutomationControllerManager.h
--- a/Engine/Source/Developer/AutomationController/Public/Interfaces/IAutomationControllerManager.h
+++ b/Engine/Source/Developer/AutomationController/Public/Interfaces/IAutomationControllerManager.h
@@ -307,16 +307,9 @@ public:
 	virtual bool IsDeveloperDirectoryIncluded(void) const = 0;
 
 	/**
-	 * Sets whether the automation tests should include visual commandlet.
+	 * Sets what flags we're considering automation tests from
 	 */
-	virtual void SetVisualCommandletFilter(const bool bInVisualCommandletFilterOn) = 0;
-
-	/**
-	 * Returns whether the automation tests should include visual commandlet.
-	 *
-	 * @return true if the visual commandlet should be included, false otherwise.
-	 */
-	virtual bool IsVisualCommandletFilterOn(void) const = 0;
+	virtual void SetRequestedTestFlags(const uint32 InRequestedTestFlags) = 0;
 
 	/**
 	 * Check if the automation tests have completed.
diff --git a/Engine/Source/Developer/AutomationController/Public/Interfaces/IAutomationReport.h b/Engine/Source/Developer/AutomationController/Public/Interfaces/IAutomationReport.h
--- a/Engine/Source/Developer/AutomationController/Public/Interfaces/IAutomationReport.h
+++ b/Engine/Source/Developer/AutomationController/Public/Interfaces/IAutomationReport.h
@@ -172,7 +172,7 @@ public:
 	 * 
 	 * @return the test type.
 	 */
-	virtual uint8 GetTestType() const = 0;
+	virtual uint32 GetTestFlags() const = 0;
 
 	/** Recursively gets the number of child nodes */
 	virtual int32 GetTotalNumChildren() const = 0;
@@ -210,11 +210,11 @@ public:
 	virtual void SetSupport(const int32 ClusterIndex) = 0;
 
 	/**
-	 * Set the test type.
+	 * Set the test flags.
 	 * 
-	 * @param InTestType The type of test.
+	 * @param InTestFlags The EAutomationTestFlags of the test.
 	 */
-	virtual void SetTestType(const uint8 InTestType) = 0;
+	virtual void SetTestFlags(const uint32 InTestFlags) = 0;
 
 	/** Returns if a particular platform is supported */
 	virtual bool IsSupported(const int32 ClusterIndex) const = 0;
diff --git a/Engine/Source/Developer/AutomationWindow/Private/AutomationFilter.h b/Engine/Source/Developer/AutomationWindow/Private/AutomationFilter.h
--- a/Engine/Source/Developer/AutomationWindow/Private/AutomationFilter.h
+++ b/Engine/Source/Developer/AutomationWindow/Private/AutomationFilter.h
@@ -92,9 +92,18 @@ public:
 	virtual bool PassesFilter( const TSharedPtr< IAutomationReport >& InReport ) const override
 	{
 		bool FilterPassed = true;
-		if( OnlySmokeTests && !InReport->IsSmokeTest( ) )
+		if (OnlySmokeTests)
 		{
-			FilterPassed = false;
+			//if we only want smoke tests and this isn't one
+			if (!InReport->IsSmokeTest())
+			{
+				FilterPassed = false;
+			}
+			//Leaf Nodes dictate this matching, not root nodes
+			if (InReport->GetTotalNumChildren())
+			{
+				FilterPassed = false;
+			}
 		}
 
 		if (ShowWarnings && ShowErrors)
diff --git a/Engine/Source/Developer/AutomationWindow/Private/SAutomationTestItem.cpp b/Engine/Source/Developer/AutomationWindow/Private/SAutomationTestItem.cpp
--- a/Engine/Source/Developer/AutomationWindow/Private/SAutomationTestItem.cpp
+++ b/Engine/Source/Developer/AutomationWindow/Private/SAutomationTestItem.cpp
@@ -308,7 +308,7 @@ END_SLATE_FUNCTION_BUILD_OPTIMIZATION
 const FSlateBrush* SAutomationTestItem::GetSmokeTestImage() const
 {
 	const FSlateBrush* ImageToUse = nullptr;
-	if ( TestStatus->GetTestType() & EAutomationTestType::ATT_SmokeTest )
+	if ( TestStatus->GetTestFlags() & EAutomationTestFlags::SmokeFilter )
 	{
 		if ( TestStatus->IsParent() )
 		{
diff --git a/Engine/Source/Developer/AutomationWindow/Private/SAutomationWindow.cpp b/Engine/Source/Developer/AutomationWindow/Private/SAutomationWindow.cpp
--- a/Engine/Source/Developer/AutomationWindow/Private/SAutomationWindow.cpp
+++ b/Engine/Source/Developer/AutomationWindow/Private/SAutomationWindow.cpp
@@ -34,18 +34,14 @@ public:
 		UI_COMMAND( FindWorkers, "Find Workers", "Find Workers", EUserInterfaceActionType::Button, FInputChord() );
 		UI_COMMAND( ErrorFilter, "Errors", "Toggle Error Filter", EUserInterfaceActionType::ToggleButton, FInputChord() );
 		UI_COMMAND( WarningFilter, "Warnings", "Toggle Warning Filter", EUserInterfaceActionType::ToggleButton, FInputChord() );
-		UI_COMMAND( SmokeTestFilter, "Smoke Tests", "Toggle Smoke Test Filter", EUserInterfaceActionType::ToggleButton, FInputChord() );
 		UI_COMMAND( DeveloperDirectoryContent, "Dev Content", "Developer Directory Content Filter (when enabled, developer directories are also included)", EUserInterfaceActionType::ToggleButton, FInputChord() );
-		UI_COMMAND( VisualCommandlet, "Vis Cmdlet", "Visual Commandlet Filter", EUserInterfaceActionType::ToggleButton, FInputChord() );
 	}
 public:
 	TSharedPtr<FUICommandInfo> RefreshTests;
 	TSharedPtr<FUICommandInfo> FindWorkers;
 	TSharedPtr<FUICommandInfo> ErrorFilter;
 	TSharedPtr<FUICommandInfo> WarningFilter;
-	TSharedPtr<FUICommandInfo> SmokeTestFilter;
 	TSharedPtr<FUICommandInfo> DeveloperDirectoryContent;
-	TSharedPtr<FUICommandInfo> VisualCommandlet;
 };
 
 //////////////////////////////////////////////////////////////////////////
@@ -470,23 +466,11 @@ void SAutomationWindow::CreateCommands()
 		FIsActionChecked::CreateRaw( this, &SAutomationWindow::IsWarningFilterOn )
 		);
 
-	ActionList.MapAction( Commands.SmokeTestFilter,
-		FExecuteAction::CreateRaw( this, &SAutomationWindow::OnToggleSmokeTestFilter ),
-		FCanExecuteAction::CreateRaw( this, &SAutomationWindow::IsAutomationControllerIdle ),
-		FIsActionChecked::CreateRaw( this, &SAutomationWindow::IsSmokeTestFilterOn )
-		);
-
 	ActionList.MapAction( Commands.DeveloperDirectoryContent,
 		FExecuteAction::CreateRaw( this, &SAutomationWindow::OnToggleDeveloperDirectoryIncluded ),
 		FCanExecuteAction::CreateRaw( this, &SAutomationWindow::IsAutomationControllerIdle ),
 		FIsActionChecked::CreateRaw( this, &SAutomationWindow::IsDeveloperDirectoryIncluded )
 		);
-
-	ActionList.MapAction( Commands.VisualCommandlet,
-		FExecuteAction::CreateRaw( this, &SAutomationWindow::OnToggleVisualCommandletFilter ),
-		FCanExecuteAction::CreateRaw( this, &SAutomationWindow::IsAutomationControllerIdle ),
-		FIsActionChecked::CreateRaw( this, &SAutomationWindow::IsVisualCommandletFilterOn )
-		);
 }
 
 TSharedRef< SWidget > SAutomationWindow::MakeAutomationWindowToolBar( const TSharedRef<FUICommandList>& InCommandList, TSharedPtr<class SAutomationWindow> InAutomationWindow )
@@ -499,7 +483,7 @@ TSharedRef< SWidget > SAutomationWindow::MakeAutomationWindowToolBar( const TSha
 {
 	struct Local
 	{
-		static void FillToolbar(FToolBarBuilder& ToolbarBuilder, TSharedRef<SWidget> RunTests, TSharedRef<SWidget> Searchbox, TSharedRef<SWidget> PresetBox, TSharedRef<SWidget> HistoryBox, TWeakPtr<class SAutomationWindow> InAutomationWindow)
+		static void FillToolbar(FToolBarBuilder& ToolbarBuilder, TSharedRef<SWidget> RunTests, TSharedRef<SWidget> Searchbox, TSharedRef<SWidget> PresetBox, TSharedRef<SWidget> HistoryBox, TSharedRef<SWidget> RequestedFilterBox, TWeakPtr<class SAutomationWindow> InAutomationWindow)
 		{
 			ToolbarBuilder.BeginSection("Automation");
 			{
@@ -520,9 +504,8 @@ TSharedRef< SWidget > SAutomationWindow::MakeAutomationWindowToolBar( const TSha
 			{
 				ToolbarBuilder.AddToolBarButton( FAutomationWindowCommands::Get().ErrorFilter );
 				ToolbarBuilder.AddToolBarButton( FAutomationWindowCommands::Get().WarningFilter );
-				ToolbarBuilder.AddToolBarButton( FAutomationWindowCommands::Get().SmokeTestFilter );
 				ToolbarBuilder.AddToolBarButton( FAutomationWindowCommands::Get().DeveloperDirectoryContent );
-				ToolbarBuilder.AddToolBarButton( FAutomationWindowCommands::Get().VisualCommandlet );
+				ToolbarBuilder.AddWidget(RequestedFilterBox);
 			}
 			ToolbarBuilder.EndSection();
 			ToolbarBuilder.BeginSection("GroupFlags");
@@ -774,9 +757,27 @@ TSharedRef< SWidget > SAutomationWindow::MakeAutomationWindowToolBar( const TSha
 			]
 		];
 
+	RequestedFilterComboList.Empty();
+	RequestedFilterComboList.Add(MakeShareable(new FString(TEXT("Smoke Tests"))));
+	RequestedFilterComboList.Add(MakeShareable(new FString(TEXT("Engine Tests"))));
+	RequestedFilterComboList.Add(MakeShareable(new FString(TEXT("Product Tests"))));
+	RequestedFilterComboList.Add(MakeShareable(new FString(TEXT("Perf Tests"))));
+	RequestedFilterComboList.Add(MakeShareable(new FString(TEXT("Stress Tests"))));
+	RequestedFilterComboList.Add(MakeShareable(new FString(TEXT("Standard Tests"))));
+	RequestedFilterComboList.Add(MakeShareable(new FString(TEXT("All Tests"))));
+
+	SAssignNew(RequestedFilterComboBox, SComboBox< TSharedPtr<FString> >)
+		.OptionsSource(&RequestedFilterComboList)
+		.OnGenerateWidget(this, &SAutomationWindow::GenerateRequestedFilterComboItem)
+		.OnSelectionChanged(this, &SAutomationWindow::HandleRequesteFilterChanged)
+		[
+			SNew(STextBlock)
+			.Text(this, &SAutomationWindow::GetRequestedFilterComboText)
+		];
+
 	FToolBarBuilder ToolbarBuilder( InCommandList, FMultiBoxCustomization::None );
 	TWeakPtr<SAutomationWindow> AutomationWindow = SharedThis(this);
-	Local::FillToolbar( ToolbarBuilder, RunTests, Searchbox, TestPresets, History, AutomationWindow );
+	Local::FillToolbar(ToolbarBuilder, RunTests, Searchbox, TestPresets, History, RequestedFilterComboBox.ToSharedRef(), AutomationWindow);
 
 	// Create the tool bar!
 	return
@@ -868,6 +869,39 @@ void SAutomationWindow::HandlePresetChanged( TSharedPtr<FAutomationTestPreset> I
 	}
 }
 
+void SAutomationWindow::HandleRequesteFilterChanged(TSharedPtr<FString> Item, ESelectInfo::Type SelectInfo)
+{
+	int32 EntryIndex = RequestedFilterComboList.Find(Item);
+	uint32 NewRequestedFlags = EAutomationTestFlags::SmokeFilter;
+
+	switch (EntryIndex)
+	{
+		case 0:	//	"Smoke Tests"
+			NewRequestedFlags = EAutomationTestFlags::SmokeFilter;
+			break;
+		case 1:	//	"Engine Tests"
+			NewRequestedFlags = EAutomationTestFlags::EngineFilter;
+			break;
+		case 2:	//	"Product Tests"
+			NewRequestedFlags = EAutomationTestFlags::ProductFilter;
+			break;
+		case 3:	//	"Perf Tests"
+			NewRequestedFlags = EAutomationTestFlags::PerfFilter;
+			break;
+		case 4:	//	"Stress Tests"
+			NewRequestedFlags = EAutomationTestFlags::StressFilter;
+			break;
+		case 5:	//	"Standard Tests"
+			NewRequestedFlags = EAutomationTestFlags::SmokeFilter | EAutomationTestFlags::EngineFilter | EAutomationTestFlags::ProductFilter | EAutomationTestFlags::PerfFilter;
+			break;
+		case 6:	//	"All Tests"
+			NewRequestedFlags = EAutomationTestFlags::FilterMask;
+			break;
+	}
+	AutomationController->SetRequestedTestFlags(NewRequestedFlags);
+}
+
+
 void SAutomationWindow::ExpandEnabledTests( TSharedPtr< IAutomationReport > InReport )
 {
 	// Expand node if the report is enabled or contains an enabled test
@@ -923,12 +957,31 @@ FText SAutomationWindow::GetPresetComboText() const
 	}
 }
 
+FText SAutomationWindow::GetRequestedFilterComboText() const
+{
+	if (RequestedFilterComboBox->GetSelectedItem().IsValid())
+	{
+		return FText::FromString(*RequestedFilterComboBox->GetSelectedItem());
+	}
+	else
+	{
+		return LOCTEXT("AutomationRequestedFilterComboLabel", "<Select A Filter>");
+	}
+}
+
+
 TSharedRef<SWidget> SAutomationWindow::GeneratePresetComboItem(TSharedPtr<FAutomationTestPreset> InItem)
 {
 	return SNew(STextBlock) 
 		.Text( InItem->GetPresetName() );
 }
 
+TSharedRef<SWidget> SAutomationWindow::GenerateRequestedFilterComboItem(TSharedPtr<FString> InItem)
+{
+	return SNew(STextBlock)
+		.Text(FText::FromString(*InItem));
+}
+
 TSharedRef< SWidget > SAutomationWindow::GenerateGroupOptionsMenuContent( TWeakPtr<class SAutomationWindow> InAutomationWindow )
 {
 	TSharedPtr<SAutomationWindow> AutomationWindow(InAutomationWindow.Pin());
@@ -1573,21 +1626,6 @@ void SAutomationWindow::OnFilterTextChanged( const FText& InFilterText )
 }
 
 
-bool SAutomationWindow::IsVisualCommandletFilterOn() const
-{
-	return AutomationController->IsVisualCommandletFilterOn();
-}
-
-
-void SAutomationWindow::OnToggleVisualCommandletFilter()
-{
-	//Change controller filter
-	AutomationController->SetVisualCommandletFilter(!IsVisualCommandletFilterOn());
-	// need to call this to request update
-	ListTests( );
-}
-
-
 bool SAutomationWindow::IsDeveloperDirectoryIncluded() const
 {
 	return AutomationController->IsDeveloperDirectoryIncluded();
diff --git a/Engine/Source/Developer/AutomationWindow/Private/SAutomationWindow.h b/Engine/Source/Developer/AutomationWindow/Private/SAutomationWindow.h
--- a/Engine/Source/Developer/AutomationWindow/Private/SAutomationWindow.h
+++ b/Engine/Source/Developer/AutomationWindow/Private/SAutomationWindow.h
@@ -180,7 +180,14 @@ private:
 	 */
 	TSharedRef<SWidget> GeneratePresetComboItem(TSharedPtr<FAutomationTestPreset> InItem);
 
-	/** 
+	/**
+	 * Creates a combo item for the requested filter.
+	 *
+	 * @return New combo item widget.
+	 */
+	TSharedRef<SWidget> GenerateRequestedFilterComboItem(TSharedPtr<FString> InItem);
+		
+	/**
 	 * Populates OutSearchStrings with the strings that should be used in searching.
 	 *
 	 * @param Report The automation report to get a text description from.
@@ -233,12 +240,6 @@ private:
 
 	/** Filter text has been updated */
 	void OnFilterTextChanged( const FText& InFilterText );
-
-	/** Returns true if automation test includes visual commandlet*/
-	bool IsVisualCommandletFilterOn() const;
-	
-	/** Toggles filter of visual commandlet */
-	void OnToggleVisualCommandletFilter();
 	
 	/** Returns if we're considering tests on content within the developer folders */
 	bool IsDeveloperDirectoryIncluded() const;
@@ -421,12 +422,16 @@ private:
 
 	/** Called when the user selects a new preset from the preset combo box. */
 	void HandlePresetChanged( TSharedPtr<FAutomationTestPreset> Item, ESelectInfo::Type SelectInfo );
+	/** Called when the user changes the requested test filter */
+	void HandleRequesteFilterChanged(TSharedPtr<FString> Item, ESelectInfo::Type SelectInfo);
 
 	/** Expands the test tree to show all enabled tests. */
 	void ExpandEnabledTests( TSharedPtr< IAutomationReport > InReport );
 
 	/** Gets the text to display for the preset combo box. */
 	FText GetPresetComboText() const;
+	/** Gets the text to display for the requested filter combo box. */
+	FText GetRequestedFilterComboText() const;
 
 	/**
 	 * Handle the copy button clicked in the command bar.
@@ -555,6 +560,10 @@ private:
 	/** Holds a pointer to the preset combo box widget. */
 	TSharedPtr< SComboBox< TSharedPtr<FAutomationTestPreset> > > PresetComboBox;
 
+	/** Holds a pointer to the preset combo box widget. */
+	TSharedPtr< SComboBox< TSharedPtr<FString> > >	RequestedFilterComboBox;
+	TArray< TSharedPtr< FString > >					RequestedFilterComboList;
+
 	/** Holds a pointer to the preset text box. */
 	TSharedPtr<SEditableTextBox> PresetTextBox;
 
diff --git a/Engine/Source/Developer/DirectoryWatcher/Private/DirectoryWatcherTests.cpp b/Engine/Source/Developer/DirectoryWatcher/Private/DirectoryWatcherTests.cpp
--- a/Engine/Source/Developer/DirectoryWatcher/Private/DirectoryWatcherTests.cpp
+++ b/Engine/Source/Developer/DirectoryWatcher/Private/DirectoryWatcherTests.cpp
@@ -98,7 +98,7 @@ FString GetWorkingDir()
 }
 
 static const float TestTickDelay = 0.1f;
-IMPLEMENT_SIMPLE_AUTOMATION_TEST(FDirectoryWatcherSimpleCreateTest, "System.Plugins.Directory Watcher.Simple Create", EAutomationTestFlags::ATF_Editor)
+IMPLEMENT_SIMPLE_AUTOMATION_TEST(FDirectoryWatcherSimpleCreateTest, "System.Plugins.Directory Watcher.Simple Create", EAutomationTestFlags::EditorContext | EAutomationTestFlags::EngineFilter)
 bool FDirectoryWatcherSimpleCreateTest::RunTest(const FString& Parameters)
 {
 	const FString WorkingDir = GetWorkingDir();
@@ -130,7 +130,7 @@ bool FDirectoryWatcherSimpleCreateTest::RunTest(const FString& Parameters)
 	return true;
 }
 
-IMPLEMENT_SIMPLE_AUTOMATION_TEST(FDirectoryWatcherSimpleModifyTest, "System.Plugins.Directory Watcher.Simple Modify", EAutomationTestFlags::ATF_Editor)
+IMPLEMENT_SIMPLE_AUTOMATION_TEST(FDirectoryWatcherSimpleModifyTest, "System.Plugins.Directory Watcher.Simple Modify", EAutomationTestFlags::EditorContext | EAutomationTestFlags::EngineFilter)
 bool FDirectoryWatcherSimpleModifyTest::RunTest(const FString& Parameters)
 {
 	const FString WorkingDir = GetWorkingDir();
@@ -197,7 +197,7 @@ bool FDirectoryWatcherSimpleModifyTest::RunTest(const FString& Parameters)
 	return true;
 }
 
-IMPLEMENT_SIMPLE_AUTOMATION_TEST(FDirectoryWatcherSimpleDeleteTest, "System.Plugins.Directory Watcher.Simple Delete", EAutomationTestFlags::ATF_Editor)
+IMPLEMENT_SIMPLE_AUTOMATION_TEST(FDirectoryWatcherSimpleDeleteTest, "System.Plugins.Directory Watcher.Simple Delete", EAutomationTestFlags::EditorContext | EAutomationTestFlags::EngineFilter)
 bool FDirectoryWatcherSimpleDeleteTest::RunTest(const FString& Parameters)
 {
 	const FString WorkingDir = GetWorkingDir();
@@ -232,7 +232,7 @@ bool FDirectoryWatcherSimpleDeleteTest::RunTest(const FString& Parameters)
 	return true;
 }
 
-IMPLEMENT_SIMPLE_AUTOMATION_TEST(FDirectoryWatcherSubFolderTest, "System.Plugins.Directory Watcher.Sub Folder", EAutomationTestFlags::ATF_Editor)
+IMPLEMENT_SIMPLE_AUTOMATION_TEST(FDirectoryWatcherSubFolderTest, "System.Plugins.Directory Watcher.Sub Folder", EAutomationTestFlags::EditorContext | EAutomationTestFlags::EngineFilter)
 bool FDirectoryWatcherSubFolderTest::RunTest(const FString& Parameters)
 {
 	const FString WorkingDir = GetWorkingDir();
@@ -337,7 +337,7 @@ bool FDirectoryWatcherSubFolderTest::RunTest(const FString& Parameters)
 	return true;
 }
 
-IMPLEMENT_SIMPLE_AUTOMATION_TEST(FDirectoryWatcherNewFolderTest, "System.Plugins.Directory Watcher.New Folder", EAutomationTestFlags::ATF_Editor)
+IMPLEMENT_SIMPLE_AUTOMATION_TEST(FDirectoryWatcherNewFolderTest, "System.Plugins.Directory Watcher.New Folder", EAutomationTestFlags::EditorContext | EAutomationTestFlags::EngineFilter)
 bool FDirectoryWatcherNewFolderTest::RunTest(const FString& Parameters)
 {
 	const FString WorkingDir = GetWorkingDir();
@@ -421,7 +421,7 @@ bool FDirectoryWatcherNewFolderTest::RunTest(const FString& Parameters)
 	return true;
 }
 
-IMPLEMENT_SIMPLE_AUTOMATION_TEST(FDirectoryWatcherIgnoreSubtreeTest, "System.Plugins.Directory Watcher.Ignore Subtree", EAutomationTestFlags::ATF_Editor)
+IMPLEMENT_SIMPLE_AUTOMATION_TEST(FDirectoryWatcherIgnoreSubtreeTest, "System.Plugins.Directory Watcher.Ignore Subtree", EAutomationTestFlags::EditorContext | EAutomationTestFlags::EngineFilter)
 bool FDirectoryWatcherIgnoreSubtreeTest::RunTest(const FString& Parameters)
 {
 	const FString WorkingDir = GetWorkingDir();
diff --git a/Engine/Source/Developer/DirectoryWatcher/Private/FileCacheUtilities.cpp b/Engine/Source/Developer/DirectoryWatcher/Private/FileCacheUtilities.cpp
--- a/Engine/Source/Developer/DirectoryWatcher/Private/FileCacheUtilities.cpp
+++ b/Engine/Source/Developer/DirectoryWatcher/Private/FileCacheUtilities.cpp
@@ -33,7 +33,7 @@ bool MatchExtensionString(const TCHAR* Filename, const TCHAR* Extensions)
 }
 
 /** Test MatchExtensionString */
-IMPLEMENT_SIMPLE_AUTOMATION_TEST(FMatchExtensionStringTest, "System.Plugins.Directory Watcher.File Cache.Extension Matching", EAutomationTestFlags::ATF_Editor)
+IMPLEMENT_SIMPLE_AUTOMATION_TEST(FMatchExtensionStringTest, "System.Plugins.Directory Watcher.File Cache.Extension Matching", EAutomationTestFlags::EditorContext | EAutomationTestFlags::EngineFilter)
 bool FMatchExtensionStringTest::RunTest(const FString& Parameters)
 {
 	auto Test = [](const TCHAR* Needle, const TCHAR* Haystack, bool bExpected) -> bool{
diff --git a/Engine/Source/Developer/FunctionalTesting/Private/Tests/FuncTestPerforming.cpp b/Engine/Source/Developer/FunctionalTesting/Private/Tests/FuncTestPerforming.cpp
--- a/Engine/Source/Developer/FunctionalTesting/Private/Tests/FuncTestPerforming.cpp
+++ b/Engine/Source/Developer/FunctionalTesting/Private/Tests/FuncTestPerforming.cpp
@@ -108,7 +108,7 @@ bool FStartFTestsOnMap::Update()
  */
 
 // create test base class
-IMPLEMENT_COMPLEX_AUTOMATION_TEST_PRIVATE(FFunctionalTestingMapsBase, "Project.Maps.Functional Testing", (EAutomationTestFlags::ATF_Game | EAutomationTestFlags::ATF_Editor))
+IMPLEMENT_COMPLEX_AUTOMATION_TEST_PRIVATE(FFunctionalTestingMapsBase, "Project.Maps.Functional Testing", (EAutomationTestFlags::ClientContext | EAutomationTestFlags::EditorContext | EAutomationTestFlags::ProductFilter))
 // implement specific class with non-standard overrides
 class FFunctionalTestingMaps : public FFunctionalTestingMapsBase
 {
diff --git a/Engine/Source/Developer/SourceControl/Private/Tests/SourceControlTests.cpp b/Engine/Source/Developer/SourceControl/Private/Tests/SourceControlTests.cpp
--- a/Engine/Source/Developer/SourceControl/Private/Tests/SourceControlTests.cpp
+++ b/Engine/Source/Developer/SourceControl/Private/Tests/SourceControlTests.cpp
@@ -85,7 +85,7 @@ bool FSetProviderLatentCommand::Update()
 	return true;
 }
 
-IMPLEMENT_COMPLEX_AUTOMATION_TEST(FSetProviderTest, "Project.Editor.Source Control.Set Provider", (EAutomationTestFlags::ATF_Editor | EAutomationTestFlags::ATF_RequiresUser))
+IMPLEMENT_COMPLEX_AUTOMATION_TEST(FSetProviderTest, "Project.Editor.Source Control.Set Provider", (EAutomationTestFlags::EditorContext | EAutomationTestFlags::RequiresUser | EAutomationTestFlags::EngineFilter))
 
 void FSetProviderTest::GetTests(TArray<FString>& OutBeautifiedNames, TArray<FString>& OutTestCommands) const
 {
@@ -116,7 +116,7 @@ bool FConnectLatentCommand::Update()
 	return AsyncHelper.IsDone();
 }
 
-IMPLEMENT_COMPLEX_AUTOMATION_TEST(FConnectTest, "Project.Editor.Source Control.Connect", (EAutomationTestFlags::ATF_Editor | EAutomationTestFlags::ATF_RequiresUser))
+IMPLEMENT_COMPLEX_AUTOMATION_TEST(FConnectTest, "Project.Editor.Source Control.Connect", (EAutomationTestFlags::EditorContext | EAutomationTestFlags::RequiresUser | EAutomationTestFlags::EngineFilter))
 
 void FConnectTest::GetTests(TArray<FString>& OutBeautifiedNames, TArray<FString>& OutTestCommands) const
 {
@@ -213,7 +213,7 @@ bool FCheckOutLatentCommand::Update()
 	return AsyncHelper.IsDone();
 }
 
-IMPLEMENT_COMPLEX_AUTOMATION_TEST(FCheckOutTest, "Project.Editor.Source Control.Check Out", (EAutomationTestFlags::ATF_Editor | EAutomationTestFlags::ATF_RequiresUser))
+IMPLEMENT_COMPLEX_AUTOMATION_TEST(FCheckOutTest, "Project.Editor.Source Control.Check Out", (EAutomationTestFlags::EditorContext | EAutomationTestFlags::RequiresUser | EAutomationTestFlags::EngineFilter))
 
 void FCheckOutTest::GetTests(TArray<FString>& OutBeautifiedNames, TArray<FString>& OutTestCommands) const
 {
@@ -364,7 +364,7 @@ bool FMarkForAddLatentCommand::Update()
 	return AsyncHelper.IsDone();
 }
 
-IMPLEMENT_COMPLEX_AUTOMATION_TEST(FMarkForAddTest, "Project.Editor.Source Control.Mark For Add", (EAutomationTestFlags::ATF_Editor | EAutomationTestFlags::ATF_RequiresUser))
+IMPLEMENT_COMPLEX_AUTOMATION_TEST(FMarkForAddTest, "Project.Editor.Source Control.Mark For Add", (EAutomationTestFlags::EditorContext | EAutomationTestFlags::RequiresUser | EAutomationTestFlags::EngineFilter))
 
 void FMarkForAddTest::GetTests(TArray<FString>& OutBeautifiedNames, TArray<FString>& OutTestCommands) const
 {
@@ -437,7 +437,7 @@ bool FDeleteLatentCommand::Update()
 	return AsyncHelper.IsDone();
 }
 
-IMPLEMENT_COMPLEX_AUTOMATION_TEST(FDeleteTest, "Project.Editor.Source Control.Delete", (EAutomationTestFlags::ATF_Editor | EAutomationTestFlags::ATF_RequiresUser))
+IMPLEMENT_COMPLEX_AUTOMATION_TEST(FDeleteTest, "Project.Editor.Source Control.Delete", (EAutomationTestFlags::EditorContext | EAutomationTestFlags::RequiresUser | EAutomationTestFlags::EngineFilter))
 
 void FDeleteTest::GetTests(TArray<FString>& OutBeautifiedNames, TArray<FString>& OutTestCommands) const
 {
@@ -538,7 +538,7 @@ bool FEditTextureLatentCommand::Update()
 	return true;
 }
 
-IMPLEMENT_COMPLEX_AUTOMATION_TEST(FCheckInTest, "Project.Editor.Source Control.Check In", (EAutomationTestFlags::ATF_Editor | EAutomationTestFlags::ATF_RequiresUser))
+IMPLEMENT_COMPLEX_AUTOMATION_TEST(FCheckInTest, "Project.Editor.Source Control.Check In", (EAutomationTestFlags::EditorContext | EAutomationTestFlags::RequiresUser | EAutomationTestFlags::EngineFilter))
 
 void FCheckInTest::GetTests(TArray<FString>& OutBeautifiedNames, TArray<FString>& OutTestCommands) const
 {
@@ -608,7 +608,7 @@ bool FSyncLatentCommand::Update()
 	return AsyncHelper.IsDone();
 }
 
-IMPLEMENT_COMPLEX_AUTOMATION_TEST(FSyncTest, "Project.Editor.Source Control.Sync", (EAutomationTestFlags::ATF_Editor | EAutomationTestFlags::ATF_RequiresUser))
+IMPLEMENT_COMPLEX_AUTOMATION_TEST(FSyncTest, "Project.Editor.Source Control.Sync", (EAutomationTestFlags::EditorContext | EAutomationTestFlags::RequiresUser | EAutomationTestFlags::EngineFilter))
 
 void FSyncTest::GetTests(TArray<FString>& OutBeautifiedNames, TArray<FString>& OutTestCommands) const
 {
@@ -631,7 +631,7 @@ bool FSyncTest::RunTest(const FString& Parameters)
 	return true;
 }
 
-IMPLEMENT_COMPLEX_AUTOMATION_TEST(FRevertTest, "Project.Editor.Source Control.Revert", (EAutomationTestFlags::ATF_Editor | EAutomationTestFlags::ATF_RequiresUser))
+IMPLEMENT_COMPLEX_AUTOMATION_TEST(FRevertTest, "Project.Editor.Source Control.Revert", (EAutomationTestFlags::EditorContext | EAutomationTestFlags::RequiresUser | EAutomationTestFlags::EngineFilter))
 
 void FRevertTest::GetTests(TArray<FString>& OutBeautifiedNames, TArray<FString>& OutTestCommands) const
 {
@@ -726,7 +726,7 @@ bool FGetStateLatentCommand::Update()
 	return true;
 }
 
-IMPLEMENT_COMPLEX_AUTOMATION_TEST(FUpdateStatusTest, "Project.Editor.Source Control.Update Status", (EAutomationTestFlags::ATF_Editor | EAutomationTestFlags::ATF_RequiresUser))
+IMPLEMENT_COMPLEX_AUTOMATION_TEST(FUpdateStatusTest, "Project.Editor.Source Control.Update Status", (EAutomationTestFlags::EditorContext | EAutomationTestFlags::RequiresUser | EAutomationTestFlags::EngineFilter))
 
 void FUpdateStatusTest::GetTests(TArray<FString>& OutBeautifiedNames, TArray<FString>& OutTestCommands) const
 {
@@ -810,7 +810,7 @@ bool FGetLabelLatentCommand::Update()
 	return true;
 }
 
-IMPLEMENT_COMPLEX_AUTOMATION_TEST(FGetLabelTest, "Project.Editor.Source Control.Get Label", (EAutomationTestFlags::ATF_Editor | EAutomationTestFlags::ATF_RequiresUser))
+IMPLEMENT_COMPLEX_AUTOMATION_TEST(FGetLabelTest, "Project.Editor.Source Control.Get Label", (EAutomationTestFlags::EditorContext | EAutomationTestFlags::RequiresUser | EAutomationTestFlags::EngineFilter))
 
 void FGetLabelTest::GetTests(TArray<FString>& OutBeautifiedNames, TArray<FString>& OutTestCommands) const
 {
@@ -855,7 +855,7 @@ bool FSyncLabelLatentCommand::Update()
 	return true;
 }
 
-IMPLEMENT_COMPLEX_AUTOMATION_TEST(FSyncLabelTest, "Project.Editor.Source Control.Sync Label", (EAutomationTestFlags::ATF_Editor | EAutomationTestFlags::ATF_RequiresUser))
+IMPLEMENT_COMPLEX_AUTOMATION_TEST(FSyncLabelTest, "Project.Editor.Source Control.Sync Label", (EAutomationTestFlags::EditorContext | EAutomationTestFlags::RequiresUser | EAutomationTestFlags::EngineFilter))
 
 void FSyncLabelTest::GetTests(TArray<FString>& OutBeautifiedNames, TArray<FString>& OutTestCommands) const
 {
@@ -921,7 +921,7 @@ bool FGetRevisionLatentCommand::Update()
 	return true;
 }
 
-IMPLEMENT_COMPLEX_AUTOMATION_TEST(FGetRevisionTest, "Project.Editor.Source Control.Get Revision", (EAutomationTestFlags::ATF_Editor | EAutomationTestFlags::ATF_RequiresUser))
+IMPLEMENT_COMPLEX_AUTOMATION_TEST(FGetRevisionTest, "Project.Editor.Source Control.Get Revision", (EAutomationTestFlags::EditorContext | EAutomationTestFlags::RequiresUser | EAutomationTestFlags::EngineFilter))
 
 void FGetRevisionTest::GetTests(TArray<FString>& OutBeautifiedNames, TArray<FString>& OutTestCommands) const
 {
diff --git a/Engine/Source/Editor/Cascade/Private/Tests/ParticleSystemTests.cpp b/Engine/Source/Editor/Cascade/Private/Tests/ParticleSystemTests.cpp
--- a/Engine/Source/Editor/Cascade/Private/Tests/ParticleSystemTests.cpp
+++ b/Engine/Source/Editor/Cascade/Private/Tests/ParticleSystemTests.cpp
@@ -226,7 +226,7 @@ bool FRunParticleEditorPromotionTestCommand::Update()
 /**
 * Automation test that handles the build promotion process
 */
-IMPLEMENT_SIMPLE_AUTOMATION_TEST(FParticleEditorPromotionTest, "System.Promotion.Editor.Particle Editor", EAutomationTestFlags::ATF_Editor)
+IMPLEMENT_SIMPLE_AUTOMATION_TEST(FParticleEditorPromotionTest, "System.Promotion.Editor.Particle Editor", EAutomationTestFlags::EditorContext | EAutomationTestFlags::EngineFilter)
 
 bool FParticleEditorPromotionTest::RunTest(const FString& Parameters)
 {
diff --git a/Engine/Source/Editor/EditorStyle/Private/SlateEditorStyle.cpp b/Engine/Source/Editor/EditorStyle/Private/SlateEditorStyle.cpp
--- a/Engine/Source/Editor/EditorStyle/Private/SlateEditorStyle.cpp
+++ b/Engine/Source/Editor/EditorStyle/Private/SlateEditorStyle.cpp
@@ -6370,12 +6370,10 @@ void FSlateEditorStyle::FStyle::SetupAutomationStyles()
 		Set( "AutomationWindow.WarningFilter", new IMAGE_BRUSH( "Automation/WarningFilter", Icon40x40) );
 		Set( "AutomationWindow.SmokeTestFilter", new IMAGE_BRUSH( "Automation/SmokeTestFilter", Icon40x40) );
 		Set( "AutomationWindow.DeveloperDirectoryContent", new IMAGE_BRUSH( "Automation/DeveloperDirectoryContent", Icon40x40) );
-		Set( "AutomationWindow.VisualCommandlet", new IMAGE_BRUSH( "Automation/VisualCommandlet", Icon40x40) );
 		Set( "AutomationWindow.ErrorFilter.Small", new IMAGE_BRUSH( "Automation/ErrorFilter", Icon20x20) );
 		Set( "AutomationWindow.WarningFilter.Small", new IMAGE_BRUSH( "Automation/WarningFilter", Icon20x20) );
 		Set( "AutomationWindow.SmokeTestFilter.Small", new IMAGE_BRUSH( "Automation/SmokeTestFilter", Icon20x20) );
 		Set( "AutomationWindow.DeveloperDirectoryContent.Small", new IMAGE_BRUSH( "Automation/DeveloperDirectoryContent", Icon20x20) );
-		Set( "AutomationWindow.VisualCommandlet.Small", new IMAGE_BRUSH( "Automation/VisualCommandlet", Icon20x20) );
 		Set( "AutomationWindow.TrackHistory", new IMAGE_BRUSH( "Automation/TrackTestHistory", Icon40x40) );
 
 		//device group settings
diff --git a/Engine/Source/Editor/GameProjectGeneration/Private/Tests/GameProjectAutomationTests.cpp b/Engine/Source/Editor/GameProjectGeneration/Private/Tests/GameProjectAutomationTests.cpp
--- a/Engine/Source/Editor/GameProjectGeneration/Private/Tests/GameProjectAutomationTests.cpp
+++ b/Engine/Source/Editor/GameProjectGeneration/Private/Tests/GameProjectAutomationTests.cpp
@@ -175,7 +175,7 @@ namespace GameProjectAutomationUtils
 /**
 * Automation test to clean up old test project files
 */
-IMPLEMENT_SIMPLE_AUTOMATION_TEST(FBuildPromotionNewProjectCleanupTest, "System.Promotion.Project Promotion Pass.Step 1 Blank Project Creation.Cleanup Potential Project Location", EAutomationTestFlags::ATF_Editor);
+IMPLEMENT_SIMPLE_AUTOMATION_TEST(FBuildPromotionNewProjectCleanupTest, "System.Promotion.Project Promotion Pass.Step 1 Blank Project Creation.Cleanup Potential Project Location", EAutomationTestFlags::EditorContext | EAutomationTestFlags::EngineFilter);
 bool FBuildPromotionNewProjectCleanupTest::RunTest(const FString& Parameters)
 {
 	FString DesiredProjectFilename = GameProjectAutomationUtils::GetDesiredProjectFilename();
@@ -210,7 +210,7 @@ bool FBuildPromotionNewProjectCleanupTest::RunTest(const FString& Parameters)
 /**
 * Automation test to create a new project
 */
-IMPLEMENT_SIMPLE_AUTOMATION_TEST(FBuildPromotionNewProjectCreateTest, "System.Promotion.Project Promotion Pass.Step 1 Blank Project Creation.Create Project", EAutomationTestFlags::ATF_Editor);
+IMPLEMENT_SIMPLE_AUTOMATION_TEST(FBuildPromotionNewProjectCreateTest, "System.Promotion.Project Promotion Pass.Step 1 Blank Project Creation.Create Project", EAutomationTestFlags::EditorContext | EAutomationTestFlags::EngineFilter);
 bool FBuildPromotionNewProjectCreateTest::RunTest(const FString& Parameters)
 {
 	FString DesiredProjectFilename = GameProjectAutomationUtils::GetDesiredProjectFilename();
@@ -255,7 +255,7 @@ bool FBuildPromotionNewProjectCreateTest::RunTest(const FString& Parameters)
 /**
 * Automation test to create a simple level and save it
 */
-IMPLEMENT_SIMPLE_AUTOMATION_TEST(FBuildPromotionNewProjectMapTest, "System.Promotion.Project Promotion Pass.Step 2 Basic Level Creation.Create Basic Level", EAutomationTestFlags::ATF_Editor);
+IMPLEMENT_SIMPLE_AUTOMATION_TEST(FBuildPromotionNewProjectMapTest, "System.Promotion.Project Promotion Pass.Step 2 Basic Level Creation.Create Basic Level", EAutomationTestFlags::EditorContext | EAutomationTestFlags::EngineFilter);
 bool FBuildPromotionNewProjectMapTest::RunTest(const FString& Parameters)
 {
 	//New level
@@ -324,7 +324,7 @@ bool FBuildPromotionNewProjectMapTest::RunTest(const FString& Parameters)
 /*
 * Template project creation test
 */
-IMPLEMENT_SIMPLE_AUTOMATION_TEST( FCreateBPTemplateProjectAutomationTests, "System.Promotion.Project Promotion Pass.Step 3 NewProjectCreationTests.CreateBlueprintProjects", EAutomationTestFlags::ATF_Editor )
+IMPLEMENT_SIMPLE_AUTOMATION_TEST(FCreateBPTemplateProjectAutomationTests, "System.Promotion.Project Promotion Pass.Step 3 NewProjectCreationTests.CreateBlueprintProjects", EAutomationTestFlags::EditorContext | EAutomationTestFlags::EngineFilter)
 
 /** 
  * Uses the new project wizard to locate all templates available for new blueprint project creation and verifies creation succeeds.
@@ -352,7 +352,7 @@ bool FCreateBPTemplateProjectAutomationTests::RunTest(const FString& Parameters)
 /*
 * Template project creation test
 */
-IMPLEMENT_SIMPLE_AUTOMATION_TEST( FCreateCPPTemplateProjectAutomationTests, "System.Promotion.Project Promotion Pass.Step 3 NewProjectCreationTests.CreateCodeProjects", EAutomationTestFlags::ATF_Editor )
+IMPLEMENT_SIMPLE_AUTOMATION_TEST(FCreateCPPTemplateProjectAutomationTests, "System.Promotion.Project Promotion Pass.Step 3 NewProjectCreationTests.CreateCodeProjects", EAutomationTestFlags::EditorContext | EAutomationTestFlags::EngineFilter)
 
 /** 
  * Uses the new project wizard to locate all templates available for new code project creation and verifies creation succeeds.
diff --git a/Engine/Source/Editor/LandscapeEditor/Private/Tests/LandscapeAutomationTests.cpp b/Engine/Source/Editor/LandscapeEditor/Private/Tests/LandscapeAutomationTests.cpp
--- a/Engine/Source/Editor/LandscapeEditor/Private/Tests/LandscapeAutomationTests.cpp
+++ b/Engine/Source/Editor/LandscapeEditor/Private/Tests/LandscapeAutomationTests.cpp
@@ -117,7 +117,7 @@ bool FEndModifyLandscapeCommand::Update()
 /**
 * Landscape creation / edit test
 */
-IMPLEMENT_SIMPLE_AUTOMATION_TEST(FLandscapeEditorTest, "System.Promotion.Editor.Landscape Editor", EAutomationTestFlags::ATF_Editor | EAutomationTestFlags::ATF_NonNullRHI );
+IMPLEMENT_SIMPLE_AUTOMATION_TEST(FLandscapeEditorTest, "System.Promotion.Editor.Landscape Editor", EAutomationTestFlags::EditorContext | EAutomationTestFlags::NonNullRHI | EAutomationTestFlags::EngineFilter);
 bool FLandscapeEditorTest::RunTest(const FString& Parameters)
 {
 	//New level
diff --git a/Engine/Source/Editor/MaterialEditor/Private/Tests/MaterialEditorTests.cpp b/Engine/Source/Editor/MaterialEditor/Private/Tests/MaterialEditorTests.cpp
--- a/Engine/Source/Editor/MaterialEditor/Private/Tests/MaterialEditorTests.cpp
+++ b/Engine/Source/Editor/MaterialEditor/Private/Tests/MaterialEditorTests.cpp
@@ -238,7 +238,7 @@ bool FRunMaterialEditorPromotionTestCommand::Update()
 /**
 * Automation test that handles the build promotion process
 */
-IMPLEMENT_SIMPLE_AUTOMATION_TEST(FMaterialEditorPromotionTest, "System.Promotion.Editor.Material Editor", EAutomationTestFlags::ATF_Editor);
+IMPLEMENT_SIMPLE_AUTOMATION_TEST(FMaterialEditorPromotionTest, "System.Promotion.Editor.Material Editor", EAutomationTestFlags::EditorContext | EAutomationTestFlags::EngineFilter);
 
 bool FMaterialEditorPromotionTest::RunTest(const FString& Parameters)
 {
diff --git a/Engine/Source/Editor/SceneOutliner/Private/SceneOutlinerSortingTests.cpp b/Engine/Source/Editor/SceneOutliner/Private/SceneOutlinerSortingTests.cpp
--- a/Engine/Source/Editor/SceneOutliner/Private/SceneOutlinerSortingTests.cpp
+++ b/Engine/Source/Editor/SceneOutliner/Private/SceneOutlinerSortingTests.cpp
@@ -3,7 +3,7 @@
 #include "SceneOutlinerPrivatePCH.h"
 #include "SortHelper.h"
 
-IMPLEMENT_SIMPLE_AUTOMATION_TEST(FSortBasicTest, "System.Editor.Scene Outliner.SortBasic", 0)
+IMPLEMENT_SIMPLE_AUTOMATION_TEST(FSortBasicTest, "System.Editor.Scene Outliner.SortBasic", EAutomationTestFlags::EditorContext | EAutomationTestFlags::EngineFilter)
 bool FSortBasicTest::RunTest(const FString& Parameters)
 {
 	TArray<SceneOutliner::FNumericStringWrapper> Muddled;
diff --git a/Engine/Source/Editor/StaticMeshEditor/Private/StaticMeshAutomationTests.cpp b/Engine/Source/Editor/StaticMeshEditor/Private/StaticMeshAutomationTests.cpp
--- a/Engine/Source/Editor/StaticMeshEditor/Private/StaticMeshAutomationTests.cpp
+++ b/Engine/Source/Editor/StaticMeshEditor/Private/StaticMeshAutomationTests.cpp
@@ -156,7 +156,7 @@ namespace EditorViewButtonHelper
 /**
  * Static mesh editor test
  */
-IMPLEMENT_SIMPLE_AUTOMATION_TEST( FStaticMeshEditorTest, "System.Editor.Content.Static Mesh Editor Test", EAutomationTestFlags::ATF_Editor )
+IMPLEMENT_SIMPLE_AUTOMATION_TEST(FStaticMeshEditorTest, "System.Editor.Content.Static Mesh Editor Test", EAutomationTestFlags::EditorContext | EAutomationTestFlags::EngineFilter)
 
 /**
  * Take screenshots of the SME window with each of the toolbar buttons toggled separately
diff --git a/Engine/Source/Editor/UnrealEd/Private/AutoReimport/AutoReimportTests.cpp b/Engine/Source/Editor/UnrealEd/Private/AutoReimport/AutoReimportTests.cpp
--- a/Engine/Source/Editor/UnrealEd/Private/AutoReimport/AutoReimportTests.cpp
+++ b/Engine/Source/Editor/UnrealEd/Private/AutoReimport/AutoReimportTests.cpp
@@ -130,7 +130,7 @@ bool CopyTestFiles(const FAutoReimportTestPayload& Test, const TArray<FSrcDstFil
 }
 
 /** Test that creating a new file gets reported correctly */
-IMPLEMENT_SIMPLE_AUTOMATION_TEST(FAutoReimportSimpleCreateTest, "System.Editor.Auto Reimport.Simple Create", EAutomationTestFlags::ATF_Editor)
+IMPLEMENT_SIMPLE_AUTOMATION_TEST(FAutoReimportSimpleCreateTest, "System.Editor.Auto Reimport.Simple Create", EAutomationTestFlags::EditorContext | EAutomationTestFlags::EngineFilter)
 bool FAutoReimportSimpleCreateTest::RunTest(const FString& Parameters)
 {
 	const FString WorkingDir = GetWorkingDir();
@@ -186,7 +186,7 @@ bool FAutoReimportSimpleCreateTest::RunTest(const FString& Parameters)
 }
 
 /** Test that modifying an existing file gets reported correctly */
-IMPLEMENT_SIMPLE_AUTOMATION_TEST(FAutoReimportSimpleModifyTest, "System.Editor.Auto Reimport.Simple Modify", EAutomationTestFlags::ATF_Editor)
+IMPLEMENT_SIMPLE_AUTOMATION_TEST(FAutoReimportSimpleModifyTest, "System.Editor.Auto Reimport.Simple Modify", EAutomationTestFlags::EditorContext | EAutomationTestFlags::EngineFilter)
 bool FAutoReimportSimpleModifyTest::RunTest(const FString& Parameters)
 {
 	const FString WorkingDir = GetWorkingDir();
@@ -258,7 +258,7 @@ bool FAutoReimportSimpleModifyTest::RunTest(const FString& Parameters)
 }
 
 /** Test that deleting an existing file gets reported correctly */
-IMPLEMENT_SIMPLE_AUTOMATION_TEST(FAutoReimportSimpleDeleteTest, "System.Editor.Auto Reimport.Simple Delete", EAutomationTestFlags::ATF_Editor)
+IMPLEMENT_SIMPLE_AUTOMATION_TEST(FAutoReimportSimpleDeleteTest, "System.Editor.Auto Reimport.Simple Delete", EAutomationTestFlags::EditorContext | EAutomationTestFlags::EngineFilter)
 bool FAutoReimportSimpleDeleteTest::RunTest(const FString& Parameters)
 {
 	const FString WorkingDir = GetWorkingDir();
@@ -323,7 +323,7 @@ bool FAutoReimportSimpleDeleteTest::RunTest(const FString& Parameters)
 }
 
 /** Test that renaming an existing file gets reported correctly */
-IMPLEMENT_SIMPLE_AUTOMATION_TEST(FAutoReimportSimpleRenameTest, "System.Editor.Auto Reimport.Simple Rename", EAutomationTestFlags::ATF_Editor)
+IMPLEMENT_SIMPLE_AUTOMATION_TEST(FAutoReimportSimpleRenameTest, "System.Editor.Auto Reimport.Simple Rename", EAutomationTestFlags::EditorContext | EAutomationTestFlags::EngineFilter)
 bool FAutoReimportSimpleRenameTest::RunTest(const FString& Parameters)
 {
 	const FString WorkingDir = GetWorkingDir();
@@ -388,7 +388,7 @@ bool FAutoReimportSimpleRenameTest::RunTest(const FString& Parameters)
 }
 
 /** Test that moving a file outside of the monitored directory gets reported correctly (should be reported as a delete) */
-IMPLEMENT_SIMPLE_AUTOMATION_TEST(FAutoReimportSimpleMoveExternallyTest, "System.Editor.Auto Reimport.Move Externally", EAutomationTestFlags::ATF_Editor)
+IMPLEMENT_SIMPLE_AUTOMATION_TEST(FAutoReimportSimpleMoveExternallyTest, "System.Editor.Auto Reimport.Move Externally", EAutomationTestFlags::EditorContext | EAutomationTestFlags::EngineFilter)
 bool FAutoReimportSimpleMoveExternallyTest::RunTest(const FString& Parameters)
 {
 	const FString WorkingDir = GetWorkingDir();
@@ -454,7 +454,7 @@ bool FAutoReimportSimpleMoveExternallyTest::RunTest(const FString& Parameters)
 }
 
 /** Test that bDetectChangesSinceLastRun works correctly when true and when false */
-IMPLEMENT_SIMPLE_AUTOMATION_TEST(FAutoReimportRestartDetectionTest, "System.Editor.Auto Reimport.Restart Detection", EAutomationTestFlags::ATF_Editor)
+IMPLEMENT_SIMPLE_AUTOMATION_TEST(FAutoReimportRestartDetectionTest, "System.Editor.Auto Reimport.Restart Detection", EAutomationTestFlags::EditorContext | EAutomationTestFlags::EngineFilter)
 bool FAutoReimportRestartDetectionTest::RunTest(const FString& Parameters)
 {
 	const FString WorkingDir = GetWorkingDir();
@@ -553,7 +553,7 @@ bool FAutoReimportRestartDetectionTest::RunTest(const FString& Parameters)
 }
 
 /** Test that making multiple changes to the same file gets picked up correctly */
-IMPLEMENT_SIMPLE_AUTOMATION_TEST(FAutoReimportMultipleChangesTest, "System.Editor.Auto Reimport.Multiple Changes", EAutomationTestFlags::ATF_Editor)
+IMPLEMENT_SIMPLE_AUTOMATION_TEST(FAutoReimportMultipleChangesTest, "System.Editor.Auto Reimport.Multiple Changes", EAutomationTestFlags::EditorContext | EAutomationTestFlags::EngineFilter)
 bool FAutoReimportMultipleChangesTest::RunTest(const FString& Parameters)
 {
 	const FString WorkingDir = GetWorkingDir();
@@ -609,7 +609,7 @@ bool FAutoReimportMultipleChangesTest::RunTest(const FString& Parameters)
 }
 
 /** Test that starting up a cache file with a different set of applicable extensions correctly ignores/updates the extensions, whilst reporting changes only for applicable extensions */
-IMPLEMENT_SIMPLE_AUTOMATION_TEST(FAutoReimportChangeExtensionsTest, "System.Editor.Auto Reimport.Change Extensions", EAutomationTestFlags::ATF_Editor)
+IMPLEMENT_SIMPLE_AUTOMATION_TEST(FAutoReimportChangeExtensionsTest, "System.Editor.Auto Reimport.Change Extensions", EAutomationTestFlags::EditorContext | EAutomationTestFlags::EngineFilter)
 bool FAutoReimportChangeExtensionsTest::RunTest(const FString& Parameters)
 {
 	const FString WorkingDir = GetWorkingDir();
@@ -701,7 +701,7 @@ bool FAutoReimportChangeExtensionsTest::RunTest(const FString& Parameters)
 }
 
 /** Test that starting up a cache file with a different set of applicable extensions correctly ignores/updates the extensions, whilst reporting changes only for applicable extensions */
-IMPLEMENT_SIMPLE_AUTOMATION_TEST(FAutoReimportWildcardFiltersTest, "System.Editor.Auto Reimport.Wildcard Filters", EAutomationTestFlags::ATF_Editor)
+IMPLEMENT_SIMPLE_AUTOMATION_TEST(FAutoReimportWildcardFiltersTest, "System.Editor.Auto Reimport.Wildcard Filters", EAutomationTestFlags::EditorContext | EAutomationTestFlags::EngineFilter)
 bool FAutoReimportWildcardFiltersTest::RunTest(const FString& Parameters)
 {
 	const FString WorkingDir = GetWorkingDir();
diff --git a/Engine/Source/Editor/UnrealEd/Private/Tests/BlueprintAutomationTests.cpp b/Engine/Source/Editor/UnrealEd/Private/Tests/BlueprintAutomationTests.cpp
--- a/Engine/Source/Editor/UnrealEd/Private/Tests/BlueprintAutomationTests.cpp
+++ b/Engine/Source/Editor/UnrealEd/Private/Tests/BlueprintAutomationTests.cpp
@@ -18,12 +18,12 @@
 
 DEFINE_LOG_CATEGORY_STATIC(LogBlueprintAutomationTests, Log, All);
 
-IMPLEMENT_COMPLEX_AUTOMATION_TEST( FBlueprintCompileOnLoadTest, "Project.Blueprints.Compile-On-Load", EAutomationTestFlags::ATF_Editor )
-IMPLEMENT_COMPLEX_AUTOMATION_TEST( FBlueprintInstancesTest, "Project.Blueprints.Instance Test", EAutomationTestFlags::ATF_Editor )
-IMPLEMENT_COMPLEX_AUTOMATION_TEST( FBlueprintReparentTest, "System.Blueprints.Reparent", EAutomationTestFlags::ATF_Editor )
-IMPLEMENT_COMPLEX_AUTOMATION_TEST( FBlueprintRenameAndCloneTest, "Project.Blueprints.Rename And Clone", (EAutomationTestFlags::ATF_Editor | EAutomationTestFlags::ATF_RequiresUser) )
-IMPLEMENT_COMPLEX_AUTOMATION_TEST( FCompileBlueprintsTest, "Project.Blueprints.Compile Blueprints", EAutomationTestFlags::ATF_Editor )
-IMPLEMENT_COMPLEX_AUTOMATION_TEST( FCompileAnimBlueprintsTest, "Project.Blueprints.Compile Anims", EAutomationTestFlags::ATF_Editor )
+IMPLEMENT_COMPLEX_AUTOMATION_TEST(FBlueprintCompileOnLoadTest, "Project.Blueprints.Compile-On-Load", EAutomationTestFlags::EditorContext | EAutomationTestFlags::StressFilter)
+IMPLEMENT_COMPLEX_AUTOMATION_TEST(FBlueprintInstancesTest, "Project.Blueprints.Instance Test", EAutomationTestFlags::EditorContext | EAutomationTestFlags::StressFilter)
+IMPLEMENT_COMPLEX_AUTOMATION_TEST(FBlueprintReparentTest, "System.Blueprints.Reparent", EAutomationTestFlags::EditorContext | EAutomationTestFlags::StressFilter)
+IMPLEMENT_COMPLEX_AUTOMATION_TEST(FBlueprintRenameAndCloneTest, "Project.Blueprints.Rename And Clone", (EAutomationTestFlags::EditorContext | EAutomationTestFlags::RequiresUser | EAutomationTestFlags::StressFilter))
+IMPLEMENT_COMPLEX_AUTOMATION_TEST(FCompileBlueprintsTest, "Project.Blueprints.Compile Blueprints", EAutomationTestFlags::EditorContext | EAutomationTestFlags::StressFilter)
+IMPLEMENT_COMPLEX_AUTOMATION_TEST(FCompileAnimBlueprintsTest, "Project.Blueprints.Compile Anims", EAutomationTestFlags::EditorContext | EAutomationTestFlags::StressFilter)
 
 class FBlueprintAutomationTestUtilities
 {
diff --git a/Engine/Source/Editor/UnrealEd/Private/Tests/CollisionAutomationTests.cpp b/Engine/Source/Editor/UnrealEd/Private/Tests/CollisionAutomationTests.cpp
--- a/Engine/Source/Editor/UnrealEd/Private/Tests/CollisionAutomationTests.cpp
+++ b/Engine/Source/Editor/UnrealEd/Private/Tests/CollisionAutomationTests.cpp
@@ -30,7 +30,7 @@ namespace CollisionAutomationTests
 	UWorld* GetAutomationWorld(const int32 TestFlags)
  	{
 		UWorld* World = nullptr;
-		if( TestFlags & EAutomationTestFlags::ATF_Game)
+		if( TestFlags & EAutomationTestFlags::ClientContext)
 		{
 			check(GEngine->GetWorldContexts().Num() == 1);
 			World = GEngine->GetWorldContexts()[0].World();
@@ -135,7 +135,7 @@ namespace CollisionAutomationTests
 /**
  * ComponentSweepMultiTest Verification
  */
-IMPLEMENT_SIMPLE_AUTOMATION_TEST( FComponentSweepMultiTest, "System.Physics.Collision.ComponentSweepMulti", EAutomationTestFlags::ATF_Editor )
+IMPLEMENT_SIMPLE_AUTOMATION_TEST(FComponentSweepMultiTest, "System.Physics.Collision.ComponentSweepMulti", EAutomationTestFlags::EditorContext | EAutomationTestFlags::EngineFilter)
 
 /** 
  * Perform some collision sweep tests. Creates a given shape mesh and checks collision normal against a collision shape type.
@@ -223,7 +223,7 @@ bool FComponentSweepMultiTest::RunTest(const FString& Parameters)
 	return true;
 }
 
-IMPLEMENT_SIMPLE_AUTOMATION_TEST( FLineTraceSingleByChannel, "System.Physics.Collision.LineTraceSingleByChannel", EAutomationTestFlags::ATF_Editor )
+IMPLEMENT_SIMPLE_AUTOMATION_TEST(FLineTraceSingleByChannel, "System.Physics.Collision.LineTraceSingleByChannel", EAutomationTestFlags::EditorContext | EAutomationTestFlags::EngineFilter)
 
 
 /** 
diff --git a/Engine/Source/Editor/UnrealEd/Private/Tests/EditorAssetAutomationTests.cpp b/Engine/Source/Editor/UnrealEd/Private/Tests/EditorAssetAutomationTests.cpp
--- a/Engine/Source/Editor/UnrealEd/Private/Tests/EditorAssetAutomationTests.cpp
+++ b/Engine/Source/Editor/UnrealEd/Private/Tests/EditorAssetAutomationTests.cpp
@@ -461,7 +461,7 @@ void AssetTestCreate(const TCHAR* NamePrefix, const FString& CurrentTimestamp, c
 /**
 * Automation test for creating, saving, and duplicating assets
 */
-IMPLEMENT_SIMPLE_AUTOMATION_TEST(FAssetEditorTest, "System.Editor.Content.Asset Creation and Duplication", EAutomationTestFlags::ATF_Editor);
+IMPLEMENT_SIMPLE_AUTOMATION_TEST(FAssetEditorTest, "System.Editor.Content.Asset Creation and Duplication", EAutomationTestFlags::EditorContext | EAutomationTestFlags::EngineFilter);
 
 bool FAssetEditorTest::RunTest(const FString& Parameters)
 {
@@ -1067,7 +1067,7 @@ bool FLogImportExportTestResultsCommand::Update()
 /**
 * Automation test to import, open, screenshot, and export assets
 */
-IMPLEMENT_SIMPLE_AUTOMATION_TEST(FAssetImportEditorTest, "Project.Editor.Content.Asset Import and Export", EAutomationTestFlags::ATF_Editor);
+IMPLEMENT_SIMPLE_AUTOMATION_TEST(FAssetImportEditorTest, "Project.Editor.Content.Asset Import and Export", EAutomationTestFlags::EditorContext | EAutomationTestFlags::EngineFilter);
 bool FAssetImportEditorTest::RunTest(const FString& Parameters)
 {
 	UAutomationTestSettings const* AutomationTestSettings = GetDefault<UAutomationTestSettings>();
diff --git a/Engine/Source/Editor/UnrealEd/Private/Tests/EditorAutomationTests.cpp b/Engine/Source/Editor/UnrealEd/Private/Tests/EditorAutomationTests.cpp
--- a/Engine/Source/Editor/UnrealEd/Private/Tests/EditorAutomationTests.cpp
+++ b/Engine/Source/Editor/UnrealEd/Private/Tests/EditorAutomationTests.cpp
@@ -114,7 +114,7 @@ void TakeLatentAutomationScreenshot(struct WindowScreenshotParameters Screenshot
  * in the Content Browser (does not allow for specific settings to be made per import factory). Cannot be run in a commandlet
  * as it executes code that routes through Slate UI.
  */
-IMPLEMENT_COMPLEX_AUTOMATION_TEST( FGenericImportAssetsAutomationTest, "System.Editor.Import", (EAutomationTestFlags::ATF_Editor | EAutomationTestFlags::ATF_NonNullRHI) )
+IMPLEMENT_COMPLEX_AUTOMATION_TEST(FGenericImportAssetsAutomationTest, "System.Editor.Import", (EAutomationTestFlags::EditorContext | EAutomationTestFlags::NonNullRHI | EAutomationTestFlags::EngineFilter))
 
 /** 
  * Requests a enumeration of all sample assets to import
@@ -193,7 +193,7 @@ bool FGenericImportAssetsAutomationTest::RunTest(const FString& Parameters)
  * Pie Test
  * Verification PIE works
  */
-IMPLEMENT_SIMPLE_AUTOMATION_TEST( FPIETest, "System.Maps.PIE", EAutomationTestFlags::ATF_Editor )
+IMPLEMENT_SIMPLE_AUTOMATION_TEST(FPIETest, "System.Maps.PIE", EAutomationTestFlags::EditorContext | EAutomationTestFlags::EngineFilter)
 
 /** 
  * Execute the loading of one map to verify PIE works
@@ -224,7 +224,7 @@ bool FPIETest::RunTest(const FString& Parameters)
  * LoadAllMaps
  * Verification automation test to make sure loading all maps succeed without crashing
  */
-IMPLEMENT_COMPLEX_AUTOMATION_TEST( FLoadAllMapsInEditorTest, "Project.Maps.Load All In Editor", EAutomationTestFlags::ATF_Editor )
+IMPLEMENT_COMPLEX_AUTOMATION_TEST(FLoadAllMapsInEditorTest, "Project.Maps.Load All In Editor", EAutomationTestFlags::EditorContext | EAutomationTestFlags::StressFilter)
 
 /** 
  * Requests a enumeration of all maps to be loaded
@@ -331,7 +331,7 @@ bool FLoadAllMapsInEditorTest::RunTest(const FString& Parameters)
 /**
  * Reinitialize all RHI resources
  */
-IMPLEMENT_SIMPLE_AUTOMATION_TEST(FReinitializeRHIResources, "System.Engine.Rendering.Reinit Resources", EAutomationTestFlags::ATF_Editor);
+IMPLEMENT_SIMPLE_AUTOMATION_TEST(FReinitializeRHIResources, "System.Engine.Rendering.Reinit Resources", EAutomationTestFlags::EditorContext | EAutomationTestFlags::EngineFilter);
 
 bool FReinitializeRHIResources::RunTest(const FString& Parameters)
 {
@@ -344,7 +344,7 @@ bool FReinitializeRHIResources::RunTest(const FString& Parameters)
 /**
  * QA Static Mesh Regression Testing
  */
-IMPLEMENT_SIMPLE_AUTOMATION_TEST(FStaticMeshValidation, "System.QA.Mesh Factory Validation", EAutomationTestFlags::ATF_Editor);
+IMPLEMENT_SIMPLE_AUTOMATION_TEST(FStaticMeshValidation, "System.QA.Mesh Factory Validation", EAutomationTestFlags::EditorContext | EAutomationTestFlags::EngineFilter);
 
 bool FStaticMeshValidation::RunTest(const FString& Parameters)
 {
@@ -420,7 +420,7 @@ bool FStaticMeshValidation::RunTest(const FString& Parameters)
 /**
  * QA Convert Meshes Regression Testing
  */
-IMPLEMENT_SIMPLE_AUTOMATION_TEST(FConvertToValidation, "System.QA.Convert Meshes", EAutomationTestFlags::ATF_Editor);
+IMPLEMENT_SIMPLE_AUTOMATION_TEST(FConvertToValidation, "System.QA.Convert Meshes", EAutomationTestFlags::EditorContext | EAutomationTestFlags::EngineFilter);
 
 //gather all brushes that exist now
 void ConvertTestFindAllBrushes(TArray<ABrush*> &PreviousBrushes)
@@ -644,7 +644,7 @@ bool FCleanupConvertToValidation::Update()
 /**
  * QA Static Mesh Regression Testing
  */
-IMPLEMENT_SIMPLE_AUTOMATION_TEST(FStaticMeshPlacement, "System.QA.Static Mesh Placement", EAutomationTestFlags::ATF_Editor);
+IMPLEMENT_SIMPLE_AUTOMATION_TEST(FStaticMeshPlacement, "System.QA.Static Mesh Placement", EAutomationTestFlags::EditorContext | EAutomationTestFlags::EngineFilter);
 
 bool FStaticMeshPlacement::RunTest(const FString& Parameters)
 {
@@ -698,7 +698,7 @@ bool FStaticMeshPlacement::RunTest(const FString& Parameters)
 /**
  * QA Light Placement Regression Testing
  */
-IMPLEMENT_SIMPLE_AUTOMATION_TEST(FLightPlacement, "System.QA.Point Light Placement", EAutomationTestFlags::ATF_Editor);
+IMPLEMENT_SIMPLE_AUTOMATION_TEST(FLightPlacement, "System.QA.Point Light Placement", EAutomationTestFlags::EditorContext | EAutomationTestFlags::EngineFilter);
 
 bool FLightPlacement::RunTest(const FString& Parameters)
 {
@@ -797,7 +797,7 @@ bool FLightPlacement::RunTest(const FString& Parameters)
  * Unit test to find all timelines in blueprints and list the events that can trigger them.
  * Timelines implicitly tick and are usually used for cosmetic events, so they can cause performance problems on dedicated servers.
  */
-IMPLEMENT_COMPLEX_AUTOMATION_TEST( FTraceAllTimelinesAutomationTest, "Project.Performance Audits.Find Timelines On Server", (EAutomationTestFlags::ATF_Editor | EAutomationTestFlags::ATF_RequiresUser))
+IMPLEMENT_COMPLEX_AUTOMATION_TEST(FTraceAllTimelinesAutomationTest, "Project.Performance Audits.Find Timelines On Server", (EAutomationTestFlags::EditorContext | EAutomationTestFlags::RequiresUser | EAutomationTestFlags::StressFilter))
 
 /** 
  * Requests an enumeration of all blueprints to be loaded
@@ -1192,7 +1192,7 @@ bool FTraceAllTimelinesAutomationTest::RunTest(const FString& BlueprintName)
 /**
 * Tool to look for overlapping UV's in static meshes.
 */
-IMPLEMENT_COMPLEX_AUTOMATION_TEST(FStaticMeshUVCheck, "Project.Tools.Static Mesh.Static Mesh UVs Check", (EAutomationTestFlags::ATF_Editor | EAutomationTestFlags::ATF_RequiresUser));
+IMPLEMENT_COMPLEX_AUTOMATION_TEST(FStaticMeshUVCheck, "Project.Tools.Static Mesh.Static Mesh UVs Check", (EAutomationTestFlags::EditorContext | EAutomationTestFlags::RequiresUser | EAutomationTestFlags::StressFilter));
 
 void FStaticMeshUVCheck::GetTests(TArray<FString>& OutBeautifiedNames, TArray<FString>& OutTestCommands) const
 {
@@ -1222,7 +1222,7 @@ bool FStaticMeshUVCheck::RunTest(const FString& Parameters)
 /**
 * Launches a map onto a specified device after making a change to it.
 */
-IMPLEMENT_COMPLEX_AUTOMATION_TEST(FLaunchOnTest, "Project.Editor.Launch On Test", (EAutomationTestFlags::ATF_Editor | EAutomationTestFlags::ATF_RequiresUser))
+IMPLEMENT_COMPLEX_AUTOMATION_TEST(FLaunchOnTest, "Project.Editor.Launch On Test", (EAutomationTestFlags::EditorContext | EAutomationTestFlags::RequiresUser | EAutomationTestFlags::EngineFilter))
 
 void FLaunchOnTest::GetTests(TArray<FString>& OutBeautifiedNames, TArray <FString>& OutTestCommands) const
 {
diff --git a/Engine/Source/Editor/UnrealEd/Private/Tests/EditorBuildPromotionTests.cpp b/Engine/Source/Editor/UnrealEd/Private/Tests/EditorBuildPromotionTests.cpp
--- a/Engine/Source/Editor/UnrealEd/Private/Tests/EditorBuildPromotionTests.cpp
+++ b/Engine/Source/Editor/UnrealEd/Private/Tests/EditorBuildPromotionTests.cpp
@@ -3660,7 +3660,7 @@ namespace BuildPromotionTestHelper
 /**
 * Automation test that handles cleanup of the build promotion test
 */
-IMPLEMENT_SIMPLE_AUTOMATION_TEST(FBuildPromotionInitialCleanupTest, "System.Promotion.Editor Promotion Pass.Step 1 Main Editor Test.Cleanup old files", EAutomationTestFlags::ATF_Editor);
+IMPLEMENT_SIMPLE_AUTOMATION_TEST(FBuildPromotionInitialCleanupTest, "System.Promotion.Editor Promotion Pass.Step 1 Main Editor Test.Cleanup old files", EAutomationTestFlags::EditorContext | EAutomationTestFlags::EngineFilter);
 bool FBuildPromotionInitialCleanupTest::RunTest(const FString& Parameters)
 {
 	EditorBuildPromotionTestUtils::PerformCleanup();
@@ -3706,7 +3706,7 @@ bool FLogTestResultCommand::Update()
 /**
 * Automation test that handles setting keybindings and editor preferences
 */
-IMPLEMENT_SIMPLE_AUTOMATION_TEST(FBuildPromotionSettingsTest, "System.Promotion.Editor Promotion Pass.Step 1 Main Editor Test.Settings", EAutomationTestFlags::ATF_Editor);
+IMPLEMENT_SIMPLE_AUTOMATION_TEST(FBuildPromotionSettingsTest, "System.Promotion.Editor Promotion Pass.Step 1 Main Editor Test.Settings", EAutomationTestFlags::EditorContext | EAutomationTestFlags::EngineFilter);
 
 bool FBuildPromotionSettingsTest::RunTest(const FString& Parameters)
 {
@@ -3795,7 +3795,7 @@ bool FRunBuildPromotionTestCommand::Update()
 /**
 * Automation test that handles the build promotion process
 */
-IMPLEMENT_SIMPLE_AUTOMATION_TEST(FBuildPromotionTest, "System.Promotion.Editor Promotion Pass.Step 1 Main Editor Test.General Editor Test", EAutomationTestFlags::ATF_Editor);
+IMPLEMENT_SIMPLE_AUTOMATION_TEST(FBuildPromotionTest, "System.Promotion.Editor Promotion Pass.Step 1 Main Editor Test.General Editor Test", EAutomationTestFlags::EditorContext | EAutomationTestFlags::EngineFilter);
 
 bool FBuildPromotionTest::RunTest(const FString& Parameters)
 {
@@ -3835,7 +3835,7 @@ bool FPIEExecCommand::Update()
 /**
 * Execute the loading of one map to verify PIE works
 */
-IMPLEMENT_SIMPLE_AUTOMATION_TEST(FBuildPromotionPIETest, "System.Promotion.Editor Promotion Pass.Step 2 Run Map After Re-launch.Run Map", EAutomationTestFlags::ATF_Editor)
+IMPLEMENT_SIMPLE_AUTOMATION_TEST(FBuildPromotionPIETest, "System.Promotion.Editor Promotion Pass.Step 2 Run Map After Re-launch.Run Map", EAutomationTestFlags::EditorContext | EAutomationTestFlags::EngineFilter)
 bool FBuildPromotionPIETest::RunTest(const FString& Parameters)
 {
 	UAutomationTestSettings const* AutomationTestSettings = GetDefault<UAutomationTestSettings>();
@@ -3882,7 +3882,7 @@ bool FBuildPromotionPIETest::RunTest(const FString& Parameters)
 /**
 * Automation test that handles cleanup of the build promotion test
 */
-IMPLEMENT_SIMPLE_AUTOMATION_TEST(FBuildPromotionCleanupTest, "System.Promotion.Editor Promotion Pass.Step 3 Test Cleanup.Cleanup", EAutomationTestFlags::ATF_Editor);
+IMPLEMENT_SIMPLE_AUTOMATION_TEST(FBuildPromotionCleanupTest, "System.Promotion.Editor Promotion Pass.Step 3 Test Cleanup.Cleanup", EAutomationTestFlags::EditorContext | EAutomationTestFlags::EngineFilter);
 bool FBuildPromotionCleanupTest::RunTest(const FString& Parameters)
 {
 	EditorBuildPromotionTestUtils::PerformCleanup();
diff --git a/Engine/Source/Editor/UnrealEd/Private/Tests/EditorOpenAssetAutomationTests.cpp b/Engine/Source/Editor/UnrealEd/Private/Tests/EditorOpenAssetAutomationTests.cpp
--- a/Engine/Source/Editor/UnrealEd/Private/Tests/EditorOpenAssetAutomationTests.cpp
+++ b/Engine/Source/Editor/UnrealEd/Private/Tests/EditorOpenAssetAutomationTests.cpp
@@ -34,7 +34,7 @@ bool FOpenActualAssetEditors(const FString& Parameters)
 * Test to open the sub editor windows for a specified list of assets.
 * This list can be setup in the Editor Preferences window within the editor or the DefaultEngine.ini file for that particular project.
 */
-IMPLEMENT_COMPLEX_AUTOMATION_TEST(FOpenAssetEditors, "System.QA.Open Asset Editors", EAutomationTestFlags::ATF_Editor);
+IMPLEMENT_COMPLEX_AUTOMATION_TEST(FOpenAssetEditors, "System.QA.Open Asset Editors", EAutomationTestFlags::EditorContext | EAutomationTestFlags::StressFilter);
 
 void FOpenAssetEditors::GetTests(TArray<FString>& OutBeautifiedNames, TArray <FString>& OutTestCommands) const
 {
@@ -216,7 +216,7 @@ void FOpenAllAssetTests(OpenAssetParameters& AssetList)
 /**
 * This is a complex test that opens every asset into it's editor.
 */
-IMPLEMENT_COMPLEX_AUTOMATION_TEST(FOpenMiscAssetEditors, "Project.Tools.Open Assets", (EAutomationTestFlags::ATF_Editor | EAutomationTestFlags::ATF_RequiresUser));
+IMPLEMENT_COMPLEX_AUTOMATION_TEST(FOpenMiscAssetEditors, "Project.Tools.Open Assets", (EAutomationTestFlags::EditorContext | EAutomationTestFlags::StressFilter | EAutomationTestFlags::RequiresUser));
 
 void FOpenMiscAssetEditors::GetTests(TArray<FString>& OutBeautifiedNames, TArray<FString>& OutTestCommands) const
 {
@@ -229,4 +229,4 @@ bool FOpenMiscAssetEditors::RunTest(const FString& Parameters)
 	bool bDidTheTestPass = FOpenActualAssetEditors(Parameters);
 	CollectGarbage(GARBAGE_COLLECTION_KEEPFLAGS);
 	return bDidTheTestPass;
-}
\ No newline at end of file
+}
diff --git a/Engine/Source/Editor/UnrealEd/Private/Tests/EditorPerformanceTests.cpp b/Engine/Source/Editor/UnrealEd/Private/Tests/EditorPerformanceTests.cpp
--- a/Engine/Source/Editor/UnrealEd/Private/Tests/EditorPerformanceTests.cpp
+++ b/Engine/Source/Editor/UnrealEd/Private/Tests/EditorPerformanceTests.cpp
@@ -219,7 +219,7 @@ bool FEditorPerfCaptureCommand::Update()
 * Map Performance in Editor tests
 * Grabs certain performance numbers and saves it to a file.
 */
-IMPLEMENT_COMPLEX_AUTOMATION_TEST(FMapPerformanceInEditor, "Project.Performance.Map Performance in Editor", EAutomationTestFlags::ATF_Editor | EAutomationTestFlags::ATF_RequiresUser)
+IMPLEMENT_COMPLEX_AUTOMATION_TEST(FMapPerformanceInEditor, "Project.Performance.Map Performance in Editor", EAutomationTestFlags::EditorContext | EAutomationTestFlags::RequiresUser | EAutomationTestFlags::EngineFilter)
 
 /**
 * Requests a enumeration of all maps to be loaded
diff --git a/Engine/Source/Editor/UnrealEd/Private/Tests/EditorViewportFlagsTest.cpp b/Engine/Source/Editor/UnrealEd/Private/Tests/EditorViewportFlagsTest.cpp
--- a/Engine/Source/Editor/UnrealEd/Private/Tests/EditorViewportFlagsTest.cpp
+++ b/Engine/Source/Editor/UnrealEd/Private/Tests/EditorViewportFlagsTest.cpp
@@ -99,7 +99,7 @@ namespace FViewportTestHelper
 /**
 * Unit test for enabling the view ports show flags
 */
-IMPLEMENT_COMPLEX_AUTOMATION_TEST( FViewportShowflagsToggleOnTest, "System.Editor.Viewport.Showflags.Toggle On", EAutomationTestFlags::ATF_Editor );
+IMPLEMENT_COMPLEX_AUTOMATION_TEST( FViewportShowflagsToggleOnTest, "System.Editor.Viewport.Showflags.Toggle On", EAutomationTestFlags::EditorContext | EAutomationTestFlags::EngineFilter );
 
 void FViewportShowflagsToggleOnTest::GetTests(TArray<FString>& OutBeautifiedNames, TArray<FString>& OutTestCommands) const
 {
@@ -134,7 +134,7 @@ bool FViewportShowflagsToggleOnTest::RunTest( const FString& Parameters )
 /**
 * Unit test for disabling the view ports show flags
 */
-IMPLEMENT_COMPLEX_AUTOMATION_TEST(FViewportShowflagsToggleOffTest, "System.Editor.Viewport.Showflags.Toggle Off", EAutomationTestFlags::ATF_Editor);
+IMPLEMENT_COMPLEX_AUTOMATION_TEST(FViewportShowflagsToggleOffTest, "System.Editor.Viewport.Showflags.Toggle Off", EAutomationTestFlags::EditorContext | EAutomationTestFlags::EngineFilter);
 
 void FViewportShowflagsToggleOffTest::GetTests(TArray<FString>& OutBeautifiedNames, TArray<FString>& OutTestCommands) const
 {
diff --git a/Engine/Source/Editor/UnrealEd/Private/Tests/FeaturePackAutomationTests.cpp b/Engine/Source/Editor/UnrealEd/Private/Tests/FeaturePackAutomationTests.cpp
--- a/Engine/Source/Editor/UnrealEd/Private/Tests/FeaturePackAutomationTests.cpp
+++ b/Engine/Source/Editor/UnrealEd/Private/Tests/FeaturePackAutomationTests.cpp
@@ -25,7 +25,7 @@ DEFINE_LOG_CATEGORY_STATIC(FeaturePackAutomationTestLog, Log, All);
 /**
  * Feature Pack validity test
  */
-IMPLEMENT_SIMPLE_AUTOMATION_TEST( FFeaturePackValidityTest, "System.Editor.Content.FeaturePackValidityTest", EAutomationTestFlags::ATF_Editor )
+IMPLEMENT_SIMPLE_AUTOMATION_TEST(FFeaturePackValidityTest, "System.Editor.Content.FeaturePackValidityTest", EAutomationTestFlags::EditorContext | EAutomationTestFlags::EngineFilter)
 
 /** 
  * Scans feature pack folder and verifies that all feature packs therein are parsed successfully.
diff --git a/Engine/Source/Editor/UnrealEd/Private/Tests/GeometryTests.cpp b/Engine/Source/Editor/UnrealEd/Private/Tests/GeometryTests.cpp
--- a/Engine/Source/Editor/UnrealEd/Private/Tests/GeometryTests.cpp
+++ b/Engine/Source/Editor/UnrealEd/Private/Tests/GeometryTests.cpp
@@ -77,7 +77,7 @@ namespace GeometryPromotionHelpers
 * There will be a few subtractive brushes intersecting with other objects.
 * Tests for number of surfaces and vertices, numbers of brushes created, and brush location.
 */
-IMPLEMENT_SIMPLE_AUTOMATION_TEST(FGeometryPromotionValidation, "System.Promotion.Editor.Geometry Validation", EAutomationTestFlags::ATF_Editor);
+IMPLEMENT_SIMPLE_AUTOMATION_TEST(FGeometryPromotionValidation, "System.Promotion.Editor.Geometry Validation", EAutomationTestFlags::EditorContext | EAutomationTestFlags::EngineFilter);
 
 bool FGeometryPromotionValidation::RunTest(const FString& Parameters)
 {
@@ -277,7 +277,7 @@ bool FGeometryPromotionValidation::RunTest(const FString& Parameters)
 * Subtractive brush intersects will the additive brush.
 * Tests for number of surfaces and vertices, numbers of brushes created, and brush location.
 */
-IMPLEMENT_SIMPLE_AUTOMATION_TEST(FBSPCubePlacement, "System.Editor.Geometry.BSP Cube.Brush Placement", EAutomationTestFlags::ATF_Editor);
+IMPLEMENT_SIMPLE_AUTOMATION_TEST(FBSPCubePlacement, "System.Editor.Geometry.BSP Cube.Brush Placement", EAutomationTestFlags::EditorContext | EAutomationTestFlags::EngineFilter);
 
 bool FBSPCubePlacement::RunTest(const FString& Parameters)
 {
@@ -365,7 +365,7 @@ bool FBSPCubePlacement::RunTest(const FString& Parameters)
 * Adds an additive brush to the world and then does an undo/redo of that transaction.
 * Tests for number vertices in the level.
 */
-IMPLEMENT_SIMPLE_AUTOMATION_TEST(FBSPCubeUndoRedoPlacement, "System.Editor.Geometry.BSP Cube.Brush Placement Undo Redo", EAutomationTestFlags::ATF_Editor);
+IMPLEMENT_SIMPLE_AUTOMATION_TEST(FBSPCubeUndoRedoPlacement, "System.Editor.Geometry.BSP Cube.Brush Placement Undo Redo", EAutomationTestFlags::EditorContext | EAutomationTestFlags::EngineFilter);
 
 bool FBSPCubeUndoRedoPlacement::RunTest(const FString& Parameters)
 {
diff --git a/Engine/Source/Editor/UnrealEd/Private/Tests/LightingTests.cpp b/Engine/Source/Editor/UnrealEd/Private/Tests/LightingTests.cpp
--- a/Engine/Source/Editor/UnrealEd/Private/Tests/LightingTests.cpp
+++ b/Engine/Source/Editor/UnrealEd/Private/Tests/LightingTests.cpp
@@ -112,7 +112,7 @@ namespace LightingTestHelpers
 /**
 * Lighting Promotion Test - Place a Point Light and move it to a new location.
 */
-IMPLEMENT_SIMPLE_AUTOMATION_TEST(FLightingPromotionPointLightPlaceRotScaleTest, "System.Promotion.Editor.Lighting.Place Scale Rotate", EAutomationTestFlags::ATF_Editor)
+IMPLEMENT_SIMPLE_AUTOMATION_TEST(FLightingPromotionPointLightPlaceRotScaleTest, "System.Promotion.Editor.Lighting.Place Scale Rotate", EAutomationTestFlags::EditorContext | EAutomationTestFlags::EngineFilter)
 
 bool FLightingPromotionPointLightPlaceRotScaleTest::RunTest(const FString& Parameters)
 {
@@ -157,7 +157,7 @@ bool FLightingPromotionPointLightPlaceRotScaleTest::RunTest(const FString& Param
 /**
 * Lighting Promotion Test - Modify a point lights properties.
 */
-IMPLEMENT_SIMPLE_AUTOMATION_TEST(FLightingPromotionModifyProperties, "System.Promotion.Editor.Lighting.Modify Properties", EAutomationTestFlags::ATF_Editor)
+IMPLEMENT_SIMPLE_AUTOMATION_TEST(FLightingPromotionModifyProperties, "System.Promotion.Editor.Lighting.Modify Properties", EAutomationTestFlags::EditorContext | EAutomationTestFlags::EngineFilter)
 
 bool FLightingPromotionModifyProperties::RunTest(const FString& Parameters)
 {
@@ -193,7 +193,7 @@ bool FLightingPromotionModifyProperties::RunTest(const FString& Parameters)
 /**
 * Lighting Promotion Test - Duplicate/Copy Paste a point light.
 */
-IMPLEMENT_SIMPLE_AUTOMATION_TEST(FLightingPromotionDuplicationTest, "System.Promotion.Editor.Lighting.Duplicate and Copy Paste", EAutomationTestFlags::ATF_Editor)
+IMPLEMENT_SIMPLE_AUTOMATION_TEST(FLightingPromotionDuplicationTest, "System.Promotion.Editor.Lighting.Duplicate and Copy Paste", EAutomationTestFlags::EditorContext | EAutomationTestFlags::EngineFilter)
 
 bool FLightingPromotionDuplicationTest::RunTest(const FString& Parameters)
 {
@@ -245,7 +245,7 @@ bool FLightingPromotionDuplicationTest::RunTest(const FString& Parameters)
 * Place a point light in the world with it's default settings
 * True if the light exists in the levels actor array, otherwise False.
 */
-IMPLEMENT_SIMPLE_AUTOMATION_TEST(FLightPointLightPlacement, "System.Editor.Lighting.Point Light.Placement", EAutomationTestFlags::ATF_Editor)
+IMPLEMENT_SIMPLE_AUTOMATION_TEST(FLightPointLightPlacement, "System.Editor.Lighting.Point Light.Placement", EAutomationTestFlags::EditorContext | EAutomationTestFlags::EngineFilter)
 
 bool FLightPointLightPlacement::RunTest(const FString& Parameters)
 {
@@ -276,7 +276,7 @@ bool FLightPointLightPlacement::RunTest(const FString& Parameters)
 * Place a point light in the world with it's default settings
 * True if the light exists in the levels actor array, otherwise False.
 */
-IMPLEMENT_SIMPLE_AUTOMATION_TEST(FLightPointLightSetLocation, "System.Editor.Lighting.Point Light.Set Location", EAutomationTestFlags::ATF_Editor)
+IMPLEMENT_SIMPLE_AUTOMATION_TEST(FLightPointLightSetLocation, "System.Editor.Lighting.Point Light.Set Location", EAutomationTestFlags::EditorContext | EAutomationTestFlags::EngineFilter)
 
 bool FLightPointLightSetLocation::RunTest(const FString& Parameters)
 {
diff --git a/Engine/Source/Runtime/AutomationMessages/Classes/AutomationWorkerMessages.h b/Engine/Source/Runtime/AutomationMessages/Classes/AutomationWorkerMessages.h
--- a/Engine/Source/Runtime/AutomationMessages/Classes/AutomationWorkerMessages.h
+++ b/Engine/Source/Runtime/AutomationMessages/Classes/AutomationWorkerMessages.h
@@ -149,9 +149,9 @@ struct FAutomationWorkerRequestTests
 	UPROPERTY()
 	bool DeveloperDirectoryIncluded;
 
-	/** Holds a flag indicating whether the visual commandlet filter is enabled. */
+	/** Holds a flag indicating which tests we'd like to request. */
 	UPROPERTY()
-	bool VisualCommandletFilterOn;
+	uint32 RequestedTestFlags;
 
 	/**
 	 * Default constructor.
@@ -161,9 +161,9 @@ struct FAutomationWorkerRequestTests
 	/**
 	 * Creates and initializes a new instance.
 	 */
-	FAutomationWorkerRequestTests( bool InDeveloperDirectoryIncluded, bool InVisualCommandletFilterOn )
+	FAutomationWorkerRequestTests(bool InDeveloperDirectoryIncluded, uint32 InRequestedTestFlags)
 		: DeveloperDirectoryIncluded(InDeveloperDirectoryIncluded)
-		, VisualCommandletFilterOn(InVisualCommandletFilterOn)
+		, RequestedTestFlags(InRequestedTestFlags)
 	{ }
 };
 
@@ -199,6 +199,14 @@ struct FAutomationWorkerRequestTestsReply
 };
 
 
+/**
+*/
+USTRUCT()
+struct FAutomationWorkerRequestTestsReplyComplete
+{
+	GENERATED_USTRUCT_BODY()
+};
+
 /**
  * Implements a message to request the running of automation tests on a worker.
  */
diff --git a/Engine/Source/Runtime/AutomationWorker/Private/AutomationWorkerModule.cpp b/Engine/Source/Runtime/AutomationWorker/Private/AutomationWorkerModule.cpp
--- a/Engine/Source/Runtime/AutomationWorker/Private/AutomationWorkerModule.cpp
+++ b/Engine/Source/Runtime/AutomationWorker/Private/AutomationWorkerModule.cpp
@@ -72,9 +72,6 @@ void FAutomationWorkerModule::Tick()
 	{
 		MessageEndpoint->ProcessInbox();
 	}
-
-	// Run any of the automation commands that was obtained during initialization.
-	//RunDeferredAutomationCommands();
 }
 
 
@@ -235,6 +232,7 @@ void FAutomationWorkerModule::SendTests( const FMessageAddress& ControllerAddres
 	{
 		MessageEndpoint->Send(new FAutomationWorkerRequestTestsReply(TestInfo[TestIndex].GetTestAsString(), TestInfo.Num()), ControllerAddress);
 	}
+	MessageEndpoint->Send(new FAutomationWorkerRequestTestsReplyComplete(), ControllerAddress);
 }
 
 
@@ -325,7 +323,7 @@ void FAutomationWorkerModule::HandleRequestTestsMessage( const FAutomationWorker
 {
 	FAutomationTestFramework::GetInstance().LoadTestModules();
 	FAutomationTestFramework::GetInstance().SetDeveloperDirectoryIncluded(Message.DeveloperDirectoryIncluded);
-	FAutomationTestFramework::GetInstance().SetVisualCommandletFilter(Message.VisualCommandletFilterOn);
+	FAutomationTestFramework::GetInstance().SetRequestedTestFilter(Message.RequestedTestFlags);
 	FAutomationTestFramework::GetInstance().GetValidTestNames( TestInfo );
 
 	SendTests(Context->GetSender());
diff --git a/Engine/Source/Runtime/AutomationWorker/Private/AutomationWorkerModule.h b/Engine/Source/Runtime/AutomationWorker/Private/AutomationWorkerModule.h
--- a/Engine/Source/Runtime/AutomationWorker/Private/AutomationWorkerModule.h
+++ b/Engine/Source/Runtime/AutomationWorker/Private/AutomationWorkerModule.h
@@ -127,7 +127,4 @@ private:
 
 	/** Delegate to fire when the test is complete */
 	FStopTestEvent StopTestEvent;
-
-	// Holds the automation command line arguments.
-	TArray<FString> DeferredAutomationCommands;
 };
diff --git a/Engine/Source/Runtime/Core/Private/Containers/ContainersTest.cpp b/Engine/Source/Runtime/Core/Private/Containers/ContainersTest.cpp
--- a/Engine/Source/Runtime/Core/Private/Containers/ContainersTest.cpp
+++ b/Engine/Source/Runtime/Core/Private/Containers/ContainersTest.cpp
@@ -4,7 +4,7 @@
 #include "AutomationTest.h"
 
 
-IMPLEMENT_SIMPLE_AUTOMATION_TEST(FContainersTest, "System.Core.Misc.Containers", EAutomationTestFlags::ATF_SmokeTest)
+IMPLEMENT_SIMPLE_AUTOMATION_TEST(FContainersTest, "System.Core.Misc.Containers", EAutomationTestFlags::ApplicationContextMask | EAutomationTestFlags::SmokeFilter)
 
 #define MAX_TEST_OBJECTS      65
 #define MAX_TEST_OBJECTS_STEP 1
diff --git a/Engine/Source/Runtime/Core/Private/Math/BasicMathExpressionEvaluator.cpp b/Engine/Source/Runtime/Core/Private/Math/BasicMathExpressionEvaluator.cpp
--- a/Engine/Source/Runtime/Core/Private/Math/BasicMathExpressionEvaluator.cpp
+++ b/Engine/Source/Runtime/Core/Private/Math/BasicMathExpressionEvaluator.cpp
@@ -222,7 +222,7 @@ bool TestInvalidExpression(FAutomationTestBase* Test, const TCHAR* Expression)
 }
 
 
-IMPLEMENT_SIMPLE_AUTOMATION_TEST(FBasicMathExpressionEvaluatorTest, "System.Core.Math.Evaluate.Valid Expressions", EAutomationTestFlags::ATF_SmokeTest)
+IMPLEMENT_SIMPLE_AUTOMATION_TEST(FBasicMathExpressionEvaluatorTest, "System.Core.Math.Evaluate.Valid Expressions", EAutomationTestFlags::ApplicationContextMask | EAutomationTestFlags::SmokeFilter)
 // Evaluates valid math expressions.
 bool FBasicMathExpressionEvaluatorTest::RunTest(const FString& Parameters)
 {
@@ -243,7 +243,7 @@ bool FBasicMathExpressionEvaluatorTest::RunTest(const FString& Parameters)
 	return true;
 }
 
-IMPLEMENT_SIMPLE_AUTOMATION_TEST(FBasicMathExpressionEvaluatorWhitespaceExpressionsTest, "System.Core.Math.Evaluate.Valid Expressions With Whitespaces", EAutomationTestFlags::ATF_SmokeTest)
+IMPLEMENT_SIMPLE_AUTOMATION_TEST(FBasicMathExpressionEvaluatorWhitespaceExpressionsTest, "System.Core.Math.Evaluate.Valid Expressions With Whitespaces", EAutomationTestFlags::ApplicationContextMask | EAutomationTestFlags::SmokeFilter)
 // Evaluates a valid math expression with leading and trailing white spaces.
 bool FBasicMathExpressionEvaluatorWhitespaceExpressionsTest::RunTest(const FString& Parameters)
 {
@@ -252,7 +252,7 @@ bool FBasicMathExpressionEvaluatorWhitespaceExpressionsTest::RunTest(const FStri
 	return true;
 }
 
-IMPLEMENT_SIMPLE_AUTOMATION_TEST(FBasicMathExpressionEvaluatorGroupedExpressionsTest, "System.Core.Math.Evaluate.Valid Grouped Expressions", EAutomationTestFlags::ATF_SmokeTest)
+IMPLEMENT_SIMPLE_AUTOMATION_TEST(FBasicMathExpressionEvaluatorGroupedExpressionsTest, "System.Core.Math.Evaluate.Valid Grouped Expressions", EAutomationTestFlags::ApplicationContextMask | EAutomationTestFlags::SmokeFilter)
 // Evaluates valid math expressions that are grouped
 bool FBasicMathExpressionEvaluatorGroupedExpressionsTest::RunTest(const FString& Parameters)
 {
@@ -263,7 +263,7 @@ bool FBasicMathExpressionEvaluatorGroupedExpressionsTest::RunTest(const FString&
 	return true;
 }
 
-IMPLEMENT_SIMPLE_AUTOMATION_TEST(FBasicMathExpressionEvaluatorInvalidExpressionTest, "System.Core.Math.Evaluate.Invalid Expressions", EAutomationTestFlags::ATF_SmokeTest)
+IMPLEMENT_SIMPLE_AUTOMATION_TEST(FBasicMathExpressionEvaluatorInvalidExpressionTest, "System.Core.Math.Evaluate.Invalid Expressions", EAutomationTestFlags::ApplicationContextMask | EAutomationTestFlags::SmokeFilter)
 
 // Evaluates invalid expressions.
 // Invalid expressions will report errors and not crash.
diff --git a/Engine/Source/Runtime/Core/Private/Misc/AutomationTest.cpp b/Engine/Source/Runtime/Core/Private/Misc/AutomationTest.cpp
--- a/Engine/Source/Runtime/Core/Private/Misc/AutomationTest.cpp
+++ b/Engine/Source/Runtime/Core/Private/Misc/AutomationTest.cpp
@@ -75,8 +75,8 @@ bool FAutomationTestFramework::UnregisterAutomationTest( const FString& InTestNa
 
 void FAutomationTestFramework::EnqueueLatentCommand(TSharedPtr<IAutomationLatentCommand> NewCommand)
 {
-	//ensure latent commands are never used within smoke tests
-	check(!bRunningSmokeTests);
+	//ensure latent commands are never used within smoke tests - will only catch when smokes are exclusively requested
+	check((RequestedTestFilter & EAutomationTestFlags::FilterMask) != EAutomationTestFlags::SmokeFilter);
 
 	//ensure we are currently "running a test"
 	check(GIsAutomationTesting);
@@ -93,7 +93,7 @@ void FAutomationTestFramework::EnqueueLatentCommand(TSharedPtr<IAutomationLatent
 void FAutomationTestFramework::EnqueueNetworkCommand(TSharedPtr<IAutomationNetworkCommand> NewCommand)
 {
 	//ensure latent commands are never used within smoke tests
-	check(!bRunningSmokeTests);
+	check((RequestedTestFilter & EAutomationTestFlags::FilterMask) != EAutomationTestFlags::SmokeFilter);
 
 	//ensure we are currently "running a test"
 	check(GIsAutomationTesting);
@@ -112,8 +112,9 @@ bool FAutomationTestFramework::RunSmokeTests()
 {
 	bool bAllSuccessful = true;
 
+	uint32 PreviousRequestedTestFilter = RequestedTestFilter;
 	//so extra log spam isn't generated
-	bRunningSmokeTests = true;
+	RequestedTestFilter = EAutomationTestFlags::SmokeFilter;
 	
 	// Skip running on cooked platforms like mobile
 	//@todo - better determination of whether to run than requires cooked data
@@ -138,7 +139,7 @@ bool FAutomationTestFramework::RunSmokeTests()
 			for ( int TestIndex = 0; TestIndex < TestInfo.Num(); ++TestIndex )
 			{
 				SlowTask.EnterProgressFrame(1);
-				if (TestInfo[TestIndex].GetTestType() == EAutomationTestType::ATT_SmokeTest )
+				if (TestInfo[TestIndex].GetTestFlags() == EAutomationTestFlags::SmokeFilter )
 				{
 					FString TestCommand = TestInfo[TestIndex].GetTestName();
 					FAutomationTestExecutionInfo& CurExecutionInfo = OutExecutionInfoMap.Add( TestCommand, FAutomationTestExecutionInfo() );
@@ -173,7 +174,7 @@ bool FAutomationTestFramework::RunSmokeTests()
 	}
 
 	//revert to allowing all logs
-	bRunningSmokeTests = false;
+	RequestedTestFilter = PreviousRequestedTestFilter;
 
 	return bAllSuccessful;
 }
@@ -306,6 +307,7 @@ void FAutomationTestFramework::LoadTestModules( )
 {
 	const bool bRunningEditor = GIsEditor && !IsRunningCommandlet();
 
+	bool bRunningSmokeTests = ((RequestedTestFilter & EAutomationTestFlags::FilterMask) == EAutomationTestFlags::SmokeFilter);
 	if( !bRunningSmokeTests )
 	{
 		TArray<FString> EngineTestModules;
@@ -364,19 +366,19 @@ void FAutomationTestFramework::GetValidTestNames( TArray<FAutomationTestInfo>& T
 	uint32 ApplicationSupportFlags = 0;
 	if ( bRunningEditor )
 	{
-		ApplicationSupportFlags |= EAutomationTestFlags::ATF_Editor;
+		ApplicationSupportFlags |= EAutomationTestFlags::EditorContext;
 	}
 	if ( bRunningGame )
 	{
-		ApplicationSupportFlags |= EAutomationTestFlags::ATF_Game;
+		ApplicationSupportFlags |= EAutomationTestFlags::ClientContext;
 	}
 	if ( bRunningCommandlet )
 	{
-		ApplicationSupportFlags |= EAutomationTestFlags::ATF_Commandlet;
+		ApplicationSupportFlags |= EAutomationTestFlags::CommandletContext;
 	}
 
 	//Feature support - assume valid RHI until told otherwise
-	uint32 FeatureSupportFlags = EAutomationTestFlags::ATF_FeatureMask;
+	uint32 FeatureSupportFlags = EAutomationTestFlags::FeatureMask;
 	// @todo: Handle this correctly. GIsUsingNullRHI is defined at Engine-level, so it can't be used directly here in Core.
 	// For now, assume Null RHI is only used for commandlets, servers, and when the command line specifies to use it.
 	if (FPlatformProperties::SupportsWindowedMode())
@@ -384,19 +386,12 @@ void FAutomationTestFramework::GetValidTestNames( TArray<FAutomationTestInfo>& T
 		bool bUsingNullRHI = FParse::Param( FCommandLine::Get(), TEXT("nullrhi") ) || IsRunningCommandlet() || IsRunningDedicatedServer();
 		if (bUsingNullRHI)
 		{
-			FeatureSupportFlags &= (~EAutomationTestFlags::ATF_NonNullRHI);
+			FeatureSupportFlags &= (~EAutomationTestFlags::NonNullRHI);
 		}
 	}
 	if (FApp::IsUnattended())
 	{
-		FeatureSupportFlags &= (~EAutomationTestFlags::ATF_RequiresUser);
-	}
-
-	// test support flag, see what kind of test is required and allowed
-	uint32 ActionSupportFlags = 0;
-	if ( bVisualCommandletFilterOn )
-	{
-		ActionSupportFlags |= EAutomationTestFlags::ATF_VisualCommandlet;
+		FeatureSupportFlags &= (~EAutomationTestFlags::RequiresUser);
 	}
 
 	for ( TMap<FString, FAutomationTestBase*>::TConstIterator TestIter( AutomationTestClassNameToInstanceMap ); TestIter; ++TestIter )
@@ -407,21 +402,19 @@ void FAutomationTestFramework::GetValidTestNames( TArray<FAutomationTestInfo>& T
 		uint32 CurTestFlags = CurTest->GetTestFlags();
 
 		//filter out full tests when running smoke tests
-		const bool bPassesSmokeTestRequirement = (!bRunningSmokeTests) || (CurTestFlags & EAutomationTestFlags::ATF_SmokeTest);
+		const bool bPassesFilterRequirement = ((CurTestFlags & RequestedTestFilter) != 0);
 
 		//Application Tests
-		uint32 CurTestApplicationFlags = (CurTestFlags & EAutomationTestFlags::ATF_ApplicationMask);
+		uint32 CurTestApplicationFlags = (CurTestFlags & EAutomationTestFlags::ApplicationContextMask);
 		const bool bPassesApplicationRequirements = (CurTestApplicationFlags == 0) || (CurTestApplicationFlags & ApplicationSupportFlags);
 		
 		//Feature Tests
-		uint32 CurTestFeatureFlags = (CurTestFlags & EAutomationTestFlags::ATF_FeatureMask);
+		uint32 CurTestFeatureFlags = (CurTestFlags & EAutomationTestFlags::FeatureMask);
 		const bool bPassesFeatureRequirements = (CurTestFeatureFlags == 0) || (CurTestFeatureFlags & FeatureSupportFlags);
 
-		// Action flag Tests
-		uint32 CurTestActionFlags = (CurTestFlags & EAutomationTestFlags::ATF_ActionMask);
-		const bool bPassesActionRequirements = (CurTestActionFlags == ActionSupportFlags);
+		const bool bEnabled = (CurTestFlags & EAutomationTestFlags::Disabled) == 0;
 
-		if (bPassesApplicationRequirements && bPassesFeatureRequirements && bPassesActionRequirements && bPassesSmokeTestRequirement)
+		if (bEnabled && bPassesApplicationRequirements && bPassesFeatureRequirements && bPassesFilterRequirement)
 		{
 			CurTest->GenerateTestNames(TestInfo);
 		}
@@ -452,7 +445,7 @@ bool FAutomationTestFramework::ShouldTestContent(const FString& Path) const
 	}
 
 	FString DevelopersPath = FPaths::GameDevelopersDir().LeftChop(1);
-	return bDeveloperDirectoryIncluded || bVisualCommandletFilterOn || !Path.StartsWith(DevelopersPath);
+	return bDeveloperDirectoryIncluded || !Path.StartsWith(DevelopersPath);
 }
 
 
@@ -462,9 +455,9 @@ void FAutomationTestFramework::SetDeveloperDirectoryIncluded(const bool bInDevel
 }
 
 
-void FAutomationTestFramework::SetVisualCommandletFilter(const bool bInVisualCommandletFilterOn)
+void FAutomationTestFramework::SetRequestedTestFilter(const uint32 InRequestedTestFlags)
 {
-	bVisualCommandletFilterOn = bInVisualCommandletFilterOn;
+	RequestedTestFilter = InRequestedTestFlags;
 }
 
 
@@ -599,7 +592,9 @@ void FAutomationTestFramework::InternalStartTest( const FString& InTestToRun )
 
 		StartTime = FPlatformTime::Seconds();
 
-		if (!bRunningSmokeTests)
+		//if non-
+		uint32 NonSmokeTestFlags = (EAutomationTestFlags::FilterMask & (~EAutomationTestFlags::SmokeFilter));
+		if (RequestedTestFilter & NonSmokeTestFlags)
 		{
 			UE_LOG(LogAutomationTest, Log, TEXT("%s %s is starting at %f"), *CurrentTest->GetBeautifiedTestName(), *Parameters, StartTime);
 		}
@@ -617,7 +612,8 @@ bool FAutomationTestFramework::InternalStopTest(FAutomationTestExecutionInfo& Ou
 
 	double EndTime = FPlatformTime::Seconds();
 	double TimeForTest = static_cast<float>(EndTime - StartTime);
-	if (!bRunningSmokeTests)
+	uint32 NonSmokeTestFlags = (EAutomationTestFlags::FilterMask & (~EAutomationTestFlags::SmokeFilter));
+	if (RequestedTestFilter & NonSmokeTestFlags)
 	{
 		UE_LOG(LogAutomationTest, Log, TEXT("%s %s ran in %f"), *CurrentTest->GetBeautifiedTestName(), *Parameters, TimeForTest);
 	}
@@ -648,12 +644,11 @@ bool FAutomationTestFramework::InternalStopTest(FAutomationTestExecutionInfo& Ou
 
 FAutomationTestFramework::FAutomationTestFramework()
 :	CachedContext( NULL )
-,	bRunningSmokeTests(false)
+,	RequestedTestFilter(EAutomationTestFlags::SmokeFilter)
 ,	StartTime(0.0f)
 ,	bTestSuccessful(false)
 ,	CurrentTest(NULL)
 ,	bDeveloperDirectoryIncluded(false)
-,	bVisualCommandletFilterOn(false)
 ,	bScreenshotsEnabled(true)
 ,	bUseFullSizeScreenShots(true)
 ,	NetworkRoleIndex(0)
@@ -738,20 +733,8 @@ void FAutomationTestBase::GenerateTestNames(TArray<FAutomationTestInfo>& TestInf
 			CompleteTestName = FString::Printf(TEXT("%s %s"), *TestName, *ParameterNames[ParameterIndex]);
 		}
 
-		// Set the test type
-		uint8 TestType = EAutomationTestType::ATT_NormalTest;
-
-		if (bComplexTask)
-		{
-			TestType = EAutomationTestType::ATT_StressTest;
-		}
-		else if(GetTestFlags() & EAutomationTestFlags::ATF_SmokeTest)
-		{
-			TestType = EAutomationTestType::ATT_SmokeTest;
-		}
-
 		// Add the test info to our collection
-		FAutomationTestInfo NewTestInfo( CompleteBeautifiedNames, CompleteTestName, TestType, GetRequiredDeviceNum(), ParameterNames[ParameterIndex] );
+		FAutomationTestInfo NewTestInfo( CompleteBeautifiedNames, CompleteTestName, GetTestFlags(), GetRequiredDeviceNum(), ParameterNames[ParameterIndex] );
 		
 		TestInfo.Add( NewTestInfo );
 	}
diff --git a/Engine/Source/Runtime/Core/Private/Misc/ExpressionParser.cpp b/Engine/Source/Runtime/Core/Private/Misc/ExpressionParser.cpp
--- a/Engine/Source/Runtime/Core/Private/Misc/ExpressionParser.cpp
+++ b/Engine/Source/Runtime/Core/Private/Misc/ExpressionParser.cpp
@@ -892,7 +892,7 @@ PRAGMA_DISABLE_OPTIMIZATION
 		return true;
 	}
 
-	IMPLEMENT_SIMPLE_AUTOMATION_TEST(FExpressionParserMoveableTypes, "System.Core.Expression Parser.Moveable Types", EAutomationTestFlags::ATF_SmokeTest)
+	IMPLEMENT_SIMPLE_AUTOMATION_TEST(FExpressionParserMoveableTypes, "System.Core.Expression Parser.Moveable Types", EAutomationTestFlags::EditorContext | EAutomationTestFlags::SmokeFilter)
 	bool FExpressionParserMoveableTypes::RunTest( const FString& Parameters )
 	{
 		return TestWithType<FMoveableType>(this);
@@ -913,7 +913,7 @@ PRAGMA_DISABLE_OPTIMIZATION
 		uint8 Padding[1024];
 	};
 
-	IMPLEMENT_SIMPLE_AUTOMATION_TEST(FExpressionParserAllocatedTypes, "System.Core.Expression Parser.Allocated Types", EAutomationTestFlags::ATF_SmokeTest)
+	IMPLEMENT_SIMPLE_AUTOMATION_TEST(FExpressionParserAllocatedTypes, "System.Core.Expression Parser.Allocated Types", EAutomationTestFlags::EditorContext | EAutomationTestFlags::SmokeFilter)
 	bool FExpressionParserAllocatedTypes::RunTest( const FString& Parameters )
 	{
 		return TestWithType<FHugeType>(this);
diff --git a/Engine/Source/Runtime/Core/Private/Misc/ExpressionParserExamples.cpp b/Engine/Source/Runtime/Core/Private/Misc/ExpressionParserExamples.cpp
--- a/Engine/Source/Runtime/Core/Private/Misc/ExpressionParserExamples.cpp
+++ b/Engine/Source/Runtime/Core/Private/Misc/ExpressionParserExamples.cpp
@@ -215,7 +215,7 @@ DEFINE_LOG_CATEGORY_STATIC(LogExampleStringTests, Log, All);
 
 #include "TextHistory.h"
 
-IMPLEMENT_SIMPLE_AUTOMATION_TEST(FExampleStringFormatterTest, "System.Core.Expression Parser", EAutomationTestFlags::ATF_Editor)
+IMPLEMENT_SIMPLE_AUTOMATION_TEST(FExampleStringFormatterTest, "System.Core.Expression Parser", EAutomationTestFlags::Editor)
 bool FExampleStringFormatterTest::RunTest( const FString& Parameters )
 {
 	FExampleStringFormatter Formatter;
diff --git a/Engine/Source/Runtime/Core/Private/Misc/TextFilterTests.cpp b/Engine/Source/Runtime/Core/Private/Misc/TextFilterTests.cpp
--- a/Engine/Source/Runtime/Core/Private/Misc/TextFilterTests.cpp
+++ b/Engine/Source/Runtime/Core/Private/Misc/TextFilterTests.cpp
@@ -215,7 +215,7 @@ namespace TextFilterTests
 	}
 };
 
-IMPLEMENT_SIMPLE_AUTOMATION_TEST(FTextFilterTests_TextFilterExpressionEvaluator, "System.Core.Text.TextFilterExpressionEvaluator", EAutomationTestFlags::ATF_ApplicationMask)
+IMPLEMENT_SIMPLE_AUTOMATION_TEST(FTextFilterTests_TextFilterExpressionEvaluator, "System.Core.Text.TextFilterExpressionEvaluator", EAutomationTestFlags::ApplicationContextMask | EAutomationTestFlags::EngineFilter)
 bool FTextFilterTests_TextFilterExpressionEvaluator::RunTest(const FString& Parameters)
 {
 	using namespace TextFilterTests;
@@ -241,7 +241,7 @@ bool FTextFilterTests_TextFilterExpressionEvaluator::RunTest(const FString& Para
 	return bResult;
 }
 
-IMPLEMENT_SIMPLE_AUTOMATION_TEST(FTextFilterTests_TextFilter, "System.Core.Text.TextFilter", EAutomationTestFlags::ATF_ApplicationMask)
+IMPLEMENT_SIMPLE_AUTOMATION_TEST(FTextFilterTests_TextFilter, "System.Core.Text.TextFilter", EAutomationTestFlags::ApplicationContextMask | EAutomationTestFlags::EngineFilter)
 bool FTextFilterTests_TextFilter::RunTest(const FString& Parameters)
 {
 	using namespace TextFilterTests;
diff --git a/Engine/Source/Runtime/Core/Private/Serialization/Csv/CsvParserTests.cpp b/Engine/Source/Runtime/Core/Private/Serialization/Csv/CsvParserTests.cpp
--- a/Engine/Source/Runtime/Core/Private/Serialization/Csv/CsvParserTests.cpp
+++ b/Engine/Source/Runtime/Core/Private/Serialization/Csv/CsvParserTests.cpp
@@ -49,7 +49,7 @@ namespace CsvParser_Tests
 	}
 }
 
-IMPLEMENT_SIMPLE_AUTOMATION_TEST(FBasicTest, "System.Core.CSV Parser.Basic", 0)
+IMPLEMENT_SIMPLE_AUTOMATION_TEST(FBasicTest, "System.Core.CSV Parser.Basic", EAutomationTestFlags::EditorContext | EAutomationTestFlags::EngineFilter)
 
 bool FBasicTest::RunTest(const FString& Parameters)
 {
@@ -64,7 +64,7 @@ bool FBasicTest::RunTest(const FString& Parameters)
 	return CsvParser_Tests::CheckRows(Parser.GetRows(), Expected, this);
 }
 
-IMPLEMENT_SIMPLE_AUTOMATION_TEST(FQuotedStringTest, "System.Core.CSV Parser.Quoted String", 0)
+IMPLEMENT_SIMPLE_AUTOMATION_TEST(FQuotedStringTest, "System.Core.CSV Parser.Quoted String", EAutomationTestFlags::EditorContext | EAutomationTestFlags::EngineFilter)
 
 bool FQuotedStringTest::RunTest(const FString& Parameters)
 {
@@ -82,7 +82,7 @@ bool FQuotedStringTest::RunTest(const FString& Parameters)
 	return CsvParser_Tests::CheckRows(Parser.GetRows(), Expected, this);
 }
 
-IMPLEMENT_SIMPLE_AUTOMATION_TEST(FEmptyCellsTest, "System.Core.CSV Parser.Empty Cells", 0)
+IMPLEMENT_SIMPLE_AUTOMATION_TEST(FEmptyCellsTest, "System.Core.CSV Parser.Empty Cells", EAutomationTestFlags::EditorContext | EAutomationTestFlags::EngineFilter)
 
 bool FEmptyCellsTest::RunTest(const FString& Parameters)
 {
@@ -100,7 +100,7 @@ bool FEmptyCellsTest::RunTest(const FString& Parameters)
 	return CsvParser_Tests::CheckRows(Parser.GetRows(), Expected, this);
 }
 
-IMPLEMENT_SIMPLE_AUTOMATION_TEST(FMixedNewLinesTest, "System.Core.CSV Parser.Mixed Newlines", 0)
+IMPLEMENT_SIMPLE_AUTOMATION_TEST(FMixedNewLinesTest, "System.Core.CSV Parser.Mixed Newlines", EAutomationTestFlags::EditorContext | EAutomationTestFlags::EngineFilter)
 
 bool FMixedNewLinesTest::RunTest(const FString& Parameters)
 {
@@ -113,7 +113,7 @@ bool FMixedNewLinesTest::RunTest(const FString& Parameters)
 	return CsvParser_Tests::CheckRows(Parser.GetRows(), Expected, this);
 }
 
-IMPLEMENT_SIMPLE_AUTOMATION_TEST(FEmptyStringTest, "System.Core.CSV Parser.Empty String", 0)
+IMPLEMENT_SIMPLE_AUTOMATION_TEST(FEmptyStringTest, "System.Core.CSV Parser.Empty String", EAutomationTestFlags::EditorContext | EAutomationTestFlags::EngineFilter)
 
 bool FEmptyStringTest::RunTest(const FString& Parameters)
 {
@@ -121,7 +121,7 @@ bool FEmptyStringTest::RunTest(const FString& Parameters)
 			FCsvParser(FString("")).GetRows().Num() == 0;
 }
 
-IMPLEMENT_SIMPLE_AUTOMATION_TEST(FOnlyNewLinesTest, "System.Core.CSV Parser.Only Newlines", 0)
+IMPLEMENT_SIMPLE_AUTOMATION_TEST(FOnlyNewLinesTest, "System.Core.CSV Parser.Only Newlines", EAutomationTestFlags::EditorContext | EAutomationTestFlags::EngineFilter)
 
 bool FOnlyNewLinesTest::RunTest(const FString& Parameters)
 {
diff --git a/Engine/Source/Runtime/Core/Private/Tests/Async/AsyncTest.cpp b/Engine/Source/Runtime/Core/Private/Tests/Async/AsyncTest.cpp
--- a/Engine/Source/Runtime/Core/Private/Tests/Async/AsyncTest.cpp
+++ b/Engine/Source/Runtime/Core/Private/Tests/Async/AsyncTest.cpp
@@ -7,7 +7,7 @@
 #include "Future.h"
 
 
-IMPLEMENT_SIMPLE_AUTOMATION_TEST(FAsyncTest, "System.Core.Async.Async", EAutomationTestFlags::ATF_Editor)
+IMPLEMENT_SIMPLE_AUTOMATION_TEST(FAsyncTest, "System.Core.Async.Async", EAutomationTestFlags::EditorContext | EAutomationTestFlags::EngineFilter)
 
 
 bool FAsyncTest::RunTest(const FString& Parameters)
diff --git a/Engine/Source/Runtime/Core/Private/Tests/HAL/PlatformTest.cpp b/Engine/Source/Runtime/Core/Private/Tests/HAL/PlatformTest.cpp
--- a/Engine/Source/Runtime/Core/Private/Tests/HAL/PlatformTest.cpp
+++ b/Engine/Source/Runtime/Core/Private/Tests/HAL/PlatformTest.cpp
@@ -34,7 +34,7 @@ struct TestC : public TestA, TestB
 };
 
 
-IMPLEMENT_SIMPLE_AUTOMATION_TEST(FPlatformVerificationTest, "System.Core.HAL.Platform Verification", EAutomationTestFlags::ATF_SmokeTest)
+IMPLEMENT_SIMPLE_AUTOMATION_TEST(FPlatformVerificationTest, "System.Core.HAL.Platform Verification", EAutomationTestFlags::EditorContext | EAutomationTestFlags::ClientContext | EAutomationTestFlags::SmokeFilter)
 
 bool FPlatformVerificationTest::RunTest (const FString& Parameters)
 {
diff --git a/Engine/Source/Runtime/Core/Private/Tests/Internationalization/ArchiveTest.cpp b/Engine/Source/Runtime/Core/Private/Tests/Internationalization/ArchiveTest.cpp
--- a/Engine/Source/Runtime/Core/Private/Tests/Internationalization/ArchiveTest.cpp
+++ b/Engine/Source/Runtime/Core/Private/Tests/Internationalization/ArchiveTest.cpp
@@ -6,7 +6,7 @@
 #include "Internationalization/InternationalizationMetadata.h"
 
 
-IMPLEMENT_SIMPLE_AUTOMATION_TEST(FArchiveTest, "System.Core.Misc.Internationalization Archive", EAutomationTestFlags::ATF_SmokeTest)
+IMPLEMENT_SIMPLE_AUTOMATION_TEST(FArchiveTest, "System.Core.Misc.Internationalization Archive", EAutomationTestFlags::EditorContext | EAutomationTestFlags::ClientContext | EAutomationTestFlags::SmokeFilter)
 
 
 bool FArchiveTest::RunTest( const FString& Parameters )
diff --git a/Engine/Source/Runtime/Core/Private/Tests/Internationalization/DateTimeFormattingRulesTest.cpp b/Engine/Source/Runtime/Core/Private/Tests/Internationalization/DateTimeFormattingRulesTest.cpp
--- a/Engine/Source/Runtime/Core/Private/Tests/Internationalization/DateTimeFormattingRulesTest.cpp
+++ b/Engine/Source/Runtime/Core/Private/Tests/Internationalization/DateTimeFormattingRulesTest.cpp
@@ -7,7 +7,7 @@
 // disable optimization for DateTimeFormattingRulesTest as it compiles very slowly in development builds
 PRAGMA_DISABLE_OPTIMIZATION
 
-	IMPLEMENT_SIMPLE_AUTOMATION_TEST(FDateTimeFormattingRulesTest, "System.Core.Misc.DateTime Formatting Rules", EAutomationTestFlags::ATF_None)
+IMPLEMENT_SIMPLE_AUTOMATION_TEST(FDateTimeFormattingRulesTest, "System.Core.Misc.DateTime Formatting Rules", EAutomationTestFlags::EditorContext | EAutomationTestFlags::ClientContext | EAutomationTestFlags::EngineFilter)
 
 	namespace
 {
diff --git a/Engine/Source/Runtime/Core/Private/Tests/Internationalization/ManifestTest.cpp b/Engine/Source/Runtime/Core/Private/Tests/Internationalization/ManifestTest.cpp
--- a/Engine/Source/Runtime/Core/Private/Tests/Internationalization/ManifestTest.cpp
+++ b/Engine/Source/Runtime/Core/Private/Tests/Internationalization/ManifestTest.cpp
@@ -15,7 +15,7 @@ int32 CountManifestEntries( const FInternationalizationManifest& Manifest )
 }
 
 
-IMPLEMENT_SIMPLE_AUTOMATION_TEST(FLocContextTest, "System.Core.Misc.Internationalization Context", EAutomationTestFlags::ATF_SmokeTest)
+IMPLEMENT_SIMPLE_AUTOMATION_TEST(FLocContextTest, "System.Core.Misc.Internationalization Context", EAutomationTestFlags::EditorContext | EAutomationTestFlags::ClientContext | EAutomationTestFlags::SmokeFilter)
 
 bool FLocContextTest::RunTest( const FString& Parameters )
 {
@@ -200,7 +200,7 @@ bool FLocContextTest::RunTest( const FString& Parameters )
 }
 
 
-IMPLEMENT_SIMPLE_AUTOMATION_TEST(FLocItemTest, "System.Core.Misc.Internationalization LocItem", EAutomationTestFlags::ATF_SmokeTest)
+IMPLEMENT_SIMPLE_AUTOMATION_TEST(FLocItemTest, "System.Core.Misc.Internationalization LocItem", EAutomationTestFlags::EditorContext | EAutomationTestFlags::ClientContext | EAutomationTestFlags::SmokeFilter)
 
 bool FLocItemTest::RunTest( const FString& Parameters )
 {
@@ -354,7 +354,7 @@ bool FLocItemTest::RunTest( const FString& Parameters )
 }
 
 
-IMPLEMENT_SIMPLE_AUTOMATION_TEST(FManifestTest, "System.Core.Misc.Internationalization Manifest", EAutomationTestFlags::ATF_SmokeTest)
+IMPLEMENT_SIMPLE_AUTOMATION_TEST(FManifestTest, "System.Core.Misc.Internationalization Manifest", EAutomationTestFlags::EditorContext | EAutomationTestFlags::ClientContext | EAutomationTestFlags::SmokeFilter)
 
 bool FManifestTest::RunTest( const FString& Parameters )
 {
diff --git a/Engine/Source/Runtime/Core/Private/Tests/Internationalization/MetadataTest.cpp b/Engine/Source/Runtime/Core/Private/Tests/Internationalization/MetadataTest.cpp
--- a/Engine/Source/Runtime/Core/Private/Tests/Internationalization/MetadataTest.cpp
+++ b/Engine/Source/Runtime/Core/Private/Tests/Internationalization/MetadataTest.cpp
@@ -5,7 +5,7 @@
 #include "Internationalization/InternationalizationMetadata.h"
 
 
-IMPLEMENT_SIMPLE_AUTOMATION_TEST(FMetadataTest, "System.Core.Misc.Internationalization Metadata", EAutomationTestFlags::ATF_SmokeTest)
+IMPLEMENT_SIMPLE_AUTOMATION_TEST(FMetadataTest, "System.Core.Misc.Internationalization Metadata", EAutomationTestFlags::EditorContext | EAutomationTestFlags::ClientContext | EAutomationTestFlags::SmokeFilter)
 
 bool FMetadataTest::RunTest( const FString& Parameters )
 {
diff --git a/Engine/Source/Runtime/Core/Private/Tests/Internationalization/NumberFormatingRulesTest.cpp b/Engine/Source/Runtime/Core/Private/Tests/Internationalization/NumberFormatingRulesTest.cpp
--- a/Engine/Source/Runtime/Core/Private/Tests/Internationalization/NumberFormatingRulesTest.cpp
+++ b/Engine/Source/Runtime/Core/Private/Tests/Internationalization/NumberFormatingRulesTest.cpp
@@ -7,7 +7,7 @@
 // Disable optimization for NumberFormatingRulesTest as it compiles very slowly in development builds
 PRAGMA_DISABLE_OPTIMIZATION
 
-	IMPLEMENT_SIMPLE_AUTOMATION_TEST(FNumberFormattingRulesTest, "System.Core.Misc.Number Formatting Rules", EAutomationTestFlags::ATF_None)
+IMPLEMENT_SIMPLE_AUTOMATION_TEST(FNumberFormattingRulesTest, "System.Core.Misc.Number Formatting Rules", EAutomationTestFlags::EditorContext | EAutomationTestFlags::ClientContext | EAutomationTestFlags::EngineFilter)
 
 	namespace
 {
diff --git a/Engine/Source/Runtime/Core/Private/Tests/Internationalization/TextTest.cpp b/Engine/Source/Runtime/Core/Private/Tests/Internationalization/TextTest.cpp
--- a/Engine/Source/Runtime/Core/Private/Tests/Internationalization/TextTest.cpp
+++ b/Engine/Source/Runtime/Core/Private/Tests/Internationalization/TextTest.cpp
@@ -11,7 +11,7 @@ PRAGMA_DISABLE_OPTIMIZATION
 #define LOCTEXT_NAMESPACE "Core.Tests.TextFormatTest"
 
 
-IMPLEMENT_SIMPLE_AUTOMATION_TEST(FTextTest, "System.Core.Misc.Text", EAutomationTestFlags::ATF_None)
+IMPLEMENT_SIMPLE_AUTOMATION_TEST(FTextTest, "System.Core.Misc.Text", EAutomationTestFlags::EditorContext | EAutomationTestFlags::ClientContext | EAutomationTestFlags::EngineFilter)
 
 namespace
 {
@@ -638,7 +638,7 @@ bool FTextTest::RunTest (const FString& Parameters)
 
 #if UE_ENABLE_ICU
 
-IMPLEMENT_SIMPLE_AUTOMATION_TEST(FICUTextTest, "System.Core.Misc.ICUText", EAutomationTestFlags::ATF_SmokeTest)
+IMPLEMENT_SIMPLE_AUTOMATION_TEST(FICUTextTest, "System.Core.Misc.ICUText", EAutomationTestFlags::EditorContext | EAutomationTestFlags::ClientContext | EAutomationTestFlags::SmokeFilter)
 
 bool FICUTextTest::RunTest (const FString& Parameters)
 {
diff --git a/Engine/Source/Runtime/Core/Private/Tests/Math/RangeBoundTest.cpp b/Engine/Source/Runtime/Core/Private/Tests/Math/RangeBoundTest.cpp
--- a/Engine/Source/Runtime/Core/Private/Tests/Math/RangeBoundTest.cpp
+++ b/Engine/Source/Runtime/Core/Private/Tests/Math/RangeBoundTest.cpp
@@ -4,7 +4,7 @@
 #include "AutomationTest.h"
 
 
-IMPLEMENT_SIMPLE_AUTOMATION_TEST(FRangeBoundTest, "System.Core.Math.RangeBound", EAutomationTestFlags::ATF_SmokeTest)
+IMPLEMENT_SIMPLE_AUTOMATION_TEST(FRangeBoundTest, "System.Core.Math.RangeBound", EAutomationTestFlags::EditorContext | EAutomationTestFlags::ClientContext | EAutomationTestFlags::SmokeFilter)
 
 bool FRangeBoundTest::RunTest( const FString& Parameters )
 {
diff --git a/Engine/Source/Runtime/Core/Private/Tests/Math/RangeSetTest.cpp b/Engine/Source/Runtime/Core/Private/Tests/Math/RangeSetTest.cpp
--- a/Engine/Source/Runtime/Core/Private/Tests/Math/RangeSetTest.cpp
+++ b/Engine/Source/Runtime/Core/Private/Tests/Math/RangeSetTest.cpp
@@ -4,7 +4,7 @@
 #include "AutomationTest.h"
 
 
-IMPLEMENT_SIMPLE_AUTOMATION_TEST(FRangeSetTest, "System.Core.Math.RangeSet", EAutomationTestFlags::ATF_SmokeTest)
+IMPLEMENT_SIMPLE_AUTOMATION_TEST(FRangeSetTest, "System.Core.Math.RangeSet", EAutomationTestFlags::EditorContext | EAutomationTestFlags::ClientContext | EAutomationTestFlags::SmokeFilter)
 
 bool FRangeSetTest::RunTest( const FString& Parameters )
 {
diff --git a/Engine/Source/Runtime/Core/Private/Tests/Math/RangeTest.cpp b/Engine/Source/Runtime/Core/Private/Tests/Math/RangeTest.cpp
--- a/Engine/Source/Runtime/Core/Private/Tests/Math/RangeTest.cpp
+++ b/Engine/Source/Runtime/Core/Private/Tests/Math/RangeTest.cpp
@@ -8,7 +8,7 @@
 PRAGMA_DISABLE_OPTIMIZATION
 
 
-IMPLEMENT_SIMPLE_AUTOMATION_TEST(FRangeTest, "System.Core.Math.Range", EAutomationTestFlags::ATF_SmokeTest)
+IMPLEMENT_SIMPLE_AUTOMATION_TEST(FRangeTest, "System.Core.Math.Range", EAutomationTestFlags::EditorContext | EAutomationTestFlags::ClientContext | EAutomationTestFlags::SmokeFilter)
 
 bool FRangeTest::RunTest( const FString& Parameters )
 {
diff --git a/Engine/Source/Runtime/Core/Private/Tests/Math/UnitConversionTests.cpp b/Engine/Source/Runtime/Core/Private/Tests/Math/UnitConversionTests.cpp
--- a/Engine/Source/Runtime/Core/Private/Tests/Math/UnitConversionTests.cpp
+++ b/Engine/Source/Runtime/Core/Private/Tests/Math/UnitConversionTests.cpp
@@ -4,7 +4,7 @@
 #include "UnitConversion.h"
 
 
-IMPLEMENT_SIMPLE_AUTOMATION_TEST(FUnitUnitTests, "System.Core.Math.Unit Conversion", EAutomationTestFlags::ATF_SmokeTest)
+IMPLEMENT_SIMPLE_AUTOMATION_TEST(FUnitUnitTests, "System.Core.Math.Unit Conversion", EAutomationTestFlags::ApplicationContextMask | EAutomationTestFlags::SmokeFilter)
 
 bool IsRoughlyEqual(double One, double Two, float Epsilon)
 {
diff --git a/Engine/Source/Runtime/Core/Private/Tests/Math/UnrealMathTest.cpp b/Engine/Source/Runtime/Core/Private/Tests/Math/UnrealMathTest.cpp
--- a/Engine/Source/Runtime/Core/Private/Tests/Math/UnrealMathTest.cpp
+++ b/Engine/Source/Runtime/Core/Private/Tests/Math/UnrealMathTest.cpp
@@ -654,7 +654,7 @@ void SetScratch( float X, float Y, float Z, float W, float U = 0.0f )
 }
 
 
-IMPLEMENT_SIMPLE_AUTOMATION_TEST(FVectorRegisterAbstractionTest, "System.Core.Math.Vector Register Abstraction Test", EAutomationTestFlags::ATF_SmokeTest)
+IMPLEMENT_SIMPLE_AUTOMATION_TEST(FVectorRegisterAbstractionTest, "System.Core.Math.Vector Register Abstraction Test", EAutomationTestFlags::ApplicationContextMask | EAutomationTestFlags::SmokeFilter)
 
 /**
  * Run a suite of vector operations to validate vector intrinsics are working on the platform
diff --git a/Engine/Source/Runtime/Core/Private/Tests/Misc/CircularBufferTest.cpp b/Engine/Source/Runtime/Core/Private/Tests/Misc/CircularBufferTest.cpp
--- a/Engine/Source/Runtime/Core/Private/Tests/Misc/CircularBufferTest.cpp
+++ b/Engine/Source/Runtime/Core/Private/Tests/Misc/CircularBufferTest.cpp
@@ -4,7 +4,7 @@
 #include "AutomationTest.h"
 
 
-IMPLEMENT_SIMPLE_AUTOMATION_TEST(FCircularBufferTest, "System.Core.Misc.CircularBuffer", EAutomationTestFlags::ATF_SmokeTest)
+IMPLEMENT_SIMPLE_AUTOMATION_TEST(FCircularBufferTest, "System.Core.Misc.CircularBuffer", EAutomationTestFlags::ApplicationContextMask | EAutomationTestFlags::SmokeFilter)
 
 bool FCircularBufferTest::RunTest( const FString& Parameters )
 {
diff --git a/Engine/Source/Runtime/Core/Private/Tests/Misc/CircularQueueTest.cpp b/Engine/Source/Runtime/Core/Private/Tests/Misc/CircularQueueTest.cpp
--- a/Engine/Source/Runtime/Core/Private/Tests/Misc/CircularQueueTest.cpp
+++ b/Engine/Source/Runtime/Core/Private/Tests/Misc/CircularQueueTest.cpp
@@ -5,7 +5,7 @@
 #include "Misc/AutomationTest.h"
 
 
-IMPLEMENT_SIMPLE_AUTOMATION_TEST(FCircularQueueTest, "System.Core.Misc.CircularQueue", EAutomationTestFlags::ATF_SmokeTest)
+IMPLEMENT_SIMPLE_AUTOMATION_TEST(FCircularQueueTest, "System.Core.Misc.CircularQueue", EAutomationTestFlags::ApplicationContextMask | EAutomationTestFlags::SmokeFilter)
 
 bool FCircularQueueTest::RunTest( const FString& Parameters )
 {
diff --git a/Engine/Source/Runtime/Core/Private/Tests/Misc/DateTimeTest.cpp b/Engine/Source/Runtime/Core/Private/Tests/Misc/DateTimeTest.cpp
--- a/Engine/Source/Runtime/Core/Private/Tests/Misc/DateTimeTest.cpp
+++ b/Engine/Source/Runtime/Core/Private/Tests/Misc/DateTimeTest.cpp
@@ -4,7 +4,7 @@
 #include "AutomationTest.h"
 
 
-IMPLEMENT_SIMPLE_AUTOMATION_TEST(FDateTimeTest, "System.Core.Misc.DateTime", EAutomationTestFlags::ATF_SmokeTest)
+IMPLEMENT_SIMPLE_AUTOMATION_TEST(FDateTimeTest, "System.Core.Misc.DateTime", EAutomationTestFlags::ApplicationContextMask | EAutomationTestFlags::SmokeFilter)
 
 bool FDateTimeTest::RunTest( const FString& Parameters )
 {
diff --git a/Engine/Source/Runtime/Core/Private/Tests/Misc/FileTest.cpp b/Engine/Source/Runtime/Core/Private/Tests/Misc/FileTest.cpp
--- a/Engine/Source/Runtime/Core/Private/Tests/Misc/FileTest.cpp
+++ b/Engine/Source/Runtime/Core/Private/Tests/Misc/FileTest.cpp
@@ -3,7 +3,7 @@
 #include "CorePrivatePCH.h"
 #include "AutomationTest.h"
 
-IMPLEMENT_SIMPLE_AUTOMATION_TEST(FFileTests, "System.Core.Misc.File", EAutomationTestFlags::ATF_SmokeTest)
+IMPLEMENT_SIMPLE_AUTOMATION_TEST(FFileTests, "System.Core.Misc.File", EAutomationTestFlags::EditorContext | EAutomationTestFlags::SmokeFilter)
 
 // These file tests are designed to ensure expected file writing behavior, as well as cross-platform consistency
 
diff --git a/Engine/Source/Runtime/Core/Private/Tests/Misc/GuidTest.cpp b/Engine/Source/Runtime/Core/Private/Tests/Misc/GuidTest.cpp
--- a/Engine/Source/Runtime/Core/Private/Tests/Misc/GuidTest.cpp
+++ b/Engine/Source/Runtime/Core/Private/Tests/Misc/GuidTest.cpp
@@ -4,7 +4,7 @@
 #include "AutomationTest.h"
 
 
-IMPLEMENT_SIMPLE_AUTOMATION_TEST(FGuidTest, "System.Core.Misc.Guid", EAutomationTestFlags::ATF_SmokeTest)
+IMPLEMENT_SIMPLE_AUTOMATION_TEST(FGuidTest, "System.Core.Misc.Guid", EAutomationTestFlags::ApplicationContextMask | EAutomationTestFlags::SmokeFilter)
 
 bool FGuidTest::RunTest( const FString& Parameters )
 {
diff --git a/Engine/Source/Runtime/Core/Private/Tests/Misc/PathsTest.cpp b/Engine/Source/Runtime/Core/Private/Tests/Misc/PathsTest.cpp
--- a/Engine/Source/Runtime/Core/Private/Tests/Misc/PathsTest.cpp
+++ b/Engine/Source/Runtime/Core/Private/Tests/Misc/PathsTest.cpp
@@ -7,7 +7,7 @@
 FString FPaths::GameProjectFilePath;
 
 
-IMPLEMENT_SIMPLE_AUTOMATION_TEST(FPathTests, "System.Core.Misc.Paths", EAutomationTestFlags::ATF_SmokeTest)
+IMPLEMENT_SIMPLE_AUTOMATION_TEST(FPathTests, "System.Core.Misc.Paths", EAutomationTestFlags::ApplicationContextMask | EAutomationTestFlags::SmokeFilter)
 
 bool FPathTests::RunTest( const FString& Parameters )
 {
diff --git a/Engine/Source/Runtime/Core/Private/Tests/Misc/QueueTest.cpp b/Engine/Source/Runtime/Core/Private/Tests/Misc/QueueTest.cpp
--- a/Engine/Source/Runtime/Core/Private/Tests/Misc/QueueTest.cpp
+++ b/Engine/Source/Runtime/Core/Private/Tests/Misc/QueueTest.cpp
@@ -4,7 +4,7 @@
 #include "AutomationTest.h"
 
 
-IMPLEMENT_SIMPLE_AUTOMATION_TEST(FQueueTest, "System.Core.Misc.Queue", EAutomationTestFlags::ATF_SmokeTest)
+IMPLEMENT_SIMPLE_AUTOMATION_TEST(FQueueTest, "System.Core.Misc.Queue", EAutomationTestFlags::ApplicationContextMask | EAutomationTestFlags::SmokeFilter)
 
 bool FQueueTest::RunTest( const FString& Parameters )
 {
diff --git a/Engine/Source/Runtime/Core/Private/Tests/Misc/TimespanTest.cpp b/Engine/Source/Runtime/Core/Private/Tests/Misc/TimespanTest.cpp
--- a/Engine/Source/Runtime/Core/Private/Tests/Misc/TimespanTest.cpp
+++ b/Engine/Source/Runtime/Core/Private/Tests/Misc/TimespanTest.cpp
@@ -4,7 +4,7 @@
 #include "AutomationTest.h"
 
 
-IMPLEMENT_SIMPLE_AUTOMATION_TEST(FTimespanTest, "System.Core.Misc.Timespan", EAutomationTestFlags::ATF_SmokeTest)
+IMPLEMENT_SIMPLE_AUTOMATION_TEST(FTimespanTest, "System.Core.Misc.Timespan", EAutomationTestFlags::ApplicationContextMask | EAutomationTestFlags::SmokeFilter)
 
 bool FTimespanTest::RunTest( const FString& Parameters )
 {
diff --git a/Engine/Source/Runtime/Core/Private/Tests/Misc/TypeContainerTest.cpp b/Engine/Source/Runtime/Core/Private/Tests/Misc/TypeContainerTest.cpp
--- a/Engine/Source/Runtime/Core/Private/Tests/Misc/TypeContainerTest.cpp
+++ b/Engine/Source/Runtime/Core/Private/Tests/Misc/TypeContainerTest.cpp
@@ -11,7 +11,7 @@
 #include "AutomationTest.h"
 
 
-IMPLEMENT_SIMPLE_AUTOMATION_TEST(FTypeContainerTest, "System.Core.Misc.TypeContainer", EAutomationTestFlags::ATF_SmokeTest)
+IMPLEMENT_SIMPLE_AUTOMATION_TEST(FTypeContainerTest, "System.Core.Misc.TypeContainer", EAutomationTestFlags::ApplicationContextMask | EAutomationTestFlags::SmokeFilter)
 
 
 /* Helpers
diff --git a/Engine/Source/Runtime/Core/Public/Misc/AutomationTest.h b/Engine/Source/Runtime/Core/Public/Misc/AutomationTest.h
--- a/Engine/Source/Runtime/Core/Public/Misc/AutomationTest.h
+++ b/Engine/Source/Runtime/Core/Public/Misc/AutomationTest.h
@@ -8,36 +8,28 @@ namespace EAutomationTestFlags
 {
 	enum Type
 	{
-		ATF_None						= 0,			//empty flags for tests that require no flags
-
-		//application type required for the test - not specifying means it will be valid for any application
-		ATF_Editor						= 0x00000001,	// Test is suitable for running within the editor
-		ATF_Game						= 0x00000002,	// Test is suitable for running within the game
-		ATF_Commandlet					= 0x00000004,	// Test is suitable for running within a commandlet
-		ATF_ApplicationMask				= ATF_Editor | ATF_Game | ATF_Commandlet,
-
-		//features required for the test - not specifying means it is valid for any feature combination
-		ATF_NonNullRHI					= 0x00000100,	// Test requires a non-null RHI to run correctly
-		ATF_RequiresUser				= 0x00000200,	// Test requires a user instigated session
-		ATF_FeatureMask					= ATF_NonNullRHI | ATF_RequiresUser,
-
-		//action options - each autotest can have only one of this and this has to match
-		ATF_VisualCommandlet			= 0x00010000,	//Test requires visual commandlet
-		ATF_ActionMask					= ATF_VisualCommandlet,
-
-		//speed of the test
-		ATF_SmokeTest					= 0x01000000,	//Test should be run in "fast" mode
-	};
-};
-
-/** Flags for specifying the type of test */
-namespace EAutomationTestType
-{
-	enum Type
-	{
-		ATT_SmokeTest = 0x00000001,	// Smoke test
-		ATT_NormalTest = 0x00000002,	// Normal test
-		ATT_StressTest = 0x00000004	// Stress test
+		//Application context required for the test - not specifying means it will be valid for any context
+		EditorContext				= 0x00000001,	// Test is suitable for running within the editor
+		ClientContext				= 0x00000002,	// Test is suitable for running within the client
+		ServerContext				= 0x00000004,	// Test is suitable for running within the server
+		CommandletContext			= 0x00000008,	// Test is suitable for running within a commandlet
+		ApplicationContextMask		= EditorContext | ClientContext | ServerContext | CommandletContext,
+
+		//Features required for the test - not specifying means it is valid for any feature combination
+		NonNullRHI					= 0x00000100,	// Test requires a non-null RHI to run correctly
+		RequiresUser				= 0x00000200,	// Test requires a user instigated session
+		FeatureMask					= NonNullRHI | RequiresUser,
+
+		//One-off flag to allow for fast disabling of tests without commenting code out
+		Disabled					= 0x00010000,	// Temp disabled and never returns for a filter
+
+		//Speed of the test
+		SmokeFilter					= 0x01000000,	//Super Fast Filter
+		EngineFilter				= 0x02000000,	//Engine Level Test
+		ProductFilter				= 0x04000000,	//Product Level Test
+		PerfFilter					= 0x08000000,	//Performance Test
+		StressFilter				= 0x10000000,	//Stress Test
+		FilterMask = SmokeFilter | EngineFilter | ProductFilter | PerfFilter | StressFilter
 	};
 };
 
@@ -85,7 +77,7 @@ public:
 
 	// Default constructor
 	FAutomationTestInfo( ) :
-		TestType( 0 )
+		TestFlags( 0 )
 		, NumParticipantsRequired( 0 )
 		, NumDevicesCurrentlyRunningTest( 0 )
 	{}
@@ -108,14 +100,14 @@ public:
 	 *
 	 * @param	InDisplayName - Name used in the UI
 	 * @param	InTestName - The test command string
-	 * @param	InTestType - Type of test e.g. smoke test
+	 * @param	InTestFlag - Test flags
 	 * @param	InParameterName - optional parameter. e.g. asset name
 	 */
-	FAutomationTestInfo( const FString& InDisplayName, const FString& InTestName, const uint8 InTestType, const int32 InNumParticipantsRequired, const FString& InParameterName = FString() )
+	FAutomationTestInfo( const FString& InDisplayName, const FString& InTestName, const uint8 InTestFlags, const int32 InNumParticipantsRequired, const FString& InParameterName = FString() )
 		: DisplayName( InDisplayName )
 		, TestName( InTestName )
 		, TestParameter( InParameterName )
-		, TestType( InTestType )
+		, TestFlags( InTestFlags )
 		, NumParticipantsRequired( InNumParticipantsRequired )
 		, NumDevicesCurrentlyRunningTest( 0 )
 	{}
@@ -123,13 +115,13 @@ public:
 public:
 
 	/**
-	 * Add a test type if a parent node.
+	 * Add a test flag if a parent node.
 	 *
-	 * @Param InTestType - the child test type to add.
+	 * @Param InTestFlags - the child test flag to add.
 	 */
-	void AddTestType( const uint8 InTestType )
+	void AddTestFlags( const uint8 InTestFlags)
 	{
-		TestType |= InTestType;
+		TestFlags |= InTestFlags;
 	}
 
 
@@ -151,7 +143,7 @@ public:
 	 */
 	const FString GetTestAsString()
 	{
-		return FString::Printf( TEXT("%s,%s,%i,%d,%s,"), *DisplayName, *TestName, TestType, NumParticipantsRequired, *TestParameter );
+		return FString::Printf( TEXT("%s,%s,%i,%d,%s,"), *DisplayName, *TestName, TestFlags, NumParticipantsRequired, *TestParameter );
 	}
 
 
@@ -182,9 +174,9 @@ public:
 	 *
 	 * @return the test type.
 	 */
-	const uint8 GetTestType() const
+	const uint8 GetTestFlags() const
 	{
-		return TestType;
+		return TestFlags;
 	}
 	
 	/**
@@ -263,7 +255,7 @@ private:
 
 		DisplayName = Pieces[0];
 		TestName = Pieces[1];
-		TestType = uint8( FCString::Atoi( *Pieces[2] ) );
+		TestFlags = uint8( FCString::Atoi( *Pieces[2] ) );
 
 		NumParticipantsRequired = FCString::Atoi( *Pieces[3] );
 
@@ -284,8 +276,8 @@ private:
 	/** Parameter - e.g. an asset name or map name */
 	FString TestParameter;
 
-	/** The test type - smoke test etc. */
-	uint8 TestType;
+	/** The test flags. */
+	uint8 TestFlags;
 
 	/** The number of participants this test requires */
 	uint32 NumParticipantsRequired;
@@ -491,9 +483,10 @@ public:
 	void SetDeveloperDirectoryIncluded(const bool bInDeveloperDirectoryIncluded);
 
 	/**
-	 * Sets whether the automation tests should include visual commandlet.
-	 */
-	void SetVisualCommandletFilter(const bool bInVisualCommandletFilterOn);
+	* Sets which set of tests to pull from.
+	*/
+	void SetRequestedTestFilter(const uint32 InRequestedTestFlags);
+	
 
 	/**
 	 * Accessor for delegate called when a png screenshot is captured 
@@ -623,7 +616,7 @@ private:
 	TQueue< TSharedPtr<IAutomationNetworkCommand> > NetworkCommands;
 
 	/** Whether we are currently executing smoke tests for startup/commandlet to minimize log spam */
-	bool bRunningSmokeTests;
+	uint32 RequestedTestFilter;
 
 	/** Time when the test began executing */
 	double StartTime;
@@ -640,9 +633,6 @@ private:
 	/** Whether we want to run automation tests on content within the Developer Directories */
 	bool bDeveloperDirectoryIncluded;
 
-	/** Whether we want to convert Anim Blueprint **/
-	bool bVisualCommandletFilterOn;
-
 	/** Wheather screenshots are enabled */
 	bool bScreenshotsEnabled;
 
@@ -1065,7 +1055,15 @@ public: \
 	{ \
 	public: \
 		TClass( const FString& InName ) \
-		:FAutomationTestBase( InName, false ) {} \
+		:FAutomationTestBase( InName, false ) {\
+			static_assert(TFlags&EAutomationTestFlags::ApplicationContextMask, "AutomationTest has no application flag.  It shouldn't run.  See AutomationTest.h."); \
+			static_assert(	(((TFlags)&EAutomationTestFlags::FilterMask) == EAutomationTestFlags::SmokeFilter) || \
+							(((TFlags)&EAutomationTestFlags::FilterMask) == EAutomationTestFlags::EngineFilter) || \
+							(((TFlags)&EAutomationTestFlags::FilterMask) == EAutomationTestFlags::ProductFilter) || \
+							(((TFlags)&EAutomationTestFlags::FilterMask) == EAutomationTestFlags::PerfFilter) || \
+							(((TFlags)&EAutomationTestFlags::FilterMask) == EAutomationTestFlags::StressFilter), \
+							"All AutomationTests must have exactly 1 filter type specified.  See AutomationTest.h."); \
+		} \
 		virtual uint32 GetTestFlags() const override { return TFlags; } \
 		virtual bool IsStressTest() const { return false; } \
 		virtual uint32 GetRequiredDeviceNum() const override { return 1; } \
@@ -1084,8 +1082,16 @@ public: \
 	{ \
 	public: \
 		TClass( const FString& InName ) \
-		:FAutomationTestBase( InName, true ) {} \
-		virtual uint32 GetTestFlags() const override { return ((TFlags) & ~(EAutomationTestFlags::ATF_SmokeTest)); } \
+		:FAutomationTestBase( InName, true ) { \
+			static_assert(TFlags&EAutomationTestFlags::ApplicationContextMask, "AutomationTest has no application flag.  It shouldn't run.  See AutomationTest.h."); \
+			static_assert(	(((TFlags)&EAutomationTestFlags::FilterMask) == EAutomationTestFlags::SmokeFilter) || \
+							(((TFlags)&EAutomationTestFlags::FilterMask) == EAutomationTestFlags::EngineFilter) || \
+							(((TFlags)&EAutomationTestFlags::FilterMask) == EAutomationTestFlags::ProductFilter) || \
+							(((TFlags)&EAutomationTestFlags::FilterMask) == EAutomationTestFlags::PerfFilter) || \
+							(((TFlags)&EAutomationTestFlags::FilterMask) == EAutomationTestFlags::StressFilter), \
+							"All AutomationTests must have exactly 1 filter type specified.  See AutomationTest.h."); \
+		} \
+		virtual uint32 GetTestFlags() const override { return (TFlags & ~(EAutomationTestFlags::SmokeFilter)); } \
 		virtual bool IsStressTest() const { return true; } \
 		virtual uint32 GetRequiredDeviceNum() const override { return 1; } \
 	protected: \
@@ -1099,8 +1105,16 @@ public: \
 	{ \
 	public: \
 		TClass( const FString& InName ) \
-		:FAutomationTestBase( InName, false ) {} \
-		virtual uint32 GetTestFlags() const override { return ((TFlags) & ~(EAutomationTestFlags::ATF_Editor | EAutomationTestFlags::ATF_Commandlet | EAutomationTestFlags::ATF_SmokeTest)); } \
+		:FAutomationTestBase( InName, false ) { \
+			static_assert(TFlags&EAutomationTestFlags::ApplicationContextMask, "AutomationTest has no application flag.  It shouldn't run.  See AutomationTest.h."); \
+			static_assert(	(((TFlags)&EAutomationTestFlags::FilterMask) == EAutomationTestFlags::SmokeFilter) || \
+							(((TFlags)&EAutomationTestFlags::FilterMask) == EAutomationTestFlags::EngineFilter) || \
+							(((TFlags)&EAutomationTestFlags::FilterMask) == EAutomationTestFlags::ProductFilter) || \
+							(((TFlags)&EAutomationTestFlags::FilterMask) == EAutomationTestFlags::PerfFilter) || \
+							(((TFlags)&EAutomationTestFlags::FilterMask) == EAutomationTestFlags::StressFilter), \
+							"All AutomationTests must have exactly 1 filter type specified.  See AutomationTest.h."); \
+		} \
+		virtual uint32 GetTestFlags() const override { return (TFlags & ~(EAutomationTestFlags::EditorContext | EAutomationTestFlags::CommandletContext | EAutomationTestFlags::SmokeFilter)); } \
 		virtual uint32 GetRequiredDeviceNum() const override { return NumParticipants; } \
 	protected: \
 		virtual void GetTests(TArray<FString>& OutBeautifiedNames, TArray <FString>& OutTestCommands) const override \
diff --git a/Engine/Source/Runtime/Core/Public/Misc/Build.h b/Engine/Source/Runtime/Core/Public/Misc/Build.h
--- a/Engine/Source/Runtime/Core/Public/Misc/Build.h
+++ b/Engine/Source/Runtime/Core/Public/Misc/Build.h
@@ -96,9 +96,10 @@
 	#define HACK_HEADER_GENERATOR 0
 #endif
 
-/** Whether we are compiling with automation worker functionality */
+/** Whether we are compiling with automation worker functionality.  Note that automation worker defaults to enabled in
+    UE_BUILD_TEST configuration, so that it can be used for performance testing on devices */
 #ifndef WITH_AUTOMATION_WORKER
-	#define WITH_AUTOMATION_WORKER !(UE_BUILD_SHIPPING || UE_BUILD_TEST || HACK_HEADER_GENERATOR)
+	#define WITH_AUTOMATION_WORKER !(UE_BUILD_SHIPPING || HACK_HEADER_GENERATOR)
 #endif
 
 /**
diff --git a/Engine/Source/Runtime/Engine/Private/Commandlets/SmokeTestCommandlet.cpp b/Engine/Source/Runtime/Engine/Private/Commandlets/SmokeTestCommandlet.cpp
--- a/Engine/Source/Runtime/Engine/Private/Commandlets/SmokeTestCommandlet.cpp
+++ b/Engine/Source/Runtime/Engine/Private/Commandlets/SmokeTestCommandlet.cpp
@@ -29,7 +29,7 @@ int32 USmokeTestCommandlet::Main( const FString& Params )
 }
 
 
-IMPLEMENT_SIMPLE_AUTOMATION_TEST(FRunServer, "System.Core.Misc.Run Server", EAutomationTestFlags::ATF_Commandlet | EAutomationTestFlags::ATF_SmokeTest)
+IMPLEMENT_SIMPLE_AUTOMATION_TEST(FRunServer, "System.Core.Misc.Run Server", EAutomationTestFlags::CommandletContext | EAutomationTestFlags::SmokeFilter)
 
 bool FRunServer::RunTest(const FString& Parameters)
 {
diff --git a/Engine/Source/Runtime/Engine/Private/Tests/AnalyticsTest.cpp b/Engine/Source/Runtime/Engine/Private/Tests/AnalyticsTest.cpp
--- a/Engine/Source/Runtime/Engine/Private/Tests/AnalyticsTest.cpp
+++ b/Engine/Source/Runtime/Engine/Private/Tests/AnalyticsTest.cpp
@@ -10,7 +10,7 @@
 /**
 * Artificial Record Event for analytics - Simulates the engine startup simulation.  
 */
-IMPLEMENT_SIMPLE_AUTOMATION_TEST(FAnalyticStartUpSimTest, "System.Engine.Analytic.Record Event - Simulate Program Start", EAutomationTestFlags::ATF_FeatureMask)
+IMPLEMENT_SIMPLE_AUTOMATION_TEST(FAnalyticStartUpSimTest, "System.Engine.Analytic.Record Event - Simulate Program Start", EAutomationTestFlags::FeatureMask | EAutomationTestFlags::ClientContext | EAutomationTestFlags::EngineFilter)
 
 bool FAnalyticStartUpSimTest::RunTest(const FString& Parameters)
 {
@@ -81,7 +81,7 @@ bool FAnalyticStartUpSimTest::RunTest(const FString& Parameters)
 /**
 * FAnalyticsEventAttribute Unit Test.
 */
-IMPLEMENT_SIMPLE_AUTOMATION_TEST(FAnalyticsEventAttributeUnitTest, "System.Engine.Analytic.EventAttribute Struct Unit Test", EAutomationTestFlags::ATF_FeatureMask)
+IMPLEMENT_SIMPLE_AUTOMATION_TEST(FAnalyticsEventAttributeUnitTest, "System.Engine.Analytic.EventAttribute Struct Unit Test", EAutomationTestFlags::FeatureMask | EAutomationTestFlags::ClientContext | EAutomationTestFlags::EngineFilter)
 
 bool FAnalyticsEventAttributeUnitTest::RunTest(const FString& Parameters)
 {
diff --git a/Engine/Source/Runtime/Engine/Private/Tests/EngineAutomationTests.cpp b/Engine/Source/Runtime/Engine/Private/Tests/EngineAutomationTests.cpp
--- a/Engine/Source/Runtime/Engine/Private/Tests/EngineAutomationTests.cpp
+++ b/Engine/Source/Runtime/Engine/Private/Tests/EngineAutomationTests.cpp
@@ -19,7 +19,7 @@ namespace
 	UWorld* GetSimpleEngineAutomationTestGameWorld(const int32 TestFlags)
 	{
 		// Accessing the game world is only valid for game-only 
-		check((TestFlags & EAutomationTestFlags::ATF_ApplicationMask) == EAutomationTestFlags::ATF_Game);
+		check((TestFlags & EAutomationTestFlags::ApplicationContextMask) == EAutomationTestFlags::ClientContext);
 		check(GEngine->GetWorldContexts().Num() == 1);
 		check(GEngine->GetWorldContexts()[0].WorldType == EWorldType::Game);
 
@@ -63,7 +63,7 @@ namespace
 /**
  * SetRes Verification - Verify changing resolution works
  */
-IMPLEMENT_SIMPLE_AUTOMATION_TEST( FSetResTest, "System.Windows.Set Resolution", EAutomationTestFlags::ATF_Game )
+IMPLEMENT_SIMPLE_AUTOMATION_TEST(FSetResTest, "System.Windows.Set Resolution", EAutomationTestFlags::ClientContext | EAutomationTestFlags::EngineFilter)
 
 /** 
  * Change resolutions, wait, and change back
@@ -97,7 +97,7 @@ bool FSetResTest::RunTest(const FString& Parameters)
 /**
  * Stats verification - Toggle various "stats" commands
  */
-IMPLEMENT_SIMPLE_AUTOMATION_TEST( FStatsVerificationMapTest, "System.Maps.Stats Verification", EAutomationTestFlags::ATF_Game )
+IMPLEMENT_SIMPLE_AUTOMATION_TEST(FStatsVerificationMapTest, "System.Maps.Stats Verification", EAutomationTestFlags::ClientContext | EAutomationTestFlags::EngineFilter)
 
 /** 
  * Execute the loading of one map to verify screen captures and performance captures work
@@ -192,7 +192,7 @@ bool FStatsVerificationMapTest::RunTest(const FString& Parameters)
  * LoadAutomationMap
  * Verification automation test to make sure features of map loading work (load, screen capture, performance capture)
  */
-IMPLEMENT_SIMPLE_AUTOMATION_TEST( FPerformanceCaptureTest, "System.Maps.Performance Capture", EAutomationTestFlags::ATF_Game )
+IMPLEMENT_SIMPLE_AUTOMATION_TEST(FPerformanceCaptureTest, "System.Maps.Performance Capture", EAutomationTestFlags::ClientContext | EAutomationTestFlags::EngineFilter)
 
 /** 
  * Execute the loading of one map to verify screen captures and performance captures work
@@ -229,7 +229,7 @@ bool FTakeViewportScreenshotCommand::Update()
  * LoadAllMapsInGame
  * Verification automation test to make sure loading all maps succeed without crashing AND does performance captures
  */
-IMPLEMENT_COMPLEX_AUTOMATION_TEST( FLoadAllMapsInGameTest, "Project.Maps.Load All In Game", EAutomationTestFlags::ATF_Game )
+IMPLEMENT_COMPLEX_AUTOMATION_TEST(FLoadAllMapsInGameTest, "Project.Maps.Load All In Game", EAutomationTestFlags::ClientContext | EAutomationTestFlags::StressFilter)
 
 /** 
  * Requests a enumeration of all maps to be loaded
@@ -277,7 +277,7 @@ bool FLoadAllMapsInGameTest::RunTest(const FString& Parameters)
  * SaveGameTest
  * Test makes sure a save game (without UI) saves and loads correctly
  */
-IMPLEMENT_SIMPLE_AUTOMATION_TEST( FSaveGameTest, "System.Engine.Game.Noninteractive Save", EAutomationTestFlags::ATF_Game )
+IMPLEMENT_SIMPLE_AUTOMATION_TEST(FSaveGameTest, "System.Engine.Game.Noninteractive Save", EAutomationTestFlags::ClientContext | EAutomationTestFlags::EngineFilter)
 
 /** 
  * Saves and loads a savegame file
@@ -339,7 +339,7 @@ bool FSaveGameTest::RunTest(const FString& Parameters)
 /**
  * Automation test to load a map and capture FPS performance charts
  */
-IMPLEMENT_COMPLEX_AUTOMATION_TEST(FCinematicFPSPerfTest, "Project.Maps.Cinematic FPS Perf Capture", (EAutomationTestFlags::ATF_Game | EAutomationTestFlags::ATF_NonNullRHI));
+IMPLEMENT_COMPLEX_AUTOMATION_TEST(FCinematicFPSPerfTest, "Project.Maps.Cinematic FPS Perf Capture", (EAutomationTestFlags::ClientContext | EAutomationTestFlags::NonNullRHI | EAutomationTestFlags::StressFilter));
 
 void FCinematicFPSPerfTest::GetTests(TArray<FString>& OutBeautifiedNames, TArray <FString>& OutTestCommands) const
 {
@@ -428,7 +428,7 @@ bool FCinematicFPSPerfTest::RunTest(const FString& Parameters)
 }
 
 
-IMPLEMENT_SIMPLE_AUTOMATION_TEST(FAutomationLogAddMessage, "System.Automation.Log.Add Log Message", EAutomationTestFlags::ATF_None)
+IMPLEMENT_SIMPLE_AUTOMATION_TEST(FAutomationLogAddMessage, "System.Automation.Log.Add Log Message", EAutomationTestFlags::ApplicationContextMask | EAutomationTestFlags::EngineFilter)
 
 bool FAutomationLogAddMessage::RunTest(const FString& Parameters)
 {
@@ -445,7 +445,7 @@ bool FAutomationLogAddMessage::RunTest(const FString& Parameters)
 	return true;
 }
 
-IMPLEMENT_SIMPLE_AUTOMATION_TEST(FAutomationLogAddWarning, "System.Automation.Log.Add Warning Message", EAutomationTestFlags::ATF_None)
+IMPLEMENT_SIMPLE_AUTOMATION_TEST(FAutomationLogAddWarning, "System.Automation.Log.Add Warning Message", EAutomationTestFlags::EditorContext | EAutomationTestFlags::EngineFilter)
 
 bool FAutomationLogAddWarning::RunTest(const FString& Parameters)
 {
@@ -462,7 +462,7 @@ bool FAutomationLogAddWarning::RunTest(const FString& Parameters)
 	return true;
 }
 
-IMPLEMENT_SIMPLE_AUTOMATION_TEST(FAutomationLogAddError, "System.Automation.Log.Add Error Message", EAutomationTestFlags::ATF_None)
+IMPLEMENT_SIMPLE_AUTOMATION_TEST(FAutomationLogAddError, "System.Automation.Log.Add Error Message", EAutomationTestFlags::EditorContext | EAutomationTestFlags::EngineFilter)
 
 bool FAutomationLogAddError::RunTest(const FString& Parameters)
 {
diff --git a/Engine/Source/Runtime/Engine/Private/Tests/ExternalToolTest.cpp b/Engine/Source/Runtime/Engine/Private/Tests/ExternalToolTest.cpp
--- a/Engine/Source/Runtime/Engine/Private/Tests/ExternalToolTest.cpp
+++ b/Engine/Source/Runtime/Engine/Private/Tests/ExternalToolTest.cpp
@@ -77,7 +77,7 @@ bool FWaitForProcessToCompleteLatentCommand::Update()
 /**
  * RunExternalToolTest
  */
-IMPLEMENT_COMPLEX_AUTOMATION_TEST(FRunExternalToolTest, "External", EAutomationTestFlags::ATF_NonNullRHI)
+IMPLEMENT_COMPLEX_AUTOMATION_TEST(FRunExternalToolTest, "External", EAutomationTestFlags::NonNullRHI | EAutomationTestFlags::EngineFilter)
 
 /** 
  * Find all the external too commands
diff --git a/Engine/Source/Runtime/Engine/Private/Tests/IntSerializationTest.cpp b/Engine/Source/Runtime/Engine/Private/Tests/IntSerializationTest.cpp
--- a/Engine/Source/Runtime/Engine/Private/Tests/IntSerializationTest.cpp
+++ b/Engine/Source/Runtime/Engine/Private/Tests/IntSerializationTest.cpp
@@ -4,7 +4,7 @@
 #include "AutomationTest.h"
 #include "Runtime/Engine/Classes/Engine/IntSerialization.h"
 
-IMPLEMENT_SIMPLE_AUTOMATION_TEST(FIntSerializationTest, "System.Engine.Int Serialization", EAutomationTestFlags::ATF_SmokeTest)
+IMPLEMENT_SIMPLE_AUTOMATION_TEST(FIntSerializationTest, "System.Engine.Int Serialization", EAutomationTestFlags::ApplicationContextMask | EAutomationTestFlags::SmokeFilter)
 
 bool FIntSerializationTest::RunTest (const FString& Parameters)
 {
diff --git a/Engine/Source/Runtime/Engine/Private/Tests/Internationalization/TextPropertyTest.cpp b/Engine/Source/Runtime/Engine/Private/Tests/Internationalization/TextPropertyTest.cpp
--- a/Engine/Source/Runtime/Engine/Private/Tests/Internationalization/TextPropertyTest.cpp
+++ b/Engine/Source/Runtime/Engine/Private/Tests/Internationalization/TextPropertyTest.cpp
@@ -11,7 +11,7 @@ UTextPropertyTestObject::UTextPropertyTestObject(const FObjectInitializer& Objec
 {
 }
 
-IMPLEMENT_SIMPLE_AUTOMATION_TEST(FTextPropertyTest, "System.Engine.Internationalization.Text Property Test", EAutomationTestFlags::ATF_SmokeTest)
+IMPLEMENT_SIMPLE_AUTOMATION_TEST(FTextPropertyTest, "System.Engine.Internationalization.Text Property Test", EAutomationTestFlags::ApplicationContextMask | EAutomationTestFlags::SmokeFilter)
 
 bool FTextPropertyTest::RunTest (const FString& Parameters)
 {
diff --git a/Engine/Source/Runtime/Engine/Private/Tests/NetworkAutomationTests.cpp b/Engine/Source/Runtime/Engine/Private/Tests/NetworkAutomationTests.cpp
--- a/Engine/Source/Runtime/Engine/Private/Tests/NetworkAutomationTests.cpp
+++ b/Engine/Source/Runtime/Engine/Private/Tests/NetworkAutomationTests.cpp
@@ -24,7 +24,7 @@ namespace EMultiplayerAutomationRoles
  * 2-Multiplayer session automation test
  * Verification for 2 player multiplayer session start up and tear down
  */
-IMPLEMENT_NETWORKED_AUTOMATION_TEST( FMultiplayer4PlayerTest, "System.Networking.Multiplayer.TwoPlayerSessionStartupShutdown", EAutomationTestFlags::ATF_Game, EMultiplayerAutomationRoles::MaxNumParticipants )
+IMPLEMENT_NETWORKED_AUTOMATION_TEST(FMultiplayer4PlayerTest, "System.Networking.Multiplayer.TwoPlayerSessionStartupShutdown", EAutomationTestFlags::ClientContext | EAutomationTestFlags::EngineFilter, EMultiplayerAutomationRoles::MaxNumParticipants)
 
 
 /** 
@@ -36,7 +36,7 @@ IMPLEMENT_NETWORKED_AUTOMATION_TEST( FMultiplayer4PlayerTest, "System.Networking
 bool FMultiplayer4PlayerTest::RunTest(const FString& Parameters)
 {
 	// Accessing the game world is only valid for game-only 
-	check((GetTestFlags() & EAutomationTestFlags::ATF_ApplicationMask) == EAutomationTestFlags::ATF_Game);
+	check((GetTestFlags() & EAutomationTestFlags::ApplicationContextMask) == EAutomationTestFlags::ClientContext);
 	check(GEngine->GetWorldContexts().Num() == 1);
 	check(GEngine->GetWorldContexts()[0].WorldType == EWorldType::Game);
 
diff --git a/Engine/Source/Runtime/Engine/Private/Tests/OSSTests.cpp b/Engine/Source/Runtime/Engine/Private/Tests/OSSTests.cpp
--- a/Engine/Source/Runtime/Engine/Private/Tests/OSSTests.cpp
+++ b/Engine/Source/Runtime/Engine/Private/Tests/OSSTests.cpp
@@ -18,7 +18,7 @@ DEFINE_LOG_CATEGORY_STATIC(LogHackAutomationTests, Log, All);
 //////////////////////////////////////////////////////////////////////////
 
 
-IMPLEMENT_SIMPLE_AUTOMATION_TEST(FHttpTest,"System.OSS.Test Http", EAutomationTestFlags::ATF_Editor )
+IMPLEMENT_SIMPLE_AUTOMATION_TEST(FHttpTest, "System.OSS.Test Http", EAutomationTestFlags::EditorContext | EAutomationTestFlags::EngineFilter)
 
 bool FHttpTest::RunTest(const FString& Parameters)
 {
@@ -28,7 +28,7 @@ bool FHttpTest::RunTest(const FString& Parameters)
 	return true;
 }
 
-IMPLEMENT_SIMPLE_AUTOMATION_TEST(FMCPFriendsTest,"System.OSS.MCP.Test Friends", EAutomationTestFlags::ATF_Editor )
+IMPLEMENT_SIMPLE_AUTOMATION_TEST(FMCPFriendsTest, "System.OSS.MCP.Test Friends", EAutomationTestFlags::EditorContext | EAutomationTestFlags::EngineFilter)
 
 bool FMCPFriendsTest::RunTest(const FString& Parameters)
 {
@@ -39,7 +39,7 @@ bool FMCPFriendsTest::RunTest(const FString& Parameters)
 	return true;
 }
 
-IMPLEMENT_SIMPLE_AUTOMATION_TEST(FAmazonFriendsTest,"System.OSS.Amazon.Test Friends", EAutomationTestFlags::ATF_Editor )
+IMPLEMENT_SIMPLE_AUTOMATION_TEST(FAmazonFriendsTest, "System.OSS.Amazon.Test Friends", EAutomationTestFlags::EditorContext | EAutomationTestFlags::EngineFilter)
 
 bool FAmazonFriendsTest::RunTest(const FString& Parameters)
 {
@@ -49,7 +49,7 @@ bool FAmazonFriendsTest::RunTest(const FString& Parameters)
 	return true;
 }
 
-IMPLEMENT_SIMPLE_AUTOMATION_TEST(FFacebookFriendsTest,"System.OSS.Facebook.Test Friends", EAutomationTestFlags::ATF_Editor )
+IMPLEMENT_SIMPLE_AUTOMATION_TEST(FFacebookFriendsTest, "System.OSS.Facebook.Test Friends", EAutomationTestFlags::EditorContext | EAutomationTestFlags::EngineFilter)
 
 bool FFacebookFriendsTest::RunTest(const FString& Parameters)
 {
@@ -59,7 +59,7 @@ bool FFacebookFriendsTest::RunTest(const FString& Parameters)
 	return true;
 }
 
-IMPLEMENT_SIMPLE_AUTOMATION_TEST(FSteamFriendsTest,"System.OSS.Steam.Test Friends", EAutomationTestFlags::ATF_Editor )
+IMPLEMENT_SIMPLE_AUTOMATION_TEST(FSteamFriendsTest, "System.OSS.Steam.Test Friends", EAutomationTestFlags::EditorContext | EAutomationTestFlags::EngineFilter)
 
 bool FSteamFriendsTest::RunTest(const FString& Parameters)
 {
@@ -69,7 +69,7 @@ bool FSteamFriendsTest::RunTest(const FString& Parameters)
 	return true;
 }
 
-IMPLEMENT_SIMPLE_AUTOMATION_TEST(FIOSFriendsTest,"System.OSS.IOS.Test Friends", EAutomationTestFlags::ATF_Editor )
+IMPLEMENT_SIMPLE_AUTOMATION_TEST(FIOSFriendsTest, "System.OSS.IOS.Test Friends", EAutomationTestFlags::EditorContext | EAutomationTestFlags::EngineFilter)
 
 bool FIOSFriendsTest::RunTest(const FString& Parameters)
 {
@@ -83,7 +83,7 @@ bool FIOSFriendsTest::RunTest(const FString& Parameters)
 
 /////////////Currently crashing//////////////
 
-IMPLEMENT_SIMPLE_AUTOMATION_TEST(FNullFriendsTest,"System.OSS.NULL.Test Friends", EAutomationTestFlags::ATF_Editor )
+IMPLEMENT_SIMPLE_AUTOMATION_TEST(FNullFriendsTest,"System.OSS.NULL.Test Friends", EAutomationTestFlags::EditorContext )
 
 bool FNullFriendsTest::RunTest(const FString& Parameters)
 {
@@ -95,7 +95,7 @@ bool FNullFriendsTest::RunTest(const FString& Parameters)
 }
 */
 
-IMPLEMENT_SIMPLE_AUTOMATION_TEST(FMCPCloudTest,"System.OSS.MCP.Test Cloud", EAutomationTestFlags::ATF_Editor )
+IMPLEMENT_SIMPLE_AUTOMATION_TEST(FMCPCloudTest, "System.OSS.MCP.Test Cloud", EAutomationTestFlags::EditorContext | EAutomationTestFlags::EngineFilter)
 
 bool FMCPCloudTest::RunTest(const FString& Parameters)
 {
@@ -105,7 +105,7 @@ bool FMCPCloudTest::RunTest(const FString& Parameters)
 	return true;
 }
 
-IMPLEMENT_SIMPLE_AUTOMATION_TEST(FSteamCloudTest,"System.OSS.Steam.Test Cloud", EAutomationTestFlags::ATF_Editor )
+IMPLEMENT_SIMPLE_AUTOMATION_TEST(FSteamCloudTest, "System.OSS.Steam.Test Cloud", EAutomationTestFlags::EditorContext | EAutomationTestFlags::EngineFilter)
 
 bool FSteamCloudTest::RunTest(const FString& Parameters)
 {
@@ -119,7 +119,7 @@ bool FSteamCloudTest::RunTest(const FString& Parameters)
 
 /////////////Currently crashing//////////////
 
-IMPLEMENT_SIMPLE_AUTOMATION_TEST(FNullCloudTest,"System.OSS.NULL.Test Cloud", EAutomationTestFlags::ATF_Editor )
+IMPLEMENT_SIMPLE_AUTOMATION_TEST(FNullCloudTest,"System.OSS.NULL.Test Cloud", EAutomationTestFlags::EditorContext )
 
 bool FNullCloudTest::RunTest(const FString& Parameters)
 {
@@ -134,7 +134,7 @@ bool FNullCloudTest::RunTest(const FString& Parameters)
 
 /////////////Currently crashing//////////////
 
-IMPLEMENT_SIMPLE_AUTOMATION_TEST(FNullLeaderBoardTest,"System.OSS.NULL.Test Leaderboards", EAutomationTestFlags::ATF_Editor )
+IMPLEMENT_SIMPLE_AUTOMATION_TEST(FNullLeaderBoardTest,"System.OSS.NULL.Test Leaderboards", EAutomationTestFlags::EditorContext)
 
 bool FNullLeaderBoardTest::RunTest(const FString& Parameters)
 {
@@ -145,7 +145,7 @@ bool FNullLeaderBoardTest::RunTest(const FString& Parameters)
 }
 */
 
-IMPLEMENT_SIMPLE_AUTOMATION_TEST(FIosLeaderBoardTest,"System.OSS.IOS.Test Leaderboards", EAutomationTestFlags::ATF_Editor )
+IMPLEMENT_SIMPLE_AUTOMATION_TEST(FIosLeaderBoardTest, "System.OSS.IOS.Test Leaderboards", EAutomationTestFlags::EditorContext | EAutomationTestFlags::EngineFilter)
 
 bool FIosLeaderBoardTest::RunTest(const FString& Parameters)
 {
@@ -156,7 +156,7 @@ bool FIosLeaderBoardTest::RunTest(const FString& Parameters)
 }
 
 
-IMPLEMENT_SIMPLE_AUTOMATION_TEST(FSteamLeaderBoardTest,"System.OSS.Steam.Test Leaderboards", EAutomationTestFlags::ATF_Editor )
+IMPLEMENT_SIMPLE_AUTOMATION_TEST(FSteamLeaderBoardTest, "System.OSS.Steam.Test Leaderboards", EAutomationTestFlags::EditorContext | EAutomationTestFlags::EngineFilter)
 
 bool FSteamLeaderBoardTest::RunTest(const FString& Parameters)
 {
@@ -166,7 +166,7 @@ bool FSteamLeaderBoardTest::RunTest(const FString& Parameters)
 	return true;
 }
 
-IMPLEMENT_SIMPLE_AUTOMATION_TEST(FMCPTimeTest,"System.OSS.MCP.Test Time", EAutomationTestFlags::ATF_Editor )
+IMPLEMENT_SIMPLE_AUTOMATION_TEST(FMCPTimeTest, "System.OSS.MCP.Test Time", EAutomationTestFlags::EditorContext | EAutomationTestFlags::EngineFilter)
 
 bool FMCPTimeTest::RunTest(const FString& Parameters)
 {
@@ -176,7 +176,7 @@ bool FMCPTimeTest::RunTest(const FString& Parameters)
 	return true;
 }
 
-IMPLEMENT_SIMPLE_AUTOMATION_TEST(FMCPIdentityTest,"System.OSS.MCP.Test Identity", EAutomationTestFlags::ATF_Editor )
+IMPLEMENT_SIMPLE_AUTOMATION_TEST(FMCPIdentityTest, "System.OSS.MCP.Test Identity", EAutomationTestFlags::EditorContext | EAutomationTestFlags::EngineFilter)
 
 bool FMCPIdentityTest::RunTest(const FString& Parameters)
 {
@@ -186,7 +186,7 @@ bool FMCPIdentityTest::RunTest(const FString& Parameters)
 	return true;
 }
 
-IMPLEMENT_SIMPLE_AUTOMATION_TEST(FSteamIdentityTest,"System.OSS.Steam.Test Identity", EAutomationTestFlags::ATF_Editor )
+IMPLEMENT_SIMPLE_AUTOMATION_TEST(FSteamIdentityTest, "System.OSS.Steam.Test Identity", EAutomationTestFlags::EditorContext | EAutomationTestFlags::EngineFilter)
 
 bool FSteamIdentityTest::RunTest(const FString& Parameters)
 {
@@ -196,7 +196,7 @@ bool FSteamIdentityTest::RunTest(const FString& Parameters)
 	return true;
 }
 
-IMPLEMENT_SIMPLE_AUTOMATION_TEST(FAmazonIdentityTest,"System.OSS.Amazon.Test Identity", EAutomationTestFlags::ATF_Editor )
+IMPLEMENT_SIMPLE_AUTOMATION_TEST(FAmazonIdentityTest, "System.OSS.Amazon.Test Identity", EAutomationTestFlags::EditorContext | EAutomationTestFlags::EngineFilter)
 
 bool FAmazonIdentityTest::RunTest(const FString& Parameters)
 {
@@ -206,7 +206,7 @@ bool FAmazonIdentityTest::RunTest(const FString& Parameters)
 	return true;
 }
 
-IMPLEMENT_SIMPLE_AUTOMATION_TEST(FFacebookIdentityTest,"System.OSS.Facebook.Test Identity", EAutomationTestFlags::ATF_Editor )
+IMPLEMENT_SIMPLE_AUTOMATION_TEST(FFacebookIdentityTest, "System.OSS.Facebook.Test Identity", EAutomationTestFlags::EditorContext | EAutomationTestFlags::EngineFilter)
 
 bool FFacebookIdentityTest::RunTest(const FString& Parameters)
 {
@@ -216,7 +216,7 @@ bool FFacebookIdentityTest::RunTest(const FString& Parameters)
 	return true;
 }
 
-IMPLEMENT_SIMPLE_AUTOMATION_TEST(FIOSIdentityTest,"System.OSS.IOS.Test Identity", EAutomationTestFlags::ATF_Editor )
+IMPLEMENT_SIMPLE_AUTOMATION_TEST(FIOSIdentityTest, "System.OSS.IOS.Test Identity", EAutomationTestFlags::EditorContext | EAutomationTestFlags::EngineFilter)
 
 bool FIOSIdentityTest::RunTest(const FString& Parameters)
 {
@@ -230,7 +230,7 @@ bool FIOSIdentityTest::RunTest(const FString& Parameters)
 
 /////////////Currently crashing//////////////
 
-IMPLEMENT_SIMPLE_AUTOMATION_TEST(FNullIdentityTest,"System.OSS.NULL.Test Identity", EAutomationTestFlags::ATF_Editor )
+IMPLEMENT_SIMPLE_AUTOMATION_TEST(FNullIdentityTest,"System.OSS.NULL.Test Identity", EAutomationTestFlags::EditorContext )
 
 bool FNullIdentityTest::RunTest(const FString& Parameters)
 {
@@ -241,7 +241,7 @@ bool FNullIdentityTest::RunTest(const FString& Parameters)
 }
 */
 
-IMPLEMENT_SIMPLE_AUTOMATION_TEST(FMCPEntitlementsTest,"System.OSS.MCP.Test Entitlements", EAutomationTestFlags::ATF_Editor )
+IMPLEMENT_SIMPLE_AUTOMATION_TEST(FMCPEntitlementsTest, "System.OSS.MCP.Test Entitlements", EAutomationTestFlags::EditorContext | EAutomationTestFlags::EngineFilter)
 
 bool FMCPEntitlementsTest::RunTest(const FString& Parameters)
 {
@@ -251,7 +251,7 @@ bool FMCPEntitlementsTest::RunTest(const FString& Parameters)
 	return true;
 }
 
-IMPLEMENT_SIMPLE_AUTOMATION_TEST(FMCPTitleFileTest,"System.OSS.MCP.Test Title File", EAutomationTestFlags::ATF_Editor )
+IMPLEMENT_SIMPLE_AUTOMATION_TEST(FMCPTitleFileTest, "System.OSS.MCP.Test Title File", EAutomationTestFlags::EditorContext | EAutomationTestFlags::EngineFilter)
 
 bool FMCPTitleFileTest::RunTest(const FString& Parameters)
 {
@@ -261,7 +261,7 @@ bool FMCPTitleFileTest::RunTest(const FString& Parameters)
 	return true;
 }
 
-IMPLEMENT_SIMPLE_AUTOMATION_TEST(FSteamAchievementsTest,"System.OSS.Steam.Test Achievements", EAutomationTestFlags::ATF_Editor )
+IMPLEMENT_SIMPLE_AUTOMATION_TEST(FSteamAchievementsTest, "System.OSS.Steam.Test Achievements", EAutomationTestFlags::EditorContext | EAutomationTestFlags::EngineFilter)
 
 bool FSteamAchievementsTest::RunTest(const FString& Parameters)
 {
@@ -275,7 +275,7 @@ bool FSteamAchievementsTest::RunTest(const FString& Parameters)
 
 /////////////Currently crashing//////////////
 
-IMPLEMENT_SIMPLE_AUTOMATION_TEST(FNullAchievementsTest,"System.OSS.NULL.Test Achievements", EAutomationTestFlags::ATF_Editor )
+IMPLEMENT_SIMPLE_AUTOMATION_TEST(FNullAchievementsTest,"System.OSS.NULL.Test Achievements", EAutomationTestFlags::EditorContext )
 
 bool FNullAchievementsTest::RunTest(const FString& Parameters)
 {
diff --git a/Engine/Source/Runtime/Engine/Private/TimerManagerTests.cpp b/Engine/Source/Runtime/Engine/Private/TimerManagerTests.cpp
--- a/Engine/Source/Runtime/Engine/Private/TimerManagerTests.cpp
+++ b/Engine/Source/Runtime/Engine/Private/TimerManagerTests.cpp
@@ -3,7 +3,7 @@
 #include "EnginePrivate.h"
 #include "TimerManager.h"
 
-IMPLEMENT_SIMPLE_AUTOMATION_TEST(FTimerManagerTest, "System.Engine.TimerManager", EAutomationTestFlags::ATF_Editor)
+IMPLEMENT_SIMPLE_AUTOMATION_TEST(FTimerManagerTest, "System.Engine.TimerManager", EAutomationTestFlags::EditorContext | EAutomationTestFlags::EngineFilter)
 
 #define TIMER_TEST_TEXT( Format, ... ) FString::Printf(TEXT("%s - %d: %s"), TEXT(__FILE__) , __LINE__ , *FString::Printf(TEXT(Format), ##__VA_ARGS__) )
 
diff --git a/Engine/Source/Runtime/Engine/Private/VisualLogger/VisualLoggerAutomationTests.cpp b/Engine/Source/Runtime/Engine/Private/VisualLogger/VisualLoggerAutomationTests.cpp
--- a/Engine/Source/Runtime/Engine/Private/VisualLogger/VisualLoggerAutomationTests.cpp
+++ b/Engine/Source/Runtime/Engine/Private/VisualLogger/VisualLoggerAutomationTests.cpp
@@ -11,7 +11,7 @@ namespace
 	UWorld* GetSimpleEngineAutomationTestWorld(const int32 TestFlags)
 	{
 		// Accessing the game world is only valid for game-only 
-		if (((TestFlags & EAutomationTestFlags::ATF_Editor) || (TestFlags & EAutomationTestFlags::ATF_Game)) == false)
+		if (((TestFlags & EAutomationTestFlags::EditorContext) || (TestFlags & EAutomationTestFlags::ClientContext)) == false)
 		{
 			return nullptr;
 		}
@@ -105,7 +105,7 @@ struct FTestDeviceContext
 	uint32 EngineDisableAILoggingFlag : 1;
 };
 
-IMPLEMENT_SIMPLE_AUTOMATION_TEST(FVisualLogTest, "System.Engine.VisualLogger.Logging simple text", EAutomationTestFlags::ATF_Game | EAutomationTestFlags::ATF_Editor)
+IMPLEMENT_SIMPLE_AUTOMATION_TEST(FVisualLogTest, "System.Engine.VisualLogger.Logging simple text", EAutomationTestFlags::ClientContext | EAutomationTestFlags::EditorContext | EAutomationTestFlags::EngineFilter)
 /**
 *
 *
@@ -166,7 +166,7 @@ bool FVisualLogTest::RunTest(const FString& Parameters)
 	return true;
 }
 
-IMPLEMENT_SIMPLE_AUTOMATION_TEST(FVisualLogSegmentsTest, "System.Engine.VisualLogger.Logging segment shape", EAutomationTestFlags::ATF_Game | EAutomationTestFlags::ATF_Editor)
+IMPLEMENT_SIMPLE_AUTOMATION_TEST(FVisualLogSegmentsTest, "System.Engine.VisualLogger.Logging segment shape", EAutomationTestFlags::ClientContext | EAutomationTestFlags::EditorContext | EAutomationTestFlags::EngineFilter)
 bool FVisualLogSegmentsTest::RunTest(const FString& Parameters)
 {
 	UWorld* World = GetSimpleEngineAutomationTestWorld(GetTestFlags());
@@ -226,7 +226,7 @@ DEFINE_VLOG_EVENT(EventTest2, All, "Second simple event for vlog tests")
 
 DEFINE_VLOG_EVENT(EventTest3, All, "Third simple event for vlog tests")
 
-IMPLEMENT_SIMPLE_AUTOMATION_TEST(FVisualLogEventsTest, "System.Engine.VisualLogger.Logging events", EAutomationTestFlags::ATF_Game | EAutomationTestFlags::ATF_Editor)
+IMPLEMENT_SIMPLE_AUTOMATION_TEST(FVisualLogEventsTest, "System.Engine.VisualLogger.Logging events", EAutomationTestFlags::ClientContext | EAutomationTestFlags::EditorContext | EAutomationTestFlags::EngineFilter)
 bool FVisualLogEventsTest::RunTest(const FString& Parameters)
 {
 	UWorld* World = GetSimpleEngineAutomationTestWorld(GetTestFlags());
diff --git a/Engine/Source/Runtime/GameplayAbilities/Private/GameplayEffectTests.cpp b/Engine/Source/Runtime/GameplayAbilities/Private/GameplayEffectTests.cpp
--- a/Engine/Source/Runtime/GameplayAbilities/Private/GameplayEffectTests.cpp
+++ b/Engine/Source/Runtime/GameplayAbilities/Private/GameplayEffectTests.cpp
@@ -283,7 +283,7 @@ public:
 		ADD_TEST(Test_PeriodicDamage);
 	}
 
-	virtual uint32 GetTestFlags() const override { return EAutomationTestFlags::ATF_Editor; }
+	virtual uint32 GetTestFlags() const override { return EAutomationTestFlags::EditorContext | EAutomationTestFlags::EngineFilter; }
 	virtual bool IsStressTest() const { return false; }
 	virtual uint32 GetRequiredDeviceNum() const override { return 1; }
 
diff --git a/Engine/Source/Runtime/GameplayTags/Private/Tests/GameplayTagTests.cpp b/Engine/Source/Runtime/GameplayTags/Private/Tests/GameplayTagTests.cpp
--- a/Engine/Source/Runtime/GameplayTags/Private/Tests/GameplayTagTests.cpp
+++ b/Engine/Source/Runtime/GameplayTags/Private/Tests/GameplayTagTests.cpp
@@ -2,7 +2,7 @@
 
 #include "GameplayTagsModulePrivatePCH.h"
 
-IMPLEMENT_SIMPLE_AUTOMATION_TEST(FGameplayTagTest, "System.GameplayTags.GameplayTag", EAutomationTestFlags::ATF_Editor)
+IMPLEMENT_SIMPLE_AUTOMATION_TEST(FGameplayTagTest, "System.GameplayTags.GameplayTag", EAutomationTestFlags::EditorContext | EAutomationTestFlags::EngineFilter)
 
 #if WITH_EDITOR
 static UDataTable* CreateGameplayDataTable()
diff --git a/Engine/Source/Runtime/Json/Private/Tests/JsonTests.cpp b/Engine/Source/Runtime/Json/Private/Tests/JsonTests.cpp
--- a/Engine/Source/Runtime/Json/Private/Tests/JsonTests.cpp
+++ b/Engine/Source/Runtime/Json/Private/Tests/JsonTests.cpp
@@ -7,7 +7,7 @@
  * FJsonAutomationTest
  * Simple unit test that runs Json's in-built test cases
  */
-IMPLEMENT_SIMPLE_AUTOMATION_TEST( FJsonAutomationTest, "System.Engine.FileSystem.JSON", EAutomationTestFlags::ATF_SmokeTest)
+IMPLEMENT_SIMPLE_AUTOMATION_TEST(FJsonAutomationTest, "System.Engine.FileSystem.JSON", EAutomationTestFlags::ApplicationContextMask | EAutomationTestFlags::SmokeFilter )
 
 typedef TJsonWriterFactory< TCHAR, TCondensedJsonPrintPolicy<TCHAR> > FCondensedJsonStringWriterFactory;
 typedef TJsonWriter< TCHAR, TCondensedJsonPrintPolicy<TCHAR> > FCondensedJsonStringWriter;
diff --git a/Engine/Source/Runtime/Launch/Launch.Build.cs b/Engine/Source/Runtime/Launch/Launch.Build.cs
--- a/Engine/Source/Runtime/Launch/Launch.Build.cs
+++ b/Engine/Source/Runtime/Launch/Launch.Build.cs
@@ -114,6 +114,13 @@ public class Launch : ModuleRules
 			PublicIncludePathModuleNames.Add("ProfilerService");
 			DynamicallyLoadedModuleNames.AddRange(new string[] { "TaskGraph", "RealtimeProfiler", "ProfilerService" });
 		}
+		
+		// The engine can use AutomationController in any connfiguration besides shipping.  This module is loaded
+		// dynamically in LaunchEngineLoop.cpp in non-shipping configurations
+		if (UEBuildConfiguration.bCompileAgainstEngine && Target.Configuration != UnrealTargetConfiguration.Shipping)
+		{
+			DynamicallyLoadedModuleNames.AddRange(new string[] { "AutomationController" });
+		}
 
 		if (UEBuildConfiguration.bBuildEditor == true)
 		{
@@ -131,7 +138,6 @@ public class Launch : ModuleRules
 			// ExtraModules that are loaded when WITH_EDITOR=1 is true
 			DynamicallyLoadedModuleNames.AddRange(
 				new string[] {
-					"AutomationController",
 					"AutomationWindow",
 					"ProfilerClient",
 					"Toolbox",
diff --git a/Engine/Source/Runtime/Launch/Private/LaunchEngineLoop.cpp b/Engine/Source/Runtime/Launch/Private/LaunchEngineLoop.cpp
--- a/Engine/Source/Runtime/Launch/Private/LaunchEngineLoop.cpp
+++ b/Engine/Source/Runtime/Launch/Private/LaunchEngineLoop.cpp
@@ -19,7 +19,6 @@
 
 #if WITH_EDITOR
 	#include "EditorStyle.h"
-	#include "AutomationController.h"
 	#include "ProfilerClient.h"
 	#include "RemoteConfigIni.h"
 	#include "EditorCommandLineUtils.h"
@@ -32,6 +31,7 @@
 #endif
 
 #if WITH_ENGINE
+	#include "AutomationController.h"
 	#include "Database.h"
 	#include "DerivedDataCacheInterface.h"
 	#include "RenderCore.h"
@@ -2124,9 +2124,13 @@ int32 FEngineLoop::Init()
 	FModuleManager::Get().LoadModule("AutomationWorker");
 #endif
 
-#if WITH_EDITOR
+	// Automation tests can be invoked locally in non-editor builds configuration (e.g. performance profiling in Test configuration)
+#if WITH_ENGINE && !UE_BUILD_SHIPPING
 	FModuleManager::Get().LoadModule("AutomationController");
 	FModuleManager::GetModuleChecked<IAutomationControllerModule>("AutomationController").Init();
+#endif
+
+#if WITH_EDITOR
 	if (GIsEditor)
 	{
 		FModuleManager::Get().LoadModule(TEXT("ProfilerClient"));
diff --git a/Engine/Source/Runtime/Networking/Private/Tests/IPv4AddressTest.cpp b/Engine/Source/Runtime/Networking/Private/Tests/IPv4AddressTest.cpp
--- a/Engine/Source/Runtime/Networking/Private/Tests/IPv4AddressTest.cpp
+++ b/Engine/Source/Runtime/Networking/Private/Tests/IPv4AddressTest.cpp
@@ -4,7 +4,7 @@
 #include "AutomationTest.h"
 
 
-IMPLEMENT_SIMPLE_AUTOMATION_TEST(FIPv4AddressTest, "System.Engine.Networking.IPv4.IPv4Address", EAutomationTestFlags::ATF_SmokeTest)
+IMPLEMENT_SIMPLE_AUTOMATION_TEST(FIPv4AddressTest, "System.Engine.Networking.IPv4.IPv4Address", EAutomationTestFlags::ApplicationContextMask | EAutomationTestFlags::SmokeFilter )
 
 
 bool FIPv4AddressTest::RunTest( const FString& Parameters )
diff --git a/Engine/Source/Runtime/Online/OnlineSubsystemUtils/Private/Tests/OnlineSubsystemUtilitiesTests.cpp b/Engine/Source/Runtime/Online/OnlineSubsystemUtils/Private/Tests/OnlineSubsystemUtilitiesTests.cpp
--- a/Engine/Source/Runtime/Online/OnlineSubsystemUtils/Private/Tests/OnlineSubsystemUtilitiesTests.cpp
+++ b/Engine/Source/Runtime/Online/OnlineSubsystemUtils/Private/Tests/OnlineSubsystemUtilitiesTests.cpp
@@ -132,7 +132,7 @@ bool FTestIdentityInterfaceLoginOut::Update()
 * login and logout from the service.
 * Use -OSSIDTESTUSER='TestUser' and -OSSIDTESTPSSWD='YouPassword'
 */
-IMPLEMENT_COMPLEX_AUTOMATION_TEST( FIdenitiyInterfacetFunctionalLoginLogoutTest, "System.OSS", EAutomationTestFlags::ATF_Game )
+IMPLEMENT_COMPLEX_AUTOMATION_TEST( FIdenitiyInterfacetFunctionalLoginLogoutTest, "System.OSS", EAutomationTestFlags::ClientContext | EAutomationTestFlags::StressFilter)
 
 void FIdenitiyInterfacetFunctionalLoginLogoutTest::GetTests( TArray<FString>& OutBeautifiedNames, TArray<FString>& OutTestCommands ) const
 {
@@ -163,7 +163,7 @@ bool FIdenitiyInterfacetFunctionalLoginLogoutTest::RunTest( const FString& Param
 }
 
 
-IMPLEMENT_COMPLEX_AUTOMATION_TEST( FTestIdentityInterfaceLogoutWhileNotLoggedIn, "System.OSS", EAutomationTestFlags::ATF_Game )
+IMPLEMENT_COMPLEX_AUTOMATION_TEST(FTestIdentityInterfaceLogoutWhileNotLoggedIn, "System.OSS", EAutomationTestFlags::ClientContext | EAutomationTestFlags::StressFilter)
 
 void FTestIdentityInterfaceLogoutWhileNotLoggedIn::GetTests( TArray<FString>& OutBeautifiedNames, TArray<FString>& OutTestCommands ) const
 {
diff --git a/Engine/Source/Runtime/Serialization/Private/Tests/StructSerializerTest.cpp b/Engine/Source/Runtime/Serialization/Private/Tests/StructSerializerTest.cpp
--- a/Engine/Source/Runtime/Serialization/Private/Tests/StructSerializerTest.cpp
+++ b/Engine/Source/Runtime/Serialization/Private/Tests/StructSerializerTest.cpp
@@ -73,7 +73,7 @@ namespace StructSerializerTest
 /* Tests
  *****************************************************************************/
 
-IMPLEMENT_SIMPLE_AUTOMATION_TEST(FJsonStructSerializerTest, "System.Core.Serialization.JsonStructSerializer", EAutomationTestFlags::ATF_Editor)
+IMPLEMENT_SIMPLE_AUTOMATION_TEST(FJsonStructSerializerTest, "System.Core.Serialization.JsonStructSerializer", EAutomationTestFlags::EditorContext | EAutomationTestFlags::EngineFilter)
 
 
 bool FJsonStructSerializerTest::RunTest( const FString& Parameters )
diff --git a/Engine/Source/Runtime/Slate/Private/Framework/Text/Tests/RichTextMarkupProcessingTest.cpp b/Engine/Source/Runtime/Slate/Private/Framework/Text/Tests/RichTextMarkupProcessingTest.cpp
--- a/Engine/Source/Runtime/Slate/Private/Framework/Text/Tests/RichTextMarkupProcessingTest.cpp
+++ b/Engine/Source/Runtime/Slate/Private/Framework/Text/Tests/RichTextMarkupProcessingTest.cpp
@@ -13,7 +13,7 @@ PRAGMA_DISABLE_OPTIMIZATION
 
 #define LOCTEXT_NAMESPACE "Slate.Tests.RickTextMarkupProcessingTest"
 
-IMPLEMENT_SIMPLE_AUTOMATION_TEST(FRTFProcessingTest, "System.Slate.Misc.RickTextMarkupProcessing", EAutomationTestFlags::ATF_SmokeTest)
+IMPLEMENT_SIMPLE_AUTOMATION_TEST(FRTFProcessingTest, "System.Slate.Misc.RickTextMarkupProcessing", EAutomationTestFlags::ApplicationContextMask | EAutomationTestFlags::SmokeFilter)
 
 namespace
 {
diff --git a/Engine/Source/Runtime/VectorVM/Private/VectorVM.cpp b/Engine/Source/Runtime/VectorVM/Private/VectorVM.cpp
--- a/Engine/Source/Runtime/VectorVM/Private/VectorVM.cpp
+++ b/Engine/Source/Runtime/VectorVM/Private/VectorVM.cpp
@@ -1610,7 +1610,7 @@ FString VectorVM::GetOperandLocationName(EVectorVMOperandLocation Location)
 	Automation test for the VM.
 ------------------------------------------------------------------------------*/
 
-IMPLEMENT_SIMPLE_AUTOMATION_TEST(FVectorVMTest, "System.Core.Math.Vector VM", EAutomationTestFlags::ATF_SmokeTest)
+IMPLEMENT_SIMPLE_AUTOMATION_TEST(FVectorVMTest, "System.Core.Math.Vector VM", EAutomationTestFlags::ApplicationContextMask | EAutomationTestFlags::SmokeFilter)
 
 bool FVectorVMTest::RunTest(const FString& Parameters)
 {
