Fixed broken "showdebug collision" display for overlapping actors

Merging 2672648

//depot/UE4-UT/...

to //depot/UE4/...

diff --git a/Engine/Source/Runtime/Engine/Private/Actor.cpp b/Engine/Source/Runtime/Engine/Private/Actor.cpp
--- a/Engine/Source/Runtime/Engine/Private/Actor.cpp
+++ b/Engine/Source/Runtime/Engine/Private/Actor.cpp
@@ -2039,10 +2039,11 @@ void AActor::DisplayDebug(UCanvas* Canvas, const FDebugDisplayInfo& DebugDisplay
 			YPos += YL;
 		}
 
-		T = FString(TEXT("Touching "));
+		T = FString(TEXT("Overlapping "));
 
 		TArray<AActor*> TouchingActors;
 		GetOverlappingActors(TouchingActors);
+		bool bFoundAnyOverlaps = false;
 		for (int32 iTouching = 0; iTouching<TouchingActors.Num(); ++iTouching)
 		{
 			AActor* const TestActor = TouchingActors[iTouching];
@@ -2052,10 +2053,11 @@ void AActor::DisplayDebug(UCanvas* Canvas, const FDebugDisplayInfo& DebugDisplay
 			{
 				AActor* A = TestActor;
 				T = T + A->GetName() + " ";
+				bFoundAnyOverlaps = true;
 			}
 		}
 
-		if ( FCString::Strcmp(*T, TEXT("Overlapping ")))
+		if (!bFoundAnyOverlaps)
 		{
 			T = TEXT("Overlapping nothing");
 		}
