Added ENGINE_API.

diff --git a/Engine/Source/Runtime/Engine/Classes/Materials/MaterialExpressionObjectOrientation.h b/Engine/Source/Runtime/Engine/Classes/Materials/MaterialExpressionObjectOrientation.h
--- a/Engine/Source/Runtime/Engine/Classes/Materials/MaterialExpressionObjectOrientation.h
+++ b/Engine/Source/Runtime/Engine/Classes/Materials/MaterialExpressionObjectOrientation.h
@@ -6,7 +6,7 @@
 #include "MaterialExpressionObjectOrientation.generated.h"
 
 UCLASS(collapsecategories, hidecategories=Object)
-class UMaterialExpressionObjectOrientation : public UMaterialExpression
+class ENGINE_API UMaterialExpressionObjectOrientation : public UMaterialExpression
 {
 	GENERATED_UCLASS_BODY()
 
