#ue4 - Don't start ticking Actors/Components until BeginPlay() has been called when in a game world. This was an issue in some situations in network games.

--------

diff --git a/Engine/Source/Runtime/Engine/Classes/Components/ActorComponent.h b/Engine/Source/Runtime/Engine/Classes/Components/ActorComponent.h
--- a/Engine/Source/Runtime/Engine/Classes/Components/ActorComponent.h
+++ b/Engine/Source/Runtime/Engine/Classes/Components/ActorComponent.h
@@ -158,6 +158,9 @@ private:
 	/** Indicates the the destruction process has begun for this component to avoid recursion */
 	uint32 bIsBeingDestroyed:1;
 
+	/** Whether we've tried to register tick functions. Reset when they are unregistered. */
+	uint32 bTickFunctionsRegistered:1;
+
 #if WITH_EDITOR
 	/** During undo/redo it isn't safe to cache owner */
 	uint32 bCanUseCachedOwner:1;
diff --git a/Engine/Source/Runtime/Engine/Classes/GameFramework/Actor.h b/Engine/Source/Runtime/Engine/Classes/GameFramework/Actor.h
--- a/Engine/Source/Runtime/Engine/Classes/GameFramework/Actor.h
+++ b/Engine/Source/Runtime/Engine/Classes/GameFramework/Actor.h
@@ -181,6 +181,9 @@ private:
 	/** Whether FinishSpawning has been called for this Actor.  If it has not, the Actor is in a mal-formed state */
 	uint32 bHasFinishedSpawning:1;
 
+	/** Whether we've tried to register tick functions. Reset when they are unregistered. */
+	uint32 bTickFunctionsRegistered : 1;
+
 	/**
 	 * Enables any collision on this actor.
 	 * @see SetActorEnableCollision(), GetActorEnableCollision()
diff --git a/Engine/Source/Runtime/Engine/Private/Actor.cpp b/Engine/Source/Runtime/Engine/Private/Actor.cpp
--- a/Engine/Source/Runtime/Engine/Private/Actor.cpp
+++ b/Engine/Source/Runtime/Engine/Private/Actor.cpp
@@ -247,7 +247,6 @@ void AActor::ResetOwnedComponents()
 void AActor::PostInitProperties()
 {
 	Super::PostInitProperties();
-	RegisterAllActorTickFunctions(true,false); // component will be handled when they are registered
 	RemoteRole = (bReplicates ? ROLE_SimulatedProxy : ROLE_None);
 
 	// Map natively-constructed component instances to any serialized properties that might reference them.
@@ -641,11 +640,16 @@ void AActor::RegisterAllActorTickFunctions(bool bRegister, bool bDoComponents)
 {
 	if(!IsTemplate())
 	{
-		FActorThreadContext& ThreadContext = FActorThreadContext::Get();
-		check(ThreadContext.TestRegisterTickFunctions == nullptr);
-		RegisterActorTickFunctions(bRegister);
-		checkf(ThreadContext.TestRegisterTickFunctions == this, TEXT("Failed to route Actor RegisterTickFunctions (%s)"), *GetFullName());
-		ThreadContext.TestRegisterTickFunctions = nullptr;
+		// Prevent repeated redundant attempts
+		if (bTickFunctionsRegistered != bRegister)
+		{
+			FActorThreadContext& ThreadContext = FActorThreadContext::Get();
+			check(ThreadContext.TestRegisterTickFunctions == nullptr);
+			RegisterActorTickFunctions(bRegister);
+			bTickFunctionsRegistered = bRegister;
+			checkf(ThreadContext.TestRegisterTickFunctions == this, TEXT("Failed to route Actor RegisterTickFunctions (%s)"), *GetFullName());
+			ThreadContext.TestRegisterTickFunctions = nullptr;
+		}
 
 		if (bDoComponents)
 		{
@@ -2783,6 +2787,7 @@ void AActor::BeginPlay()
 {
 	ensure(!bActorHasBegunPlay);
 	SetLifeSpan( InitialLifeSpan );
+	RegisterAllActorTickFunctions(true, false); // Components are done below.
 
 	TInlineComponentArray<UActorComponent*> Components;
 	GetComponents(Components);
@@ -2792,7 +2797,11 @@ void AActor::BeginPlay()
 		// bHasBegunPlay will be true for the component if the component was renamed and moved to a new outer during initialization
 		if (Component->IsRegistered() && !Component->HasBegunPlay())
 		{
-			Component->BeginPlay();
+			Component->RegisterAllComponentTickFunctions(true);
+			if (Component->bWantsBeginPlay)
+			{
+				Component->BeginPlay();
+			}
 		}
 		else
 		{
@@ -3622,6 +3631,16 @@ bool AActor::IncrementalRegisterComponents(int32 NumComponentsToRegister)
 		// 0 - means register all components
 		NumComponentsToRegister = MAX_int32;
 	}
+
+	UWorld* const World = GetWorld();
+	check(World);
+
+	// If we are not a game world, then register tick functions now. If we are a game world we wait until right before BeginPlay(),
+	// so as to not actually tick until BeginPlay() executes (which could otherwise happen in network games).
+	if (!World->IsGameWorld())
+	{
+		RegisterAllActorTickFunctions(true, false); // components will be handled when they are registered
+	}
 	
 	// Register RootComponent first so all other components can reliable use it (ie call GetLocation) when they register
 	if( RootComponent != NULL && !RootComponent->IsRegistered() )
@@ -3636,8 +3655,7 @@ bool AActor::IncrementalRegisterComponents(int32 NumComponentsToRegister)
 		//This should prevent unwanted components hanging around when undoing a copy/paste or duplication action.
 		RootComponent->Modify(false);
 
-		check(GetWorld());
-		RootComponent->RegisterComponentWithWorld(GetWorld());
+		RootComponent->RegisterComponentWithWorld(World);
 	}
 
 	int32 NumTotalRegisteredComponents = 0;
@@ -3676,8 +3694,7 @@ bool AActor::IncrementalRegisterComponents(int32 NumComponentsToRegister)
 			//This should prevent unwanted components hanging around when undoing a copy/paste or duplication action.
 			Component->Modify(false);
 
-			check(GetWorld());
-			Component->RegisterComponentWithWorld(GetWorld());
+			Component->RegisterComponentWithWorld(World);
 			NumRegisteredComponentsThisRun++;
 		}
 
diff --git a/Engine/Source/Runtime/Engine/Private/Components/ActorComponent.cpp b/Engine/Source/Runtime/Engine/Private/Components/ActorComponent.cpp
--- a/Engine/Source/Runtime/Engine/Private/Components/ActorComponent.cpp
+++ b/Engine/Source/Runtime/Engine/Private/Components/ActorComponent.cpp
@@ -636,6 +636,8 @@ void UActorComponent::BeginPlay()
 {
 	check(bRegistered);
 	check(!bHasBegunPlay);
+	checkSlow(bWantsBeginPlay);
+	checkSlow(bTickFunctionsRegistered); // If this fails, someone called BeginPlay() without first calling RegisterAllComponentTickFunctions().
 
 	ReceiveBeginPlay();
 
@@ -755,9 +757,14 @@ void UActorComponent::RegisterAllComponentTickFunctions(bool bRegister)
 	// Components don't have tick functions until they are registered with the world
 	if (bRegistered)
 	{
-		RegisterComponentTickFunctions(bRegister);
-		checkf(GTestRegisterComponentTickFunctions == this, TEXT("Failed to route component RegisterTickFunctions (%s)"), *GetFullName());
-		GTestRegisterComponentTickFunctions = NULL;
+		// Prevent repeated redundant attempts
+		if (bTickFunctionsRegistered != bRegister)
+		{
+			RegisterComponentTickFunctions(bRegister);
+			bTickFunctionsRegistered = bRegister;
+			checkf(GTestRegisterComponentTickFunctions == this, TEXT("Failed to route component RegisterTickFunctions (%s)"), *GetFullName());
+			GTestRegisterComponentTickFunctions = NULL;
+		}
 	}
 }
 
@@ -822,7 +829,12 @@ void UActorComponent::RegisterComponentWithWorld(UWorld* InWorld)
 	World = InWorld;
 
 	ExecuteRegisterEvents();
-	RegisterAllComponentTickFunctions(true);
+
+	// If not in a game world register ticks now, otherwise defer until BeginPlay
+	if (!InWorld->IsGameWorld())
+	{
+		RegisterAllComponentTickFunctions(true);
+	}
 
 	if (MyOwner == nullptr || MyOwner->IsActorInitialized())
 	{
@@ -832,9 +844,10 @@ void UActorComponent::RegisterComponentWithWorld(UWorld* InWorld)
 		}
 	}
 
-	if (MyOwner && MyOwner->HasActorBegunPlay())
+	if (MyOwner && MyOwner->HasActorBegunPlay() && !bHasBegunPlay)
 	{
-		if (!bHasBegunPlay && bWantsBeginPlay)
+		RegisterAllComponentTickFunctions(true);
+		if (bWantsBeginPlay)
 		{
 			BeginPlay();
 		}
