#merge, replays: Don't delete actors when we stop demos, let that naturally happen when the DemoNetDriver gets destroyed. This is so the player controller sticks around until the next map is loaded.

diff --git a/Engine/Source/Runtime/Engine/Classes/Engine/DemoNetDriver.h b/Engine/Source/Runtime/Engine/Classes/Engine/DemoNetDriver.h
--- a/Engine/Source/Runtime/Engine/Classes/Engine/DemoNetDriver.h
+++ b/Engine/Source/Runtime/Engine/Classes/Engine/DemoNetDriver.h
@@ -127,6 +127,9 @@ public:
 
 	void GotoTimeInSeconds(const float TimeInSeconds, const FOnGotoTimeDelegate& InOnGotoTimeDelegate = FOnGotoTimeDelegate());
 
+	bool IsRecording();
+	bool IsPlaying();
+
 public:
 
 	// FExec interface
diff --git a/Engine/Source/Runtime/Engine/Private/DemoNetDriver.cpp b/Engine/Source/Runtime/Engine/Private/DemoNetDriver.cpp
--- a/Engine/Source/Runtime/Engine/Private/DemoNetDriver.cpp
+++ b/Engine/Source/Runtime/Engine/Private/DemoNetDriver.cpp
@@ -314,8 +314,8 @@ void UDemoNetDriver::FinishDestroy()
 {
 	if ( !HasAnyFlags( RF_ClassDefaultObject ) )
 	{
-		// Make sure we stop any recording that might be going on
-		if ( ClientConnections.Num() > 0 )
+		// Make sure we stop any recording/playing that might be going on
+		if ( IsRecording() || IsPlaying() )
 		{
 			StopDemo();
 		}
@@ -703,11 +703,21 @@ bool UDemoNetDriver::InitListen( FNetworkNotify* InNotify, FURL& ListenURL, bool
 	return true;
 }
 
-void UDemoNetDriver::TickDispatch(float DeltaSeconds)
+bool UDemoNetDriver::IsRecording()
+{
+	return ClientConnections.Num() > 0 && ClientConnections[0] != nullptr && ClientConnections[0]->State != USOCK_Closed;
+}
+
+bool UDemoNetDriver::IsPlaying()
+{
+	return ServerConnection != nullptr && ServerConnection->State != USOCK_Closed;
+}
+
+void UDemoNetDriver::TickDispatch( float DeltaSeconds )
 {
 	Super::TickDispatch( DeltaSeconds );
 
-	if ( ClientConnections.Num() == 0 && ServerConnection == nullptr )
+	if ( !IsRecording() && !IsPlaying() )
 	{
 		// Nothing to do
 		return;
@@ -730,7 +740,7 @@ void UDemoNetDriver::TickDispatch(float DeltaSeconds)
 		return;
 	}
 
-	if ( ClientConnections.Num() > 0 )
+	if ( IsRecording() )
 	{
 		const double StartTime = FPlatformTime::Seconds();
 
@@ -766,8 +776,10 @@ void UDemoNetDriver::TickDispatch(float DeltaSeconds)
 			RecordCountSinceFlush	= 0;
 		}
 	}
-	else if ( ServerConnection != NULL )
+	else
 	{
+		check( IsPlaying() );
+
 		// Wait until all levels are streamed in
 		for ( int32 i = 0; i < World->StreamingLevels.Num(); ++i )
 		{
@@ -827,7 +839,7 @@ void UDemoNetDriver::TickDispatch(float DeltaSeconds)
 
 void UDemoNetDriver::ProcessRemoteFunction( class AActor* Actor, class UFunction* Function, void* Parameters, struct FOutParmRec* OutParms, struct FFrame* Stack, class UObject* SubObject )
 {
-	if (ClientConnections.Num() > 0 && ClientConnections[0] != NULL)
+	if ( IsRecording() )
 	{
 		if ((Function->FunctionFlags & FUNC_NetMulticast))
 		{
@@ -863,7 +875,7 @@ bool UDemoNetDriver::Exec( UWorld* InWorld, const TCHAR* Cmd, FOutputDevice& Ar
 
 void UDemoNetDriver::StopDemo()
 {
-	if ( !ServerConnection && ClientConnections.Num() == 0 )
+	if ( !IsRecording() && !IsPlaying() )
 	{
 		UE_LOG( LogDemo, Warning, TEXT( "StopDemo: No demo is playing" ) );
 		return;
@@ -930,8 +942,8 @@ void UDemoNetDriver::StopDemo()
 		// let GC cleanup the object
 		if ( ClientConnections.Num() > 0 && ClientConnections[0] != NULL )
 		{
+			ServerConnection->State = USOCK_Closed;
 			ClientConnections[0]->Close();
-			ClientConnections[0]->CleanUp(); // make sure DemoRecSpectator gets destroyed immediately
 		}
 	}
 	else
@@ -941,15 +953,12 @@ void UDemoNetDriver::StopDemo()
 
 		ServerConnection->State = USOCK_Closed;
 		ServerConnection->Close();
-		ServerConnection->CleanUp(); // make sure DemoRecSpectator gets destroyed immediately
-		ServerConnection = NULL;
 	}
 
 	ReplayStreamer->StopStreaming();
 	ClearReplayTasks();
 
-	check( ClientConnections.Num() == 0 );
-	check( ServerConnection == NULL );
+	check( !IsRecording() && !IsPlaying() );
 }
 
 /*-----------------------------------------------------------------------------
@@ -1214,7 +1223,7 @@ void UDemoNetDriver::RequestEventData(const FString& EventID, FOnRequestEventDat
 
 void UDemoNetDriver::TickDemoRecord( float DeltaSeconds )
 {
-	if ( ClientConnections.Num() == 0 )
+	if ( !IsRecording() )
 	{
 		return;
 	}
@@ -1576,9 +1585,8 @@ void UDemoNetDriver::AddUserToReplay( const FString& UserString )
 
 void UDemoNetDriver::TickDemoPlayback( float DeltaSeconds )
 {
-	if ( ServerConnection == NULL || ServerConnection->State == USOCK_Closed )
+	if ( !IsPlaying() )
 	{
-		StopDemo();
 		return;
 	}
 
