#merge, replays: Fix a crash when the server stops recording a replay.

diff --git a/Engine/Source/Runtime/Engine/Private/DemoNetDriver.cpp b/Engine/Source/Runtime/Engine/Private/DemoNetDriver.cpp
--- a/Engine/Source/Runtime/Engine/Private/DemoNetDriver.cpp
+++ b/Engine/Source/Runtime/Engine/Private/DemoNetDriver.cpp
@@ -942,7 +942,6 @@ void UDemoNetDriver::StopDemo()
 		// let GC cleanup the object
 		if ( ClientConnections.Num() > 0 && ClientConnections[0] != NULL )
 		{
-			ServerConnection->State = USOCK_Closed;
 			ClientConnections[0]->Close();
 		}
 	}
