#merge, replays: Add ability to force replay loading failures at various stages

diff --git a/Engine/Source/Runtime/Engine/Private/DemoNetDriver.cpp b/Engine/Source/Runtime/Engine/Private/DemoNetDriver.cpp
--- a/Engine/Source/Runtime/Engine/Private/DemoNetDriver.cpp
+++ b/Engine/Source/Runtime/Engine/Private/DemoNetDriver.cpp
@@ -35,6 +35,10 @@ static TAutoConsoleVariable<int32> CVarDemoFastForwardDestroyTearOffActors( TEXT
 static TAutoConsoleVariable<int32> CVarDemoFastForwardSkipRepNotifies( TEXT( "demo.FastForwardSkipRepNotifies" ), 1, TEXT( "If true, the driver will optimize fast-forwarding by deferring calls to RepNotify functions until the fast-forward is complete. " ) );
 static TAutoConsoleVariable<int32> CVarDemoQueueCheckpointChannels( TEXT( "demo.QueueCheckpointChannels" ), 1, TEXT( "If true, the driver will put all channels created during checkpoint loading into queuing mode, to amortize the cost of spawning new actors across multiple frames." ) );
 
+#if !(UE_BUILD_SHIPPING || UE_BUILD_TEST)
+static TAutoConsoleVariable<int32> CVarDemoForceFailure( TEXT( "demo.ForceFailure" ), 0, TEXT( "" ) );
+#endif
+
 static const int32 MAX_DEMO_READ_WRITE_BUFFER = 1024 * 2;
 
 #define DEMO_CHECKSUMS 0		// When setting this to 1, this will invalidate all demos, you will need to re-record and playback
@@ -581,7 +585,16 @@ bool UDemoNetDriver::InitConnectInternal( FString& Error )
 
 	WorldContext->PendingNetGame = NewPendingNetGame;
 
-	if ( !GEngine->LoadMap( *WorldContext, DemoURL, NewPendingNetGame, LoadMapError ) )
+	bool bSuccess = GEngine->LoadMap( *WorldContext, DemoURL, NewPendingNetGame, LoadMapError );
+
+#if !(UE_BUILD_SHIPPING || UE_BUILD_TEST)
+	if ( CVarDemoForceFailure.GetValueOnGameThread() == 2 )
+	{
+		bSuccess = false;
+	}
+#endif
+
+	if ( !bSuccess )
 	{
 		StopDemo();
 
@@ -591,10 +604,12 @@ bool UDemoNetDriver::InitConnectInternal( FString& Error )
 		if ( WorldContext->World() == NULL )
 		{
 			GEngine->DestroyNamedNetDriver( WorldContext->PendingNetGame, NetDriverName );
-			WorldContext->PendingNetGame = NULL;
-			GEngine->BrowseToDefaultMap( *WorldContext );
 		}
 
+		WorldContext->PendingNetGame = NULL;
+
+		GEngine->BrowseToDefaultMap( *WorldContext );
+
 		Error = LoadMapError;
 		UE_LOG( LogDemo, Error, TEXT( "UDemoNetDriver::InitConnect: LoadMap failed: failed: %s" ), *Error );
 		GameInstance->HandleDemoPlaybackFailure( EDemoPlayFailure::Generic, FString( TEXT( "LoadMap failed" ) ) );
@@ -1727,6 +1742,13 @@ void UDemoNetDriver::ReplayStreamingReady( bool bSuccess, bool bRecord )
 {
 	bWasStartStreamingSuccessful = bSuccess;
 
+#if !(UE_BUILD_SHIPPING || UE_BUILD_TEST)
+	if ( CVarDemoForceFailure.GetValueOnGameThread() == 1 )
+	{
+		bSuccess = false;
+	}
+#endif
+
 	if ( !bSuccess )
 	{
 		UE_LOG( LogDemo, Warning, TEXT( "UDemoNetConnection::ReplayStreamingReady: Failed." ) );
