attachment bugfix

passing through the bone (socketname) name along with the parent parameter so physics object knows how to properly detach from skeletal mesh component that it was attached to.

diff --git a/Engine/Source/Runtime/Engine/Classes/Components/SceneComponent.h b/Engine/Source/Runtime/Engine/Classes/Components/SceneComponent.h
--- a/Engine/Source/Runtime/Engine/Classes/Components/SceneComponent.h
+++ b/Engine/Source/Runtime/Engine/Classes/Components/SceneComponent.h
@@ -722,7 +722,7 @@ protected:
 private:
 
 	void PropagateTransformUpdate(bool bTransformChanged, bool bSkipPhysicsMove = false, ETeleportType Teleport = ETeleportType::None);
-	void UpdateComponentToWorldWithParent(USceneComponent* Parent, bool bSkipPhysicsMove, const FQuat& RelativeRotationQuat, ETeleportType Teleport = ETeleportType::None);
+	void UpdateComponentToWorldWithParent(USceneComponent* Parent, FName SocketName, bool bSkipPhysicsMove, const FQuat& RelativeRotationQuat, ETeleportType Teleport = ETeleportType::None);
 
 
 public:
@@ -918,8 +918,8 @@ protected:
 
 	/** Calculate the new ComponentToWorld transform for this component.
 		Parent is optional and can be used for computing ComponentToWorld based on arbitrary USceneComponent.
-		If Parent is not passed in we use the component's AttachParent*/
-	virtual FTransform CalcNewComponentToWorld(const FTransform& NewRelativeTransform, const USceneComponent* Parent = NULL) const;
+		If Parent is not passed in or is NULL then we use the component's existing AttachParent and AttachSocket */
+	virtual FTransform CalcNewComponentToWorld(const FTransform& NewRelativeTransform, const USceneComponent* Parent = NULL, FName SocketName = NAME_None) const;
 
 	
 public:
diff --git a/Engine/Source/Runtime/Engine/Private/Components/SceneComponent.cpp b/Engine/Source/Runtime/Engine/Private/Components/SceneComponent.cpp
--- a/Engine/Source/Runtime/Engine/Private/Components/SceneComponent.cpp
+++ b/Engine/Source/Runtime/Engine/Private/Components/SceneComponent.cpp
@@ -297,12 +297,13 @@ void USceneComponent::PostEditChangeChainProperty(FPropertyChangedChainEvent& Pr
 
 #endif
 
-FTransform USceneComponent::CalcNewComponentToWorld(const FTransform& NewRelativeTransform, const USceneComponent* Parent) const
+FTransform USceneComponent::CalcNewComponentToWorld(const FTransform& NewRelativeTransform, const USceneComponent* Parent, FName SocketName) const
 {
+	SocketName = Parent ? SocketName : AttachSocketName;
 	Parent = Parent ? Parent : AttachParent;
 	if (Parent != NULL)
 	{
-		const FTransform ParentToWorld = Parent->GetSocketTransform(AttachSocketName);
+		const FTransform ParentToWorld = Parent->GetSocketTransform(SocketName);
 		FTransform NewCompToWorld = NewRelativeTransform * ParentToWorld;
 
 		if(bAbsoluteLocation)
@@ -332,7 +333,7 @@ void USceneComponent::OnUpdateTransform(bool bSkipPhysicsMove, ETeleportType Tel
 {
 }
 
-void USceneComponent::UpdateComponentToWorldWithParent(USceneComponent* Parent, bool bSkipPhysicsMove, const FQuat& RelativeRotationQuat, ETeleportType Teleport)
+void USceneComponent::UpdateComponentToWorldWithParent(USceneComponent* Parent,FName SocketName, bool bSkipPhysicsMove, const FQuat& RelativeRotationQuat, ETeleportType Teleport)
 {
 	// If our parent hasn't been updated before, we'll need walk up our parent attach hierarchy
 	if (Parent && !Parent->bWorldToComponentUpdated)
@@ -350,7 +351,7 @@ void USceneComponent::UpdateComponentToWorldWithParent(USceneComponent* Parent,
 
 	// Calculate the new ComponentToWorld transform
 	const FTransform RelativeTransform(RelativeRotationQuat, RelativeLocation, RelativeScale3D);
-	FTransform NewTransform = CalcNewComponentToWorld(RelativeTransform, Parent);
+	FTransform NewTransform = CalcNewComponentToWorld(RelativeTransform, Parent, SocketName);
 
 #if !(UE_BUILD_SHIPPING || UE_BUILD_TEST)
 	ensure(NewTransform.IsValid());
@@ -411,7 +412,7 @@ void USceneComponent::OnRegister()
 
 void USceneComponent::UpdateComponentToWorld(bool bSkipPhysicsMove, ETeleportType Teleport)
 {
-	UpdateComponentToWorldWithParent(AttachParent, bSkipPhysicsMove, RelativeRotationCache.RotatorToQuat(RelativeRotation), Teleport);
+	UpdateComponentToWorldWithParent(AttachParent,AttachSocketName, bSkipPhysicsMove, RelativeRotationCache.RotatorToQuat(RelativeRotation), Teleport);
 }
 
 
@@ -1248,7 +1249,7 @@ bool USceneComponent::AttachTo(class USceneComponent* Parent, FName InSocketName
 				//Since the object is physically simulated it can't be the case that it's a child of object A and being attached to object B (at runtime)
 				if (bMaintainWorldPosition == false)	//User tried to attach but physically based so detach. However, if they provided relative coordinates we should still get the correct position
 				{
-					UpdateComponentToWorldWithParent(Parent, false, RelativeRotationCache.RotatorToQuat(RelativeRotation));
+					UpdateComponentToWorldWithParent(Parent, InSocketName, false, RelativeRotationCache.RotatorToQuat(RelativeRotation));
 					RelativeLocation = ComponentToWorld.GetLocation(); // or GetComponentLocation(), but worried about custom location...
 					RelativeRotation = GetComponentRotation();
 					RelativeScale3D = GetComponentScale();
@@ -1909,7 +1910,7 @@ bool USceneComponent::InternalSetWorldLocationAndRotation(FVector NewLocation, c
 	{
 		RelativeLocation = NewLocation;
 		RelativeRotation = RelativeRotationCache.QuatToRotator(NewRotationQuat); // Normalizes quat, if this is a new rotation. Then we'll use it below.
-		UpdateComponentToWorldWithParent(AttachParent, bNoPhysics, RelativeRotationCache.GetCachedQuat(), Teleport);
+		UpdateComponentToWorldWithParent(AttachParent,AttachSocketName, bNoPhysics, RelativeRotationCache.GetCachedQuat(), Teleport);
 
 		PostUpdateNavigationData();
 		return true;
