#replays: Don't call playback error handler when recording fails

diff --git a/Engine/Source/Runtime/Engine/Private/DemoNetDriver.cpp b/Engine/Source/Runtime/Engine/Private/DemoNetDriver.cpp
--- a/Engine/Source/Runtime/Engine/Private/DemoNetDriver.cpp
+++ b/Engine/Source/Runtime/Engine/Private/DemoNetDriver.cpp
@@ -741,8 +741,12 @@ void UDemoNetDriver::TickDispatch( float DeltaSeconds )
 	if ( ReplayStreamer->GetLastError() != ENetworkReplayError::None )
 	{
 		UE_LOG( LogDemo, Error, TEXT( "UDemoNetDriver::TickFlush: ReplayStreamer ERROR: %s" ), ENetworkReplayError::ToString( ReplayStreamer->GetLastError() ) );
+		const bool bIsPlaying = IsPlaying();
 		StopDemo();
-		World->GetGameInstance()->HandleDemoPlaybackFailure( EDemoPlayFailure::Generic, FString( EDemoPlayFailure::ToString( EDemoPlayFailure::Generic ) ) );
+		if ( bIsPlaying )
+		{
+			World->GetGameInstance()->HandleDemoPlaybackFailure( EDemoPlayFailure::Generic, FString( EDemoPlayFailure::ToString( EDemoPlayFailure::Generic ) ) );
+		}
 		return;
 	}
 
@@ -1752,8 +1756,13 @@ void UDemoNetDriver::ReplayStreamingReady( bool bSuccess, bool bRecord )
 	if ( !bSuccess )
 	{
 		UE_LOG( LogDemo, Warning, TEXT( "UDemoNetConnection::ReplayStreamingReady: Failed." ) );
+
 		StopDemo();
-		GetWorld()->GetGameInstance()->HandleDemoPlaybackFailure( EDemoPlayFailure::DemoNotFound, FString( EDemoPlayFailure::ToString( EDemoPlayFailure::DemoNotFound ) ) );
+
+		if ( !bRecord )
+		{
+			GetWorld()->GetGameInstance()->HandleDemoPlaybackFailure( EDemoPlayFailure::DemoNotFound, FString( EDemoPlayFailure::ToString( EDemoPlayFailure::DemoNotFound ) ) );
+		}
 		return;
 	}
 
