Merging

//depot/UE4-Orion/Engine/Source/Runtime/Engine/Private/Animation/Skeleton.cpp

to //depot/UE4-Orion/Engine/Source/Runtime/Engine/Private/Animation/Skeleton.cpp

diff --git a/Engine/Source/Runtime/Engine/Private/Animation/Skeleton.cpp b/Engine/Source/Runtime/Engine/Private/Animation/Skeleton.cpp
--- a/Engine/Source/Runtime/Engine/Private/Animation/Skeleton.cpp
+++ b/Engine/Source/Runtime/Engine/Private/Animation/Skeleton.cpp
@@ -269,6 +269,12 @@ bool USkeleton::IsCompatibleMesh(const USkeletalMesh* InSkelMesh) const
 		if( SkeletonBoneIndex != INDEX_NONE )
 		{
 			++NumOfBoneMatches;
+
+			// follow the parent chain to verify the chain is same
+			if(!DoesParentChainMatch(SkeletonBoneIndex, InSkelMesh))
+			{
+				return false;
+			}
 		}
 		else
 		{
