UE-19859 Remove debug/show rendering features on shipping
I did the first part, lock some showflags to be 0 or 1 in SHIPPING
If this causes problems for some applications (e.g. using showflag as scalability which they shouldn't) the user needs to update the code.
or we remove the lock from that showflag.
For editor or non shipping features it's fine to use the SetSHOWFLAGNAME() calls, they become NOPS in SHIPPING.

This is a performance optimization (when  used in inner loop) and prepares for removing more

diff --git a/Engine/Plugins/2D/Paper2D/Source/Paper2DEditor/Private/FlipbookEditor/FlipbookEditorViewportClient.cpp b/Engine/Plugins/2D/Paper2D/Source/Paper2DEditor/Private/FlipbookEditor/FlipbookEditorViewportClient.cpp
--- a/Engine/Plugins/2D/Paper2D/Source/Paper2DEditor/Private/FlipbookEditor/FlipbookEditorViewportClient.cpp
+++ b/Engine/Plugins/2D/Paper2D/Source/Paper2DEditor/Private/FlipbookEditor/FlipbookEditorViewportClient.cpp
@@ -35,7 +35,7 @@ FFlipbookEditorViewportClient::FFlipbookEditorViewportClient(const TAttribute<UP
 	DrawHelper.bDrawGrid = GetDefault<UFlipbookEditorSettings>()->bShowGridByDefault;
 
 	EngineShowFlags.DisableAdvancedFeatures();
-	EngineShowFlags.CompositeEditorPrimitives = true;
+	EngineShowFlags.SetCompositeEditorPrimitives(true);
 }
 
 void FFlipbookEditorViewportClient::DrawCanvas(FViewport& Viewport, FSceneView& View, FCanvas& Canvas)
diff --git a/Engine/Plugins/2D/Paper2D/Source/Paper2DEditor/Private/SpriteEditor/SpriteEditorViewportClient.cpp b/Engine/Plugins/2D/Paper2D/Source/Paper2DEditor/Private/SpriteEditor/SpriteEditorViewportClient.cpp
--- a/Engine/Plugins/2D/Paper2D/Source/Paper2DEditor/Private/SpriteEditor/SpriteEditorViewportClient.cpp
+++ b/Engine/Plugins/2D/Paper2D/Source/Paper2DEditor/Private/SpriteEditor/SpriteEditorViewportClient.cpp
@@ -91,7 +91,7 @@ FSpriteEditorViewportClient::FSpriteEditorViewportClient(TWeakPtr<FSpriteEditor>
 	DrawHelper.bDrawGrid = GetDefault<USpriteEditorSettings>()->bShowGridByDefault;
 
 	EngineShowFlags.DisableAdvancedFeatures();
-	EngineShowFlags.CompositeEditorPrimitives = true;
+	EngineShowFlags.SetCompositeEditorPrimitives(true);
 
 	// Create a render component for the sprite being edited
 	{
@@ -671,7 +671,7 @@ void FSpriteEditorViewportClient::ToggleShowSpriteNames()
 
 void FSpriteEditorViewportClient::ToggleShowMeshEdges()
 {
-	EngineShowFlags.MeshEdges = !EngineShowFlags.MeshEdges;
+	EngineShowFlags.SetMeshEdges(!EngineShowFlags.MeshEdges);
 	Invalidate(); 
 }
 
diff --git a/Engine/Plugins/2D/Paper2D/Source/Paper2DEditor/Private/TileMapEditing/TileMapEditorViewportClient.cpp b/Engine/Plugins/2D/Paper2D/Source/Paper2DEditor/Private/TileMapEditing/TileMapEditorViewportClient.cpp
--- a/Engine/Plugins/2D/Paper2D/Source/Paper2DEditor/Private/TileMapEditing/TileMapEditorViewportClient.cpp
+++ b/Engine/Plugins/2D/Paper2D/Source/Paper2DEditor/Private/TileMapEditing/TileMapEditorViewportClient.cpp
@@ -44,7 +44,7 @@ FTileMapEditorViewportClient::FTileMapEditorViewportClient(TWeakPtr<FTileMapEdit
 	bShowPivot = false;
 
 	EngineShowFlags.DisableAdvancedFeatures();
-	EngineShowFlags.CompositeEditorPrimitives = true;
+	EngineShowFlags.SetCompositeEditorPrimitives(true);
 
 	// Create a render component for the tile map being edited
 	{
@@ -259,7 +259,7 @@ bool FTileMapEditorViewportClient::IsShowLayerGridChecked() const
 
 void FTileMapEditorViewportClient::ToggleShowMeshEdges()
 {
-	EngineShowFlags.MeshEdges = !EngineShowFlags.MeshEdges;
+	EngineShowFlags.SetMeshEdges(!EngineShowFlags.MeshEdges);
 	Invalidate();
 }
 
diff --git a/Engine/Plugins/2D/Paper2D/Source/Paper2DEditor/Private/TileSetEditor/SingleTileEditorViewportClient.cpp b/Engine/Plugins/2D/Paper2D/Source/Paper2DEditor/Private/TileSetEditor/SingleTileEditorViewportClient.cpp
--- a/Engine/Plugins/2D/Paper2D/Source/Paper2DEditor/Private/TileSetEditor/SingleTileEditorViewportClient.cpp
+++ b/Engine/Plugins/2D/Paper2D/Source/Paper2DEditor/Private/TileSetEditor/SingleTileEditorViewportClient.cpp
@@ -39,7 +39,7 @@ FSingleTileEditorViewportClient::FSingleTileEditorViewportClient(UPaperTileSet*
 	((FAssetEditorModeManager*)ModeTools)->SetPreviewScene(PreviewScene);
 
 	EngineShowFlags.DisableAdvancedFeatures();
-	EngineShowFlags.CompositeEditorPrimitives = true;
+	EngineShowFlags.SetCompositeEditorPrimitives(true);
 
 	// Create a render component for the tile preview
 	PreviewTileSpriteComponent = NewObject<UPaperSpriteComponent>();
diff --git a/Engine/Plugins/Experimental/SimpleHMD/Source/SimpleHMD/Private/SimpleHMD.cpp b/Engine/Plugins/Experimental/SimpleHMD/Source/SimpleHMD/Private/SimpleHMD.cpp
--- a/Engine/Plugins/Experimental/SimpleHMD/Source/SimpleHMD/Private/SimpleHMD.cpp
+++ b/Engine/Plugins/Experimental/SimpleHMD/Source/SimpleHMD/Private/SimpleHMD.cpp
@@ -334,7 +334,7 @@ void FSimpleHMD::SetupViewFamily(FSceneViewFamily& InViewFamily)
 {
 	InViewFamily.EngineShowFlags.MotionBlur = 0;
 	InViewFamily.EngineShowFlags.HMDDistortion = true;
-	InViewFamily.EngineShowFlags.ScreenPercentage = 1.0f;
+	InViewFamily.EngineShowFlags.SetScreenPercentage(true);
 	InViewFamily.EngineShowFlags.StereoRendering = IsStereoEnabled();
 }
 
diff --git a/Engine/Source/Editor/Cascade/Private/Cascade.cpp b/Engine/Source/Editor/Cascade/Private/Cascade.cpp
--- a/Engine/Source/Editor/Cascade/Private/Cascade.cpp
+++ b/Engine/Source/Editor/Cascade/Private/Cascade.cpp
@@ -3820,7 +3820,7 @@ void FCascade::OnToggleGrid()
 		EditorOptions->SaveConfig();
 		DrawHelper.bDrawGrid = bShowGrid;
 
-		PreviewViewport->GetViewportClient()->EngineShowFlags.Grid = bShowGrid;
+		PreviewViewport->GetViewportClient()->EngineShowFlags.SetGrid(bShowGrid);
 		PreviewViewport->RefreshViewport();
 	}
 }
diff --git a/Engine/Source/Editor/Cascade/Private/CascadePreviewViewportClient.cpp b/Engine/Source/Editor/Cascade/Private/CascadePreviewViewportClient.cpp
--- a/Engine/Source/Editor/Cascade/Private/CascadePreviewViewportClient.cpp
+++ b/Engine/Source/Editor/Cascade/Private/CascadePreviewViewportClient.cpp
@@ -62,7 +62,7 @@ FCascadeEdPreviewViewportClient::FCascadeEdPreviewViewportClient(TWeakPtr<FCasca
 	SetViewMode(VMI_Lit);
 
 	EngineShowFlags.DisableAdvancedFeatures();
-	EngineShowFlags.CompositeEditorPrimitives = true;
+	EngineShowFlags.SetCompositeEditorPrimitives(true);
 
 	OverrideNearClipPlane(1.0f);
 
@@ -105,7 +105,7 @@ FCascadeEdPreviewViewportClient::FCascadeEdPreviewViewportClient(TWeakPtr<FCasca
 	
 	if (DrawHelper.bDrawGrid)
 	{
-		EngineShowFlags.Grid = 1;
+		EngineShowFlags.SetGrid(true);
 	}
 
 	if (EditorOptions->FloorMesh == TEXT(""))
@@ -255,11 +255,11 @@ void FCascadeEdPreviewViewportClient::Draw(FViewport* Viewport, FCanvas* Canvas)
 
 	if (GetDrawElement(Bounds))
 	{
-		EngineShowFlags.Bounds = 1;
+		EngineShowFlags.SetBounds(true);
 		EngineShowFlags.Game = 1;
 	}
 
-	EngineShowFlags.VectorFields = GetDrawElement(VectorFields);
+	EngineShowFlags.SetVectorFields(GetDrawElement(VectorFields));
 
 	CascadePreviewScene.AddComponent(LineBatcher,FTransform::Identity);
 
diff --git a/Engine/Source/Editor/DetailCustomizations/Private/AnimMontageSegmentDetails.cpp b/Engine/Source/Editor/DetailCustomizations/Private/AnimMontageSegmentDetails.cpp
--- a/Engine/Source/Editor/DetailCustomizations/Private/AnimMontageSegmentDetails.cpp
+++ b/Engine/Source/Editor/DetailCustomizations/Private/AnimMontageSegmentDetails.cpp
@@ -14,7 +14,7 @@ FAnimationSegmentViewportClient::FAnimationSegmentViewportClient(FPreviewScene&
 	SetViewMode(VMI_Lit);
 
 	// Always composite editor objects after post processing in the editor
-	EngineShowFlags.CompositeEditorPrimitives = true;
+	EngineShowFlags.SetCompositeEditorPrimitives(true);
 	EngineShowFlags.DisableAdvancedFeatures();
 	
 	UpdateLighting();
diff --git a/Engine/Source/Editor/DetailCustomizations/Private/SceneCaptureDetails.cpp b/Engine/Source/Editor/DetailCustomizations/Private/SceneCaptureDetails.cpp
--- a/Engine/Source/Editor/DetailCustomizations/Private/SceneCaptureDetails.cpp
+++ b/Engine/Source/Editor/DetailCustomizations/Private/SceneCaptureDetails.cpp
@@ -66,7 +66,6 @@ void FSceneCaptureDetails::CustomizeDetails( IDetailLayoutBuilder& DetailLayout
 	// Show flags that should be exposed for Scene Captures
 	TArray<FEngineShowFlags::EShowFlag> ShowFlagsToAllowForCaptures;
 
-	// General
 	ShowFlagsToAllowForCaptures.Add(FEngineShowFlags::EShowFlag::SF_AtmosphericFog);
 	ShowFlagsToAllowForCaptures.Add(FEngineShowFlags::EShowFlag::SF_BSP);
 	ShowFlagsToAllowForCaptures.Add(FEngineShowFlags::EShowFlag::SF_Decals);
diff --git a/Engine/Source/Editor/Kismet/Private/SCSEditorViewportClient.cpp b/Engine/Source/Editor/Kismet/Private/SCSEditorViewportClient.cpp
--- a/Engine/Source/Editor/Kismet/Private/SCSEditorViewportClient.cpp
+++ b/Engine/Source/Editor/Kismet/Private/SCSEditorViewportClient.cpp
@@ -121,8 +121,7 @@ FSCSEditorViewportClient::FSCSEditorViewportClient(TWeakPtr<FBlueprintEditor>& I
 	Widget->SetSnapEnabled(true);
 
 	// Selectively set particular show flags that we need
-	EngineShowFlags.TextRender = 1;
-	EngineShowFlags.SelectionOutline = GetDefault<ULevelEditorViewportSettings>()->bUseSelectionOutline;
+	EngineShowFlags.SetSelectionOutline(GetDefault<ULevelEditorViewportSettings>()->bUseSelectionOutline);
 
 	// Set if the grid will be drawn
 	DrawHelper.bDrawGrid = GetDefault<UEditorPerProjectUserSettings>()->bSCSEditorShowGrid;
diff --git a/Engine/Source/Editor/LevelEditor/Private/HighResScreenshotUI.cpp b/Engine/Source/Editor/LevelEditor/Private/HighResScreenshotUI.cpp
--- a/Engine/Source/Editor/LevelEditor/Private/HighResScreenshotUI.cpp
+++ b/Engine/Source/Editor/LevelEditor/Private/HighResScreenshotUI.cpp
@@ -254,7 +254,7 @@ void SHighResScreenshotDialog::ResetViewport()
 			if (ConfigViewport->GetClient() &&
 				ConfigViewport->GetClient()->GetEngineShowFlags())
 			{
-				ConfigViewport->GetClient()->GetEngineShowFlags()->HighResScreenshotMask = bMaskVisualizationWasEnabled;
+				ConfigViewport->GetClient()->GetEngineShowFlags()->SetHighResScreenshotMask(bMaskVisualizationWasEnabled);
 			}
 		}
 	}
@@ -330,7 +330,7 @@ TWeakPtr<class SWindow> SHighResScreenshotDialog::OpenDialog(const TSharedPtr<FS
 		if (InViewport.IsValid())
 		{
 			bMaskVisualizationWasEnabled = InViewport->GetClient()->GetEngineShowFlags()->HighResScreenshotMask;
-			InViewport->GetClient()->GetEngineShowFlags()->HighResScreenshotMask = Config.bMaskEnabled;
+			InViewport->GetClient()->GetEngineShowFlags()->SetHighResScreenshotMask(Config.bMaskEnabled);
 		}
 	}
 
diff --git a/Engine/Source/Editor/LevelEditor/Private/HighresScreenshotUI.h b/Engine/Source/Editor/LevelEditor/Private/HighresScreenshotUI.h
--- a/Engine/Source/Editor/LevelEditor/Private/HighresScreenshotUI.h
+++ b/Engine/Source/Editor/LevelEditor/Private/HighresScreenshotUI.h
@@ -82,7 +82,7 @@ private:
 		auto ConfigViewport = Config.TargetViewport.Pin();
 		if (ConfigViewport.IsValid())
 		{
-			ConfigViewport->GetClient()->GetEngineShowFlags()->HighResScreenshotMask = Config.bMaskEnabled;
+			ConfigViewport->GetClient()->GetEngineShowFlags()->SetHighResScreenshotMask(Config.bMaskEnabled);
 			ConfigViewport->Invalidate();
 		}
 	}
diff --git a/Engine/Source/Editor/LevelEditor/Private/SLevelViewport.cpp b/Engine/Source/Editor/LevelEditor/Private/SLevelViewport.cpp
--- a/Engine/Source/Editor/LevelEditor/Private/SLevelViewport.cpp
+++ b/Engine/Source/Editor/LevelEditor/Private/SLevelViewport.cpp
@@ -334,10 +334,10 @@ void SLevelViewport::ConstructLevelEditorViewportClient( const FArguments& InArg
 	{
 		EditorShowFlags.MotionBlur = 0;
 		EditorShowFlags.Fog = 0;
-		EditorShowFlags.DepthOfField = 0;
+		EditorShowFlags.SetDepthOfField(false);
 		GameShowFlags.MotionBlur = 0;
 		GameShowFlags.Fog = 0;
-		GameShowFlags.DepthOfField = 0;
+		GameShowFlags.SetDepthOfField(false);
 	}
 
 	EditorShowFlags.SetSnap(1);
@@ -376,10 +376,10 @@ void SLevelViewport::ConstructLevelEditorViewportClient( const FArguments& InArg
 	LevelViewportClient->OverrideFarClipPlane( ViewportInstanceSettings.FarViewPlane );
 	
 	// Set the selection outline flag based on preferences
-	LevelViewportClient->EngineShowFlags.SelectionOutline = GetDefault<ULevelEditorViewportSettings>()->bUseSelectionOutline;
+	LevelViewportClient->EngineShowFlags.SetSelectionOutline(GetDefault<ULevelEditorViewportSettings>()->bUseSelectionOutline);
 	
 	// Always composite editor objects after post processing in the editor
-	LevelViewportClient->EngineShowFlags.CompositeEditorPrimitives = true;
+	LevelViewportClient->EngineShowFlags.SetCompositeEditorPrimitives(true);
 
 	LevelViewportClient->SetViewModes(ViewportInstanceSettings.PerspViewModeIndex, ViewportInstanceSettings.OrthoViewModeIndex );
 }
@@ -1928,7 +1928,7 @@ void SLevelViewport::OnUseDefaultShowFlags(bool bUseSavedDefaults)
 	// this trashes the current viewmode!
 	LevelViewportClient->EngineShowFlags = EditorShowFlags;
 	// Restore the state of SelectionOutline based on user settings
-	LevelViewportClient->EngineShowFlags.SelectionOutline = GetDefault<ULevelEditorViewportSettings>()->bUseSelectionOutline;
+	LevelViewportClient->EngineShowFlags.SetSelectionOutline(GetDefault<ULevelEditorViewportSettings>()->bUseSelectionOutline);
 	LevelViewportClient->LastEngineShowFlags = GameShowFlags;
 
 	// re-apply the cached viewmode, as it was trashed with FEngineShowFlags()
@@ -2436,9 +2436,9 @@ void SLevelViewport::OnActorSelectionChanged(const TArray<UObject*>& NewSelectio
 	// On the first actor selection after entering Game View, enable the selection show flag
 	if (IsVisible() && IsInGameView() && NewSelection.Num() != 0)
 	{
-		LevelViewportClient->EngineShowFlags.ModeWidgets = 1;
-		LevelViewportClient->EngineShowFlags.Selection = 1;
-		LevelViewportClient->EngineShowFlags.SelectionOutline = GetDefault<ULevelEditorViewportSettings>()->bUseSelectionOutline;
+		LevelViewportClient->EngineShowFlags.SetModeWidgets(true);
+		LevelViewportClient->EngineShowFlags.SetSelection(true);
+		LevelViewportClient->EngineShowFlags.SetSelectionOutline(GetDefault<ULevelEditorViewportSettings>()->bUseSelectionOutline);
 	}
 
 
@@ -3020,7 +3020,7 @@ void SLevelViewport::PreviewActors( const TArray< AActor* >& ActorsToPreview )
 				// Default to "game" show flags for camera previews
 				// Still draw selection highlight though
 				ActorPreviewLevelViewportClient->EngineShowFlags = FEngineShowFlags(ESFIM_Game);
-				ActorPreviewLevelViewportClient->EngineShowFlags.Selection = 1;
+				ActorPreviewLevelViewportClient->EngineShowFlags.SetSelection(true);
 				ActorPreviewLevelViewportClient->LastEngineShowFlags = FEngineShowFlags(ESFIM_Editor);
 				
 				// We don't use view modes for preview viewports
@@ -3129,8 +3129,8 @@ void SLevelViewport::UpdateActorPreviewViewports()
 		CurActorPreview.LevelViewportClient->SetRealtime( LevelViewportClient->IsRealtime() );
 		CurActorPreview.LevelViewportClient->bDrawBaseInfo = LevelViewportClient->bDrawBaseInfo;
 		CurActorPreview.LevelViewportClient->bDrawVertices = LevelViewportClient->bDrawVertices;
-		CurActorPreview.LevelViewportClient->EngineShowFlags.SelectionOutline = LevelViewportClient->EngineShowFlags.SelectionOutline;
-		CurActorPreview.LevelViewportClient->EngineShowFlags.CompositeEditorPrimitives = LevelViewportClient->EngineShowFlags.CompositeEditorPrimitives;
+		CurActorPreview.LevelViewportClient->EngineShowFlags.SetSelectionOutline(LevelViewportClient->EngineShowFlags.SelectionOutline);
+		CurActorPreview.LevelViewportClient->EngineShowFlags.SetCompositeEditorPrimitives(LevelViewportClient->EngineShowFlags.CompositeEditorPrimitives);
 	}
 }
 
diff --git a/Engine/Source/Editor/Matinee/Private/Matinee.cpp b/Engine/Source/Editor/Matinee/Private/Matinee.cpp
--- a/Engine/Source/Editor/Matinee/Private/Matinee.cpp
+++ b/Engine/Source/Editor/Matinee/Private/Matinee.cpp
@@ -944,7 +944,7 @@ void FMatinee::InitMatinee(const EToolkitMode::Type Mode, const TSharedPtr< clas
 			}
 
 			// Turn on 'show camera frustums' flag
-			LevelVC->EngineShowFlags.CameraFrustums = 1;
+			LevelVC->EngineShowFlags.SetCameraFrustums(true);
 		}
 	}
 
@@ -2067,7 +2067,7 @@ void FMatinee::OnClose()
 			}
 
 			// Turn off 'show camera frustums' flag.
-			LevelVC->EngineShowFlags.CameraFrustums = 0;
+			LevelVC->EngineShowFlags.SetCameraFrustums(false);
 		}
 	}
 
diff --git a/Engine/Source/Editor/MeshPaint/Private/MeshPaintEdMode.cpp b/Engine/Source/Editor/MeshPaint/Private/MeshPaintEdMode.cpp
--- a/Engine/Source/Editor/MeshPaint/Private/MeshPaintEdMode.cpp
+++ b/Engine/Source/Editor/MeshPaint/Private/MeshPaintEdMode.cpp
@@ -2953,7 +2953,7 @@ void FEdModeMeshPaint::SetViewportShowFlags( const bool bAllowColorViewModes, FE
 				{
 					// If we're transitioning to normal mode then restore the backup
 					// Clear the flags relevant to vertex color modes
-					Viewport.EngineShowFlags.VertexColors = 0;
+					Viewport.EngineShowFlags.SetVertexColors(false);
 						
 					// Restore the vertex color mode flags that were set when we last entered vertex color mode
 					ApplyViewMode(Viewport.GetViewMode(), Viewport.IsPerspective(), Viewport.EngineShowFlags);
@@ -2962,13 +2962,13 @@ void FEdModeMeshPaint::SetViewportShowFlags( const bool bAllowColorViewModes, FE
 			}
 			else
 			{
-				Viewport.EngineShowFlags.Materials = 1;
-				Viewport.EngineShowFlags.Lighting = 0;
-				Viewport.EngineShowFlags.BSPTriangles = 1;
-				Viewport.EngineShowFlags.VertexColors = 1;
-				Viewport.EngineShowFlags.PostProcessing = 0;
-				Viewport.EngineShowFlags.HMDDistortion = 0;
-					
+				Viewport.EngineShowFlags.SetMaterials(true);
+				Viewport.EngineShowFlags.SetLighting(false);
+				Viewport.EngineShowFlags.SetBSPTriangles(true);
+				Viewport.EngineShowFlags.SetVertexColors(true);
+				Viewport.EngineShowFlags.SetPostProcessing(false);
+				Viewport.EngineShowFlags.SetHMDDistortion(false);
+
 				switch( ColorViewMode )
 				{
 					case EMeshPaintColorViewMode::RGB:
diff --git a/Engine/Source/Editor/Persona/Private/AnimationEditorViewportClient.cpp b/Engine/Source/Editor/Persona/Private/AnimationEditorViewportClient.cpp
--- a/Engine/Source/Editor/Persona/Private/AnimationEditorViewportClient.cpp
+++ b/Engine/Source/Editor/Persona/Private/AnimationEditorViewportClient.cpp
@@ -126,7 +126,7 @@ FAnimationViewportClient::FAnimationViewportClient(FAnimationEditorPreviewScene&
 	ViewFOV = FMath::Clamp<float>(ConfigOption->ViewFOV, FOVMin, FOVMax);
 
 	EngineShowFlags.DisableAdvancedFeatures();
-	EngineShowFlags.CompositeEditorPrimitives = true;
+	EngineShowFlags.SetCompositeEditorPrimitives(true);
 
 	// set camera mode
 	bCameraFollow = false;
diff --git a/Engine/Source/Editor/Persona/Private/SAnimationEditorViewport.cpp b/Engine/Source/Editor/Persona/Private/SAnimationEditorViewport.cpp
--- a/Engine/Source/Editor/Persona/Private/SAnimationEditorViewport.cpp
+++ b/Engine/Source/Editor/Persona/Private/SAnimationEditorViewport.cpp
@@ -1388,7 +1388,7 @@ void SAnimationEditorViewportTabBody::UpdateShowFlagForMeshEdges()
 	bShowMeshEdgesViewMode = (CurrentViewMode == EAnimationEditorViewportMode::UnlitWithMeshEdges);
 #endif
 
-	LevelViewportClient->EngineShowFlags.MeshEdges = (bDrawBonesInfluence || bShowMeshEdgesViewMode) ? 1 : 0;
+	LevelViewportClient->EngineShowFlags.SetMeshEdges(bDrawBonesInfluence || bShowMeshEdgesViewMode);
 }
 
 bool SAnimationEditorViewportTabBody::IsLODModelSelected(ELODViewSelection LODSelectionType) const
diff --git a/Engine/Source/Editor/Persona/Private/SAnimationSequenceBrowser.cpp b/Engine/Source/Editor/Persona/Private/SAnimationSequenceBrowser.cpp
--- a/Engine/Source/Editor/Persona/Private/SAnimationSequenceBrowser.cpp
+++ b/Engine/Source/Editor/Persona/Private/SAnimationSequenceBrowser.cpp
@@ -1045,7 +1045,7 @@ FAnimationAssetViewportClient::FAnimationAssetViewportClient(FPreviewScene& InPr
 	SetViewMode(VMI_Lit);
 
 	// Always composite editor objects after post processing in the editor
-	EngineShowFlags.CompositeEditorPrimitives = true;
+	EngineShowFlags.SetCompositeEditorPrimitives(true);
 	EngineShowFlags.DisableAdvancedFeatures();
 
 	// Setup defaults for the common draw helper.
diff --git a/Engine/Source/Editor/PhAT/Private/PhATPreviewViewportClient.cpp b/Engine/Source/Editor/PhAT/Private/PhATPreviewViewportClient.cpp
--- a/Engine/Source/Editor/PhAT/Private/PhATPreviewViewportClient.cpp
+++ b/Engine/Source/Editor/PhAT/Private/PhATPreviewViewportClient.cpp
@@ -52,7 +52,7 @@ FPhATEdPreviewViewportClient::FPhATEdPreviewViewportClient(TWeakPtr<FPhAT> InPhA
 	check(PhATFont);
 
 	EngineShowFlags.DisableAdvancedFeatures();
-	EngineShowFlags.CompositeEditorPrimitives = true;
+	EngineShowFlags.SetCompositeEditorPrimitives(true);
 
 	// Get actors asset collision bounding box, and move actor so its not intersection the floor plane at Z = 0.
 	FBox CollBox = SharedData->PhysicsAsset->CalcAABB(SharedData->EditorSkelComp, SharedData->EditorSkelComp->ComponentToWorld);	
diff --git a/Engine/Source/Editor/StaticMeshEditor/Private/SStaticMeshEditorViewport.cpp b/Engine/Source/Editor/StaticMeshEditor/Private/SStaticMeshEditorViewport.cpp
--- a/Engine/Source/Editor/StaticMeshEditor/Private/SStaticMeshEditorViewport.cpp
+++ b/Engine/Source/Editor/StaticMeshEditor/Private/SStaticMeshEditorViewport.cpp
@@ -339,13 +339,13 @@ void SStaticMeshEditorViewport::SetViewModeVertexColor()
 {
 	if (!EditorViewportClient->EngineShowFlags.VertexColors)
 	{
-		EditorViewportClient->EngineShowFlags.VertexColors = true;
-		EditorViewportClient->EngineShowFlags.Lighting = false;
+		EditorViewportClient->EngineShowFlags.SetVertexColors(true);
+		EditorViewportClient->EngineShowFlags.SetLighting(false);
 	}
 	else
 	{
-		EditorViewportClient->EngineShowFlags.VertexColors = false;
-		EditorViewportClient->EngineShowFlags.Lighting = true;
+		EditorViewportClient->EngineShowFlags.SetVertexColors(false);
+		EditorViewportClient->EngineShowFlags.SetLighting(true);
 	}
 	if (FEngineAnalytics::IsAvailable())
 	{
diff --git a/Engine/Source/Editor/StaticMeshEditor/Private/StaticMeshAutomationTests.cpp b/Engine/Source/Editor/StaticMeshEditor/Private/StaticMeshAutomationTests.cpp
--- a/Engine/Source/Editor/StaticMeshEditor/Private/StaticMeshAutomationTests.cpp
+++ b/Engine/Source/Editor/StaticMeshEditor/Private/StaticMeshAutomationTests.cpp
@@ -93,27 +93,20 @@ namespace EditorViewButtonHelper
 			case EStaticMeshFlag::Vert:
 				if (!AutomationParameters.ViewportClient->EngineShowFlags.VertexColors)
 				{
-					AutomationParameters.ViewportClient->EngineShowFlags.VertexColors = true;
-					AutomationParameters.ViewportClient->EngineShowFlags.Lighting = false;
+					AutomationParameters.ViewportClient->EngineShowFlags.SetVertexColors(true);
+					AutomationParameters.ViewportClient->EngineShowFlags.SetLighting(false);
 				}
 				else
 				{
-					AutomationParameters.ViewportClient->EngineShowFlags.VertexColors = false;
-					AutomationParameters.ViewportClient->EngineShowFlags.Lighting = true;
+					AutomationParameters.ViewportClient->EngineShowFlags.SetVertexColors(false);
+					AutomationParameters.ViewportClient->EngineShowFlags.SetLighting(true);
 				}
 				break;
 			case EStaticMeshFlag::Grid:
 				AutomationParameters.ViewportClient->SetShowGrid();
 				break;
 			case EStaticMeshFlag::Bounds:
-				if (AutomationParameters.ViewportClient->EngineShowFlags.Bounds)
-				{
-					AutomationParameters.ViewportClient->EngineShowFlags.Bounds = false;
-				}
-				else
-				{
-					AutomationParameters.ViewportClient->EngineShowFlags.Bounds = true;
-				}
+				AutomationParameters.ViewportClient->EngineShowFlags.SetBounds(!AutomationParameters.ViewportClient->EngineShowFlags.Bounds);
 				break;
 			case EStaticMeshFlag::Collision:
 				AutomationParameters.ViewportClient->SetShowCollision();
diff --git a/Engine/Source/Editor/StaticMeshEditor/Private/StaticMeshEditorViewportClient.cpp b/Engine/Source/Editor/StaticMeshEditor/Private/StaticMeshEditorViewportClient.cpp
--- a/Engine/Source/Editor/StaticMeshEditor/Private/StaticMeshEditorViewportClient.cpp
+++ b/Engine/Source/Editor/StaticMeshEditor/Private/StaticMeshEditorViewportClient.cpp
@@ -71,7 +71,7 @@ FStaticMeshEditorViewportClient::FStaticMeshEditorViewportClient(TWeakPtr<IStati
 
 	EngineShowFlags.DisableAdvancedFeatures();
 	EngineShowFlags.SetSnap(0);
-	EngineShowFlags.CompositeEditorPrimitives = true;
+	EngineShowFlags.SetCompositeEditorPrimitives(true);
 	OverrideNearClipPlane(1.0f);
 	bUsingOrbitCamera = true;
 
diff --git a/Engine/Source/Editor/UnrealEd/Private/EditorViewportClient.cpp b/Engine/Source/Editor/UnrealEd/Private/EditorViewportClient.cpp
--- a/Engine/Source/Editor/UnrealEd/Private/EditorViewportClient.cpp
+++ b/Engine/Source/Editor/UnrealEd/Private/EditorViewportClient.cpp
@@ -1554,7 +1554,7 @@ void FEditorViewportClient::UpdateLightingShowFlags( FEngineShowFlags& InOutShow
 				{
 					// We have lights in the scene now so go ahead and turn lighting back on
 					// designer can see what they're interacting with!
-					InOutShowFlags.Lighting = true;
+					InOutShowFlags.SetLighting(true);
 				}
 
 				// No longer forcing lighting to be off
@@ -1567,7 +1567,7 @@ void FEditorViewportClient::UpdateLightingShowFlags( FEngineShowFlags& InOutShow
 				{
 					// No lights in the scene, so make sure that lighting is turned off so the level
 					// designer can see what they're interacting with!
-					InOutShowFlags.Lighting = false;
+					InOutShowFlags.SetLighting(false);
 				}
 			}
 		}
@@ -4258,12 +4258,12 @@ bool FEditorViewportClient::IsSetShowGridChecked() const
 
 void FEditorViewportClient::SetShowBounds(bool bShow)
 {
-	EngineShowFlags.Bounds = bShow;
+	EngineShowFlags.SetBounds(bShow);
 }
 
 void FEditorViewportClient::ToggleShowBounds()
 {
-	EngineShowFlags.Bounds = 1 - EngineShowFlags.Bounds;
+	EngineShowFlags.SetBounds(!EngineShowFlags.Bounds);
 	if (FEngineAnalytics::IsAvailable())
 	{
 		FEngineAnalytics::GetProvider().RecordEvent(TEXT("Editor.Usage.StaticMesh.Toolbar"), TEXT("Bounds"), FString::Printf(TEXT("%d"), EngineShowFlags.Bounds));
@@ -4278,7 +4278,7 @@ bool FEditorViewportClient::IsSetShowBoundsChecked() const
 
 void FEditorViewportClient::SetShowCollision()
 {
-	EngineShowFlags.Collision = !EngineShowFlags.Collision;
+	EngineShowFlags.SetCollision(!EngineShowFlags.Collision);
 	Invalidate();
 }
 
@@ -4787,17 +4787,17 @@ void FEditorViewportClient::SetGameView(bool bGameViewEnable)
 	}
 
 	// maintain this state
-	EngineShowFlags.CompositeEditorPrimitives = bCompositeEditorPrimitives;
-	LastEngineShowFlags.CompositeEditorPrimitives = bCompositeEditorPrimitives;
+	EngineShowFlags.SetCompositeEditorPrimitives(bCompositeEditorPrimitives);
+	LastEngineShowFlags.SetCompositeEditorPrimitives(bCompositeEditorPrimitives);
 
 	//reset game engine show flags that may have been turned on by making a selection in game view
 	if(bGameViewEnable)
 	{
-		EngineShowFlags.ModeWidgets = 0;
-		EngineShowFlags.Selection = 0;
+		EngineShowFlags.SetModeWidgets(false);
+		EngineShowFlags.SetSelection(false);
 	}
 
-	EngineShowFlags.SelectionOutline = bGameViewEnable ? false : GetDefault<ULevelEditorViewportSettings>()->bUseSelectionOutline;
+	EngineShowFlags.SetSelectionOutline(bGameViewEnable ? false : GetDefault<ULevelEditorViewportSettings>()->bUseSelectionOutline);
 
 	ApplyViewMode(GetViewMode(), IsPerspective(), EngineShowFlags);
 
diff --git a/Engine/Source/Editor/UnrealEd/Private/LevelEditorViewport.cpp b/Engine/Source/Editor/UnrealEd/Private/LevelEditorViewport.cpp
--- a/Engine/Source/Editor/UnrealEd/Private/LevelEditorViewport.cpp
+++ b/Engine/Source/Editor/UnrealEd/Private/LevelEditorViewport.cpp
@@ -953,9 +953,9 @@ bool FLevelEditorViewportClient::DropObjectsOnWidget(FSceneView* View, FViewport
 	const bool bOldModeWidgets1 = EngineShowFlags.ModeWidgets;
 	const bool bOldModeWidgets2 = View->Family->EngineShowFlags.ModeWidgets;
 
-	EngineShowFlags.ModeWidgets = 0;
+	EngineShowFlags.SetModeWidgets(false);
 	FSceneViewFamily* SceneViewFamily = const_cast< FSceneViewFamily* >( View->Family );
-	SceneViewFamily->EngineShowFlags.ModeWidgets = 0;
+	SceneViewFamily->EngineShowFlags.SetModeWidgets(false);
 
 	// Invalidate the hit proxy map so it will be rendered out again when GetHitProxy is called
 	Viewport->InvalidateHitProxy();
@@ -976,8 +976,8 @@ bool FLevelEditorViewportClient::DropObjectsOnWidget(FSceneView* View, FViewport
 	bResult = DropObjectsAtCoordinates(CursorPos.X, CursorPos.Y, DroppedObjects, TemporaryActors, bOnlyDropOnTarget, bCreateDropPreview);
 
 	// Restore the original flags
-	EngineShowFlags.ModeWidgets = bOldModeWidgets1;
-	SceneViewFamily->EngineShowFlags.ModeWidgets = bOldModeWidgets2;
+	EngineShowFlags.SetModeWidgets(bOldModeWidgets1);
+	SceneViewFamily->EngineShowFlags.SetModeWidgets(bOldModeWidgets2);
 
 	return bResult;
 }
@@ -1878,9 +1878,9 @@ void FLevelEditorViewportClient::ProcessClick(FSceneView& View, HHitProxy* HitPr
 			const bool bOldModeWidgets1 = EngineShowFlags.ModeWidgets;
 			const bool bOldModeWidgets2 = View.Family->EngineShowFlags.ModeWidgets;
 
-			EngineShowFlags.ModeWidgets = 0;
+			EngineShowFlags.SetModeWidgets(false);
 			FSceneViewFamily* SceneViewFamily = const_cast<FSceneViewFamily*>(View.Family);
-			SceneViewFamily->EngineShowFlags.ModeWidgets = 0;
+			SceneViewFamily->EngineShowFlags.SetModeWidgets(false);
 			bool bWasWidgetDragging = Widget->IsDragging();
 			Widget->SetDragging(false);
 
@@ -1897,8 +1897,8 @@ void FLevelEditorViewportClient::ProcessClick(FSceneView& View, HHitProxy* HitPr
 			}
 
 			// Undo the evil
-			EngineShowFlags.ModeWidgets = bOldModeWidgets1;
-			SceneViewFamily->EngineShowFlags.ModeWidgets = bOldModeWidgets2;
+			EngineShowFlags.SetModeWidgets(bOldModeWidgets1);
+			SceneViewFamily->EngineShowFlags.SetModeWidgets(bOldModeWidgets2);
 
 			Widget->SetDragging(bWasWidgetDragging);
 
@@ -2779,7 +2779,7 @@ void FLevelEditorViewportClient::HandleViewportSettingChanged(FName PropertyName
 {
 	if (PropertyName == GET_MEMBER_NAME_CHECKED(ULevelEditorViewportSettings, bUseSelectionOutline))
 	{
-		EngineShowFlags.SelectionOutline = GetDefault<ULevelEditorViewportSettings>()->bUseSelectionOutline;
+		EngineShowFlags.SetSelectionOutline(GetDefault<ULevelEditorViewportSettings>()->bUseSelectionOutline);
 	}
 }
 
@@ -3992,14 +3992,14 @@ void FLevelEditorViewportClient::SetupViewForRendering( FSceneViewFamily& ViewFa
 		// unless the view port is Matinee controlled
 		ViewFamily.EngineShowFlags.CameraInterpolation = 0;
 		// keep the image sharp - ScreenPercentage is an optimization and should not affect the editor
-		ViewFamily.EngineShowFlags.ScreenPercentage = 0;
+		ViewFamily.EngineShowFlags.SetScreenPercentage(false);
 	}
 
 	TSharedPtr<FDragDropOperation> DragOperation = FSlateApplication::Get().GetDragDroppingContent();
 	if (!(DragOperation.IsValid() && DragOperation->IsOfType<FBrushBuilderDragDropOp>()))
 	{
 		// Hide the builder brush when not in geometry mode
-		ViewFamily.EngineShowFlags.BuilderBrush = 0;
+		ViewFamily.EngineShowFlags.SetBuilderBrush(false);
 	}
 
 	// Update the listener.
diff --git a/Engine/Source/Editor/UnrealEd/Private/SSkeletonWidget.cpp b/Engine/Source/Editor/UnrealEd/Private/SSkeletonWidget.cpp
--- a/Engine/Source/Editor/UnrealEd/Private/SSkeletonWidget.cpp
+++ b/Engine/Source/Editor/UnrealEd/Private/SSkeletonWidget.cpp
@@ -1269,7 +1269,7 @@ public:
 		SetViewMode(VMI_Lit);
 
 		// Always composite editor objects after post processing in the editor
-		EngineShowFlags.CompositeEditorPrimitives = true;
+		EngineShowFlags.SetCompositeEditorPrimitives(true);
 		EngineShowFlags.DisableAdvancedFeatures();
 
 		UpdateLighting();
diff --git a/Engine/Source/Editor/UnrealEd/Private/ThumbnailRendering/LevelThumbnailRenderer.cpp b/Engine/Source/Editor/UnrealEd/Private/ThumbnailRendering/LevelThumbnailRenderer.cpp
--- a/Engine/Source/Editor/UnrealEd/Private/ThumbnailRendering/LevelThumbnailRenderer.cpp
+++ b/Engine/Source/Editor/UnrealEd/Private/ThumbnailRendering/LevelThumbnailRenderer.cpp
@@ -20,8 +20,7 @@ void ULevelThumbnailRenderer::Draw(UObject* Object, int32 X, int32 Y, uint32 Wid
 
 		ViewFamily.EngineShowFlags.DisableAdvancedFeatures();
 		ViewFamily.EngineShowFlags.MotionBlur = 0;
-		ViewFamily.EngineShowFlags.Lighting = 1;
-		ViewFamily.EngineShowFlags.PostProcessing = 0;
+		ViewFamily.EngineShowFlags.SetPostProcessing(false);
 
 		GetView(Level, &ViewFamily, X, Y, Width, Height);
 
diff --git a/Engine/Source/Editor/UnrealEd/Public/EditorViewportClient.h b/Engine/Source/Editor/UnrealEd/Public/EditorViewportClient.h
--- a/Engine/Source/Editor/UnrealEd/Public/EditorViewportClient.h
+++ b/Engine/Source/Editor/UnrealEd/Public/EditorViewportClient.h
@@ -965,13 +965,13 @@ public:
 	 */
 	void SetShowAspectRatioBarDisplay(bool bEnable)
 	{
-		EngineShowFlags.CameraAspectRatioBars = bEnable ? 1 : 0;
+		EngineShowFlags.SetCameraAspectRatioBars(bEnable);
 		Invalidate(false,false);
 	}
 
 	void SetShowSafeFrameBoxDisplay(bool bEnable)
 	{
-		EngineShowFlags.CameraSafeFrames = bEnable ? 1 : 0;
+		EngineShowFlags.SetCameraSafeFrames(bEnable);
 		Invalidate(false,false);
 	}
 
diff --git a/Engine/Source/Runtime/Core/Private/HAL/ConsoleManager.cpp b/Engine/Source/Runtime/Core/Private/HAL/ConsoleManager.cpp
--- a/Engine/Source/Runtime/Core/Private/HAL/ConsoleManager.cpp
+++ b/Engine/Source/Runtime/Core/Private/HAL/ConsoleManager.cpp
@@ -2037,16 +2037,6 @@ static TAutoConsoleVariable<int32> CVarDetailMode(
 	TEXT(" 2: high, show all objects (default)"),
 	ECVF_Scalability | ECVF_RenderThreadSafe);
 
-static TAutoConsoleVariable<int32> CVarRefractionQuality(
-	TEXT("r.RefractionQuality"),
-	2,
-	TEXT("Defines the distorion/refraction quality which allows to adjust for quality or performance.\n")
-	TEXT("<=0: off (fastest)\n")
-	TEXT("  1: low quality (not yet implemented)\n")
-	TEXT("  2: normal quality (default)\n")
-	TEXT("  3: high quality (e.g. color fringe, not yet implemented)"),
-	ECVF_Scalability | ECVF_RenderThreadSafe);
-
 static TAutoConsoleVariable<int32> CVarDBuffer(
 	TEXT("r.DBuffer"),
 	0,
diff --git a/Engine/Source/Runtime/Engine/Private/Components/SceneCaptureComponent.cpp b/Engine/Source/Runtime/Engine/Private/Components/SceneCaptureComponent.cpp
--- a/Engine/Source/Runtime/Engine/Private/Components/SceneCaptureComponent.cpp
+++ b/Engine/Source/Runtime/Engine/Private/Components/SceneCaptureComponent.cpp
@@ -155,9 +155,9 @@ USceneCaptureComponent::USceneCaptureComponent(const FObjectInitializer& ObjectI
 	MaxViewDistanceOverride = -1;
 
 	// Disable features that are not desired when capturing the scene
-	ShowFlags.MotionBlur = 0; // motion blur doesn't work correctly with scene captures.
-	ShowFlags.SeparateTranslucency = 0;
-	ShowFlags.HMDDistortion = 0;
+	ShowFlags.SetMotionBlur(0); // motion blur doesn't work correctly with scene captures.
+	ShowFlags.SetSeparateTranslucency(0);
+	ShowFlags.SetHMDDistortion(0);
 
     CaptureStereoPass = EStereoscopicPass::eSSP_FULL;
 }
diff --git a/Engine/Source/Runtime/Engine/Private/GameViewportClient.cpp b/Engine/Source/Runtime/Engine/Private/GameViewportClient.cpp
--- a/Engine/Source/Runtime/Engine/Private/GameViewportClient.cpp
+++ b/Engine/Source/Runtime/Engine/Private/GameViewportClient.cpp
@@ -175,7 +175,7 @@ UGameViewportClient::~UGameViewportClient()
 {
 	if (EngineShowFlags.Collision)
 	{
-		EngineShowFlags.Collision = false;
+		EngineShowFlags.SetCollision(false);
 		ToggleShowCollision();
 	}
 
@@ -2309,7 +2309,7 @@ void UGameViewportClient::ToggleShowVolumes()
 	// Don't allow 'show collision' and 'show volumes' at the same time, so turn collision off
 	if (EngineShowFlags.Volumes && EngineShowFlags.Collision)
 	{
-		EngineShowFlags.Collision = false;
+		EngineShowFlags.SetCollision(false);
 		ToggleShowCollision();
 	}
 
@@ -2350,7 +2350,7 @@ void UGameViewportClient::ToggleShowCollision()
 		// Don't allow 'show collision' and 'show volumes' at the same time, so turn collision off
 		if (EngineShowFlags.Volumes)
 		{
-			EngineShowFlags.Volumes = false;
+			EngineShowFlags.SetVolumes(false);
 			ToggleShowVolumes();
 		}
 
diff --git a/Engine/Source/Runtime/Engine/Private/SceneView.cpp b/Engine/Source/Runtime/Engine/Private/SceneView.cpp
--- a/Engine/Source/Runtime/Engine/Private/SceneView.cpp
+++ b/Engine/Source/Runtime/Engine/Private/SceneView.cpp
@@ -1178,6 +1178,8 @@ void FSceneView::StartFinalPostprocessSettings(FVector InViewLocation)
 
 void FSceneView::EndFinalPostprocessSettings(const FSceneViewInitOptions& ViewInitOptions)
 {
+	const auto SceneViewFeatureLevel = GetFeatureLevel();
+
 	// will be deprecated soon, use the new asset LightPropagationVolumeBlendable instead
 	{
 		FLightPropagationVolumeSettings& Dest = FinalPostProcessSettings.BlendableManager.GetSingleFinalData<FLightPropagationVolumeSettings>();
@@ -1322,7 +1324,10 @@ void FSceneView::EndFinalPostprocessSettings(const FSceneViewInitOptions& ViewIn
 		}
 	}
 
-	if(!Family->EngineShowFlags.ScreenPercentage || bIsSceneCapture || bIsReflectionCapture)
+	// Not supported in ES2/3.
+	bool bES2Or3 = (SceneViewFeatureLevel == ERHIFeatureLevel::ES2) || (SceneViewFeatureLevel == ERHIFeatureLevel::ES3_1);
+
+	if(!Family->EngineShowFlags.ScreenPercentage || bIsSceneCapture || bIsReflectionCapture || bES2Or3)
 	{
 		FinalPostProcessSettings.ScreenPercentage = 100;
 	}
@@ -1383,8 +1388,6 @@ void FSceneView::EndFinalPostprocessSettings(const FSceneViewInitOptions& ViewIn
 
 	// Anti-Aliasing
 	{
-		const auto SceneViewFeatureLevel = GetFeatureLevel();
-
 		static const auto CVar = IConsoleManager::Get().FindTConsoleVariableDataInt(TEXT("r.PostProcessAAQuality")); 
 		static auto* MobileHDRCvar = IConsoleManager::Get().FindTConsoleVariableDataInt(TEXT("r.MobileHDR"));
 		static auto* MobileMSAACvar = IConsoleManager::Get().FindTConsoleVariableDataInt(TEXT("r.MobileMSAA"));
@@ -1657,13 +1660,6 @@ FSceneViewFamily::FSceneViewFamily( const ConstructionValues& CVS )
 	HierarchicalLODOverride = -1;
 	bDrawBaseInfo = true;
 #endif
-
-	// Not supported in ES2.
-	auto FeatureLevel = GetFeatureLevel();
-	if (FeatureLevel == ERHIFeatureLevel::ES2 || FeatureLevel == ERHIFeatureLevel::ES3_1)
-	{
-		EngineShowFlags.ScreenPercentage = false;
-	}
 }
 
 void FSceneViewFamily::ComputeFamilySize()
diff --git a/Engine/Source/Runtime/Engine/Private/ShowFlags.cpp b/Engine/Source/Runtime/Engine/Private/ShowFlags.cpp
--- a/Engine/Source/Runtime/Engine/Private/ShowFlags.cpp
+++ b/Engine/Source/Runtime/Engine/Private/ShowFlags.cpp
@@ -140,7 +140,7 @@ void FEngineShowFlags::SetSingleFlag(uint32 Index, bool bSet)
 	{
 	#define SHOWFLAG_ALWAYS_ACCESSIBLE(a,...) case SF_##a: a = bSet?1:0; break;
 	#if UE_BUILD_OPTIMIZED_SHOWFLAGS 
-		#define SHOWFLAG_FIXED_IN_SHIPPING(a,...) case SF_##a: break;
+		#define SHOWFLAG_FIXED_IN_SHIPPING(v,a,...) case SF_##a: break;
 	#endif
 	#include "ShowFlagsValues.inl"
 	default:
@@ -316,7 +316,7 @@ void EngineShowFlagOverride(EShowFlagInitMode ShowFlagInitMode, EViewModeIndex V
 	if(ShowFlagInitMode == ESFIM_Game)
 	{
 		// editor only features
-		EngineShowFlags.AudioRadius = 0;
+		EngineShowFlags.SetAudioRadius(false);
 	}
 
 	{
@@ -326,9 +326,9 @@ void EngineShowFlagOverride(EShowFlagInitMode ShowFlagInitMode, EViewModeIndex V
 			// disabled as it requires multiple frames, AA can be done by downsampling, more control and better masking
 			EngineShowFlags.TemporalAA = 0;
 			// no editor gizmos / selection
-			EngineShowFlags.ModeWidgets = 0;
-			EngineShowFlags.Selection = 0;
-			EngineShowFlags.SelectionOutline = 0;
+			EngineShowFlags.SetModeWidgets(false);
+			EngineShowFlags.SetSelection(false);
+			EngineShowFlags.SetSelectionOutline(false);
 		}
 	}
 
@@ -361,14 +361,6 @@ void EngineShowFlagOverride(EShowFlagInitMode ShowFlagInitMode, EViewModeIndex V
 		}
 	}
 
-	{
-		static const auto ICVar = IConsoleManager::Get().FindTConsoleVariableDataInt(TEXT("r.RefractionQuality"));
-		if(ICVar->GetValueOnGameThread() <= 0)
-		{
-			EngineShowFlags.Refraction = 0;
-		}
-	}
-
 	{
 		static const auto ICVar = IConsoleManager::Get().FindTConsoleVariableDataInt(TEXT("r.EyeAdaptationQuality"));
 		if(ICVar->GetValueOnGameThread() <= 0)
@@ -402,7 +394,7 @@ void EngineShowFlagOverride(EShowFlagInitMode ShowFlagInitMode, EViewModeIndex V
 
 		if( ViewModeIndex == VMI_BrushWireframe)
 		{
-			EngineShowFlags.Brushes = 1;
+			EngineShowFlags.SetBrushes(true);
 		}
 
 		if( ViewModeIndex == VMI_Wireframe ||
@@ -413,7 +405,7 @@ void EngineShowFlagOverride(EShowFlagInitMode ShowFlagInitMode, EViewModeIndex V
 			ViewModeIndex == VMI_VertexDensities ||
 			ViewModeIndex == VMI_LightmapDensity)
 		{
-			EngineShowFlags.Lighting = 0;
+			EngineShowFlags.SetLighting(false);
 			EngineShowFlags.AtmosphericFog = 0;
 		}
 
@@ -421,14 +413,14 @@ void EngineShowFlagOverride(EShowFlagInitMode ShowFlagInitMode, EViewModeIndex V
 			ViewModeIndex == VMI_LightingOnly ||
 			ViewModeIndex == VMI_LitLightmapDensity)
 		{
-			EngineShowFlags.Lighting = 1;
+			EngineShowFlags.SetLighting(true);
 		}
 
 		if( ViewModeIndex == VMI_LightingOnly ||
 			ViewModeIndex == VMI_BrushWireframe ||
 			ViewModeIndex == VMI_StationaryLightOverlap)
 		{
-			EngineShowFlags.Materials = 0;
+			EngineShowFlags.SetMaterials(false);
 		}
 
 		if( ViewModeIndex == VMI_LightComplexity )
@@ -438,7 +430,7 @@ void EngineShowFlagOverride(EShowFlagInitMode ShowFlagInitMode, EViewModeIndex V
 
 		if( ViewModeIndex == VMI_LODColoration )
 		{
-			EngineShowFlags.Lighting = 1;		// Best currently otherwise the image becomes hard to read.
+			EngineShowFlags.SetLighting(true);	// Best currently otherwise the image becomes hard to read.
 			EngineShowFlags.Fog = 0;			// Removed fog to improve color readability.
 			EngineShowFlags.AtmosphericFog = 0;
 			EngineShowFlags.Translucency = 0;	// Translucent are off because there are no color override shader currently for translucency.
@@ -448,7 +440,7 @@ void EngineShowFlagOverride(EShowFlagInitMode ShowFlagInitMode, EViewModeIndex V
 	// disable AA in full screen GBuffer visualization
 	if(EngineShowFlags.VisualizeBuffer && CurrentBufferVisualizationMode != NAME_None)
 	{
-		EngineShowFlags.Tonemapper = 0;
+		EngineShowFlags.SetTonemapper(false);
 	}
 
 #if !(UE_BUILD_SHIPPING || UE_BUILD_TEST)
diff --git a/Engine/Source/Runtime/Engine/Private/SkeletalRenderCPUSkin.cpp b/Engine/Source/Runtime/Engine/Private/SkeletalRenderCPUSkin.cpp
--- a/Engine/Source/Runtime/Engine/Private/SkeletalRenderCPUSkin.cpp
+++ b/Engine/Source/Runtime/Engine/Private/SkeletalRenderCPUSkin.cpp
@@ -201,37 +201,6 @@ void FSkeletalMeshObjectCPUSkin::UpdateDynamicData_RenderThread(FRHICommandListI
 	CacheVertices(DynamicData->LODIndex,true);
 }
 
-static bool ComputeTangent(FVector &t,
-							const FVector &p0, const FVector2D &c0,
-							const FVector &p1, const FVector2D &c1,
-							const FVector &p2, const FVector2D &c2)
-{
-	const float epsilon = 0.0001f;
-	bool   Ret = false;
-	FVector dp1 = p1 - p0;
-	FVector dp2 = p2 - p0;
-	float   du1 = c1.X - c0.X;
-	float   dv1 = c1.Y - c0.Y;
-	if(FMath::Abs(dv1) < epsilon && FMath::Abs(du1) >= epsilon)
-	{
-		t = dp1 / du1;
-		Ret = true;
-	}
-	else
-	{
-		float du2 = c2.X - c0.X;
-		float dv2 = c2.Y - c0.Y;
-		float det = dv1*du2 - dv2*du1;
-		if(FMath::Abs(det) >= epsilon)
-		{
-			t = (dp2*dv1-dp1*dv2)/det;
-			Ret = true;
-		}
-	}
-	return Ret;
-}
-
-
 void FSkeletalMeshObjectCPUSkin::CacheVertices(int32 LODIndex, bool bForce) const
 {
 	SCOPE_CYCLE_COUNTER( STAT_CPUSkinUpdateRTTime);
diff --git a/Engine/Source/Runtime/Engine/Private/UnrealClient.cpp b/Engine/Source/Runtime/Engine/Private/UnrealClient.cpp
--- a/Engine/Source/Runtime/Engine/Private/UnrealClient.cpp
+++ b/Engine/Source/Runtime/Engine/Private/UnrealClient.cpp
@@ -755,17 +755,17 @@ void FViewport::HighResScreenshot()
 
 	DummyViewport->EnqueueBeginRenderFrame();
 
-	uint32 MaskShowFlagBackup = ViewportClient->GetEngineShowFlags()->HighResScreenshotMask;
+	bool MaskShowFlagBackup = ViewportClient->GetEngineShowFlags()->HighResScreenshotMask;
 	uint32 MotionBlurShowFlagBackup = ViewportClient->GetEngineShowFlags()->MotionBlur;
-	ViewportClient->GetEngineShowFlags()->HighResScreenshotMask = GetHighResScreenshotConfig().bMaskEnabled ? 1 : 0;
-	ViewportClient->GetEngineShowFlags()->MotionBlur = 0;
+	ViewportClient->GetEngineShowFlags()->SetHighResScreenshotMask(GetHighResScreenshotConfig().bMaskEnabled);
+	ViewportClient->GetEngineShowFlags()->SetMotionBlur(false);
 
 	FCanvas Canvas(DummyViewport, NULL, ViewportClient->GetWorld(), ViewportClient->GetWorld()->FeatureLevel);
 	{
 		ViewportClient->Draw(DummyViewport, &Canvas);
 	}
 	Canvas.Flush_GameThread();
-	ViewportClient->GetEngineShowFlags()->HighResScreenshotMask = MaskShowFlagBackup;
+	ViewportClient->GetEngineShowFlags()->SetHighResScreenshotMask(MaskShowFlagBackup);
 	ViewportClient->GetEngineShowFlags()->MotionBlur = MotionBlurShowFlagBackup;
 	ViewportClient->ProcessScreenShots(DummyViewport);
 
diff --git a/Engine/Source/Runtime/Engine/Public/ShowFlags.h b/Engine/Source/Runtime/Engine/Public/ShowFlags.h
--- a/Engine/Source/Runtime/Engine/Public/ShowFlags.h
+++ b/Engine/Source/Runtime/Engine/Public/ShowFlags.h
@@ -39,9 +39,10 @@ enum EShowFlagInitMode
 #endif
 
 /**
- * Meant to replace the old FShowflag system, first we start with some new flags and
- * then incrementally we can move over to the new system. ViewModes should be a separate layer that can manipulate show flags before they get used.
- * So the view needs it's own show flags and the ones that should be finally used after manipulation from view mode and user (editor/show command)
+ * ShowFlags are a set of bits (some are fixed in SHIPPING) that are stored in the ViewFamily.
+ * They are use by the user to debug/profile parts of the rendering.
+ * They should not be used for scalability (as some might be compiled out for SHIPPING), there we use console variables.
+ * ViewModes are on a higher level, they can manipulate show flags before they get used.
  */
 struct FEngineShowFlags
 {
@@ -51,9 +52,9 @@ struct FEngineShowFlags
 	#define SHOWFLAG_ALWAYS_ACCESSIBLE(a,...) uint32 a : 1; void Set##a(bool bVal){ a = bVal?1:0;}
 
 	#if UE_BUILD_OPTIMIZED_SHOWFLAGS 
-		#define SHOWFLAG_FIXED_IN_SHIPPING(a,b,...) static const bool a = b; void Set##a(bool bVal){}
+		#define SHOWFLAG_FIXED_IN_SHIPPING(v,a,...) static const bool a = v; void Set##a(bool bVal){}
 	#else
-		#define SHOWFLAG_FIXED_IN_SHIPPING(a,b,c,d) SHOWFLAG_ALWAYS_ACCESSIBLE(a,c,d)
+		#define SHOWFLAG_FIXED_IN_SHIPPING(v,a,b,c) SHOWFLAG_ALWAYS_ACCESSIBLE(a,b,c)
 	#endif
 
 	#include "ShowFlagsValues.inl"
@@ -79,7 +80,7 @@ struct FEngineShowFlags
 		if (LocNames.Num() == 0)
 		{
 			#define SHOWFLAG_ALWAYS_ACCESSIBLE(a,b,c) LocNames.Add( TEXT(PREPROCESSOR_TO_STRING(a)), c);
-			#define SHOWFLAG_FIXED_IN_SHIPPING(a,b,c,d) LocNames.Add( TEXT(PREPROCESSOR_TO_STRING(a)), d);
+			#define SHOWFLAG_FIXED_IN_SHIPPING(v,a,b,c) LocNames.Add( TEXT(PREPROCESSOR_TO_STRING(a)), c);
 			#include "ShowFlagsValues.inl"
 
 			// Additional strings that don't correspond to a showflag variable:
@@ -123,7 +124,7 @@ struct FEngineShowFlags
 		switch (FlagIndex)
 		{
 			#define SHOWFLAG_ALWAYS_ACCESSIBLE(a,b,c) case SF_##a: return b;
-			#define SHOWFLAG_FIXED_IN_SHIPPING(a,b,c,d) case SF_##a: return c;
+			#define SHOWFLAG_FIXED_IN_SHIPPING(v,a,b,c) case SF_##a: return b;
 			#include "ShowFlagsValues.inl"
 
 			default:
@@ -171,7 +172,6 @@ struct FEngineShowFlags
 		SetScreenPercentage(false);
 		SetScreenSpaceReflections(false);
 		SetTemporalAA(false);
-		SetTextRender(false);
 
 		// might cause reallocation if we render rarely to it - for now off
 		SetAmbientOcclusion(false);
@@ -186,8 +186,6 @@ struct FEngineShowFlags
 		SetStereoRendering(false);
 		SetDistanceFieldAO(false);
 		SetDistanceFieldGI(false);
-
-		SetForceGBuffer(true);
 	}
 
 	// ---------------------------------------------------------
@@ -338,7 +336,6 @@ private:
 		SetVisualizeDistanceFieldAO(false);
 		SetVisualizeDistanceFieldGI(false);
 		SetVisualizeSSR(false);
-		SetForceGBuffer(false);
 		SetVisualizeSSS(false);
 		SetVisualizeBloom(false);
 	}
diff --git a/Engine/Source/Runtime/Engine/Public/ShowFlagsValues.inl b/Engine/Source/Runtime/Engine/Public/ShowFlagsValues.inl
--- a/Engine/Source/Runtime/Engine/Public/ShowFlagsValues.inl
+++ b/Engine/Source/Runtime/Engine/Public/ShowFlagsValues.inl
@@ -13,285 +13,283 @@
 
 // the default case for SHOWFLAG_FIXED_IN_SHIPPING is to give flag name.
 #ifndef SHOWFLAG_FIXED_IN_SHIPPING
-#define SHOWFLAG_FIXED_IN_SHIPPING(a,...) SHOWFLAG_ALWAYS_ACCESSIBLE(a,__VA_ARGS__)
+#define SHOWFLAG_FIXED_IN_SHIPPING(v,a,...) SHOWFLAG_ALWAYS_ACCESSIBLE(a,__VA_ARGS__)
 #endif
 
-/** Affects all postprocessing features, depending on viewmode this is on or off */
+/** Affects all postprocessing features, depending on viewmode this is on or off, for now SHOWFLAG_ALWAYS_ACCESSIBLE because it's used by ReflectionEnviromentCapture */
 SHOWFLAG_ALWAYS_ACCESSIBLE(PostProcessing, SFG_Hidden, LOCTEXT("PostProcessingSF", "Post-processing"))
-/** Bloom */
+/** Bloom, for now SHOWFLAG_ALWAYS_ACCESSIBLE because it's exposed in SceneCapture */
 SHOWFLAG_ALWAYS_ACCESSIBLE(Bloom, SFG_PostProcess, LOCTEXT("BloomSF", "Bloom"))
 /** HDR->LDR conversion is done through a tone mapper (otherwise linear mapping is used) */
-SHOWFLAG_ALWAYS_ACCESSIBLE(Tonemapper, SFG_PostProcess, LOCTEXT("TonemapperSF", "Tonemapper"))
-/** Any Anti-aliasing e.g. FXAA, Temporal AA */
+SHOWFLAG_FIXED_IN_SHIPPING(1, Tonemapper, SFG_PostProcess, LOCTEXT("TonemapperSF", "Tonemapper"))
+/** Any Anti-aliasing e.g. FXAA, Temporal AA, for now SHOWFLAG_ALWAYS_ACCESSIBLE because it's exposed in SceneCapture */
 SHOWFLAG_ALWAYS_ACCESSIBLE(AntiAliasing, SFG_Normal, LOCTEXT("AntiAliasingSF", "Anti-aliasing"))
-/** Only used in AntiAliasing is on, true:uses Temporal AA, otherwise FXAA */
+/** Only used in AntiAliasing is on, true:uses Temporal AA, otherwise FXAA, for now SHOWFLAG_ALWAYS_ACCESSIBLE because it's exposed in SceneCapture  */
 SHOWFLAG_ALWAYS_ACCESSIBLE(TemporalAA, SFG_Advanced, LOCTEXT("TemporalAASF", "Temporal AA (instead FXAA)"))
-/** e.g. Ambient cube map  */
+/** e.g. Ambient cube map, for now SHOWFLAG_ALWAYS_ACCESSIBLE because it's exposed in SceneCapture */
 SHOWFLAG_ALWAYS_ACCESSIBLE(AmbientCubemap, SFG_LightingFeatures, LOCTEXT("AmbientCubemapSF", "Ambient Cubemap"))
-/** Human like eye simulation to adapt to the brightness of the view */
+/** Human like eye simulation to adapt to the brightness of the view, for now SHOWFLAG_ALWAYS_ACCESSIBLE because it's exposed in SceneCapture */
 SHOWFLAG_ALWAYS_ACCESSIBLE(EyeAdaptation, SFG_PostProcess, LOCTEXT("EyeAdaptationSF", "Eye Adaptation"))
 /** Display a histogram of the scene HDR color */
-SHOWFLAG_ALWAYS_ACCESSIBLE(VisualizeHDR, SFG_Visualize, LOCTEXT("VisualizeHDRSF", "HDR (Eye Adaptation)"))
+SHOWFLAG_FIXED_IN_SHIPPING(0, VisualizeHDR, SFG_Visualize, LOCTEXT("VisualizeHDRSF", "HDR (Eye Adaptation)"))
 /** Image based lens flares (Simulate artifact of reflections within a camera system) */
-SHOWFLAG_ALWAYS_ACCESSIBLE(LensFlares, SFG_PostProcess, LOCTEXT("LensFlaresSF", "Lens Flares"))
-/** show indirect lighting component */
+SHOWFLAG_FIXED_IN_SHIPPING(1, LensFlares, SFG_PostProcess, LOCTEXT("LensFlaresSF", "Lens Flares"))
+/** show indirect lighting component, for now SHOWFLAG_ALWAYS_ACCESSIBLE because it's needed by r.GBuffer */
 SHOWFLAG_ALWAYS_ACCESSIBLE(GlobalIllumination, SFG_LightingComponents, LOCTEXT("GlobalIlluminationSF", "Global Illumination"))
 /** Darkens the screen borders (Camera artifact and artistic effect) */
-SHOWFLAG_ALWAYS_ACCESSIBLE(Vignette, SFG_PostProcess, LOCTEXT("VignetteSF", "Vignette"))
+SHOWFLAG_FIXED_IN_SHIPPING(1, Vignette, SFG_PostProcess, LOCTEXT("VignetteSF", "Vignette"))
 /** Fine film grain */
-SHOWFLAG_ALWAYS_ACCESSIBLE(Grain, SFG_PostProcess, LOCTEXT("GrainSF", "Grain"))
-/** Screen Space Ambient Occlusion */
+SHOWFLAG_FIXED_IN_SHIPPING(1, Grain, SFG_PostProcess, LOCTEXT("GrainSF", "Grain"))
+/** Screen Space Ambient Occlusion, for now SHOWFLAG_ALWAYS_ACCESSIBLE because it's exposed in SceneCapture */
 SHOWFLAG_ALWAYS_ACCESSIBLE(AmbientOcclusion, SFG_LightingComponents, LOCTEXT("AmbientOcclusionSF", "Ambient Occlusion"))
-/** Decal rendering */
+/** Decal rendering, for now SHOWFLAG_ALWAYS_ACCESSIBLE because it's exposed in SceneCapture */
 SHOWFLAG_ALWAYS_ACCESSIBLE(Decals, SFG_Normal, LOCTEXT("DecalsSF", "Decals"))
 /** like bloom dirt mask */
-SHOWFLAG_ALWAYS_ACCESSIBLE(CameraImperfections, SFG_PostProcess, LOCTEXT("CameraImperfectionsSF", "Camera Imperfections"))
+SHOWFLAG_FIXED_IN_SHIPPING(1, CameraImperfections, SFG_PostProcess, LOCTEXT("CameraImperfectionsSF", "Camera Imperfections"))
 /** to allow to disable visualizetexture for some editor rendering (e.g. thumbnail rendering) */
 SHOWFLAG_ALWAYS_ACCESSIBLE(OnScreenDebug, SFG_Developer, LOCTEXT("OnScreenDebugSF", "On Screen Debug"))
 /** needed for VMI_Lit_DetailLighting, Whether to override material diffuse and specular with constants, used by the Detail Lighting viewmode. */
-SHOWFLAG_ALWAYS_ACCESSIBLE(OverrideDiffuseAndSpecular, SFG_Hidden, LOCTEXT("OverrideDiffuseAndSpecularSF", "Override Diffuse And Specular"))
+SHOWFLAG_FIXED_IN_SHIPPING(0, OverrideDiffuseAndSpecular, SFG_Hidden, LOCTEXT("OverrideDiffuseAndSpecularSF", "Override Diffuse And Specular"))
 /** needed for VMI_ReflectionOverride, Whether to override all materials to be smooth, mirror reflections. */
-SHOWFLAG_ALWAYS_ACCESSIBLE(ReflectionOverride, SFG_Hidden, LOCTEXT("ReflectionOverrideSF", "Reflections"))
+SHOWFLAG_FIXED_IN_SHIPPING(0, ReflectionOverride, SFG_Hidden, LOCTEXT("ReflectionOverrideSF", "Reflections"))
 /** needed for VMI_VisualizeBuffer, Whether to enable the buffer visualization mode. */
-SHOWFLAG_ALWAYS_ACCESSIBLE(VisualizeBuffer, SFG_Hidden, LOCTEXT("VisualizeBufferSF", "Buffer Visualization"))
+SHOWFLAG_FIXED_IN_SHIPPING(0, VisualizeBuffer, SFG_Hidden, LOCTEXT("VisualizeBufferSF", "Buffer Visualization"))
 /** Allows to disable all direct lighting (does not affect indirect light) */
-SHOWFLAG_FIXED_IN_SHIPPING(DirectLighting, 1, SFG_LightingComponents, LOCTEXT("DirectLightingSF", "Direct Lighting"))
+SHOWFLAG_FIXED_IN_SHIPPING(1, DirectLighting, SFG_LightingComponents, LOCTEXT("DirectLightingSF", "Direct Lighting"))
 /** Allows to disable lighting from Directional Lights */
-SHOWFLAG_FIXED_IN_SHIPPING(DirectionalLights, 1, SFG_LightTypes, LOCTEXT("DirectionalLightsSF", "Directional Lights"))
+SHOWFLAG_FIXED_IN_SHIPPING(1, DirectionalLights, SFG_LightTypes, LOCTEXT("DirectionalLightsSF", "Directional Lights"))
 /** Allows to disable lighting from Point Lights */
-SHOWFLAG_FIXED_IN_SHIPPING(PointLights, 1, SFG_LightTypes, LOCTEXT("PointLightsSF", "Point Lights"))
+SHOWFLAG_FIXED_IN_SHIPPING(1, PointLights, SFG_LightTypes, LOCTEXT("PointLightsSF", "Point Lights"))
 /** Allows to disable lighting from Spot Lights */
-SHOWFLAG_FIXED_IN_SHIPPING(SpotLights, 1, SFG_LightTypes, LOCTEXT("SpotLightsSF", "Spot Lights"))
+SHOWFLAG_FIXED_IN_SHIPPING(1, SpotLights, SFG_LightTypes, LOCTEXT("SpotLightsSF", "Spot Lights"))
 /** Color correction after tone mapping */
-SHOWFLAG_ALWAYS_ACCESSIBLE(ColorGrading, SFG_PostProcess, LOCTEXT("ColorGradingSF", "Color Grading"))
+SHOWFLAG_FIXED_IN_SHIPPING(1, ColorGrading, SFG_PostProcess, LOCTEXT("ColorGradingSF", "Color Grading"))
 /** Visualize vector fields. */
-SHOWFLAG_ALWAYS_ACCESSIBLE(VectorFields, SFG_Developer, LOCTEXT("VectorFieldsSF", "Vector Fields"))
+SHOWFLAG_FIXED_IN_SHIPPING(0, VectorFields, SFG_Developer, LOCTEXT("VectorFieldsSF", "Vector Fields"))
 /** Depth of Field */
-SHOWFLAG_ALWAYS_ACCESSIBLE(DepthOfField, SFG_PostProcess, LOCTEXT("DepthOfFieldSF", "Depth Of Field"))
+SHOWFLAG_FIXED_IN_SHIPPING(1, DepthOfField, SFG_PostProcess, LOCTEXT("DepthOfFieldSF", "Depth Of Field"))
 /** Whether editor-hidden primitives cast shadows. */
-SHOWFLAG_ALWAYS_ACCESSIBLE(ShadowsFromEditorHiddenActors, SFG_Developer, LOCTEXT("ShadowsFromEditorHiddenActorsSF", "Shadows of Editor-Hidden Actors"))
+SHOWFLAG_FIXED_IN_SHIPPING(1, ShadowsFromEditorHiddenActors, SFG_Developer, LOCTEXT("ShadowsFromEditorHiddenActorsSF", "Shadows of Editor-Hidden Actors"))
 /** Highlight materials that indicate performance issues or show unrealistic materials */
-SHOWFLAG_ALWAYS_ACCESSIBLE(GBufferHints, SFG_Developer, LOCTEXT("GBufferHintsSF", "GBuffer Hints (material attributes)"))
-/** MotionBlur, for now only camera motion blur */
+SHOWFLAG_FIXED_IN_SHIPPING(0, GBufferHints, SFG_Developer, LOCTEXT("GBufferHintsSF", "GBuffer Hints (material attributes)"))
+/** MotionBlur, for now only camera motion blur, for now SHOWFLAG_ALWAYS_ACCESSIBLE because it's exposed in SceneCapture */
 SHOWFLAG_ALWAYS_ACCESSIBLE(MotionBlur, SFG_PostProcess, LOCTEXT("MotionBlurSF", "Motion Blur"))
-/** Whether to render the editor gizmos and other foreground editor widgets off screen and apply them after post process */
-SHOWFLAG_ALWAYS_ACCESSIBLE(CompositeEditorPrimitives, SFG_Developer, LOCTEXT("CompositeEditorPrimitivesSF", "Composite Editor Primitives"))
+/** Whether to render the editor gizmos and other foreground editor widgets off screen and apply them after post process, only needed for the editor */
+SHOWFLAG_FIXED_IN_SHIPPING(0, CompositeEditorPrimitives, SFG_Developer, LOCTEXT("CompositeEditorPrimitivesSF", "Composite Editor Primitives"))
 /** Shows a test image that allows to tweak the monitor colors, borders and allows to judge image and temporal aliasing  */
-SHOWFLAG_ALWAYS_ACCESSIBLE(TestImage, SFG_Developer, LOCTEXT("TestImageSF", "Test Image"))
+SHOWFLAG_FIXED_IN_SHIPPING(0, TestImage, SFG_Developer, LOCTEXT("TestImageSF", "Test Image"))
 /** Helper to tweak depth of field */
-SHOWFLAG_ALWAYS_ACCESSIBLE(VisualizeDOF, SFG_Visualize, LOCTEXT("VisualizeDOFSF", "Depth of Field Layers"))
+SHOWFLAG_FIXED_IN_SHIPPING(0, VisualizeDOF, SFG_Visualize, LOCTEXT("VisualizeDOFSF", "Depth of Field Layers"))
 /** Helper to tweak depth of field */
-SHOWFLAG_ALWAYS_ACCESSIBLE(VisualizeAdaptiveDOF, SFG_Visualize, LOCTEXT("VisualizeAdaptiveDOFSF", "Adaptive Depth of Field"))
+SHOWFLAG_FIXED_IN_SHIPPING(0, VisualizeAdaptiveDOF, SFG_Visualize, LOCTEXT("VisualizeAdaptiveDOFSF", "Adaptive Depth of Field"))
 /** Show Vertex Colors */
-SHOWFLAG_ALWAYS_ACCESSIBLE(VertexColors, SFG_Advanced, LOCTEXT("VertexColorsSF", "Vertex Colors"))
+SHOWFLAG_FIXED_IN_SHIPPING(0, VertexColors, SFG_Advanced, LOCTEXT("VertexColorsSF", "Vertex Colors"))
 /** Render Post process (screen space) distortion/refraction */
-SHOWFLAG_ALWAYS_ACCESSIBLE(Refraction, SFG_Developer, LOCTEXT("RefractionSF", "Refraction"))
+SHOWFLAG_FIXED_IN_SHIPPING(1, Refraction, SFG_Developer, LOCTEXT("RefractionSF", "Refraction"))
 /** Usually set in game or when previewing Matinee but not in editor, used for motion blur or any kind of rendering features that rely on the former frame */
 SHOWFLAG_ALWAYS_ACCESSIBLE(CameraInterpolation, SFG_Hidden, LOCTEXT("CameraInterpolationSF", "Camera Interpolation"))
 /** Post processing color fringe (chromatic aberration) */
-SHOWFLAG_ALWAYS_ACCESSIBLE(SceneColorFringe, SFG_PostProcess, LOCTEXT("SceneColorFringeSF", "Scene Color Fringe"))
-/** If Translucency should be rendered into a separate RT and composited without DepthOfField, can be disabled in the materials (affects sorting) */
+SHOWFLAG_FIXED_IN_SHIPPING(1, SceneColorFringe, SFG_PostProcess, LOCTEXT("SceneColorFringeSF", "Scene Color Fringe"))
+/** If Translucency should be rendered into a separate RT and composited without DepthOfField, can be disabled in the materials (affects sorting), SHOWFLAG_ALWAYS_ACCESSIBLE for now because USceneCaptureComponent needs that */
 SHOWFLAG_ALWAYS_ACCESSIBLE(SeparateTranslucency, SFG_Advanced, LOCTEXT("SeparateTranslucencySF", "Separate Translucency"))
 /** If Screen Percentage should be applied (upscaling), useful to disable it in editor */
-SHOWFLAG_ALWAYS_ACCESSIBLE(ScreenPercentage, SFG_PostProcess, LOCTEXT("ScreenPercentageSF", "Screen Percentage"))
+SHOWFLAG_FIXED_IN_SHIPPING(1, ScreenPercentage, SFG_PostProcess, LOCTEXT("ScreenPercentageSF", "Screen Percentage"))
 /** Helper to tweak motion blur settings */
-SHOWFLAG_ALWAYS_ACCESSIBLE(VisualizeMotionBlur, SFG_Visualize, LOCTEXT("VisualizeMotionBlurSF", "Motion Blur"))
-/** Whether to display the Reflection Environment feature, which has local reflections from Reflection Capture actors. */
+SHOWFLAG_FIXED_IN_SHIPPING(0, VisualizeMotionBlur, SFG_Visualize, LOCTEXT("VisualizeMotionBlurSF", "Motion Blur"))
+/** Whether to display the Reflection Environment feature, which has local reflections from Reflection Capture actors, for now SHOWFLAG_ALWAYS_ACCESSIBLE because it's exposed in SceneCapture */
 SHOWFLAG_ALWAYS_ACCESSIBLE(ReflectionEnvironment, SFG_LightingFeatures, LOCTEXT("ReflectionEnvironmentSF", "Reflection Environment"))
 /** Visualize pixels that are outside of their object's bounding box (content error). */
-SHOWFLAG_ALWAYS_ACCESSIBLE(VisualizeOutOfBoundsPixels, SFG_Visualize, LOCTEXT("VisualizeOutOfBoundsPixelsSF", "Out of Bounds Pixels"))
+SHOWFLAG_FIXED_IN_SHIPPING(0, VisualizeOutOfBoundsPixels, SFG_Visualize, LOCTEXT("VisualizeOutOfBoundsPixelsSF", "Out of Bounds Pixels"))
 /** Whether to display the scene's diffuse. */
-SHOWFLAG_ALWAYS_ACCESSIBLE(Diffuse, SFG_LightingComponents, LOCTEXT("DiffuseSF", "Diffuse"))
+SHOWFLAG_FIXED_IN_SHIPPING(1, Diffuse, SFG_LightingComponents, LOCTEXT("DiffuseSF", "Diffuse"))
 /** Whether to display the scene's specular, including reflections. */
-SHOWFLAG_ALWAYS_ACCESSIBLE(Specular, SFG_LightingComponents, LOCTEXT("SpecularSF", "Specular"))
+SHOWFLAG_FIXED_IN_SHIPPING(1, Specular, SFG_LightingComponents, LOCTEXT("SpecularSF", "Specular"))
 /** Outline around selected objects in the editor */
-SHOWFLAG_ALWAYS_ACCESSIBLE(SelectionOutline, SFG_Hidden, LOCTEXT("SelectionOutlineSF", "Selection Outline"))
-/** If screen space reflections are enabled */
+SHOWFLAG_FIXED_IN_SHIPPING(0, SelectionOutline, SFG_Hidden, LOCTEXT("SelectionOutlineSF", "Selection Outline"))
+/** If screen space reflections are enabled, for now SHOWFLAG_ALWAYS_ACCESSIBLE because it's exposed in SceneCapture */
 SHOWFLAG_ALWAYS_ACCESSIBLE(ScreenSpaceReflections, SFG_LightingFeatures, LOCTEXT("ScreenSpaceReflectionsSF", "Screen Space Reflections"))
 /** If Screen Space Subsurface Scattering enabled */
-SHOWFLAG_FIXED_IN_SHIPPING(SubsurfaceScattering, 1, SFG_LightingFeatures, LOCTEXT("SubsurfaceScatteringSF", "Subsurface Scattering (Screen Space)"))
+SHOWFLAG_FIXED_IN_SHIPPING(1, SubsurfaceScattering, SFG_LightingFeatures, LOCTEXT("SubsurfaceScatteringSF", "Subsurface Scattering (Screen Space)"))
 /** If Screen Space Subsurface Scattering visualization is enabled */
-SHOWFLAG_FIXED_IN_SHIPPING(VisualizeSSS, 0, SFG_Visualize, LOCTEXT("VisualizeSSSSF", "Subsurface Scattering (Screen Space)"))
+SHOWFLAG_FIXED_IN_SHIPPING(0, VisualizeSSS, SFG_Visualize, LOCTEXT("VisualizeSSSSF", "Subsurface Scattering (Screen Space)"))
 /** If the indirect lighting cache is enabled */
-SHOWFLAG_ALWAYS_ACCESSIBLE(IndirectLightingCache, SFG_LightingFeatures, LOCTEXT("IndirectLightingCacheSF", "Indirect Lighting Cache"))
+SHOWFLAG_FIXED_IN_SHIPPING(1, IndirectLightingCache, SFG_LightingFeatures, LOCTEXT("IndirectLightingCacheSF", "Indirect Lighting Cache"))
 /** calls debug drawing for AIs */
-SHOWFLAG_ALWAYS_ACCESSIBLE(DebugAI, SFG_Developer, LOCTEXT("DebugAISF", "AI Debug"))
+SHOWFLAG_FIXED_IN_SHIPPING(0, DebugAI, SFG_Developer, LOCTEXT("DebugAISF", "AI Debug"))
 /** calls debug drawing for whatever LogVisualizer wants to draw */
-SHOWFLAG_ALWAYS_ACCESSIBLE(VisLog, SFG_Developer, LOCTEXT("VisLogSF", "Log Visualizer"))
+SHOWFLAG_FIXED_IN_SHIPPING(0, VisLog, SFG_Developer, LOCTEXT("VisLogSF", "Log Visualizer"))
 /** whether to draw navigation data */
-SHOWFLAG_ALWAYS_ACCESSIBLE(Navigation, SFG_Normal, LOCTEXT("NavigationSF", "Navigation"))
+SHOWFLAG_FIXED_IN_SHIPPING(0, Navigation, SFG_Normal, LOCTEXT("NavigationSF", "Navigation"))
 /** used by gameplay debugging components to debug-draw on screen */
-SHOWFLAG_ALWAYS_ACCESSIBLE(GameplayDebug, SFG_Developer, LOCTEXT("GameplayDebugSF", "Gameplay Debug"))
-/** LightProfiles, usually 1d textures to have a light  (IES) */
-SHOWFLAG_ALWAYS_ACCESSIBLE(TexturedLightProfiles, SFG_LightingFeatures, LOCTEXT("TexturedLightProfilesSF", "Textured Light Profiles (IES Texture)"))
-/** LightFunctions (masking light sources with a material) */
+SHOWFLAG_FIXED_IN_SHIPPING(0, GameplayDebug, SFG_Developer, LOCTEXT("GameplayDebugSF", "Gameplay Debug"))
+/** LightProfiles, usually 1d textures to have a light (IES), for now SHOWFLAG_ALWAYS_ACCESSIBLE because it's exposed in SceneCapture */
+SHOWFLAG_ALWAYS_ACCESSIBLE(TexturedLightProfiles,  SFG_LightingFeatures, LOCTEXT("TexturedLightProfilesSF", "Textured Light Profiles (IES Texture)"))
+/** LightFunctions (masking light sources with a material), for now SHOWFLAG_ALWAYS_ACCESSIBLE because it's exposed in SceneCapture */
 SHOWFLAG_ALWAYS_ACCESSIBLE(LightFunctions, SFG_LightingFeatures, LOCTEXT("LightFunctionsSF", "Light Functions"))
 /** Hardware Tessellation (DX11 feature) */
-SHOWFLAG_ALWAYS_ACCESSIBLE(Tessellation, SFG_Advanced, LOCTEXT("TessellationSF", "Tessellation"))
-/** Draws instanced static meshes that are not foliage or grass. */
+SHOWFLAG_FIXED_IN_SHIPPING(1, Tessellation, SFG_Advanced, LOCTEXT("TessellationSF", "Tessellation"))
+/** Draws instanced static meshes that are not foliage or grass, for now SHOWFLAG_ALWAYS_ACCESSIBLE because it's exposed in SceneCapture */
 SHOWFLAG_ALWAYS_ACCESSIBLE(InstancedStaticMeshes, SFG_Advanced, LOCTEXT("InstancedStaticMeshesSF", "Instanced Static Meshes"))
-/** Draws instanced foliage. */
+/** Draws instanced foliage, for now SHOWFLAG_ALWAYS_ACCESSIBLE because it's exposed in SceneCapture */
 SHOWFLAG_ALWAYS_ACCESSIBLE(InstancedFoliage, SFG_Advanced, LOCTEXT("InstancedFoliageSF", "Foliage"))
-/** Draws instanced grass. */
+/** Draws instanced grass, for now SHOWFLAG_ALWAYS_ACCESSIBLE because it's exposed in SceneCapture */
 SHOWFLAG_ALWAYS_ACCESSIBLE(InstancedGrass, SFG_Advanced, LOCTEXT("InstancedGrassSF", "Grass"))
-/** non baked shadows */
+/** non baked shadows, for now SHOWFLAG_ALWAYS_ACCESSIBLE because it's exposed in SceneCapture */
 SHOWFLAG_ALWAYS_ACCESSIBLE(DynamicShadows, SFG_LightingComponents, LOCTEXT("DynamicShadowsSF", "Dynamic Shadows"))
-/** Particles */
+/** Particles, for now SHOWFLAG_ALWAYS_ACCESSIBLE because it's exposed in SceneCapture */
 SHOWFLAG_ALWAYS_ACCESSIBLE(Particles, SFG_Normal, LOCTEXT("ParticlesSF", "Particles Sprite"))
-/** if SkeletalMeshes are getting rendered */
+/** if SkeletalMeshes are getting rendered, for now SHOWFLAG_ALWAYS_ACCESSIBLE because it's exposed in SceneCapture */
 SHOWFLAG_ALWAYS_ACCESSIBLE(SkeletalMeshes, SFG_Normal, LOCTEXT("SkeletalMeshesSF", "Skeletal Meshes"))
 /** if the builder brush (editor) is getting rendered */
-SHOWFLAG_ALWAYS_ACCESSIBLE(BuilderBrush, SFG_Hidden, LOCTEXT("BuilderBrushSF", "Builder Brush"))
-/** Render translucency */
+SHOWFLAG_FIXED_IN_SHIPPING(0, BuilderBrush, SFG_Hidden, LOCTEXT("BuilderBrushSF", "Builder Brush"))
+/** Render translucency, for now SHOWFLAG_ALWAYS_ACCESSIBLE because it's exposed in SceneCapture */
 SHOWFLAG_ALWAYS_ACCESSIBLE(Translucency, SFG_Normal, LOCTEXT("TranslucencySF", "Translucency"))
 /** Draw billboard components */
-SHOWFLAG_ALWAYS_ACCESSIBLE(BillboardSprites, SFG_Advanced, LOCTEXT("BillboardSpritesSF", "Billboard Sprites"))
+SHOWFLAG_FIXED_IN_SHIPPING(1, BillboardSprites, SFG_Advanced, LOCTEXT("BillboardSpritesSF", "Billboard Sprites"))
 /** Use LOD parenting, MinDrawDistance, etc. If disabled, will show LOD parenting lines */
 SHOWFLAG_ALWAYS_ACCESSIBLE(LOD, SFG_Advanced, LOCTEXT("LODSF", "LOD Parenting"))
 /** needed for VMI_LightComplexity */
-SHOWFLAG_ALWAYS_ACCESSIBLE(LightComplexity, SFG_Hidden, LOCTEXT("LightComplexitySF", "Light Complexity"))
+SHOWFLAG_FIXED_IN_SHIPPING(0, LightComplexity, SFG_Hidden, LOCTEXT("LightComplexitySF", "Light Complexity"))
 /** needed for VMI_ShaderComplexity, render world colored by shader complexity */
-SHOWFLAG_ALWAYS_ACCESSIBLE(ShaderComplexity, SFG_Hidden, LOCTEXT("ShaderComplexitySF", "Shader Complexity"))
+SHOWFLAG_FIXED_IN_SHIPPING(0, ShaderComplexity, SFG_Hidden, LOCTEXT("ShaderComplexitySF", "Shader Complexity"))
 /** needed for VMI_StationaryLightOverlap, render world colored by stationary light overlap */
-SHOWFLAG_ALWAYS_ACCESSIBLE(StationaryLightOverlap, SFG_Hidden, LOCTEXT("StationaryLightOverlapSF", "Stationary Light Overlap"))
+SHOWFLAG_FIXED_IN_SHIPPING(0, StationaryLightOverlap,  SFG_Hidden, LOCTEXT("StationaryLightOverlapSF", "Stationary Light Overlap"))
 /** needed for VMI_LightmapDensity and VMI_LitLightmapDensity, render checkerboard material with UVs scaled by lightmap resolution w. color tint for world-space lightmap density */
-SHOWFLAG_ALWAYS_ACCESSIBLE(LightMapDensity, SFG_Hidden, LOCTEXT("LightMapDensitySF", "Light Map Density"))
+SHOWFLAG_FIXED_IN_SHIPPING(0, LightMapDensity, SFG_Hidden, LOCTEXT("LightMapDensitySF", "Light Map Density"))
 /** Render streaming bounding volumes for the currently selected texture */
-SHOWFLAG_ALWAYS_ACCESSIBLE(StreamingBounds, SFG_Advanced, LOCTEXT("StreamingBoundsSF", "Streaming Bounds"))
+SHOWFLAG_FIXED_IN_SHIPPING(0, StreamingBounds, SFG_Advanced, LOCTEXT("StreamingBoundsSF", "Streaming Bounds"))
 /** Render joint limits */
-SHOWFLAG_FIXED_IN_SHIPPING(Constraints, 0, SFG_Advanced, LOCTEXT("ConstraintsSF", "Constraints"))
+SHOWFLAG_FIXED_IN_SHIPPING(0, Constraints, SFG_Advanced, LOCTEXT("ConstraintsSF", "Constraints"))
 /** Draws camera frustums */
-SHOWFLAG_ALWAYS_ACCESSIBLE(CameraFrustums, SFG_Advanced, LOCTEXT("CameraFrustumsSF", "Camera Frustums"))
+SHOWFLAG_FIXED_IN_SHIPPING(0, CameraFrustums, SFG_Advanced, LOCTEXT("CameraFrustumsSF", "Camera Frustums"))
 /** Draw sound actor radii */
-SHOWFLAG_ALWAYS_ACCESSIBLE(AudioRadius, SFG_Advanced, LOCTEXT("AudioRadiusSF", "Audio Radius"))
+SHOWFLAG_FIXED_IN_SHIPPING(0, AudioRadius, SFG_Advanced, LOCTEXT("AudioRadiusSF", "Audio Radius"))
 /** Colors BSP based on model component association */
-SHOWFLAG_ALWAYS_ACCESSIBLE(BSPSplit, SFG_Advanced, LOCTEXT("BSPSplitSF", "BSP Split"))
-/** show editor (wireframe) brushes */
-SHOWFLAG_ALWAYS_ACCESSIBLE(Brushes, SFG_Hidden, LOCTEXT("BrushesSF", "Brushes"))
+SHOWFLAG_FIXED_IN_SHIPPING(0, BSPSplit, SFG_Advanced, LOCTEXT("BSPSplitSF", "BSP Split"))
+/** show editor (wireframe) brushes, for now SHOWFLAG_ALWAYS_ACCESSIBLE because it's exposed in SceneCapture */
+SHOWFLAG_FIXED_IN_SHIPPING(0, Brushes, SFG_Hidden, LOCTEXT("BrushesSF", "Brushes"))
 /** Show the usual material light interaction */
-SHOWFLAG_ALWAYS_ACCESSIBLE(Lighting, SFG_Hidden, LOCTEXT("LightingSF", "Lighting"))
-/** Execute the deferred light passes, can be disabled for debugging */
+SHOWFLAG_FIXED_IN_SHIPPING(1, Lighting, SFG_Hidden, LOCTEXT("LightingSF", "Lighting"))
+/** Execute the deferred light passes, can be disabled for debugging, for now SHOWFLAG_ALWAYS_ACCESSIBLE because it's exposed in SceneCapture */
 SHOWFLAG_ALWAYS_ACCESSIBLE(DeferredLighting, SFG_Advanced, LOCTEXT("DeferredLightingSF", "DeferredLighting"))
 /** Special: Allows to hide objects in the editor, is evaluated per primitive */
-SHOWFLAG_ALWAYS_ACCESSIBLE(Editor, SFG_Hidden, LOCTEXT("EditorSF", "Editor"))
+SHOWFLAG_FIXED_IN_SHIPPING(0, Editor, SFG_Hidden, LOCTEXT("EditorSF", "Editor"))
 /** needed for VMI_BrushWireframe and VMI_LitLightmapDensity, Draws BSP triangles */
-SHOWFLAG_ALWAYS_ACCESSIBLE(BSPTriangles, SFG_Hidden, LOCTEXT("BSPTrianglesSF", "BSP Triangles"))
-/** Displays large clickable icons on static mesh vertices */
-SHOWFLAG_ALWAYS_ACCESSIBLE(LargeVertices, SFG_Advanced, LOCTEXT("LargeVerticesSF", "Large Vertices"))
+SHOWFLAG_FIXED_IN_SHIPPING(1, BSPTriangles, SFG_Hidden, LOCTEXT("BSPTrianglesSF", "BSP Triangles"))
+/** Displays large clickable icons on static mesh vertices, only needed for the editor */
+SHOWFLAG_FIXED_IN_SHIPPING(0, LargeVertices, SFG_Advanced, LOCTEXT("LargeVerticesSF", "Large Vertices"))
 /** To show the grid in editor (grey lines and red dots) */
-SHOWFLAG_ALWAYS_ACCESSIBLE(Grid, SFG_Normal, LOCTEXT("GridSF", "Grid"))
+SHOWFLAG_FIXED_IN_SHIPPING(0, Grid, SFG_Normal, LOCTEXT("GridSF", "Grid"))
 /** To show the snap in editor (only for editor view ports, red dots) */
-SHOWFLAG_ALWAYS_ACCESSIBLE(Snap, SFG_Hidden, LOCTEXT("SnapSF", "Snap"))
+SHOWFLAG_FIXED_IN_SHIPPING(0, Snap, SFG_Hidden, LOCTEXT("SnapSF", "Snap"))
 /** In the filled view modeModeWidgetss, render mesh edges as well as the filled surfaces. */
-SHOWFLAG_ALWAYS_ACCESSIBLE(MeshEdges, SFG_Advanced, LOCTEXT("MeshEdgesSF", "Mesh Edges"))
+SHOWFLAG_FIXED_IN_SHIPPING(0, MeshEdges, SFG_Advanced, LOCTEXT("MeshEdgesSF", "Mesh Edges"))
 /** Complex cover rendering */
-SHOWFLAG_ALWAYS_ACCESSIBLE(Cover, SFG_Hidden, LOCTEXT("CoverSF", "Cover"))
+SHOWFLAG_FIXED_IN_SHIPPING(0, Cover, SFG_Hidden, LOCTEXT("CoverSF", "Cover"))
 /** Spline rendering */
-SHOWFLAG_ALWAYS_ACCESSIBLE(Splines, SFG_Advanced, LOCTEXT("SplinesSF", "Splines"))
+SHOWFLAG_FIXED_IN_SHIPPING(1, Splines, SFG_Advanced, LOCTEXT("SplinesSF", "Splines"))
 /** Selection rendering, could be useful in game as well */
-SHOWFLAG_ALWAYS_ACCESSIBLE(Selection, SFG_Advanced, LOCTEXT("SelectionSF", "Selection"))
+SHOWFLAG_FIXED_IN_SHIPPING(0, Selection, SFG_Advanced, LOCTEXT("SelectionSF", "Selection"))
 /** Draws mode specific widgets and controls in the viewports (should only be set on viewport clients that are editing the level itself) */
-SHOWFLAG_ALWAYS_ACCESSIBLE(ModeWidgets, SFG_Advanced, LOCTEXT("ModeWidgetsSF", "Mode Widgets"))
+SHOWFLAG_FIXED_IN_SHIPPING(0, ModeWidgets, SFG_Advanced, LOCTEXT("ModeWidgetsSF", "Mode Widgets"))
 /**  */
-SHOWFLAG_ALWAYS_ACCESSIBLE(Bounds, SFG_Advanced, LOCTEXT("BoundsSF", "Bounds"))
-/** Draws each hit proxy in the scene with a different color */
-SHOWFLAG_ALWAYS_ACCESSIBLE(HitProxies, SFG_Developer, LOCTEXT("HitProxiesSF", "Hit Proxies"))
+SHOWFLAG_FIXED_IN_SHIPPING(0, Bounds,  SFG_Advanced, LOCTEXT("BoundsSF", "Bounds"))
+/** Draws each hit proxy in the scene with a different color, for now only available in the editor */
+SHOWFLAG_FIXED_IN_SHIPPING(0, HitProxies, SFG_Developer, LOCTEXT("HitProxiesSF", "Hit Proxies"))
 /** Render objects with colors based on the property values */
-SHOWFLAG_ALWAYS_ACCESSIBLE(PropertyColoration, SFG_Advanced, LOCTEXT("PropertyColorationSF", "Property Coloration"))
+SHOWFLAG_FIXED_IN_SHIPPING(0, PropertyColoration, SFG_Advanced, LOCTEXT("PropertyColorationSF", "Property Coloration"))
 /** Draw lines to lights affecting this mesh if its selected. */
-SHOWFLAG_ALWAYS_ACCESSIBLE(LightInfluences, SFG_Advanced, LOCTEXT("LightInfluencesSF", "Light Influences"))
-/**  */
-SHOWFLAG_ALWAYS_ACCESSIBLE(Pivot, SFG_Hidden, LOCTEXT("PivotSF", "Pivot"))
+SHOWFLAG_FIXED_IN_SHIPPING(0, LightInfluences, SFG_Advanced, LOCTEXT("LightInfluencesSF", "Light Influences"))
+/** for the Editor */
+SHOWFLAG_FIXED_IN_SHIPPING(0, Pivot, SFG_Hidden, LOCTEXT("PivotSF", "Pivot"))
 /** Draws un-occluded shadow frustums in wireframe */
-SHOWFLAG_ALWAYS_ACCESSIBLE(ShadowFrustums, SFG_Advanced, LOCTEXT("ShadowFrustumsSF", "Shadow Frustums"))
+SHOWFLAG_FIXED_IN_SHIPPING(0, ShadowFrustums, SFG_Advanced, LOCTEXT("ShadowFrustumsSF", "Shadow Frustums"))
 /** needed for VMI_Wireframe and VMI_BrushWireframe */
-SHOWFLAG_ALWAYS_ACCESSIBLE(Wireframe, SFG_Hidden, LOCTEXT("WireframeSF", "Wireframe"))
-/**  */
-SHOWFLAG_ALWAYS_ACCESSIBLE(Materials, SFG_Hidden, LOCTEXT("MaterialsSF", "Materials"))
+SHOWFLAG_FIXED_IN_SHIPPING(0, Wireframe, SFG_Hidden, LOCTEXT("WireframeSF", "Wireframe"))
 /**  */
+SHOWFLAG_FIXED_IN_SHIPPING(1, Materials, SFG_Hidden, LOCTEXT("MaterialsSF", "Materials"))
+/** for now SHOWFLAG_ALWAYS_ACCESSIBLE because it's exposed in SceneCapture */
 SHOWFLAG_ALWAYS_ACCESSIBLE(StaticMeshes, SFG_Normal, LOCTEXT("StaticMeshesSF", "Static Meshes"))
-/**  */
+/** for now SHOWFLAG_ALWAYS_ACCESSIBLE because it's exposed in SceneCapture */
 SHOWFLAG_ALWAYS_ACCESSIBLE(Landscape, SFG_Normal, LOCTEXT("LandscapeSF", "Landscape"))
 /**  */
-SHOWFLAG_ALWAYS_ACCESSIBLE(LightRadius, SFG_Advanced, LOCTEXT("LightRadiusSF", "Light Radius"))
-/** Draws fog */
+SHOWFLAG_FIXED_IN_SHIPPING(0, LightRadius, SFG_Advanced, LOCTEXT("LightRadiusSF", "Light Radius"))
+/** Draws fog, for now SHOWFLAG_ALWAYS_ACCESSIBLE because it's exposed in SceneCapture */
 SHOWFLAG_ALWAYS_ACCESSIBLE(Fog, SFG_Normal, LOCTEXT("FogSF", "Fog"))
 /** Draws Volumes */
-SHOWFLAG_ALWAYS_ACCESSIBLE(Volumes, SFG_Advanced, LOCTEXT("VolumesSF", "Volumes"))
+SHOWFLAG_FIXED_IN_SHIPPING(0, Volumes, SFG_Advanced, LOCTEXT("VolumesSF", "Volumes"))
 /** if this is a game viewport, needed? */
 SHOWFLAG_ALWAYS_ACCESSIBLE(Game, SFG_Hidden, LOCTEXT("GameSF", "Game"))
 /** Render objects with colors based on what the level they belong to */
-SHOWFLAG_ALWAYS_ACCESSIBLE(LevelColoration, SFG_Advanced, LOCTEXT("LevelColorationSF", "Level Coloration"))
-/** Draws BSP brushes (in game or editor textured triangles usually with lightmaps) */
+SHOWFLAG_FIXED_IN_SHIPPING(0, LevelColoration, SFG_Advanced, LOCTEXT("LevelColorationSF", "Level Coloration"))
+/** Draws BSP brushes (in game or editor textured triangles usually with lightmaps), for now SHOWFLAG_ALWAYS_ACCESSIBLE because it's exposed in SceneCapture */
 SHOWFLAG_ALWAYS_ACCESSIBLE(BSP, SFG_Normal, LOCTEXT("BSPSF", "BSP"))
 /** Collision drawing */
-SHOWFLAG_ALWAYS_ACCESSIBLE(Collision, SFG_Normal, LOCTEXT("CollisionWireFrame", "Collision"))
+SHOWFLAG_FIXED_IN_SHIPPING(0, Collision, SFG_Normal, LOCTEXT("CollisionWireFrame", "Collision"))
 /** Collision blocking visibility against complex **/
-SHOWFLAG_ALWAYS_ACCESSIBLE(CollisionVisibility, SFG_Hidden, LOCTEXT("CollisionVisibility", "Visibility"))
+SHOWFLAG_FIXED_IN_SHIPPING(0, CollisionVisibility, SFG_Hidden, LOCTEXT("CollisionVisibility", "Visibility"))
 /** Collision blocking pawn against simple collision **/
-SHOWFLAG_ALWAYS_ACCESSIBLE(CollisionPawn, SFG_Hidden, LOCTEXT("CollisionPawn", "Pawn"))
-/** Render LightShafts */
+SHOWFLAG_FIXED_IN_SHIPPING(0, CollisionPawn, SFG_Hidden, LOCTEXT("CollisionPawn", "Pawn"))
+/** Render LightShafts, for now SHOWFLAG_ALWAYS_ACCESSIBLE because it's exposed in SceneCapture */
 SHOWFLAG_ALWAYS_ACCESSIBLE(LightShafts, SFG_LightingFeatures, LOCTEXT("LightShaftsSF", "Light Shafts"))
 /** Render the PostProcess Material */
-SHOWFLAG_ALWAYS_ACCESSIBLE(PostProcessMaterial, SFG_PostProcess, LOCTEXT("PostProcessMaterialSF", "Post Process Material"))
-/** Render Atmospheric scattering (Atmospheric Fog) */
+SHOWFLAG_FIXED_IN_SHIPPING(1, PostProcessMaterial, SFG_PostProcess, LOCTEXT("PostProcessMaterialSF", "Post Process Material"))
+/** Render Atmospheric scattering (Atmospheric Fog), for now SHOWFLAG_ALWAYS_ACCESSIBLE because it's exposed in SceneCapture */
 SHOWFLAG_ALWAYS_ACCESSIBLE(AtmosphericFog, SFG_Advanced, LOCTEXT("AtmosphereSF", "Atmospheric Fog"))
 /** Render safe frames bars*/
-SHOWFLAG_ALWAYS_ACCESSIBLE(CameraAspectRatioBars, SFG_Advanced, LOCTEXT("CameraAspectRatioBarsSF", "Camera Aspect Ratio Bars"))
+SHOWFLAG_FIXED_IN_SHIPPING(0, CameraAspectRatioBars, SFG_Advanced, LOCTEXT("CameraAspectRatioBarsSF", "Camera Aspect Ratio Bars"))
 /** Render safe frames */
-SHOWFLAG_ALWAYS_ACCESSIBLE(CameraSafeFrames, SFG_Advanced, LOCTEXT("CameraSafeFramesSF", "Camera Safe Frames"))
-/** Render TextRenderComponents (3D text) */
+SHOWFLAG_FIXED_IN_SHIPPING(1, CameraSafeFrames, SFG_Advanced, LOCTEXT("CameraSafeFramesSF", "Camera Safe Frames"))
+/** Render TextRenderComponents (3D text), for now SHOWFLAG_ALWAYS_ACCESSIBLE because it's exposed in SceneCapture */
 SHOWFLAG_ALWAYS_ACCESSIBLE(TextRender, SFG_Advanced, LOCTEXT("TextRenderSF", "Render (3D) Text"))
 /** Any rendering/buffer clearing  (good for benchmarking) */
-SHOWFLAG_ALWAYS_ACCESSIBLE(Rendering, SFG_Hidden, LOCTEXT("RenderingSF", "Any Rendering"))
+SHOWFLAG_FIXED_IN_SHIPPING(1, Rendering, SFG_Hidden, LOCTEXT("RenderingSF", "Any Rendering"))
 /** Show the current mask being used by the highres screenshot capture */
-SHOWFLAG_ALWAYS_ACCESSIBLE(HighResScreenshotMask, SFG_Hidden, LOCTEXT("HighResScreenshotMaskSF", "High Res Screenshot Mask"))
-/** Distortion of output for HMD devices */
+SHOWFLAG_FIXED_IN_SHIPPING(0, HighResScreenshotMask, SFG_Hidden, LOCTEXT("HighResScreenshotMaskSF", "High Res Screenshot Mask"))
+/** Distortion of output for HMD devices, SHOWFLAG_ALWAYS_ACCESSIBLE for now because USceneCaptureComponent needs that */
 SHOWFLAG_ALWAYS_ACCESSIBLE(HMDDistortion, SFG_PostProcess, LOCTEXT("HMDDistortionSF", "HMD Distortion"))
-/** Whether to render in stereoscopic 3d */
+/** Whether to render in stereoscopic 3d, for now SHOWFLAG_ALWAYS_ACCESSIBLE because it's used by StereoRendering */
 SHOWFLAG_ALWAYS_ACCESSIBLE(StereoRendering, SFG_Hidden, LOCTEXT("StereoRenderingSF", "Stereoscopic Rendering"))
 /** Show objects even if they should be distance culled */
-SHOWFLAG_ALWAYS_ACCESSIBLE(DistanceCulledPrimitives, SFG_Hidden, LOCTEXT("DistanceCulledPrimitivesSF", "Distance Culled Primitives"))
+SHOWFLAG_FIXED_IN_SHIPPING(0, DistanceCulledPrimitives, SFG_Hidden, LOCTEXT("DistanceCulledPrimitivesSF", "Distance Culled Primitives"))
 /** To visualize the culling in Tile Based Deferred Lighting, later for non tiled as well */
-SHOWFLAG_ALWAYS_ACCESSIBLE(VisualizeLightCulling, SFG_Hidden, LOCTEXT("VisualizeLightCullingSF", "Light Culling"))
+SHOWFLAG_FIXED_IN_SHIPPING(0, VisualizeLightCulling, SFG_Hidden, LOCTEXT("VisualizeLightCullingSF", "Light Culling"))
 /** To disable precomputed visibility */
-SHOWFLAG_ALWAYS_ACCESSIBLE(PrecomputedVisibility, SFG_Advanced, LOCTEXT("PrecomputedVisibilitySF", "Precomputed Visibility"))
-/** Contribution from sky light */
+SHOWFLAG_FIXED_IN_SHIPPING(1, PrecomputedVisibility, SFG_Advanced, LOCTEXT("PrecomputedVisibilitySF", "Precomputed Visibility"))
+/** Contribution from sky light, for now SHOWFLAG_ALWAYS_ACCESSIBLE because it's exposed in SceneCapture */
 SHOWFLAG_ALWAYS_ACCESSIBLE(SkyLighting, SFG_LightTypes, LOCTEXT("SkyLightingSF", "Sky Lighting"))
 /** Visualize Light Propagation Volume, for developer (by default off): */
-SHOWFLAG_ALWAYS_ACCESSIBLE(VisualizeLPV, SFG_Visualize, LOCTEXT("VisualizeLPVSF", "Light Propagation Volume"))
+SHOWFLAG_FIXED_IN_SHIPPING(0, VisualizeLPV, SFG_Visualize, LOCTEXT("VisualizeLPVSF", "Light Propagation Volume"))
 /** Visualize preview shadow indicator */
-SHOWFLAG_ALWAYS_ACCESSIBLE(PreviewShadowsIndicator, SFG_Visualize, LOCTEXT("PreviewShadowIndicatorSF", "Preview Shadows Indicator"))
-/** Visualize precomputed visibilty cells */
-SHOWFLAG_ALWAYS_ACCESSIBLE(PrecomputedVisibilityCells, SFG_Visualize, LOCTEXT("PrecomputedVisibilityCellsSF", "Precomputed Visibility Cells"))
+SHOWFLAG_FIXED_IN_SHIPPING(0, PreviewShadowsIndicator, SFG_Visualize, LOCTEXT("PreviewShadowIndicatorSF", "Preview Shadows Indicator"))
+/** Visualize precomputed visibility cells */
+SHOWFLAG_FIXED_IN_SHIPPING(0, PrecomputedVisibilityCells, SFG_Visualize, LOCTEXT("PrecomputedVisibilityCellsSF", "Precomputed Visibility Cells"))
 /** Visualize volume lighting samples used for GI on dynamic objects */
-SHOWFLAG_ALWAYS_ACCESSIBLE(VolumeLightingSamples, SFG_Visualize, LOCTEXT("VolumeLightingSamplesSF", "Volume Lighting Samples"))
-/** Render Paper2D sprites */
+SHOWFLAG_FIXED_IN_SHIPPING(0, VolumeLightingSamples, SFG_Visualize, LOCTEXT("VolumeLightingSamplesSF", "Volume Lighting Samples"))
+/** Render Paper2D sprites, for now SHOWFLAG_ALWAYS_ACCESSIBLE because it's exposed in SceneCapture */
 SHOWFLAG_ALWAYS_ACCESSIBLE(Paper2DSprites, SFG_Advanced, LOCTEXT("Paper2DSpritesSF", "Paper 2D Sprites"))
 /** Visualization of distance field AO */
-SHOWFLAG_ALWAYS_ACCESSIBLE(VisualizeDistanceFieldAO, SFG_Visualize, LOCTEXT("VisualizeDistanceFieldAOSF", "Distance Field Ambient Occlusion"))
+SHOWFLAG_FIXED_IN_SHIPPING(0, VisualizeDistanceFieldAO, SFG_Visualize, LOCTEXT("VisualizeDistanceFieldAOSF", "Distance Field Ambient Occlusion"))
 /** Visualization of distance field GI */
-SHOWFLAG_ALWAYS_ACCESSIBLE(VisualizeDistanceFieldGI, SFG_Visualize, LOCTEXT("VisualizeDistanceFieldGISF", "Distance Field Global Illumination"))
+SHOWFLAG_FIXED_IN_SHIPPING(0, VisualizeDistanceFieldGI, SFG_Visualize, LOCTEXT("VisualizeDistanceFieldGISF", "Distance Field Global Illumination"))
 /** Mesh Distance fields */
-SHOWFLAG_ALWAYS_ACCESSIBLE(VisualizeMeshDistanceFields, SFG_Visualize, LOCTEXT("MeshDistanceFieldsSF", "Mesh DistanceFields"))
-/** Screen space AO */
+SHOWFLAG_FIXED_IN_SHIPPING(0, VisualizeMeshDistanceFields, SFG_Visualize, LOCTEXT("MeshDistanceFieldsSF", "Mesh DistanceFields"))
+/** Screen space AO, for now SHOWFLAG_ALWAYS_ACCESSIBLE because r.GBuffer need that */
 SHOWFLAG_ALWAYS_ACCESSIBLE(ScreenSpaceAO, SFG_LightingFeatures, LOCTEXT("ScreenSpaceAOSF", "Screen Space Ambient Occlusion"))
-/** Distance field AO */
+/** Distance field AO, for now SHOWFLAG_ALWAYS_ACCESSIBLE because it's exposed in SceneCapture */
 SHOWFLAG_ALWAYS_ACCESSIBLE(DistanceFieldAO, SFG_LightingFeatures, LOCTEXT("DistanceFieldAOSF", "Distance Field Ambient Occlusion"))
 /** Distance field GI */
-SHOWFLAG_ALWAYS_ACCESSIBLE(DistanceFieldGI, SFG_LightingFeatures, LOCTEXT("DistanceFieldGISF", "Distance Field Global Illumination"))
+SHOWFLAG_FIXED_IN_SHIPPING(1, DistanceFieldGI, SFG_LightingFeatures, LOCTEXT("DistanceFieldGISF", "Distance Field Global Illumination"))
 /** Visualize screen space reflections, for developer (by default off): */
-SHOWFLAG_FIXED_IN_SHIPPING(VisualizeSSR, 0, SFG_Visualize, LOCTEXT("VisualizeSSR", "Screen Space Reflections"))
-/** Force the use of the GBuffer. */
-SHOWFLAG_ALWAYS_ACCESSIBLE(ForceGBuffer, SFG_Hidden, LOCTEXT("ForceGBuffer", "Force usage of GBuffer"))
+SHOWFLAG_FIXED_IN_SHIPPING(0, VisualizeSSR, SFG_Visualize, LOCTEXT("VisualizeSSR", "Screen Space Reflections"))
 /** Visualize the senses configuration of AIs' PawnSensingComponent */
-SHOWFLAG_ALWAYS_ACCESSIBLE(VisualizeSenses, SFG_Advanced, LOCTEXT("VisualizeSenses", "Senses"))
+SHOWFLAG_FIXED_IN_SHIPPING(0, VisualizeSenses, SFG_Advanced, LOCTEXT("VisualizeSenses", "Senses"))
 /** Visualize the bloom, for developer (by default off): */
-SHOWFLAG_FIXED_IN_SHIPPING(VisualizeBloom, 0, SFG_Visualize, LOCTEXT("VisualizeBloom", "Bloom"))
+SHOWFLAG_FIXED_IN_SHIPPING(0, VisualizeBloom, SFG_Visualize, LOCTEXT("VisualizeBloom", "Bloom"))
 /** Visualize Vertex Density */
-SHOWFLAG_ALWAYS_ACCESSIBLE(VertexDensities, SFG_Hidden, LOCTEXT("VisualizeVertexDensities", "Visualize Vertex Densities"))
+SHOWFLAG_FIXED_IN_SHIPPING(0, VertexDensities, SFG_Hidden, LOCTEXT("VisualizeVertexDensities", "Visualize Vertex Densities"))
 /** Visualize LOD Coloration */
-SHOWFLAG_FIXED_IN_SHIPPING(LODColoration, 0, SFG_Hidden, LOCTEXT("VisualizeLODColoration", "Visualize LOD Coloration"))
+SHOWFLAG_FIXED_IN_SHIPPING(0, LODColoration, SFG_Hidden, LOCTEXT("VisualizeLODColoration", "Visualize LOD Coloration"))
 
 #undef SHOWFLAG_ALWAYS_ACCESSIBLE
 #undef SHOWFLAG_FIXED_IN_SHIPPING
diff --git a/Engine/Source/Runtime/Renderer/Private/DeferredShadingRenderer.cpp b/Engine/Source/Runtime/Renderer/Private/DeferredShadingRenderer.cpp
--- a/Engine/Source/Runtime/Renderer/Private/DeferredShadingRenderer.cpp
+++ b/Engine/Source/Runtime/Renderer/Private/DeferredShadingRenderer.cpp
@@ -869,10 +869,6 @@ void FDeferredShadingSceneRenderer::Render(FRHICommandListImmediate& RHICmdList)
 
 	static const auto GBufferCVar = IConsoleManager::Get().FindTConsoleVariableDataInt(TEXT("r.GBuffer"));
 	bool bGBuffer = GBufferCVar ? (GBufferCVar->GetValueOnRenderThread() != 0) : true;
-	if (ViewFamily.EngineShowFlags.ForceGBuffer)
-	{
-		bGBuffer = true;
-	}
 
 	if (ClearMethodCVar)
 	{
@@ -1284,7 +1280,7 @@ void FDeferredShadingSceneRenderer::Render(FRHICommandListImmediate& RHICmdList)
 		RenderTranslucency(RHICmdList);
 		ServiceLocalQueue();
 
-		if (ViewFamily.EngineShowFlags.Refraction)
+		if(GetRefractionQuality(ViewFamily) > 0)
 		{
 			// To apply refraction effect by distorting the scene color.
 			// After non separate translucency as that is considered at scene depth anyway
diff --git a/Engine/Source/Runtime/Renderer/Private/DistortionRendering.cpp b/Engine/Source/Runtime/Renderer/Private/DistortionRendering.cpp
--- a/Engine/Source/Runtime/Renderer/Private/DistortionRendering.cpp
+++ b/Engine/Source/Runtime/Renderer/Private/DistortionRendering.cpp
@@ -794,6 +794,20 @@ bool FDistortionPrimSet::DrawAccumulatedOffsets(FRHICommandListImmediate& RHICmd
 	return bDirty;
 }
 
+int32 FSceneRenderer::GetRefractionQuality(const FSceneViewFamily& ViewFamily)
+{
+	static const auto ICVar = IConsoleManager::Get().FindTConsoleVariableDataInt(TEXT("r.RefractionQuality"));
+	
+	int32 Value = 0;
+	
+	if(ViewFamily.EngineShowFlags.Refraction)
+	{
+		Value = ICVar->GetValueOnRenderThread();
+	}
+	
+	return Value;
+}
+
 /** 
  * Renders the scene's distortion 
  */
diff --git a/Engine/Source/Runtime/Renderer/Private/ForwardShadingRenderer.cpp b/Engine/Source/Runtime/Renderer/Private/ForwardShadingRenderer.cpp
--- a/Engine/Source/Runtime/Renderer/Private/ForwardShadingRenderer.cpp
+++ b/Engine/Source/Runtime/Renderer/Private/ForwardShadingRenderer.cpp
@@ -162,7 +162,7 @@ void FForwardShadingSceneRenderer::Render(FRHICommandListImmediate& RHICmdList)
 		// Having the distortion applied between two different translucency passes would make it consistent with the deferred pass.
 		// This is not done yet.
 
-		if (ViewFamily.EngineShowFlags.Refraction)
+		if(GetRefractionQuality(ViewFamily) > 0)
 		{
 			// to apply refraction effect by distorting the scene color
 			RenderDistortion(RHICmdList);
diff --git a/Engine/Source/Runtime/Renderer/Private/IndirectLightingCache.cpp b/Engine/Source/Runtime/Renderer/Private/IndirectLightingCache.cpp
--- a/Engine/Source/Runtime/Renderer/Private/IndirectLightingCache.cpp
+++ b/Engine/Source/Runtime/Renderer/Private/IndirectLightingCache.cpp
@@ -377,7 +377,7 @@ void FIndirectLightingCache::UpdateCache(FScene* Scene, FSceneRenderer& Renderer
 
 		for (int32 ViewIndex = 0; ViewIndex < Renderer.Views.Num(); ViewIndex++)
 		{
-			bAnyViewAllowsIndirectLightingCache |= Renderer.Views[ViewIndex].Family->EngineShowFlags.IndirectLightingCache;
+			bAnyViewAllowsIndirectLightingCache |= (bool)Renderer.Views[ViewIndex].Family->EngineShowFlags.IndirectLightingCache;
 		}
 
 		if (bAnyViewAllowsIndirectLightingCache)
diff --git a/Engine/Source/Runtime/Renderer/Private/ReflectionEnvironmentCapture.cpp b/Engine/Source/Runtime/Renderer/Private/ReflectionEnvironmentCapture.cpp
--- a/Engine/Source/Runtime/Renderer/Private/ReflectionEnvironmentCapture.cpp
+++ b/Engine/Source/Runtime/Renderer/Private/ReflectionEnvironmentCapture.cpp
@@ -1165,7 +1165,7 @@ void CaptureSceneIntoScratchCubemap(FScene* Scene, FVector CapturePosition, bool
 		// Exclude particles and light functions as they are usually dynamic, and can't be captured well
 		ViewFamily.EngineShowFlags.Particles = 0;
 		ViewFamily.EngineShowFlags.LightFunctions = 0;
-		ViewFamily.EngineShowFlags.CompositeEditorPrimitives = 0;
+		ViewFamily.EngineShowFlags.SetCompositeEditorPrimitives(false);
 		// These are highly dynamic and can't be captured effectively
 		ViewFamily.EngineShowFlags.LightShafts = 0;
 
diff --git a/Engine/Source/Runtime/Renderer/Private/SceneRendering.cpp b/Engine/Source/Runtime/Renderer/Private/SceneRendering.cpp
--- a/Engine/Source/Runtime/Renderer/Private/SceneRendering.cpp
+++ b/Engine/Source/Runtime/Renderer/Private/SceneRendering.cpp
@@ -54,6 +54,16 @@ static TAutoConsoleVariable<int32> CVarForwardLighting(
 	TEXT("1: on (some CPU setup cost on GPU and extra per pixel cost)"),
 	ECVF_RenderThreadSafe);
 
+static TAutoConsoleVariable<int32> CVarRefractionQuality(
+	TEXT("r.RefractionQuality"),
+	2,
+	TEXT("Defines the distorion/refraction quality which allows to adjust for quality or performance.\n")
+	TEXT("<=0: off (fastest)\n")
+	TEXT("  1: low quality (not yet implemented)\n")
+	TEXT("  2: normal quality (default)\n")
+	TEXT("  3: high quality (e.g. color fringe, not yet implemented)"),
+	ECVF_Scalability | ECVF_RenderThreadSafe);
+
 #if !(UE_BUILD_SHIPPING || UE_BUILD_TEST)
 static TAutoConsoleVariable<float> CVarGeneralPurposeTweak(
 	TEXT("r.GeneralPurposeTweak"),
@@ -1175,7 +1185,7 @@ FSceneRenderer::FSceneRenderer(const FSceneViewFamily* InViewFamily,FHitProxyCon
 	if(HitProxyConsumer)
 	{
 		// Set the hit proxies show flag.
-		ViewFamily.EngineShowFlags.HitProxies = 1;
+		ViewFamily.EngineShowFlags.SetHitProxies(1);
 	}
 
 	// launch custom visibility queries for views
diff --git a/Engine/Source/Runtime/Renderer/Private/SceneRendering.h b/Engine/Source/Runtime/Renderer/Private/SceneRendering.h
--- a/Engine/Source/Runtime/Renderer/Private/SceneRendering.h
+++ b/Engine/Source/Runtime/Renderer/Private/SceneRendering.h
@@ -856,7 +856,9 @@ protected:
 	void OnStartFrame();
 
 	/** Renders the scene's distortion */
-	void RenderDistortion(FRHICommandListImmediate& RHICmdList);	
+	void RenderDistortion(FRHICommandListImmediate& RHICmdList);
+
+	static int32 GetRefractionQuality(const FSceneViewFamily& ViewFamily);
 };
 
 
