comment

diff --git a/Engine/Source/Runtime/Engine/Public/ShowFlags.h b/Engine/Source/Runtime/Engine/Public/ShowFlags.h
--- a/Engine/Source/Runtime/Engine/Public/ShowFlags.h
+++ b/Engine/Source/Runtime/Engine/Public/ShowFlags.h
@@ -40,7 +40,8 @@ enum EShowFlagInitMode
 
 /**
  * ShowFlags are a set of bits (some are fixed in SHIPPING) that are stored in the ViewFamily.
- * They are use by the user to debug/profile parts of the rendering.
+ * They can be used by the artist/designer/developer to debug/profile parts of the rendering.
+ * Some ShowFlags are used to customize the rendering e.g. by the SceneCaptureActor, those should not be locked in shipping and not used in inner loops (for performance)
  * They should not be used for scalability (as some might be compiled out for SHIPPING), there we use console variables.
  * ViewModes are on a higher level, they can manipulate show flags before they get used.
  */
