Fix occlusion query flickering when updating HISMCs at runtime (UE-20457)

diff --git a/Engine/Source/Runtime/Engine/Private/HierarchicalInstancedStaticMesh.cpp b/Engine/Source/Runtime/Engine/Private/HierarchicalInstancedStaticMesh.cpp
--- a/Engine/Source/Runtime/Engine/Private/HierarchicalInstancedStaticMesh.cpp
+++ b/Engine/Source/Runtime/Engine/Private/HierarchicalInstancedStaticMesh.cpp
@@ -651,6 +651,7 @@ class FHierarchicalStaticMeshSceneProxy : public FInstancedStaticMeshSceneProxy
 	TArray<FBoxSphereBounds> OcclusionBounds;
 	TMap<uint32, FFoliageOcclusionResults> OcclusionResults;
 	bool bIsGrass;
+	uint32 SceneProxyCreatedFrameNumberRenderThread;
 
 #if !(UE_BUILD_SHIPPING || UE_BUILD_TEST)
 	mutable TArray<uint32> SingleDebugRuns[MAX_STATIC_MESH_LODS];
@@ -670,6 +671,7 @@ public:
 		, FirstUnbuiltIndex(InComponent->NumBuiltInstances)
 		, LastUnbuiltIndex(InComponent->PerInstanceSMData.Num()+InComponent->RemovedInstances.Num()-1)
 		, bIsGrass(bInIsGrass)
+		, SceneProxyCreatedFrameNumberRenderThread(UINT32_MAX)
 #if !(UE_BUILD_SHIPPING || UE_BUILD_TEST)
 		, CaptureTag(0)
 #endif
@@ -744,6 +746,12 @@ public:
 
 	// FPrimitiveSceneProxy interface.
 	
+	virtual void CreateRenderThreadResources() override
+	{
+		FInstancedStaticMeshSceneProxy::CreateRenderThreadResources();
+		SceneProxyCreatedFrameNumberRenderThread = GFrameNumberRenderThread;
+	}
+	
 	virtual FPrimitiveViewRelevance GetViewRelevance(const FSceneView* View) override
 	{
 		FPrimitiveViewRelevance Result;
@@ -1561,7 +1569,8 @@ void FHierarchicalStaticMeshSceneProxy::GetDynamicMeshElements(const TArray<cons
 
 void FHierarchicalStaticMeshSceneProxy::AcceptOcclusionResults(const FSceneView* View, const bool* Results, int32 NumResults)
 {
-	if (OcclusionBounds.Num() == NumResults)
+	// Don't accept subprimitive occlusion results from a previously-created sceneproxy - the tree may have been different
+	if (OcclusionBounds.Num() == NumResults && SceneProxyCreatedFrameNumberRenderThread < GFrameNumberRenderThread)
 	{
 		uint32 ViewId = View->GetViewKey();
 		FFoliageOcclusionResults* OldResults = OcclusionResults.Find(ViewId);
