GeneratedCodeHelpers.h includes BPGC

diff --git a/Engine/Source/Runtime/Engine/Public/GeneratedCodeHelpers.h b/Engine/Source/Runtime/Engine/Public/GeneratedCodeHelpers.h
--- a/Engine/Source/Runtime/Engine/Public/GeneratedCodeHelpers.h
+++ b/Engine/Source/Runtime/Engine/Public/GeneratedCodeHelpers.h
@@ -2,16 +2,20 @@
 
 #pragma once
 
+// Common includes
+#include "UObject/Stack.h"
+#include "Engine/BlueprintGeneratedClass.h"
+
 // Common libraries
 #include "Kismet/KismetArrayLibrary.h"
 #include "Kismet/KismetSystemLibrary.h"
 #include "Kismet/KismetMathLibrary.h"
 #include "Kismet/BlueprintFunctionLibrary.h"
 #include "Kismet/KismetMathLibrary.h"
+
 // Special libraries
 #include "Kismet/DataTableFunctionLibrary.h"
 
-#include "Stack.h"
 
 inline UClass* DynamicMetaCast(const UClass* DesiredClass, UClass* SourceClass)
 {
