ScreenPercentage showflag is no longer be fixed in shipping - it can change at runtime but it seems no good use

diff --git a/Engine/Plugins/Runtime/GearVR/Source/GearVR/Private/GearVR.cpp b/Engine/Plugins/Runtime/GearVR/Source/GearVR/Private/GearVR.cpp
--- a/Engine/Plugins/Runtime/GearVR/Source/GearVR/Private/GearVR.cpp
+++ b/Engine/Plugins/Runtime/GearVR/Source/GearVR/Private/GearVR.cpp
@@ -632,7 +632,7 @@ void FGearVR::SetupViewFamily(FSceneViewFamily& InViewFamily)
 {
 	InViewFamily.EngineShowFlags.MotionBlur = 0;
 	InViewFamily.EngineShowFlags.HMDDistortion = false;
-	InViewFamily.EngineShowFlags.SetScreenPercentage(false);
+	InViewFamily.EngineShowFlags.ScreenPercentage =false;
 	InViewFamily.EngineShowFlags.StereoRendering = IsStereoEnabled();
 }
 
diff --git a/Engine/Source/Runtime/Engine/Public/ShowFlagsValues.inl b/Engine/Source/Runtime/Engine/Public/ShowFlagsValues.inl
--- a/Engine/Source/Runtime/Engine/Public/ShowFlagsValues.inl
+++ b/Engine/Source/Runtime/Engine/Public/ShowFlagsValues.inl
@@ -92,8 +92,8 @@ SHOWFLAG_ALWAYS_ACCESSIBLE(CameraInterpolation, SFG_Hidden, LOCTEXT("CameraInter
 SHOWFLAG_FIXED_IN_SHIPPING(1, SceneColorFringe, SFG_PostProcess, LOCTEXT("SceneColorFringeSF", "Scene Color Fringe"))
 /** If Translucency should be rendered into a separate RT and composited without DepthOfField, can be disabled in the materials (affects sorting), SHOWFLAG_ALWAYS_ACCESSIBLE for now because USceneCaptureComponent needs that */
 SHOWFLAG_ALWAYS_ACCESSIBLE(SeparateTranslucency, SFG_Advanced, LOCTEXT("SeparateTranslucencySF", "Separate Translucency"))
-/** If Screen Percentage should be applied (upscaling), useful to disable it in editor */
-SHOWFLAG_FIXED_IN_SHIPPING(1, ScreenPercentage, SFG_PostProcess, LOCTEXT("ScreenPercentageSF", "Screen Percentage"))
+/** If Screen Percentage should be applied (upscaling), useful to disable it in editor, SHOWFLAG_ALWAYS_ACCESSIBLE for now because some VR code is using it  */
+SHOWFLAG_ALWAYS_ACCESSIBLE(ScreenPercentage, SFG_PostProcess, LOCTEXT("ScreenPercentageSF", "Screen Percentage"))
 /** Helper to tweak motion blur settings */
 SHOWFLAG_FIXED_IN_SHIPPING(0, VisualizeMotionBlur, SFG_Visualize, LOCTEXT("VisualizeMotionBlurSF", "Motion Blur"))
 /** Whether to display the Reflection Environment feature, which has local reflections from Reflection Capture actors, for now SHOWFLAG_ALWAYS_ACCESSIBLE because it's exposed in SceneCapture */
