#ue4 - Stop ticking particle systems on the dedicated server if they are not enabled, rather than bailing out of the tick each time. Clarified comment on bAutoManageAttachment to indicate it avoids initial attachment since we don't activate them.

--------

diff --git a/Engine/Source/Runtime/Engine/Classes/Particles/ParticleSystemComponent.h b/Engine/Source/Runtime/Engine/Classes/Particles/ParticleSystemComponent.h
--- a/Engine/Source/Runtime/Engine/Classes/Particles/ParticleSystemComponent.h
+++ b/Engine/Source/Runtime/Engine/Classes/Particles/ParticleSystemComponent.h
@@ -300,6 +300,7 @@ public:
 	 * True if we should automatically attach to AutoAttachParent when activated, and detach from our parent when completed.
 	 * This overrides any current attachment that may be present at the time of activation (deferring initial attachment until activation, if AutoAttachParent is null).
 	 * When enabled, detachment occurs regardless of whether AutoAttachParent is assigned, and the relative transform from the time of activation is restored.
+	 * This also disables attachment on dedicated servers, where we don't actually activate even if bAutoActivate is true.
 	 * @see AutoAttachParent, AutoAttachSocketName, AutoAttachLocationType
 	 */
 	UPROPERTY(EditAnywhere, BlueprintReadWrite, Category=Attachment)
diff --git a/Engine/Source/Runtime/Engine/Private/Particles/ParticleComponents.cpp b/Engine/Source/Runtime/Engine/Private/Particles/ParticleComponents.cpp
--- a/Engine/Source/Runtime/Engine/Private/Particles/ParticleComponents.cpp
+++ b/Engine/Source/Runtime/Engine/Private/Particles/ParticleComponents.cpp
@@ -3643,6 +3643,10 @@ void UParticleSystemComponent::TickComponent(float DeltaTime, enum ELevelTick Ti
 			// We need to destroy the component if the user is expecting us to do it automatically otherwise this component will live forever because HasCompleted() will never get checked
 			DestroyComponent();
 		}
+		else
+		{
+			SetComponentTickEnabled(false);
+		}
 		return;
 	}
 
