Made ARecastNavMesh::OnNavMeshGenerationFinished virtual #UE4

github PR#1508 by slonopotamus
UE-20548

diff --git a/Engine/Source/Runtime/Engine/Classes/AI/Navigation/RecastNavMesh.h b/Engine/Source/Runtime/Engine/Classes/AI/Navigation/RecastNavMesh.h
--- a/Engine/Source/Runtime/Engine/Classes/AI/Navigation/RecastNavMesh.h
+++ b/Engine/Source/Runtime/Engine/Classes/AI/Navigation/RecastNavMesh.h
@@ -784,7 +784,7 @@ public:
 	void InvalidateAffectedPaths(const TArray<uint32>& ChangedTiles);
 
 	/** Event from generator that navmesh build has finished */
-	void OnNavMeshGenerationFinished();
+	virtual void OnNavMeshGenerationFinished();
 
 	virtual void SetConfig(const FNavDataConfig& Src) override;
 protected:
