#Optimization: SetMin LOD function for SkinnedMeshComponent

Merging

//depot/UE4-Orion/Engine/Source/Runtime/Engine/...

to //depot/UE4/Engine/Source/Runtime/Engine/...

diff --git a/Engine/Source/Runtime/Engine/Classes/Components/SkinnedMeshComponent.h b/Engine/Source/Runtime/Engine/Classes/Components/SkinnedMeshComponent.h
--- a/Engine/Source/Runtime/Engine/Classes/Components/SkinnedMeshComponent.h
+++ b/Engine/Source/Runtime/Engine/Classes/Components/SkinnedMeshComponent.h
@@ -370,6 +370,14 @@ public:
 	UFUNCTION(BlueprintCallable, Category="Components|SkinnedMesh")
 	virtual void SetPhysicsAsset(class UPhysicsAsset* NewPhysicsAsset, bool bForceReInit = false);
 
+	/**
+	 * Set MinLodModel of the mesh component
+	 *
+	 * @param	InNewMinLODModel	Set new MinLodModel that make sure the LOD does not go below of this value
+	 */
+	UFUNCTION(BlueprintCallable, Category="Components|SkinnedMesh")
+	void SetMinLOD(int32 InNewMinLODModel);
+
 	/**
 	 * Find the index of bone by name. Looks in the current SkeletalMesh being used by this SkeletalMeshComponent.
 	 * 
diff --git a/Engine/Source/Runtime/Engine/Private/Components/SkinnedMeshComponent.cpp b/Engine/Source/Runtime/Engine/Private/Components/SkinnedMeshComponent.cpp
--- a/Engine/Source/Runtime/Engine/Private/Components/SkinnedMeshComponent.cpp
+++ b/Engine/Source/Runtime/Engine/Private/Components/SkinnedMeshComponent.cpp
@@ -1013,7 +1013,6 @@ void USkinnedMeshComponent::SetPhysicsAsset(class UPhysicsAsset* InPhysicsAsset,
 	PhysicsAssetOverride = InPhysicsAsset;
 }
 
-
 void USkinnedMeshComponent::SetMasterPoseComponent(class USkinnedMeshComponent* NewMasterBoneComponent)
 {
 	MasterPoseComponent = NewMasterBoneComponent;
@@ -1807,6 +1806,25 @@ void USkinnedMeshComponent::UnHideBoneByName( FName BoneName )
 	}
 }
 
+void USkinnedMeshComponent::SetMinLOD(int32 InNewMinLODModel)
+{
+	int32 MaxLODIndex = 0;
+	if(MeshObject)
+	{
+		MaxLODIndex = MeshObject->GetSkeletalMeshResource().LODModels.Num() - 1;
+	}
+
+	MinLodModel = FMath::Clamp(InNewMinLODModel, 0, MaxLODIndex);
+
+	// if current LOD is lower than min LOD, 
+	// update LOD status
+	// Otherwise, we're fine to stay this way
+	if (PredictedLODLevel < MinLodModel)
+	{
+		UpdateLODStatus();
+	}
+}
+
 bool USkinnedMeshComponent::UpdateLODStatus()
 {
 	// Predict the best (min) LOD level we are going to need. Basically we use the Min (best) LOD the renderer desired last frame.
