UE-18862 Improve comments for DrawDebugCone

diff --git a/Engine/Source/Runtime/Engine/Public/DrawDebugHelpers.h b/Engine/Source/Runtime/Engine/Public/DrawDebugHelpers.h
--- a/Engine/Source/Runtime/Engine/Public/DrawDebugHelpers.h
+++ b/Engine/Source/Runtime/Engine/Public/DrawDebugHelpers.h
@@ -33,6 +33,7 @@ ENGINE_API void DrawDebug2DDonut(const UWorld* InWorld, const FMatrix& Transform
 ENGINE_API void DrawDebugSphere(const UWorld* InWorld, FVector const& Center, float Radius, int32 Segments, FColor const& Color, bool bPersistentLines = false, float LifeTime=-1.f, uint8 DepthPriority = 0);
 /** Draw a debug cylinder */
 ENGINE_API void DrawDebugCylinder(const UWorld* InWorld, FVector const& Start, FVector const& End, float Radius, int32 Segments, FColor const& Color, bool bPersistentLines = false, float LifeTime=-1.f, uint8 DepthPriority = 0);
+/** Draw a debug cone. AngleWidth and AngleHeight are given in radians */
 ENGINE_API void DrawDebugCone(const UWorld* InWorld, FVector const& Origin, FVector const& Direction, float Length, float AngleWidth, float AngleHeight, int32 NumSides, FColor const& Color, bool bPersistentLines=false, float LifeTime=-1.f, uint8 DepthPriority = 0);
 /** Used by gameplay when defining a cone by a vertical and horizontal dot products. */
 ENGINE_API void DrawDebugAltCone(const UWorld* InWorld, FVector const& Origin, FRotator const& Rotation, float Length, float AngleWidth, float AngleHeight, FColor const& DrawColor, bool bPersistentLines=false, float LifeTime=-1.f, uint8 DepthPriority=0);
