Made anim sequence asset import data VisibleAnywhere

This is to fix a bug where the import data could be set to null, and resulted in poor UI on the details panel. The object is still entirely editable due to the 'instanced' flag.
This addresses UE-20379.

diff --git a/Engine/Source/Runtime/Engine/Classes/Animation/AnimSequence.h b/Engine/Source/Runtime/Engine/Classes/Animation/AnimSequence.h
--- a/Engine/Source/Runtime/Engine/Classes/Animation/AnimSequence.h
+++ b/Engine/Source/Runtime/Engine/Classes/Animation/AnimSequence.h
@@ -508,7 +508,7 @@ class ENGINE_API UAnimSequence : public UAnimSequenceBase
 	uint32 bWasCompressedWithoutTranslations:1;
 
 	/** Importing data and options used for this mesh */
-	UPROPERTY(EditAnywhere, Instanced, Category=ImportSettings)
+	UPROPERTY(VisibleAnywhere, Instanced, Category=ImportSettings)
 	class UAssetImportData* AssetImportData;
 
 	/***  for Reimport **/
