[UE-18722] Added some extra keywords to Delay and RetriggerableDelay in KismetSystemLibrary. (Contributed by [REDACTED_PERSON], GitHub pull request 1457)

diff --git a/Engine/Source/Runtime/Engine/Classes/Kismet/KismetSystemLibrary.h b/Engine/Source/Runtime/Engine/Classes/Kismet/KismetSystemLibrary.h
--- a/Engine/Source/Runtime/Engine/Classes/Kismet/KismetSystemLibrary.h
+++ b/Engine/Source/Runtime/Engine/Classes/Kismet/KismetSystemLibrary.h
@@ -277,7 +277,7 @@ class ENGINE_API UKismetSystemLibrary : public UBlueprintFunctionLibrary
 	 * @param Duration 		length of delay (in seconds).
 	 * @param LatentInfo 	The latent action.
 	 */
-	UFUNCTION(BlueprintCallable, Category="Utilities|FlowControl", meta=(Latent, WorldContext="WorldContextObject", LatentInfo="LatentInfo", Duration="0.2"))
+	UFUNCTION(BlueprintCallable, Category="Utilities|FlowControl", meta=(Latent, WorldContext="WorldContextObject", LatentInfo="LatentInfo", Duration="0.2", Keywords="sleep"))
 	static void	Delay(UObject* WorldContextObject, float Duration, struct FLatentActionInfo LatentInfo );
 
 	/** 
@@ -287,7 +287,7 @@ class ENGINE_API UKismetSystemLibrary : public UBlueprintFunctionLibrary
 	 * @param Duration 		length of delay (in seconds).
 	 * @param LatentInfo 	The latent action.
 	 */
-	UFUNCTION(BlueprintCallable, meta=(Latent, LatentInfo="LatentInfo", WorldContext="WorldContextObject", Duration="0.2"), Category="Utilities|FlowControl")
+	UFUNCTION(BlueprintCallable, meta=(Latent, LatentInfo="LatentInfo", WorldContext="WorldContextObject", Duration="0.2", Keywords="sleep"), Category="Utilities|FlowControl")
 	static void RetriggerableDelay(UObject* WorldContextObject, float Duration, FLatentActionInfo LatentInfo);
 
 	/** Interpolate a component to the specified relative location and rotation over the course of OverTime seconds. */
