[UE-18722] Added "set controller" keywords to Posses and UnPosses in controller.h (Contributed by [REDACTED_PERSON], GitHub pull request 1457)

diff --git a/Engine/Source/Runtime/Engine/Classes/GameFramework/Controller.h b/Engine/Source/Runtime/Engine/Classes/GameFramework/Controller.h
--- a/Engine/Source/Runtime/Engine/Classes/GameFramework/Controller.h
+++ b/Engine/Source/Runtime/Engine/Classes/GameFramework/Controller.h
@@ -223,13 +223,13 @@ public:
 	 * @param InPawn The Pawn to be possessed.
 	 * @see HasAuthority()
 	 */
-	UFUNCTION(BlueprintCallable, BlueprintAuthorityOnly, Category="Pawn")
+	UFUNCTION(BlueprintCallable, BlueprintAuthorityOnly, Category="Pawn", meta=(Keywords="set controller"))
 	virtual void Possess(APawn* InPawn);
 
 	/**
 	 * Called to unpossess our pawn for any reason that is not the pawn being destroyed (destruction handled by PawnDestroyed()).
 	 */
-	UFUNCTION(BlueprintCallable, Category="Pawn")
+	UFUNCTION(BlueprintCallable, Category="Pawn", meta=(Keywords="set controller"))
 	virtual void UnPossess();
 
 	/**
