accessor for UConnection off of UPackagemapClient. Useful for cases in custom ::NetSerialize functions.

--------
Integrated using branch UE4-To-UE4-Orion (reversed) of change#2675327 by [REDACTED_PERSON] on 2015/09/01 09:36:33.

diff --git a/Engine/Source/Runtime/Engine/Classes/Engine/PackageMapClient.h b/Engine/Source/Runtime/Engine/Classes/Engine/PackageMapClient.h
--- a/Engine/Source/Runtime/Engine/Classes/Engine/PackageMapClient.h
+++ b/Engine/Source/Runtime/Engine/Classes/Engine/PackageMapClient.h
@@ -156,6 +156,8 @@ public:
 
 	TArray< FNetworkGUID > & GetMustBeMappedGuidsInLastBunch() { return MustBeMappedGuidsInLastBunch; }
 
+	class UNetConnection* GetConnection() { return Connection; }
+
 protected:
 
 	bool	ExportNetGUID( FNetworkGUID NetGUID, const UObject* Object, FString PathName, UObject* ObjOuter );
