OnRep function for WorldGravityZ to fix clients sometimes not getting bWorldGravitySet set.

--------
Integrated using branch UT_To_UE4 of change#2674694 by [REDACTED_PERSON] on 2015/08/31 17:13:59.

diff --git a/Engine/Source/Runtime/Engine/Classes/GameFramework/WorldSettings.h b/Engine/Source/Runtime/Engine/Classes/GameFramework/WorldSettings.h
--- a/Engine/Source/Runtime/Engine/Classes/GameFramework/WorldSettings.h
+++ b/Engine/Source/Runtime/Engine/Classes/GameFramework/WorldSettings.h
@@ -296,9 +296,12 @@ class ENGINE_API AWorldSettings : public AInfo, public IInterface_AssetUserData
 	TSubclassOf<UDamageType> KillZDamageType;
 
 	// current gravity actually being used
-	UPROPERTY(transient, replicated)
+	UPROPERTY(transient, ReplicatedUsing=OnRep_WorldGravityZ)
 	float WorldGravityZ;
 
+	UFUNCTION()
+	virtual void OnRep_WorldGravityZ();
+
 	// optional level specific gravity override set by level designer
 	UPROPERTY(EditAnywhere, BlueprintReadOnly, Category=Physics, meta=(editcondition = "bGlobalGravitySet"))
 	float GlobalGravityZ;
diff --git a/Engine/Source/Runtime/Engine/Private/WorldSettings.cpp b/Engine/Source/Runtime/Engine/Private/WorldSettings.cpp
--- a/Engine/Source/Runtime/Engine/Private/WorldSettings.cpp
+++ b/Engine/Source/Runtime/Engine/Private/WorldSettings.cpp
@@ -122,6 +122,11 @@ float AWorldSettings::GetGravityZ() const
 	return WorldGravityZ;
 }
 
+void AWorldSettings::OnRep_WorldGravityZ()
+{
+	bWorldGravitySet = true;
+}
+
 float AWorldSettings::FixupDeltaSeconds(float DeltaSeconds, float RealDeltaSeconds)
 {
 	// Clamp time between 2000 fps and 2.5 fps.
