Added alternative way of gathering neighbours of navmesh poly #UE4

github PR#1511 by slonopotamus
with some minor improvements

UE-20586

diff --git a/Engine/Source/Runtime/Engine/Classes/AI/Navigation/RecastNavMesh.h b/Engine/Source/Runtime/Engine/Classes/AI/Navigation/RecastNavMesh.h
--- a/Engine/Source/Runtime/Engine/Classes/AI/Navigation/RecastNavMesh.h
+++ b/Engine/Source/Runtime/Engine/Classes/AI/Navigation/RecastNavMesh.h
@@ -889,6 +889,9 @@ public:
 	/** Finds all polys connected with specified one */
 	bool GetPolyNeighbors(NavNodeRef PolyID, TArray<FNavigationPortalEdge>& Neighbors) const;
 
+	/** Finds all polys connected with specified one, results expressed as array of NavNodeRefs */
+	bool GetPolyNeighbors(NavNodeRef PolyID, TArray<NavNodeRef>& Neighbors) const;
+
 	/** Finds closest point constrained to given poly */
 	bool GetClosestPointOnPoly(NavNodeRef PolyID, const FVector& TestPt, FVector& PointOnPoly) const;
 
diff --git a/Engine/Source/Runtime/Engine/Private/AI/Navigation/PImplRecastNavMesh.cpp b/Engine/Source/Runtime/Engine/Private/AI/Navigation/PImplRecastNavMesh.cpp
--- a/Engine/Source/Runtime/Engine/Private/AI/Navigation/PImplRecastNavMesh.cpp
+++ b/Engine/Source/Runtime/Engine/Private/AI/Navigation/PImplRecastNavMesh.cpp
@@ -1757,6 +1757,36 @@ bool FPImplRecastNavMesh::GetPolyNeighbors(NavNodeRef PolyID, TArray<FNavigation
 	return false;
 }
 
+bool FPImplRecastNavMesh::GetPolyNeighbors(NavNodeRef PolyID, TArray<NavNodeRef>& Neighbors) const
+{
+	if (DetourNavMesh)
+	{
+		const dtPolyRef PolyRef = static_cast<dtPolyRef>(PolyID);
+		dtPoly const* Poly = 0;
+		dtMeshTile const* Tile = 0;
+
+		const dtStatus Status = DetourNavMesh->getTileAndPolyByRef(PolyRef, &Tile, &Poly);
+
+		if (dtStatusSucceed(Status))
+		{
+			uint32 LinkIdx = Poly->firstLink;
+			Neighbors.Reserve(DT_VERTS_PER_POLYGON);
+
+			while (LinkIdx != DT_NULL_LINK)
+			{
+				const dtLink& Link = DetourNavMesh->getLink(Tile, LinkIdx);
+				LinkIdx = Link.next;
+
+				Neighbors.Add(Link.ref);
+			}
+
+			return true;
+		}
+	}
+
+	return false;
+}
+
 bool FPImplRecastNavMesh::GetPolyTileIndex(NavNodeRef PolyID, uint32& PolyIndex, uint32& TileIndex) const
 {
 	if (DetourNavMesh && PolyID)
diff --git a/Engine/Source/Runtime/Engine/Private/AI/Navigation/RecastNavMesh.cpp b/Engine/Source/Runtime/Engine/Private/AI/Navigation/RecastNavMesh.cpp
--- a/Engine/Source/Runtime/Engine/Private/AI/Navigation/RecastNavMesh.cpp
+++ b/Engine/Source/Runtime/Engine/Private/AI/Navigation/RecastNavMesh.cpp
@@ -794,22 +794,12 @@ FBox ARecastNavMesh::GetNavMeshTileBounds(int32 TileIndex) const
 
 bool ARecastNavMesh::GetNavMeshTileXY(int32 TileIndex, int32& OutX, int32& OutY, int32& OutLayer) const
 {
-	if (RecastNavMeshImpl)
-	{
-		return RecastNavMeshImpl->GetNavMeshTileXY(TileIndex, OutX, OutY, OutLayer);
-	}
-
-	return false;
+	return RecastNavMeshImpl && RecastNavMeshImpl->GetNavMeshTileXY(TileIndex, OutX, OutY, OutLayer);
 }
 
 bool ARecastNavMesh::GetNavMeshTileXY(const FVector& Point, int32& OutX, int32& OutY) const
 {
-	if (RecastNavMeshImpl)
-	{
-		return RecastNavMeshImpl->GetNavMeshTileXY(Point, OutX, OutY);
-	}
-
-	return false;
+	return RecastNavMeshImpl && RecastNavMeshImpl->GetNavMeshTileXY(Point, OutX, OutY);
 }
 
 void ARecastNavMesh::GetNavMeshTilesAt(int32 TileX, int32 TileY, TArray<int32>& Indices) const
@@ -822,13 +812,7 @@ void ARecastNavMesh::GetNavMeshTilesAt(int32 TileX, int32 TileY, TArray<int32>&
 
 bool ARecastNavMesh::GetPolysInTile(int32 TileIndex, TArray<FNavPoly>& Polys) const
 {
-	bool bSuccess = false;
-	if (RecastNavMeshImpl)
-	{
-		bSuccess = RecastNavMeshImpl->GetPolysInTile(TileIndex, Polys);
-	}
-
-	return bSuccess;
+	return RecastNavMeshImpl && RecastNavMeshImpl->GetPolysInTile(TileIndex, Polys);
 }
 
 int32 ARecastNavMesh::GetNavMeshTilesCount() const
@@ -963,13 +947,7 @@ bool ARecastNavMesh::GetRandomPointInNavigableRadius(const FVector& Origin, floa
 
 bool ARecastNavMesh::GetRandomPointInCluster(NavNodeRef ClusterRef, FNavLocation& OutLocation) const
 {
-	bool bSuccess = false;
-	if (RecastNavMeshImpl)
-	{
-		bSuccess = RecastNavMeshImpl->GetRandomPointInCluster(ClusterRef, OutLocation);
-	}
-
-	return bSuccess;
+	return RecastNavMeshImpl && RecastNavMeshImpl->GetRandomPointInCluster(ClusterRef, OutLocation);
 }
 
 NavNodeRef ARecastNavMesh::GetClusterRef(NavNodeRef PolyRef) const
@@ -985,13 +963,7 @@ NavNodeRef ARecastNavMesh::GetClusterRef(NavNodeRef PolyRef) const
 
 bool ARecastNavMesh::ProjectPoint(const FVector& Point, FNavLocation& OutLocation, const FVector& Extent, TSharedPtr<const FNavigationQueryFilter> Filter, const UObject* QueryOwner) const
 {
-	bool bSuccess = false;
-	if (RecastNavMeshImpl)
-	{
-		bSuccess = RecastNavMeshImpl->ProjectPointToNavMesh(Point, OutLocation, Extent, GetRightFilterRef(Filter), QueryOwner);
-	}
-
-	return bSuccess;
+	return RecastNavMeshImpl && RecastNavMeshImpl->ProjectPointToNavMesh(Point, OutLocation, Extent, GetRightFilterRef(Filter), QueryOwner);
 }
 
 void ARecastNavMesh::BatchProjectPoints(TArray<FNavigationProjectionWork>& Workload, const FVector& Extent, TSharedPtr<const FNavigationQueryFilter> Filter, const UObject* Querier) const 
@@ -1098,13 +1070,7 @@ bool ARecastNavMesh::GetPolysInBox(const FBox& Box, TArray<FNavPoly>& Polys, TSh
 bool ARecastNavMesh::ProjectPointMulti(const FVector& Point, TArray<FNavLocation>& OutLocations, const FVector& Extent,
 	float MinZ, float MaxZ, TSharedPtr<const FNavigationQueryFilter> Filter, const UObject* QueryOwner) const
 {
-	bool bSuccess = false;
-	if (RecastNavMeshImpl)
-	{
-		bSuccess = RecastNavMeshImpl->ProjectPointMulti(Point, OutLocations, Extent, MinZ, MaxZ, GetRightFilterRef(Filter), QueryOwner);
-	}
-
-	return bSuccess;
+	return RecastNavMeshImpl && RecastNavMeshImpl->ProjectPointMulti(Point, OutLocations, Extent, MinZ, MaxZ, GetRightFilterRef(Filter), QueryOwner);
 }
 
 ENavigationQueryResult::Type ARecastNavMesh::CalcPathCost(const FVector& PathStart, const FVector& PathEnd, float& OutPathCost, TSharedPtr<const FNavigationQueryFilter> QueryFilter, const UObject* QueryOwner) const
@@ -1260,24 +1226,12 @@ void ARecastNavMesh::UpdateSegmentLinkArea(int32 UserId, TSubclassOf<UNavArea> A
 
 bool ARecastNavMesh::GetPolyCenter(NavNodeRef PolyID, FVector& OutCenter) const
 {
-	bool bSuccess = false;
-	if (RecastNavMeshImpl)
-	{
-		bSuccess = RecastNavMeshImpl->GetPolyCenter(PolyID, OutCenter);
-	}
-
-	return bSuccess;
+	return RecastNavMeshImpl && RecastNavMeshImpl->GetPolyCenter(PolyID, OutCenter);
 }
 
 bool ARecastNavMesh::GetPolyVerts(NavNodeRef PolyID, TArray<FVector>& OutVerts) const
 {
-	bool bSuccess = false;
-	if (RecastNavMeshImpl)
-	{
-		bSuccess = RecastNavMeshImpl->GetPolyVerts(PolyID, OutVerts);
-	}
-
-	return bSuccess;
+	return RecastNavMeshImpl && RecastNavMeshImpl->GetPolyVerts(PolyID, OutVerts);
 }
 
 uint32 ARecastNavMesh::GetPolyAreaID(NavNodeRef PolyID) const
@@ -1348,80 +1302,43 @@ bool ARecastNavMesh::GetPolyFlags(NavNodeRef PolyID, FNavMeshNodeFlags& Flags) c
 
 bool ARecastNavMesh::GetPolyNeighbors(NavNodeRef PolyID, TArray<FNavigationPortalEdge>& Neighbors) const
 {
-	bool bFound = false;
-	if (RecastNavMeshImpl)
-	{
-		bFound = RecastNavMeshImpl->GetPolyNeighbors(PolyID, Neighbors);
-	}
+	return RecastNavMeshImpl && RecastNavMeshImpl->GetPolyNeighbors(PolyID, Neighbors);
+}
 
-	return bFound;
+bool ARecastNavMesh::GetPolyNeighbors(NavNodeRef PolyID, TArray<NavNodeRef>& Neighbors) const
+{
+	return RecastNavMeshImpl && RecastNavMeshImpl->GetPolyNeighbors(PolyID, Neighbors);
 }
 
 bool ARecastNavMesh::GetClosestPointOnPoly(NavNodeRef PolyID, const FVector& TestPt, FVector& PointOnPoly) const
 {
-	bool bFound = false;
-	if (RecastNavMeshImpl)
-	{
-		bFound = RecastNavMeshImpl->GetClosestPointOnPoly(PolyID, TestPt, PointOnPoly);
-	}
-
-	return bFound;
+	return RecastNavMeshImpl && RecastNavMeshImpl->GetClosestPointOnPoly(PolyID, TestPt, PointOnPoly);
 }
 
 bool ARecastNavMesh::GetPolyTileIndex(NavNodeRef PolyID, uint32& PolyIndex, uint32& TileIndex) const
 {
-	bool bFound = false;
-	if (RecastNavMeshImpl)
-	{
-		bFound = RecastNavMeshImpl->GetPolyTileIndex(PolyID, PolyIndex, TileIndex);
-	}
-
-	return bFound;
+	return RecastNavMeshImpl && RecastNavMeshImpl->GetPolyTileIndex(PolyID, PolyIndex, TileIndex);
 }
 
 bool ARecastNavMesh::GetLinkEndPoints(NavNodeRef LinkPolyID, FVector& PointA, FVector& PointB) const
 {
-	bool bSuccess = false;
-	if (RecastNavMeshImpl)
-	{
-		bSuccess = RecastNavMeshImpl->GetLinkEndPoints(LinkPolyID, PointA, PointB);
-	}
-
-	return bSuccess;
+	return RecastNavMeshImpl && RecastNavMeshImpl->GetLinkEndPoints(LinkPolyID, PointA, PointB);
 }
 
 bool ARecastNavMesh::IsCustomLink(NavNodeRef LinkPolyID) const
 {
-	bool bSuccess = false;
-	if (RecastNavMeshImpl)
-	{
-		bSuccess = RecastNavMeshImpl->IsCustomLink(LinkPolyID);
-	}
-
-	return bSuccess;
+	return RecastNavMeshImpl && RecastNavMeshImpl->IsCustomLink(LinkPolyID);
 }
 
 bool ARecastNavMesh::GetClusterBounds(NavNodeRef ClusterRef, FBox& OutBounds) const
 {
-	bool bSuccess = false;
-	if (RecastNavMeshImpl)
-	{
-		bSuccess = RecastNavMeshImpl->GetClusterBounds(ClusterRef, OutBounds);
-	}
-
-	return bSuccess;
+	return RecastNavMeshImpl && RecastNavMeshImpl->GetClusterBounds(ClusterRef, OutBounds);
 }
 
 bool ARecastNavMesh::GetPolysWithinPathingDistance(FVector const& StartLoc, const float PathingDistance, TArray<NavNodeRef>& FoundPolys,
 	TSharedPtr<const FNavigationQueryFilter> Filter, const UObject* QueryOwner, FRecastDebugPathfindingData* DebugData) const
 {
-	bool bSuccess = false;
-	if (RecastNavMeshImpl)
-	{
-		bSuccess = RecastNavMeshImpl->GetPolysWithinPathingDistance(StartLoc, PathingDistance, GetRightFilterRef(Filter), QueryOwner, FoundPolys, DebugData);
-	}
-
-	return bSuccess;
+	return RecastNavMeshImpl && RecastNavMeshImpl->GetPolysWithinPathingDistance(StartLoc, PathingDistance, GetRightFilterRef(Filter), QueryOwner, FoundPolys, DebugData);
 }
 
 void ARecastNavMesh::GetDebugGeometry(FRecastDebugGeometry& OutGeometry, int32 TileIndex) const
diff --git a/Engine/Source/Runtime/Engine/Public/AI/Navigation/PImplRecastNavMesh.h b/Engine/Source/Runtime/Engine/Public/AI/Navigation/PImplRecastNavMesh.h
--- a/Engine/Source/Runtime/Engine/Public/AI/Navigation/PImplRecastNavMesh.h
+++ b/Engine/Source/Runtime/Engine/Public/AI/Navigation/PImplRecastNavMesh.h
@@ -177,6 +177,8 @@ public:
 	void SetPolyAreaID(NavNodeRef PolyID, uint8 AreaID);
 	/** Finds all polys connected with specified one */
 	bool GetPolyNeighbors(NavNodeRef PolyID, TArray<FNavigationPortalEdge>& Neighbors) const;
+	/** Finds all polys connected with specified one, results expressed as array of NavNodeRefs */
+	bool GetPolyNeighbors(NavNodeRef PolyID, TArray<NavNodeRef>& Neighbors) const;
 	/** Finds closest point constrained to given poly */
 	bool GetClosestPointOnPoly(NavNodeRef PolyID, const FVector& TestPt, FVector& PointOnPoly) const;
 	/** Decode poly ID into tile index and poly index */
