Fix for the 3d spread being ignored in SoundCue  nodes.

Also fixed the same for SpatializationAlgorithm.

diff --git a/Engine/Source/Runtime/Engine/Private/SoundNodeAttenuation.cpp b/Engine/Source/Runtime/Engine/Private/SoundNodeAttenuation.cpp
--- a/Engine/Source/Runtime/Engine/Private/SoundNodeAttenuation.cpp
+++ b/Engine/Source/Runtime/Engine/Private/SoundNodeAttenuation.cpp
@@ -56,6 +56,8 @@ void USoundNodeAttenuation::ParseNodes( FAudioDevice* AudioDevice, const UPTRINT
 		const FListener& Listener = AudioDevice->Listeners[0];
 		Settings->ApplyAttenuation(UpdatedParseParams.Transform, Listener.Transform.GetTranslation(), UpdatedParseParams.Volume, UpdatedParseParams.HighFrequencyGain);
 		UpdatedParseParams.OmniRadius = Settings->OmniRadius;
+		UpdatedParseParams.StereoSpread = Settings->StereoSpread;
+		UpdatedParseParams.SpatializationAlgorithm = Settings->SpatializationAlgorithm;
 		UpdatedParseParams.bUseSpatialization |= Settings->bSpatialize;
 	}
 	else
