Merging

//depot/UE4-Orion/Engine/Source/Runtime/Engine/...

to //depot/UE4-Orion/Engine/Source/Runtime/Engine/...

diff --git a/Engine/Source/Runtime/Engine/Classes/Components/SkinnedMeshComponent.h b/Engine/Source/Runtime/Engine/Classes/Components/SkinnedMeshComponent.h
--- a/Engine/Source/Runtime/Engine/Classes/Components/SkinnedMeshComponent.h
+++ b/Engine/Source/Runtime/Engine/Classes/Components/SkinnedMeshComponent.h
@@ -373,10 +373,18 @@ public:
 	/**
 	 * Set MinLodModel of the mesh component
 	 *
-	 * @param	InNewMinLODModel	Set new MinLodModel that make sure the LOD does not go below of this value
+	 * @param	InNewMinLOD	Set new MinLodModel that make sure the LOD does not go below of this value. This will affect in the next tick update. 
 	 */
 	UFUNCTION(BlueprintCallable, Category="Components|SkinnedMesh")
-	void SetMinLOD(int32 InNewMinLODModel);
+	void SetMinLOD(int32 InNewMinLOD);
+
+	/**
+	 * Set MinLodModel of the mesh component
+	 *
+	 * @param	InNewForcedLOD	Set new ForcedLODModel that forces to set the incoming LOD. This will affect in the next tick update. 
+	 */
+	UFUNCTION(BlueprintCallable, Category="Components|SkinnedMesh")
+	void SetForcedLOD(int32 InNewForcedLOD);
 
 	/**
 	 * Find the index of bone by name. Looks in the current SkeletalMesh being used by this SkeletalMeshComponent.
diff --git a/Engine/Source/Runtime/Engine/Private/Components/SkinnedMeshComponent.cpp b/Engine/Source/Runtime/Engine/Private/Components/SkinnedMeshComponent.cpp
--- a/Engine/Source/Runtime/Engine/Private/Components/SkinnedMeshComponent.cpp
+++ b/Engine/Source/Runtime/Engine/Private/Components/SkinnedMeshComponent.cpp
@@ -1806,7 +1806,7 @@ void USkinnedMeshComponent::UnHideBoneByName( FName BoneName )
 	}
 }
 
-void USkinnedMeshComponent::SetMinLOD(int32 InNewMinLODModel)
+void USkinnedMeshComponent::SetForcedLOD(int32 InNewForcedLOD)
 {
 	int32 MaxLODIndex = 0;
 	if(MeshObject)
@@ -1814,15 +1814,18 @@ void USkinnedMeshComponent::SetMinLOD(int32 InNewMinLODModel)
 		MaxLODIndex = MeshObject->GetSkeletalMeshResource().LODModels.Num() - 1;
 	}
 
-	MinLodModel = FMath::Clamp(InNewMinLODModel, 0, MaxLODIndex);
+	ForcedLodModel = FMath::Clamp(InNewForcedLOD, 0, MaxLODIndex);
+}
 
-	// if current LOD is lower than min LOD, 
-	// update LOD status
-	// Otherwise, we're fine to stay this way
-	if (PredictedLODLevel < MinLodModel)
+void USkinnedMeshComponent::SetMinLOD(int32 InNewMinLOD)
+{
+	int32 MaxLODIndex = 0;
+	if(MeshObject)
 	{
-		UpdateLODStatus();
+		MaxLODIndex = MeshObject->GetSkeletalMeshResource().LODModels.Num() - 1;
 	}
+
+	MinLodModel = FMath::Clamp(InNewMinLOD, 0, MaxLODIndex);
 }
 
 bool USkinnedMeshComponent::UpdateLODStatus()
