Changes FPrimitiveSceneProxy::GetViewRelevance(const FSceneView* View) to FPrimitiveSceneProxy::GetViewRelevance(const FSceneView* View) const

diff --git a/Engine/Plugins/2D/Paper2D/Source/Paper2D/Private/PaperRenderSceneProxy.cpp b/Engine/Plugins/2D/Paper2D/Source/Paper2D/Private/PaperRenderSceneProxy.cpp
--- a/Engine/Plugins/2D/Paper2D/Source/Paper2D/Private/PaperRenderSceneProxy.cpp
+++ b/Engine/Plugins/2D/Paper2D/Source/Paper2D/Private/PaperRenderSceneProxy.cpp
@@ -566,7 +566,7 @@ void FPaperRenderSceneProxy::GetBatchMesh(const FSceneView* View, class UMateria
 	}
 }
 
-FPrimitiveViewRelevance FPaperRenderSceneProxy::GetViewRelevance(const FSceneView* View)
+FPrimitiveViewRelevance FPaperRenderSceneProxy::GetViewRelevance(const FSceneView* View) const
 {
 	const FEngineShowFlags& EngineShowFlags = View->Family->EngineShowFlags;
 
diff --git a/Engine/Plugins/2D/Paper2D/Source/Paper2D/Private/PaperRenderSceneProxy.h b/Engine/Plugins/2D/Paper2D/Source/Paper2D/Private/PaperRenderSceneProxy.h
--- a/Engine/Plugins/2D/Paper2D/Source/Paper2D/Private/PaperRenderSceneProxy.h
+++ b/Engine/Plugins/2D/Paper2D/Source/Paper2D/Private/PaperRenderSceneProxy.h
@@ -145,7 +145,7 @@ public:
 
 	// FPrimitiveSceneProxy interface.
 	virtual void GetDynamicMeshElements(const TArray<const FSceneView*>& Views, const FSceneViewFamily& ViewFamily, uint32 VisibilityMap, FMeshElementCollector& Collector) const override;
-	virtual FPrimitiveViewRelevance GetViewRelevance(const FSceneView* View) override;
+	virtual FPrimitiveViewRelevance GetViewRelevance(const FSceneView* View) const override;
 	virtual uint32 GetMemoryFootprint() const override;
 	virtual bool CanBeOccluded() const override;
 	// End of FPrimitiveSceneProxy interface.
diff --git a/Engine/Plugins/Experimental/StereoPanorama/Source/StereoPanorama/Private/StereoStaticMeshComponent.cpp b/Engine/Plugins/Experimental/StereoPanorama/Source/StereoPanorama/Private/StereoStaticMeshComponent.cpp
--- a/Engine/Plugins/Experimental/StereoPanorama/Source/StereoPanorama/Private/StereoStaticMeshComponent.cpp
+++ b/Engine/Plugins/Experimental/StereoPanorama/Source/StereoPanorama/Private/StereoStaticMeshComponent.cpp
@@ -16,7 +16,7 @@ public:
         EyeToRender = Component->EyeToRender;
     }
 
-    virtual FPrimitiveViewRelevance GetViewRelevance(const FSceneView* View) override
+    virtual FPrimitiveViewRelevance GetViewRelevance(const FSceneView* View) const override
     {
         FPrimitiveViewRelevance viewRelevance = FStaticMeshSceneProxy::GetViewRelevance(View);
         bool bVisible = true;
diff --git a/Engine/Plugins/Runtime/CableComponent/Source/CableComponent/Private/CableComponent.cpp b/Engine/Plugins/Runtime/CableComponent/Source/CableComponent/Private/CableComponent.cpp
--- a/Engine/Plugins/Runtime/CableComponent/Source/CableComponent/Private/CableComponent.cpp
+++ b/Engine/Plugins/Runtime/CableComponent/Source/CableComponent/Private/CableComponent.cpp
@@ -299,7 +299,7 @@ public:
 		}
 	}
 
-	virtual FPrimitiveViewRelevance GetViewRelevance(const FSceneView* View) override
+	virtual FPrimitiveViewRelevance GetViewRelevance(const FSceneView* View) const override
 	{
 		FPrimitiveViewRelevance Result;
 		Result.bDrawRelevance = IsShown(View);
diff --git a/Engine/Plugins/Runtime/CustomMeshComponent/Source/CustomMeshComponent/Private/CustomMeshComponent.cpp b/Engine/Plugins/Runtime/CustomMeshComponent/Source/CustomMeshComponent/Private/CustomMeshComponent.cpp
--- a/Engine/Plugins/Runtime/CustomMeshComponent/Source/CustomMeshComponent/Private/CustomMeshComponent.cpp
+++ b/Engine/Plugins/Runtime/CustomMeshComponent/Source/CustomMeshComponent/Private/CustomMeshComponent.cpp
@@ -207,7 +207,7 @@ public:
 		}
 	}
 
-	virtual FPrimitiveViewRelevance GetViewRelevance(const FSceneView* View) override
+	virtual FPrimitiveViewRelevance GetViewRelevance(const FSceneView* View) const override
 	{
 		FPrimitiveViewRelevance Result;
 		Result.bDrawRelevance = IsShown(View);
diff --git a/Engine/Source/Developer/FunctionalTesting/Private/FuncTestRenderingComponent.cpp b/Engine/Source/Developer/FunctionalTesting/Private/FuncTestRenderingComponent.cpp
--- a/Engine/Source/Developer/FunctionalTesting/Private/FuncTestRenderingComponent.cpp
+++ b/Engine/Source/Developer/FunctionalTesting/Private/FuncTestRenderingComponent.cpp
@@ -31,7 +31,7 @@ public:
 		}
 	}
 
-	virtual FPrimitiveViewRelevance GetViewRelevance(const FSceneView* View) override
+	virtual FPrimitiveViewRelevance GetViewRelevance(const FSceneView* View) const override
 	{
 		FPrimitiveViewRelevance Result;
 		Result.bDrawRelevance = IsShown(View) && IsSelected();
diff --git a/Engine/Source/Developer/GameplayDebugger/Private/GameplayDebuggingComponent.cpp b/Engine/Source/Developer/GameplayDebugger/Private/GameplayDebuggingComponent.cpp
--- a/Engine/Source/Developer/GameplayDebugger/Private/GameplayDebuggingComponent.cpp
+++ b/Engine/Source/Developer/GameplayDebugger/Private/GameplayDebuggingComponent.cpp
@@ -82,7 +82,7 @@ public:
 		}
 	}
 
-	virtual FPrimitiveViewRelevance GetViewRelevance(const FSceneView* View) override
+	virtual FPrimitiveViewRelevance GetViewRelevance(const FSceneView* View) const override
 	{
 		FPrimitiveViewRelevance Result;
 		for (int32 Index = 0; Index < ChildProxies.Num(); ++Index)
@@ -1389,7 +1389,7 @@ public:
 	}
 
 
-	FPrimitiveViewRelevance GetViewRelevance(const FSceneView* View) override
+	FPrimitiveViewRelevance GetViewRelevance(const FSceneView* View) const override
 	{
 		FPrimitiveViewRelevance Result;
 		Result.bDrawRelevance = View->Family->EngineShowFlags.GetSingleFlag(ViewFlagIndex);// IsShown(View);
diff --git a/Engine/Source/Developer/LogVisualizer/Private/VisualLoggerRenderingActor.cpp b/Engine/Source/Developer/LogVisualizer/Private/VisualLoggerRenderingActor.cpp
--- a/Engine/Source/Developer/LogVisualizer/Private/VisualLoggerRenderingActor.cpp
+++ b/Engine/Source/Developer/LogVisualizer/Private/VisualLoggerRenderingActor.cpp
@@ -26,7 +26,7 @@ public:
 		ViewFlagIndex = uint32(FEngineShowFlags::FindIndexByName(*ViewFlagName));
 	}
 
-	virtual FPrimitiveViewRelevance GetViewRelevance(const FSceneView* View) override
+	virtual FPrimitiveViewRelevance GetViewRelevance(const FSceneView* View) const override
 	{
 		FPrimitiveViewRelevance Result;
 		Result.bDrawRelevance = IsShown(View);
diff --git a/Engine/Source/Runtime/AIModule/Classes/EnvironmentQuery/EQSRenderingComponent.h b/Engine/Source/Runtime/AIModule/Classes/EnvironmentQuery/EQSRenderingComponent.h
--- a/Engine/Source/Runtime/AIModule/Classes/EnvironmentQuery/EQSRenderingComponent.h
+++ b/Engine/Source/Runtime/AIModule/Classes/EnvironmentQuery/EQSRenderingComponent.h
@@ -18,7 +18,7 @@ public:
 
 	virtual void DrawDebugLabels(UCanvas* Canvas, APlayerController*) override;
 	
-	virtual FPrimitiveViewRelevance GetViewRelevance(const FSceneView* View) override;
+	virtual FPrimitiveViewRelevance GetViewRelevance(const FSceneView* View) const override;
 
 #if  USE_EQS_DEBUGGER 
 	static void CollectEQSData(const UPrimitiveComponent* InComponent, const IEQSQueryResultSourceInterface* QueryDataSource, TArray<FSphere>& Spheres, TArray<FText3d>& Texts, TArray<EQSDebug::FDebugHelper>& DebugItems);
diff --git a/Engine/Source/Runtime/AIModule/Private/EnvironmentQuery/EQSRenderingComponent.cpp b/Engine/Source/Runtime/AIModule/Private/EnvironmentQuery/EQSRenderingComponent.cpp
--- a/Engine/Source/Runtime/AIModule/Private/EnvironmentQuery/EQSRenderingComponent.cpp
+++ b/Engine/Source/Runtime/AIModule/Private/EnvironmentQuery/EQSRenderingComponent.cpp
@@ -292,7 +292,7 @@ bool FEQSSceneProxy::SafeIsActorSelected() const
 	return false;
 }
 
-FPrimitiveViewRelevance FEQSSceneProxy::GetViewRelevance(const FSceneView* View)
+FPrimitiveViewRelevance FEQSSceneProxy::GetViewRelevance(const FSceneView* View) const
 {
 	FPrimitiveViewRelevance Result;
 	Result.bDrawRelevance = View->Family->EngineShowFlags.GetSingleFlag(ViewFlagIndex) && IsShown(View) 
diff --git a/Engine/Source/Runtime/Engine/Classes/Engine/NiagaraEffectRenderer.h b/Engine/Source/Runtime/Engine/Classes/Engine/NiagaraEffectRenderer.h
--- a/Engine/Source/Runtime/Engine/Classes/Engine/NiagaraEffectRenderer.h
+++ b/Engine/Source/Runtime/Engine/Classes/Engine/NiagaraEffectRenderer.h
@@ -70,7 +70,7 @@ public:
 
 	virtual const TArray<FNiagaraVariableInfo>& GetRequiredAttributes() = 0;
 
-	virtual FPrimitiveViewRelevance GetViewRelevance(const FSceneView* View, FNiagaraSceneProxy *SceneProxy)
+	virtual FPrimitiveViewRelevance GetViewRelevance(const FSceneView* View, const FNiagaraSceneProxy *SceneProxy)
 	{
 		FPrimitiveViewRelevance Result;
 		bool bHasDynamicData = HasDynamicData();
diff --git a/Engine/Source/Runtime/Engine/Private/AI/Navigation/NavLinkRenderingComponent.cpp b/Engine/Source/Runtime/Engine/Private/AI/Navigation/NavLinkRenderingComponent.cpp
--- a/Engine/Source/Runtime/Engine/Private/AI/Navigation/NavLinkRenderingComponent.cpp
+++ b/Engine/Source/Runtime/Engine/Private/AI/Navigation/NavLinkRenderingComponent.cpp
@@ -417,7 +417,7 @@ void FNavLinkRenderingProxy::DrawLinks(FPrimitiveDrawInterface* PDI, TArray<FNav
 	}
 }
 
-FPrimitiveViewRelevance FNavLinkRenderingProxy::GetViewRelevance(const FSceneView* View)
+FPrimitiveViewRelevance FNavLinkRenderingProxy::GetViewRelevance(const FSceneView* View) const
 {
 	FPrimitiveViewRelevance Result;
 	Result.bDrawRelevance = IsShown(View) && IsSelected() && (View && View->Family && View->Family->EngineShowFlags.Navigation);
diff --git a/Engine/Source/Runtime/Engine/Private/AI/Navigation/NavTestRenderingComponent.cpp b/Engine/Source/Runtime/Engine/Private/AI/Navigation/NavTestRenderingComponent.cpp
--- a/Engine/Source/Runtime/Engine/Private/AI/Navigation/NavTestRenderingComponent.cpp
+++ b/Engine/Source/Runtime/Engine/Private/AI/Navigation/NavTestRenderingComponent.cpp
@@ -376,7 +376,7 @@ public:
 		Canvas->SetDrawColor(OldDrawColor);
 	}
 
-	virtual FPrimitiveViewRelevance GetViewRelevance(const FSceneView* View) override
+	virtual FPrimitiveViewRelevance GetViewRelevance(const FSceneView* View) const override
 	{
 		FPrimitiveViewRelevance Result;
 		Result.bDrawRelevance = IsShown(View);
diff --git a/Engine/Source/Runtime/Engine/Private/Components/ArrowComponent.cpp b/Engine/Source/Runtime/Engine/Private/Components/ArrowComponent.cpp
--- a/Engine/Source/Runtime/Engine/Private/Components/ArrowComponent.cpp
+++ b/Engine/Source/Runtime/Engine/Private/Components/ArrowComponent.cpp
@@ -222,7 +222,7 @@ public:
 		}
 	}
 
-	virtual FPrimitiveViewRelevance GetViewRelevance(const FSceneView* View) override
+	virtual FPrimitiveViewRelevance GetViewRelevance(const FSceneView* View) const override
 	{
 		FPrimitiveViewRelevance Result;
 		Result.bDrawRelevance = IsShown(View) && (View->Family->EngineShowFlags.BillboardSprites);
diff --git a/Engine/Source/Runtime/Engine/Private/Components/BillboardComponent.cpp b/Engine/Source/Runtime/Engine/Private/Components/BillboardComponent.cpp
--- a/Engine/Source/Runtime/Engine/Private/Components/BillboardComponent.cpp
+++ b/Engine/Source/Runtime/Engine/Private/Components/BillboardComponent.cpp
@@ -200,7 +200,7 @@ public:
 		}
 	}
 
-	virtual FPrimitiveViewRelevance GetViewRelevance(const FSceneView* View) override
+	virtual FPrimitiveViewRelevance GetViewRelevance(const FSceneView* View) const override
 	{
 		bool bVisible = View->Family->EngineShowFlags.BillboardSprites;
 #if WITH_EDITOR
diff --git a/Engine/Source/Runtime/Engine/Private/Components/BoxComponent.cpp b/Engine/Source/Runtime/Engine/Private/Components/BoxComponent.cpp
--- a/Engine/Source/Runtime/Engine/Private/Components/BoxComponent.cpp
+++ b/Engine/Source/Runtime/Engine/Private/Components/BoxComponent.cpp
@@ -116,7 +116,7 @@ FPrimitiveSceneProxy* UBoxComponent::CreateSceneProxy()
 			}
 		}
 
-		virtual FPrimitiveViewRelevance GetViewRelevance(const FSceneView* View) override
+		virtual FPrimitiveViewRelevance GetViewRelevance(const FSceneView* View) const override
 		{
 			const bool bVisible = !bDrawOnlyIfSelected || IsSelected();
 
diff --git a/Engine/Source/Runtime/Engine/Private/Components/BrushComponent.cpp b/Engine/Source/Runtime/Engine/Private/Components/BrushComponent.cpp
--- a/Engine/Source/Runtime/Engine/Private/Components/BrushComponent.cpp
+++ b/Engine/Source/Runtime/Engine/Private/Components/BrushComponent.cpp
@@ -329,7 +329,7 @@ public:
 		}
 	}
 
-	virtual FPrimitiveViewRelevance GetViewRelevance(const FSceneView* View) override
+	virtual FPrimitiveViewRelevance GetViewRelevance(const FSceneView* View) const override
 	{
 		bool bVisible = false;
 
diff --git a/Engine/Source/Runtime/Engine/Private/Components/CapsuleComponent.cpp b/Engine/Source/Runtime/Engine/Private/Components/CapsuleComponent.cpp
--- a/Engine/Source/Runtime/Engine/Private/Components/CapsuleComponent.cpp
+++ b/Engine/Source/Runtime/Engine/Private/Components/CapsuleComponent.cpp
@@ -55,7 +55,7 @@ FPrimitiveSceneProxy* UCapsuleComponent::CreateSceneProxy()
 			}
 		}
 
-		virtual FPrimitiveViewRelevance GetViewRelevance(const FSceneView* View) override
+		virtual FPrimitiveViewRelevance GetViewRelevance(const FSceneView* View) const override
 		{
 			const bool bVisible = !bDrawOnlyIfSelected || IsSelected();
 			FPrimitiveViewRelevance Result;
diff --git a/Engine/Source/Runtime/Engine/Private/Components/DrawFrustumComponent.cpp b/Engine/Source/Runtime/Engine/Private/Components/DrawFrustumComponent.cpp
--- a/Engine/Source/Runtime/Engine/Private/Components/DrawFrustumComponent.cpp
+++ b/Engine/Source/Runtime/Engine/Private/Components/DrawFrustumComponent.cpp
@@ -108,7 +108,7 @@ public:
 		}
 	}
 
-	virtual FPrimitiveViewRelevance GetViewRelevance(const FSceneView* View) override
+	virtual FPrimitiveViewRelevance GetViewRelevance(const FSceneView* View) const override
 	{
 		FPrimitiveViewRelevance Result;
 		Result.bDrawRelevance = IsShown(View) && View->Family->EngineShowFlags.CameraFrustums;
diff --git a/Engine/Source/Runtime/Engine/Private/Components/LineBatchComponent.cpp b/Engine/Source/Runtime/Engine/Private/Components/LineBatchComponent.cpp
--- a/Engine/Source/Runtime/Engine/Private/Components/LineBatchComponent.cpp
+++ b/Engine/Source/Runtime/Engine/Private/Components/LineBatchComponent.cpp
@@ -80,7 +80,7 @@ public:
 	*	@param		Scene view to use to determine our relevence.
 	*  @return		View relevance struct
 	*/
-	virtual FPrimitiveViewRelevance GetViewRelevance(const FSceneView* View) override
+	virtual FPrimitiveViewRelevance GetViewRelevance(const FSceneView* View) const override
 	{
 		return ViewRelevance;
 	}
diff --git a/Engine/Source/Runtime/Engine/Private/Components/MaterialBillboardComponent.cpp b/Engine/Source/Runtime/Engine/Private/Components/MaterialBillboardComponent.cpp
--- a/Engine/Source/Runtime/Engine/Private/Components/MaterialBillboardComponent.cpp
+++ b/Engine/Source/Runtime/Engine/Private/Components/MaterialBillboardComponent.cpp
@@ -230,7 +230,7 @@ public:
 		}
 	}
 
-	virtual FPrimitiveViewRelevance GetViewRelevance(const FSceneView* View) override
+	virtual FPrimitiveViewRelevance GetViewRelevance(const FSceneView* View) const override
 	{
 		bool bVisible = View->Family->EngineShowFlags.BillboardSprites;
 		FPrimitiveViewRelevance Result;
diff --git a/Engine/Source/Runtime/Engine/Private/Components/SphereComponent.cpp b/Engine/Source/Runtime/Engine/Private/Components/SphereComponent.cpp
--- a/Engine/Source/Runtime/Engine/Private/Components/SphereComponent.cpp
+++ b/Engine/Source/Runtime/Engine/Private/Components/SphereComponent.cpp
@@ -129,7 +129,7 @@ FPrimitiveSceneProxy* USphereComponent::CreateSceneProxy()
 			}
 		}
 
-		virtual FPrimitiveViewRelevance GetViewRelevance(const FSceneView* View)  override
+		virtual FPrimitiveViewRelevance GetViewRelevance(const FSceneView* View) const override
 		{
 			const bool bVisibleForSelection = !bDrawOnlyIfSelected || IsSelected();
 			const bool bVisibleForShowFlags = true; // @TODO
diff --git a/Engine/Source/Runtime/Engine/Private/Components/TextRenderComponent.cpp b/Engine/Source/Runtime/Engine/Private/Components/TextRenderComponent.cpp
--- a/Engine/Source/Runtime/Engine/Private/Components/TextRenderComponent.cpp
+++ b/Engine/Source/Runtime/Engine/Private/Components/TextRenderComponent.cpp
@@ -365,7 +365,7 @@ public:
 	// Begin FPrimitiveSceneProxy interface
 	virtual void GetDynamicMeshElements(const TArray<const FSceneView*>& Views, const FSceneViewFamily& ViewFamily, uint32 VisibilityMap, FMeshElementCollector& Collector) const override;
 	virtual void DrawStaticElements(FStaticPrimitiveDrawInterface* PDI) override;
-	virtual FPrimitiveViewRelevance GetViewRelevance(const FSceneView* View) override;
+	virtual FPrimitiveViewRelevance GetViewRelevance(const FSceneView* View) const override;
 	virtual bool CanBeOccluded() const override;
 	virtual uint32 GetMemoryFootprint() const override;
 	uint32 GetAllocatedSize() const;
@@ -532,7 +532,7 @@ bool FTextRenderSceneProxy::CanBeOccluded() const
 	return !MaterialRelevance.bDisableDepthTest;
 }
 
-FPrimitiveViewRelevance FTextRenderSceneProxy::GetViewRelevance(const FSceneView* View)
+FPrimitiveViewRelevance FTextRenderSceneProxy::GetViewRelevance(const FSceneView* View) const
 {
 	FPrimitiveViewRelevance Result;
 	Result.bDrawRelevance = IsShown(View) && View->Family->EngineShowFlags.TextRender;
diff --git a/Engine/Source/Runtime/Engine/Private/HierarchicalInstancedStaticMesh.cpp b/Engine/Source/Runtime/Engine/Private/HierarchicalInstancedStaticMesh.cpp
--- a/Engine/Source/Runtime/Engine/Private/HierarchicalInstancedStaticMesh.cpp
+++ b/Engine/Source/Runtime/Engine/Private/HierarchicalInstancedStaticMesh.cpp
@@ -752,7 +752,7 @@ public:
 		SceneProxyCreatedFrameNumberRenderThread = GFrameNumberRenderThread;
 	}
 	
-	virtual FPrimitiveViewRelevance GetViewRelevance(const FSceneView* View) override
+	virtual FPrimitiveViewRelevance GetViewRelevance(const FSceneView* View) const override
 	{
 		FPrimitiveViewRelevance Result;
 		if (bIsGrass ? View->Family->EngineShowFlags.InstancedGrass : View->Family->EngineShowFlags.InstancedFoliage)
diff --git a/Engine/Source/Runtime/Engine/Private/InstancedStaticMesh.h b/Engine/Source/Runtime/Engine/Private/InstancedStaticMesh.h
--- a/Engine/Source/Runtime/Engine/Private/InstancedStaticMesh.h
+++ b/Engine/Source/Runtime/Engine/Private/InstancedStaticMesh.h
@@ -528,7 +528,7 @@ public:
 
 	// FPrimitiveSceneProxy interface.
 	
-	virtual FPrimitiveViewRelevance GetViewRelevance(const FSceneView* View) override
+	virtual FPrimitiveViewRelevance GetViewRelevance(const FSceneView* View) const override
 	{
 		FPrimitiveViewRelevance Result;
 		if(View->Family->EngineShowFlags.InstancedStaticMeshes)
diff --git a/Engine/Source/Runtime/Engine/Private/ModelRender.cpp b/Engine/Source/Runtime/Engine/Private/ModelRender.cpp
--- a/Engine/Source/Runtime/Engine/Private/ModelRender.cpp
+++ b/Engine/Source/Runtime/Engine/Private/ModelRender.cpp
@@ -442,7 +442,7 @@ public:
 		}
 	}
 
-	virtual FPrimitiveViewRelevance GetViewRelevance(const FSceneView* View) override
+	virtual FPrimitiveViewRelevance GetViewRelevance(const FSceneView* View) const override
 	{
 		FPrimitiveViewRelevance Result;
 		Result.bDrawRelevance = IsShown(View) && View->Family->EngineShowFlags.BSPTriangles && View->Family->EngineShowFlags.BSP;
diff --git a/Engine/Source/Runtime/Engine/Private/NavMeshRenderingHelpers.h b/Engine/Source/Runtime/Engine/Private/NavMeshRenderingHelpers.h
--- a/Engine/Source/Runtime/Engine/Private/NavMeshRenderingHelpers.h
+++ b/Engine/Source/Runtime/Engine/Private/NavMeshRenderingHelpers.h
@@ -561,7 +561,7 @@ public:
 		Canvas->SetDrawColor(OldDrawColor);
 	}
 
-	virtual FPrimitiveViewRelevance GetViewRelevance(const FSceneView* View) override
+	virtual FPrimitiveViewRelevance GetViewRelevance(const FSceneView* View) const override
 	{
 		const bool bVisible = !!View->Family->EngineShowFlags.Navigation || bForceRendering;
 		FPrimitiveViewRelevance Result;
diff --git a/Engine/Source/Runtime/Engine/Private/Particles/ParticleSystemRender.cpp b/Engine/Source/Runtime/Engine/Private/Particles/ParticleSystemRender.cpp
--- a/Engine/Source/Runtime/Engine/Private/Particles/ParticleSystemRender.cpp
+++ b/Engine/Source/Runtime/Engine/Private/Particles/ParticleSystemRender.cpp
@@ -6553,7 +6553,7 @@ void FParticleSystemSceneProxy::GetObjectPositionAndScale(const FSceneView& View
 /**
 * @return Relevance for rendering the particle system primitive component in the given View
 */
-FPrimitiveViewRelevance FParticleSystemSceneProxy::GetViewRelevance(const FSceneView* View)
+FPrimitiveViewRelevance FParticleSystemSceneProxy::GetViewRelevance(const FSceneView* View) const
 {
 	FPrimitiveViewRelevance Result;
 	Result.bDrawRelevance = IsShown(View) && View->Family->EngineShowFlags.Particles;
diff --git a/Engine/Source/Runtime/Engine/Private/PrimitiveSceneProxy.cpp b/Engine/Source/Runtime/Engine/Private/PrimitiveSceneProxy.cpp
--- a/Engine/Source/Runtime/Engine/Private/PrimitiveSceneProxy.cpp
+++ b/Engine/Source/Runtime/Engine/Private/PrimitiveSceneProxy.cpp
@@ -172,7 +172,7 @@ HHitProxy* FPrimitiveSceneProxy::CreateHitProxies(UPrimitiveComponent* Component
 	}
 }
 
-FPrimitiveViewRelevance FPrimitiveSceneProxy::GetViewRelevance(const FSceneView* View)
+FPrimitiveViewRelevance FPrimitiveSceneProxy::GetViewRelevance(const FSceneView* View) const
 {
 	return FPrimitiveViewRelevance();
 }
diff --git a/Engine/Source/Runtime/Engine/Private/SkeletalMesh.cpp b/Engine/Source/Runtime/Engine/Private/SkeletalMesh.cpp
--- a/Engine/Source/Runtime/Engine/Private/SkeletalMesh.cpp
+++ b/Engine/Source/Runtime/Engine/Private/SkeletalMesh.cpp
@@ -4741,7 +4741,7 @@ bool FSkeletalMeshSceneProxy::GetWorldMatrices( FMatrix& OutLocalToWorld, FMatri
 /**
  * Relevance is always dynamic for skel meshes unless they are disabled
  */
-FPrimitiveViewRelevance FSkeletalMeshSceneProxy::GetViewRelevance(const FSceneView* View)
+FPrimitiveViewRelevance FSkeletalMeshSceneProxy::GetViewRelevance(const FSceneView* View) const
 {
 	FPrimitiveViewRelevance Result;
 	Result.bDrawRelevance = IsShown(View) && View->Family->EngineShowFlags.SkeletalMeshes;
diff --git a/Engine/Source/Runtime/Engine/Private/SplineMeshSceneProxy.h b/Engine/Source/Runtime/Engine/Private/SplineMeshSceneProxy.h
--- a/Engine/Source/Runtime/Engine/Private/SplineMeshSceneProxy.h
+++ b/Engine/Source/Runtime/Engine/Private/SplineMeshSceneProxy.h
@@ -162,7 +162,7 @@ public:
 	/** Sets up a wireframe FMeshBatch for a specific LOD. */
 	virtual bool GetWireframeMeshElement(int32 LODIndex, int32 BatchIndex, const FMaterialRenderProxy* WireframeRenderProxy, uint8 InDepthPriorityGroup, FMeshBatch& OutMeshBatch) const override;
 
-	virtual FPrimitiveViewRelevance GetViewRelevance(const FSceneView* View) override
+	virtual FPrimitiveViewRelevance GetViewRelevance(const FSceneView* View) const override
 	{
 		return FStaticMeshSceneProxy::GetViewRelevance(View);
 	}
diff --git a/Engine/Source/Runtime/Engine/Private/StaticMeshRender.cpp b/Engine/Source/Runtime/Engine/Private/StaticMeshRender.cpp
--- a/Engine/Source/Runtime/Engine/Private/StaticMeshRender.cpp
+++ b/Engine/Source/Runtime/Engine/Private/StaticMeshRender.cpp
@@ -918,7 +918,7 @@ bool FStaticMeshSceneProxy::CanBeOccluded() const
 	return !MaterialRelevance.bDisableDepthTest && !ShouldRenderCustomDepth();
 }
 
-FPrimitiveViewRelevance FStaticMeshSceneProxy::GetViewRelevance(const FSceneView* View)
+FPrimitiveViewRelevance FStaticMeshSceneProxy::GetViewRelevance(const FSceneView* View) const
 {   
 	checkSlow(IsInParallelRenderingThread());
 
diff --git a/Engine/Source/Runtime/Engine/Private/VectorField.cpp b/Engine/Source/Runtime/Engine/Private/VectorField.cpp
--- a/Engine/Source/Runtime/Engine/Private/VectorField.cpp
+++ b/Engine/Source/Runtime/Engine/Private/VectorField.cpp
@@ -440,7 +440,7 @@ public:
 	/**
 	 * Computes view relevance for this scene proxy.
 	 */
-	virtual FPrimitiveViewRelevance GetViewRelevance(const FSceneView* View) override
+	virtual FPrimitiveViewRelevance GetViewRelevance(const FSceneView* View) const override
 	{
 		FPrimitiveViewRelevance Result;
 		Result.bDrawRelevance = IsShown(View); 
diff --git a/Engine/Source/Runtime/Engine/Public/AI/NavLinkRenderingProxy.h b/Engine/Source/Runtime/Engine/Public/AI/NavLinkRenderingProxy.h
--- a/Engine/Source/Runtime/Engine/Public/AI/NavLinkRenderingProxy.h
+++ b/Engine/Source/Runtime/Engine/Public/AI/NavLinkRenderingProxy.h
@@ -40,7 +40,7 @@ public:
 	/** Initialization constructor. */
 	FNavLinkRenderingProxy(const UPrimitiveComponent* InComponent);
 	virtual void GetDynamicMeshElements(const TArray<const FSceneView*>& Views, const FSceneViewFamily& ViewFamily, uint32 VisibilityMap, FMeshElementCollector& Collector) const override;
-	virtual FPrimitiveViewRelevance GetViewRelevance(const FSceneView* View) override;
+	virtual FPrimitiveViewRelevance GetViewRelevance(const FSceneView* View) const override;
 	virtual uint32 GetMemoryFootprint( void ) const override;
 	uint32 GetAllocatedSize( void ) const;
 	void StorePointLinks(const FTransform& LocalToWorld, const TArray<FNavigationLink>& LinksArray);
diff --git a/Engine/Source/Runtime/Engine/Public/ParticleHelper.h b/Engine/Source/Runtime/Engine/Public/ParticleHelper.h
--- a/Engine/Source/Runtime/Engine/Public/ParticleHelper.h
+++ b/Engine/Source/Runtime/Engine/Public/ParticleHelper.h
@@ -2413,7 +2413,7 @@ public:
 		return false;
 	}
 	virtual void GetDynamicMeshElements(const TArray<const FSceneView*>& Views, const FSceneViewFamily& ViewFamily, uint32 VisibilityMap, FMeshElementCollector& Collector) const override;
-	virtual FPrimitiveViewRelevance GetViewRelevance(const FSceneView* View) override;
+	virtual FPrimitiveViewRelevance GetViewRelevance(const FSceneView* View) const override;
 	virtual void OnActorPositionChanged() override;
 	virtual void OnTransformChanged() override;
 
diff --git a/Engine/Source/Runtime/Engine/Public/PrimitiveSceneProxy.h b/Engine/Source/Runtime/Engine/Public/PrimitiveSceneProxy.h
--- a/Engine/Source/Runtime/Engine/Public/PrimitiveSceneProxy.h
+++ b/Engine/Source/Runtime/Engine/Public/PrimitiveSceneProxy.h
@@ -211,7 +211,7 @@ public:
 	 * @param View - The view to determine relevance for.
 	 * @return The relevance of the primitive's elements to the view.
 	 */
-	ENGINE_API virtual FPrimitiveViewRelevance GetViewRelevance(const FSceneView* View);
+	ENGINE_API virtual FPrimitiveViewRelevance GetViewRelevance(const FSceneView* View) const;
 
 	/** Callback from the renderer to gather simple lights that this proxy wants renderered. */
 	virtual void GatherSimpleLights(const FSceneViewFamily& ViewFamily, FSimpleLightArray& OutParticleLights) const {}
diff --git a/Engine/Source/Runtime/Engine/Public/SkeletalMeshTypes.h b/Engine/Source/Runtime/Engine/Public/SkeletalMeshTypes.h
--- a/Engine/Source/Runtime/Engine/Public/SkeletalMeshTypes.h
+++ b/Engine/Source/Runtime/Engine/Public/SkeletalMeshTypes.h
@@ -1742,7 +1742,7 @@ public:
 	virtual HHitProxy* CreateHitProxies(UPrimitiveComponent* Component, TArray<TRefCountPtr<HHitProxy> >& OutHitProxies) override;
 #endif
 	virtual void GetDynamicMeshElements(const TArray<const FSceneView*>& Views, const FSceneViewFamily& ViewFamily, uint32 VisibilityMap, FMeshElementCollector& Collector) const override;
-	virtual FPrimitiveViewRelevance GetViewRelevance(const FSceneView* View) override;
+	virtual FPrimitiveViewRelevance GetViewRelevance(const FSceneView* View) const override;
 	virtual bool CanBeOccluded() const override;
 	
 	/**
diff --git a/Engine/Source/Runtime/Engine/Public/StaticMeshResources.h b/Engine/Source/Runtime/Engine/Public/StaticMeshResources.h
--- a/Engine/Source/Runtime/Engine/Public/StaticMeshResources.h
+++ b/Engine/Source/Runtime/Engine/Public/StaticMeshResources.h
@@ -772,7 +772,7 @@ public:
 	virtual void DrawStaticElements(FStaticPrimitiveDrawInterface* PDI) override;
 	virtual void OnTransformChanged() override;
 	virtual int32 GetLOD(const FSceneView* View) const override;
-	virtual FPrimitiveViewRelevance GetViewRelevance(const FSceneView* View) override;
+	virtual FPrimitiveViewRelevance GetViewRelevance(const FSceneView* View) const override;
 	virtual bool CanBeOccluded() const override;
 	virtual void GetLightRelevance(const FLightSceneProxy* LightSceneProxy, bool& bDynamic, bool& bRelevant, bool& bLightMapped, bool& bShadowMapped) const override;
 	virtual void GetDistancefieldAtlasData(FBox& LocalVolumeBounds, FIntVector& OutBlockMin, FIntVector& OutBlockSize, bool& bOutBuiltAsIfTwoSided, bool& bMeshWasPlane, TArray<FMatrix>& ObjectLocalToWorldTransforms) const override;
diff --git a/Engine/Source/Runtime/GeometryCache/Private/GeometryCacheSceneProxy.cpp b/Engine/Source/Runtime/GeometryCache/Private/GeometryCacheSceneProxy.cpp
--- a/Engine/Source/Runtime/GeometryCache/Private/GeometryCacheSceneProxy.cpp
+++ b/Engine/Source/Runtime/GeometryCache/Private/GeometryCacheSceneProxy.cpp
@@ -151,7 +151,7 @@ void FGeometryCacheSceneProxy::GetDynamicMeshElements(const TArray<const FSceneV
 #endif
 }
 
-FPrimitiveViewRelevance FGeometryCacheSceneProxy::GetViewRelevance(const FSceneView* View)
+FPrimitiveViewRelevance FGeometryCacheSceneProxy::GetViewRelevance(const FSceneView* View) const
 {
 	FPrimitiveViewRelevance Result;
 	Result.bDrawRelevance = IsShown(View);
diff --git a/Engine/Source/Runtime/GeometryCache/Private/GeometryCacheSceneProxy.h b/Engine/Source/Runtime/GeometryCache/Private/GeometryCacheSceneProxy.h
--- a/Engine/Source/Runtime/GeometryCache/Private/GeometryCacheSceneProxy.h
+++ b/Engine/Source/Runtime/GeometryCache/Private/GeometryCacheSceneProxy.h
@@ -92,7 +92,7 @@ public:
 
 	virtual void GetDynamicMeshElements(const TArray<const FSceneView*>& Views, const FSceneViewFamily& ViewFamily, uint32 VisibilityMap, FMeshElementCollector& Collector) const override;
 
-	virtual FPrimitiveViewRelevance GetViewRelevance(const FSceneView* View);
+	virtual FPrimitiveViewRelevance GetViewRelevance(const FSceneView* View) const override;
 
 	virtual bool CanBeOccluded() const override;
 
diff --git a/Engine/Source/Runtime/Landscape/Private/LandscapeGizmoActor.cpp b/Engine/Source/Runtime/Landscape/Private/LandscapeGizmoActor.cpp
--- a/Engine/Source/Runtime/Landscape/Private/LandscapeGizmoActor.cpp
+++ b/Engine/Source/Runtime/Landscape/Private/LandscapeGizmoActor.cpp
@@ -323,7 +323,7 @@ public:
 #endif
 	};
 
-	virtual FPrimitiveViewRelevance GetViewRelevance(const FSceneView* View) override
+	virtual FPrimitiveViewRelevance GetViewRelevance(const FSceneView* View) const override
 	{
 		FPrimitiveViewRelevance Result;
 #if WITH_EDITOR
diff --git a/Engine/Source/Runtime/Landscape/Private/LandscapeRender.cpp b/Engine/Source/Runtime/Landscape/Private/LandscapeRender.cpp
--- a/Engine/Source/Runtime/Landscape/Private/LandscapeRender.cpp
+++ b/Engine/Source/Runtime/Landscape/Private/LandscapeRender.cpp
@@ -793,7 +793,7 @@ bool FLandscapeComponentSceneProxy::CanBeOccluded() const
 	return !MaterialRelevance.bDisableDepthTest;
 }
 
-FPrimitiveViewRelevance FLandscapeComponentSceneProxy::GetViewRelevance(const FSceneView* View)
+FPrimitiveViewRelevance FLandscapeComponentSceneProxy::GetViewRelevance(const FSceneView* View) const
 {
 	FPrimitiveViewRelevance Result;
 	Result.bDrawRelevance = IsShown(View) && View->Family->EngineShowFlags.Landscape;
diff --git a/Engine/Source/Runtime/Landscape/Private/LandscapeSplines.cpp b/Engine/Source/Runtime/Landscape/Private/LandscapeSplines.cpp
--- a/Engine/Source/Runtime/Landscape/Private/LandscapeSplines.cpp
+++ b/Engine/Source/Runtime/Landscape/Private/LandscapeSplines.cpp
@@ -284,7 +284,7 @@ public:
 		}
 	}
 
-	virtual FPrimitiveViewRelevance GetViewRelevance(const FSceneView* View) override
+	virtual FPrimitiveViewRelevance GetViewRelevance(const FSceneView* View) const override
 	{
 		FPrimitiveViewRelevance Result;
 		Result.bDrawRelevance = IsShown(View) && View->Family->EngineShowFlags.Splines;
diff --git a/Engine/Source/Runtime/Landscape/Public/LandscapeRender.h b/Engine/Source/Runtime/Landscape/Public/LandscapeRender.h
--- a/Engine/Source/Runtime/Landscape/Public/LandscapeRender.h
+++ b/Engine/Source/Runtime/Landscape/Public/LandscapeRender.h
@@ -572,7 +572,7 @@ public:
 	virtual void DrawStaticElements(FStaticPrimitiveDrawInterface* PDI) override;
 	virtual void GetDynamicMeshElements(const TArray<const FSceneView*>& Views, const FSceneViewFamily& ViewFamily, uint32 VisibilityMap, FMeshElementCollector& Collector) const override;
 	virtual uint32 GetMemoryFootprint() const override { return(sizeof(*this) + GetAllocatedSize()); }
-	virtual FPrimitiveViewRelevance GetViewRelevance(const FSceneView* View) override;
+	virtual FPrimitiveViewRelevance GetViewRelevance(const FSceneView* View) const override;
 	virtual bool CanBeOccluded() const override;
 	virtual void GetLightRelevance(const FLightSceneProxy* LightSceneProxy, bool& bDynamic, bool& bRelevant, bool& bLightMapped, bool& bShadowMapped) const override;
 	virtual void OnTransformChanged() override;
diff --git a/Engine/Source/Runtime/Niagara/Private/NiagaraComponent.cpp b/Engine/Source/Runtime/Niagara/Private/NiagaraComponent.cpp
--- a/Engine/Source/Runtime/Niagara/Private/NiagaraComponent.cpp
+++ b/Engine/Source/Runtime/Niagara/Private/NiagaraComponent.cpp
@@ -89,7 +89,7 @@ void FNiagaraSceneProxy::OnTransformChanged()
 	//WorldSpacePrimitiveUniformBuffer.ReleaseResource();
 }
 
-FPrimitiveViewRelevance FNiagaraSceneProxy::GetViewRelevance(const FSceneView* View)
+FPrimitiveViewRelevance FNiagaraSceneProxy::GetViewRelevance(const FSceneView* View) const
 {
 	FPrimitiveViewRelevance Relevance;
 	Relevance.bDynamicRelevance = true;
diff --git a/Engine/Source/Runtime/Niagara/Public/NiagaraComponent.h b/Engine/Source/Runtime/Niagara/Public/NiagaraComponent.h
--- a/Engine/Source/Runtime/Niagara/Public/NiagaraComponent.h
+++ b/Engine/Source/Runtime/Niagara/Public/NiagaraComponent.h
@@ -87,7 +87,7 @@ private:
 
 	virtual void GetDynamicMeshElements(const TArray<const FSceneView*>& Views, const FSceneViewFamily& ViewFamily, uint32 VisibilityMap, FMeshElementCollector& Collector) const override;
 
-	virtual FPrimitiveViewRelevance GetViewRelevance(const FSceneView* View)  override;
+	virtual FPrimitiveViewRelevance GetViewRelevance(const FSceneView* View) const override;
 
 	/*
 	virtual bool CanBeOccluded() const override
diff --git a/Engine/Source/Runtime/UMG/Private/Components/WidgetComponent.cpp b/Engine/Source/Runtime/UMG/Private/Components/WidgetComponent.cpp
--- a/Engine/Source/Runtime/UMG/Private/Components/WidgetComponent.cpp
+++ b/Engine/Source/Runtime/UMG/Private/Components/WidgetComponent.cpp
@@ -285,7 +285,7 @@ public:
 #endif
 	}
 
-	virtual FPrimitiveViewRelevance GetViewRelevance(const FSceneView* View) override
+	virtual FPrimitiveViewRelevance GetViewRelevance(const FSceneView* View) const override
 	{
 		bool bVisible = true;
 
