#ue4 - Don't start ticking Actors/Components until BeginPlay() has been called when in a game world. This was an issue in some situations in network games.

- Actors may opt out of this by setting bAllowTickBeforeBeginPlay to true. APlayerController, AInfo and derived classes (AGameMode, AWorldSettings, etc) enable this by default because they typically have match start up logic.
- Components with no Actor owner will start ticking immediately before the world begins play. This should be rare however.
- Editor worlds don't "begin play" so ticking begins immediately in those.

--------

diff --git a/Engine/Source/Runtime/Engine/Classes/Components/ActorComponent.h b/Engine/Source/Runtime/Engine/Classes/Components/ActorComponent.h
--- a/Engine/Source/Runtime/Engine/Classes/Components/ActorComponent.h
+++ b/Engine/Source/Runtime/Engine/Classes/Components/ActorComponent.h
@@ -158,6 +158,9 @@ private:
 	/** Indicates the the destruction process has begun for this component to avoid recursion */
 	uint32 bIsBeingDestroyed:1;
 
+	/** Whether we've tried to register tick functions. Reset when they are unregistered. */
+	uint32 bTickFunctionsRegistered:1;
+
 #if WITH_EDITOR
 	/** During undo/redo it isn't safe to cache owner */
 	uint32 bCanUseCachedOwner:1;
diff --git a/Engine/Source/Runtime/Engine/Classes/GameFramework/Actor.h b/Engine/Source/Runtime/Engine/Classes/GameFramework/Actor.h
--- a/Engine/Source/Runtime/Engine/Classes/GameFramework/Actor.h
+++ b/Engine/Source/Runtime/Engine/Classes/GameFramework/Actor.h
@@ -177,10 +177,21 @@ public:
 	/** True if this actor is currently running user construction script (used to defer component registration) */
 	uint32 bRunningUserConstructionScript:1;
 
+	/**
+	 * Whether we allow this Actor to tick before it receives the BeginPlay event.
+	 * Normally we don't tick actors until after BeginPlay; this setting allows this behavior to be overridden.
+	 * This Actor must be able to tick for this setting to be relevant.
+	 */
+	UPROPERTY(EditDefaultsOnly, Category="Tick")
+	uint32 bAllowTickBeforeBeginPlay:1;
+
 private:
 	/** Whether FinishSpawning has been called for this Actor.  If it has not, the Actor is in a mal-formed state */
 	uint32 bHasFinishedSpawning:1;
 
+	/** Whether we've tried to register tick functions. Reset when they are unregistered. */
+	uint32 bTickFunctionsRegistered : 1;
+
 	/**
 	 * Enables any collision on this actor.
 	 * @see SetActorEnableCollision(), GetActorEnableCollision()
diff --git a/Engine/Source/Runtime/Engine/Private/Actor.cpp b/Engine/Source/Runtime/Engine/Private/Actor.cpp
--- a/Engine/Source/Runtime/Engine/Private/Actor.cpp
+++ b/Engine/Source/Runtime/Engine/Private/Actor.cpp
@@ -247,7 +247,6 @@ void AActor::ResetOwnedComponents()
 void AActor::PostInitProperties()
 {
 	Super::PostInitProperties();
-	RegisterAllActorTickFunctions(true,false); // component will be handled when they are registered
 	RemoteRole = (bReplicates ? ROLE_SimulatedProxy : ROLE_None);
 
 	// Map natively-constructed component instances to any serialized properties that might reference them.
@@ -641,11 +640,16 @@ void AActor::RegisterAllActorTickFunctions(bool bRegister, bool bDoComponents)
 {
 	if(!IsTemplate())
 	{
-		FActorThreadContext& ThreadContext = FActorThreadContext::Get();
-		check(ThreadContext.TestRegisterTickFunctions == nullptr);
-		RegisterActorTickFunctions(bRegister);
-		checkf(ThreadContext.TestRegisterTickFunctions == this, TEXT("Failed to route Actor RegisterTickFunctions (%s)"), *GetFullName());
-		ThreadContext.TestRegisterTickFunctions = nullptr;
+		// Prevent repeated redundant attempts
+		if (bTickFunctionsRegistered != bRegister)
+		{
+			FActorThreadContext& ThreadContext = FActorThreadContext::Get();
+			check(ThreadContext.TestRegisterTickFunctions == nullptr);
+			RegisterActorTickFunctions(bRegister);
+			bTickFunctionsRegistered = bRegister;
+			checkf(ThreadContext.TestRegisterTickFunctions == this, TEXT("Failed to route Actor RegisterTickFunctions (%s)"), *GetFullName());
+			ThreadContext.TestRegisterTickFunctions = nullptr;
+		}
 
 		if (bDoComponents)
 		{
@@ -2783,6 +2787,7 @@ void AActor::BeginPlay()
 {
 	ensure(!bActorHasBegunPlay);
 	SetLifeSpan( InitialLifeSpan );
+	RegisterAllActorTickFunctions(true, false); // Components are done below.
 
 	TInlineComponentArray<UActorComponent*> Components;
 	GetComponents(Components);
@@ -2792,7 +2797,11 @@ void AActor::BeginPlay()
 		// bHasBegunPlay will be true for the component if the component was renamed and moved to a new outer during initialization
 		if (Component->IsRegistered() && !Component->HasBegunPlay())
 		{
-			Component->BeginPlay();
+			Component->RegisterAllComponentTickFunctions(true);
+			if (Component->bWantsBeginPlay)
+			{
+				Component->BeginPlay();
+			}
 		}
 		else
 		{
@@ -3622,6 +3631,16 @@ bool AActor::IncrementalRegisterComponents(int32 NumComponentsToRegister)
 		// 0 - means register all components
 		NumComponentsToRegister = MAX_int32;
 	}
+
+	UWorld* const World = GetWorld();
+	check(World);
+
+	// If we are not a game world, then register tick functions now. If we are a game world we wait until right before BeginPlay(),
+	// so as to not actually tick until BeginPlay() executes (which could otherwise happen in network games).
+	if (bAllowTickBeforeBeginPlay || !World->IsGameWorld())
+	{
+		RegisterAllActorTickFunctions(true, false); // components will be handled when they are registered
+	}
 	
 	// Register RootComponent first so all other components can reliable use it (ie call GetLocation) when they register
 	if( RootComponent != NULL && !RootComponent->IsRegistered() )
@@ -3636,8 +3655,7 @@ bool AActor::IncrementalRegisterComponents(int32 NumComponentsToRegister)
 		//This should prevent unwanted components hanging around when undoing a copy/paste or duplication action.
 		RootComponent->Modify(false);
 
-		check(GetWorld());
-		RootComponent->RegisterComponentWithWorld(GetWorld());
+		RootComponent->RegisterComponentWithWorld(World);
 	}
 
 	int32 NumTotalRegisteredComponents = 0;
@@ -3676,8 +3694,7 @@ bool AActor::IncrementalRegisterComponents(int32 NumComponentsToRegister)
 			//This should prevent unwanted components hanging around when undoing a copy/paste or duplication action.
 			Component->Modify(false);
 
-			check(GetWorld());
-			Component->RegisterComponentWithWorld(GetWorld());
+			Component->RegisterComponentWithWorld(World);
 			NumRegisteredComponentsThisRun++;
 		}
 
diff --git a/Engine/Source/Runtime/Engine/Private/Components/ActorComponent.cpp b/Engine/Source/Runtime/Engine/Private/Components/ActorComponent.cpp
--- a/Engine/Source/Runtime/Engine/Private/Components/ActorComponent.cpp
+++ b/Engine/Source/Runtime/Engine/Private/Components/ActorComponent.cpp
@@ -636,6 +636,8 @@ void UActorComponent::BeginPlay()
 {
 	check(bRegistered);
 	check(!bHasBegunPlay);
+	checkSlow(bWantsBeginPlay);
+	checkSlow(bTickFunctionsRegistered); // If this fails, someone called BeginPlay() without first calling RegisterAllComponentTickFunctions().
 
 	ReceiveBeginPlay();
 
@@ -755,9 +757,14 @@ void UActorComponent::RegisterAllComponentTickFunctions(bool bRegister)
 	// Components don't have tick functions until they are registered with the world
 	if (bRegistered)
 	{
-		RegisterComponentTickFunctions(bRegister);
-		checkf(GTestRegisterComponentTickFunctions == this, TEXT("Failed to route component RegisterTickFunctions (%s)"), *GetFullName());
-		GTestRegisterComponentTickFunctions = NULL;
+		// Prevent repeated redundant attempts
+		if (bTickFunctionsRegistered != bRegister)
+		{
+			RegisterComponentTickFunctions(bRegister);
+			bTickFunctionsRegistered = bRegister;
+			checkf(GTestRegisterComponentTickFunctions == this, TEXT("Failed to route component RegisterTickFunctions (%s)"), *GetFullName());
+			GTestRegisterComponentTickFunctions = NULL;
+		}
 	}
 }
 
@@ -822,7 +829,12 @@ void UActorComponent::RegisterComponentWithWorld(UWorld* InWorld)
 	World = InWorld;
 
 	ExecuteRegisterEvents();
-	RegisterAllComponentTickFunctions(true);
+
+	// If not in a game world register ticks now, otherwise defer until BeginPlay. If no owner we won't trigger BeginPlay either so register now in that case as well.
+	if (MyOwner == nullptr || !InWorld->IsGameWorld())
+	{
+		RegisterAllComponentTickFunctions(true);
+	}
 
 	if (MyOwner == nullptr || MyOwner->IsActorInitialized())
 	{
@@ -832,9 +844,10 @@ void UActorComponent::RegisterComponentWithWorld(UWorld* InWorld)
 		}
 	}
 
-	if (MyOwner && MyOwner->HasActorBegunPlay())
+	if (MyOwner && MyOwner->HasActorBegunPlay() && !bHasBegunPlay)
 	{
-		if (!bHasBegunPlay && bWantsBeginPlay)
+		RegisterAllComponentTickFunctions(true);
+		if (bWantsBeginPlay)
 		{
 			BeginPlay();
 		}
diff --git a/Engine/Source/Runtime/Engine/Private/Info.cpp b/Engine/Source/Runtime/Engine/Private/Info.cpp
--- a/Engine/Source/Runtime/Engine/Private/Info.cpp
+++ b/Engine/Source/Runtime/Engine/Private/Info.cpp
@@ -32,6 +32,7 @@ AInfo::AInfo(const FObjectInitializer& ObjectInitializer)
 #endif // WITH_EDITORONLY_DATA
 
 	PrimaryActorTick.bCanEverTick = false;
+	bAllowTickBeforeBeginPlay = true;
 	bReplicates = false;
 	NetUpdateFrequency = 10.0f;
 	bHidden = true;
diff --git a/Engine/Source/Runtime/Engine/Private/PlayerController.cpp b/Engine/Source/Runtime/Engine/Private/PlayerController.cpp
--- a/Engine/Source/Runtime/Engine/Private/PlayerController.cpp
+++ b/Engine/Source/Runtime/Engine/Private/PlayerController.cpp
@@ -67,6 +67,7 @@ APlayerController::APlayerController(const FObjectInitializer& ObjectInitializer
 	bIsUsingStreamingVolumes = true;
 	PrimaryActorTick.TickGroup = TG_PrePhysics;
 	PrimaryActorTick.bTickEvenWhenPaused = true;
+	bAllowTickBeforeBeginPlay = true;
 	bShouldPerformFullTickWhenPaused = false;
 	LastRetryPlayerTime = 0.f;
 	DefaultMouseCursor = EMouseCursor::Default;
