Editor: Fix inverted mirrored brushes on load.
Note: Unshelved from pending changelist '2675825' (by RichTW)
[INTEGRATE] Change 2676478 by [REDACTED_PERSON][REDACTED_HANDLE] on 2015/09/01 20:38:54

diff --git a/Engine/Source/Editor/UnrealEd/Classes/Editor/UnrealEdEngine.h b/Engine/Source/Editor/UnrealEd/Classes/Editor/UnrealEdEngine.h
--- a/Engine/Source/Editor/UnrealEd/Classes/Editor/UnrealEdEngine.h
+++ b/Engine/Source/Editor/UnrealEd/Classes/Editor/UnrealEdEngine.h
@@ -685,6 +685,11 @@ public:
 	 */
 	void UpdateVolumeActorVisibility( UClass* InVolumeActorClass = NULL , FLevelEditorViewportClient* InViewport = NULL);
 
+	/**
+	 * Identify any brushes whose sense is inverted and repair them
+	 */
+	void FixAnyInvertedBrushes(UWorld* World);
+
 	/**
 	 * Get the index of the provided sprite category
 	 *
diff --git a/Engine/Source/Editor/UnrealEd/Private/FileHelpers.cpp b/Engine/Source/Editor/UnrealEd/Private/FileHelpers.cpp
--- a/Engine/Source/Editor/UnrealEd/Private/FileHelpers.cpp
+++ b/Engine/Source/Editor/UnrealEd/Private/FileHelpers.cpp
@@ -2063,6 +2063,9 @@ void FEditorFileUtils::LoadMap(const FString& InFilename, bool LoadAsTemplate, b
 	// potentially contain volumes.
 	GUnrealEd->UpdateVolumeActorVisibility(NULL);
 
+	// If there are any old mirrored brushes in the map with inverted polys, fix them here
+	GUnrealEd->FixAnyInvertedBrushes(World);
+
 	// Fire delegate when a new map is opened, with name of map
 	FEditorDelegates::OnMapOpened.Broadcast(InFilename, LoadAsTemplate);
 }
diff --git a/Engine/Source/Editor/UnrealEd/Private/UnrealEdEngine.cpp b/Engine/Source/Editor/UnrealEd/Private/UnrealEdEngine.cpp
--- a/Engine/Source/Editor/UnrealEd/Private/UnrealEdEngine.cpp
+++ b/Engine/Source/Editor/UnrealEd/Private/UnrealEdEngine.cpp
@@ -1192,6 +1192,56 @@ void UUnrealEdEngine::UpdateVolumeActorVisibility( UClass* InVolumeActorClass, F
 }
 
 
+void UUnrealEdEngine::FixAnyInvertedBrushes(UWorld* World)
+{
+	// Get list of brushes with inverted polys
+	TArray<ABrush*> Brushes;
+	for (TActorIterator<ABrush> It(World); It; ++It)
+	{
+		ABrush* Brush = *It;
+		if (Brush->BrushComponent && Brush->BrushComponent->HasInvertedPolys())
+		{
+			Brushes.Add(Brush);
+		}
+	}
+
+	bool bAnyStaticBrushesFixed = false;
+	if (Brushes.Num() > 0)
+	{
+		for (ABrush* Brush : Brushes)
+		{
+			UE_LOG(LogUnrealEdEngine, Warning, TEXT("Brush '%s' appears to be inside out - fixing."), *Brush->GetName());
+
+			// Invert the polys of the brush
+			for (auto& Poly : Brush->BrushComponent->Brush->Polys->Element)
+			{
+				Poly.Reverse();
+				Poly.CalcNormal();
+			}
+
+			if (Brush->IsStaticBrush())
+			{
+				// Static brushes require a full BSP rebuild
+				bAnyStaticBrushesFixed = true;
+			}
+			else
+			{
+				// Dynamic brushes can be fixed up here
+				FBSPOps::csgPrepMovingBrush(Brush);
+				Brush->BrushComponent->BuildSimpleBrushCollision();
+			}
+
+			Brush->MarkPackageDirty();
+		}
+
+		if (bAnyStaticBrushesFixed)
+		{
+			RebuildAlteredBSP();
+		}
+	}
+}
+
+
 void UUnrealEdEngine::RegisterComponentVisualizer(FName ComponentClassName, TSharedPtr<FComponentVisualizer> Visualizer)
 {
 	if( ComponentClassName != NAME_Name )
diff --git a/Engine/Source/Runtime/Engine/Classes/Components/BrushComponent.h b/Engine/Source/Runtime/Engine/Classes/Components/BrushComponent.h
--- a/Engine/Source/Runtime/Engine/Classes/Components/BrushComponent.h
+++ b/Engine/Source/Runtime/Engine/Classes/Components/BrushComponent.h
@@ -49,6 +49,9 @@ public:
 	virtual bool ComponentIsTouchingSelectionFrustum(const FConvexVolume& InFrustum, const FEngineShowFlags& ShowFlags, const bool bConsiderOnlyBSP, const bool bMustEncompassEntireComponent) const override;
 	// End UPrimitiveComponent interface.
 
+	/** Return true if the brush appears to have inverted polys */
+	ENGINE_API bool HasInvertedPolys() const;
+
 	/** If the transform mirroring no longer reflects the body setup, request its recalculation */
 	ENGINE_API void RequestUpdateBrushCollision();
 #endif
diff --git a/Engine/Source/Runtime/Engine/Private/Components/BrushComponent.cpp b/Engine/Source/Runtime/Engine/Private/Components/BrushComponent.cpp
--- a/Engine/Source/Runtime/Engine/Private/Components/BrushComponent.cpp
+++ b/Engine/Source/Runtime/Engine/Private/Components/BrushComponent.cpp
@@ -780,4 +780,77 @@ void UBrushComponent::BuildSimpleBrushCollision()
 }
 
 
+#if WITH_EDITOR
+static FVector GetPolyCenter(const FPoly& Poly)
+{
+	FVector Result = FVector::ZeroVector;
+	for (const auto& Vertex : Poly.Vertices)
+	{
+		Result += Vertex;
+	}
+
+	return Result / Poly.Vertices.Num();
+}
+
+bool UBrushComponent::HasInvertedPolys() const
+{
+	// Determine if a brush looks as if it has had its sense inverted
+	// (due to the old behavior of inverting the poly winding and normal when performing a Mirror operation).
+
+	const bool bIsMirrored = (RelativeScale3D.X * RelativeScale3D.Y * RelativeScale3D.Z < 0.0f);
+	// Only attempt to fix up brushes with negative scale
+	if (bIsMirrored)
+	{
+		int NumInwardFacingPolys = 0;
+		for (auto& Poly : Brush->Polys->Element)
+		{
+			// Calculate a nominal center point for the poly
+			const FVector PolyCenter = GetPolyCenter(Poly);
+			bool bIntersected = false;
 
+			// Find intersections of a ray cast out from the center in the normal direction with the other polys
+			for (auto& OtherPoly : Brush->Polys->Element)
+			{
+				if (&Poly != &OtherPoly)
+				{
+					// Calculate a nominal center point for the poly being tested for intersection
+					const FVector OtherPolyCenter = GetPolyCenter(OtherPoly);
+					const float Dot = FVector::DotProduct(Poly.Normal, OtherPoly.Normal);
+					// If normals are perpendicular, skip it - this implies that the poly normal is parallel to the plane
+					if (Dot != 0.0f)
+					{
+						const float Distance = FVector::DotProduct(OtherPolyCenter - PolyCenter, OtherPoly.Normal) / Dot;
+						// Only consider intersections in the direction of the poly normal
+						if (Distance > 0.0f)
+						{
+							const FVector Intersection = PolyCenter + Poly.Normal * Distance;
+
+							// Does the ray intersect with the actual poly?
+							if (OtherPoly.OnPoly(Intersection))
+							{
+								// If so, toggle the intersected flag.
+								// An odd number of intersections implies an inwards facing poly.
+								// An even number of intersections implies an outwards facing poly.
+								bIntersected = !bIntersected;
+							}
+						}
+					}
+				}
+			}
+
+			if (bIntersected)
+			{
+				NumInwardFacingPolys++;
+			}
+		}
+
+		// If more than half of the polys are deemed to be inwards facing, consider this to be an inside out brush
+		if (NumInwardFacingPolys > Brush->Polys->Element.Num() / 2)
+		{
+			return true;
+		}
+	}
+
+	return false;
+}
+#endif
