AbstractNavData will be restored as part of SpawnMissingNavigationData function, it doesn't need to be serialized and loaded

diff --git a/Engine/Source/Runtime/Engine/Private/AI/Navigation/NavigationSystem.cpp b/Engine/Source/Runtime/Engine/Private/AI/Navigation/NavigationSystem.cpp
--- a/Engine/Source/Runtime/Engine/Private/AI/Navigation/NavigationSystem.cpp
+++ b/Engine/Source/Runtime/Engine/Private/AI/Navigation/NavigationSystem.cpp
@@ -341,7 +341,7 @@ void UNavigationSystem::DoInitialSetup()
 
 void UNavigationSystem::UpdateAbstractNavData()
 {
-	if (AbstractNavData != nullptr)
+	if (AbstractNavData != nullptr && !AbstractNavData->IsPendingKill())
 	{
 		return;
 	}
@@ -3217,6 +3217,8 @@ void UNavigationSystem::SpawnMissingNavigationData()
 		// update 
 		MainNavData = GetMainNavData(FNavigationSystem::DontCreate);
 	}
+
+	UpdateAbstractNavData();
 }
 
 ANavigationData* UNavigationSystem::CreateNavigationDataInstance(const FNavDataConfig& NavConfig)
