Hopefully fix sporadic assert when switching sync group leader
Merged from 2664514

diff --git a/Engine/Source/Runtime/Engine/Classes/Animation/AnimationAsset.h b/Engine/Source/Runtime/Engine/Classes/Animation/AnimationAsset.h
--- a/Engine/Source/Runtime/Engine/Classes/Animation/AnimationAsset.h
+++ b/Engine/Source/Runtime/Engine/Classes/Animation/AnimationAsset.h
@@ -337,7 +337,7 @@ public:
 	ENGINE_API void TestTickRecordForLeadership(EAnimGroupRole::Type MembershipType);
 
 	// Called after all tick records have been added but before assets are actually ticked
-	ENGINE_API void Finalize(const TArray<FName>& PreviousValidMarkers);
+	ENGINE_API void Finalize(const TArray<FName>& PreviousValidMarkers, int32 PreviousGroupLeader);
 };
 
 /** Utility struct to accumulate root motion. */
diff --git a/Engine/Source/Runtime/Engine/Private/Animation/AnimInstance.cpp b/Engine/Source/Runtime/Engine/Private/Animation/AnimInstance.cpp
--- a/Engine/Source/Runtime/Engine/Private/Animation/AnimInstance.cpp
+++ b/Engine/Source/Runtime/Engine/Private/Animation/AnimInstance.cpp
@@ -425,8 +425,10 @@ void UAnimInstance::UpdateAnimation(float DeltaSeconds)
 			if (SyncGroup.ActivePlayers.Num() > 0)
 			{
 				static TArray<FName> TempNames;
-				const TArray<FName>& PreviousValidMarkers = PreviousSyncGroups.IsValidIndex(GroupIndex) ? PreviousSyncGroups[GroupIndex].ValidMarkers : TempNames;
-				SyncGroup.Finalize(PreviousValidMarkers);
+				const bool bSyncGroupValidLastFrame = PreviousSyncGroups.IsValidIndex(GroupIndex);
+				const TArray<FName>& PreviousValidMarkers = bSyncGroupValidLastFrame ? PreviousSyncGroups[GroupIndex].ValidMarkers : TempNames;
+				const int32 PreviousGroupLeader = bSyncGroupValidLastFrame ? PreviousSyncGroups[GroupIndex].GroupLeaderIndex : -1;
+				SyncGroup.Finalize(PreviousValidMarkers, PreviousGroupLeader);
     
 				// Tick the group leader
 				FAnimAssetTickContext TickContext(DeltaSeconds, RootMotionMode, SyncGroup.ValidMarkers);
diff --git a/Engine/Source/Runtime/Engine/Private/AnimationAsset.cpp b/Engine/Source/Runtime/Engine/Private/AnimationAsset.cpp
--- a/Engine/Source/Runtime/Engine/Private/AnimationAsset.cpp
+++ b/Engine/Source/Runtime/Engine/Private/AnimationAsset.cpp
@@ -33,7 +33,7 @@ void FAnimGroupInstance::TestTickRecordForLeadership(EAnimGroupRole::Type Member
 	}
 }
 
-void FAnimGroupInstance::Finalize(const TArray<FName>& PreviousValidMarkers)
+void FAnimGroupInstance::Finalize(const TArray<FName>& PreviousValidMarkers, int32 PreviousGroupLeader)
 {
 	GroupLeaderIndex = FMath::Max(GroupLeaderIndex, 0);
 
@@ -72,7 +72,7 @@ void FAnimGroupInstance::Finalize(const TArray<FName>& PreviousValidMarkers)
 			}
 		}
 		ValidMarkers.Sort();
-		if (ValidMarkers != PreviousValidMarkers)
+		if (ValidMarkers != PreviousValidMarkers || PreviousGroupLeader != GroupLeaderIndex)
 		{
 			for (int32 Idx = 0; Idx < ActivePlayers.Num(); ++Idx)
 			{
