Fixed an issue with wrong Device Profile Texture LOD Settings being used at cook time.

diff --git a/Engine/Source/Runtime/Engine/Private/TextureDerivedData.cpp b/Engine/Source/Runtime/Engine/Private/TextureDerivedData.cpp
--- a/Engine/Source/Runtime/Engine/Private/TextureDerivedData.cpp
+++ b/Engine/Source/Runtime/Engine/Private/TextureDerivedData.cpp
@@ -368,7 +368,7 @@ static void GetTextureBuildSettings(
 	OutBuildSettings.bFlipGreenChannel = Texture.bFlipGreenChannel;
 	OutBuildSettings.CompositeTextureMode = Texture.CompositeTextureMode;
 	OutBuildSettings.CompositePower = Texture.CompositePower;
-	OutBuildSettings.LODBias = UDeviceProfileManager::Get().GetActiveProfile()->GetTextureLODSettings()->CalculateLODBias(Texture.Source.GetSizeX(), Texture.Source.GetSizeY(), Texture.LODGroup, Texture.LODBias, Texture.NumCinematicMipLevels, Texture.MipGenSettings);
+	OutBuildSettings.LODBias = TextureLODSettings.CalculateLODBias(Texture.Source.GetSizeX(), Texture.Source.GetSizeY(), Texture.LODGroup, Texture.LODBias, Texture.NumCinematicMipLevels, Texture.MipGenSettings);
 	OutBuildSettings.bStreamable = !Texture.NeverStream && (Texture.LODGroup != TEXTUREGROUP_UI) && (Cast<const UTexture2D>(&Texture) != NULL);
 	OutBuildSettings.PowerOfTwoMode = Texture.PowerOfTwoMode;
 	OutBuildSettings.PaddingColor = Texture.PaddingColor;
