[Merging using UE4-Orion-To-UE4] - Fix shadow warning from Blend Profile checkin

diff --git a/Engine/Source/Runtime/Engine/Private/Animation/BlendProfile.cpp b/Engine/Source/Runtime/Engine/Private/Animation/BlendProfile.cpp
--- a/Engine/Source/Runtime/Engine/Private/Animation/BlendProfile.cpp
+++ b/Engine/Source/Runtime/Engine/Private/Animation/BlendProfile.cpp
@@ -143,9 +143,9 @@ int32 UBlendProfile::GetPerBoneInterpolationIndex(int32 BoneIndex, const FBoneCo
 
 void UBlendProfile::SetSingleBoneBlendScale(int32 InBoneIdx, float InScale, bool bCreate /*= false*/)
 {
-	FBlendProfileBoneEntry* Entry = ProfileEntries.FindByPredicate([InBoneIdx](const FBlendProfileBoneEntry& Entry)
+	FBlendProfileBoneEntry* Entry = ProfileEntries.FindByPredicate([InBoneIdx](const FBlendProfileBoneEntry& InEntry)
 	{
-		return Entry.BoneReference.BoneIndex == InBoneIdx;
+		return InEntry.BoneReference.BoneIndex == InBoneIdx;
 	});
 
 	if(!Entry && bCreate)
