Fix multiple assert issues in blend spaces plus stop anim sequences evaluation by marker sync in Persona

Merged from 2666209

diff --git a/Engine/Source/Runtime/Engine/Classes/Animation/AnimationAsset.h b/Engine/Source/Runtime/Engine/Classes/Animation/AnimationAsset.h
--- a/Engine/Source/Runtime/Engine/Classes/Animation/AnimationAsset.h
+++ b/Engine/Source/Runtime/Engine/Classes/Animation/AnimationAsset.h
@@ -190,6 +190,9 @@ struct FMarkerSyncAnimPosition
 	1   we are at NextMarker
 	0.5 we are half way between the two */
 	float PositionBetweenMarkers;
+
+	/** Is this a valid Marker Sync Position */
+	bool IsValid() const { return !(PreviousMarkerName == NAME_None && NextMarkerName == NAME_None); }
 };
 
 struct FPassedMarker
@@ -452,7 +455,7 @@ struct FRootMotionMovementParams
 struct FAnimAssetTickContext
 {
 public:
-	FAnimAssetTickContext(float InDeltaTime, ERootMotionMode::Type InRootMotionMode, const TArray<FName>& ValidMarkerNames)
+	FAnimAssetTickContext(float InDeltaTime, ERootMotionMode::Type InRootMotionMode, bool bInOnlyOneAnimationInGroup, const TArray<FName>& ValidMarkerNames)
 		: RootMotionMode(InRootMotionMode)
 		, MarkerTickContext(ValidMarkerNames)
 		, DeltaTime(InDeltaTime)
@@ -460,16 +463,18 @@ public:
 		, AnimLengthRatio(0.0f)
 		, bIsMarkerPositionValid(ValidMarkerNames.Num() > 0)
 		, bIsLeader(true)
+		, bOnlyOneAnimationInGroup(bInOnlyOneAnimationInGroup)
 	{
 	}
 
-	FAnimAssetTickContext(float InDeltaTime, ERootMotionMode::Type InRootMotionMode)
+	FAnimAssetTickContext(float InDeltaTime, ERootMotionMode::Type InRootMotionMode, bool bInOnlyOneAnimationInGroup)
 		: RootMotionMode(InRootMotionMode)
 		, DeltaTime(InDeltaTime)
 		, LeaderDelta(0.f)
 		, AnimLengthRatio(0.0f)
 		, bIsMarkerPositionValid(false)
 		, bIsLeader(true)
+		, bOnlyOneAnimationInGroup(bInOnlyOneAnimationInGroup)
 	{
 	}
 	
@@ -527,6 +532,11 @@ public:
 		return IsLeader();
 	}
 
+	bool IsSingleAnimationContext() const
+	{
+		return bOnlyOneAnimationInGroup;
+	}
+
 	//Root Motion accumulated from this tick context
 	FRootMotionMovementParams RootMotionMovementParams;
 
@@ -546,6 +556,8 @@ private:
 	bool bIsMarkerPositionValid;
 
 	bool bIsLeader;
+
+	bool bOnlyOneAnimationInGroup;
 };
 
 USTRUCT()
diff --git a/Engine/Source/Runtime/Engine/Private/Animation/AnimInstance.cpp b/Engine/Source/Runtime/Engine/Private/Animation/AnimInstance.cpp
--- a/Engine/Source/Runtime/Engine/Private/Animation/AnimInstance.cpp
+++ b/Engine/Source/Runtime/Engine/Private/Animation/AnimInstance.cpp
@@ -109,7 +109,7 @@ void UAnimInstance::SequenceAdvanceImmediate(UAnimSequenceBase* Sequence, bool b
 	FAnimTickRecord TickRecord;
 	MakeSequenceTickRecord(TickRecord, Sequence, bLooping, PlayRate, /*FinalBlendWeight=*/ 1.0f, CurrentTime);
 
-	FAnimAssetTickContext TickContext(DeltaSeconds, RootMotionMode);
+	FAnimAssetTickContext TickContext(DeltaSeconds, RootMotionMode, true);
 	TickRecord.SourceAsset->TickAssetPlayerInstance(TickRecord, this, TickContext);
 }
 
@@ -118,7 +118,7 @@ void UAnimInstance::BlendSpaceAdvanceImmediate(class UBlendSpaceBase* BlendSpace
 	FAnimTickRecord TickRecord;
 	MakeBlendSpaceTickRecord(TickRecord, BlendSpace, BlendInput, BlendSampleDataCache, BlendFilter, bLooping, PlayRate, /*FinalBlendWeight=*/ 1.0f, CurrentTime);
 	
-	FAnimAssetTickContext TickContext(DeltaSeconds, RootMotionMode);
+	FAnimAssetTickContext TickContext(DeltaSeconds, RootMotionMode, true);
 	TickRecord.SourceAsset->TickAssetPlayerInstance(TickRecord, this, TickContext);
 }
 
@@ -430,8 +430,9 @@ void UAnimInstance::UpdateAnimation(float DeltaSeconds)
 				const int32 PreviousGroupLeader = bSyncGroupValidLastFrame ? PreviousSyncGroups[GroupIndex].GroupLeaderIndex : -1;
 				SyncGroup.Finalize(PreviousValidMarkers, PreviousGroupLeader);
     
+				const bool bOnlyOneAnimationInGroup = SyncGroup.ActivePlayers.Num() == 1;
 				// Tick the group leader
-				FAnimAssetTickContext TickContext(DeltaSeconds, RootMotionMode, SyncGroup.ValidMarkers);
+				FAnimAssetTickContext TickContext(DeltaSeconds, RootMotionMode, bOnlyOneAnimationInGroup, SyncGroup.ValidMarkers);
 				FAnimTickRecord& GroupLeader = SyncGroup.ActivePlayers[SyncGroup.GroupLeaderIndex];
 				{
 					QUICK_SCOPE_CYCLE_COUNTER(STAT_UAnimInstance_UpdateAnimation_TickAssetPlayerInstance);
@@ -475,12 +476,12 @@ void UAnimInstance::UpdateAnimation(float DeltaSeconds)
 		for (int32 TickIndex = 0; TickIndex < UngroupedActivePlayers.Num(); ++TickIndex)
 		{
 			FAnimTickRecord& AssetPlayerToTick = UngroupedActivePlayers[TickIndex];
-			AssetPlayerToTick.bCanUseMarkerSync = true;
 			static TArray<FName> TempNames;
 			const TArray<FName>* UniqueNames = AssetPlayerToTick.SourceAsset->GetUniqueMarkerNames();
 			const TArray<FName>& ValidMarkers = UniqueNames ? *UniqueNames : TempNames;
 
-			FAnimAssetTickContext TickContext(DeltaSeconds, RootMotionMode, ValidMarkers);
+			const bool bOnlyOneAnimationInGroup = true;
+			FAnimAssetTickContext TickContext(DeltaSeconds, RootMotionMode, bOnlyOneAnimationInGroup, ValidMarkers);
 		    {
 			    QUICK_SCOPE_CYCLE_COUNTER(STAT_UAnimInstance_UpdateAnimation_TickAssetPlayerInstance);
 			    FScopeCycleCounterUObject Scope(AssetPlayerToTick.SourceAsset);
diff --git a/Engine/Source/Runtime/Engine/Private/Animation/BlendSpaceBase.cpp b/Engine/Source/Runtime/Engine/Private/Animation/BlendSpaceBase.cpp
--- a/Engine/Source/Runtime/Engine/Private/Animation/BlendSpaceBase.cpp
+++ b/Engine/Source/Runtime/Engine/Private/Animation/BlendSpaceBase.cpp
@@ -262,7 +262,7 @@ void UBlendSpaceBase::TickAssetPlayerInstance(FAnimTickRecord& Instance, class U
 				SampleDataList.Append(NewSampleDataList);
 			}
 
-			bool bCanDoMarkerSync = Instance.bCanUseMarkerSync && Context.CanUseMarkerPosition();
+			bool bCanDoMarkerSync = bAllSequencesHaveMatchingMarkers && (Context.IsSingleAnimationContext() || (Instance.bCanUseMarkerSync && Context.CanUseMarkerPosition()));
 
 			if (bCanDoMarkerSync && Instance.MarkerTickRecord.IsValid())
 			{
@@ -305,8 +305,13 @@ void UBlendSpaceBase::TickAssetPlayerInstance(FAnimTickRecord& Instance, class U
 					{
 						SampleDataItem.Time = NormalizedCurrentTime * Sample.Animation->SequenceLength;
 					}
+					if (!SampleDataItem.MarkerTickRecord.IsValid() && Context.MarkerTickContext.GetMarkerSyncStartPosition().IsValid())
+					{
+						Sample.Animation->GetMarkerIndicesForPosition(Context.MarkerTickContext.GetMarkerSyncStartPosition(), true, SampleDataItem.MarkerTickRecord.PreviousMarker, SampleDataItem.MarkerTickRecord.NextMarker, SampleDataItem.Time);
+					}
 					if (Context.GetDeltaTime() != 0.f)
 					{
+						Context.SetLeaderDelta(Context.GetDeltaTime());
 						Sample.Animation->TickByMarkerAsLeader(SampleDataItem.MarkerTickRecord, Context.MarkerTickContext, SampleDataItem.Time, SampleDataItem.PreviousTime, Context.GetDeltaTime(), true);
 						TickFollowerSamples(SampleDataList, HighestWeightIndex, Context);
 					}
