Correct marker positions after animation trimming

Merged from 2667787

diff --git a/Engine/Source/Runtime/Engine/Private/Animation/AnimSequence.cpp b/Engine/Source/Runtime/Engine/Private/Animation/AnimSequence.cpp
--- a/Engine/Source/Runtime/Engine/Private/Animation/AnimSequence.cpp
+++ b/Engine/Source/Runtime/Engine/Private/Animation/AnimSequence.cpp
@@ -1395,6 +1395,31 @@ void UAnimSequence::ResizeSequence(float NewLength, int32 NewNumFrames, bool bIn
 		}
 	}
 
+	for (FAnimSyncMarker& Marker : AuthoredSyncMarkers)
+	{
+		float CurrentTime = Marker.Time;
+		if (bInsert)
+		{
+			// when insert, we only care about start time
+			// if it's later than start time
+			if (CurrentTime >= OldStartTime)
+			{
+				CurrentTime += Duration;
+			}
+		}
+		else
+		{
+			if (CurrentTime >= OldStartTime && CurrentTime <= OldEndTime)
+			{
+				CurrentTime = OldStartTime;
+			}
+			else if (CurrentTime > OldEndTime)
+			{
+				CurrentTime -= Duration;
+			}
+		}
+		Marker.Time = FMath::Clamp(CurrentTime, 0.f, SequenceLength);
+	}
 	// resize curves
 	RawCurveData.Resize(NewLength, bInsert, OldStartTime, OldEndTime);
 }
