[UE-20165] Restrict native constructed component to property mapping to Actor-based BP object loads only (perf).

change summary:
- added an override for AActor::Serialize()
- relocated code from AActor::PostInitProperties() to Serialize() and added some additional constraints to restrict its execution to load/BP only
- changed mapping to use component FName as the key rather than the raw pointer value (since the address can change during serialization/CoL)

diff --git a/Engine/Source/Runtime/Engine/Classes/GameFramework/Actor.h b/Engine/Source/Runtime/Engine/Classes/GameFramework/Actor.h
--- a/Engine/Source/Runtime/Engine/Classes/GameFramework/Actor.h
+++ b/Engine/Source/Runtime/Engine/Classes/GameFramework/Actor.h
@@ -1332,6 +1332,7 @@ public:
 	virtual void ProcessEvent( UFunction* Function, void* Parameters ) override;
 	virtual int32 GetFunctionCallspace( UFunction* Function, void* Parameters, FFrame* Stack ) override;
 	virtual bool CallRemoteFunction( UFunction* Function, void* Parameters, FOutParmRec* OutParms, FFrame* Stack ) override;
+	virtual void Serialize(FArchive& Ar) override;
 	virtual void PostLoad() override;
 	virtual void PostLoadSubobjects( FObjectInstancingGraph* OuterInstanceGraph ) override;
 	virtual void BeginDestroy() override;
@@ -2515,8 +2516,10 @@ private:
 	/** List of replicated components. */
 	TArray<UActorComponent*> ReplicatedComponents;
 
+#if WITH_EDITOR
 	/** Maps natively-constructed components to properties that reference them. */
-	TMultiMap<UActorComponent*, UObjectProperty*> NativeConstructedComponentToPropertyMap;
+	TMultiMap<FName, UObjectProperty*> NativeConstructedComponentToPropertyMap;
+#endif
 
 public:
 
diff --git a/Engine/Source/Runtime/Engine/Private/Actor.cpp b/Engine/Source/Runtime/Engine/Private/Actor.cpp
--- a/Engine/Source/Runtime/Engine/Private/Actor.cpp
+++ b/Engine/Source/Runtime/Engine/Private/Actor.cpp
@@ -193,6 +193,7 @@ bool AActor::CheckActorComponents()
 
 void AActor::ResetOwnedComponents()
 {
+#if WITH_EDITOR
 	// Identify any natively-constructed components referenced by properties that either failed to serialize or came in as NULL.
 	if(HasAnyFlags(RF_WasLoaded) && NativeConstructedComponentToPropertyMap.Num() > 0)
 	{
@@ -204,7 +205,7 @@ void AActor::ResetOwnedComponents()
 			{
 				// Find the property or properties that previously referenced the natively-constructed component.
 				TArray<UObjectProperty*> Properties;
-				NativeConstructedComponentToPropertyMap.MultiFind(Component, Properties);
+				NativeConstructedComponentToPropertyMap.MultiFind(Component->GetFName(), Properties);
 
 				// Determine if the property or properties are no longer valid references (either it got serialized out that way or something failed during load)
 				for (auto ObjProp : Properties)
@@ -223,6 +224,7 @@ void AActor::ResetOwnedComponents()
 		// Clear out the mapping as we don't need it anymore
 		NativeConstructedComponentToPropertyMap.Empty();
 	}
+#endif
 
 	TArray<UObject*> ActorChildren;
 	OwnedComponents.Empty();
@@ -249,27 +251,6 @@ void AActor::PostInitProperties()
 	Super::PostInitProperties();
 	RemoteRole = (bReplicates ? ROLE_SimulatedProxy : ROLE_None);
 
-	// Map natively-constructed component instances to any serialized properties that might reference them.
-	// We'll use this information post-load to determine if any owned components may not have been serialized through the reference property (i.e. in case the serialized property value ends up being NULL).
-	if(!HasAnyFlags(RF_WasLoaded))
-	{
-		NativeConstructedComponentToPropertyMap.Empty(OwnedComponents.Num());
-		for(TFieldIterator<UObjectProperty> PropertyIt(GetClass(), EFieldIteratorFlags::IncludeSuper); PropertyIt; ++PropertyIt)
-		{
-			UObjectProperty* ObjProp = *PropertyIt;
-
-			// Ignore transient properties since they won't be serialized
-			if(!ObjProp->HasAnyPropertyFlags(CPF_Transient))
-			{
-				UActorComponent* ActorComponent = Cast<UActorComponent>(ObjProp->GetObjectPropertyValue_InContainer(this));
-				if(ActorComponent != nullptr && ActorComponent->CreationMethod == EComponentCreationMethod::Native)
-				{
-					NativeConstructedComponentToPropertyMap.Add(ActorComponent, ObjProp);
-				}
-			}
-		}
-	}
-
 	// Make sure the OwnedComponents list correct.  
 	// Under some circumstances sub-object instancing can result in bogus/duplicate entries.
 	// This is not necessary for CDOs
@@ -476,6 +457,40 @@ bool AActor::IsReadyForFinishDestroy()
 	return Super::IsReadyForFinishDestroy() && DetachFence.IsFenceComplete();
 }
 
+void AActor::Serialize(FArchive& Ar)
+{
+#if WITH_EDITOR
+	// Prior to load, map natively-constructed component instances for Blueprint-generated class types to any serialized properties that might reference them.
+	// We'll use this information post-load to determine if any owned components may not have been serialized through the reference property (i.e. in case the serialized property value ends up being NULL).
+	if (Ar.IsLoading()
+		&& OwnedComponents.Num() > 0
+		&& !(Ar.GetPortFlags() & PPF_Duplicate)
+		&& HasAllFlags(RF_WasLoaded|RF_NeedPostLoad))
+	{
+		if (const UBlueprintGeneratedClass* BPGC = Cast<UBlueprintGeneratedClass>(GetClass()))
+		{
+			NativeConstructedComponentToPropertyMap.Empty(OwnedComponents.Num());
+			for(TFieldIterator<UObjectProperty> PropertyIt(BPGC, EFieldIteratorFlags::IncludeSuper); PropertyIt; ++PropertyIt)
+			{
+				UObjectProperty* ObjProp = *PropertyIt;
+
+				// Ignore transient properties since they won't be serialized
+				if(!ObjProp->HasAnyPropertyFlags(CPF_Transient))
+				{
+					UActorComponent* ActorComponent = Cast<UActorComponent>(ObjProp->GetObjectPropertyValue_InContainer(this));
+					if(ActorComponent != nullptr && ActorComponent->CreationMethod == EComponentCreationMethod::Native)
+					{
+						NativeConstructedComponentToPropertyMap.Add(ActorComponent->GetFName(), ObjProp);
+					}
+				}
+			}
+		}
+	}
+#endif
+
+	Super::Serialize(Ar);
+}
+
 void AActor::PostLoad()
 {
 	Super::PostLoad();
