relative reattach teleport

https://jira.ol.epicgames.net/browse/UE-19890
https://udn.unrealengine.com/questions/256055/scenecomponentattachto-problem.html

When doing relative reattachment (keep the relative pose not the global pose) we dont want to sweep across the universe and obliterate things along of the way.  Better to do a teleport.   So, got this working for the kinematic components. Supposedly this will apply for dynamic bodies within the same skeletal mesh component (untested).    Not yet for separate dynamic object components connected to the kinematic.   So on my test weapon rig, the kinematic nunchuck is now ok, but, as before, dynamic one still can affect things and gain energy/momentum during the big move.

since i was asked, this changelist doesn't directly affect the main cause of problems in UE-19100  however, the codefix here potentially would fix any problems had there been two inputs activated in a single frame (one causeing a teleport and another causing a kinematic move).

diff --git a/Engine/Source/Runtime/Engine/Private/Components/SceneComponent.cpp b/Engine/Source/Runtime/Engine/Private/Components/SceneComponent.cpp
--- a/Engine/Source/Runtime/Engine/Private/Components/SceneComponent.cpp
+++ b/Engine/Source/Runtime/Engine/Private/Components/SceneComponent.cpp
@@ -4,6 +4,7 @@
 	SceneComponent.cpp
 =============================================================================*/
 
+
 #include "EnginePrivate.h"
 #include "PhysicsPublic.h"
 #include "MessageLog.h"
@@ -1357,8 +1358,7 @@ bool USceneComponent::AttachTo(class USceneComponent* Parent, FName InSocketName
 
 		AttachParent->OnChildAttached(this);
 
-		// calculate transform with new attachment condition
-		UpdateComponentToWorld();
+		UpdateComponentToWorld(false, ETeleportType::TeleportPhysics);
 
 		if (UPrimitiveComponent * PrimitiveComponent = Cast<UPrimitiveComponent>(this))
 		{
diff --git a/Engine/Source/Runtime/Engine/Private/PhysicsEngine/BodyInstance.cpp b/Engine/Source/Runtime/Engine/Private/PhysicsEngine/BodyInstance.cpp
--- a/Engine/Source/Runtime/Engine/Private/PhysicsEngine/BodyInstance.cpp
+++ b/Engine/Source/Runtime/Engine/Private/PhysicsEngine/BodyInstance.cpp
@@ -2609,6 +2609,10 @@ void FBodyInstance::SetBodyTransform(const FTransform& NewTransform, ETeleportTy
 				// Otherwise, set global pose
 				else
 				{
+					if (!IsRigidBodyNonKinematic_AssumesLocked(PRigidDynamic))  // check if kinematic  (checks the physx bit for this)
+					{
+						PRigidDynamic->setKinematicTarget(PNewPose);  // physx doesn't clear target on setGlobalPose, so overwrite any previous attempt to set this that wasn't yet resolved
+					}
 					PRigidDynamic->setGlobalPose(PNewPose);
 				}
 			}
