Update comment on AActor::Tick to accurately represent what delta seconds is

diff --git a/Engine/Source/Runtime/Engine/Classes/GameFramework/Actor.h b/Engine/Source/Runtime/Engine/Classes/GameFramework/Actor.h
--- a/Engine/Source/Runtime/Engine/Classes/GameFramework/Actor.h
+++ b/Engine/Source/Runtime/Engine/Classes/GameFramework/Actor.h
@@ -1862,7 +1862,7 @@ public:
 	 *	Function called every frame on this Actor. Override this function to implement custom logic to be executed every frame.
 	 *	Note that Tick is disabled by default, and you will need to check PrimaryActorTick.bCanEverTick is set to true to enable it.
 	 *
-	 *	@param	DeltaSeconds	Game time elapsed since last call to Tick
+	 *	@param	DeltaSeconds	Game time elapsed during last frame modified by the time dilation
 	 */
 	virtual void Tick( float DeltaSeconds );
 
