Clear ClientWorldPackageName between seamless travel, fixes GameState being NULL on world when travelling to same map

diff --git a/Engine/Source/Runtime/Engine/Private/NetConnection.cpp b/Engine/Source/Runtime/Engine/Private/NetConnection.cpp
--- a/Engine/Source/Runtime/Engine/Private/NetConnection.cpp
+++ b/Engine/Source/Runtime/Engine/Private/NetConnection.cpp
@@ -1704,6 +1704,7 @@ void UNetConnection::ResetGameWorldState()
 	RecentlyDormantActors.Empty();
 	DormantActors.Empty();
 	ClientVisibleLevelNames.Empty();
+	ClientWorldPackageName = NAME_None;
 	KeepProcessingActorChannelBunchesMap.Empty();
 	DormantReplicatorMap.Empty();
 
