#ue4 - Revert change to check Component->bWantsBeginPlay in Actor::BeginPlay(), for now until upgrade issues are resolved with components that have not set this flag.

--------

diff --git a/Engine/Source/Runtime/Engine/Private/Actor.cpp b/Engine/Source/Runtime/Engine/Private/Actor.cpp
--- a/Engine/Source/Runtime/Engine/Private/Actor.cpp
+++ b/Engine/Source/Runtime/Engine/Private/Actor.cpp
@@ -2813,7 +2813,7 @@ void AActor::BeginPlay()
 		if (Component->IsRegistered() && !Component->HasBegunPlay())
 		{
 			Component->RegisterAllComponentTickFunctions(true);
-			if (Component->bWantsBeginPlay)
+			//if (Component->bWantsBeginPlay) // TODO: this should be enforced, but we need to address an upgrade path first to not quietly break code. Not checking this flag is the old behavior.
 			{
 				Component->BeginPlay();
 			}
diff --git a/Engine/Source/Runtime/Engine/Private/Components/ActorComponent.cpp b/Engine/Source/Runtime/Engine/Private/Components/ActorComponent.cpp
--- a/Engine/Source/Runtime/Engine/Private/Components/ActorComponent.cpp
+++ b/Engine/Source/Runtime/Engine/Private/Components/ActorComponent.cpp
@@ -636,7 +636,6 @@ void UActorComponent::BeginPlay()
 {
 	check(bRegistered);
 	check(!bHasBegunPlay);
-	checkSlow(bWantsBeginPlay);
 	checkSlow(bTickFunctionsRegistered); // If this fails, someone called BeginPlay() without first calling RegisterAllComponentTickFunctions().
 
 	ReceiveBeginPlay();
