Unshelved from pending changelist '2668808':

PR #1499: Exposed the underlying SoundAttenuation struct to BP (Contributed by [REDACTED_PERSON])

Note: I made the change BlueprintReadOnly

diff --git a/Engine/Source/Runtime/Engine/Classes/Sound/SoundAttenuation.h b/Engine/Source/Runtime/Engine/Classes/Sound/SoundAttenuation.h
--- a/Engine/Source/Runtime/Engine/Classes/Sound/SoundAttenuation.h
+++ b/Engine/Source/Runtime/Engine/Classes/Sound/SoundAttenuation.h
@@ -187,6 +187,6 @@ class USoundAttenuation : public UObject
 {
 	GENERATED_UCLASS_BODY()
 
-	UPROPERTY(EditAnywhere, Category=Settings)
+	UPROPERTY(EditAnywhere, BlueprintReadOnly, Category=Settings)
 	FAttenuationSettings Attenuation;
 };
