Track real time separately to compare against current time used to timeout connections

diff --git a/Engine/Source/Runtime/Engine/Classes/Engine/NetConnection.h b/Engine/Source/Runtime/Engine/Classes/Engine/NetConnection.h
--- a/Engine/Source/Runtime/Engine/Classes/Engine/NetConnection.h
+++ b/Engine/Source/Runtime/Engine/Classes/Engine/NetConnection.h
@@ -225,6 +225,8 @@ public:
 	// Internal.
 	UPROPERTY()
 	double			LastReceiveTime;		// Last time a packet was received, for timeout checking.
+	double			LastReceiveRealtime;	// Last time a packet was received, using real time seconds (FPlatformTime::Seconds)
+	double			LastGoodPacketRealtime;	// Last real time a packet was considered valid
 	double			LastSendTime;			// Last time a packet was sent, for keepalives.
 	double			LastTickTime;			// Last time of polling.
 	int32			QueuedBytes;			// Bytes assumed to be queued up.
diff --git a/Engine/Source/Runtime/Engine/Private/DemoNetDriver.cpp b/Engine/Source/Runtime/Engine/Private/DemoNetDriver.cpp
--- a/Engine/Source/Runtime/Engine/Private/DemoNetDriver.cpp
+++ b/Engine/Source/Runtime/Engine/Private/DemoNetDriver.cpp
@@ -2148,6 +2148,8 @@ void UDemoNetConnection::HandleClientPlayer( APlayerController* PC, UNetConnecti
 		PC->Role			= ROLE_AutonomousProxy;
 		PC->NetConnection	= NetConnection;
 		LastReceiveTime		= Driver->Time;
+		LastReceiveRealtime = FPlatformTime::Seconds();
+		LastGoodPacketRealtime = FPlatformTime::Seconds();
 		State				= USOCK_Open;
 		PlayerController	= PC;
 		OwningActor			= PC;
diff --git a/Engine/Source/Runtime/Engine/Private/NetConnection.cpp b/Engine/Source/Runtime/Engine/Private/NetConnection.cpp
--- a/Engine/Source/Runtime/Engine/Private/NetConnection.cpp
+++ b/Engine/Source/Runtime/Engine/Private/NetConnection.cpp
@@ -86,6 +86,8 @@ void UNetConnection::InitBase(UNetDriver* InDriver,class FSocket* InSocket, cons
 	// Stats
 	StatUpdateTime = Driver->Time;
 	LastReceiveTime = Driver->Time;
+	LastReceiveRealtime = FPlatformTime::Seconds();
+	LastGoodPacketRealtime = FPlatformTime::Seconds();
 	LastSendTime = Driver->Time;
 	LastTickTime = Driver->Time;
 	LastRecvAckTime = Driver->Time;
@@ -665,6 +667,7 @@ void UNetConnection::ReceivedPacket( FBitReader& Reader )
 
 	// Update receive time to avoid timeout.
 	LastReceiveTime = Driver->Time;
+	LastReceiveRealtime = FPlatformTime::Seconds();
 
 	// Check packet ordering.
 	const int32 PacketId = InternalAck ? InPacketId + 1 : MakeRelative(Reader.ReadInt(MAX_PACKETID),InPacketId,MAX_PACKETID);
@@ -1020,6 +1023,7 @@ void UNetConnection::ReceivedPacket( FBitReader& Reader )
 			{
 				UE_LOG( LogNetTraffic, Error, TEXT("Received corrupted packet data from client %s.  Disconnecting."), *LowLevelGetRemoteAddress() );
 				State = USOCK_Closed;
+				bSkipAck = true;
 			}
 		}
 	}
@@ -1031,6 +1035,8 @@ void UNetConnection::ReceivedPacket( FBitReader& Reader )
 	// Acknowledge the packet.
 	if ( !bSkipAck )
 	{
+		LastGoodPacketRealtime = FPlatformTime::Seconds();
+
 		SendAck(PacketId, true);
 	}
 }
@@ -1357,6 +1363,8 @@ void UNetConnection::Tick()
 		OutAckPacketId = OutPacketId;
 
 		LastReceiveTime = Driver->Time;
+		LastReceiveRealtime = FPlatformTime::Seconds();
+		LastGoodPacketRealtime = FPlatformTime::Seconds();
 		for( int32 i=OpenChannels.Num()-1; i>=0; i-- )
 		{
 			UChannel* It = OpenChannels[i];
@@ -1408,10 +1416,18 @@ void UNetConnection::Tick()
 	if ((Driver->Time - LastReceiveTime) > Timeout)
 #endif
 	{
+		// Compute true realtime since packet was received (as well as truly processed)
+		const double Seconds = FPlatformTime::Seconds();
+
+		const float ReceiveRealtimeDelta	= Seconds - LastReceiveRealtime;
+		const float GoodRealtimeDelta		= Seconds - LastGoodPacketRealtime;
+
 		// Timeout.
-		FString Error = FString::Printf(TEXT("UNetConnection::Tick: Connection TIMED OUT. Closing connection. Elapsed: %f, Threshold: %f, %s"),
+		FString Error = FString::Printf(TEXT("UNetConnection::Tick: Connection TIMED OUT. Closing connection. Elapsed: %2.2f, Real: %2.2f, Good: %2.2f, Threshold: %2.2f, %s"),
 			Driver->Time - LastReceiveTime,
-			Timeout, 
+			ReceiveRealtimeDelta,
+			GoodRealtimeDelta,
+			Timeout,
 			*Describe());
 		UE_LOG(LogNet, Warning, TEXT("%s"), *Error);
 		GEngine->BroadcastNetworkFailure(Driver->GetWorld(), Driver, ENetworkFailure::ConnectionTimeout, Error);
