#ANIM: Montage on memory leak issue
Merging

//depot/UE4-Orion/Engine/Source/Runtime/Engine/Private/Animation/AnimInstance.cpp

to //depot/UE4-Orion/Engine/Source/Runtime/Engine/Private/Animation/AnimInstance.cpp

diff --git a/Engine/Source/Runtime/Engine/Private/Animation/AnimInstance.cpp b/Engine/Source/Runtime/Engine/Private/Animation/AnimInstance.cpp
--- a/Engine/Source/Runtime/Engine/Private/Animation/AnimInstance.cpp
+++ b/Engine/Source/Runtime/Engine/Private/Animation/AnimInstance.cpp
@@ -290,6 +290,13 @@ void UAnimInstance::UninitializeAnimation()
 {
 	StopAllMontages(0.f);
 
+	for(int32 Index = 0; Index < MontageInstances.Num(); ++Index)
+	{
+		delete MontageInstances[Index];
+	}
+
+	MontageInstances.Empty();
+
 	USkeletalMeshComponent* SkelMeshComp = GetSkelMeshComponent();
 	if (SkelMeshComp)
 	{
