#ANIM: Fix crash due to over-fixing skeleton - causing compressed data would get lost.

Merging

//depot/UE4-Orion/Engine/Source/Runtime/Engine/Private/...

to //depot/UE4/Engine/Source/Runtime/Engine/Private/...

diff --git a/Engine/Source/Runtime/Engine/Private/Animation/AnimSequence.cpp b/Engine/Source/Runtime/Engine/Private/Animation/AnimSequence.cpp
--- a/Engine/Source/Runtime/Engine/Private/Animation/AnimSequence.cpp
+++ b/Engine/Source/Runtime/Engine/Private/Animation/AnimSequence.cpp
@@ -1666,47 +1666,51 @@ bool UAnimSequence::CompressRawAnimData(float MaxPosDiff, float MaxAngleDiff)
 	bool bRemovedKeys = false;
 #if WITH_EDITORONLY_DATA
 
-	// This removes trivial keys, and this has to happen before the removing tracks
-	for(int32 TrackIndex=0; TrackIndex<RawAnimationData.Num(); TrackIndex++)
+	if ( ensureMsgf (RawAnimationData.Num() > 0, TEXT("%s is trying to compress while raw animation is missing"), *GetName()) )
 	{
-		bRemovedKeys |= CompressRawAnimSequenceTrack(RawAnimationData[TrackIndex], MaxPosDiff, MaxAngleDiff);
-	}
+		// This removes trivial keys, and this has to happen before the removing tracks
+		for(int32 TrackIndex=0; TrackIndex<RawAnimationData.Num(); TrackIndex++)
+		{
+			bRemovedKeys |= CompressRawAnimSequenceTrack(RawAnimationData[TrackIndex], MaxPosDiff, MaxAngleDiff);
+		}
 
-	const USkeleton* MySkeleton = GetSkeleton();
+		const USkeleton* MySkeleton = GetSkeleton();
 
-	if (MySkeleton)
-	{
-		bool bCompressScaleKeys = false;
-		// go through remove keys if not needed
-		for ( int32 TrackIndex=0; TrackIndex<RawAnimationData.Num(); TrackIndex++ )
+		if(MySkeleton)
 		{
-			FRawAnimSequenceTrack const& RawData = RawAnimationData[TrackIndex];
-			if ( RawData.ScaleKeys.Num() > 0 )
+			bool bCompressScaleKeys = false;
+			// go through remove keys if not needed
+			for(int32 TrackIndex=0; TrackIndex<RawAnimationData.Num(); TrackIndex++)
 			{
-				// if scale key exists, see if we can just empty it
-				if((RawData.ScaleKeys.Num() > 1) || (RawData.ScaleKeys[0].Equals(FVector(1.f)) == false))
+				FRawAnimSequenceTrack const& RawData = RawAnimationData[TrackIndex];
+				if(RawData.ScaleKeys.Num() > 0)
 				{
-					bCompressScaleKeys = true;
-					break;
+					// if scale key exists, see if we can just empty it
+					if((RawData.ScaleKeys.Num() > 1) || (RawData.ScaleKeys[0].Equals(FVector(1.f)) == false))
+					{
+						bCompressScaleKeys = true;
+						break;
+					}
 				}
 			}
-		}
 
-		// if we don't have scale, we should delete all scale keys
-		// if you have one track that has scale, we still should support scale, so compress scale
-		if (!bCompressScaleKeys)
-		{
-			// then remove all scale keys
-			for ( int32 TrackIndex=0; TrackIndex<RawAnimationData.Num(); TrackIndex++ )
+			// if we don't have scale, we should delete all scale keys
+			// if you have one track that has scale, we still should support scale, so compress scale
+			if(!bCompressScaleKeys)
 			{
-				FRawAnimSequenceTrack& RawData = RawAnimationData[TrackIndex];
-				RawData.ScaleKeys.Empty();
+				// then remove all scale keys
+				for(int32 TrackIndex=0; TrackIndex<RawAnimationData.Num(); TrackIndex++)
+				{
+					FRawAnimSequenceTrack& RawData = RawAnimationData[TrackIndex];
+					RawData.ScaleKeys.Empty();
+				}
 			}
 		}
-	}
 
-	CompressedTrackOffsets.Empty();
-	CompressedScaleOffsets.Empty();
+		CompressedTrackOffsets.Empty();
+		CompressedScaleOffsets.Empty();
+	}
+	
 #endif
 	return bRemovedKeys;
 }
diff --git a/Engine/Source/Runtime/Engine/Private/AnimationAsset.cpp b/Engine/Source/Runtime/Engine/Private/AnimationAsset.cpp
--- a/Engine/Source/Runtime/Engine/Private/AnimationAsset.cpp
+++ b/Engine/Source/Runtime/Engine/Private/AnimationAsset.cpp
@@ -104,7 +104,6 @@ void UAnimationAsset::PostLoad()
 
 	ValidateSkeleton();
 
-
 	check( Skeleton==NULL || SkeletonGuid.IsValid() );
 }
 
@@ -153,7 +152,7 @@ bool UAnimationAsset::ReplaceSkeleton(USkeleton* NewSkeleton, bool bConvertSpace
 			for (auto Iter = AnimSeqsToReplace.CreateIterator(); Iter; ++Iter)
 			{
 				UAnimSequence* AnimSeq = *Iter;
-				if (AnimSeq)
+				if (AnimSeq && AnimSeq->Skeleton != NewSkeleton)
 				{
 					AnimSeq->RemapTracksToNewSkeleton(NewSkeleton, bConvertSpaces);
 				}
