#ANIM : Fixed Force LOD not working when Max LOD

Merging

//depot/UE4-Orion/Engine/Source/Runtime/Engine/Classes/Components/SkinnedMeshComponent.h

to //depot/UE4-Orion/Engine/Source/Runtime/Engine/Classes/Components/SkinnedMeshComponent.h

diff --git a/Engine/Source/Runtime/Engine/Classes/Components/SkinnedMeshComponent.h b/Engine/Source/Runtime/Engine/Classes/Components/SkinnedMeshComponent.h
--- a/Engine/Source/Runtime/Engine/Classes/Components/SkinnedMeshComponent.h
+++ b/Engine/Source/Runtime/Engine/Classes/Components/SkinnedMeshComponent.h
@@ -373,7 +373,7 @@ public:
 	/**
 	 * Set MinLodModel of the mesh component
 	 *
-	 * @param	InNewMinLOD	Set new MinLodModel that make sure the LOD does not go below of this value. This will affect in the next tick update. 
+	 * @param	InNewMinLOD	Set new MinLodModel that make sure the LOD does not go below of this value. Range from [0, Max Number of LOD - 1]. This will affect in the next tick update. 
 	 */
 	UFUNCTION(BlueprintCallable, Category="Components|SkinnedMesh")
 	void SetMinLOD(int32 InNewMinLOD);
@@ -381,7 +381,7 @@ public:
 	/**
 	 * Set MinLodModel of the mesh component
 	 *
-	 * @param	InNewForcedLOD	Set new ForcedLODModel that forces to set the incoming LOD. This will affect in the next tick update. 
+	 * @param	InNewForcedLOD	Set new ForcedLODModel that forces to set the incoming LOD. Range from [1, Max Number of LOD]. This will affect in the next tick update. 
 	 */
 	UFUNCTION(BlueprintCallable, Category="Components|SkinnedMesh")
 	void SetForcedLOD(int32 InNewForcedLOD);
diff --git a/Engine/Source/Runtime/Engine/Private/Components/SkinnedMeshComponent.cpp b/Engine/Source/Runtime/Engine/Private/Components/SkinnedMeshComponent.cpp
--- a/Engine/Source/Runtime/Engine/Private/Components/SkinnedMeshComponent.cpp
+++ b/Engine/Source/Runtime/Engine/Private/Components/SkinnedMeshComponent.cpp
@@ -1811,7 +1811,7 @@ void USkinnedMeshComponent::SetForcedLOD(int32 InNewForcedLOD)
 	int32 MaxLODIndex = 0;
 	if(MeshObject)
 	{
-		MaxLODIndex = MeshObject->GetSkeletalMeshResource().LODModels.Num() - 1;
+		MaxLODIndex = MeshObject->GetSkeletalMeshResource().LODModels.Num();
 	}
 
 	ForcedLodModel = FMath::Clamp(InNewForcedLOD, 0, MaxLODIndex);
