Merging ZebraBranch CL #2675564: Fix the swizzle for the vertex attribute in curve-injection that is actually an FColor & restore a necessary alignment workaround for Metal. Force a rebuild of shaders to ensure that the alignment change takes effect.

diff --git a/Engine/Shaders/ParticleCurveInjectionShader.usf b/Engine/Shaders/ParticleCurveInjectionShader.usf
--- a/Engine/Shaders/ParticleCurveInjectionShader.usf
+++ b/Engine/Shaders/ParticleCurveInjectionShader.usf
@@ -47,7 +47,7 @@ void VertexMain(
 		);
 
 	// Pass sample on to the pixel shader.
-	Interpolants.Sample = Input.Sample;
+	Interpolants.Sample = Input.Sample FCOLOR_COMPONENT_SWIZZLE;
 }
 
 #endif // #if VERTEXSHADER
diff --git a/Engine/Shaders/ShaderVersion.usf b/Engine/Shaders/ShaderVersion.usf
--- a/Engine/Shaders/ShaderVersion.usf
+++ b/Engine/Shaders/ShaderVersion.usf
@@ -2,4 +2,4 @@
 // Each time the console command is executed it generates a new GUID. As this file is included
 // in common.usf (which should be included in any shader) it allows to invalidate the shader DDC.
 // 
-// GUID = 4C4BD3C9571F4F3A93CE647EC1EAACF6
\ No newline at end of file
+// GUID = C182C749FEB14FCD933644E0FB379E66
\ No newline at end of file
diff --git a/Engine/Source/Runtime/Engine/Private/Particles/ParticleGpuSimulation.cpp b/Engine/Source/Runtime/Engine/Private/Particles/ParticleGpuSimulation.cpp
--- a/Engine/Source/Runtime/Engine/Private/Particles/ParticleGpuSimulation.cpp
+++ b/Engine/Source/Runtime/Engine/Private/Particles/ParticleGpuSimulation.cpp
@@ -761,12 +761,12 @@ FArchive& operator<<(FArchive& Ar, FParticlePerFrameSimulationShaderParameters&
  * simulation.
  */
 BEGIN_UNIFORM_BUFFER_STRUCT( FVectorFieldUniformParameters,)
+	DECLARE_UNIFORM_BUFFER_STRUCT_MEMBER( int32, Count )
 	DECLARE_UNIFORM_BUFFER_STRUCT_MEMBER_ARRAY( FMatrix, WorldToVolume, [MAX_VECTOR_FIELDS] )
 	DECLARE_UNIFORM_BUFFER_STRUCT_MEMBER_ARRAY( FMatrix, VolumeToWorld, [MAX_VECTOR_FIELDS] )
-	DECLARE_UNIFORM_BUFFER_STRUCT_MEMBER( int32, Count )
-	DECLARE_UNIFORM_BUFFER_STRUCT_MEMBER_ARRAY( FVector, TilingAxes, [MAX_VECTOR_FIELDS] )
+	DECLARE_UNIFORM_BUFFER_STRUCT_MEMBER_ARRAY( FVector4, IntensityAndTightness, [MAX_VECTOR_FIELDS] )
 	DECLARE_UNIFORM_BUFFER_STRUCT_MEMBER_ARRAY( FVector, VolumeSize, [MAX_VECTOR_FIELDS] )
-	DECLARE_UNIFORM_BUFFER_STRUCT_MEMBER_ARRAY( FVector2D, IntensityAndTightness, [MAX_VECTOR_FIELDS] )
+	DECLARE_UNIFORM_BUFFER_STRUCT_MEMBER_ARRAY( FVector, TilingAxes, [MAX_VECTOR_FIELDS] )
 END_UNIFORM_BUFFER_STRUCT( FVectorFieldUniformParameters )
 
 IMPLEMENT_UNIFORM_BUFFER_STRUCT(FVectorFieldUniformParameters,TEXT("VectorFields"));
@@ -4156,7 +4156,7 @@ static void SetParametersForVectorField(FVectorFieldUniformParameters& OutParame
 	OutParameters.WorldToVolume[Index] = VectorFieldInstance->WorldToVolume;
 	OutParameters.VolumeToWorld[Index] = VectorFieldInstance->VolumeToWorldNoScale;
 	OutParameters.VolumeSize[Index] = FVector(Resource->SizeX, Resource->SizeY, Resource->SizeZ);
-	OutParameters.IntensityAndTightness[Index] = FVector2D(Intensity, Tightness );
+	OutParameters.IntensityAndTightness[Index] = FVector4(Intensity, Tightness, 0, 0 );
 	OutParameters.TilingAxes[Index].X = VectorFieldInstance->bTileX ? 1.0f : 0.0f;
 	OutParameters.TilingAxes[Index].Y = VectorFieldInstance->bTileY ? 1.0f : 0.0f;
 	OutParameters.TilingAxes[Index].Z = VectorFieldInstance->bTileZ ? 1.0f : 0.0f;
@@ -4238,7 +4238,7 @@ void FFXSystem::SimulateGPUParticles(
 			VectorFieldParameters.WorldToVolume[Index] = FMatrix::Identity;
 			VectorFieldParameters.VolumeToWorld[Index] = FMatrix::Identity;
 			VectorFieldParameters.VolumeSize[Index] = FVector(1.0f);
-			VectorFieldParameters.IntensityAndTightness[Index] = FVector2D::ZeroVector;
+			VectorFieldParameters.IntensityAndTightness[Index] = FVector4(0.0f);
 		}
 		VectorFieldParameters.Count = 0;
 		EmptyVectorFieldUniformBuffer = FVectorFieldUniformBufferRef::CreateUniformBufferImmediate(VectorFieldParameters, UniformBuffer_SingleFrame);
@@ -4312,7 +4312,7 @@ void FFXSystem::SimulateGPUParticles(
 						VectorFieldParameters.WorldToVolume[Index] = FMatrix::Identity;
 						VectorFieldParameters.VolumeToWorld[Index] = FMatrix::Identity;
 						VectorFieldParameters.VolumeSize[Index] = FVector(1.0f);
-						VectorFieldParameters.IntensityAndTightness[Index] = FVector2D::ZeroVector;
+						VectorFieldParameters.IntensityAndTightness[Index] = FVector4(0.0f);
 					}
 					SimulationCommand->VectorFieldsUniformBuffer = FVectorFieldUniformBufferRef::CreateUniformBufferImmediate(VectorFieldParameters, UniformBuffer_SingleFrame);
 				}
