Added ability for games to have custom control channel messages

Merging 2674765

//depot/UE4-UT/...

to //depot/UE4/...

diff --git a/Engine/Source/Runtime/Engine/Classes/Engine/Engine.h b/Engine/Source/Runtime/Engine/Classes/Engine/Engine.h
--- a/Engine/Source/Runtime/Engine/Classes/Engine/Engine.h
+++ b/Engine/Source/Runtime/Engine/Classes/Engine/Engine.h
@@ -2525,7 +2525,7 @@ public:
 	/** Browse to a specified URL, relative to the current one. */
 	virtual EBrowseReturnVal::Type Browse( FWorldContext& WorldContext, FURL URL, FString& Error );
 
-	void TickWorldTravel(FWorldContext& WorldContext, float DeltaSeconds);
+	virtual void TickWorldTravel(FWorldContext& WorldContext, float DeltaSeconds);
 
 	void BrowseToDefaultMap( FWorldContext& WorldContext );
 
diff --git a/Engine/Source/Runtime/Engine/Classes/Engine/GameInstance.h b/Engine/Source/Runtime/Engine/Classes/Engine/GameInstance.h
--- a/Engine/Source/Runtime/Engine/Classes/Engine/GameInstance.h
+++ b/Engine/Source/Runtime/Engine/Classes/Engine/GameInstance.h
@@ -206,6 +206,13 @@ public:
 	 */
 	virtual void AddUserToReplay(const FString& UserString);
 
+	/** handle a game specific net control message (NMT_GameSpecific)
+	 * this allows games to insert their own logic into the control channel
+	 * the meaning of both data parameters is game-specific
+	 */
+	virtual void HandleGameNetControlMessage(class UNetConnection* Connection, uint8 MessageByte, const FString& MessageStr)
+	{}
+
 	FTimerManager* TimerManager;
 
 	/** @return online session management object associated with this game instance */
diff --git a/Engine/Source/Runtime/Engine/Classes/GameFramework/GameMode.h b/Engine/Source/Runtime/Engine/Classes/GameFramework/GameMode.h
--- a/Engine/Source/Runtime/Engine/Classes/GameFramework/GameMode.h
+++ b/Engine/Source/Runtime/Engine/Classes/GameFramework/GameMode.h
@@ -640,8 +640,15 @@ public:
 	virtual void GetSeamlessTravelActorList(bool bToEntry, TArray<AActor*>& ActorList);
 
 	/** Allow the game to specify a place for clients to download MapName */
+	DEPRECATED(4.10, "Override GameWelcomePlayer() instead")
 	virtual FString GetRedirectURL(const FString& MapName) const;
 
+	/** allows game to send network messages to provide more information to the client joining the game via NMT_GameSpecific
+	 * (for example required DLC)
+	 * the out string RedirectURL is built in and send automatically if only a simple URL is needed
+	 */
+	virtual void GameWelcomePlayer(UNetConnection* Connection, FString& RedirectURL);
+
 	/** 
 	 * used to swap a viewport/connection's PlayerControllers when seamless traveling and the new GameMode's
 	 * controller class is different than the previous
diff --git a/Engine/Source/Runtime/Engine/Private/DataChannel.cpp b/Engine/Source/Runtime/Engine/Private/DataChannel.cpp
--- a/Engine/Source/Runtime/Engine/Private/DataChannel.cpp
+++ b/Engine/Source/Runtime/Engine/Private/DataChannel.cpp
@@ -1180,6 +1180,17 @@ void UControlChannel::ReceivedBunch( FInBunch& Bunch )
 
 			}
 		}
+		else if (MessageType == NMT_GameSpecific)
+		{
+			// the most common Notify handlers do not support subclasses by default and so we redirect the game specific messaging to the GameInstance instead
+			uint8 MessageByte;
+			FString MessageStr;
+			FNetControlMessage<NMT_GameSpecific>::Receive(Bunch, MessageByte, MessageStr);
+			if (Connection->Driver->World != NULL && Connection->Driver->World->GetGameInstance() != NULL)
+			{
+				Connection->Driver->World->GetGameInstance()->HandleGameNetControlMessage(Connection, MessageByte, MessageStr);
+			}
+		}
 		else if(MessageType == NMT_SecurityViolation)
 		{
 			FString DebugMessage;
diff --git a/Engine/Source/Runtime/Engine/Private/GameMode.cpp b/Engine/Source/Runtime/Engine/Private/GameMode.cpp
--- a/Engine/Source/Runtime/Engine/Private/GameMode.cpp
+++ b/Engine/Source/Runtime/Engine/Private/GameMode.cpp
@@ -1680,6 +1680,12 @@ FString AGameMode::GetRedirectURL(const FString& MapName) const
 {
 	return FString();
 }
+void AGameMode::GameWelcomePlayer(UNetConnection* Connection, FString& RedirectURL)
+{
+	PRAGMA_DISABLE_DEPRECATION_WARNINGS
+	RedirectURL = GetRedirectURL(Connection->ClientWorldPackageName.ToString());
+	PRAGMA_ENABLE_DEPRECATION_WARNINGS
+}
 
 bool AGameMode::IsHandlingReplays()
 {
diff --git a/Engine/Source/Runtime/Engine/Private/World.cpp b/Engine/Source/Runtime/Engine/Private/World.cpp
--- a/Engine/Source/Runtime/Engine/Private/World.cpp
+++ b/Engine/Source/Runtime/Engine/Private/World.cpp
@@ -3743,7 +3743,7 @@ void UWorld::WelcomePlayer(UNetConnection* Connection)
 	if (AuthorityGameMode != NULL)
 	{
 		GameName = AuthorityGameMode->GetClass()->GetPathName();
-		RedirectURL = AuthorityGameMode->GetRedirectURL(LevelName);
+		AuthorityGameMode->GameWelcomePlayer(Connection, RedirectURL);
 	}
 
 	FNetControlMessage<NMT_Welcome>::Send(Connection, LevelName, GameName, RedirectURL);
diff --git a/Engine/Source/Runtime/Engine/Public/Net/DataChannel.h b/Engine/Source/Runtime/Engine/Public/Net/DataChannel.h
--- a/Engine/Source/Runtime/Engine/Public/Net/DataChannel.h
+++ b/Engine/Source/Runtime/Engine/Public/Net/DataChannel.h
@@ -386,6 +386,7 @@ DEFINE_CONTROL_CHANNEL_MESSAGE_ONEPARAM(ActorChannelFailure, 16, int32); // clie
 DEFINE_CONTROL_CHANNEL_MESSAGE_ONEPARAM(DebugText, 17, FString); // debug text sent to all clients or to server
 DEFINE_CONTROL_CHANNEL_MESSAGE_TWOPARAM(NetGUIDAssign, 18, FNetworkGUID, FString); // Explicit NetworkGUID assignment. This is rare and only happens if a netguid is only serialized client->server (this msg goes server->client to tell client what ID to use in that case)
 DEFINE_CONTROL_CHANNEL_MESSAGE_ONEPARAM(SecurityViolation, 19, FString); // server tells client that it has violated security and has been disconnected
+DEFINE_CONTROL_CHANNEL_MESSAGE_TWOPARAM(GameSpecific, 20, uint8, FString); // custom game-specific message routed to UGameInstance for processing
 
 // 			Beacon control channel flow
 // Client												Server
