Missing WITH_EDITOR define

diff --git a/Engine/Source/Runtime/Engine/Private/LODActor.cpp b/Engine/Source/Runtime/Engine/Private/LODActor.cpp
--- a/Engine/Source/Runtime/Engine/Private/LODActor.cpp
+++ b/Engine/Source/Runtime/Engine/Private/LODActor.cpp
@@ -28,8 +28,10 @@ static void ToggleHLODEnabled(UWorld* InWorld)
 		ALODActor* LODActor = Cast<ALODActor>(Actor);
 		if (LODActor)
 		{
-			LODActor->SetHidden(bHLODEnabled);
+			LODActor->SetActorHiddenInGame(bHLODEnabled);
+#if WITH_EDITOR
 			LODActor->SetIsTemporarilyHiddenInEditor(bHLODEnabled);
+#endif // WITH_EDITOR
 			LODActor->MarkComponentsRenderStateDirty();			
 		}
 	}
