Integrated from Orion:

CL2679029
	fixed potential Subsurface render artifact, will integrate into main

diff --git a/Engine/Source/Runtime/Engine/Private/Rendering/SubsurfaceProfile.cpp b/Engine/Source/Runtime/Engine/Private/Rendering/SubsurfaceProfile.cpp
--- a/Engine/Source/Runtime/Engine/Private/Rendering/SubsurfaceProfile.cpp
+++ b/Engine/Source/Runtime/Engine/Private/Rendering/SubsurfaceProfile.cpp
@@ -40,7 +40,9 @@ int32 FSubsurfaceProfileTexture::AddProfile(const FSubsurfaceProfileStruct Setti
 		{
 			if (SubsurfaceProfileEntries[i].Profile == 0)
 			{
-				RetAllocationId = i; break;
+				RetAllocationId = i;
+				SubsurfaceProfileEntries[RetAllocationId].Profile = InProfile;
+				break;
 			}
 		}
 
