#ue4 - Prevent OnComponentBeginOverlap from triggering more than once in some cases, usually when the event tries to disable collision or overlaps.

Related to CL 2648081.

diff --git a/Engine/Source/Runtime/Engine/Private/Components/PrimitiveComponent.cpp b/Engine/Source/Runtime/Engine/Private/Components/PrimitiveComponent.cpp
--- a/Engine/Source/Runtime/Engine/Private/Components/PrimitiveComponent.cpp
+++ b/Engine/Source/Runtime/Engine/Private/Components/PrimitiveComponent.cpp
@@ -105,6 +105,16 @@ private:
 	const AActor& MyOwner;
 };
 
+
+FORCEINLINE_DEBUGGABLE static bool CanComponentsGenerateOverlap(const UPrimitiveComponent* MyComponent, /*const*/ UPrimitiveComponent* OtherComp)
+{
+	return OtherComp
+		&& OtherComp->bGenerateOverlapEvents
+		&& MyComponent
+		&& MyComponent->bGenerateOverlapEvents
+		&& MyComponent->GetCollisionResponseToComponent(OtherComp) == ECR_Overlap;
+}
+
 // Predicate to remove components from overlaps array that can no longer overlap
 struct FPredicateFilterCannotOverlap
 {
@@ -113,14 +123,9 @@ struct FPredicateFilterCannotOverlap
 	{
 	}
 
-	// return true if not collideable
 	bool operator() (const FOverlapInfo& Info) const
 	{
-		UPrimitiveComponent* OtherComp = Info.OverlapInfo.GetComponent();
-		return !OtherComp
-			|| !OtherComp->bGenerateOverlapEvents
-			|| !MyComponent.bGenerateOverlapEvents
-			|| MyComponent.GetCollisionResponseToComponent(OtherComp) != ECR_Overlap;
+		return !CanComponentsGenerateOverlap(&MyComponent, Info.OverlapInfo.GetComponent());
 	}
 
 private:
@@ -1991,7 +1996,7 @@ void UPrimitiveComponent::BeginComponentOverlap(const FOverlapInfo& OtherOverlap
 	if (OtherComp)
 	{
 		bool const bComponentsAlreadyTouching = IsOverlappingComponent(OtherOverlap);
-		if (!bComponentsAlreadyTouching)
+		if (!bComponentsAlreadyTouching && CanComponentsGenerateOverlap(this, OtherComp))
 		{
 			AActor* const OtherActor = OtherComp->GetOwner();
 			AActor* const MyActor = GetOwner();
