Added more montage stats to UAnimInstance

diff --git a/Engine/Source/Runtime/Engine/Private/Animation/AnimInstance.cpp b/Engine/Source/Runtime/Engine/Private/Animation/AnimInstance.cpp
--- a/Engine/Source/Runtime/Engine/Private/Animation/AnimInstance.cpp
+++ b/Engine/Source/Runtime/Engine/Private/Animation/AnimInstance.cpp
@@ -43,6 +43,11 @@ DEFINE_STAT(STAT_TickUpdateRate);
 DEFINE_STAT(STAT_BlueprintUpdateAnimation);
 DEFINE_STAT(STAT_BlueprintPostEvaluateAnimation);
 DEFINE_STAT(STAT_NativeUpdateAnimation);
+DEFINE_STAT(STAT_Montage_Advance);
+DEFINE_STAT(STAT_Montage_UpdateWeight);
+DEFINE_STAT(STAT_AnimMontageInstance_Advance);
+DEFINE_STAT(STAT_AnimMontageInstance_TickBranchPoints);
+DEFINE_STAT(STAT_AnimMontageInstance_Advance_Iteration);
 
 DECLARE_CYCLE_STAT_EXTERN(TEXT("Anim Init Time"), STAT_AnimInitTime, STATGROUP_Anim, );
 DEFINE_STAT(STAT_AnimInitTime);
@@ -1758,7 +1763,7 @@ FName UAnimInstance::GetCurrentStateName(int32 MachineIndex)
 
 void UAnimInstance::Montage_UpdateWeight(float DeltaSeconds)
 {
-	QUICK_SCOPE_CYCLE_COUNTER(STAT_UAnimInstance_Montage_UpdateWeight);
+	SCOPE_CYCLE_COUNTER(STAT_Montage_UpdateWeight);
 
 	// go through all montage instances, and update them
 	// and make sure their weight is updated properly
@@ -1773,7 +1778,7 @@ void UAnimInstance::Montage_UpdateWeight(float DeltaSeconds)
 
 void UAnimInstance::Montage_Advance(float DeltaSeconds)
 {
-	QUICK_SCOPE_CYCLE_COUNTER(STAT_UAnimInstance_Montage_Advance);
+	SCOPE_CYCLE_COUNTER(STAT_Montage_Advance);
 
 	// We're about to tick montages, queue their events to they're triggered after batched anim notifies.
 	bQueueMontageEvents = true;
diff --git a/Engine/Source/Runtime/Engine/Private/Animation/AnimMontage.cpp b/Engine/Source/Runtime/Engine/Private/Animation/AnimMontage.cpp
--- a/Engine/Source/Runtime/Engine/Private/Animation/AnimMontage.cpp
+++ b/Engine/Source/Runtime/Engine/Private/Animation/AnimMontage.cpp
@@ -1285,6 +1285,8 @@ bool FAnimMontageInstance::SimulateAdvance(float DeltaTime, float& InOutPosition
 
 void FAnimMontageInstance::Advance(float DeltaTime, struct FRootMotionMovementParams * OutRootMotionParams, bool bBlendRootMotion)
 {
+	SCOPE_CYCLE_COUNTER(STAT_AnimMontageInstance_Advance);
+
 	if( IsValid() )
 	{
 #if WITH_EDITOR
@@ -1311,6 +1313,8 @@ void FAnimMontageInstance::Advance(float DeltaTime, struct FRootMotionMovementPa
 
 			while( bPlaying && (FMath::Abs(DesiredDeltaMove) > KINDA_SMALL_NUMBER) && ((OriginalMoveDelta * DesiredDeltaMove) > 0.f) )
 			{
+				SCOPE_CYCLE_COUNTER(STAT_AnimMontageInstance_Advance_Iteration);
+
 				// Get position relative to current montage section.
 				float PosInSection;
 				int32 CurrentSectionIndex = Montage->GetAnimCompositeSectionIndexFromPos(Position, PosInSection);
@@ -1449,6 +1453,8 @@ void FAnimMontageInstance::Advance(float DeltaTime, struct FRootMotionMovementPa
 
 	if (!bInterrupted)
 	{
+		SCOPE_CYCLE_COUNTER(STAT_AnimMontageInstance_TickBranchPoints);
+
 		// Tick all active state branching points
 		for (int32 Index = 0; Index < ActiveStateBranchingPoints.Num(); Index++)
 		{
diff --git a/Engine/Source/Runtime/Engine/Public/Animation/AnimStats.h b/Engine/Source/Runtime/Engine/Public/Animation/AnimStats.h
--- a/Engine/Source/Runtime/Engine/Public/Animation/AnimStats.h
+++ b/Engine/Source/Runtime/Engine/Public/Animation/AnimStats.h
@@ -30,6 +30,11 @@ DECLARE_CYCLE_STAT_EXTERN(TEXT("AnimGameThreadTime"), STAT_AnimGameThreadTime, S
 DECLARE_CYCLE_STAT_EXTERN(TEXT("BlueprintUpdateAnimation"), STAT_BlueprintUpdateAnimation, STATGROUP_Anim, );
 DECLARE_CYCLE_STAT_EXTERN(TEXT("NativeUpdateAnimation"), STAT_NativeUpdateAnimation, STATGROUP_Anim, );
 DECLARE_CYCLE_STAT_EXTERN(TEXT("BlueprintPostEvaluateAnimation"), STAT_BlueprintPostEvaluateAnimation, STATGROUP_Anim, );
+DECLARE_CYCLE_STAT_EXTERN(TEXT("MontageAdvance"), STAT_Montage_Advance, STATGROUP_Anim, );
+DECLARE_CYCLE_STAT_EXTERN(TEXT("MontageUpdateWeight"), STAT_Montage_UpdateWeight, STATGROUP_Anim, );
+DECLARE_CYCLE_STAT_EXTERN(TEXT("AnimMontageInstance_Advance"), STAT_AnimMontageInstance_Advance, STATGROUP_Anim, );
+DECLARE_CYCLE_STAT_EXTERN(TEXT("AnimMontageInstance_TickBranchPoints"), STAT_AnimMontageInstance_TickBranchPoints, STATGROUP_Anim, );
+DECLARE_CYCLE_STAT_EXTERN(TEXT("AnimMontageInstance_Advance_Iteration"), STAT_AnimMontageInstance_Advance_Iteration, STATGROUP_Anim, );
 
 #define DO_ANIMSTAT_PROCESSING(StatName) DECLARE_CYCLE_STAT_EXTERN(TEXT(#StatName), STAT_ ## StatName, STATGROUP_Anim, ENGINE_API)
 #include "AnimMTStats.h"
