Fixed issue with rebuilding the collision mesh on SplineMeshComponents when spline parameters are changed. Also added a new Blueprint callable method to force an update to the mesh, and added it as an optional parameter to spline parameter mutator calls so that mesh / collision update can be deferred if desired.

diff --git a/Engine/Source/Runtime/Engine/Classes/Components/SplineMeshComponent.h b/Engine/Source/Runtime/Engine/Classes/Components/SplineMeshComponent.h
--- a/Engine/Source/Runtime/Engine/Classes/Components/SplineMeshComponent.h
+++ b/Engine/Source/Runtime/Engine/Classes/Components/SplineMeshComponent.h
@@ -134,6 +134,10 @@ class ENGINE_API USplineMeshComponent : public UStaticMeshComponent, public IInt
 	uint32 bSelected:1;
 #endif
 
+	// Indicates that the mesh needs updating
+	UPROPERTY(transient)
+	uint32 bMeshDirty:1;
+
 	//Begin UObject Interface
 	virtual void Serialize(FArchive& Ar) override;
 	virtual bool Modify(bool bAlwaysMarkDirty = true) override;
@@ -173,8 +177,15 @@ class ENGINE_API USplineMeshComponent : public UStaticMeshComponent, public IInt
 	virtual void GetMeshId(FString& OutMeshId) override;
 	// End Interface_CollisionDataProvider Interface
 
-	/** Called when spline params are changed, to notify render thread and possibly collision */
-	void MarkSplineParamsDirty();
+	/** Called to notify render thread and possibly collision of a change in spline params or mesh */
+	void UpdateRenderStateAndCollision();
+
+	DEPRECATED(4.10, "This method has been renamed to UpdateRenderStateAndCollision, but use of UpdateMesh is preferred")
+	void MarkSplineParamsDirty() { UpdateRenderStateAndCollision(); }
+
+	/** Update the collision and render state on the spline mesh following changes to its geometry */
+	UFUNCTION(BlueprintCallable, Category = SplineMesh)
+	void UpdateMesh();
 
 	/** Get the start position of spline in local space */
 	UFUNCTION(BlueprintCallable, Category = SplineMesh)
@@ -182,7 +193,7 @@ class ENGINE_API USplineMeshComponent : public UStaticMeshComponent, public IInt
 
 	/** Set the start position of spline in local space */
 	UFUNCTION(BlueprintCallable, Category=SplineMesh)
-	void SetStartPosition(FVector StartPos);
+	void SetStartPosition(FVector StartPos, bool bUpdateMesh = true);
 
 	/** Get the start tangent vector of spline in local space */
 	UFUNCTION(BlueprintCallable, Category = SplineMesh)
@@ -190,7 +201,7 @@ class ENGINE_API USplineMeshComponent : public UStaticMeshComponent, public IInt
 
 	/** Set the start tangent vector of spline in local space */
 	UFUNCTION(BlueprintCallable, Category = SplineMesh)
-	void SetStartTangent(FVector StartTangent);
+	void SetStartTangent(FVector StartTangent, bool bUpdateMesh = true);
 
 	/** Get the end position of spline in local space */
 	UFUNCTION(BlueprintCallable, Category = SplineMesh)
@@ -198,7 +209,7 @@ class ENGINE_API USplineMeshComponent : public UStaticMeshComponent, public IInt
 
 	/** Set the end position of spline in local space */
 	UFUNCTION(BlueprintCallable, Category = SplineMesh)
-	void SetEndPosition(FVector EndPos);
+	void SetEndPosition(FVector EndPos, bool bUpdateMesh = true);
 
 	/** Get the end tangent vector of spline in local space */
 	UFUNCTION(BlueprintCallable, Category = SplineMesh)
@@ -206,11 +217,11 @@ class ENGINE_API USplineMeshComponent : public UStaticMeshComponent, public IInt
 
 	/** Set the end tangent vector of spline in local space */
 	UFUNCTION(BlueprintCallable, Category = SplineMesh)
-	void SetEndTangent(FVector EndTangent);
+	void SetEndTangent(FVector EndTangent, bool bUpdateMesh = true);
 
 	/** Set the start and end, position and tangent, all in local space */
 	UFUNCTION(BlueprintCallable, Category = SplineMesh)
-	void SetStartAndEnd(FVector StartPos, FVector StartTangent, FVector EndPos, FVector EndTangent);
+	void SetStartAndEnd(FVector StartPos, FVector StartTangent, FVector EndPos, FVector EndTangent, bool bUpdateMesh = true);
 
 	/** Get the start scaling */
 	UFUNCTION(BlueprintCallable, Category = SplineMesh)
@@ -218,7 +229,7 @@ class ENGINE_API USplineMeshComponent : public UStaticMeshComponent, public IInt
 
 	/** Set the start scaling */
 	UFUNCTION(BlueprintCallable, Category = SplineMesh)
-	void SetStartScale(FVector2D StartScale = FVector2D(1,1));
+	void SetStartScale(FVector2D StartScale = FVector2D(1,1), bool bUpdateMesh = true);
 
 	/** Get the start roll */
 	UFUNCTION(BlueprintCallable, Category = SplineMesh)
@@ -226,7 +237,7 @@ class ENGINE_API USplineMeshComponent : public UStaticMeshComponent, public IInt
 
 	/** Set the start roll */
 	UFUNCTION(BlueprintCallable, Category = SplineMesh)
-	void SetStartRoll(float StartRoll);
+	void SetStartRoll(float StartRoll, bool bUpdateMesh = true);
 
 	/** Get the start offset */
 	UFUNCTION(BlueprintCallable, Category = SplineMesh)
@@ -234,7 +245,7 @@ class ENGINE_API USplineMeshComponent : public UStaticMeshComponent, public IInt
 
 	/** Set the start offset */
 	UFUNCTION(BlueprintCallable, Category = SplineMesh)
-	void SetStartOffset(FVector2D StartOffset);
+	void SetStartOffset(FVector2D StartOffset, bool bUpdateMesh = true);
 
 	/** Get the end scaling */
 	UFUNCTION(BlueprintCallable, Category = SplineMesh)
@@ -242,7 +253,7 @@ class ENGINE_API USplineMeshComponent : public UStaticMeshComponent, public IInt
 
 	/** Set the end scaling */
 	UFUNCTION(BlueprintCallable, Category = SplineMesh)
-	void SetEndScale(FVector2D EndScale = FVector2D(1,1));
+	void SetEndScale(FVector2D EndScale = FVector2D(1,1), bool bUpdateMesh = true);
 
 	/** Get the end roll */
 	UFUNCTION(BlueprintCallable, Category = SplineMesh)
@@ -250,7 +261,7 @@ class ENGINE_API USplineMeshComponent : public UStaticMeshComponent, public IInt
 
 	/** Set the end roll */
 	UFUNCTION(BlueprintCallable, Category = SplineMesh)
-	void SetEndRoll(float EndRoll);
+	void SetEndRoll(float EndRoll, bool bUpdateMesh = true);
 
 	/** Get the end offset */
 	UFUNCTION(BlueprintCallable, Category = SplineMesh)
@@ -258,7 +269,7 @@ class ENGINE_API USplineMeshComponent : public UStaticMeshComponent, public IInt
 
 	/** Set the end offset */
 	UFUNCTION(BlueprintCallable, Category = SplineMesh)
-	void SetEndOffset(FVector2D EndOffset);
+	void SetEndOffset(FVector2D EndOffset, bool bUpdateMesh = true);
 
 	/** Get the forward axis */
 	UFUNCTION(BlueprintCallable, Category = SplineMesh)
@@ -266,7 +277,7 @@ class ENGINE_API USplineMeshComponent : public UStaticMeshComponent, public IInt
 
 	/** Set the forward axis */
 	UFUNCTION(BlueprintCallable, Category = SplineMesh)
-	void SetForwardAxis(ESplineMeshAxis::Type InForwardAxis);
+	void SetForwardAxis(ESplineMeshAxis::Type InForwardAxis, bool bUpdateMesh = true);
 
 	/** Get the spline up direction */
 	UFUNCTION(BlueprintCallable, Category = SplineMesh)
@@ -274,7 +285,7 @@ class ENGINE_API USplineMeshComponent : public UStaticMeshComponent, public IInt
 
 	/** Set the spline up direction */
 	UFUNCTION(BlueprintCallable, Category = SplineMesh)
-	void SetSplineUpDir(const FVector& InSplineUpDir);
+	void SetSplineUpDir(const FVector& InSplineUpDir, bool bUpdateMesh = true);
 
 	/** Get the boundary min */
 	UFUNCTION(BlueprintCallable, Category = SplineMesh)
@@ -282,7 +293,7 @@ class ENGINE_API USplineMeshComponent : public UStaticMeshComponent, public IInt
 
 	/** Set the boundary min */
 	UFUNCTION(BlueprintCallable, Category = SplineMesh)
-	void SetBoundaryMin(float InBoundaryMin);
+	void SetBoundaryMin(float InBoundaryMin, bool bUpdateMesh = true);
 
 	/** Get the boundary max */
 	UFUNCTION(BlueprintCallable, Category = SplineMesh)
@@ -290,7 +301,7 @@ class ENGINE_API USplineMeshComponent : public UStaticMeshComponent, public IInt
 
 	/** Set the boundary max */
 	UFUNCTION(BlueprintCallable, Category = SplineMesh)
-	void SetBoundaryMax(float InBoundaryMax);
+	void SetBoundaryMax(float InBoundaryMax, bool bUpdateMesh = true);
 
 	// Destroys the body setup, used to clear collision if the mesh goes missing
 	void DestroyBodySetup();
diff --git a/Engine/Source/Runtime/Engine/Private/Components/SplineMeshComponent.cpp b/Engine/Source/Runtime/Engine/Private/Components/SplineMeshComponent.cpp
--- a/Engine/Source/Runtime/Engine/Private/Components/SplineMeshComponent.cpp
+++ b/Engine/Source/Runtime/Engine/Private/Components/SplineMeshComponent.cpp
@@ -244,6 +244,8 @@ USplineMeshComponent::USplineMeshComponent(const FObjectInitializer& ObjectIniti
 
 	SplineBoundaryMin = 0;
 	SplineBoundaryMax = 0;
+
+	bMeshDirty = false;
 }
 
 FVector USplineMeshComponent::GetStartPosition() const
@@ -251,10 +253,14 @@ FVector USplineMeshComponent::GetStartPosition() const
 	return SplineParams.StartPos;
 }
 
-void USplineMeshComponent::SetStartPosition(FVector StartPos)
+void USplineMeshComponent::SetStartPosition(FVector StartPos, bool bUpdateMesh)
 {
 	SplineParams.StartPos = StartPos;
-	MarkSplineParamsDirty();
+	bMeshDirty = true;
+	if (bUpdateMesh)
+	{
+		UpdateRenderStateAndCollision();
+	}
 }
 
 FVector USplineMeshComponent::GetStartTangent() const
@@ -262,10 +268,14 @@ FVector USplineMeshComponent::GetStartTangent() const
 	return SplineParams.StartTangent;
 }
 
-void USplineMeshComponent::SetStartTangent(FVector StartTangent)
+void USplineMeshComponent::SetStartTangent(FVector StartTangent, bool bUpdateMesh)
 {
 	SplineParams.StartTangent = StartTangent;
-	MarkSplineParamsDirty();
+	bMeshDirty = true;
+	if (bUpdateMesh)
+	{
+		UpdateRenderStateAndCollision();
+	}
 }
 
 FVector USplineMeshComponent::GetEndPosition() const
@@ -273,10 +283,14 @@ FVector USplineMeshComponent::GetEndPosition() const
 	return SplineParams.EndPos;
 }
 
-void USplineMeshComponent::SetEndPosition(FVector EndPos)
+void USplineMeshComponent::SetEndPosition(FVector EndPos, bool bUpdateMesh)
 {
 	SplineParams.EndPos = EndPos;
-	MarkSplineParamsDirty();
+	bMeshDirty = true;
+	if (bUpdateMesh)
+	{
+		UpdateRenderStateAndCollision();
+	}
 }
 
 FVector USplineMeshComponent::GetEndTangent() const
@@ -284,19 +298,27 @@ FVector USplineMeshComponent::GetEndTangent() const
 	return SplineParams.EndTangent;
 }
 
-void USplineMeshComponent::SetEndTangent(FVector EndTangent)
+void USplineMeshComponent::SetEndTangent(FVector EndTangent, bool bUpdateMesh)
 {
 	SplineParams.EndTangent = EndTangent;
-	MarkSplineParamsDirty();
+	bMeshDirty = true;
+	if (bUpdateMesh)
+	{
+		UpdateRenderStateAndCollision();
+	}
 }
 
-void USplineMeshComponent::SetStartAndEnd(FVector StartPos, FVector StartTangent, FVector EndPos, FVector EndTangent)
+void USplineMeshComponent::SetStartAndEnd(FVector StartPos, FVector StartTangent, FVector EndPos, FVector EndTangent, bool bUpdateMesh)
 {
 	SplineParams.StartPos = StartPos;
 	SplineParams.StartTangent = StartTangent;
 	SplineParams.EndPos = EndPos;
 	SplineParams.EndTangent = EndTangent;
-	MarkSplineParamsDirty();
+	bMeshDirty = true;
+	if (bUpdateMesh)
+	{
+		UpdateRenderStateAndCollision();
+	}
 }
 
 FVector2D USplineMeshComponent::GetStartScale() const
@@ -304,10 +326,14 @@ FVector2D USplineMeshComponent::GetStartScale() const
 	return SplineParams.StartScale;
 }
 
-void USplineMeshComponent::SetStartScale(FVector2D StartScale)
+void USplineMeshComponent::SetStartScale(FVector2D StartScale, bool bUpdateMesh)
 {
 	SplineParams.StartScale = StartScale;
-	MarkSplineParamsDirty();
+	bMeshDirty = true;
+	if (bUpdateMesh)
+	{
+		UpdateRenderStateAndCollision();
+	}
 }
 
 float USplineMeshComponent::GetStartRoll() const
@@ -315,10 +341,14 @@ float USplineMeshComponent::GetStartRoll() const
 	return SplineParams.StartRoll;
 }
 
-void USplineMeshComponent::SetStartRoll(float StartRoll)
+void USplineMeshComponent::SetStartRoll(float StartRoll, bool bUpdateMesh)
 {
 	SplineParams.StartRoll = StartRoll;
-	MarkSplineParamsDirty();
+	bMeshDirty = true;
+	if (bUpdateMesh)
+	{
+		UpdateRenderStateAndCollision();
+	}
 }
 
 FVector2D USplineMeshComponent::GetStartOffset() const
@@ -326,10 +356,14 @@ FVector2D USplineMeshComponent::GetStartOffset() const
 	return SplineParams.StartOffset;
 }
 
-void USplineMeshComponent::SetStartOffset(FVector2D StartOffset)
+void USplineMeshComponent::SetStartOffset(FVector2D StartOffset, bool bUpdateMesh)
 {
 	SplineParams.StartOffset = StartOffset;
-	MarkSplineParamsDirty();
+	bMeshDirty = true;
+	if (bUpdateMesh)
+	{
+		UpdateRenderStateAndCollision();
+	}
 }
 
 FVector2D USplineMeshComponent::GetEndScale() const
@@ -337,10 +371,14 @@ FVector2D USplineMeshComponent::GetEndScale() const
 	return SplineParams.EndScale;
 }
 
-void USplineMeshComponent::SetEndScale(FVector2D EndScale)
+void USplineMeshComponent::SetEndScale(FVector2D EndScale, bool bUpdateMesh)
 {
 	SplineParams.EndScale = EndScale;
-	MarkSplineParamsDirty();
+	bMeshDirty = true;
+	if (bUpdateMesh)
+	{
+		UpdateRenderStateAndCollision();
+	}
 }
 
 float USplineMeshComponent::GetEndRoll() const
@@ -348,10 +386,14 @@ float USplineMeshComponent::GetEndRoll() const
 	return SplineParams.EndRoll;
 }
 
-void USplineMeshComponent::SetEndRoll(float EndRoll)
+void USplineMeshComponent::SetEndRoll(float EndRoll, bool bUpdateMesh)
 {
 	SplineParams.EndRoll = EndRoll;
-	MarkSplineParamsDirty();
+	bMeshDirty = true;
+	if (bUpdateMesh)
+	{
+		UpdateRenderStateAndCollision();
+	}
 }
 
 FVector2D USplineMeshComponent::GetEndOffset() const
@@ -359,10 +401,14 @@ FVector2D USplineMeshComponent::GetEndOffset() const
 	return SplineParams.EndOffset;
 }
 
-void USplineMeshComponent::SetEndOffset(FVector2D EndOffset)
+void USplineMeshComponent::SetEndOffset(FVector2D EndOffset, bool bUpdateMesh)
 {
 	SplineParams.EndOffset = EndOffset;
-	MarkSplineParamsDirty();
+	bMeshDirty = true;
+	if (bUpdateMesh)
+	{
+		UpdateRenderStateAndCollision();
+	}
 }
 
 ESplineMeshAxis::Type USplineMeshComponent::GetForwardAxis() const
@@ -370,10 +416,14 @@ ESplineMeshAxis::Type USplineMeshComponent::GetForwardAxis() const
 	return ForwardAxis;
 }
 
-void USplineMeshComponent::SetForwardAxis(ESplineMeshAxis::Type InForwardAxis)
+void USplineMeshComponent::SetForwardAxis(ESplineMeshAxis::Type InForwardAxis, bool bUpdateMesh)
 {
 	ForwardAxis = InForwardAxis;
-	MarkSplineParamsDirty();
+	bMeshDirty = true;
+	if (bUpdateMesh)
+	{
+		UpdateRenderStateAndCollision();
+	}
 }
 
 FVector USplineMeshComponent::GetSplineUpDir() const
@@ -381,10 +431,14 @@ FVector USplineMeshComponent::GetSplineUpDir() const
 	return SplineUpDir;
 }
 
-void USplineMeshComponent::SetSplineUpDir(const FVector& InSplineUpDir)
+void USplineMeshComponent::SetSplineUpDir(const FVector& InSplineUpDir, bool bUpdateMesh)
 {
 	SplineUpDir = InSplineUpDir.GetSafeNormal();
-	MarkSplineParamsDirty();
+	bMeshDirty = true;
+	if (bUpdateMesh)
+	{
+		UpdateRenderStateAndCollision();
+	}
 }
 
 float USplineMeshComponent::GetBoundaryMin() const
@@ -392,10 +446,14 @@ float USplineMeshComponent::GetBoundaryMin() const
 	return SplineBoundaryMin;
 }
 
-void USplineMeshComponent::SetBoundaryMin(float InBoundaryMin)
+void USplineMeshComponent::SetBoundaryMin(float InBoundaryMin, bool bUpdateMesh)
 {
 	SplineBoundaryMin = InBoundaryMin;
-	MarkSplineParamsDirty();
+	bMeshDirty = true;
+	if (bUpdateMesh)
+	{
+		UpdateRenderStateAndCollision();
+	}
 }
 
 float USplineMeshComponent::GetBoundaryMax() const
@@ -403,23 +461,37 @@ float USplineMeshComponent::GetBoundaryMax() const
 	return SplineBoundaryMax;
 }
 
-void USplineMeshComponent::SetBoundaryMax(float InBoundaryMax)
+void USplineMeshComponent::SetBoundaryMax(float InBoundaryMax, bool bUpdateMesh)
 {
 	SplineBoundaryMax = InBoundaryMax;
-	MarkSplineParamsDirty();
+	bMeshDirty = true;
+	if (bUpdateMesh)
+	{
+		UpdateRenderStateAndCollision();
+	}
 }
 
+void USplineMeshComponent::UpdateMesh()
+{
+	if (bMeshDirty)
+	{
+		UpdateRenderStateAndCollision();
+	}
+}
 
-void USplineMeshComponent::MarkSplineParamsDirty()
+void USplineMeshComponent::UpdateRenderStateAndCollision()
 {
 	MarkRenderStateDirty();
 
 #if WITH_EDITOR
 	if (!GetWorld()->AreActorsInitialized())
 	{
-		RecreateCollision();
+		CachedMeshBodySetupGuid.Invalidate();
+		RecreatePhysicsState();
 	}
 #endif // WITH_EDITOR
+
+	bMeshDirty = false;
 }
 
 void USplineMeshComponent::Serialize(FArchive& Ar)
@@ -996,7 +1068,7 @@ void USplineMeshComponent::ApplyComponentInstanceData(FSplineMeshInstanceData* S
 			SplineParams.EndPos = SplineMeshInstanceData->EndPos;
 			SplineParams.StartTangent = SplineMeshInstanceData->StartTangent;
 			SplineParams.EndTangent = SplineMeshInstanceData->EndTangent;
-			MarkSplineParamsDirty();
+			UpdateRenderStateAndCollision();
 		}
 	}
 }
