Fix CIS non-unity/Mac errors

diff --git a/Engine/Source/Runtime/Engine/Private/Animation/AnimInstance.cpp b/Engine/Source/Runtime/Engine/Private/Animation/AnimInstance.cpp
--- a/Engine/Source/Runtime/Engine/Private/Animation/AnimInstance.cpp
+++ b/Engine/Source/Runtime/Engine/Private/Animation/AnimInstance.cpp
@@ -45,9 +45,6 @@ DEFINE_STAT(STAT_BlueprintPostEvaluateAnimation);
 DEFINE_STAT(STAT_NativeUpdateAnimation);
 DEFINE_STAT(STAT_Montage_Advance);
 DEFINE_STAT(STAT_Montage_UpdateWeight);
-DEFINE_STAT(STAT_AnimMontageInstance_Advance);
-DEFINE_STAT(STAT_AnimMontageInstance_TickBranchPoints);
-DEFINE_STAT(STAT_AnimMontageInstance_Advance_Iteration);
 
 DECLARE_CYCLE_STAT_EXTERN(TEXT("Anim Init Time"), STAT_AnimInitTime, STATGROUP_Anim, );
 DEFINE_STAT(STAT_AnimInitTime);
diff --git a/Engine/Source/Runtime/Engine/Private/Animation/AnimMontage.cpp b/Engine/Source/Runtime/Engine/Private/Animation/AnimMontage.cpp
--- a/Engine/Source/Runtime/Engine/Private/Animation/AnimMontage.cpp
+++ b/Engine/Source/Runtime/Engine/Private/Animation/AnimMontage.cpp
@@ -12,6 +12,10 @@
 #include "Animation/AnimNotifies/AnimNotify.h"
 #include "Animation/AnimSequenceBase.h"
 
+DEFINE_STAT(STAT_AnimMontageInstance_Advance);
+DEFINE_STAT(STAT_AnimMontageInstance_TickBranchPoints);
+DEFINE_STAT(STAT_AnimMontageInstance_Advance_Iteration);
+
 ///////////////////////////////////////////////////////////////////////////
 //
 UAnimMontage::UAnimMontage(const FObjectInitializer& ObjectInitializer)
