Another CIS fix

diff --git a/Engine/Source/Runtime/Engine/Private/Animation/AnimMontage.cpp b/Engine/Source/Runtime/Engine/Private/Animation/AnimMontage.cpp
--- a/Engine/Source/Runtime/Engine/Private/Animation/AnimMontage.cpp
+++ b/Engine/Source/Runtime/Engine/Private/Animation/AnimMontage.cpp
@@ -11,6 +11,7 @@
 #include "Animation/AnimNotifies/AnimNotifyState.h"
 #include "Animation/AnimNotifies/AnimNotify.h"
 #include "Animation/AnimSequenceBase.h"
+#include "AnimStats.h"
 
 DEFINE_STAT(STAT_AnimMontageInstance_Advance);
 DEFINE_STAT(STAT_AnimMontageInstance_TickBranchPoints);
