#ANIM: Fix removing double calculation

Merging

//depot/UE4-Orion/Engine/Source/Runtime/Engine/Public/BonePose.h

to //depot/UE4-Orion/Engine/Source/Runtime/Engine/Public/BonePose.h

diff --git a/Engine/Source/Runtime/Engine/Public/BonePose.h b/Engine/Source/Runtime/Engine/Public/BonePose.h
--- a/Engine/Source/Runtime/Engine/Public/BonePose.h
+++ b/Engine/Source/Runtime/Engine/Public/BonePose.h
@@ -439,9 +439,9 @@ void FCSPose<PoseType>::CalculateComponentSpaceTransform(BoneIndexType BoneIndex
 	FTransform& ParentBone = Pose[ParentIndex];
 	check(!Pose[BoneIndex].ContainsNaN());
 	check(!Pose[ParentIndex].ContainsNaN());
-	FTransform Suspect = Pose[BoneIndex] * Pose[ParentIndex];
-	check(!Suspect.ContainsNaN());
-	Pose[BoneIndex] = Pose[BoneIndex] * Pose[ParentIndex];
+	FTransform ComponentTransform = Pose[BoneIndex] * Pose[ParentIndex];
+	check(!ComponentTransform.ContainsNaN());
+	Pose[BoneIndex] = ComponentTransform;
 	Pose[BoneIndex].NormalizeRotation();
 	check(!Pose[BoneIndex].ContainsNaN());
 	ComponentSpaceFlags[BoneIndex] = 1;
