#ANIM : debug pose system
Merging

//depot/UE4-Orion/Engine/Source/Runtime/Engine/...

to //depot/UE4/Engine/Source/Runtime/Engine/...

diff --git a/Engine/Source/Runtime/Engine/Classes/Animation/AnimNodeBase.h b/Engine/Source/Runtime/Engine/Classes/Animation/AnimNodeBase.h
--- a/Engine/Source/Runtime/Engine/Classes/Animation/AnimNodeBase.h
+++ b/Engine/Source/Runtime/Engine/Classes/Animation/AnimNodeBase.h
@@ -325,6 +325,8 @@ public:
 	void AttemptRelink(const FAnimationBaseContext& Context);
 };
 
+#define ENABLE_ANIMNODE_POSE_DEBUG 0
+
 /** A local-space pose link to another node */
 USTRUCT()
 struct ENGINE_API FPoseLink : public FPoseLinkBase
@@ -334,6 +336,12 @@ struct ENGINE_API FPoseLink : public FPoseLinkBase
 public:
 	// Interface
 	void Evaluate(FPoseContext& Output);
+
+#if ENABLE_ANIMNODE_POSE_DEBUG
+private:
+	// forwarded pose data from the wired node which current node's skeletal control is not applied yet
+	FCompactPose CurrentPose;
+#endif //#if ENABLE_ANIMNODE_POSE_DEBUG
 };
 
 /** A component-space pose link to another node */
diff --git a/Engine/Source/Runtime/Engine/Private/Animation/AnimNodeBase.cpp b/Engine/Source/Runtime/Engine/Private/Animation/AnimNodeBase.cpp
--- a/Engine/Source/Runtime/Engine/Private/Animation/AnimNodeBase.cpp
+++ b/Engine/Source/Runtime/Engine/Private/Animation/AnimNodeBase.cpp
@@ -131,7 +131,13 @@ void FPoseLink::Evaluate(FPoseContext& Output)
 
 	if (LinkedNode != NULL)
 	{
+#if ENABLE_ANIMNODE_POSE_DEBUG
+		CurrentPose.ResetToIdentity();
+#endif
 		LinkedNode->Evaluate(Output);
+#if ENABLE_ANIMNODE_POSE_DEBUG
+		CurrentPose = Output.Pose;
+#endif
 	}
 	else
 	{
