Big refactor of cloth to allow for running cloth in parallel and off game thread. Submitting with parallel cloth turned off to shake out any bugs from refactor, but parallel can be turned on via p.ParallelCloth 1

diff --git a/Engine/Source/Editor/UnrealEd/Classes/Animation/DebugSkelMeshComponent.h b/Engine/Source/Editor/UnrealEd/Classes/Animation/DebugSkelMeshComponent.h
--- a/Engine/Source/Editor/UnrealEd/Classes/Animation/DebugSkelMeshComponent.h
+++ b/Engine/Source/Editor/UnrealEd/Classes/Animation/DebugSkelMeshComponent.h
@@ -270,7 +270,7 @@ class UDebugSkelMeshComponent : public USkeletalMeshComponent
 	int32 FindCurrentSectionDisplayMode();
 
 	/** to avoid clothing reset while modifying properties in Persona */
-	virtual void CheckClothTeleport(float DeltaTime) override;
+	virtual void CheckClothTeleport(float DeltaTime, FClothSimulationContext& ClothSimulationContext) const override;
 
 #endif //#if WITH_APEX_CLOTHING
 
diff --git a/Engine/Source/Editor/UnrealEd/Private/Animation/DebugSkelMeshComponent.cpp b/Engine/Source/Editor/UnrealEd/Private/Animation/DebugSkelMeshComponent.cpp
--- a/Engine/Source/Editor/UnrealEd/Private/Animation/DebugSkelMeshComponent.cpp
+++ b/Engine/Source/Editor/UnrealEd/Private/Animation/DebugSkelMeshComponent.cpp
@@ -706,7 +706,7 @@ int32 UDebugSkelMeshComponent::FindCurrentSectionDisplayMode()
 	return DisplayMode;
 }
 
-void UDebugSkelMeshComponent::CheckClothTeleport(float DeltaTime)
+void UDebugSkelMeshComponent::CheckClothTeleport(float DeltaTime, FClothSimulationContext& ClothSimulationContext) const
 {
 	// do nothing to avoid clothing reset while modifying properties
 	// modifying values can cause frame delay and clothes will be reset by a large delta time (low fps)
diff --git a/Engine/Source/Runtime/Engine/Classes/Components/SkeletalMeshComponent.h b/Engine/Source/Runtime/Engine/Classes/Components/SkeletalMeshComponent.h
--- a/Engine/Source/Runtime/Engine/Classes/Components/SkeletalMeshComponent.h
+++ b/Engine/Source/Runtime/Engine/Classes/Components/SkeletalMeshComponent.h
@@ -8,12 +8,82 @@
 #include "SkeletalMeshTypes.h"
 #include "Animation/AnimationAsset.h"
 #include "AnimCurveTypes.h"
+#include "ClothSimData.h"
+#include "SingleAnimationPlayData.h"
 #include "SkeletalMeshComponent.generated.h"
 
+
+
 class UAnimInstance;
 struct FEngineShowFlags;
 struct FConvexVolume;
 
+
+
+
+namespace physx
+{ 
+	namespace apex 
+	{
+		class NxClothingAsset;
+		class NxClothingActor;
+		class NxClothingCollision;
+	}
+}
+
+class FPhysScene;
+
+/** a class to manage an APEX clothing actor */
+class FClothingActor
+{
+public:
+	enum TeleportMode
+	{
+		/** Simulation continues smoothly. This is the most commonly used mode */
+		Continuous,
+		/**
+		 * Transforms the current simulation state from the old global pose to the new global pose.
+		 * This will transform positions and velocities and thus keep the simulation state, just translate it to a new pose.
+		 */
+		Teleport,
+
+		/**
+		 * Forces the cloth to the animated position in the next frame.
+		 * This can be used to reset it from a bad state or by a teleport where the old state is not important anymore.
+		 */
+		TeleportAndReset,
+
+		/** The client hasn't specified an override. This is needed for double buffering to determine intent of game code (Do not actually pass into APEX) */
+		Default,
+	};
+
+	void Clear(bool bReleaseResource = false);
+
+	/** 
+	 * to check whether this actor is valid or not 
+	 * because clothing asset can be changed by editing 
+	 */
+	physx::apex::NxClothingAsset*	ParentClothingAsset;
+	/** APEX clothing actor is created from APEX clothing asset for cloth simulation */
+	physx::apex::NxClothingActor*		ApexClothingActor;
+};
+
+//The data that cloth needs for simulation prep in parallel. These properties are accessible via double buffer
+struct FClothSimulationContext
+{
+	/** used for pre-computation using TeleportRotationThreshold property */
+	float ClothTeleportCosineThresholdInRad;
+	/** used for pre-computation using tTeleportDistanceThreshold property */
+	float ClothTeleportDistThresholdSquared;
+private:
+	/** whether we need to teleport cloth. There are functions which allow you to modify this accordingly. Do not access directly as double buffer strategy relies on internal logic being consistent */
+	FClothingActor::TeleportMode ClothTeleportMode;
+	/** previous root bone matrix to compare the difference and decide to do clothing teleport  */
+	FMatrix	PrevRootBoneMatrix;
+
+	friend class USkeletalMeshComponent;
+};
+
 struct FAnimationEvaluationContext
 {
 	// The anim instance we are evaluating
@@ -74,61 +144,6 @@ struct FAnimationEvaluationContext
 
 };
 
-//#if WITH_APEX
-namespace physx
-{ 
-	namespace apex 
-	{
-		class NxClothingAsset;
-		class NxClothingActor;
-		class NxClothingCollision;
-	}
-}
-//#endif // WITH_APEX
-
-class FPhysScene;
-
-/** a class to manage an APEX clothing actor */
-class FClothingActor
-{
-public:
-	enum TeleportMode
-	{
-		/** Simulation continues smoothly. This is the most commonly used mode */
-		Continuous,
-		/**
-		 * Transforms the current simulation state from the old global pose to the new global pose.
-		 * This will transform positions and velocities and thus keep the simulation state, just translate it to a new pose.
-		 */
-		Teleport,
-
-		/**
-		 * Forces the cloth to the animated position in the next frame.
-		 * This can be used to reset it from a bad state or by a teleport where the old state is not important anymore.
-		 */
-		TeleportAndReset,
-	};
-
-	void Clear(bool bReleaseResource = false);
-
-	/** 
-	 * to check whether this actor is valid or not 
-	 * because clothing asset can be changed by editing 
-	 */
-	physx::apex::NxClothingAsset*	ParentClothingAsset;
-	/** APEX clothing actor is created from APEX clothing asset for cloth simulation */
-	physx::apex::NxClothingActor*		ApexClothingActor;
-	FPhysScene * PhysScene;
-	uint32 SceneType;
-};
-
-/** data for updating cloth section from the results of clothing simulation */
-struct FClothSimulData
-{
-	TArray<FVector4> ClothSimulPositions;
-	TArray<FVector4> ClothSimulNormals;
-};
-
 /**  for storing precomputed cloth morph target data */
 struct FClothMorphTargetData
 {
@@ -215,55 +230,6 @@ namespace EAnimationMode
 	};
 }
 
-USTRUCT()
-struct FSingleAnimationPlayData
-{
-	GENERATED_USTRUCT_BODY()
-
-	// @todo in the future, we should make this one UObject
-	// and have detail customization to display different things
-	// The default sequence to play on this skeletal mesh
-	UPROPERTY(EditAnywhere, BlueprintReadWrite, Category=Animation)
-	class UAnimationAsset* AnimToPlay;
-
-	// @fixme : until we properly support it I'm commenting out editable property part
-	// The default sequence to play on this skeletal mesh
-	UPROPERTY()//EditAnywhere, BlueprintReadWrite, Category=Animation)
-	class UVertexAnimation* VertexAnimToPlay;
-
-	/** Default setting for looping for SequenceToPlay. This is not current state of looping. */
-	UPROPERTY(EditAnywhere, BlueprintReadWrite, Category=Animation, meta=(DisplayName="Looping"))
-	uint32 bSavedLooping:1;
-
-	/** Default setting for playing for SequenceToPlay. This is not current state of playing. */
-	UPROPERTY(EditAnywhere, BlueprintReadWrite, Category=Animation, meta=(DisplayName="Playing"))
-	uint32 bSavedPlaying:1;
-
-	/** Default setting for position of SequenceToPlay to play. */
-	UPROPERTY(EditAnywhere, BlueprintReadWrite, Category=Animation, meta=(DisplayName="Initial Position"))
-	float SavedPosition;
-
-	/** Default setting for play rate of SequenceToPlay to play. */
-	UPROPERTY(EditAnywhere, AdvancedDisplay, BlueprintReadWrite, Category=Animation, meta=(DisplayName="PlayRate"))
-	float SavedPlayRate;
-
-	FSingleAnimationPlayData()
-	{
-		SavedPlayRate = 1.0f;
-		bSavedLooping = true;
-		bSavedPlaying = true;
-		SavedPosition = 0.f;
-	}
-
-	/** Called when initialized. */
-	ENGINE_API void Initialize(class UAnimSingleNodeInstance* Instance);
-
-	/** Populates this play data with the current state of the supplied instance. */
-	ENGINE_API void PopulateFrom(UAnimSingleNodeInstance* Instance);
-
-	void ValidatePosition();
-};
-
 /**
 * Tick function that prepares for cloth tick
 **/
@@ -714,47 +680,52 @@ public:
 	FSkeletalMeshComponentPreClothTickFunction PreClothTickFunction;
 
 #if WITH_APEX_CLOTHING
-	/** 
+
+	/* The public facing editable cloth simulation context. Game code should only modify this from the game thread!*/
+	FClothSimulationContext EditableClothSimulationContext;
+
+	bool IsClothBoundToMasterComponent() const { return bBindClothToMasterComponent; }
+
+private:
+
+   /** Double buffer for the current cloth simulation context */
+	FClothSimulationContext InternalClothSimulationContext;
+
+	/** Whether or not we're taking cloth sim information from our master component */
+	bool bBindClothToMasterComponent;
+	/** The previous state of the master component simulation coord space, so we can restore on unbind */
+	bool bPrevMasterSimulateLocalSpace;
+
+	/** Copies the data from the external cloth simulation context. We copy instead of flipping because the API has to return the full struct to make backwards compat easy*/
+	void SubmitClothSimulationContext();
+
+   /** 
 	* clothing actors will be created from clothing assets for cloth simulation 
 	* 1 actor should correspond to 1 asset
 	*/
 	TArray<FClothingActor> ClothingActors;
 
-	float ClothMaxDistanceScale;
+public:
 
-	FClothingActor::TeleportMode ClothTeleportMode;
-	/** previous root bone matrix to compare the difference and decide to do clothing teleport  */
-	FMatrix	PrevRootBoneMatrix;
-	/** used for pre-computation using TeleportRotationThreshold property */
-	float ClothTeleportCosineThresholdInRad;
-	/** used for pre-computation using tTeleportDistanceThreshold property */
-	float ClothTeleportDistThresholdSquared;
-	/** 
-	 * clothing reset is needed once more to avoid clothing pop up 
-	 * use until Apex clothing bug is resolved 
-	 */
-	bool bNeedTeleportAndResetOnceMore;
+	const TArray<FClothingActor>& GetClothingActors(){ return ClothingActors; }
+
+	float ClothMaxDistanceScale;
 	/** used for checking whether cloth morph target data were pre-computed or not */
 	bool bPreparedClothMorphTargets;
+
 	/** precomputed actual cloth morph target data */
 	TArray<FClothMorphTargetData> ClothMorphTargets;
 
-	/** Whether or not we're taking cloth sim information from our master component */
-	bool bBindClothToMasterComponent;
-	/** The previous state of the master component simulation coord space, so we can restore on unbind */
-	bool bPrevMasterSimulateLocalSpace;
-
-#if WITH_CLOTH_COLLISION_DETECTION
+	#if WITH_CLOTH_COLLISION_DETECTION
 	/** increase every tick to update clothing collision  */
-	uint32 ClothingCollisionRevision; 
+	uint32 ClothingCollisionRevision;
 
 	TArray<physx::apex::NxClothingCollision*>	ParentCollisions;
 	TArray<physx::apex::NxClothingCollision*>	EnvironmentCollisions;
 	TArray<physx::apex::NxClothingCollision*>	ChildrenCollisions;
 
 	TMap<TWeakObjectPtr<UPrimitiveComponent>, FApexClothCollisionInfo> ClothOverlappedComponentsMap;
-#endif // WITH_CLOTH_COLLISION_DETECTION
-
+	#endif // WITH_CLOTH_COLLISION_DETECTION
 #endif // WITH_APEX_CLOTHING
 
 	/** 
@@ -776,9 +747,10 @@ public:
 
 	/** Store cloth simulation data into OutClothSimData */
 	void GetUpdateClothSimulationData(TArray<FClothSimulData>& OutClothSimData, USkeletalMeshComponent* OverrideLocalRootComponent = nullptr);
-	void ApplyWindForCloth(FClothingActor& ClothingActor);
+	void ApplyWindForCloth(const FClothingActor& ClothingActor) const;
 	void RemoveAllClothingActors();
 	void ReleaseAllClothingResources();
+	bool IsExecutingParallelCloth() const;
 
 	bool IsValidClothingActor(int32 ActorIndex) const;
 	/** Draws APEX Clothing simulated normals on cloth meshes **/
@@ -1163,13 +1135,13 @@ public:
 	*/
 	bool CreateClothingActor(int32 AssetIndex, physx::apex::NxClothingAsset* ClothingAsset, TArray<FVector>* BlendedDelta = NULL);
 	/** should call this method if occurred any changes in clothing assets */
-	void ValidateClothingActors();
+	void RecreateClothingActors();
 	/** add bounding box for cloth */
 	void AddClothingBounds(FBoxSphereBounds& InOutBounds) const;
 	/** changes clothing LODs, if clothing LOD is disabled or LODIndex is greater than apex clothing LODs, simulation will be disabled */
 	void SetClothingLOD(int32 LODIndex);
 	/** check whether clothing teleport is needed or not to avoid a weird simulation result */
-	virtual void CheckClothTeleport(float DeltaTime);
+	virtual void CheckClothTeleport(float DeltaTime, FClothSimulationContext& ClothSimulationContext) const;
 
 	/** 
 	* methods for cloth morph targets 
@@ -1316,11 +1288,15 @@ private:
 
 	void PerformBlendPhysicsBones(const TArray<FBoneIndexType>& InRequiredBones, TArray<FTransform>& InLocalAtoms);
 
+	void ParallelTickClothing(float DeltaTime, FClothSimulationContext& ClothSimulationContext) const;
+	void ParallelUpdateClothState(float DeltaTime, FClothSimulationContext& ClothSimulationContext) const;
+
 	friend class FParallelBlendPhysicsCompletionTask;
 	void CompleteParallelBlendPhysics();
 	void PostBlendPhysics();
 
 	friend class FTickClothingTask;
+	friend class FClothManager;
 	void PerformTickClothing(float DeltaTime);
 
 	// these are deprecated variables from removing SingleAnimSkeletalComponent
diff --git a/Engine/Source/Runtime/Engine/Classes/Engine/World.h b/Engine/Source/Runtime/Engine/Classes/Engine/World.h
--- a/Engine/Source/Runtime/Engine/Classes/Engine/World.h
+++ b/Engine/Source/Runtime/Engine/Classes/Engine/World.h
@@ -374,7 +374,7 @@ struct FEndPhysicsTickFunction : public FTickFunction
 * Tick function that starts the cloth tick
 **/
 USTRUCT()
-struct FStartClothSimulationFunction : public FTickFunction
+struct FStartClothAndAsyncSimulationFunction : public FTickFunction
 {
 	GENERATED_USTRUCT_BODY()
 
@@ -870,7 +870,7 @@ public:
 	FEndPhysicsTickFunction EndPhysicsTickFunction;
 
 	/** Tick function for starting cloth simulation																				*/
-	FStartClothSimulationFunction StartClothTickFunction;
+	FStartClothAndAsyncSimulationFunction StartClothTickFunction;
 	/** Tick function for ending cloth simulation																				*/
 	FEndClothSimulationFunction EndClothTickFunction;
 
@@ -2112,6 +2112,14 @@ public:
 	/** Returns a reference to the game viewport displaying this world if one exists. */
 	UGameViewportClient* GetGameViewport() const;
 
+private:
+	/** Begin cloth and async simulation */
+	void StartClothAndAsyncSim();
+
+	friend FStartClothAndAsyncSimulationFunction;
+
+public:
+
 	DEPRECATED(4.3, "GetBrush is deprecated use GetDefaultBrush instead.")
 	ABrush* GetBrush() const;
 	/** 
@@ -2981,9 +2989,6 @@ public:
 	/** Waits for the physics scene to be done processing */
 	void FinishPhysicsSim();
 
-	/** Begin cloth simulation */
-	void StartClothSim();
-
 	/** Spawns GameMode for the level. */
 	bool SetGameMode(const FURL& InURL);
 
diff --git a/Engine/Source/Runtime/Engine/Classes/PhysicsEngine/PhysicsSettings.h b/Engine/Source/Runtime/Engine/Classes/PhysicsEngine/PhysicsSettings.h
--- a/Engine/Source/Runtime/Engine/Classes/PhysicsEngine/PhysicsSettings.h
+++ b/Engine/Source/Runtime/Engine/Classes/PhysicsEngine/PhysicsSettings.h
@@ -177,6 +177,10 @@ class ENGINE_API UPhysicsSettings : public UDeveloperSettings
 	UPROPERTY(config, EditAnywhere, Category = Simulation)
 	bool bDisableCCD;
 
+	/**
+	*  If true cloth data will be updated in parallel. With this mode MorphTarget and child collision is not supported.*/
+	bool bParallelCloth;
+
 
 	/** Max Physics Delta Time to be clamped. */
 	UPROPERTY(config, EditAnywhere, meta=(ClampMin="0.0013", UIMin = "0.0013", ClampMax="1.0", UIMax="1.0"), Category=Framerate)
diff --git a/Engine/Source/Runtime/Engine/Private/Components/SkeletalMeshComponent.cpp b/Engine/Source/Runtime/Engine/Private/Components/SkeletalMeshComponent.cpp
--- a/Engine/Source/Runtime/Engine/Private/Components/SkeletalMeshComponent.cpp
+++ b/Engine/Source/Runtime/Engine/Private/Components/SkeletalMeshComponent.cpp
@@ -5,6 +5,7 @@
 =============================================================================*/
 
 #include "EnginePrivate.h"
+#include "Components/SkeletalMeshComponent.h"
 #include "ParticleDefinitions.h"
 #include "BlueprintUtilities.h"
 #include "SkeletalRenderCPUSkin.h"
@@ -20,6 +21,7 @@
 #include "Animation/AnimSingleNodeInstance.h"
 #include "Engine/SkeletalMeshSocket.h"
 #include "AI/NavigationSystemHelpers.h"
+#include "PhysicsPublic.h"
 #if WITH_EDITOR
 #include "ShowFlags.h"
 #include "Collision.h"
@@ -27,6 +29,7 @@
 #endif
 #include "PhysicsEngine/PhysicsSettings.h"
 #include "PhysicsEngine/PhysicsAsset.h"
+#include "PhysicsEngine/ClothManager.h"
 
 #if WITH_APEX_CLOTHING
 #include "PhysicsEngine/PhysXSupport.h"
@@ -128,19 +131,22 @@ USkeletalMeshComponent::USkeletalMeshComponent(const FObjectInitializer& ObjectI
 
 #if WITH_APEX_CLOTHING
 	ClothMaxDistanceScale = 1.0f;
-	ClothTeleportMode = FClothingActor::Continuous;
-	PrevRootBoneMatrix = GetBoneMatrix(0); // save the root bone transform
 	bResetAfterTeleport = true;
 	TeleportDistanceThreshold = 300.0f;
 	TeleportRotationThreshold = 0.0f;// angles in degree, disabled by default
-	// pre-compute cloth teleport thresholds for performance
-	ClothTeleportCosineThresholdInRad = FMath::Cos(FMath::DegreesToRadians(TeleportRotationThreshold));
-	ClothTeleportDistThresholdSquared = TeleportDistanceThreshold * TeleportDistanceThreshold;
-	bNeedTeleportAndResetOnceMore = false;
 	ClothBlendWeight = 1.0f;
 	bPreparedClothMorphTargets = false;
+
+	EditableClothSimulationContext.ClothTeleportMode = FClothingActor::Default;
+	InternalClothSimulationContext.ClothTeleportMode = FClothingActor::Continuous;
+	InternalClothSimulationContext.PrevRootBoneMatrix = GetBoneMatrix(0); // save the root bone transform
+
+	// pre-compute cloth teleport thresholds for performance
+	EditableClothSimulationContext.ClothTeleportCosineThresholdInRad = FMath::Cos(FMath::DegreesToRadians(TeleportRotationThreshold));
+	EditableClothSimulationContext.ClothTeleportDistThresholdSquared = TeleportDistanceThreshold * TeleportDistanceThreshold;
 	bBindClothToMasterComponent = false;
 	bPrevMasterSimulateLocalSpace = false;
+
 #if WITH_CLOTH_COLLISION_DETECTION
 	ClothingCollisionRevision = 0;
 #endif// #if WITH_CLOTH_COLLISION_DETECTION
@@ -166,8 +172,33 @@ void USkeletalMeshComponent::RegisterComponentTickFunctions(bool bRegister)
 	UpdatePreClothTickRegisteredState();
 }
 
+static TAutoConsoleVariable<int32> CVarAsyncCloth(
+	TEXT("p.ParallelCloth"),
+	0,
+	TEXT("Whether cloth prep and simulation is done in off the game thread."));
+
+
 void USkeletalMeshComponent::RegisterPreClothTick(bool bRegister)
 {
+	if(CVarAsyncCloth.GetValueOnGameThread())
+	{
+		UPhysicsSettings::Get()->bParallelCloth = true;
+	}else 
+	{
+		UPhysicsSettings::Get()->bParallelCloth = false;
+	}
+
+#if WITH_APEX_CLOTHING
+	if(bRegister && UPhysicsSettings::Get()->bParallelCloth)	//If we're using parallel cloth we break the cloth prep code out of the game thread. To do this we must register with the ClothManager which will do the work async
+	{
+		if(FPhysScene* PhysScene = GetWorld()->GetPhysicsScene())
+		{
+			FClothManager* ClothManager = PhysScene->GetClothManager();
+			ClothManager->RegisterForClothThisFrame(this);
+		}
+	}
+#endif
+
 	if (bRegister != PreClothTickFunction.IsTickFunctionRegistered())
 	{
 		if (bRegister)
@@ -230,6 +261,10 @@ void USkeletalMeshComponent::OnRegister()
 	{
 		SetComponentTickEnabled(false);
 	}
+
+#if WITH_APEX_CLOTHING
+	RecreateClothingActors();
+#endif
 }
 
 void USkeletalMeshComponent::OnUnregister()
@@ -1049,6 +1084,19 @@ int32 GetCurveNumber(USkeleton* Skeleton)
 	return 0;
 }
 
+void USkeletalMeshComponent::SubmitClothSimulationContext()
+{
+	
+	InternalClothSimulationContext.ClothTeleportCosineThresholdInRad = EditableClothSimulationContext.ClothTeleportCosineThresholdInRad;
+	InternalClothSimulationContext.ClothTeleportDistThresholdSquared = EditableClothSimulationContext.ClothTeleportDistThresholdSquared;
+
+	//If the editable teleport mode is non default it means the user has actively written to it since the previous flip. This means we should use it. Otherwise use the internal because the cloth code set it as needed
+	InternalClothSimulationContext.ClothTeleportMode = EditableClothSimulationContext.ClothTeleportMode == FClothingActor::Default ? InternalClothSimulationContext.ClothTeleportMode : EditableClothSimulationContext.ClothTeleportMode;
+	EditableClothSimulationContext.ClothTeleportMode = FClothingActor::Default;	//reset editable teleport mode to default because user has to set it directly
+
+	//we intentionally ignore PrevRootBone which is only needed internally, but is still in the Context struct for the sake of const
+}
+
 void USkeletalMeshComponent::RefreshBoneTransforms(FActorComponentTickFunction* TickFunction)
 {
 	SCOPE_CYCLE_COUNTER(STAT_AnimGameThreadTime);
@@ -1353,7 +1401,7 @@ void USkeletalMeshComponent::SetSkeletalMesh(USkeletalMesh* InSkelMesh)
 	InitAnim(false);
 
 #if WITH_APEX_CLOTHING
-	ValidateClothingActors();
+	RecreateClothingActors();
 #endif
 }
 
@@ -1843,44 +1891,6 @@ class UAnimSingleNodeInstance* USkeletalMeshComponent::GetSingleNodeInstance() c
 	return Cast<class UAnimSingleNodeInstance>(AnimScriptInstance);
 }
 
-void FSingleAnimationPlayData::Initialize(UAnimSingleNodeInstance* Instance)
-{
-	Instance->SetAnimationAsset(AnimToPlay);
-	Instance->SetPosition(SavedPosition, false);
-	Instance->SetPlayRate(SavedPlayRate);
-	Instance->SetPlaying(bSavedPlaying);
-	Instance->SetLooping(bSavedLooping);
-}
-
-void FSingleAnimationPlayData::PopulateFrom(UAnimSingleNodeInstance* Instance)
-{
-	AnimToPlay = Instance->CurrentAsset;
-	SavedPosition = Instance->CurrentTime;
-	SavedPlayRate = Instance->PlayRate;
-	bSavedPlaying = Instance->bPlaying;
-	bSavedLooping = Instance->bLooping;
-}
-
-void FSingleAnimationPlayData::ValidatePosition()
-{
-	float Min=0,Max=0;
-
-	if (AnimToPlay)
-	{
-		UAnimSequenceBase* SequenceBase = Cast<UAnimSequenceBase>(AnimToPlay);
-		if (SequenceBase)
-		{
-			Max = SequenceBase->SequenceLength;
-		}
-	}
-	else if (VertexAnimToPlay)
-	{
-		Max = VertexAnimToPlay->GetAnimLength();
-	}
-
-	SavedPosition = FMath::Clamp<float>(SavedPosition, Min, Max);
-}
-
 FTransform USkeletalMeshComponent::ConvertLocalRootMotionToWorld(const FTransform& InTransform)
 {
 	// Make sure component to world is up to date
@@ -2242,9 +2252,6 @@ void USkeletalMeshComponent::BindClothToMasterPoseComponent()
 			return;
 		}
 
-		ValidateClothingActors();
-		MasterComp->ValidateClothingActors();
-
 		int32 NumClothingActors = ClothingActors.Num();
 		
 		for(int32 ActorIdx = 0 ; ActorIdx < NumClothingActors ; ++ActorIdx)
@@ -2255,6 +2262,7 @@ void USkeletalMeshComponent::BindClothToMasterPoseComponent()
 			NxClothingActor* MasterApexActor = MasterActor.ApexClothingActor;
 			if(ApexActor && MasterApexActor)
 			{
+				//TODO: wait on parallel animation if needed
 				// Disable our actors
 				ApexActor->setFrozen(true);
 
@@ -2298,6 +2306,7 @@ void USkeletalMeshComponent::UnbindClothFromMasterPoseComponent(bool bRestoreSim
 
 			if(ApexActor)
 			{
+				//TODO: wait on parallel animation if needed
 				ApexActor->setFrozen(false);
 
 				bool bMasterPoseSpaceChanged = (MasterComp->bLocalSpaceSimulation && !bPrevMasterSimulateLocalSpace);
diff --git a/Engine/Source/Runtime/Engine/Private/PhysicsEngine/ClothManager.cpp b/Engine/Source/Runtime/Engine/Private/PhysicsEngine/ClothManager.cpp
new file mode 100644
--- /dev/null
+++ b/Engine/Source/Runtime/Engine/Private/PhysicsEngine/ClothManager.cpp
@@ -0,0 +1,85 @@
+// Copyright 1998-2015 Epic Games, Inc. All Rights Reserved.
+
+#include "EnginePrivate.h"
+#include "ClothManager.h"
+#include "PhysicsPublic.h"
+#include "Components/SkeletalMeshComponent.h"
+
+FClothManager::FClothManager(UWorld* AssociatedWorld, FPhysScene* InPhysScene)
+ : PhysScene(InPhysScene) 
+ {
+	ClothManagerTickFunction.bCanEverTick = true;
+	ClothManagerTickFunction.Target = this;
+	ClothManagerTickFunction.TickGroup = TG_StartPhysics;
+	ClothManagerTickFunction.bRunOnAnyThread = true;	//this value seems to be ignored
+	ClothManagerTickFunction.RegisterTickFunction(AssociatedWorld->PersistentLevel);
+ }
+
+void FClothManager::Tick(float DeltaTime)
+{
+	if(SkeletalMeshComponents.Num())
+	{
+		IsPreparingCloth.AtomicSet(true);
+		for (USkeletalMeshComponent* SkeletalMeshComponent : SkeletalMeshComponents)
+		{	
+			SkeletalMeshComponent->SubmitClothSimulationContext();	//make sure user params are passed internally
+
+			FClothSimulationContext& ClothSimulationContext = SkeletalMeshComponent->InternalClothSimulationContext;
+			SkeletalMeshComponent->ParallelTickClothing(DeltaTime, ClothSimulationContext);
+		}
+
+		SkeletalMeshComponents.Empty(SkeletalMeshComponents.Num());
+		PhysScene->StartCloth();
+		IsPreparingCloth.AtomicSet(false);
+	}
+}
+
+void FClothManager::RegisterForClothThisFrame(USkeletalMeshComponent* SkeletalMeshComponent)
+{
+	SkeletalMeshComponents.Add(SkeletalMeshComponent);
+}
+
+
+class FClothManagerTickTask
+{
+	FClothManager* ClothManager;
+	float DeltaTime;
+
+public:
+	FClothManagerTickTask(FClothManager* InClothManager, float InDeltaTime)
+		: ClothManager(InClothManager)
+		, DeltaTime(InDeltaTime)
+	{
+	}
+
+	FORCEINLINE TStatId GetStatId() const
+	{
+		RETURN_QUICK_DECLARE_CYCLE_STAT(FClothManagerTickTask, STATGROUP_TaskGraphTasks);
+	}
+	static ENamedThreads::Type GetDesiredThread()
+	{
+		return ENamedThreads::AnyThread;
+	}
+	static ESubsequentsMode::Type GetSubsequentsMode()
+	{
+		return ESubsequentsMode::FireAndForget;
+	}
+
+	void DoTask(ENamedThreads::Type CurrentThread, const FGraphEventRef& MyCompletionGraphEvent)
+	{
+		ClothManager->Tick(DeltaTime);
+	}
+};
+
+
+void FClothManagerTickFunction::ExecuteTick(float DeltaTime, enum ELevelTick TickType, ENamedThreads::Type CurrentThread, const FGraphEventRef& MyCompletionGraphEvent)
+{
+	QUICK_SCOPE_CYCLE_COUNTER(FClothManagerTickFunction_ExecuteTick);
+
+	TGraphTask<FClothManagerTickTask>::CreateTask().ConstructAndDispatchWhenReady(Target, DeltaTime);
+}
+
+FString FClothManagerTickFunction::DiagnosticMessage()
+{
+	return TEXT("FClothManagerTickFunction");
+}
diff --git a/Engine/Source/Runtime/Engine/Private/PhysicsEngine/ClothManager.h b/Engine/Source/Runtime/Engine/Private/PhysicsEngine/ClothManager.h
new file mode 100644
--- /dev/null
+++ b/Engine/Source/Runtime/Engine/Private/PhysicsEngine/ClothManager.h
@@ -0,0 +1,45 @@
+// Copyright 1998-2015 Epic Games, Inc. All Rights Reserved.
+
+/*=============================================================================
+	ClothManager.h
+	In charge of passing data between skeletal meshes and cloth solver
+=============================================================================*/
+
+#pragma once
+
+class FPhysScene;
+class USkeletalMeshComponent;
+
+struct FClothManagerTickFunction : public FTickFunction
+{
+	class FClothManager* Target;
+
+	// FTickFunction interface
+	virtual void ExecuteTick(float DeltaTime, enum ELevelTick TickType, ENamedThreads::Type CurrentThread, const FGraphEventRef& MyCompletionGraphEvent) override;
+	virtual FString DiagnosticMessage() override;
+	// End of FTickFunction interface
+};
+
+
+class FClothManager
+{
+public:
+
+	FClothManager(UWorld* AssociatedWorld, FPhysScene* InPhysScene);
+
+	/** Go through all SkeletalMeshComponents registered for work this frame and call the solver as needed*/
+	void Tick(float DeltaTime);
+
+	/** Register the SkeletalMeshComponent so that it does cloth solving work this frame */
+	void RegisterForClothThisFrame(USkeletalMeshComponent* SkeletalMeshComponent);
+
+	/** Whether the cloth manager is currently doing work in an async thread */
+	bool IsPreparingClothAsync() const { return IsPreparingCloth; }
+	
+private:
+	TArray<USkeletalMeshComponent*> SkeletalMeshComponents;
+	FClothManagerTickFunction ClothManagerTickFunction;
+	FPhysScene* PhysScene;
+
+	FThreadSafeBool IsPreparingCloth;
+};
\ No newline at end of file
diff --git a/Engine/Source/Runtime/Engine/Private/PhysicsEngine/PhysLevel.cpp b/Engine/Source/Runtime/Engine/Private/PhysicsEngine/PhysLevel.cpp
--- a/Engine/Source/Runtime/Engine/Private/PhysicsEngine/PhysLevel.cpp
+++ b/Engine/Source/Runtime/Engine/Private/PhysicsEngine/PhysLevel.cpp
@@ -210,7 +210,7 @@ void UWorld::FinishPhysicsSim()
 	PhysScene->EndFrame(LineBatcher);
 }
 
-void UWorld::StartClothSim()
+void UWorld::StartClothAndAsyncSim()
 {
 	FPhysScene* PhysScene = GetPhysicsScene();
 	if (PhysScene == NULL)
@@ -218,7 +218,12 @@ void UWorld::StartClothSim()
 		return;
 	}
 
-	PhysScene->StartCloth();
+	if(!UPhysicsSettings::Get()->bParallelCloth)	//if we're doing parallel cloth, the cloth sim will be ticked from cloth manager
+	{
+		PhysScene->StartCloth();
+	}
+
+	PhysScene->StartAsync();
 }
 
 // the physics tick functions
@@ -281,15 +286,15 @@ FString FEndPhysicsTickFunction::DiagnosticMessage()
 	return TEXT("FEndPhysicsTickFunction");
 }
 
-void FStartClothSimulationFunction::ExecuteTick(float DeltaTime, enum ELevelTick TickType, ENamedThreads::Type CurrentThread, const FGraphEventRef& MyCompletionGraphEvent)
+void FStartClothAndAsyncSimulationFunction::ExecuteTick(float DeltaTime, enum ELevelTick TickType, ENamedThreads::Type CurrentThread, const FGraphEventRef& MyCompletionGraphEvent)
 {
 	QUICK_SCOPE_CYCLE_COUNTER(FStartClothSimulationFunction_ExecuteTick);
 
 	check(Target);
-	Target->StartClothSim();
+	Target->StartClothAndAsyncSim();
 }
 
-FString FStartClothSimulationFunction::DiagnosticMessage()
+FString FStartClothAndAsyncSimulationFunction::DiagnosticMessage()
 {
 	return TEXT("FStartClothSimulationFunction");
 }
diff --git a/Engine/Source/Runtime/Engine/Private/PhysicsEngine/PhysScene.cpp b/Engine/Source/Runtime/Engine/Private/PhysicsEngine/PhysScene.cpp
--- a/Engine/Source/Runtime/Engine/Private/PhysicsEngine/PhysScene.cpp
+++ b/Engine/Source/Runtime/Engine/Private/PhysicsEngine/PhysScene.cpp
@@ -3,6 +3,8 @@
 
 #include "EnginePrivate.h"
 #include "PhysicsPublic.h"
+#include "PhysicsEngine/ClothManager.h"
+#include "Components/SkeletalMeshComponent.h"
 
 #if WITH_PHYSX
 	#include "PhysXSupport.h"
@@ -98,6 +100,8 @@ FPhysScene::FPhysScene()
 #endif
 	PhysxUserData = FPhysxUserData(this);
 
+	ClothManager = nullptr;
+
 	
 	int64 NumPhysxDispatcher = 0;
 	FParse::Value(FCommandLine::Get(), TEXT("physxDispatcher="), NumPhysxDispatcher);
@@ -180,7 +184,21 @@ FPhysScene::FPhysScene()
 #endif
 }
 
+void FPhysScene::SetOwningWorld(UWorld* InOwningWorld)
+{
+	if(OwningWorld != nullptr)
+	{
+		delete ClothManager;
+		ClothManager = nullptr;
+	}
 
+	OwningWorld = InOwningWorld;
+
+	if (OwningWorld)
+	{
+		ClothManager = new FClothManager(OwningWorld, this);
+	}
+}
 
 /** Exposes destruction of physics-engine scene outside Engine. */
 FPhysScene::~FPhysScene()
@@ -588,7 +606,8 @@ void GatherApexStats(const UWorld* World, NxApexScene* ApexScene)
 	for (TObjectIterator<USkeletalMeshComponent> Itr; Itr; ++Itr)
 	{
 		if (Itr->GetWorld() != World) { continue; }
-		for (const FClothingActor& ClothingActor : Itr->ClothingActors)
+		const TArray<FClothingActor>& ClothingActors = Itr->GetClothingActors();
+		for (const FClothingActor& ClothingActor : ClothingActors)
 		{
 			if (ClothingActor.ApexClothingActor && ClothingActor.ApexClothingActor->getActivePhysicalLod() == 1)
 			{
@@ -1115,7 +1134,7 @@ TAutoConsoleVariable<int32> CVarEnableClothPhysics(TEXT("p.ClothPhysics"), 1, TE
 void FPhysScene::StartCloth()
 {
 	FGraphEventArray FinishPrerequisites;
-	if(CVarEnableClothPhysics.GetValueOnGameThread())
+	if(CVarEnableClothPhysics.GetValueOnAnyThread())
 	{
 		TickPhysScene(PST_Cloth, PhysicsSubsceneCompletion[PST_Cloth]);
 		{
@@ -1135,7 +1154,12 @@ void FPhysScene::StartCloth()
 			}
 		}
 	}
-	
+}
+
+
+void FPhysScene::StartAsync()
+{
+	FGraphEventArray FinishPrerequisites;
 
 	//If the async scene is lagged we start it here to make sure any cloth in the async scene is using the results of the previous simulation.
 	if (FrameLagAsync() && bAsyncSceneEnabled)
@@ -1544,6 +1568,12 @@ void FPhysScene::TermPhysScene(uint32 SceneType)
 		GPhysXSceneMap.Remove(PhysXSceneIndex[SceneType]);
 	}
 #endif
+
+	if (SceneType == PST_Cloth)
+	{
+		delete ClothManager;
+		ClothManager = nullptr;
+	}
 }
 
 #if WITH_PHYSX
diff --git a/Engine/Source/Runtime/Engine/Private/PhysicsEngine/PhysicsSettings.cpp b/Engine/Source/Runtime/Engine/Private/PhysicsEngine/PhysicsSettings.cpp
--- a/Engine/Source/Runtime/Engine/Private/PhysicsEngine/PhysicsSettings.cpp
+++ b/Engine/Source/Runtime/Engine/Private/PhysicsEngine/PhysicsSettings.cpp
@@ -22,6 +22,7 @@ UPhysicsSettings::UPhysicsSettings(const FObjectInitializer& ObjectInitializer)
 	, MaxAngularVelocity(3600)	//10 revolutions per second
 	, bSimulateSkeletalMeshOnDedicatedServer(true)
 	, bDefaultHasComplexCollision(true)
+	, bParallelCloth(false)
 	, MaxPhysicsDeltaTime(1.f / 30.f)
 	, bSubstepping(false)
 	, bSubsteppingAsync(false)
diff --git a/Engine/Source/Runtime/Engine/Private/SingleAnimationPlayData.cpp b/Engine/Source/Runtime/Engine/Private/SingleAnimationPlayData.cpp
new file mode 100644
--- /dev/null
+++ b/Engine/Source/Runtime/Engine/Private/SingleAnimationPlayData.cpp
@@ -0,0 +1,46 @@
+﻿// Copyright 1998-2015 Epic Games, Inc. All Rights Reserved.
+
+#include "EnginePrivate.h"
+#include "SingleAnimationPlayData.h"
+#include "Animation/AnimSingleNodeInstance.h"
+#include "Animation/AnimSequenceBase.h"
+#include "Animation/VertexAnim/VertexAnimation.h"
+
+
+void FSingleAnimationPlayData::Initialize(UAnimSingleNodeInstance* Instance)
+{
+	Instance->SetAnimationAsset(AnimToPlay);
+	Instance->SetPosition(SavedPosition, false);
+	Instance->SetPlayRate(SavedPlayRate);
+	Instance->SetPlaying(bSavedPlaying);
+	Instance->SetLooping(bSavedLooping);
+}
+
+void FSingleAnimationPlayData::PopulateFrom(UAnimSingleNodeInstance* Instance)
+{
+	AnimToPlay = Instance->CurrentAsset;
+	SavedPosition = Instance->CurrentTime;
+	SavedPlayRate = Instance->PlayRate;
+	bSavedPlaying = Instance->bPlaying;
+	bSavedLooping = Instance->bLooping;
+}
+
+void FSingleAnimationPlayData::ValidatePosition()
+{
+	float Min = 0, Max = 0;
+
+	if (AnimToPlay)
+	{
+		UAnimSequenceBase* SequenceBase = Cast<UAnimSequenceBase>(AnimToPlay);
+		if (SequenceBase)
+		{
+			Max = SequenceBase->SequenceLength;
+		}
+	}
+	else if (VertexAnimToPlay)
+	{
+		Max = VertexAnimToPlay->GetAnimLength();
+	}
+
+	SavedPosition = FMath::Clamp<float>(SavedPosition, Min, Max);
+}
\ No newline at end of file
diff --git a/Engine/Source/Runtime/Engine/Private/SkeletalMeshComponentPhysics.cpp b/Engine/Source/Runtime/Engine/Private/SkeletalMeshComponentPhysics.cpp
--- a/Engine/Source/Runtime/Engine/Private/SkeletalMeshComponentPhysics.cpp
+++ b/Engine/Source/Runtime/Engine/Private/SkeletalMeshComponentPhysics.cpp
@@ -4,6 +4,7 @@
 #include "PhysicsPublic.h"
 #include "SkeletalRender.h"
 #include "SkeletalRenderPublic.h"
+#include "Components/SkeletalMeshComponent.h"
 
 #include "MessageLog.h"
 #include "CollisionDebugDrawingPublic.h"
@@ -32,6 +33,8 @@
 #endif//#if WITH_APEX
 #include "PhysicsEngine/PhysicsConstraintTemplate.h"
 #include "PhysicsEngine/PhysicsAsset.h"
+#include "PhysicsEngine/PhysicsSettings.h"
+#include "PhysicsEngine/ClothManager.h"
 
 #define LOCTEXT_NAMESPACE "SkeletalMeshComponentPhysics"
 
@@ -895,10 +898,10 @@ void USkeletalMeshComponent::InitArticulated(FPhysScene* PhysScene)
 	// Update Flag
 	ResetAllBodiesSimulatePhysics();
 #if WITH_APEX_CLOTHING
-	PrevRootBoneMatrix = GetBoneMatrix(0); // save the root bone transform
+	InternalClothSimulationContext.PrevRootBoneMatrix = GetBoneMatrix(0); // save the root bone transform
 	// pre-compute cloth teleport thresholds for performance
-	ClothTeleportCosineThresholdInRad = FMath::Cos(FMath::DegreesToRadians(TeleportRotationThreshold));
-	ClothTeleportDistThresholdSquared = TeleportDistanceThreshold * TeleportDistanceThreshold;
+	EditableClothSimulationContext.ClothTeleportCosineThresholdInRad = FMath::Cos(FMath::DegreesToRadians(TeleportRotationThreshold));
+	EditableClothSimulationContext.ClothTeleportDistThresholdSquared = TeleportDistanceThreshold * TeleportDistanceThreshold;	
 #endif // #if WITH_APEX_CLOTHING
 }
 
@@ -1992,8 +1995,8 @@ bool USkeletalMeshComponent::ComponentOverlapMultiImpl(TArray<struct FOverlapRes
 // convert a bone name from APEX stype to FBX style
 static FName GetConvertedBoneName(NxClothingAsset* ApexClothingAsset, int32 BoneIndex)
 {
-	QUICK_SCOPE_CYCLE_COUNTER(STAT_GetConvertedBoneName);
-	static FString Work;
+	//QUICK_SCOPE_CYCLE_COUNTER(STAT_GetConvertedBoneName);
+	/*static FString Work;
 	check(IsInGameThread() && !Work.Len());
 	// FIXME, this still needs work, string conversion and hashing
 
@@ -2012,7 +2015,9 @@ static FName GetConvertedBoneName(NxClothingAsset* ApexClothingAsset, int32 Bone
 		//check(Result == FName(*FString(ApexClothingAsset->getBoneName(BoneIndex)).Replace(TEXT(" "), TEXT("-"))));
 		return Result;
 	}
-	return NAME_None;
+	return NAME_None;*/
+	//Going back to allocation approach because we now call this from many threads
+	return *FString(ApexClothingAsset->getBoneName(BoneIndex)).Replace(TEXT(" "), TEXT("-"));
 }
 
 void USkeletalMeshComponent::AddClothingBounds(FBoxSphereBounds& InOutBounds) const
@@ -2053,36 +2058,16 @@ bool USkeletalMeshComponent::HasValidClothingActors()
 }
 
 //if any changes in clothing assets, will create new actors 
-void USkeletalMeshComponent::ValidateClothingActors()
+void USkeletalMeshComponent::RecreateClothingActors()
 {
-	//newly spawned component's tick group could be "specified tick group + 1" for the first few frames
-	//but it comes back to original tick group soon
-	if(PrimaryComponentTick.GetActualTickGroup() != PrimaryComponentTick.TickGroup)
+	if(SkeletalMesh == nullptr)
 	{
 		return;
 	}
 
-	if(!SkeletalMesh)
-	{
-		return;
-	}
-
-	int32 NumAssets = SkeletalMesh->ClothingAssets.Num();
-
-	//if clothing assets are added or removed, re-create after removing all
-	if(ClothingActors.Num() != NumAssets)
-	{
-		RemoveAllClothingActors();
-
-		if(NumAssets > 0)
-		{
-			ClothingActors.Empty(NumAssets);
-			ClothingActors.AddZeroed(NumAssets);
-		}
-	}
-
-	bool bNeedUpdateLOD = false;
+	ReleaseAllClothingResources();
 
+	const int32 NumAssets = SkeletalMesh->ClothingAssets.Num();
 	for(int32 AssetIdx=0; AssetIdx < NumAssets; AssetIdx++)
 	{
 		FClothingAssetData& ClothAsset = SkeletalMesh->ClothingAssets[AssetIdx];
@@ -2090,25 +2075,9 @@ void USkeletalMeshComponent::ValidateClothingActors()
 		// if there exist mapped sections for all LODs, create a clothing actor
 		if (SkeletalMesh->HasClothSectionsInAllLODs(AssetIdx))
 		{
-			if(CreateClothingActor(AssetIdx, ClothAsset.ApexClothingAsset))
-			{
-				bNeedUpdateLOD = true;
-			}
-		}
-		else 
-		{	// don't have cloth sections but a clothing actor is alive
-			if(IsValidClothingActor(AssetIdx))
-			{
-				// clear this clothing actor because mapped sections are removed
-				ClothingActors[AssetIdx].Clear(true);
-			}
+			CreateClothingActor(AssetIdx, ClothAsset.ApexClothingAsset);
 		}
 	}
-
-	if(bNeedUpdateLOD)
-	{
-		SetClothingLOD(PredictedLODLevel);
-	}
 }
 
 /** 
@@ -2117,31 +2086,32 @@ void USkeletalMeshComponent::ValidateClothingActors()
 */
 bool USkeletalMeshComponent::CreateClothingActor(int32 AssetIndex, physx::apex::NxClothingAsset* ClothingAsset, TArray<FVector>* BlendedDelta)
 {	
-	int32 NumActors = ClothingActors.Num();
-	int32 ActorIndex = -1;
+	check(ClothingActors.Num() == AssetIndex);
+	ClothingActors.AddZeroed();
 
-	//check we need to expand ClothingActors array 
-	//actor is totally corresponding to asset, 1-on-1, so asset index should be same as actor index
-	if(AssetIndex < NumActors)
+	FPhysScene* PhysScene = NULL;
+
+	if (GetWorld())
 	{
-		ActorIndex = AssetIndex;
+		PhysScene = GetWorld()->GetPhysicsScene();
+	}
 
-		if(IsValidClothingActor(ActorIndex))
-		{
-			// a valid actor already exists
-			return false;
-		}
-		else
-		{
-			//removed or changed or not-created yet
-			ClothingActors[ActorIndex].Clear();
-		}
+	//this clothing actor would be crated later because this becomes invalid
+	if (!PhysScene)
+	{
+		return false;
 	}
 
-	if(ActorIndex < 0)
+	physx::apex::NxApexScene* ScenePtr = PhysScene->GetApexScene(PST_Cloth);
+
+	if (!ScenePtr)
 	{
-		ActorIndex = ClothingActors.AddZeroed();
+		// can't create clothing actor
+		UE_LOG(LogSkeletalMesh, Log, TEXT("CreateClothingActor: Failed to create an actor becauase PhysX Scene doesn't exist"));
+		return false;
 	}
+
+	SCOPED_APEX_SCENE_WRITE_LOCK(ScenePtr);	//we need to lock here because this function could get called from multiple threads and there's memory allocation done inside a singleton asset
 	 
 	// Get the (singleton!) default actor descriptor.
 	NxParameterized::Interface* ActorDesc = ClothingAsset->getDefaultActorDesc();
@@ -2199,34 +2169,10 @@ bool USkeletalMeshComponent::CreateClothingActor(int32 AssetIndex, physx::apex::
 		md.setParamVec3Array(PxBlendedData.GetData(), PxBlendedData.Num());
 	}
 
-	FPhysScene* PhysScene = NULL;
-
-	if(GetWorld())
-	{
-		PhysScene = GetWorld()->GetPhysicsScene();
-	}
-
-	//this clothing actor would be crated later because this becomes invalid
-	if(!PhysScene)
-	{
-		return false;
-	}
-
-	physx::apex::NxApexScene* ScenePtr = PhysScene->GetApexScene(PST_Cloth);
-
-	if(!ScenePtr)
-	{
-		// can't create clothing actor
-		UE_LOG(LogSkeletalMesh, Log, TEXT("CreateClothingActor: Failed to create an actor becauase PhysX Scene doesn't exist"));
-		return false;
-	}
-
 	physx::apex::NxApexActor* apexActor = ClothingAsset->createApexActor(*ActorDesc, *ScenePtr);
 	physx::apex::NxClothingActor* ClothingActor = static_cast<physx::apex::NxClothingActor*>(apexActor);
 
-	ClothingActors[ActorIndex].ApexClothingActor = ClothingActor;
-	ClothingActors[ActorIndex].SceneType = PST_Cloth;
-	ClothingActors[ActorIndex].PhysScene = PhysScene;
+	ClothingActors[AssetIndex].ApexClothingActor = ClothingActor;
 
 	if(!ClothingActor)
 	{
@@ -2235,7 +2181,7 @@ bool USkeletalMeshComponent::CreateClothingActor(int32 AssetIndex, physx::apex::
 	}
 
 	//set parent pointer to verify later whether became invalid or not
-	ClothingActors[ActorIndex].ParentClothingAsset = ClothingAsset;
+	ClothingActors[AssetIndex].ParentClothingAsset = ClothingAsset;
 
 	// budget is millisecond units
 	ScenePtr->setLODResourceBudget(100); // for temporary, 100ms
@@ -2259,6 +2205,7 @@ bool USkeletalMeshComponent::CreateClothingActor(int32 AssetIndex, physx::apex::
 
 void USkeletalMeshComponent::SetClothingLOD(int32 LODIndex)
 {
+	//TODO: are we sure this can't be called during parallel animation?
 	int32 NumActors = ClothingActors.Num();
 
 	bool bFrozen = false;
@@ -2314,8 +2261,16 @@ void USkeletalMeshComponent::SetClothingLOD(int32 LODIndex)
 
 }
 
+bool USkeletalMeshComponent::IsExecutingParallelCloth() const
+{
+	UWorld* World = GetWorld();
+	FPhysScene* PhysScene = World ? World->GetPhysicsScene() : nullptr;
+	return UPhysicsSettings::Get()->bParallelCloth && PhysScene && PhysScene->GetClothManager()->IsPreparingClothAsync();
+}
+
 void USkeletalMeshComponent::RemoveAllClothingActors()
 {
+	check(!IsExecutingParallelCloth());
 	int32 NumActors = ClothingActors.Num();
 
 	for(int32 i=0; i<NumActors; i++)
@@ -2340,7 +2295,7 @@ void USkeletalMeshComponent::ReleaseAllClothingResources()
 #endif// #if WITH_APEX_CLOTHING
 }
 
-void USkeletalMeshComponent::ApplyWindForCloth(FClothingActor& ClothingActor)
+void USkeletalMeshComponent::ApplyWindForCloth(const FClothingActor& ClothingActor) const
 {
 	//to convert from normalized value( usually 0.0 to 1.0 ) to Apex clothing wind value
 	const float WindUnitAmout = 2500.0f;
@@ -2768,6 +2723,7 @@ void USkeletalMeshComponent::FindClothCollisions(TArray<FApexClothCollisionVolum
 
 void USkeletalMeshComponent::CreateInternalClothCollisions(TArray<FApexClothCollisionVolumeData>& InCollisions, TArray<physx::apex::NxClothingCollision*>& OutCollisions)
 {
+	check(!IsExecutingParallelCloth());
 	int32 NumCollisions = InCollisions.Num();
 
 	const int32 MaxNumCapsules = 16;
@@ -2839,6 +2795,7 @@ void USkeletalMeshComponent::CreateInternalClothCollisions(TArray<FApexClothColl
 
 void USkeletalMeshComponent::CopyClothCollisionsToChildren()
 {
+	check(!IsExecutingParallelCloth());
 	// 3 steps
 	// 1. release all previous parent collisions
 	// 2. find new collisions from parent(this class)
@@ -2881,6 +2838,7 @@ void USkeletalMeshComponent::CopyClothCollisionsToChildren()
 
 void USkeletalMeshComponent::ReleaseAllChildrenCollisions()
 {
+	check(!IsExecutingParallelCloth());
 	for(int32 i=0; i<ChildrenCollisions.Num(); i++)
 	{
 		ReleaseClothingCollision(ChildrenCollisions[i]);
@@ -2892,6 +2850,7 @@ void USkeletalMeshComponent::ReleaseAllChildrenCollisions()
 // children's collisions can affect to parent's cloth reversely
 void USkeletalMeshComponent::CopyChildrenClothCollisionsToParent()
 {
+	check(!IsExecutingParallelCloth());
 	// 3 steps
 	// 1. release all previous children collisions
 	// 2. find new collisions from children
@@ -2948,6 +2907,7 @@ void USkeletalMeshComponent::ReleaseClothingCollision(NxClothingCollision* Colli
 
 FApexClothCollisionInfo* USkeletalMeshComponent::CreateNewClothingCollsions(UPrimitiveComponent* PrimitiveComponent)
 {
+	check(!IsExecutingParallelCloth());
 	FApexClothCollisionInfo NewInfo;
 
 	TArray<FClothCollisionPrimitive> CollisionPrims;
@@ -3069,6 +3029,7 @@ FApexClothCollisionInfo* USkeletalMeshComponent::CreateNewClothingCollsions(UPri
 
 void USkeletalMeshComponent::RemoveAllOverlappedComponentMap()
 {
+	check(!IsExecutingParallelCloth());
 	for( auto It = ClothOverlappedComponentsMap.CreateConstIterator(); It; ++It )
 	{
 		const FApexClothCollisionInfo& Info = It.Value();
@@ -3086,6 +3047,7 @@ void USkeletalMeshComponent::RemoveAllOverlappedComponentMap()
 
 void USkeletalMeshComponent::ReleaseAllParentCollisions()
 {
+	check(!IsExecutingParallelCloth());
 	for(int32 i=0; i<ParentCollisions.Num(); i++)
 	{
 		ReleaseClothingCollision(ParentCollisions[i]);
@@ -3128,6 +3090,8 @@ void USkeletalMeshComponent::ProcessClothCollisionWithEnvironment()
 		return;
 	}
 
+	check(!IsExecutingParallelCloth());
+
 	// Increment the revision number
 	ClothingCollisionRevision++;
 
@@ -3227,7 +3191,7 @@ void USkeletalMeshComponent::PreClothTick(float DeltaTime, FTickFunction& ThisTi
 
 #if WITH_APEX_CLOTHING
 	// if skeletal mesh has clothing assets, call TickClothing
-	if (SkeletalMesh && SkeletalMesh->ClothingAssets.Num() > 0)
+	if (SkeletalMesh && SkeletalMesh->ClothingAssets.Num() > 0 && !UPhysicsSettings::Get()->bParallelCloth)
 	{
 		TickClothing(DeltaTime, ThisTickFunction);
 	}
@@ -3274,54 +3238,39 @@ void USkeletalMeshComponent::UpdateClothTransform()
 	}
 }
 
-void USkeletalMeshComponent::CheckClothTeleport(float DeltaTime)
+void USkeletalMeshComponent::CheckClothTeleport(float DeltaTime, FClothSimulationContext& ClothSimulationContext) const
 {
 	// Get the root bone transform
 	FMatrix CurRootBoneMat = GetBoneMatrix(0);
-
-	if(bNeedTeleportAndResetOnceMore)
-	{
-		ForceClothNextUpdateTeleportAndReset();
-		bNeedTeleportAndResetOnceMore = false;
-	}
-
-	float DeltaTimeThreshold = 0.2f;
-	// clothing simulation could be broken if frame-rate goes under 5 fps
-	if(DeltaTime > DeltaTimeThreshold)
-	{
-		ForceClothNextUpdateTeleportAndReset();
-	}
-
+	
 	// distance check 
 	// TeleportDistanceThreshold is greater than Zero and not teleported yet
-	if(TeleportDistanceThreshold > 0 && ClothTeleportMode == FClothingActor::Continuous)
+	if(TeleportDistanceThreshold > 0 && ClothSimulationContext.ClothTeleportMode == FClothingActor::Continuous)
 	{
-		float DistSquared = FVector::DistSquared(PrevRootBoneMatrix.GetOrigin(), CurRootBoneMat.GetOrigin());
-		if ( DistSquared > ClothTeleportDistThresholdSquared ) // if it has traveled too far
+		float DistSquared = FVector::DistSquared(ClothSimulationContext.PrevRootBoneMatrix.GetOrigin(), CurRootBoneMat.GetOrigin());
+		if ( DistSquared > ClothSimulationContext.ClothTeleportDistThresholdSquared ) // if it has traveled too far
 		{
-			ClothTeleportMode = bResetAfterTeleport ? FClothingActor::TeleportAndReset : FClothingActor::Teleport;
-			// clothing reset is needed once more to avoid clothing pop up when moved the component too far
-			bNeedTeleportAndResetOnceMore = true;
+			ClothSimulationContext.ClothTeleportMode = bResetAfterTeleport ? FClothingActor::TeleportAndReset : FClothingActor::Teleport;
 		}
 	}
 
 	// rotation check
 	// if TeleportRotationThreshold is greater than Zero and the user didn't do force teleport
-	if(TeleportRotationThreshold > 0 && ClothTeleportMode == FClothingActor::Continuous)
+	if(TeleportRotationThreshold > 0 && ClothSimulationContext.ClothTeleportMode == FClothingActor::Continuous)
 	{
 		// Detect whether teleportation is needed or not
 		// Rotation matrix's transpose means an inverse but can't use a transpose because this matrix includes scales
-		FMatrix AInvB = CurRootBoneMat * PrevRootBoneMatrix.InverseFast();
+		FMatrix AInvB = CurRootBoneMat * ClothSimulationContext.PrevRootBoneMatrix.InverseFast();
 		float Trace = AInvB.M[0][0] + AInvB.M[1][1] + AInvB.M[2][2];
 		float CosineTheta = (Trace - 1.0f) / 2.0f; // trace = 1+2cos(theta) for a 3x3 matrix
 
-		if ( CosineTheta < ClothTeleportCosineThresholdInRad ) // has the root bone rotated too much
+		if ( CosineTheta < ClothSimulationContext.ClothTeleportCosineThresholdInRad ) // has the root bone rotated too much
 		{
-			ClothTeleportMode = bResetAfterTeleport ? FClothingActor::TeleportAndReset : FClothingActor::Teleport;
+			ClothSimulationContext.ClothTeleportMode = bResetAfterTeleport ? FClothingActor::TeleportAndReset : FClothingActor::Teleport;
 		}
 	}
 
-	PrevRootBoneMatrix = CurRootBoneMat;
+	ClothSimulationContext.PrevRootBoneMatrix = CurRootBoneMat;
 }
 
 void USkeletalMeshComponent::ChangeClothMorphTargetMapping(FClothMorphTargetData& MorphData, FName CurrentActivatedMorphName)
@@ -3368,6 +3317,7 @@ void USkeletalMeshComponent::ChangeClothMorphTargetMapping(FClothMorphTargetData
 
 void USkeletalMeshComponent::PrepareClothMorphTargets()
 {	
+	check(!IsExecutingParallelCloth());
 	// if didn't turn on the cloth morph target option or already precomputed
 	if (!bClothMorphTarget || bPreparedClothMorphTargets)
 	{ 
@@ -3496,6 +3446,8 @@ void USkeletalMeshComponent::PrepareClothMorphTargets()
 
 void USkeletalMeshComponent::UpdateClothMorphTarget()
 {
+	check(!IsExecutingParallelCloth());
+
 	if (!bClothMorphTarget)
 	{
 		return;
@@ -3571,6 +3523,9 @@ void USkeletalMeshComponent::UpdateClothState(float DeltaTime)
 		return;
 	}
 
+	check(!IsExecutingParallelCloth());
+	check(IsInGameThread());
+	
 	// if turned on bClothMorphTarget option
 	if (bClothMorphTarget)
 	{
@@ -3584,41 +3539,50 @@ void USkeletalMeshComponent::UpdateClothState(float DeltaTime)
 	}
 
 #if WITH_CLOTH_COLLISION_DETECTION
-
-	if(bCollideWithAttachedChildren)
+	if (bCollideWithAttachedChildren)
 	{
 		CopyClothCollisionsToChildren();
 		CopyChildrenClothCollisionsToParent();
 	}
 #endif // WITH_CLOTH_COLLISION_DETECTION
 
+	ParallelUpdateClothState(DeltaTime, EditableClothSimulationContext);
+}
+
+void USkeletalMeshComponent::ParallelUpdateClothState(float DeltaTime, FClothSimulationContext& ClothSimulationContext) const
+{
+	if (CVarEnableClothPhysics.GetValueOnAnyThread() == 0)
+	{
+		return;
+	}
+
 	int32 NumActors = ClothingActors.Num();
-	if(NumActors == 0)
+	if (NumActors == 0)
 	{
 		return;
 	}
 
 	const TArray<FTransform>& BoneTransforms = MasterPoseComponent.IsValid() ? MasterPoseComponent.Get()->GetSpaceBases() : GetSpaceBases();
 
-	if(BoneTransforms.Num() == 0)
+	if (BoneTransforms.Num() == 0)
 	{
 		return;
 	}
 
 	physx::PxMat44 PxGlobalPose = U2PMatrix(ComponentToWorld.ToMatrixWithScale());
 
-	CheckClothTeleport(DeltaTime);
+	CheckClothTeleport(DeltaTime, ClothSimulationContext);
 
 	// convert teleport mode to apex clothing teleport enum
-	physx::apex::ClothingTeleportMode::Enum CurTeleportMode = (physx::apex::ClothingTeleportMode::Enum)ClothTeleportMode;
+	physx::apex::ClothingTeleportMode::Enum CurTeleportMode = (physx::apex::ClothingTeleportMode::Enum)ClothSimulationContext.ClothTeleportMode;
 
-	static TArray<physx::PxMat44> BoneMatrices;
-	check(IsInGameThread() && !BoneMatrices.Num());
+	TArray<physx::PxMat44> BoneMatrices;
+	check(!BoneMatrices.Num());
 
-	for(int32 ActorIdx=0; ActorIdx<NumActors; ActorIdx++)
+	for (int32 ActorIdx = 0; ActorIdx < NumActors; ActorIdx++)
 	{
 		// skip if ClothingActor is NULL or invalid
-		if(!IsValidClothingActor(ActorIdx))
+		if (!IsValidClothingActor(ActorIdx))
 		{
 			continue;
 		}
@@ -3632,22 +3596,22 @@ void USkeletalMeshComponent::UpdateClothState(float DeltaTime)
 		check(!BoneMatrices.Num());
 		BoneMatrices.AddUninitialized(NumUsedBones);
 
-		for(uint32 Index=0; Index < NumUsedBones; Index++)
+		for (uint32 Index = 0; Index < NumUsedBones; Index++)
 		{
 			FName BoneName = GetConvertedBoneName(ClothingActors[ActorIdx].ParentClothingAsset, Index);
 
-		   int32 BoneIndex = GetBoneIndex(BoneName);
+			int32 BoneIndex = GetBoneIndex(BoneName);
 
-			if(MasterPoseComponent.IsValid())
+			if (MasterPoseComponent.IsValid())
 			{
 				int32 TempBoneIndex = BoneIndex;
 				BoneIndex = INDEX_NONE;
-				if(TempBoneIndex < MasterBoneMap.Num())
+				if (TempBoneIndex < MasterBoneMap.Num())
 				{
 					int32 MasterBoneIndex = MasterBoneMap[TempBoneIndex];
- 
+
 					// If ParentBoneIndex is valid, grab matrix from MasterPoseComponent.
-					if( MasterBoneIndex != INDEX_NONE && 
+					if (MasterBoneIndex != INDEX_NONE &&
 						MasterBoneIndex < MasterPoseComponent->GetNumSpaceBases())
 					{
 						BoneIndex = MasterBoneIndex;
@@ -3655,14 +3619,14 @@ void USkeletalMeshComponent::UpdateClothState(float DeltaTime)
 				}
 			}
 
-		   if(BoneIndex != INDEX_NONE)
-		   {
-			   BoneMatrices[Index] = U2PMatrix(BoneTransforms[BoneIndex].ToMatrixWithScale());
-		   }
-		   else
-		   {
-			   BoneMatrices[Index] = PxMat44::createIdentity();
-		   }
+			if (BoneIndex != INDEX_NONE)
+			{
+				BoneMatrices[Index] = U2PMatrix(BoneTransforms[BoneIndex].ToMatrixWithScale());
+			}
+			else
+			{
+				BoneMatrices[Index] = PxMat44::createIdentity();
+			}
 		}
 
 		NxClothingActor* ClothingActor = ClothingActors[ActorIdx].ApexClothingActor;
@@ -3671,17 +3635,18 @@ void USkeletalMeshComponent::UpdateClothState(float DeltaTime)
 
 		// if bUseInternalboneOrder is set, "NumUsedBones" works, otherwise have to use "getNumBones" 
 		ClothingActor->updateState(
-			PxGlobalPose, 
-			BoneMatrices.GetData(), 
-			sizeof(physx::PxMat44), 
+			PxGlobalPose,
+			BoneMatrices.GetData(),
+			sizeof(physx::PxMat44),
 			NumUsedBones,
 			CurTeleportMode);
 
 		BoneMatrices.Reset();
 	}
 
+	//move to post cloth tick
 	// reset to Continuous
-	ClothTeleportMode = FClothingActor::Continuous;
+	ClothSimulationContext.ClothTeleportMode = FClothingActor::Continuous;
 }
 
 void USkeletalMeshComponent::GetClothRootBoneMatrix(int32 AssetIndex, FMatrix& OutRootBoneMatrix) const
@@ -3722,6 +3687,7 @@ bool USkeletalMeshComponent::GetClothSimulatedPosition(int32 AssetIndex, int32 V
 			// handles only simulated vertices, indices of fixed vertices are bigger than # of simulated vertices
 			if ((uint32)VertexIndex < NumSimulVertices)
 			{
+				//TODO: this is not thread safe! apex requires a lock. Not going to add it here because it's called for many vertices, but this should be addressed!
 				bSucceed = true;
 				const physx::PxVec3* Vertices = ApexClothingActor->getSimulationPositions();
 
@@ -3794,7 +3760,6 @@ void USkeletalMeshComponent::PerformTickClothing(float DeltaTime)
 	}
 	else
 	{
-		ValidateClothingActors();
 		UpdateClothState(DeltaTime);
 	}
 #if 0 //if turn on this flag, you can check which primitive objects are activated for collision detection
@@ -3803,8 +3768,24 @@ void USkeletalMeshComponent::PerformTickClothing(float DeltaTime)
 #endif// #if WITH_APEX_CLOTHING
 }
 
+void USkeletalMeshComponent::ParallelTickClothing(float DeltaTime, FClothSimulationContext& ClothSimulationContext) const
+{
+#if WITH_APEX_CLOTHING
+	// animated but bone transforms were not updated because it was not rendered
+	if (PoseTickedThisFrame() && !bRecentlyRendered)
+	{
+		ClothSimulationContext.ClothTeleportMode = FClothingActor::TeleportAndReset;
+	}
+	else
+	{
+		ParallelUpdateClothState(DeltaTime, ClothSimulationContext);
+	}
+#endif 
+}
+
 void USkeletalMeshComponent::TickClothing(float DeltaTime, FTickFunction& ThisTickFunction)
 {
+
 	if(IsValidRef(ParallelBlendPhysicsCompletionTask))
 	{
 		FGraphEventArray Prerequistes;
@@ -3909,6 +3890,7 @@ void USkeletalMeshComponent::FreezeClothSection(bool bFreeze)
 {
 #if WITH_APEX_CLOTHING
 
+	check(!IsExecutingParallelCloth());
 	int32 NumActors = ClothingActors.Num();
 
 	for(int32 ActorIdx=0; ActorIdx<NumActors; ActorIdx++)
@@ -5081,6 +5063,7 @@ float USkeletalMeshComponent::GetClothMaxDistanceScale()
 void USkeletalMeshComponent::SetClothMaxDistanceScale(float Scale)
 {
 #if WITH_APEX_CLOTHING
+	///TODO: add APEX lock
 
 	//this scale parameter is also used when new clothing actor is created
 	ClothMaxDistanceScale = Scale;
@@ -5110,21 +5093,21 @@ void USkeletalMeshComponent::SetClothMaxDistanceScale(float Scale)
 void USkeletalMeshComponent::ResetClothTeleportMode()
 {
 #if WITH_APEX_CLOTHING
-	ClothTeleportMode = FClothingActor::Continuous;
+	EditableClothSimulationContext.ClothTeleportMode = FClothingActor::Continuous;
 #endif// #if WITH_APEX_CLOTHING
 }
 
 void USkeletalMeshComponent::ForceClothNextUpdateTeleport()
 {
 #if WITH_APEX_CLOTHING
-	ClothTeleportMode = FClothingActor::Teleport;
+	EditableClothSimulationContext.ClothTeleportMode = FClothingActor::Teleport;
 #endif// #if WITH_APEX_CLOTHING
 }
 
 void USkeletalMeshComponent::ForceClothNextUpdateTeleportAndReset()
 {
 #if WITH_APEX_CLOTHING
-	ClothTeleportMode = FClothingActor::TeleportAndReset;
+	EditableClothSimulationContext.ClothTeleportMode = FClothingActor::TeleportAndReset;
 #endif// #if WITH_APEX_CLOTHING
 }
 
diff --git a/Engine/Source/Runtime/Engine/Private/SkeletalRenderGPUSkin.cpp b/Engine/Source/Runtime/Engine/Private/SkeletalRenderGPUSkin.cpp
--- a/Engine/Source/Runtime/Engine/Private/SkeletalRenderGPUSkin.cpp
+++ b/Engine/Source/Runtime/Engine/Private/SkeletalRenderGPUSkin.cpp
@@ -26,6 +26,7 @@
 #include "SkeletalRenderCPUSkin.h"
 #include "GPUSkinCache.h"
 #include "Animation/VertexAnim/VertexAnimBase.h"
+#include "Components/SkeletalMeshComponent.h"
 
 DEFINE_LOG_CATEGORY_STATIC(LogSkeletalGPUSkinMesh, Warning, All);
 
@@ -1126,7 +1127,7 @@ bool FDynamicSkelMeshObjectDataGPUSkin::UpdateClothSimulationData(USkinnedMeshCo
 	USkeletalMeshComponent * SkelMeshComponent = Cast<USkeletalMeshComponent>(InMeshComponent);
 
 #if WITH_APEX_CLOTHING
-	if(InMeshComponent->MasterPoseComponent.IsValid() && SkelMeshComponent->bBindClothToMasterComponent)
+	if(InMeshComponent->MasterPoseComponent.IsValid() && SkelMeshComponent->IsClothBoundToMasterComponent())
 	{
 		USkeletalMeshComponent* OriginalComponent = Cast<USkeletalMeshComponent>(InMeshComponent);
 		SkelMeshComponent = Cast<USkeletalMeshComponent>(InMeshComponent->MasterPoseComponent.Get());
diff --git a/Engine/Source/Runtime/Engine/Private/SkeletalRenderGPUSkin.h b/Engine/Source/Runtime/Engine/Private/SkeletalRenderGPUSkin.h
--- a/Engine/Source/Runtime/Engine/Private/SkeletalRenderGPUSkin.h
+++ b/Engine/Source/Runtime/Engine/Private/SkeletalRenderGPUSkin.h
@@ -8,7 +8,8 @@
 
 #include "SkeletalRender.h"
 #include "SkeletalRenderPublic.h"
-#include "GPUSkinVertexFactory.h" 
+#include "GPUSkinVertexFactory.h"
+#include "ClothSimData.h"
 
 // 1 for a single buffer, this creates and releases textures every frame (the driver has to keep the reference and need to defer the release, low memory as it only occupies rendered buffers (up to 3 copies), best Xbox360 method?)
 // 2 for double buffering (works well for PC, caused Xbox360 to stall)
diff --git a/Engine/Source/Runtime/Engine/Private/World.cpp b/Engine/Source/Runtime/Engine/Private/World.cpp
--- a/Engine/Source/Runtime/Engine/Private/World.cpp
+++ b/Engine/Source/Runtime/Engine/Private/World.cpp
@@ -40,6 +40,7 @@
 #include "GameFramework/GameMode.h"
 #include "GameFramework/GameState.h"
 #include "GameFramework/PlayerState.h"
+#include "PhysicsEngine/PhysicsConstraintComponent.h"
 
 #include "Materials/MaterialParameterCollectionInstance.h"
 
@@ -3449,7 +3450,7 @@ void UWorld::SetPhysicsScene(FPhysScene* InScene)
 	// Clear world pointer in old FPhysScene (if there is one)
 	if(PhysicsScene != NULL)
 	{
-		PhysicsScene->OwningWorld = NULL;
+		PhysicsScene->SetOwningWorld(nullptr);
 		delete PhysicsScene;
 	}
 
@@ -3459,7 +3460,7 @@ void UWorld::SetPhysicsScene(FPhysScene* InScene)
 	// Set pointer in scene to know which world its coming from
 	if(PhysicsScene != NULL)
 	{
-		PhysicsScene->OwningWorld = this;
+		PhysicsScene->SetOwningWorld(this);
 	}
 }
 
diff --git a/Engine/Source/Runtime/Engine/Public/ClothSimData.h b/Engine/Source/Runtime/Engine/Public/ClothSimData.h
new file mode 100644
--- /dev/null
+++ b/Engine/Source/Runtime/Engine/Public/ClothSimData.h
@@ -0,0 +1,10 @@
+// Copyright 1998-2015 Epic Games, Inc. All Rights Reserved.
+
+#pragma once
+
+/** data for updating cloth section from the results of clothing simulation */
+struct FClothSimulData
+{
+	TArray<FVector4> ClothSimulPositions;
+	TArray<FVector4> ClothSimulNormals;
+};
diff --git a/Engine/Source/Runtime/Engine/Public/PhysicsPublic.h b/Engine/Source/Runtime/Engine/Public/PhysicsPublic.h
--- a/Engine/Source/Runtime/Engine/Public/PhysicsPublic.h
+++ b/Engine/Source/Runtime/Engine/Public/PhysicsPublic.h
@@ -266,6 +266,8 @@ namespace SleepEvent
 
 }
 
+class FClothManager;
+
 /** Container object for a physics engine 'scene'. */
 
 class FPhysScene
@@ -287,9 +289,16 @@ public:
 	/** Gets the array of collision notifications, pending execution at the end of the physics engine run. */
 	TArray<FCollisionNotifyInfo>& GetPendingCollisionNotifies(int32 SceneType){ return PendingCollisionData[SceneType].PendingCollisionNotifies; }
 
+private:
 	/** World that owns this physics scene */
 	UWorld*							OwningWorld;
 
+public:
+	//Owning world is made private so that any code which depends on setting an owning world can update
+	void SetOwningWorld(UWorld* InOwningWorld);
+	UWorld* GetOwningWorld(){ return OwningWorld; }
+	const UWorld* GetOwningWorld() const { return OwningWorld; }
+
 	/** These indices are used to get the actual PxScene or NxApexScene from the GPhysXSceneMap. */
 	int16								PhysXSceneIndex[PST_MAX];
 
@@ -338,6 +347,8 @@ public:
 	void AddPendingSleepingEvent(PxActor* Actor, SleepEvent::Type SleepEventType, int32 SceneType);
 #endif
 
+	FClothManager* GetClothManager() const { return ClothManager; }
+
 private:
 	/** DeltaSeconds from UWorld. */
 	float										DeltaSeconds;
@@ -355,6 +366,8 @@ private:
 	/** Completion events (task) for the physics scenes	(both apex and non-apex). This is a "join" of the above. */
 	FGraphEventRef PhysicsSceneCompletion;
 
+	FClothManager* ClothManager;
+
 #if WITH_PHYSX
 
 	struct FDeferredSceneData
@@ -433,6 +446,9 @@ public:
 	/** Starts cloth Simulation*/
 	ENGINE_API void StartCloth();
 
+	/** Starts cloth Simulation*/
+	ENGINE_API void StartAsync();
+
 	/** returns the completion event for a frame */
 	FGraphEventRef GetCompletionEvent()
 	{
diff --git a/Engine/Source/Runtime/Engine/Public/SingleAnimationPlayData.h b/Engine/Source/Runtime/Engine/Public/SingleAnimationPlayData.h
new file mode 100644
--- /dev/null
+++ b/Engine/Source/Runtime/Engine/Public/SingleAnimationPlayData.h
@@ -0,0 +1,54 @@
+// Copyright 1998-2015 Epic Games, Inc. All Rights Reserved.
+
+#pragma once
+
+#include "SingleAnimationPlayData.generated.h"
+
+USTRUCT()
+struct FSingleAnimationPlayData
+{
+	GENERATED_USTRUCT_BODY()
+
+	// @todo in the future, we should make this one UObject
+	// and have detail customization to display different things
+	// The default sequence to play on this skeletal mesh
+	UPROPERTY(EditAnywhere, BlueprintReadWrite, Category = Animation)
+	class UAnimationAsset* AnimToPlay;
+
+	// @fixme : until we properly support it I'm commenting out editable property part
+	// The default sequence to play on this skeletal mesh
+	UPROPERTY()//EditAnywhere, BlueprintReadWrite, Category=Animation)
+	class UVertexAnimation* VertexAnimToPlay;
+
+	/** Default setting for looping for SequenceToPlay. This is not current state of looping. */
+	UPROPERTY(EditAnywhere, BlueprintReadWrite, Category = Animation, meta = (DisplayName = "Looping"))
+		uint32 bSavedLooping : 1;
+
+	/** Default setting for playing for SequenceToPlay. This is not current state of playing. */
+	UPROPERTY(EditAnywhere, BlueprintReadWrite, Category = Animation, meta = (DisplayName = "Playing"))
+		uint32 bSavedPlaying : 1;
+
+	/** Default setting for position of SequenceToPlay to play. */
+	UPROPERTY(EditAnywhere, BlueprintReadWrite, Category = Animation, meta = (DisplayName = "Initial Position"))
+		float SavedPosition;
+
+	/** Default setting for play rate of SequenceToPlay to play. */
+	UPROPERTY(EditAnywhere, AdvancedDisplay, BlueprintReadWrite, Category = Animation, meta = (DisplayName = "PlayRate"))
+		float SavedPlayRate;
+
+	FSingleAnimationPlayData()
+	{
+		SavedPlayRate = 1.0f;
+		bSavedLooping = true;
+		bSavedPlaying = true;
+		SavedPosition = 0.f;
+	}
+
+	/** Called when initialized. */
+	ENGINE_API void Initialize(class UAnimSingleNodeInstance* Instance);
+
+	/** Populates this play data with the current state of the supplied instance. */
+	ENGINE_API void PopulateFrom(UAnimSingleNodeInstance* Instance);
+
+	void ValidatePosition();
+};
\ No newline at end of file
