Fix non APEX build

diff --git a/Engine/Source/Runtime/Engine/Private/Components/SkeletalMeshComponent.cpp b/Engine/Source/Runtime/Engine/Private/Components/SkeletalMeshComponent.cpp
--- a/Engine/Source/Runtime/Engine/Private/Components/SkeletalMeshComponent.cpp
+++ b/Engine/Source/Runtime/Engine/Private/Components/SkeletalMeshComponent.cpp
@@ -1086,7 +1086,7 @@ int32 GetCurveNumber(USkeleton* Skeleton)
 
 void USkeletalMeshComponent::SubmitClothSimulationContext()
 {
-	
+#if WITH_APEX_CLOTHING
 	InternalClothSimulationContext.ClothTeleportCosineThresholdInRad = EditableClothSimulationContext.ClothTeleportCosineThresholdInRad;
 	InternalClothSimulationContext.ClothTeleportDistThresholdSquared = EditableClothSimulationContext.ClothTeleportDistThresholdSquared;
 
@@ -1095,6 +1095,7 @@ void USkeletalMeshComponent::SubmitClothSimulationContext()
 	EditableClothSimulationContext.ClothTeleportMode = FClothingActor::Default;	//reset editable teleport mode to default because user has to set it directly
 
 	//we intentionally ignore PrevRootBone which is only needed internally, but is still in the Context struct for the sake of const
+#endif
 }
 
 void USkeletalMeshComponent::RefreshBoneTransforms(FActorComponentTickFunction* TickFunction)
