non-APEX build fix take 2

diff --git a/Engine/Source/Runtime/Engine/Private/PhysicsEngine/ClothManager.cpp b/Engine/Source/Runtime/Engine/Private/PhysicsEngine/ClothManager.cpp
--- a/Engine/Source/Runtime/Engine/Private/PhysicsEngine/ClothManager.cpp
+++ b/Engine/Source/Runtime/Engine/Private/PhysicsEngine/ClothManager.cpp
@@ -17,6 +17,7 @@ FClothManager::FClothManager(UWorld* AssociatedWorld, FPhysScene* InPhysScene)
 
 void FClothManager::Tick(float DeltaTime)
 {
+#if WITH_APEX_CLOTHING
 	if(SkeletalMeshComponents.Num())
 	{
 		IsPreparingCloth.AtomicSet(true);
@@ -32,6 +33,7 @@ void FClothManager::Tick(float DeltaTime)
 		PhysScene->StartCloth();
 		IsPreparingCloth.AtomicSet(false);
 	}
+#endif
 }
 
 void FClothManager::RegisterForClothThisFrame(USkeletalMeshComponent* SkeletalMeshComponent)
