#ue4 Fix for HMDFunctionLibrary IsHMDEnabled() functionality.

diff --git a/Engine/Source/Runtime/Engine/Private/HeadMountedDisplayFunctionLibrary.cpp b/Engine/Source/Runtime/Engine/Private/HeadMountedDisplayFunctionLibrary.cpp
--- a/Engine/Source/Runtime/Engine/Private/HeadMountedDisplayFunctionLibrary.cpp
+++ b/Engine/Source/Runtime/Engine/Private/HeadMountedDisplayFunctionLibrary.cpp
@@ -13,7 +13,7 @@ UHeadMountedDisplayFunctionLibrary::UHeadMountedDisplayFunctionLibrary(const FOb
 
 bool UHeadMountedDisplayFunctionLibrary::IsHeadMountedDisplayEnabled()
 {
-	return GEngine->HMDDevice.IsValid() && GEngine->HMDDevice->IsHMDEnabled();
+	return GEngine->HMDDevice.IsValid() && GEngine->HMDDevice->IsHeadTrackingAllowed();
 }
 
 bool UHeadMountedDisplayFunctionLibrary::EnableHMD(bool bEnable)
