FPendingLatentAction instances leaked if world is destroyed while they're still in progress

https://udn.unrealengine.com/questions/260117/fpendinglatentaction-instances-leaked-if-world-is.html

diff --git a/Engine/Source/Runtime/Engine/Classes/Engine/LatentActionManager.h b/Engine/Source/Runtime/Engine/Classes/Engine/LatentActionManager.h
--- a/Engine/Source/Runtime/Engine/Classes/Engine/LatentActionManager.h
+++ b/Engine/Source/Runtime/Engine/Classes/Engine/LatentActionManager.h
@@ -146,6 +146,8 @@ public:
 	FString GetDescription(UObject* InObject, int32 UUID) const;
 #endif
 
+	~FLatentActionManager();
+
 protected:
 	/** 
 	 * Finds the action instance for the supplied object will return NULL if one does not exist.
diff --git a/Engine/Source/Runtime/Engine/Private/LatentActionManager.cpp b/Engine/Source/Runtime/Engine/Private/LatentActionManager.cpp
--- a/Engine/Source/Runtime/Engine/Private/LatentActionManager.cpp
+++ b/Engine/Source/Runtime/Engine/Private/LatentActionManager.cpp
@@ -230,4 +230,22 @@ void FLatentActionManager::GetActiveUUIDs(UObject* InObject, TSet<int32>& UUIDLi
 	}
 }
 
-#endif
\ No newline at end of file
+#endif
+
+FLatentActionManager::~FLatentActionManager()
+{
+	for (auto& ObjectActionListIterator : ObjectToActionListMap)
+	{
+		TSharedPtr<FActionList>& ActionList = ObjectActionListIterator.Value;
+		if (ActionList.IsValid())
+		{
+			for (auto& ActionIterator : *ActionList.Get())
+			{
+				FPendingLatentAction* Action = ActionIterator.Value;
+				ActionIterator.Value = nullptr;
+				delete Action;
+			}
+			ActionList->Reset();
+		}
+	}
+}
\ No newline at end of file
