Fix RepairWorldSettings crash

diff --git a/Engine/Source/Runtime/Engine/Private/World.cpp b/Engine/Source/Runtime/Engine/Private/World.cpp
--- a/Engine/Source/Runtime/Engine/Private/World.cpp
+++ b/Engine/Source/Runtime/Engine/Private/World.cpp
@@ -836,6 +836,14 @@ void UWorld::RepairWorldSettings()
 			PersistentLevel->Actors[0]->Rename(NULL, PersistentLevel, REN_ForceNoResetLoaders);
 		}
 		
+		bool bClearOwningWorld = false;
+
+		if (PersistentLevel->OwningWorld == nullptr)
+		{
+			bClearOwningWorld = true;
+			PersistentLevel->OwningWorld = this;
+		}
+
 		FActorSpawnParameters SpawnInfo;
 		SpawnInfo.SpawnCollisionHandlingOverride = ESpawnActorCollisionHandlingMethod::AlwaysSpawn;
 		SpawnInfo.Name = GEngine->WorldSettingsClass->GetFName();
@@ -861,6 +869,10 @@ void UWorld::RepairWorldSettings()
 		// Re-sort actor list as we just shuffled things around.
 		PersistentLevel->SortActorList();
 
+		if (bClearOwningWorld)
+		{
+			PersistentLevel->OwningWorld = nullptr;
+		}
 	}
 	check(GetWorldSettings());
 }
