UE-19098 Add ENGINE_API to WindDirectionalSourceComponent.
Remove WindDirectionalSource.h and WindDirectionalSourceComponent.h from Engine.h

diff --git a/Engine/Source/Runtime/Engine/Classes/Components/WindDirectionalSourceComponent.h b/Engine/Source/Runtime/Engine/Classes/Components/WindDirectionalSourceComponent.h
--- a/Engine/Source/Runtime/Engine/Classes/Components/WindDirectionalSourceComponent.h
+++ b/Engine/Source/Runtime/Engine/Classes/Components/WindDirectionalSourceComponent.h
@@ -2,11 +2,12 @@
 
 
 #pragma once
+#include "Components/SceneComponent.h"
 #include "WindDirectionalSourceComponent.generated.h"
 
 /** Component that provides a directional wind source. Only affects SpeedTree assets. */
 UCLASS(collapsecategories, hidecategories=(Object, Mobility), editinlinenew)
-class UWindDirectionalSourceComponent : public USceneComponent
+class ENGINE_API UWindDirectionalSourceComponent : public USceneComponent
 {
 	GENERATED_UCLASS_BODY()
 
diff --git a/Engine/Source/Runtime/Engine/Public/Engine.h b/Engine/Source/Runtime/Engine/Public/Engine.h
--- a/Engine/Source/Runtime/Engine/Public/Engine.h
+++ b/Engine/Source/Runtime/Engine/Public/Engine.h
@@ -68,7 +68,6 @@
 #include "GameFramework/GameState.h"
 #include "GameFramework/PlayerState.h"
 #include "Engine/SkyLight.h"
-#include "Engine/WindDirectionalSource.h"
 #include "GameFramework/WorldSettings.h"
 #include "Components/LightComponentBase.h"
 #include "Engine/LevelBounds.h"
@@ -162,7 +161,6 @@
 #include "Components/SceneCaptureComponent.h"
 #include "Components/SceneCaptureComponent2D.h"
 #include "Components/SceneCaptureComponentCube.h"
-#include "Components/WindDirectionalSourceComponent.h"
 #include "Components/TimelineComponent.h"
 #include "Animation/AnimCurveTypes.h"
 #include "Animation/AnimSequenceBase.h"
diff --git a/Engine/Source/Runtime/Renderer/Private/RendererScene.cpp b/Engine/Source/Runtime/Renderer/Private/RendererScene.cpp
--- a/Engine/Source/Runtime/Renderer/Private/RendererScene.cpp
+++ b/Engine/Source/Runtime/Renderer/Private/RendererScene.cpp
@@ -16,6 +16,7 @@
 #include "DistanceFieldLightingShared.h"
 #include "SpeedTreeWind.h"
 #include "HeightfieldLighting.h"
+#include "Components/WindDirectionalSourceComponent.h"
 
 // Enable this define to do slow checks for components being added to the wrong
 // world's scene, when using PIE. This can happen if a PIE component is reattached
