UE-18030 : Disable trying to draw constraints in level viewport with 'show constraints', because it currently does not work

diff --git a/Engine/Source/Runtime/Engine/Private/SkeletalMesh.cpp b/Engine/Source/Runtime/Engine/Private/SkeletalMesh.cpp
--- a/Engine/Source/Runtime/Engine/Private/SkeletalMesh.cpp
+++ b/Engine/Source/Runtime/Engine/Private/SkeletalMesh.cpp
@@ -4801,7 +4801,8 @@ void FSkeletalMeshSceneProxy::DebugDrawPhysicsAsset(int32 ViewIndex, FMeshElemen
 			}
 			if( EngineShowFlags.Constraints )
 			{
-				PhysicsAssetForDebug->DrawConstraints(Collector.GetPDI(ViewIndex), SkeletalMeshForDebug, *BoneSpaceBases, LocalToWorldTransform, TotalScale.X);
+				// TODO: Currently doesn't work because the PDI returned here (FSimpleElementCollector) can't draw meshes
+				//PhysicsAssetForDebug->DrawConstraints(Collector.GetPDI(ViewIndex), SkeletalMeshForDebug, *BoneSpaceBases, LocalToWorldTransform, TotalScale.X);
 			}
 		}
 	}
