UE-13367 : Remove duplicate (non-BP exposed) AHUD:DrawText function that uses different co-ordinates

diff --git a/Engine/Source/Runtime/Engine/Classes/GameFramework/HUD.h b/Engine/Source/Runtime/Engine/Classes/GameFramework/HUD.h
--- a/Engine/Source/Runtime/Engine/Classes/GameFramework/HUD.h
+++ b/Engine/Source/Runtime/Engine/Classes/GameFramework/HUD.h
@@ -464,9 +464,6 @@ public:
 	// Messaging.
 	//=============================================================================
 
-	/** Display current messages */
-	virtual void DrawText(const FString& Text, FVector2D Position, UFont* TextFont, FVector2D FontScale, FColor TextColor);
-	
 	/** @return UFont* from given FontSize index */
 	virtual  UFont* GetFontFromSizeIndex(int32 FontSize) const;
 
diff --git a/Engine/Source/Runtime/Engine/Private/HUD.cpp b/Engine/Source/Runtime/Engine/Private/HUD.cpp
--- a/Engine/Source/Runtime/Engine/Private/HUD.cpp
+++ b/Engine/Source/Runtime/Engine/Private/HUD.cpp
@@ -373,18 +373,6 @@ void AHUD::DrawHUD()
 	ReceiveDrawHUD(Canvas->SizeX, Canvas->SizeY);
 }
 
-void AHUD::DrawText(const FString& Text, FVector2D Position, UFont* TextFont, FVector2D FontScale, FColor TextColor)
-{
-	float XL, YL;
-	Canvas->TextSize(TextFont, Text, XL, YL);
-	const float X = Canvas->ClipX/2.0f - XL/2.0f + Position.X;
-	const float Y = Canvas->ClipY/2.0f - YL/2.0f + Position.Y;
-	FCanvasTextItem TextItem( FVector2D( X, Y ), FText::FromString( Text ), TextFont, FLinearColor( TextColor ) );
-	TextItem.Scale = FontScale;
-	
-	Canvas->DrawItem(TextItem);
-}
-
 UFont* AHUD::GetFontFromSizeIndex(int32 FontSizeIndex) const
 {
 	switch (FontSizeIndex)
