Remove SetHidden UFUNCTION() which isn't necessary anymore after CL#2668779.

diff --git a/Engine/Source/Runtime/Engine/Classes/GameFramework/Actor.h b/Engine/Source/Runtime/Engine/Classes/GameFramework/Actor.h
--- a/Engine/Source/Runtime/Engine/Classes/GameFramework/Actor.h
+++ b/Engine/Source/Runtime/Engine/Classes/GameFramework/Actor.h
@@ -108,8 +108,6 @@ public:
 	 */
 	UPROPERTY(Interp, EditAnywhere, Category=Rendering, BlueprintReadOnly, replicated, meta=(DisplayName = "Actor Hidden In Game", SequencerTrackClass = "MovieSceneVisibilityTrack"))
 	uint32 bHidden:1;
-	UFUNCTION()
-	void SetHidden(bool bInHidden) { bHidden = bInHidden; }
 
 	/** If true, when the actor is spawned it will be sent to the client but receive no further replication updates from the server afterwards. */
 	UPROPERTY()
