#ANIM: Better fix for kurohane crash due to inconsistent weight of montages

- Now it gets weight data and use that weight data without recalc in eval
- Also made montage instance's members to be private
- changed single node to get weight and use that weight for eval - same as normal anim BP
- changed update weight of montage to happen before native update animation

Merging

//depot/UE4-Orion/Engine/Source/...

to //depot/UE4/Engine/Source/...

diff --git a/Engine/Source/Editor/AnimGraph/Private/AnimPreviewInstance.cpp b/Engine/Source/Editor/AnimGraph/Private/AnimPreviewInstance.cpp
--- a/Engine/Source/Editor/AnimGraph/Private/AnimPreviewInstance.cpp
+++ b/Engine/Source/Editor/AnimGraph/Private/AnimPreviewInstance.cpp
@@ -973,10 +973,7 @@ int32 UAnimPreviewInstance::MontagePreview_FindLastSection(int32 StartSectionIdx
 				{
 					AlreadyVisited[CurrentSectionIdx] = true;
 					ResultIdx = CurrentSectionIdx;
-					if (CurMontageInstance->NextSections.IsValidIndex(CurrentSectionIdx))
-					{
-						CurrentSectionIdx = CurMontageInstance->NextSections[CurrentSectionIdx];
-					}
+					CurrentSectionIdx = CurMontageInstance->GetNextSectionID(CurrentSectionIdx);
 				}
 			}
 		}
diff --git a/Engine/Source/Runtime/AnimGraphRuntime/Private/AnimNodes/AnimNode_Slot.cpp b/Engine/Source/Runtime/AnimGraphRuntime/Private/AnimNodes/AnimNode_Slot.cpp
--- a/Engine/Source/Runtime/AnimGraphRuntime/Private/AnimNodes/AnimNode_Slot.cpp
+++ b/Engine/Source/Runtime/AnimGraphRuntime/Private/AnimNodes/AnimNode_Slot.cpp
@@ -12,9 +12,7 @@ void FAnimNode_Slot::Initialize(const FAnimationInitializeContext& Context)
 	FAnimNode_Base::Initialize(Context);
 
 	Source.Initialize(Context);
-
-	SourceWeight = 0.f;
-	SlotNodeWeight = 0.f;
+	WeightData.Reset();
 
 	// If this node has not already been registered with the AnimInstance, do it.
 	if ((LastSlotNodeInitializationCounter == INDEX_NONE) || (LastSlotNodeInitializationCounter != Context.AnimInstance->GetSlotNodeInitializationCounter()))
@@ -32,36 +30,34 @@ void FAnimNode_Slot::CacheBones(const FAnimationCacheBonesContext& Context)
 void FAnimNode_Slot::Update(const FAnimationUpdateContext& Context)
 {
 	// Update weights.
-	Context.AnimInstance->GetSlotWeight(SlotName, SlotNodeWeight, SourceWeight);
+	Context.AnimInstance->GetSlotWeight(SlotName, WeightData.SlotNodeWeight, WeightData.SourceWeight, WeightData.TotalNodeWeight);
 
 	// Update cache in AnimInstance.
-	Context.AnimInstance->UpdateSlotNodeWeight(SlotName, SlotNodeWeight);
+	Context.AnimInstance->UpdateSlotNodeWeight(SlotName, WeightData.SlotNodeWeight);
 	Context.AnimInstance->UpdateSlotRootMotionWeight(SlotName, Context.GetFinalBlendWeight());
 
-
-	if (SourceWeight > ZERO_ANIMWEIGHT_THRESH)
+	if (WeightData.SourceWeight > ZERO_ANIMWEIGHT_THRESH)
 	{
-		Source.Update(Context.FractionalWeight(SourceWeight));
+		Source.Update(Context.FractionalWeight(WeightData.SourceWeight));
 	}
 }
 
 void FAnimNode_Slot::Evaluate(FPoseContext & Output)
 {
 	// If not playing a montage, just pass through
-	if (SlotNodeWeight <= ZERO_ANIMWEIGHT_THRESH)
+	if (WeightData.SlotNodeWeight <= ZERO_ANIMWEIGHT_THRESH)
 	{
 		Source.Evaluate(Output);
 	}
 	else
 	{
 		FPoseContext SourceContext(Output);
-// @TODO: FIXME temporary fix to fix crash 
-//		if (SourceWeight > ZERO_ANIMWEIGHT_THRESH)
+		if (WeightData.SourceWeight > ZERO_ANIMWEIGHT_THRESH)
 		{
 			Source.Evaluate(SourceContext);
 		}
 
-		Output.AnimInstance->SlotEvaluatePose(SlotName, SourceContext.Pose, SourceContext.Curve, Output.Pose, Output.Curve, SlotNodeWeight);
+		Output.AnimInstance->SlotEvaluatePose(SlotName, SourceContext.Pose, SourceContext.Curve, WeightData.SourceWeight, Output.Pose, Output.Curve, WeightData.SlotNodeWeight, WeightData.TotalNodeWeight);
 
 		checkSlow(!Output.ContainsNaN());
 		checkSlow(Output.IsNormalized());
@@ -71,11 +67,11 @@ void FAnimNode_Slot::Evaluate(FPoseContext & Output)
 void FAnimNode_Slot::GatherDebugData(FNodeDebugData& DebugData)
 {
 	FString DebugLine = DebugData.GetNodeName(this);
-	DebugLine += FString::Printf(TEXT("(Slot Name: '%s' Weight:%.1f%%)"), *SlotName.ToString(), SlotNodeWeight*100.f);
+	DebugLine += FString::Printf(TEXT("(Slot Name: '%s' Weight:%.1f%%)"), *SlotName.ToString(), WeightData.SlotNodeWeight*100.f);
 	
-	bool const bIsPoseSource = (SourceWeight <= ZERO_ANIMWEIGHT_THRESH);
+	bool const bIsPoseSource = (WeightData.SourceWeight <= ZERO_ANIMWEIGHT_THRESH);
 	DebugData.AddDebugItem(DebugLine, bIsPoseSource);
-	Source.GatherDebugData(DebugData.BranchFlow(SourceWeight));
+	Source.GatherDebugData(DebugData.BranchFlow(WeightData.SourceWeight));
 
 	for (FAnimMontageInstance* MontageInstance : DebugData.AnimInstance->MontageInstances)
 	{
@@ -99,8 +95,6 @@ void FAnimNode_Slot::GatherDebugData(FNodeDebugData& DebugData)
 
 FAnimNode_Slot::FAnimNode_Slot()
 	: SlotName(FAnimSlotGroup::DefaultSlotName)
-	, SourceWeight(0.f)
-	, SlotNodeWeight(0.f)
 	, LastSlotNodeInitializationCounter(INDEX_NONE)
 {
 }
diff --git a/Engine/Source/Runtime/AnimGraphRuntime/Public/AnimNodes/AnimNode_Slot.h b/Engine/Source/Runtime/AnimGraphRuntime/Public/AnimNodes/AnimNode_Slot.h
--- a/Engine/Source/Runtime/AnimGraphRuntime/Public/AnimNodes/AnimNode_Slot.h
+++ b/Engine/Source/Runtime/AnimGraphRuntime/Public/AnimNodes/AnimNode_Slot.h
@@ -21,12 +21,7 @@ struct ANIMGRAPHRUNTIME_API FAnimNode_Slot : public FAnimNode_Base
 	FName SlotName;
 
 protected:
-	/** Weight of Source Branch. This is the weight of the input pose coming from children.
-	This is different than (1.f - SourceWeight) since the Slot can play additive animations, which are overlayed on top of the source pose. */
-	float SourceWeight;
-
-	/** Weight of Slot Node. Determined by Montages weight playing on this slot */
-	float SlotNodeWeight;
+	FSlotNodeWeightInfo WeightData;
 
 private:
 	/** Counter to register slot node just once per initialization */
diff --git a/Engine/Source/Runtime/Engine/Classes/Animation/AnimInstance.h b/Engine/Source/Runtime/Engine/Classes/Animation/AnimInstance.h
--- a/Engine/Source/Runtime/Engine/Classes/Animation/AnimInstance.h
+++ b/Engine/Source/Runtime/Engine/Classes/Animation/AnimInstance.h
@@ -370,9 +370,17 @@ public:
 	// Populate a tick record with previous frames data
 	void InitTickRecordFromLastFrame(int32 GroupIndex, FAnimTickRecord& NewTickRecord) const;
 
-	// slotnode interfaces
-	void GetSlotWeight(FName const & SlotNodeName, float& out_SlotNodeWeight, float& out_SourceWeight) const;
-	void SlotEvaluatePose(FName SlotNodeName, const FCompactPose& SourcePose, const FBlendedCurve& SourceCurve, FCompactPose& BlendedPose, FBlendedCurve& BlendedCurve, float SlotNodeWeight);
+	/**
+	 * Get Slot Node Weight : this returns new Slot Node Weight, Source Weight, Original TotalNodeWeight
+	 *							this 3 values can't be derived from each other
+	 *
+	 * @param SlotNodeName : the name of the slot node you're querying
+	 * @param out_SlotNodeWeight : The node weight for this slot node in the range of [0, 1]
+	 * @param out_SourceWeight : The Source weight for this node. 
+	 * @param out_TotalNodeWeight : Total weight of this node
+	 */
+	void GetSlotWeight(FName const& SlotNodeName, float& out_SlotNodeWeight, float& out_SourceWeight, float& out_TotalNodeWeight) const;
+	void SlotEvaluatePose(FName SlotNodeName, const FCompactPose& SourcePose, const FBlendedCurve& SourceCurve, float InSourceWeight, FCompactPose& BlendedPose, FBlendedCurve& BlendedCurve, float InBlendWeight, float InTotalNodeWeight);
 
 	// slot node run-time functions
 	void ReinitializeSlotNodes();
diff --git a/Engine/Source/Runtime/Engine/Classes/Animation/AnimMontage.h b/Engine/Source/Runtime/Engine/Classes/Animation/AnimMontage.h
--- a/Engine/Source/Runtime/Engine/Classes/Animation/AnimMontage.h
+++ b/Engine/Source/Runtime/Engine/Classes/Animation/AnimMontage.h
@@ -166,20 +166,18 @@ struct FAnimMontageInstance
 	UPROPERTY()
 	class UAnimMontage* Montage;
 
-public: 
-	UPROPERTY()
-	float DesiredWeight;
+	// delegates
+	FOnMontageEnded OnMontageEnded;
+	FOnMontageBlendingOutStarted OnMontageBlendingOutStarted;
 
 	UPROPERTY()
-	float Weight;
-
-	UPROPERTY(transient)
-	float BlendTime;
+	bool bPlaying;
 
 	// Blend Time multiplier to allow extending and narrowing blendtimes
 	UPROPERTY(transient)
 	float DefaultBlendTimeMultiplier;
 
+private:
 	// list of next sections per section - index of array is section id
 	UPROPERTY()
 	TArray<int32> NextSections;
@@ -188,40 +186,40 @@ public:
 	UPROPERTY()
 	TArray<int32> PrevSections;
 
-	UPROPERTY()
-	bool bPlaying;
-
-	// delegates
-	FOnMontageEnded OnMontageEnded;
-	FOnMontageBlendingOutStarted OnMontageBlendingOutStarted;
-
 	// reference to AnimInstance
 	TWeakObjectPtr<UAnimInstance> AnimInstance;
 
-	// need to save if it's interrupted or not
-	// this information is crucial for gameplay
-	bool bInterrupted;
-
-	// transient PreviousWeight - Weight of previous tick
-	float PreviousWeight;
-
-	// transient NotifyWeight   - Weight for spawned notifies, modified slightly to make sure
-	//                          - we spawn all notifies
-	float NotifyWeight;
-
 	/** Currently Active AnimNotifyState, stored as a copy of the event as we need to
 		call NotifyEnd on the event after a deletion in the editor. After this the event
 		is removed correctly. */
 	UPROPERTY(Transient)
 	TArray<FAnimNotifyEvent> ActiveStateBranchingPoints;
 
-private:
 	UPROPERTY()
 	float Position;
 
 	UPROPERTY()
 	float PlayRate;
 
+	UPROPERTY()
+	float DesiredWeight;
+
+	UPROPERTY()
+	float Weight;
+
+	UPROPERTY(transient)
+	float BlendTime;
+
+	// need to save if it's interrupted or not
+	// this information is crucial for gameplay
+	bool bInterrupted;
+
+	// transient PreviousWeight - Weight of previous tick
+	float PreviousWeight;
+
+	// transient NotifyWeight   - Weight for spawned notifies, modified slightly to make sure
+	//                          - we spawn all notifies
+	float NotifyWeight;
 public:
 	/** Montage to Montage Synchronization.
 	 *
@@ -242,6 +240,10 @@ public:
 	/** PostUpdate - Sync if updated after Leader. */
 	void MontageSync_PostUpdate();
 
+	/** Get Weight */
+	float GetWeight() const { return Weight; }
+	float GetDesiredWeight() const { return DesiredWeight; }
+	float GetBlendTime() const { return BlendTime; }
 private:
 	/** Followers this Montage will synchronize */
 	TArray<struct FAnimMontageInstance*> MontageSyncFollowers;
@@ -341,7 +343,7 @@ public:
 
 	FName GetCurrentSection() const;
 	FName GetNextSection() const;
-	int32 GetNextSectionID(int32 const & CurrentSectionID) const;
+	ENGINE_API int32 GetNextSectionID(int32 const & CurrentSectionID) const;
 	FName GetSectionNameFromID(int32 const & SectionID) const;
 
 	// reference has to be managed manually
@@ -360,6 +362,11 @@ private:
 	/** Trigger associated events when Montage ticking reaches given FBranchingPointMarker */
 	void BranchingPointEventHandler(const FBranchingPointMarker* BranchingPointMarker);
 	void RefreshNextPrevSections();
+
+public:
+	/** static functions that are used by matinee functionality */
+	static void SetMatineeAnimPositionInner(FName SlotName, USkeletalMeshComponent* SkeletalMeshComponent, UAnimSequence* InAnimSequence, TWeakObjectPtr<UAnimMontage>& CurrentlyPlayingMontage, float InPosition, bool bLooping);
+	static void PreviewMatineeSetAnimPositionInner(FName SlotName, USkeletalMeshComponent* SkeletalMeshComponent, UAnimSequence* InAnimSequence, TWeakObjectPtr<UAnimMontage>& CurrentlyPlayingMontage, float InPosition, bool bLooping, bool bFireNotifies, float DeltaTime);
 };
 
 UCLASS(config=Engine, hidecategories=(UObject, Length), MinimalAPI, BlueprintType)
diff --git a/Engine/Source/Runtime/Engine/Classes/Animation/AnimSingleNodeInstance.h b/Engine/Source/Runtime/Engine/Classes/Animation/AnimSingleNodeInstance.h
--- a/Engine/Source/Runtime/Engine/Classes/Animation/AnimSingleNodeInstance.h
+++ b/Engine/Source/Runtime/Engine/Classes/Animation/AnimSingleNodeInstance.h
@@ -112,6 +112,9 @@ public:
 #if WITH_EDITORONLY_DATA
 	float PreviewPoseCurrentTime;
 #endif
+private:
+	/** Slot node weight transient data */
+	FSlotNodeWeightInfo WeightInfo;
 };
 
 
diff --git a/Engine/Source/Runtime/Engine/Private/Animation/AnimInstance.cpp b/Engine/Source/Runtime/Engine/Private/Animation/AnimInstance.cpp
--- a/Engine/Source/Runtime/Engine/Private/Animation/AnimInstance.cpp
+++ b/Engine/Source/Runtime/Engine/Private/Animation/AnimInstance.cpp
@@ -401,6 +401,10 @@ void UAnimInstance::UpdateAnimation(float DeltaSeconds)
 		}
 	}
 
+	// update weight before all nodes update comes in
+	// when native update animation for single node, it needs to know latest weight of montage
+	Montage_UpdateWeight(DeltaSeconds);
+
 	{
 		SCOPE_CYCLE_COUNTER(STAT_NativeUpdateAnimation);
 		NativeUpdateAnimation(DeltaSeconds);
@@ -410,9 +414,6 @@ void UAnimInstance::UpdateAnimation(float DeltaSeconds)
 		BlueprintUpdateAnimation(DeltaSeconds);
 	}
 
-	// update weight before all nodes update comes in
-	Montage_UpdateWeight(DeltaSeconds);
-
 	// Update the anim graph
 	if (RootNode != NULL)
 	{
@@ -917,7 +918,7 @@ void UAnimInstance::DisplayDebug(class UCanvas* Canvas, const FDebugDisplayInfo&
 
 			Canvas->SetLinearDrawColor((MontageInstance->IsActive()) ? ActiveColor : TextWhite);
 
-			FString MontageEntry = FString::Printf(TEXT("%i) %s CurrSec: %s NextSec: %s W:%.3f DW:%.3f"), MontageIndex, *MontageInstance->Montage->GetName(), *MontageInstance->GetCurrentSection().ToString(), *MontageInstance->GetNextSection().ToString(), MontageInstance->Weight, MontageInstance->DesiredWeight);
+			FString MontageEntry = FString::Printf(TEXT("%i) %s CurrSec: %s NextSec: %s W:%.3f DW:%.3f"), MontageIndex, *MontageInstance->Montage->GetName(), *MontageInstance->GetCurrentSection().ToString(), *MontageInstance->GetNextSection().ToString(), MontageInstance->GetWeight(), MontageInstance->GetDesiredWeight());
 			YL = Canvas->DrawText(RenderFont, MontageEntry, Indent, YPos, 1.f, 1.f, RenderInfo);
 			YPos += YL;
 		}
@@ -1370,9 +1371,12 @@ void UAnimInstance::AnimNotify_Sound(UAnimNotify* AnimNotify)
 //to debug montage weight
 #define DEBUGMONTAGEWEIGHT 0
 
-void UAnimInstance::GetSlotWeight(FName const & SlotNodeName, float& out_SlotNodeWeight, float& out_SourceWeight) const
+void UAnimInstance::GetSlotWeight(FName const& SlotNodeName, float& out_SlotNodeWeight, float& out_SourceWeight, float& out_TotalNodeWeight) const
 {
-	float NodeTotalWeight = 0.f;
+	// node total weight 
+	float NewSlotNodeWeight = 0.f;
+	// this is required to track, because it will be 1-SourceWeight
+	// if additive, it can be applied more
 	float NonAdditiveTotalWeight = 0.f;
 
 #if DEBUGMONTAGEWEIGHT
@@ -1384,11 +1388,11 @@ void UAnimInstance::GetSlotWeight(FName const & SlotNodeName, float& out_SlotNod
 		FAnimMontageInstance const* const MontageInstance = MontageInstances[Index];
 		if (MontageInstance && MontageInstance->IsValid() && MontageInstance->Montage->IsValidSlot(SlotNodeName))
 		{
-			NodeTotalWeight += MontageInstance->Weight;
+			NewSlotNodeWeight += MontageInstance->GetWeight();
 
 			if( !MontageInstance->Montage->IsValidAdditiveSlot(SlotNodeName) )
 			{
-				NonAdditiveTotalWeight += MontageInstance->Weight;
+				NonAdditiveTotalWeight += MontageInstance->GetWeight();
 			}
 
 #if DEBUGMONTAGEWEIGHT			
@@ -1397,24 +1401,23 @@ void UAnimInstance::GetSlotWeight(FName const & SlotNodeName, float& out_SlotNod
 		}
 	}
 
+	// save the total node weight, it can be more than 1
+	// we need this so that when we eval, we normalized by this weight
+	// calculating there can cause inconsistency if some data changes
+	out_TotalNodeWeight = NewSlotNodeWeight;
+
 	// this can happen when it's blending in OR when newer animation comes in with shorter blendtime
 	// say #1 animation was blending out time with current blendtime 1.0 #2 animation was blending in with 1.0 (old) but got blend out with new blendtime 0.2f
 	// #3 animation was blending in with the new blendtime 0.2f, you'll have sum of #1, 2, 3 exceeds 1.f
-	if (NodeTotalWeight > (1.f + ZERO_ANIMWEIGHT_THRESH))
+	if (NewSlotNodeWeight > (1.f + ZERO_ANIMWEIGHT_THRESH))
 	{
-		// Re-normalize instance weights.
-		for (int32 Index = 0; Index < MontageInstances.Num(); Index++)
-		{
-			FAnimMontageInstance* MontageInstance = MontageInstances[Index];
-			if (MontageInstance && MontageInstance->IsValid() && MontageInstance->Montage->IsValidSlot(SlotNodeName))
-			{
-				MontageInstance->Weight /= NodeTotalWeight;
-			}
-		} 
-
-		// Re-normalize totals
-		NonAdditiveTotalWeight /= NodeTotalWeight;
-		NodeTotalWeight = 1.f;
+		// you don't want to change weight of montage instance since it can play multiple slots
+		// if you change one, it will apply to all slots in that montage
+		// instead we should renormalize when we eval
+		// this should happen in the eval phase
+		NonAdditiveTotalWeight /= NewSlotNodeWeight;
+		// since we normalized, we reset
+		NewSlotNodeWeight = 1.f;
 	}
 #if DEBUGMONTAGEWEIGHT
 	else if (TotalDesiredWeight == 1.f && TotalSum < 1.f - ZERO_ANIMWEIGHT_THRESH)
@@ -1426,16 +1429,21 @@ void UAnimInstance::GetSlotWeight(FName const & SlotNodeName, float& out_SlotNod
 	}
 #endif
 
-	out_SlotNodeWeight = NodeTotalWeight;
+	out_SlotNodeWeight = NewSlotNodeWeight;
 	out_SourceWeight = 1.f - NonAdditiveTotalWeight;
 }
 
-void UAnimInstance::SlotEvaluatePose(FName SlotNodeName, const FCompactPose& SourcePose, const FBlendedCurve& SourceCurve, FCompactPose& BlendedPose, FBlendedCurve& BlendedCurve, float SlotNodeWeight)
+#define DEBUG_MONTAGEINSTANCE_WEIGHT 1
+#if (UE_BUILD_SHIPPING || UE_BUILD_TEST)
+#define DEBUG_MONTAGEINSTANCE_WEIGHT 0
+#endif // !(UE_BUILD_SHIPPING || UE_BUILD_TEST)
+
+void UAnimInstance::SlotEvaluatePose(FName SlotNodeName, const FCompactPose& SourcePose, const FBlendedCurve& SourceCurve, float InSourceWeight, FCompactPose& BlendedPose, FBlendedCurve& BlendedCurve, float InBlendWeight, float InTotalNodeWeight)
 {
 	//Accessing MontageInstances from this function is not safe (as this can be called during Parallel Anim Evaluation!
 	//Any montage data you need to add should be part of MontageEvaluationData
-
-	if (SlotNodeWeight <= ZERO_ANIMWEIGHT_THRESH)
+	// nothing to blend, just get it out
+	if (InBlendWeight <= ZERO_ANIMWEIGHT_THRESH)
 	{
 		BlendedPose = SourcePose;
 		BlendedCurve = SourceCurve;
@@ -1447,8 +1455,9 @@ void UAnimInstance::SlotEvaluatePose(FName SlotNodeName, const FCompactPose& Sou
 	TArray<FSlotEvaluationPose> NonAdditivePoses;
 
 	// first pass we go through collect weights and valid montages. 
+#if DEBUG_MONTAGEINSTANCE_WEIGHT
 	float TotalWeight = 0.f;
-	float NonAdditiveWeight = 0.f;
+#endif // DEBUG_MONTAGEINSTANCE_WEIGHT
 	for (const FMontageEvaluationState& EvalState : MontageEvaluationData)
 	{
 		if (EvalState.Montage->IsValidSlot(SlotNodeName))
@@ -1476,43 +1485,41 @@ void UAnimInstance::SlotEvaluatePose(FName SlotNodeName, const FCompactPose& Sou
 			EvalState.Montage->EvaluateCurveData(MontageCurve, EvalState.MontagePosition);
 			NewPose.Curve.Combine(MontageCurve);
 
+#if DEBUG_MONTAGEINSTANCE_WEIGHT
 			TotalWeight += EvalState.MontageWeight;
+#endif // DEBUG_MONTAGEINSTANCE_WEIGHT
 			if (AdditiveAnimType == AAT_None)
 			{
-				NonAdditiveWeight += EvalState.MontageWeight;
 				NonAdditivePoses.Add(NewPose);
 			}
 			else
 			{
 				AdditivePoses.Add(NewPose);
 			}
-		
 		}
 	}
 
-	// nothing else to do here, there is no weight
-	if (TotalWeight <= ZERO_ANIMWEIGHT_THRESH)
-	{
-		BlendedPose = SourcePose;
-		BlendedCurve = SourceCurve;
-		return;
-	}
-	// Make sure weights don't exceed 1.f, otherwise re-normalize.
-	else if (TotalWeight > (1.f + ZERO_ANIMWEIGHT_THRESH))
+	// allocate for blending
+	// If source has any weight, add it to the blend array.
+	float const SourceWeight = FMath::Clamp<float>(InSourceWeight, 0.f, 1.f);
+
+#if DEBUG_MONTAGEINSTANCE_WEIGHT
+	ensure (FMath::IsNearlyEqual(InTotalNodeWeight, TotalWeight, KINDA_SMALL_NUMBER));
+#endif // DEBUG_MONTAGEINSTANCE_WEIGHT
+	ensure (InTotalNodeWeight > ZERO_ANIMWEIGHT_THRESH);
+
+	if (InTotalNodeWeight > (1.f + ZERO_ANIMWEIGHT_THRESH))
 	{
 		// Re-normalize additive poses
 		for (int32 Index = 0; Index < AdditivePoses.Num(); Index++)
 		{
-			AdditivePoses[Index].Weight /= TotalWeight;
+			AdditivePoses[Index].Weight /= InTotalNodeWeight;
 		}
 		// Re-normalize non-additive poses
 		for (int32 Index = 0; Index < NonAdditivePoses.Num(); Index++)
 		{
-			NonAdditivePoses[Index].Weight /= TotalWeight;
+			NonAdditivePoses[Index].Weight /= InTotalNodeWeight;
 		}
-		// Re-normalize totals.
-		NonAdditiveWeight /= TotalWeight;
-		TotalWeight = 1.f;
 	}
 
 	// Make sure we have at least one montage here.
@@ -1529,9 +1536,6 @@ void UAnimInstance::SlotEvaluatePose(FName SlotNodeName, const FCompactPose& Sou
 		// Otherwise we need to blend non additive poses together
 		else
 		{
-			// allocate for blending
-			// If source has any weight, add it to the blend array.
-			float const SourceWeight = FMath::Clamp<float>(1.f - NonAdditiveWeight, 0.f, 1.f);
 			int32 const NumPoses = NonAdditivePoses.Num() + ((SourceWeight > ZERO_ANIMWEIGHT_THRESH) ? 1 : 0);
 
 			TArray<const FCompactPose*> BlendingPoses;
@@ -1820,8 +1824,8 @@ void UAnimInstance::Montage_Advance(float DeltaSeconds)
 				FAnimMontageInstance* AnimMontageInstance = MontageInstances(I);
 				// print blending time and weight and montage name
 				UE_LOG(LogAnimation, Warning, TEXT("%d. Montage (%s), DesiredWeight(%0.2f), CurrentWeight(%0.2f), BlendingTime(%0.2f)"), 
-					I+1, *AnimMontageInstance->Montage->GetName(), AnimMontageInstance->DesiredWeight, AnimMontageInstance->Weight,  
-					AnimMontageInstance->BlendTime );
+					I+1, *AnimMontageInstance->Montage->GetName(), AnimMontageInstance->GetDesiredWeight(), AnimMontageInstance->GetWeight(),  
+					AnimMontageInstance->GetBlendTime() );
 			}
 #endif
 		}
@@ -2489,7 +2493,7 @@ float UAnimInstance::Montage_GetBlendTime(UAnimMontage* Montage)
 		FAnimMontageInstance* MontageInstance = GetActiveInstanceForMontage(*Montage);
 		if (MontageInstance)
 		{
-			return MontageInstance->BlendTime;
+			return MontageInstance->GetBlendTime();
 		}
 	}
 	else
@@ -2500,7 +2504,7 @@ float UAnimInstance::Montage_GetBlendTime(UAnimMontage* Montage)
 			FAnimMontageInstance* MontageInstance = MontageInstances[InstanceIndex];
 			if (MontageInstance && MontageInstance->IsActive())
 			{
-				return MontageInstance->BlendTime;
+				return MontageInstance->GetBlendTime();
 			}
 		}
 	}
@@ -2541,7 +2545,7 @@ int32 UAnimInstance::Montage_GetNextSectionID(UAnimMontage const* const Montage,
 		FAnimMontageInstance* MontageInstance = GetActiveInstanceForMontage(*Montage);
 		if (MontageInstance)
 		{
-			return (CurrentSectionID < MontageInstance->NextSections.Num()) ? MontageInstance->NextSections[CurrentSectionID] : INDEX_NONE;
+			return MontageInstance->GetNextSectionID(CurrentSectionID);
 		}
 	}
 	else
@@ -2552,7 +2556,7 @@ int32 UAnimInstance::Montage_GetNextSectionID(UAnimMontage const* const Montage,
 			FAnimMontageInstance* MontageInstance = MontageInstances[InstanceIndex];
 			if (MontageInstance && MontageInstance->IsActive())
 			{
-				return (CurrentSectionID < MontageInstance->NextSections.Num()) ? MontageInstance->NextSections[CurrentSectionID] : INDEX_NONE;
+				return MontageInstance->GetNextSectionID(CurrentSectionID);
 			}
 		}
 	}
@@ -2778,9 +2782,11 @@ void UAnimInstance::UpdateMontageEvaluationData()
 	MontageEvaluationData.Reset(MontageInstances.Num());
 	for (FAnimMontageInstance* MontageInstance : MontageInstances)
 	{
-		if (MontageInstance->Montage && MontageInstance->Weight > ZERO_ANIMWEIGHT_THRESH)
+		if (MontageInstance->Montage && MontageInstance->GetWeight() > ZERO_ANIMWEIGHT_THRESH)
 		{
-			MontageEvaluationData.Add(FMontageEvaluationState(MontageInstance->Montage, MontageInstance->Weight, MontageInstance->GetPosition()));
+			UE_LOG(LogAnimation, Verbose, TEXT("UpdateMontageEvaluationData : AnimMontage: %s,  (DesiredWeight:%0.2f, Weight:%0.2f)"),
+						*MontageInstance->Montage->GetName(), MontageInstance->GetDesiredWeight(), MontageInstance->GetWeight());
+			MontageEvaluationData.Add(FMontageEvaluationState(MontageInstance->Montage, MontageInstance->GetWeight(), MontageInstance->GetPosition()));
 		}
 	}
 }
diff --git a/Engine/Source/Runtime/Engine/Private/Animation/AnimMontage.cpp b/Engine/Source/Runtime/Engine/Private/Animation/AnimMontage.cpp
--- a/Engine/Source/Runtime/Engine/Private/Animation/AnimMontage.cpp
+++ b/Engine/Source/Runtime/Engine/Private/Animation/AnimMontage.cpp
@@ -11,7 +11,8 @@
 #include "Animation/AnimNotifies/AnimNotifyState.h"
 #include "Animation/AnimNotifies/AnimNotify.h"
 #include "Animation/AnimSequenceBase.h"
-#include "AnimStats.h"
+#include "Animation/AnimStats.h"
+#include "Animation/AnimSingleNodeInstance.h"
 
 DEFINE_STAT(STAT_AnimMontageInstance_Advance);
 DEFINE_STAT(STAT_AnimMontageInstance_TickBranchPoints);
@@ -1072,7 +1073,7 @@ FName FAnimMontageInstance::GetNextSection() const
 
 int32 FAnimMontageInstance::GetNextSectionID(int32 const & CurrentSectionID) const
 {
-	return (IsActive() && (CurrentSectionID < NextSections.Num())) ? NextSections[CurrentSectionID] : INDEX_NONE;
+	return NextSections.IsValidIndex(CurrentSectionID) ? NextSections[CurrentSectionID] : INDEX_NONE;
 }
 
 FName FAnimMontageInstance::GetSectionNameFromID(int32 const & SectionID) const
@@ -1602,6 +1603,137 @@ void FAnimMontageInstance::BranchingPointEventHandler(const FBranchingPointMarke
 	}
 }
 
+void FAnimMontageInstance::SetMatineeAnimPositionInner(FName SlotName, USkeletalMeshComponent* SkeletalMeshComponent, UAnimSequence* InAnimSequence, TWeakObjectPtr<UAnimMontage>& CurrentlyPlayingMontage, float InPosition, bool bLooping)
+{
+	UAnimInstance* AnimInst = SkeletalMeshComponent->GetAnimInstance();
+	if(AnimInst)
+	{
+		UAnimSingleNodeInstance * SingleNodeInst = SkeletalMeshComponent->GetSingleNodeInstance();
+		if(SingleNodeInst)
+		{
+			if(SingleNodeInst->CurrentAsset != InAnimSequence)
+			{
+				SingleNodeInst->SetAnimationAsset(InAnimSequence, bLooping);
+				SingleNodeInst->SetPosition(0.0f);
+				SingleNodeInst->bPlaying = false;
+			}
+
+			if(SingleNodeInst->bLooping!=bLooping)
+			{
+				SingleNodeInst->SetLooping(bLooping);
+			}
+			if(SingleNodeInst->CurrentTime != InPosition)
+			{
+				SingleNodeInst->SetPosition(InPosition);
+			}
+		}
+		else
+		{
+			bool bShouldChange = AnimInst->IsPlayingSlotAnimation(InAnimSequence, SlotName) == false;
+			if(bShouldChange)
+			{
+				if(CurrentlyPlayingMontage.IsValid())
+				{
+					// set it's weight to 0
+					struct FAnimMontageInstance* PrevAnimMontageInst = AnimInst->GetActiveInstanceForMontage(*CurrentlyPlayingMontage);
+					if(PrevAnimMontageInst)
+					{
+						PrevAnimMontageInst->Weight=0.f;
+					}
+				}
+
+				AnimInst->PlaySlotAnimationAsDynamicMontage(InAnimSequence, SlotName, 0.0f, 0.0f, 0.f, 1);
+				CurrentlyPlayingMontage = AnimInst->GetCurrentActiveMontage();
+			}
+
+			ensure(CurrentlyPlayingMontage.IsValid());
+
+			struct FAnimMontageInstance* AnimMontageInst = AnimInst->GetActiveInstanceForMontage(*CurrentlyPlayingMontage);
+			if(AnimMontageInst)
+			{
+				AnimMontageInst->Weight = 1.f;
+
+				float OldMontagePosition = AnimInst->Montage_GetPosition(CurrentlyPlayingMontage.Get());
+				AnimInst->Montage_SetPosition(CurrentlyPlayingMontage.Get(), InPosition);
+
+				// since we don't advance montage in the tick, we manually have to handle notifies
+				AnimMontageInst->HandleEvents(OldMontagePosition, InPosition, NULL);
+				AnimInst->TriggerAnimNotifies(0.f);
+			}
+			else
+			{
+				UE_LOG(LogSkeletalMesh, Warning, TEXT("Invalid Slot Node Name: %s"), *SlotName.ToString());
+			}
+		}
+	}
+}
+
+void FAnimMontageInstance::PreviewMatineeSetAnimPositionInner(FName SlotName, USkeletalMeshComponent* SkeletalMeshComponent, UAnimSequence* InAnimSequence, TWeakObjectPtr<UAnimMontage>& CurrentlyPlayingMontage, float InPosition, bool bLooping, bool bFireNotifies, float DeltaTime)
+{
+	UAnimInstance* AnimInst = SkeletalMeshComponent->GetAnimInstance();
+	UAnimSingleNodeInstance * SingleNodeInst = SkeletalMeshComponent->GetSingleNodeInstance();
+	if(SingleNodeInst)
+	{
+		if(SingleNodeInst->CurrentAsset != InAnimSequence)
+		{
+			SingleNodeInst->SetAnimationAsset(InAnimSequence, bLooping);
+		}
+
+		SingleNodeInst->SetLooping(bLooping);
+		SingleNodeInst->SetPosition(InPosition, bFireNotifies);
+	}
+	else
+	{
+		bool bShouldChange = AnimInst->IsPlayingSlotAnimation(InAnimSequence, SlotName) == false;
+		if(bShouldChange)
+		{
+			if(CurrentlyPlayingMontage.IsValid())
+			{
+				// set it's weight to 0
+				struct FAnimMontageInstance* PrevAnimMontageInst = AnimInst->GetActiveInstanceForMontage(*CurrentlyPlayingMontage);
+				if(PrevAnimMontageInst)
+				{
+					PrevAnimMontageInst->Weight=0.f;
+				}
+			}
+
+			AnimInst->PlaySlotAnimationAsDynamicMontage(InAnimSequence, SlotName, 0.0f, 0.0f, 0.f, 1);
+			CurrentlyPlayingMontage = AnimInst->GetCurrentActiveMontage();
+		}
+
+		ensure(CurrentlyPlayingMontage.IsValid());
+
+		struct FAnimMontageInstance* AnimMontageInst = AnimInst->GetActiveInstanceForMontage(*CurrentlyPlayingMontage);
+		if(AnimMontageInst)
+		{
+			AnimMontageInst->Weight = 1.f;
+
+			float OldMontagePosition = AnimInst->Montage_GetPosition(CurrentlyPlayingMontage.Get());
+			AnimInst->Montage_SetPosition(CurrentlyPlayingMontage.Get(), InPosition);
+			AnimInst->UpdateAnimation(DeltaTime);
+
+			// since we don't advance montage in the tick, we manually have to handle notifies
+			AnimMontageInst->HandleEvents(OldMontagePosition, InPosition, NULL);
+			AnimInst->TriggerAnimNotifies(DeltaTime);
+		}
+		else
+		{
+			UE_LOG(LogSkeletalMesh, Warning, TEXT("Invalid Slot Node Name: %s"), *SlotName.ToString());
+		}
+	}
+
+	// Update space bases so new animation position has an effect.
+	SkeletalMeshComponent->UpdateMaterialParameters();
+	SkeletalMeshComponent->RefreshBoneTransforms();
+	SkeletalMeshComponent->RefreshSlaveComponents();
+	SkeletalMeshComponent->UpdateComponentToWorld();
+	SkeletalMeshComponent->FinalizeBoneTransform();
+}
+//////////////////////////////////////////////////
+//
+// AnimMontage
+//
+//////////////////////////////////////////////////
 float UAnimMontage::CalculateSequenceLength()
 {
 	float CalculatedSequenceLength = 0.f;
diff --git a/Engine/Source/Runtime/Engine/Private/Animation/AnimSingleNodeInstance.cpp b/Engine/Source/Runtime/Engine/Private/Animation/AnimSingleNodeInstance.cpp
--- a/Engine/Source/Runtime/Engine/Private/Animation/AnimSingleNodeInstance.cpp
+++ b/Engine/Source/Runtime/Engine/Private/Animation/AnimSingleNodeInstance.cpp
@@ -183,6 +183,7 @@ void UAnimSingleNodeInstance::RestartMontage(UAnimMontage * Montage, FName FromS
 {
 	if( Montage == CurrentAsset )
 	{
+		WeightInfo.Reset();
 		Montage_Play(Montage, PlayRate);
 		if( FromSection != NAME_None )
 		{
@@ -277,7 +278,10 @@ void UAnimSingleNodeInstance::NativeUpdateAnimation(float DeltaTimeX)
 			// Full weight , if you don't have slot track, you won't be able to see animation playing
 			if ( Montage->SlotAnimTracks.Num() > 0 )
 			{
-				UpdateSlotNodeWeight(Montage->SlotAnimTracks[0].SlotName, 1.f);		
+				// in the future, maybe we can support which slot
+				const FName CurrentSlotNodeName = Montage->SlotAnimTracks[0].SlotName;
+				GetSlotWeight(CurrentSlotNodeName, WeightInfo.SlotNodeWeight, WeightInfo.SourceWeight, WeightInfo.TotalNodeWeight);
+				UpdateSlotNodeWeight(CurrentSlotNodeName, WeightInfo.SlotNodeWeight);		
 			}
 			// get the montage position
 			// @todo anim: temporarily just choose first slot and show the location
@@ -406,7 +410,7 @@ bool UAnimSingleNodeInstance::NativeEvaluateAnimation(FPoseContext& Output)
 					SourcePose.ResetToRefPose();			
 				}
 
-				SlotEvaluatePose(Montage->SlotAnimTracks[0].SlotName, SourcePose, SourceCurve, Output.Pose, Output.Curve, 1.f);
+				SlotEvaluatePose(Montage->SlotAnimTracks[0].SlotName, SourcePose, SourceCurve, WeightInfo.SourceWeight, Output.Pose, Output.Curve, WeightInfo.SlotNodeWeight, WeightInfo.TotalNodeWeight);
 			}
 		}
 	}
diff --git a/Engine/Source/Runtime/Engine/Private/Character.cpp b/Engine/Source/Runtime/Engine/Private/Character.cpp
--- a/Engine/Source/Runtime/Engine/Private/Character.cpp
+++ b/Engine/Source/Runtime/Engine/Private/Character.cpp
@@ -1109,7 +1109,7 @@ bool ACharacter::CanUseRootMotionRepMove(const FSimulatedRootMotionReplicatedMov
 			const int32 CurrentSectionIndex = AnimMontage->GetSectionIndexFromPosition(ClientPosition);
 			if( CurrentSectionIndex != INDEX_NONE )
 			{
-				const int32 NextSectionIndex = (CurrentSectionIndex < ClientMontageInstance.NextSections.Num()) ? ClientMontageInstance.NextSections[CurrentSectionIndex] : INDEX_NONE;
+				const int32 NextSectionIndex = ClientMontageInstance.GetNextSectionID(CurrentSectionIndex);
 
 				// We can only extract root motion if we are within the same section.
 				// It's not trivial to jump through sections in a deterministic manner, but that is luckily not frequent. 
diff --git a/Engine/Source/Runtime/Engine/Private/SkeletalMesh.cpp b/Engine/Source/Runtime/Engine/Private/SkeletalMesh.cpp
--- a/Engine/Source/Runtime/Engine/Private/SkeletalMesh.cpp
+++ b/Engine/Source/Runtime/Engine/Private/SkeletalMesh.cpp
@@ -3893,70 +3893,12 @@ void ASkeletalMeshActor::PreviewFinishAnimControl(UInterpGroup* InInterpGroup)
 
 void ASkeletalMeshActor::PreviewSetAnimPosition(FName SlotName, int32 ChannelIndex, UAnimSequence* InAnimSequence, float InPosition, bool bLooping, bool bFireNotifies, float DeltaTime)
 {
-	UAnimInstance* AnimInst = SkeletalMeshComponent->GetAnimInstance();
-	if (AnimInst && CanPlayAnimation(InAnimSequence))
+	if(CanPlayAnimation(InAnimSequence))
 	{
-		UAnimSingleNodeInstance * SingleNodeInst = SkeletalMeshComponent->GetSingleNodeInstance();
-		if(SingleNodeInst)
-		{
-			if(SingleNodeInst->CurrentAsset != InAnimSequence)
-			{
-				SingleNodeInst->SetAnimationAsset(InAnimSequence, bLooping);
-			}
-
-			SingleNodeInst->SetLooping(bLooping);
-			SingleNodeInst->SetPosition(InPosition, bFireNotifies);
-		}
-		else
-		{
-			bool bShouldChange = AnimInst->IsPlayingSlotAnimation(InAnimSequence, SlotName ) == false; 
-			if(bShouldChange)
-			{
-				if (CurrentlyPlayingMontage.IsValid())
-				{
-					// set it's weight to 0
-					struct FAnimMontageInstance* PrevAnimMontageInst = AnimInst->GetActiveInstanceForMontage(*CurrentlyPlayingMontage);
-					if (PrevAnimMontageInst)
-					{
-						PrevAnimMontageInst->Weight=0.f;
-					}
-				}
-
-				AnimInst->PlaySlotAnimationAsDynamicMontage(InAnimSequence, SlotName, 0.0f, 0.0f, 0.f, 1);
-				CurrentlyPlayingMontage = AnimInst->GetCurrentActiveMontage();
-			}
-
-			ensure(CurrentlyPlayingMontage.IsValid());
-
-			struct FAnimMontageInstance* AnimMontageInst = AnimInst->GetActiveInstanceForMontage(*CurrentlyPlayingMontage);
-			if (AnimMontageInst)
-			{
-				AnimMontageInst->Weight = 1.f;
-
-				float OldMontagePosition = AnimInst->Montage_GetPosition(CurrentlyPlayingMontage.Get());
-				AnimInst->Montage_SetPosition(CurrentlyPlayingMontage.Get(), InPosition);
-				AnimInst->UpdateAnimation(DeltaTime);
-
-				// since we don't advance montage in the tick, we manually have to handle notifies
-				AnimMontageInst->HandleEvents(OldMontagePosition, InPosition, NULL);
-				AnimInst->TriggerAnimNotifies(DeltaTime);
-			}
-			else
-			{
-				UE_LOG(LogSkeletalMesh, Warning, TEXT("Invalid Slot Node Name: %s"), *SlotName.ToString() );
-			}
-		}
-
-		// Update space bases so new animation position has an effect.
-		SkeletalMeshComponent->UpdateMaterialParameters();
-		SkeletalMeshComponent->RefreshBoneTransforms();
-		SkeletalMeshComponent->RefreshSlaveComponents();
-		SkeletalMeshComponent->UpdateComponentToWorld();
+		FAnimMontageInstance::PreviewMatineeSetAnimPositionInner(SlotName, SkeletalMeshComponent, InAnimSequence, CurrentlyPlayingMontage, InPosition, bLooping, bFireNotifies, DeltaTime);
 	}
-	SkeletalMeshComponent->FinalizeBoneTransform();
 }
 
-
 void ASkeletalMeshActor::PreviewSetAnimWeights(TArray<FAnimSlotInfo>& SlotInfos)
 {
 	//no support yet
@@ -4095,75 +4037,11 @@ bool ASkeletalMeshActor::CanPlayAnimation(class UAnimSequenceBase* AnimAssetBase
 		(!AnimAssetBase || SkeletalMeshComponent->SkeletalMesh->Skeleton->IsCompatible(AnimAssetBase->GetSkeleton())));
 }
 
-void SetAnimPositionInner(FName SlotName, USkeletalMeshComponent* SkeletalMeshComponent, UAnimSequence* InAnimSequence, TWeakObjectPtr<UAnimMontage>& CurrentlyPlayingMontage, float InPosition, bool bLooping)
-{
-	UAnimInstance* AnimInst = SkeletalMeshComponent->GetAnimInstance();
-	if(AnimInst)
-	{
-		UAnimSingleNodeInstance * SingleNodeInst = SkeletalMeshComponent->GetSingleNodeInstance();
-		if(SingleNodeInst)
-		{
-			if(SingleNodeInst->CurrentAsset != InAnimSequence)
-			{
-				SingleNodeInst->SetAnimationAsset(InAnimSequence, bLooping);
-				SingleNodeInst->SetPosition(0.0f);
-				SingleNodeInst->bPlaying = false;
-			}
-
-			if(SingleNodeInst->bLooping!=bLooping)
-			{
-				SingleNodeInst->SetLooping(bLooping);
-			}
-			if(SingleNodeInst->CurrentTime != InPosition)
-			{
-				SingleNodeInst->SetPosition(InPosition);
-			}
-		}
-		else
-		{
-			bool bShouldChange = AnimInst->IsPlayingSlotAnimation(InAnimSequence, SlotName) == false;
-			if(bShouldChange)
-			{
-				if(CurrentlyPlayingMontage.IsValid())
-				{
-					// set it's weight to 0
-					struct FAnimMontageInstance* PrevAnimMontageInst = AnimInst->GetActiveInstanceForMontage(*CurrentlyPlayingMontage);
-					if(PrevAnimMontageInst)
-					{
-						PrevAnimMontageInst->Weight=0.f;
-					}
-				}
-
-				AnimInst->PlaySlotAnimationAsDynamicMontage(InAnimSequence, SlotName, 0.0f, 0.0f, 0.f, 1);
-				CurrentlyPlayingMontage = AnimInst->GetCurrentActiveMontage();
-			}
-
-			ensure(CurrentlyPlayingMontage.IsValid());
-
-			struct FAnimMontageInstance* AnimMontageInst = AnimInst->GetActiveInstanceForMontage(*CurrentlyPlayingMontage);
-			if(AnimMontageInst)
-			{
-				AnimMontageInst->Weight = 1.f;
-
-				float OldMontagePosition = AnimInst->Montage_GetPosition(CurrentlyPlayingMontage.Get());
-				AnimInst->Montage_SetPosition(CurrentlyPlayingMontage.Get(), InPosition);
-
-				// since we don't advance montage in the tick, we manually have to handle notifies
-				AnimMontageInst->HandleEvents(OldMontagePosition, InPosition, NULL);
-				AnimInst->TriggerAnimNotifies(0.f);
-			}
-			else
-			{
-				UE_LOG(LogSkeletalMesh, Warning, TEXT("Invalid Slot Node Name: %s"), *SlotName.ToString() );
-			}
-		}
-	}
-}
 void ASkeletalMeshActor::SetAnimPosition(FName SlotName, int32 ChannelIndex, UAnimSequence* InAnimSequence, float InPosition, bool bFireNotifies, bool bLooping)
 {
 	if (CanPlayAnimation(InAnimSequence))
 	{
-		SetAnimPositionInner(SlotName, SkeletalMeshComponent, InAnimSequence, CurrentlyPlayingMontage, InPosition, bLooping);
+		FAnimMontageInstance::SetMatineeAnimPositionInner(SlotName, SkeletalMeshComponent, InAnimSequence, CurrentlyPlayingMontage, InPosition, bLooping);
 	}
 }
 
diff --git a/Engine/Source/Runtime/Engine/Public/Animation/AnimTypes.h b/Engine/Source/Runtime/Engine/Public/Animation/AnimTypes.h
--- a/Engine/Source/Runtime/Engine/Public/Animation/AnimTypes.h
+++ b/Engine/Source/Runtime/Engine/Public/Animation/AnimTypes.h
@@ -283,4 +283,33 @@ namespace ECurveBlendOption
 		/** Blend By Weight without normalizing*/
 		BlendByWeight
 	};
-}
\ No newline at end of file
+}
+
+/**
+ * Slot node weight information - this is transient data that is used by slot node
+ */
+struct FSlotNodeWeightInfo
+{
+	/** Weight of Source Branch. This is the weight of the input pose coming from children.
+	This is different than (1.f - SourceWeight) since the Slot can play additive animations, which are overlayed on top of the source pose. */
+	float SourceWeight;
+
+	/** Weight of Slot Node. Determined by Montages weight playing on this slot */
+	float SlotNodeWeight;
+
+	/** Total Weight of Slot Node. Determined by Montages weight playing on this slot, it can be more than 1 */
+	float TotalNodeWeight;
+
+	FSlotNodeWeightInfo()
+		: SourceWeight(0.f)
+		, SlotNodeWeight(0.f)
+		, TotalNodeWeight(0.f)
+	{}
+
+	void Reset()
+	{
+		SourceWeight = 0.f;
+		SlotNodeWeight = 0.f;
+		TotalNodeWeight = 0.f;
+	}
+};
\ No newline at end of file
