#ANIM: Revert changes to weight update, it's going to cause some issue, so reverting it.

The issue was single node, but single node instance should solve the other way. Disabled debug because of it.

Merging

//depot/UE4-Orion/Engine/Source/Runtime/Engine/Private/Animation/AnimInstance.cpp

to //depot/UE4-Orion/Engine/Source/Runtime/Engine/Private/Animation/AnimInstance.cpp

diff --git a/Engine/Source/Runtime/Engine/Private/Animation/AnimInstance.cpp b/Engine/Source/Runtime/Engine/Private/Animation/AnimInstance.cpp
--- a/Engine/Source/Runtime/Engine/Private/Animation/AnimInstance.cpp
+++ b/Engine/Source/Runtime/Engine/Private/Animation/AnimInstance.cpp
@@ -401,10 +401,6 @@ void UAnimInstance::UpdateAnimation(float DeltaSeconds)
 		}
 	}
 
-	// update weight before all nodes update comes in
-	// when native update animation for single node, it needs to know latest weight of montage
-	Montage_UpdateWeight(DeltaSeconds);
-
 	{
 		SCOPE_CYCLE_COUNTER(STAT_NativeUpdateAnimation);
 		NativeUpdateAnimation(DeltaSeconds);
@@ -414,6 +410,8 @@ void UAnimInstance::UpdateAnimation(float DeltaSeconds)
 		BlueprintUpdateAnimation(DeltaSeconds);
 	}
 
+	Montage_UpdateWeight(DeltaSeconds);
+
 	// Update the anim graph
 	if (RootNode != NULL)
 	{
@@ -1433,7 +1431,8 @@ void UAnimInstance::GetSlotWeight(FName const& SlotNodeName, float& out_SlotNode
 	out_SourceWeight = 1.f - NonAdditiveTotalWeight;
 }
 
-#define DEBUG_MONTAGEINSTANCE_WEIGHT 1
+// for now disable becauase it will not work with single node instance
+#define DEBUG_MONTAGEINSTANCE_WEIGHT 0
 #if (UE_BUILD_SHIPPING || UE_BUILD_TEST)
 #define DEBUG_MONTAGEINSTANCE_WEIGHT 0
 #endif // !(UE_BUILD_SHIPPING || UE_BUILD_TEST)
