Integrate 2682946 from UE4-Orion

Fix some PS4 compile warnings

diff --git a/Engine/Source/Runtime/Engine/Classes/Animation/AnimMontage.h b/Engine/Source/Runtime/Engine/Classes/Animation/AnimMontage.h
--- a/Engine/Source/Runtime/Engine/Classes/Animation/AnimMontage.h
+++ b/Engine/Source/Runtime/Engine/Classes/Animation/AnimMontage.h
@@ -262,16 +262,16 @@ private:
 public:
 	FAnimMontageInstance()
 		: Montage(NULL)
+		, bPlaying(false)
+		, DefaultBlendTimeMultiplier(1.0f)
+		, AnimInstance(NULL)
+		, Position(0.f)
+		, PlayRate(1.f)
 		, DesiredWeight(0.f)
 		, Weight(0.f)
 		, BlendTime(0.f)
-		, DefaultBlendTimeMultiplier(1.0f)
-		, bPlaying(false)	
-		, AnimInstance(NULL)
 		, bInterrupted(false)
 		, PreviousWeight(0.f)
-		, Position(0.f)
-		, PlayRate(1.f)
 		, MontageSyncLeader(NULL)
 		, MontageSyncUpdateFrameCounter(INDEX_NONE)
 	{
@@ -279,16 +279,16 @@ public:
 
 	FAnimMontageInstance(UAnimInstance * InAnimInstance)
 		: Montage(NULL)
+		, bPlaying(false)
+		, DefaultBlendTimeMultiplier(1.0f)
+		, AnimInstance(InAnimInstance)
+		, Position(0.f)
+		, PlayRate(1.f)
 		, DesiredWeight(0.f)
 		, Weight(0.f)
 		, BlendTime(0.f)
-		, DefaultBlendTimeMultiplier(1.0f)
-		, bPlaying(false)		
-		, AnimInstance(InAnimInstance)
 		, bInterrupted(false)
 		, PreviousWeight(0.f)	
-		, Position(0.f)
-		, PlayRate(1.f)
 		, MontageSyncLeader(NULL)
 		, MontageSyncUpdateFrameCounter(INDEX_NONE)
 	{
