Changed a bunch of debug material names to FStrings instead of FStringAssetReferneces because asset references are resolved at times when we don't want them to be.  Also these materials are resolved manually.

diff --git a/Engine/Source/Runtime/Engine/Classes/Engine/Engine.h b/Engine/Source/Runtime/Engine/Classes/Engine/Engine.h
--- a/Engine/Source/Runtime/Engine/Classes/Engine/Engine.h
+++ b/Engine/Source/Runtime/Engine/Classes/Engine/Engine.h
@@ -805,7 +805,7 @@ public:
 
 	/** @todo document */
 	UPROPERTY(globalconfig)
-	FStringAssetReference WireframeMaterialName;
+	FString WireframeMaterialName;
 
 #if WITH_EDITORONLY_DATA
 	/** A translucent material used to render things in geometry mode. */
@@ -831,7 +831,7 @@ public:
 
 	/** @todo document */
 	UPROPERTY(globalconfig)
-	FStringAssetReference LevelColorationLitMaterialName;
+	FString LevelColorationLitMaterialName;
 
 	/** Material used for visualizing level membership in unlit view port modes. */
 	UPROPERTY()
@@ -839,7 +839,7 @@ public:
 
 	/** @todo document */
 	UPROPERTY(globalconfig)
-	FStringAssetReference LevelColorationUnlitMaterialName;
+	FString LevelColorationUnlitMaterialName;
 
 	/** Material used for visualizing lighting only w/ lightmap texel density. */
 	UPROPERTY()
@@ -847,7 +847,7 @@ public:
 
 	/** @todo document */
 	UPROPERTY(globalconfig)
-	FStringAssetReference LightingTexelDensityName;
+	FString LightingTexelDensityName;
 
 	/** Material used for visualizing level membership in lit view port modes. Uses shading to show axis directions. */
 	UPROPERTY()
@@ -855,7 +855,7 @@ public:
 
 	/** @todo document */
 	UPROPERTY(globalconfig)
-	FStringAssetReference ShadedLevelColorationLitMaterialName;
+	FString ShadedLevelColorationLitMaterialName;
 
 	/** Material used for visualizing level membership in unlit view port modes.  Uses shading to show axis directions. */
 	UPROPERTY()
@@ -863,7 +863,7 @@ public:
 
 	/** @todo document */
 	UPROPERTY(globalconfig)
-	FStringAssetReference ShadedLevelColorationUnlitMaterialName;
+	FString ShadedLevelColorationUnlitMaterialName;
 
 	/** Material used to indicate that the associated BSP surface should be removed. */
 	UPROPERTY()
@@ -879,7 +879,7 @@ public:
 
 	/** @todo document */
 	UPROPERTY(globalconfig)
-	FStringAssetReference VertexColorMaterialName;
+	FString VertexColorMaterialName;
 
 	/** Material for visualizing vertex colors on meshes in the scene (color only, no alpha) */
 	UPROPERTY()
@@ -887,7 +887,7 @@ public:
 
 	/** @todo document */
 	UPROPERTY(globalconfig)
-	FStringAssetReference VertexColorViewModeMaterialName_ColorOnly;
+	FString VertexColorViewModeMaterialName_ColorOnly;
 
 	/** Material for visualizing vertex colors on meshes in the scene (alpha channel as color) */
 	UPROPERTY()
@@ -895,7 +895,7 @@ public:
 
 	/** @todo document */
 	UPROPERTY(globalconfig)
-	FStringAssetReference VertexColorViewModeMaterialName_AlphaAsColor;
+	FString VertexColorViewModeMaterialName_AlphaAsColor;
 
 	/** Material for visualizing vertex colors on meshes in the scene (red only) */
 	UPROPERTY()
@@ -903,7 +903,7 @@ public:
 
 	/** @todo document */
 	UPROPERTY(globalconfig)
-	FStringAssetReference VertexColorViewModeMaterialName_RedOnly;
+	FString VertexColorViewModeMaterialName_RedOnly;
 
 	/** Material for visualizing vertex colors on meshes in the scene (green only) */
 	UPROPERTY()
@@ -911,7 +911,7 @@ public:
 
 	/** @todo document */
 	UPROPERTY(globalconfig)
-	FStringAssetReference VertexColorViewModeMaterialName_GreenOnly;
+	FString VertexColorViewModeMaterialName_GreenOnly;
 
 	/** Material for visualizing vertex colors on meshes in the scene (blue only) */
 	UPROPERTY()
@@ -919,7 +919,7 @@ public:
 
 	/** @todo document */
 	UPROPERTY(globalconfig)
-	FStringAssetReference VertexColorViewModeMaterialName_BlueOnly;
+	FString VertexColorViewModeMaterialName_BlueOnly;
 
 #if WITH_EDITORONLY_DATA
 	/** Material used to render bone weights on skeletal meshes */
diff --git a/Engine/Source/Runtime/Engine/Private/UnrealEngine.cpp b/Engine/Source/Runtime/Engine/Private/UnrealEngine.cpp
--- a/Engine/Source/Runtime/Engine/Private/UnrealEngine.cpp
+++ b/Engine/Source/Runtime/Engine/Private/UnrealEngine.cpp
@@ -1750,18 +1750,18 @@ void UEngine::InitializeObjectReferences()
 	if (AllowDebugViewmodes())
 	{
 		// Materials that are needed in-game if debug viewmodes are allowed
-		LoadSpecialMaterial(WireframeMaterialName.ToString(), WireframeMaterial, true);
-		LoadSpecialMaterial(LevelColorationLitMaterialName.ToString(), LevelColorationLitMaterial, true);
-		LoadSpecialMaterial(LevelColorationUnlitMaterialName.ToString(), LevelColorationUnlitMaterial, true);
-		LoadSpecialMaterial(LightingTexelDensityName.ToString(), LightingTexelDensityMaterial, false);
-		LoadSpecialMaterial(ShadedLevelColorationLitMaterialName.ToString(), ShadedLevelColorationLitMaterial, true);
-		LoadSpecialMaterial(ShadedLevelColorationUnlitMaterialName.ToString(), ShadedLevelColorationUnlitMaterial, true);
-		LoadSpecialMaterial(VertexColorMaterialName.ToString(), VertexColorMaterial, false);
-		LoadSpecialMaterial(VertexColorViewModeMaterialName_ColorOnly.ToString(), VertexColorViewModeMaterial_ColorOnly, false);
-		LoadSpecialMaterial(VertexColorViewModeMaterialName_AlphaAsColor.ToString(), VertexColorViewModeMaterial_AlphaAsColor, false);
-		LoadSpecialMaterial(VertexColorViewModeMaterialName_RedOnly.ToString(), VertexColorViewModeMaterial_RedOnly, false);
-		LoadSpecialMaterial(VertexColorViewModeMaterialName_GreenOnly.ToString(), VertexColorViewModeMaterial_GreenOnly, false);
-		LoadSpecialMaterial(VertexColorViewModeMaterialName_BlueOnly.ToString(), VertexColorViewModeMaterial_BlueOnly, false);
+		LoadSpecialMaterial(WireframeMaterialName, WireframeMaterial, true);
+		LoadSpecialMaterial(LevelColorationLitMaterialName, LevelColorationLitMaterial, true);
+		LoadSpecialMaterial(LevelColorationUnlitMaterialName, LevelColorationUnlitMaterial, true);
+		LoadSpecialMaterial(LightingTexelDensityName, LightingTexelDensityMaterial, false);
+		LoadSpecialMaterial(ShadedLevelColorationLitMaterialName, ShadedLevelColorationLitMaterial, true);
+		LoadSpecialMaterial(ShadedLevelColorationUnlitMaterialName, ShadedLevelColorationUnlitMaterial, true);
+		LoadSpecialMaterial(VertexColorMaterialName, VertexColorMaterial, false);
+		LoadSpecialMaterial(VertexColorViewModeMaterialName_ColorOnly, VertexColorViewModeMaterial_ColorOnly, false);
+		LoadSpecialMaterial(VertexColorViewModeMaterialName_AlphaAsColor, VertexColorViewModeMaterial_AlphaAsColor, false);
+		LoadSpecialMaterial(VertexColorViewModeMaterialName_RedOnly, VertexColorViewModeMaterial_RedOnly, false);
+		LoadSpecialMaterial(VertexColorViewModeMaterialName_GreenOnly, VertexColorViewModeMaterial_GreenOnly, false);
+		LoadSpecialMaterial(VertexColorViewModeMaterialName_BlueOnly, VertexColorViewModeMaterial_BlueOnly, false);
 	}
 
 	// Materials that may or may not be needed when debug viewmodes are disabled but haven't been fixed up yet
