SetConstrainedComponents needs to update the its stored component reference names in addition to the component pointers.

it looks like this physics constraint wrapper thing uses two systems to refer to the guys its attached to - a pointer and a fname.  previously this set the bodies function was just updating the pointers.  Regular blueprints might have worked ok, but   construction script didn't cause names to be updated, and it seems that the pointers didn't help when playing in editor (I assume during some sort of system replication system).  i.e. the joint didn't know who to connect with.  I confirmed by putting breakpoint in the joint creating code.

Adding code to set the names and it all works now.

with timelapse did some digging deeper. it looks like the overridecomponent was added 2 years ago to enable blueprints setting these.  the name 'ovverride' suggests this was just for blueprints and i dont see it being used elsewhere.

we might want to consider not having the cached weak pointers if its not perf critical (just keep the string names) - or not use string references as the main way of connecting things (likely a non-trivial change).   just a suggestion.

Jira
# 7878

diff --git a/Engine/Source/Runtime/Engine/Private/PhysicsEngine/PhysicsConstraintComponent.cpp b/Engine/Source/Runtime/Engine/Private/PhysicsEngine/PhysicsConstraintComponent.cpp
--- a/Engine/Source/Runtime/Engine/Private/PhysicsEngine/PhysicsConstraintComponent.cpp
+++ b/Engine/Source/Runtime/Engine/Private/PhysicsEngine/PhysicsConstraintComponent.cpp
@@ -208,12 +208,14 @@ void UPhysicsConstraintComponent::SetConstrainedComponents(UPrimitiveComponent*
 {
 	if(Component1 != NULL)
 	{
+		this->ComponentName1.ComponentName = Component1->GetFName();
 		OverrideComponent1 = Component1;
 		ConstraintInstance.ConstraintBone1 = BoneName1;
 	}
 
 	if(Component2 != NULL)
 	{
+		this->ComponentName2.ComponentName = Component2->GetFName();
 		OverrideComponent2 = Component2;
 		ConstraintInstance.ConstraintBone2 = BoneName2;
 	}
