Broke my own changes (expanding after making changes to clusters works again now)
Removed force lod column
Added more info to triangle count column (wip)
Fix for HLOD outliner crashing due to garbage collected SubActors

diff --git a/Engine/Source/Editor/HierarchicalLODOutliner/Private/HLODOutliner.cpp b/Engine/Source/Editor/HierarchicalLODOutliner/Private/HLODOutliner.cpp
--- a/Engine/Source/Editor/HierarchicalLODOutliner/Private/HLODOutliner.cpp
+++ b/Engine/Source/Editor/HierarchicalLODOutliner/Private/HLODOutliner.cpp
@@ -175,15 +175,11 @@ namespace HLODOutliner
 				SNew(SHeaderRow)
 				+ SHeaderRow::Column("SceneActorName")
 				.DefaultLabel(LOCTEXT("SceneActorName", "Scene Actor Name"))
-				.FillWidth(0.5f)
-				+ SHeaderRow::Column("ForceCheckbox")
-				.DefaultLabel(LOCTEXT("ForceCheckbox", "Force HLOD"))
-				.DefaultTooltip(LOCTEXT("ForceCheckboxTooltip", "Force Hierarchical LOD level visibility"))
-				.FillWidth(0.25f)
+				.FillWidth(0.5f)				
 				+ SHeaderRow::Column("TriangleCount")
 				.DefaultLabel(LOCTEXT("TriangleCount", "Number of Triangles"))
 				.DefaultTooltip(LOCTEXT("TriangleCountToolTip", "Number of Triangles in a LOD Mesh"))
-				.FillWidth(0.25f)				
+				.FillWidth(0.5f)				
 				);
 	}
 
@@ -255,14 +251,14 @@ namespace HLODOutliner
 
 	void SHLODOutliner::Tick(const FGeometry& AllottedGeometry, const double InCurrentTime, const float InDeltaTime)
 	{
+		// Get a collection of items and folders which were formerly collapsed
+		const FParentsExpansionState ExpansionStateInfo = GetParentsExpansionState();
+
 		if (bNeedsRefresh)
 		{
 			Populate();
 		}
 
-		// Get a collection of items and folders which were formerly collapsed
-		const FParentsExpansionState ExpansionStateInfo = GetParentsExpansionState();
-
 		bool bChangeMade = false;
 
 		// Only deal with 256 at a time
@@ -291,13 +287,14 @@ namespace HLODOutliner
 			}
 		}
 		PendingActions.RemoveAt(0, End);
-
+				
 		if (bChangeMade)
 		{
-			// Restore expansion states
+			// Restore the expansion states
 			SetParentsExpansionState(ExpansionStateInfo);
 
-			TreeView->RequestTreeRefresh();
+			// Restore expansion states
+			TreeView->RequestTreeRefresh();		
 		}			
 
 		// Update the forced LOD level, as the slider for it is being dragged
@@ -309,7 +306,7 @@ namespace HLODOutliner
 			{
 				RestoreForcedLODLevel(ForcedLODLevel);
 				ForcedLODLevel = -1;
-				SetForcedLODLevel(SnappedValue - 1);
+				SetForcedLODLevel(SnappedValue - 1);				
 			}
 		}
 	}
@@ -485,35 +482,6 @@ namespace HLODOutliner
 		}
 	}
 
-	ECheckBoxState SHLODOutliner::IsHLODLevelChecked(TSharedRef<ITreeItem> Item) const
-	{
-		check(Item->GetTreeItemType() == ITreeItem::HierarchicalLODLevel);
-
-		FLODLevelItem* LevelItem = static_cast<FLODLevelItem*>(&Item.Get());
-
-		return (LevelItem->LODLevelIndex == ForcedLODLevel) ? ECheckBoxState::Checked : ECheckBoxState::Unchecked;
-	}
-
-	void SHLODOutliner::HandleCheckBoxCheckedStateChanged(ECheckBoxState NewState, TSharedRef<ITreeItem> Item)
-	{
-		check(Item->GetTreeItemType() == ITreeItem::HierarchicalLODLevel);
-
-		FLODLevelItem* LevelItem = static_cast<FLODLevelItem*>(&Item.Get());
-		if (LevelItem->LODLevelIndex == ForcedLODLevel && NewState == ECheckBoxState::Unchecked)
-		{
-			// Reset all
-			RestoreForcedLODLevel(LevelItem->LODLevelIndex);
-			ForcedLODLevel = -1;
-		}
-		else if ( LevelItem->LODLevelIndex != ForcedLODLevel && NewState == ECheckBoxState::Checked)
-		{
-			RestoreForcedLODLevel(ForcedLODLevel);
-			SetForcedLODLevel(LevelItem->LODLevelIndex);			
-		}
-
-		ForcedLODSliderValue = ((1.0f / (LODLevelDrawDistances.Num())) * (ForcedLODLevel + 1));
-	}
-
 	bool SHLODOutliner::HandleForcedLevelSliderIsEnabled() const
 	{
 		bool bHLODsBuild = true;
@@ -824,7 +792,6 @@ namespace HLODOutliner
 			.Outliner(this)
 			.World(CurrentWorld);
 
-		//ItemMap.Add(InTreeItem->GetHash(), InTreeItem);
 		return Widget;
 	}
 
@@ -1102,7 +1069,26 @@ namespace HLODOutliner
 	void SHLODOutliner::OnLevelActorsRemoved(AActor* InActor)
 	{
 		if (!InActor->IsA<AHLODSelectionActor>() && !InActor->IsA<AWorldSettings>())
+		{
+			// Remove InActor from LOD actor which contains it
+			bool bRemoved = false;
+			for (auto& ActorArray : LODLevelActors)
+			{
+				if (bRemoved)
+					break;
+
+				for (auto& Actor : ActorArray)
+				{
+					if (Actor->RemoveSubActor(InActor))
+					{
+						bRemoved = true;
+						break;
+					}
+				}
+			}
+
 			FullRefresh();
+		}
 	}
 	
 	void SHLODOutliner::OnActorLabelChanged(AActor* ChangedActor)
diff --git a/Engine/Source/Editor/HierarchicalLODOutliner/Private/HLODOutliner.h b/Engine/Source/Editor/HierarchicalLODOutliner/Private/HLODOutliner.h
--- a/Engine/Source/Editor/HierarchicalLODOutliner/Private/HLODOutliner.h
+++ b/Engine/Source/Editor/HierarchicalLODOutliner/Private/HLODOutliner.h
@@ -121,22 +121,6 @@ namespace HLODOutliner
 		*/
 		void ForceViewLODActor(TSharedRef<ITreeItem> Item);
 
-		/**
-		* UI callback for the state of Force HLOD View checkbox
-		*
-		* @param Item - Treeview Item representation for Hierarchical LOD level
-		* @return ECheckBoxState
-		*/
-		ECheckBoxState IsHLODLevelChecked(TSharedRef<ITreeItem> Item) const;
-
-		/**
-		* UI callback for changes in Force HLOD view checkbox
-		*
-		* @param NewState - State of the check Box
-		* @param Item - Treeview Item representation for Hierarchical LOD level
-		*/
-		void HandleCheckBoxCheckedStateChanged(ECheckBoxState NewState, TSharedRef<ITreeItem> Item);
-
 		/**
 		* Determines if the HLOD slider is enabled, only if all HLODs are buil
 		*
diff --git a/Engine/Source/Editor/HierarchicalLODOutliner/Private/HLODTreeWidgetItem.cpp b/Engine/Source/Editor/HierarchicalLODOutliner/Private/HLODTreeWidgetItem.cpp
--- a/Engine/Source/Editor/HierarchicalLODOutliner/Private/HLODTreeWidgetItem.cpp
+++ b/Engine/Source/Editor/HierarchicalLODOutliner/Private/HLODTreeWidgetItem.cpp
@@ -133,25 +133,7 @@ BEGIN_SLATE_FUNCTION_BUILD_OPTIMIZATION
 					.Text(this, &SHLODWidgetItem::GetItemDisplayString)
 					.ColorAndOpacity(this, &SHLODWidgetItem::GetTint)
 				];
-		}
-		else if (ColumnName == TEXT("ForceCheckbox") && TreeItem->GetTreeItemType() == ITreeItem::HierarchicalLODLevel)
-		{
-			FLODLevelItem* Item = static_cast<FLODLevelItem*>(TreeItem);
-
-			return SNew(SHorizontalBox)
-				+ SHorizontalBox::Slot()
-				.AutoWidth()
-				.VAlign(VAlign_Center)
-				[
-					SNew(SBox)
-					.IsEnabled(Outliner, &SHLODOutliner::CanHLODLevelBeForced, Item->LODLevelIndex)
-					[
-						SNew(SCheckBox)
-						.IsChecked(Outliner, &SHLODOutliner::IsHLODLevelChecked, TreeItem->AsShared())
-						.OnCheckStateChanged(Outliner, &SHLODOutliner::HandleCheckBoxCheckedStateChanged, TreeItem->AsShared())
-					]					
-				];			
-		}
+		}		
 		else if (ColumnName == TEXT("TriangleCount") && TreeItem->GetTreeItemType() == ITreeItem::HierarchicalLODActor)
 		{
 			FLODActorItem* Item = static_cast<FLODActorItem*>(TreeItem);
diff --git a/Engine/Source/Editor/HierarchicalLODOutliner/Private/LODActorItem.cpp b/Engine/Source/Editor/HierarchicalLODOutliner/Private/LODActorItem.cpp
--- a/Engine/Source/Editor/HierarchicalLODOutliner/Private/LODActorItem.cpp
+++ b/Engine/Source/Editor/HierarchicalLODOutliner/Private/LODActorItem.cpp
@@ -86,7 +86,10 @@ FText HLODOutliner::FLODActorItem::GetNumTrianglesAsText() const
 {
 	if (LODActor.IsValid())
 	{
-		return FText::FromString(FString::FromInt(LODActor->GetNumTriangles()));
+		const uint32 SubActorCount = LODActor->GetNumTrianglesInSubActors();
+		const uint32 MergedCount = LODActor->GetNumTrianglesInMergedMesh();
+		const uint32 PercentageOfOriginal = ((float)MergedCount / (float)SubActorCount) * 100;
+		return FText::FromString(FString::FromInt(SubActorCount) + TEXT(" (Original) - ") + FString::FromInt(MergedCount) + TEXT(" (Merged) - ") + FString::FromInt(PercentageOfOriginal) + "% (of Orginal)" );
 	}
 	else
 	{
diff --git a/Engine/Source/Runtime/Engine/Classes/Engine/LODActor.h b/Engine/Source/Runtime/Engine/Classes/Engine/LODActor.h
--- a/Engine/Source/Runtime/Engine/Classes/Engine/LODActor.h
+++ b/Engine/Source/Runtime/Engine/Classes/Engine/LODActor.h
@@ -58,7 +58,7 @@ public:
 	* Removes InActor from the SubActors array and sets its LODParent to nullptr
 	* @param InActor - Actor to remove
 	*/
-	void RemoveSubActor(AActor* InActor);
+	const bool RemoveSubActor(AActor* InActor);
 
 	/**
 	* Returns whether or not this LODActor is dirty
@@ -88,7 +88,10 @@ public:
 	void SetHiddenFromEditorView(const bool InState, const int32 ForceLODLevel);
 
 	/** Returns the number of triangles this LODActor's SubActors contain */
-	const uint32 GetNumTriangles();
+	const uint32 GetNumTrianglesInSubActors();
+
+	/** Returns the number of triangles this LODActor's SubActors contain */
+	const uint32 GetNumTrianglesInMergedMesh();
 	
 	/** Updates the LODParents for the SubActors (and the drawing distance)*/
 	void UpdateSubActorLODParents();
@@ -120,6 +123,10 @@ public:
 	/** Cached number of triangles contained in the SubActors*/
 	UPROPERTY()
 	uint32 NumTrianglesInSubActors;
+
+	/** Cached number of triangles contained in the SubActors*/
+	UPROPERTY()
+	uint32 NumTrianglesInMergedMesh;
 #endif // WITH_EDITORONLY_DATA
 
 	/** Returns StaticMeshComponent subobject **/
diff --git a/Engine/Source/Runtime/Engine/Private/LODActor.cpp b/Engine/Source/Runtime/Engine/Private/LODActor.cpp
--- a/Engine/Source/Runtime/Engine/Private/LODActor.cpp
+++ b/Engine/Source/Runtime/Engine/Private/LODActor.cpp
@@ -57,6 +57,42 @@ ALODActor::ALODActor(const FObjectInitializer& ObjectInitializer)
 	bDirty = true;
 
 	NumTrianglesInSubActors = 0;
+	NumTrianglesInMergedMesh = 0;
+
+	if (SubActors.Num() != 0)
+	{
+		for (auto& Actor : SubActors)
+		{
+			// Adding number of triangles
+			if (!Actor->IsA<ALODActor>())
+			{
+				TArray<UStaticMeshComponent*> StaticMeshComponents;
+				Actor->GetComponents<UStaticMeshComponent>(StaticMeshComponents);
+				for (UStaticMeshComponent* Component : StaticMeshComponents)
+				{
+					if (Component && Component->StaticMesh && Component->StaticMesh->RenderData)
+					{
+						NumTrianglesInSubActors += Component->StaticMesh->RenderData->LODResources[0].GetNumTriangles();
+					}
+
+					Component->MarkRenderStateDirty();
+				}
+			}
+			else
+			{
+				ALODActor* LODActor = Cast<ALODActor>(Actor);
+				NumTrianglesInSubActors += LODActor->GetNumTrianglesInSubActors();
+			}
+		}
+	
+	}
+
+	if (StaticMeshComponent && StaticMeshComponent->StaticMesh && StaticMeshComponent->StaticMesh->RenderData)
+	{
+		NumTrianglesInMergedMesh = StaticMeshComponent->StaticMesh->RenderData->LODResources[0].GetNumTriangles();
+	}
+
+	
 #endif // WITH_EDITORONLY_DATA
 
 	StaticMeshComponent = CreateDefaultSubobject<UStaticMeshComponent>(TEXT("StaticMeshComponent0"));
@@ -222,44 +258,50 @@ void ALODActor::AddSubActor(AActor* InActor)
 	else
 	{
 		ALODActor* LODActor = Cast<ALODActor>(InActor);
-		NumTrianglesInSubActors += LODActor->GetNumTriangles();
+		NumTrianglesInSubActors += LODActor->GetNumTrianglesInSubActors();
 	}
 
 	StaticMeshComponent->MarkRenderStateDirty();
 	
 }
 
-void ALODActor::RemoveSubActor(AActor* InActor)
+const bool ALODActor::RemoveSubActor(AActor* InActor)
 {
-	SubActors.Remove(InActor);
-	InActor->SetLODParent(nullptr, 0);
-	SetIsDirty(true);
-	
-	// Deducting number of triangles
-	if (!InActor->IsA<ALODActor>())
+	if (SubActors.Contains(InActor))
 	{
-		TArray<UStaticMeshComponent*> StaticMeshComponents;
-		InActor->GetComponents<UStaticMeshComponent>(StaticMeshComponents);
-		for (UStaticMeshComponent* Component : StaticMeshComponents)
+		SubActors.Remove(InActor);
+		InActor->SetLODParent(nullptr, 0);
+		SetIsDirty(true);
+
+		// Deducting number of triangles
+		if (!InActor->IsA<ALODActor>())
 		{
-			if (Component && Component->StaticMesh && Component->StaticMesh->RenderData)
+			TArray<UStaticMeshComponent*> StaticMeshComponents;
+			InActor->GetComponents<UStaticMeshComponent>(StaticMeshComponents);
+			for (UStaticMeshComponent* Component : StaticMeshComponents)
 			{
-				NumTrianglesInSubActors -= Component->StaticMesh->RenderData->LODResources[0].GetNumTriangles();
-			}
+				if (Component && Component->StaticMesh && Component->StaticMesh->RenderData)
+				{
+					NumTrianglesInSubActors -= Component->StaticMesh->RenderData->LODResources[0].GetNumTriangles();
+				}
 
-			Component->MarkRenderStateDirty();
+				Component->MarkRenderStateDirty();
+			}
+		}
+		else
+		{
+			ALODActor* LODActor = Cast<ALODActor>(InActor);
+			NumTrianglesInSubActors -= LODActor->GetNumTrianglesInSubActors();
 		}
-	}
-	else
-	{
-		ALODActor* LODActor = Cast<ALODActor>(InActor);
-		NumTrianglesInSubActors -= LODActor->GetNumTriangles();
-	}
 
-	StaticMeshComponent->MarkRenderStateDirty();
+		StaticMeshComponent->MarkRenderStateDirty();
+
+		// In case the user removes an actor while the HLOD system is force viewing one LOD level
+		InActor->SetIsTemporarilyHiddenInEditor(false);
+		return true;
+	}
 
-	// In case the user removes an actor while the HLOD system is force viewing one LOD level
-	InActor->SetIsTemporarilyHiddenInEditor(false);	
+	return false;
 }
 
 void ALODActor::SetIsDirty(const bool bNewState)
@@ -361,17 +403,27 @@ void ALODActor::SetHiddenFromEditorView(const bool InState, const int32 ForceLOD
 	StaticMeshComponent->MarkRenderStateDirty();
 }
 
-const uint32 ALODActor::GetNumTriangles()
+const uint32 ALODActor::GetNumTrianglesInSubActors()
 {
 	return NumTrianglesInSubActors;
 }
 
+const uint32 ALODActor::GetNumTrianglesInMergedMesh()
+{
+	return NumTrianglesInMergedMesh;
+}
+
 void ALODActor::SetStaticMesh(class UStaticMesh* InStaticMesh)
 {
 	if (StaticMeshComponent)
 	{
 		StaticMeshComponent->StaticMesh = InStaticMesh;
 		SetIsDirty(false);
+
+		if (StaticMeshComponent && StaticMeshComponent->StaticMesh && StaticMeshComponent->StaticMesh->RenderData)
+		{
+			NumTrianglesInMergedMesh = StaticMeshComponent->StaticMesh->RenderData->LODResources[0].GetNumTriangles();
+		}
 	}
 }
 
