#ue4 - Remove redundant IsPendingKill() checks, WeakObjectPtr Get() already covers it.

--------

diff --git a/Engine/Source/Runtime/Engine/Private/LevelTick.cpp b/Engine/Source/Runtime/Engine/Private/LevelTick.cpp
--- a/Engine/Source/Runtime/Engine/Private/LevelTick.cpp
+++ b/Engine/Source/Runtime/Engine/Private/LevelTick.cpp
@@ -814,7 +814,7 @@ void UWorld::SendAllEndOfFrameUpdates()
 			UActorComponent* NextComponent = LocalComponentsThatNeedEndOfFrameUpdate[Index].Get();
 			if (NextComponent)
 			{
-				if (!NextComponent->IsPendingKill() && NextComponent->IsRegistered() && !NextComponent->IsTemplate())
+				if (NextComponent->IsRegistered() && !NextComponent->IsTemplate())
 				{
 					FScopeCycleCounterUObject ComponentScope(NextComponent);
 					FScopeCycleCounterUObject AdditionalScope(STATS ? NextComponent->AdditionalStatObject() : NULL);
@@ -833,7 +833,7 @@ void UWorld::SendAllEndOfFrameUpdates()
 				UActorComponent* Component = It->Get();
 				if (Component)
 				{
-					if ( !Component->IsPendingKill() && Component->IsRegistered() && !Component->IsTemplate())
+					if (Component->IsRegistered() && !Component->IsTemplate())
 					{
 						FScopeCycleCounterUObject ComponentScope(Component);
 						FScopeCycleCounterUObject AdditionalScope(STATS ? Component->AdditionalStatObject() : NULL);
