Engine: Removed misleading/incorrect duplicate signature comments from the FGameDelegates delegate declarations, replacing them with inline parameter names
[INTEGRATE] Change 2683785 by [REDACTED_PERSON][REDACTED_HANDLE] on 2015/09/08 17:26:48

diff --git a/Engine/Source/Runtime/Engine/Public/GameDelegates.h b/Engine/Source/Runtime/Engine/Public/GameDelegates.h
--- a/Engine/Source/Runtime/Engine/Public/GameDelegates.h
+++ b/Engine/Source/Runtime/Engine/Public/GameDelegates.h
@@ -17,25 +17,20 @@ enum class EGameDelegates_SaveGame : short
 };
 
 /** Delegate to modify cooking behavior - return extra packages to cook, load up the asset registry, etc */
-// FCookModificationDelegate(TArray<FString>& ExtraPackagesToCook);
-DECLARE_DELEGATE_OneParam(FCookModificationDelegate, TArray<FString>&);
-// FAssignStreamingChunkDelegate(const FString& PackageToAdd, const FString& LastLoadedMapName, const TArray<int32>& AssetRegistryChunkIDs, const TArray<int32>& ExistingChunkIds, int32& OutChunkIndex);
-DECLARE_DELEGATE_FiveParams(FAssignStreamingChunkDelegate, const FString&, const FString&, const TArray<int32>&, const TArray<int32>&, TArray<int32>&);
+DECLARE_DELEGATE_OneParam(FCookModificationDelegate, TArray<FString>& /*ExtraPackagesToCook*/);
+DECLARE_DELEGATE_FiveParams(FAssignStreamingChunkDelegate, const FString& /*PackageToAdd*/, const FString& /*LastLoadedMapName*/, const TArray<int32>& /*AssetRegistryChunkIDs*/, const TArray<int32>& /*ExistingChunkIds*/, TArray<int32>& /*OutChunkIndexList*/);
 
 /** Delegate to assign a disc layer to a chunk */
-// FAssignLayerChunkDelegate(const FChunkManifest* ChunkManifest, const FString& Platform, const int32 ChunkIndex, int32& OutChunkLayer);
 typedef const TMap<FName, FString> FAssignLayerChunkMap;
-DECLARE_DELEGATE_FourParams(FAssignLayerChunkDelegate, const FAssignLayerChunkMap*, const FString&, const int32, int32&);
+DECLARE_DELEGATE_FourParams(FAssignLayerChunkDelegate, const FAssignLayerChunkMap* /*ChunkManifest*/, const FString& /*Platform*/, const int32 /*ChunkIndex*/, int32& /*OutChunkLayer*/);
 
 /** A delegate for platforms that need extra information to flesh out save data information (name of an icon, for instance) */
-// FExtendedSaveGameInfoDelegate(const TCHAR* SaveName, const EGameDelegates_SaveGame Key, FString& Value); 
-DECLARE_DELEGATE_ThreeParams(FExtendedSaveGameInfoDelegate, const TCHAR*, const EGameDelegates_SaveGame, FString&);
+DECLARE_DELEGATE_ThreeParams(FExtendedSaveGameInfoDelegate, const TCHAR* /*SaveName*/, const EGameDelegates_SaveGame /*Key*/, FString& /*Value*/);
 
 /** A delegate for a web server running in engine to tell the game about events received from a client, and for game to respond to the client */
 // using a TMap in the DECLARE_DELEGATE macro caused compiler problems (in clang anyway), a typedef solves it
 typedef TMap<FString, FString> StringStringMap;
-// FWebServerActionDelegate(int32 UserIndex, const FString& Action, const FString& URL, const& TMap<FString, FString> Params, TMap<FString, FString>& Response);
-DECLARE_DELEGATE_FiveParams(FWebServerActionDelegate, int32, const FString&, const FString&, const StringStringMap&, StringStringMap&);
+DECLARE_DELEGATE_FiveParams(FWebServerActionDelegate, int32 /*UserIndex*/, const FString& /*Action*/, const FString& /*URL*/, const StringStringMap& /*Params*/, StringStringMap& /*Response*/);
 
 
 
