Engine: Overhauled how obj list is formatted
- Removed meaningless 'ResourceKB' number (it is never actionable due to double-counting, etc...); ExclusiveResKB is still present and useful
- Changed all sizes to be float KB instead of int KB (10.2f)
- Added commas between entries and dropped the K suffixes, to make it easier to import into other software
[INTEGRATE] Change 2684072 by [REDACTED_PERSON][REDACTED_HANDLE] on 2015/09/08 21:46:58

diff --git a/Engine/Source/Runtime/Engine/Private/UnrealEngine.cpp b/Engine/Source/Runtime/Engine/Private/UnrealEngine.cpp
--- a/Engine/Source/Runtime/Engine/Private/UnrealEngine.cpp
+++ b/Engine/Source/Runtime/Engine/Private/UnrealEngine.cpp
@@ -2587,30 +2587,26 @@ struct FItem
 	int32 Count;
 	SIZE_T Num;
 	SIZE_T Max;
-	/** Resource size of the object and all of its references, the 'old-style'. */
-	SIZE_T ResourceSize;
 	/** Only exclusive resource size, the truer resource size. */
 	SIZE_T TrueResourceSize;
 
 	FItem( UClass* InClass=NULL )
-	: Class(InClass), Count(0), Num(0), Max(0), ResourceSize(0), TrueResourceSize(0)
+	: Class(InClass), Count(0), Num(0), Max(0), TrueResourceSize(0)
 	{}
 
-	FItem( UClass* InClass, int32 InCount, SIZE_T InNum, SIZE_T InMax, SIZE_T InResourceSize, SIZE_T InTrueResourceSize ) :
+	FItem( UClass* InClass, int32 InCount, SIZE_T InNum, SIZE_T InMax, SIZE_T InTrueResourceSize ) :
 		Class( InClass ),
 		Count( InCount ),
 		Num( InNum ), 
 		Max( InMax ), 
-		ResourceSize( InResourceSize ), 
 		TrueResourceSize( InTrueResourceSize )
 	{}
 
-	void Add( FArchiveCountMem& Ar, SIZE_T InResourceSize, SIZE_T InTrueResourceSize )
+	void Add( FArchiveCountMem& Ar, SIZE_T InTrueResourceSize )
 	{
 		Count++;
 		Num += Ar.GetNum();
 		Max += Ar.GetMax();
-		ResourceSize += InResourceSize;
 		TrueResourceSize += InTrueResourceSize;
 	}
 };
@@ -2618,15 +2614,17 @@ struct FItem
 struct FSubItem
 {
 	UObject* Object;
+	/** Size of the object, counting containers as current usage */
 	SIZE_T Num;
+	/** Size of the object, counting containers as total allocated (max usage) */
 	SIZE_T Max;
 	/** Resource size of the object and all of its references, the 'old-style'. */
 	SIZE_T ResourceSize;
 	/** Only exclusive resource size, the truer resource size. */
 	SIZE_T TrueResourceSize;
 
-	FSubItem( UObject* InObject, SIZE_T InNum, SIZE_T InMax, SIZE_T InResourceSize, SIZE_T InTrueResourceSize )
-	: Object( InObject ), Num( InNum ), Max( InMax ), ResourceSize( InResourceSize ), TrueResourceSize( InTrueResourceSize )
+	FSubItem( UObject* InObject, SIZE_T InNum, SIZE_T InMax, SIZE_T InTrueResourceSize )
+	: Object( InObject ), Num( InNum ), Max( InMax ), TrueResourceSize( InTrueResourceSize )
 	{}
 };
 
@@ -5340,27 +5338,23 @@ bool UEngine::HandleObjCommand( const TCHAR* Cmd, FOutputDevice& Ar )
 			for( UObject*& RefObj : ReferencedObjects )
 			{
 				FArchiveCountMem Count( RefObj );
-				// Get the 'old-style' resource size and the truer resource size
-				const SIZE_T ResourceSize = It->GetResourceSize( EResourceSizeMode::Inclusive );
 				const SIZE_T TrueResourceSize = It->GetResourceSize( EResourceSizeMode::Exclusive );
-				ThisObject.Add( Count, ResourceSize, TrueResourceSize );
+				ThisObject.Add( Count, TrueResourceSize );
 			}
 
 			FItem& ClassObjects = ObjectsByClass.FindOrAdd( ThisObject.Class );
 			ClassObjects.Count++;
 			ClassObjects.Num += ThisObject.Num;
 			ClassObjects.Max += ThisObject.Max;
-			ClassObjects.ResourceSize += ThisObject.ResourceSize;
 			ClassObjects.TrueResourceSize += ThisObject.TrueResourceSize;
 		}
 
 		ObjectsByClass.ValueSort( FCompareByInclusiveSize() );
 
-		UE_LOG( LogEngine, Log, TEXT("%32s, %6s %6s %6s %6s, %6s"), 
+		UE_LOG( LogEngine, Log, TEXT("%32s, %6s %6s %6s, %6s"), 
 			TEXT("Class"),
 			TEXT("IncMax"),
 			TEXT("IncNum"),
-			TEXT("ResInc"),
 			TEXT("ResExc"),
 			TEXT("Count") );
 
@@ -5376,15 +5370,14 @@ bool UEngine::HandleObjCommand( const TCHAR* Cmd, FOutputDevice& Ar )
 				Culled.Count += ClassObjects.Count;
 				Culled.Num += ClassObjects.Num;
 				Culled.Max += ClassObjects.Max;
-				Culled.ResourceSize += ClassObjects.ResourceSize;
 				Culled.TrueResourceSize += ClassObjects.TrueResourceSize;
 			}
 			else
 			{
-				UE_LOG( LogEngine, Log, TEXT("%32s, %6s %6s %6s %6s, %6i"), 
+				UE_LOG( LogEngine, Log, TEXT("%32s, %6s %6s %6s, %6i"), 
 					*Class->GetName(), 
 					*FHierarchy::Size(ClassObjects.Max), *FHierarchy::Size(ClassObjects.Num), 
-					*FHierarchy::Size(ClassObjects.ResourceSize), *FHierarchy::Size(ClassObjects.TrueResourceSize),
+					*FHierarchy::Size(ClassObjects.TrueResourceSize),
 					ClassObjects.Count 
 					);
 
@@ -5393,26 +5386,25 @@ bool UEngine::HandleObjCommand( const TCHAR* Cmd, FOutputDevice& Ar )
 			Total.Count += ClassObjects.Count;
 			Total.Num += ClassObjects.Num;
 			Total.Max += ClassObjects.Max;
-			Total.ResourceSize += ClassObjects.ResourceSize;
 			Total.TrueResourceSize += ClassObjects.TrueResourceSize;
 		}
 
 		if( Culled.Count > 0 )
 		{
 			UE_LOG( LogEngine, Log, TEXT("") );
-			UE_LOG( LogEngine, Log, TEXT("%32s, %6s %6s %6s %6s, %6i"), 
+			UE_LOG( LogEngine, Log, TEXT("%32s, %6s %6s %6s, %6i"), 
 				TEXT("(Culled)"), 
 				*FHierarchy::Size(Culled.Max), *FHierarchy::Size(Culled.Num), 
-				*FHierarchy::Size(Culled.ResourceSize), *FHierarchy::Size(Culled.TrueResourceSize),
+				*FHierarchy::Size(Culled.TrueResourceSize),
 				Culled.Count 
 				);
 		}
 
 		UE_LOG( LogEngine, Log, TEXT("") );
-		UE_LOG( LogEngine, Log, TEXT("%32s, %6s %6s %6s %6s, %6i"), 
+		UE_LOG( LogEngine, Log, TEXT("%32s, %6s %6s %6s, %6i"), 
 			TEXT("Total"), 
 			*FHierarchy::Size(Total.Max), *FHierarchy::Size(Total.Num), 
-			*FHierarchy::Size(Total.ResourceSize), *FHierarchy::Size(Total.TrueResourceSize),
+			*FHierarchy::Size(Total.TrueResourceSize),
 			Total.Count 
 			);
 		UE_LOG( LogEngine, Log, TEXT("**********************************************") );
@@ -5445,10 +5437,8 @@ bool UEngine::HandleObjCommand( const TCHAR* Cmd, FOutputDevice& Ar )
 		for( FObjectIterator It; It; ++It )
 		{
 			FArchiveCountMem Count( *It );
-			// Get the 'old-style' resource size and the truer resource size
-			const SIZE_T ResourceSize = It->GetResourceSize(EResourceSizeMode::Inclusive);
 			const SIZE_T TrueResourceSize = It->GetResourceSize(EResourceSizeMode::Exclusive);
-			Objects.Add(*It, FSubItem(*It, Count.GetNum(), Count.GetMax(), ResourceSize, TrueResourceSize));
+			Objects.Add(*It, FSubItem(*It, Count.GetNum(), Count.GetMax(), TrueResourceSize));
 			Classes.AddClassInstance(*It);
 			Outers.AddOuter(*It);
 			Flat.AddFlat(*It);
@@ -5477,9 +5467,8 @@ bool UEngine::HandleObjCommand( const TCHAR* Cmd, FOutputDevice& Ar )
 		Ar.Log( TEXT("Objects:") );
 		Ar.Log( TEXT("") );
 
-		UClass*   CheckType     = NULL;
-		UClass*   MetaClass		= NULL;
-		const bool bExportToFile = FParse::Param(Cmd,TEXT("FILE"));
+		UClass* CheckType = NULL;
+		UClass* MetaClass = NULL;
 
 		// allow checking for any Outer, not just a UPackage
 		UObject* CheckOuter = NULL;
@@ -5587,8 +5576,6 @@ bool UEngine::HandleObjCommand( const TCHAR* Cmd, FOutputDevice& Ar )
 
 				FArchiveCountMem Count( *It );
 
-				// Get the 'old-style' resource size and the truer resource size
-				const SIZE_T ResourceSize = It->GetResourceSize(EResourceSizeMode::Inclusive);
 				const SIZE_T TrueResourceSize = It->GetResourceSize(EResourceSizeMode::Exclusive);
 
 				int32 i;
@@ -5626,10 +5613,10 @@ bool UEngine::HandleObjCommand( const TCHAR* Cmd, FOutputDevice& Ar )
 
 				if( bShowDetailedObjectInfo )
 				{
-					new(Objects)FSubItem( *It, Count.GetNum(), Count.GetMax(), ResourceSize, TrueResourceSize );
+					new(Objects)FSubItem( *It, Count.GetNum(), Count.GetMax(), TrueResourceSize );
 				}
-				List[i].Add( Count, ResourceSize, TrueResourceSize );
-				Total.Add( Count, ResourceSize, TrueResourceSize );
+				List[i].Add( Count, TrueResourceSize );
+				Total.Add( Count, TrueResourceSize );
 			}
 		}
 
@@ -5652,28 +5639,15 @@ bool UEngine::HandleObjCommand( const TCHAR* Cmd, FOutputDevice& Ar )
 			};
 			Objects.Sort( FCompareFSubItem( bAlphaSort ) );
 
-			Ar.Logf( TEXT("%140s % 10s % 10s % 10s % 10s"), TEXT("Object"), TEXT("NumKBytes"), TEXT("MaxKBytes"), TEXT("ResKBytes"), TEXT("ExclusiveResKBytes") );
+			Ar.Logf( TEXT("%140s, % 10s, % 10s, % 10s"), TEXT("Object"), TEXT("NumKBytes"), TEXT("MaxKBytes"), TEXT("ExclusiveResKBytes") );
 
-			for( int32 ObjIndex=0; ObjIndex<Objects.Num(); ObjIndex++ )
+			for (const FSubItem& ObjItem : Objects)
 			{
-				const FSubItem& ObjItem = Objects[ObjIndex];
-
-				///MSSTART DAN PRICE MICROSOFT Mar 12th, 2007 export object data to a file
-				if(bExportToFile)
-				{
-					FString Path = TEXT("./ObjExport");
-					FString MungedPath = ObjItem.Object->GetOutermost()->GetName();
-					MungedPath.ReplaceInline(TEXT("/"), TEXT("_"));
-					const FString Filename = Path /  + TEXT(".") + MungedPath / ObjItem.Object->GetName() + TEXT(".t3d");
-					Ar.Logf( TEXT("%s"),*Filename);
-					UExporter::ExportToFile(ObjItem.Object, NULL, *Filename, 1, 0);
-				}					
-				//MSEND
-
-				Ar.Logf( TEXT("%140s % 10iK % 10iK % 10iK % 10iK"), *ObjItem.Object->GetFullName(), (int32)ObjItem.Num / 1024, (int32)ObjItem.Max / 1024, (int32)ObjItem.ResourceSize / 1024, (int32)ObjItem.TrueResourceSize / 1024 );
+				Ar.Logf(TEXT("%140s, % 10.2f, % 10.2f, % 10.2f"), *ObjItem.Object->GetFullName(), ObjItem.Num / 1024.0f, ObjItem.Max / 1024.0f, ObjItem.TrueResourceSize / 1024.0f);
 			}
-			Ar.Log( TEXT("") );
+			Ar.Log(TEXT(""));
 		}
+
 		if( List.Num() )
 		{
 			struct FCompareFItem
@@ -5689,15 +5663,15 @@ bool UEngine::HandleObjCommand( const TCHAR* Cmd, FOutputDevice& Ar )
 				}
 			};
 			List.Sort( FCompareFItem( bAlphaSort, bCountSort ) );
-			Ar.Logf(TEXT(" %100s % 6s % 10s % 10s % 10s % 10s"), TEXT("Class"), TEXT("Count"), TEXT("NumKBytes"), TEXT("MaxKBytes"), TEXT("ResKBytes"), TEXT("ExclusiveResKBytes") );
+			Ar.Logf(TEXT(" %100s, % 6s, % 10s, % 10s, % 10s"), TEXT("Class"), TEXT("Count"), TEXT("NumKBytes"), TEXT("MaxKBytes"), TEXT("ExclusiveResKBytes") );
 
 			for( int32 i=0; i<List.Num(); i++ )
 			{
-				Ar.Logf(TEXT(" %100s % 6i % 10iK % 10iK % 10iK % 10iK"), *List[i].Class->GetName(), (int32)List[i].Count, (int32)(List[i].Num/1024), (int32)(List[i].Max/1024), (int32)(List[i].ResourceSize/1024), (int32)(List[i].TrueResourceSize/1024) );
+				Ar.Logf(TEXT(" %100s, % 6i, % 10.2f, % 10.2f, % 10.2f"), *List[i].Class->GetName(), (int32)List[i].Count, List[i].Num / 1024.0f, List[i].Max / 1024.0f, List[i].TrueResourceSize / 1024.0f);
 			}
 			Ar.Log( TEXT("") );
 		}
-		Ar.Logf( TEXT("%i Objects (%.3fM / %.3fM / %.3fM / %.3fM)"), Total.Count, (float)Total.Num/1024.0/1024.0, (float)Total.Max/1024.0/1024.0, (float)Total.ResourceSize/1024.0/1024.0, (float)Total.TrueResourceSize/1024.0/1024.0 );
+		Ar.Logf( TEXT("%i Objects (%.3fM / %.3fM / %.3fM)"), Total.Count, (float)Total.Num/1024.0/1024.0, (float)Total.Max/1024.0/1024.0, (float)Total.TrueResourceSize/1024.0/1024.0 );
 		return true;
 
 	}
