jira UE-20672
-Fix a debug crash (check) when we turn on the Levels->LockReadOnlyFile, by clearing the selection for every readonly level
-Verify the read only state of the level file only once per frame instead of 2 time by actor in the level by tick

diff --git a/Engine/Source/Editor/WorldBrowser/Private/LevelCollectionModel.cpp b/Engine/Source/Editor/WorldBrowser/Private/LevelCollectionModel.cpp
--- a/Engine/Source/Editor/WorldBrowser/Private/LevelCollectionModel.cpp
+++ b/Engine/Source/Editor/WorldBrowser/Private/LevelCollectionModel.cpp
@@ -1487,6 +1487,15 @@ void FLevelCollectionModel::UnockAllLevels_Executed()
 
 void FLevelCollectionModel::ToggleReadOnlyLevels_Executed()
 {
+	//We are about to lock some Levels, deselect all actor and surfaces from the read only levels
+	if (!GEditor->bLockReadOnlyLevels)
+	{
+		DeselectActorsInAllReadOnlyLevel(GetAllLevels());
+		DeselectSurfaceInAllReadOnlyLevel(GetAllLevels());
+		// Tell the editor selection status was changed.
+		GEditor->NoteSelectionChange();
+	}
+
 	GEditor->bLockReadOnlyLevels = !GEditor->bLockReadOnlyLevels;
 }
 
@@ -1642,6 +1651,33 @@ void FLevelCollectionModel::FillSourceControlSubMenu(FMenuBuilder& InMenuBuilder
 	}
 }
 
+void FLevelCollectionModel::DeselectActorsInAllReadOnlyLevel(const FLevelModelList& InLevelList)
+{
+	const FScopedTransaction Transaction(LOCTEXT("DeselectActorsInReadOnlyLevel", "Deselect Actors in all read only Level"));
+
+	for (auto It = InLevelList.CreateConstIterator(); It; ++It)
+	{
+		if ((*It)->IsFileReadOnly())
+		{
+			(*It)->SelectActors(/*bSelect*/ false, /*bNotify*/ true, /*bSelectEvenIfHidden*/ true);
+		}
+	}
+}
+
+void FLevelCollectionModel::DeselectSurfaceInAllReadOnlyLevel(const FLevelModelList& InLevelList)
+{
+	const FScopedTransaction Transaction(LOCTEXT("DeselectSurfacesInReadOnlyLevel", "Deselect Surfaces in all read only Level"));
+
+	for (auto It = InLevelList.CreateConstIterator(); It; ++It)
+	{
+		if ((*It)->IsFileReadOnly())
+		{
+			(*It)->DeselectAllSurfaces();
+		}
+	}
+}
+
+
 void FLevelCollectionModel::OnLevelsCollectionChanged()
 {
 	UpdateAllLevels();
diff --git a/Engine/Source/Editor/WorldBrowser/Private/LevelCollectionModel.h b/Engine/Source/Editor/WorldBrowser/Private/LevelCollectionModel.h
--- a/Engine/Source/Editor/WorldBrowser/Private/LevelCollectionModel.h
+++ b/Engine/Source/Editor/WorldBrowser/Private/LevelCollectionModel.h
@@ -430,6 +430,12 @@ protected:
 	
 	/**  */
 	virtual void BindCommands();
+
+	/** Removes the Actors in all read-only Levels from the viewport's existing selection */
+	void DeselectActorsInAllReadOnlyLevel(const FLevelModelList& InLevelList);
+
+	/** Removes the Actors in all read-only Levels from the viewport's existing selection */
+	void DeselectSurfaceInAllReadOnlyLevel(const FLevelModelList& InLevelList);
 	
 	/** Called whenever level collection has been changed */
 	virtual void OnLevelsCollectionChanged();
diff --git a/Engine/Source/Runtime/Engine/Private/LevelUtils.cpp b/Engine/Source/Runtime/Engine/Private/LevelUtils.cpp
--- a/Engine/Source/Runtime/Engine/Private/LevelUtils.cpp
+++ b/Engine/Source/Runtime/Engine/Private/LevelUtils.cpp
@@ -11,6 +11,25 @@
 
 #define LOCTEXT_NAMESPACE "LevelUtils"
 
+
+#if WITH_EDITOR
+// Structure to hold the state of the Level file on disk, the goal is to query it only one time per frame.
+struct FLevelReadOnlyData
+{
+	FLevelReadOnlyData()
+		:IsReadOnly(false)
+		,LastUpdateTime(-1.0f)
+	{}
+	/** the current level file state */
+	bool IsReadOnly;
+	/** Last time when the level file state was update */
+	float LastUpdateTime;
+};
+// Map to link the level data with a level
+static TMap<ULevel*, FLevelReadOnlyData> LevelReadOnlyCache;
+
+#endif
+
 /////////////////////////////////////////////////////////////////////////////////////////
 //
 //	FindStreamingLevel methods.
@@ -87,28 +106,44 @@ ULevelStreaming* FLevelUtils::FindStreamingLevel(UWorld* InWorld, const TCHAR* I
  */
 bool FLevelUtils::IsLevelLocked(ULevel* Level)
 {
+//We should not check file status on disk if we are not running the editor
+#if WITH_EDITOR
 	// Don't permit spawning in read only levels if they are locked
 	if ( GIsEditor && !GIsEditorLoadingPackage )
 	{
 		if ( GEngine && GEngine->bLockReadOnlyLevels )
 		{
-			const UPackage* pPackage = Cast<UPackage>( Level->GetOutermost() );
-			if ( pPackage )
+			if (!LevelReadOnlyCache.Contains(Level))
+			{
+				LevelReadOnlyCache.Add(Level, FLevelReadOnlyData());
+			}
+			check(LevelReadOnlyCache.Contains(Level));
+			FLevelReadOnlyData &LevelData = LevelReadOnlyCache[Level];
+			//Make sure we test if the level file on disk is readonly only once a frame,
+			//when the frame time get updated.
+			if (LevelData.LastUpdateTime < Level->OwningWorld->GetRealTimeSeconds())
 			{
-				FString PackageFileName;
-				if ( FPackageName::DoesPackageExist( pPackage->GetName(), NULL, &PackageFileName ) )
+				LevelData.LastUpdateTime = Level->OwningWorld->GetRealTimeSeconds();
+				//If we dont find package we dont consider it as readonly
+				LevelData.IsReadOnly = false;
+				const UPackage* pPackage = Cast<UPackage>(Level->GetOutermost());
+				if (pPackage)
 				{
-					const bool bIsReadOnly = IFileManager::Get().IsReadOnly( *PackageFileName );
-					if( bIsReadOnly )
+					FString PackageFileName;
+					if (FPackageName::DoesPackageExist(pPackage->GetName(), NULL, &PackageFileName))
 					{
-						//This function is called from every slate tick so read-only levels spam this warning
-						//UE_LOG(LogLevel, Warning, TEXT( "Read-only packages are locked, unlock them from the Tools menu to permit changes to level: %s" ), *pPackage->GetName() );
-						return true;
+						LevelData.IsReadOnly = IFileManager::Get().IsReadOnly(*PackageFileName);
 					}
 				}
 			}
+
+			if (LevelData.IsReadOnly)
+			{
+				return true;
+			}
 		}
 	}
+#endif //#if WITH_EDITOR
 
 	// PIE levels, the persistent level, and transient move levels are usually never locked.
 	if ( Level->RootPackageHasAnyFlags(PKG_PlayInEditor) || Level->IsPersistentLevel() || Level->GetName() == TEXT("TransLevelMoveBuffer") )
