Added extra methods to SplineComponent for finding the closest point on the spline to a given world location.
PR #1537: Adds series of Find X Closest To World Location methods to USplineCom. (Contributed by [REDACTED_PERSON])

diff --git a/Engine/Source/Runtime/Engine/Classes/Components/SplineComponent.h b/Engine/Source/Runtime/Engine/Classes/Components/SplineComponent.h
--- a/Engine/Source/Runtime/Engine/Classes/Components/SplineComponent.h
+++ b/Engine/Source/Runtime/Engine/Classes/Components/SplineComponent.h
@@ -405,6 +405,51 @@ public:
 	UFUNCTION(BlueprintCallable, Category=Spline)
 	FVector GetScaleAtTime(float Time, bool bUseConstantVelocity = false) const;
 
+    /** Given a location, in world space, return the input key closest to that location. */
+	UFUNCTION(BlueprintCallable, Category=Spline)
+	float FindInputKeyClosestToWorldLocation(const FVector& WorldLocation) const;
+
+	/** Given a location, in world space, return the point on the curve that is closest to the location. */
+	UFUNCTION(BlueprintCallable, Category=Spline)
+	FVector FindLocationClosestToWorldLocation(const FVector& WorldLocation, ESplineCoordinateSpace::Type CoordinateSpace) const;
+
+	/** Given a location, in world spcae, return a unit direction vector of the spline tangent closest to the location. */
+	UFUNCTION(BlueprintCallable, Category=Spline)
+	FVector FindDirectionClosestToWorldLocation(const FVector& WorldLocation, ESplineCoordinateSpace::Type CoordinateSpace) const;
+
+	/** Given a location, in world space, return the tangent vector of the spline closest to the location. */
+	UFUNCTION(BlueprintCallable, Category = Spline)
+	FVector FindTangentClosestToWorldLocation(const FVector& WorldLocation, ESplineCoordinateSpace::Type CoordinateSpace) const;
+
+	/** Given a location, in world space, return a quaternion corresponding to the spline's rotation closest to the location. */
+	UFUNCTION(BlueprintCallable, Category = Spline)
+	FQuat FindQuaternionClosestToWorldLocation(const FVector& WorldLocation, ESplineCoordinateSpace::Type CoordinateSpace) const;
+
+	/** Given a location, in world space, return rotation corresponding to the spline's rotation closest to the location. */
+	UFUNCTION(BlueprintCallable, Category = Spline)
+	FRotator FindRotationClosestToWorldLocation(const FVector& WorldLocation, ESplineCoordinateSpace::Type CoordinateSpace) const;
+
+	/** Given a location, in world space, return a unit direction vector corresponding to the spline's up vector closest to the location. */
+	UFUNCTION(BlueprintCallable, Category = Spline)
+	FVector FindUpVectorClosestToWorldLocation(const FVector& WorldLocation, ESplineCoordinateSpace::Type CoordinateSpace) const;
+
+    /** Given a location, in world space, return a unit direction vector corresponding to the spline's right vector closest to the location. */
+	UFUNCTION(BlueprintCallable, Category = Spline)
+	FVector FindRightVectorClosestToWorldLocation(const FVector& WorldLocation, ESplineCoordinateSpace::Type CoordinateSpace) const;
+
+    /** Given a location, in world space, return the spline's roll closest to the location, in degrees. */
+	UFUNCTION(BlueprintCallable, Category = Spline)
+    float FindRollClosestToWorldLocation(const FVector& WorldLocation, ESplineCoordinateSpace::Type CoordinateSpace) const;
+
+	/** Given a location, in world space, return the spline's scale closest to the location. */
+	UFUNCTION(BlueprintCallable, Category = Spline)
+	FVector FindScaleClosestToWorldLocation(const FVector& WorldLocation) const;
+
+	/** Given a location, in world space, return an FTransform closest to that location. */
+	UFUNCTION(BlueprintCallable, Category = Spline)
+	FTransform FindTransformClosestToWorldLocation(const FVector& WorldLocation, ESplineCoordinateSpace::Type CoordinateSpace, bool bUseScale = false) const;
+
+
 	// UObject interface
 	virtual void PostLoad() override;
 	virtual void PostEditImport() override;
diff --git a/Engine/Source/Runtime/Engine/Private/Components/SplineComponent.cpp b/Engine/Source/Runtime/Engine/Private/Components/SplineComponent.cpp
--- a/Engine/Source/Runtime/Engine/Private/Components/SplineComponent.cpp
+++ b/Engine/Source/Runtime/Engine/Private/Components/SplineComponent.cpp
@@ -1007,6 +1007,84 @@ FVector USplineComponent::GetScaleAtTime(float Time, bool bUseConstantVelocity)
 }
 
 
+float USplineComponent::FindInputKeyClosestToWorldLocation(const FVector& WorldLocation) const
+{
+	const FVector LocalLocation = ComponentToWorld.InverseTransformPosition(WorldLocation);
+	float Dummy;
+	return SplineInfo.InaccurateFindNearest(LocalLocation, Dummy);
+}
+
+
+FVector USplineComponent::FindLocationClosestToWorldLocation(const FVector& WorldLocation, ESplineCoordinateSpace::Type CoordinateSpace) const
+{
+	const float Param = FindInputKeyClosestToWorldLocation(WorldLocation);
+	return GetLocationAtSplineInputKey(Param, CoordinateSpace);
+}
+
+
+FVector USplineComponent::FindDirectionClosestToWorldLocation(const FVector& WorldLocation, ESplineCoordinateSpace::Type CoordinateSpace) const
+{
+	const float Param = FindInputKeyClosestToWorldLocation(WorldLocation);
+	return GetDirectionAtSplineInputKey(Param, CoordinateSpace);
+}
+
+
+FVector USplineComponent::FindTangentClosestToWorldLocation(const FVector& WorldLocation, ESplineCoordinateSpace::Type CoordinateSpace) const
+{
+	const float Param = FindInputKeyClosestToWorldLocation(WorldLocation);
+	return GetTangentAtSplineInputKey(Param, CoordinateSpace);
+}
+
+
+FQuat USplineComponent::FindQuaternionClosestToWorldLocation(const FVector& WorldLocation, ESplineCoordinateSpace::Type CoordinateSpace) const
+{
+	const float Param = FindInputKeyClosestToWorldLocation(WorldLocation);
+	return GetQuaternionAtSplineInputKey(Param, CoordinateSpace);
+}
+
+
+FRotator USplineComponent::FindRotationClosestToWorldLocation(const FVector& WorldLocation, ESplineCoordinateSpace::Type CoordinateSpace) const
+{
+	const float Param = FindInputKeyClosestToWorldLocation(WorldLocation);
+	return GetRotationAtSplineInputKey(Param, CoordinateSpace);
+}
+
+
+FVector USplineComponent::FindUpVectorClosestToWorldLocation(const FVector& WorldLocation, ESplineCoordinateSpace::Type CoordinateSpace) const
+{
+	const float Param = FindInputKeyClosestToWorldLocation(WorldLocation);
+	return GetUpVectorAtSplineInputKey(Param, CoordinateSpace);
+}
+
+
+FVector USplineComponent::FindRightVectorClosestToWorldLocation(const FVector& WorldLocation, ESplineCoordinateSpace::Type CoordinateSpace) const
+{
+	const float Param = FindInputKeyClosestToWorldLocation(WorldLocation);
+	return GetRightVectorAtSplineInputKey(Param, CoordinateSpace);
+}
+
+
+float USplineComponent::FindRollClosestToWorldLocation(const FVector& WorldLocation, ESplineCoordinateSpace::Type CoordinateSpace) const
+{
+	const float Param = FindInputKeyClosestToWorldLocation(WorldLocation);
+	return GetRollAtSplineInputKey(Param, CoordinateSpace);
+}
+
+
+FVector USplineComponent::FindScaleClosestToWorldLocation(const FVector& WorldLocation) const
+{
+	const float Param = FindInputKeyClosestToWorldLocation(WorldLocation);
+	return GetScaleAtSplineInputKey(Param);
+}
+
+
+FTransform USplineComponent::FindTransformClosestToWorldLocation(const FVector& WorldLocation, ESplineCoordinateSpace::Type CoordinateSpace, bool bUseScale) const
+{
+	const float Param = FindInputKeyClosestToWorldLocation(WorldLocation);
+	return GetTransformAtSplineInputKey(Param, CoordinateSpace, bUseScale);
+}
+
+
 /** Used to store spline data during RerunConstructionScripts */
 class FSplineInstanceData : public FSceneComponentInstanceData
 {
