Client replay recording fixes.
-Pass URL options from LoadMap to the demo net driver, if ?DemoRec was specified
-Set the replay spectator's NetDriverName to the demo net driver, so that GetNetDriver and therefore IsLocalController work as expected for it
-Fix a crash in UDemoNetDriver::InitListen if the GameState hasn't been replicated to the recording client yet
-Use the "game=" option from the URL to get the replay spectator class to spawn if the game state or game mode class hasn't replicated yet

--------
Integrated using branch UE4-Orion-To-UE4 of change#2684826 by [REDACTED_PERSON] on 2015/09/09 12:03:38.

diff --git a/Engine/Source/Runtime/Engine/Classes/Engine/DemoNetDriver.h b/Engine/Source/Runtime/Engine/Classes/Engine/DemoNetDriver.h
--- a/Engine/Source/Runtime/Engine/Classes/Engine/DemoNetDriver.h
+++ b/Engine/Source/Runtime/Engine/Classes/Engine/DemoNetDriver.h
@@ -153,7 +153,7 @@ public:
 	bool ReadDemoFrame( FArchive* Archive );
 	void TickDemoPlayback( float DeltaSeconds );
 	void FinalizeFastForward( const float StartTime );
-	void SpawnDemoRecSpectator( UNetConnection* Connection );
+	void SpawnDemoRecSpectator( UNetConnection* Connection, const FURL& ListenURL );
 	void ResetDemoState();
 	void JumpToEndOfLiveReplay();
 	void AddEvent(const FString& Group, const FString& Meta, const TArray<uint8>& Data);
diff --git a/Engine/Source/Runtime/Engine/Classes/Engine/GameInstance.h b/Engine/Source/Runtime/Engine/Classes/Engine/GameInstance.h
--- a/Engine/Source/Runtime/Engine/Classes/Engine/GameInstance.h
+++ b/Engine/Source/Runtime/Engine/Classes/Engine/GameInstance.h
@@ -190,8 +190,14 @@ public:
 
 	inline FTimerManager& GetTimerManager() const { return *TimerManager; }
 
-	/** Start recording a replay with the given custom name and friendly name. */
-	virtual void StartRecordingReplay(const FString& InName, const FString& FriendlyName);
+	/**
+	 * Start recording a replay with the given custom name and friendly name.
+	 *
+	 * @param InName If not empty, the unique name to use as an identifier for the replay. If empty, a name will be automatically generated by the replay streamer implementation.
+	 * @param FriendlyName An optional (may be empty) descriptive name for the replay. Does not have to be unique.
+	 * @param AdditionalOptions Additional URL options to append to the URL that will be processed by the replay net driver. Will usually remain empty.
+	 */
+	virtual void StartRecordingReplay(const FString& InName, const FString& FriendlyName, const TArray<FString>& AdditionalOptions = TArray<FString>());
 
 	/** Stop recording a replay if one is currently in progress */
 	virtual void StopRecordingReplay();
diff --git a/Engine/Source/Runtime/Engine/Private/DemoNetDriver.cpp b/Engine/Source/Runtime/Engine/Private/DemoNetDriver.cpp
--- a/Engine/Source/Runtime/Engine/Private/DemoNetDriver.cpp
+++ b/Engine/Source/Runtime/Engine/Private/DemoNetDriver.cpp
@@ -680,13 +680,17 @@ bool UDemoNetDriver::InitListen( FNetworkNotify* InNotify, FURL& ListenURL, bool
 
 	TArray< FString > UserNames;
 
-	for ( int32 i = 0; i < GetWorld()->GameState->PlayerArray.Num(); i++ )
+	// If a client is recording a replay, GameState may not have replicated yet
+	if ( GetWorld()->GameState != nullptr )
 	{
-		APlayerState* PlayerState = GetWorld()->GameState->PlayerArray[i];
-
-		if ( !PlayerState->bIsABot && !PlayerState->bIsSpectator )
+		for ( int32 i = 0; i < GetWorld()->GameState->PlayerArray.Num(); i++ )
 		{
-			UserNames.Add( PlayerState->UniqueId.ToString() );
+			APlayerState* PlayerState = GetWorld()->GameState->PlayerArray[i];
+
+			if ( !PlayerState->bIsABot && !PlayerState->bIsSpectator )
+			{
+				UserNames.Add( PlayerState->UniqueId.ToString() );
+			}
 		}
 	}
 
@@ -715,7 +719,7 @@ bool UDemoNetDriver::InitListen( FNetworkNotify* InNotify, FURL& ListenURL, bool
 	FileAr->Flush();
 
 	// Spawn the demo recording spectator.
-	SpawnDemoRecSpectator( Connection );
+	SpawnDemoRecSpectator( Connection, ListenURL );
 
 	return true;
 }
@@ -1702,7 +1706,7 @@ void UDemoNetDriver::FinalizeFastForward( const float StartTime )
 	UE_LOG( LogDemo, Log, TEXT( "Fast forward took %.2f seconds." ), FastForwardTotalSeconds );
 }
 
-void UDemoNetDriver::SpawnDemoRecSpectator( UNetConnection* Connection )
+void UDemoNetDriver::SpawnDemoRecSpectator( UNetConnection* Connection, const FURL& ListenURL )
 {
 	check( Connection != NULL );
 
@@ -1710,7 +1714,20 @@ void UDemoNetDriver::SpawnDemoRecSpectator( UNetConnection* Connection )
 	// since the game mode instance isn't replicated to clients of live games.
 	AGameState* GameState = GetWorld() != nullptr ? GetWorld()->GetGameState() : nullptr;
 	TSubclassOf<AGameMode> DefaultGameModeClass = GameState != nullptr ? GameState->GameModeClass : nullptr;
-	AGameMode* DefaultGameMode = CastChecked<AGameMode>(DefaultGameModeClass.GetDefaultObject());
+	
+	// If we don't have a game mode class from the world, try to get it from the URL option.
+	// This may be true on clients who are recording a replay before the game mode class was replicated to them.
+	if (DefaultGameModeClass == nullptr)
+	{
+		const TCHAR* URLGameModeClass = ListenURL.GetOption(TEXT("game="), nullptr);
+		if (URLGameModeClass != nullptr)
+		{
+			UClass* GameModeFromURL = StaticLoadClass(AGameMode::StaticClass(), nullptr, URLGameModeClass);
+			DefaultGameModeClass = GameModeFromURL;
+		}
+	}
+
+	AGameMode* DefaultGameMode = Cast<AGameMode>(DefaultGameModeClass.GetDefaultObject());
 	UClass* C = DefaultGameMode != nullptr ? DefaultGameMode->ReplaySpectatorPlayerControllerClass : nullptr;
 
 	if ( C == NULL )
@@ -1738,6 +1755,9 @@ void UDemoNetDriver::SpawnDemoRecSpectator( UNetConnection* Connection )
 		}
 	}
 
+	// Make sure SpectatorController->GetNetDriver returns this driver. Ensures functions that depend on it,
+	// such as IsLocalController, work as expected.
+	SpectatorController->NetDriverName = NetDriverName;
 	SpectatorController->SetReplicates( true );
 	SpectatorController->SetAutonomousProxy( true );
 
diff --git a/Engine/Source/Runtime/Engine/Private/GameInstance.cpp b/Engine/Source/Runtime/Engine/Private/GameInstance.cpp
--- a/Engine/Source/Runtime/Engine/Private/GameInstance.cpp
+++ b/Engine/Source/Runtime/Engine/Private/GameInstance.cpp
@@ -740,7 +740,7 @@ void UGameInstance::OnSessionUserInviteAccepted(const bool bWasSuccess, const in
 	}
 }
 
-void UGameInstance::StartRecordingReplay(const FString& Name, const FString& FriendlyName)
+void UGameInstance::StartRecordingReplay(const FString& Name, const FString& FriendlyName, const TArray<FString>& AdditionalOptions)
 {
 	if ( FParse::Param( FCommandLine::Get(),TEXT( "NOREPLAYS" ) ) )
 	{
@@ -765,6 +765,11 @@ void UGameInstance::StartRecordingReplay(const FString& Name, const FString& Fri
 	DemoURL.Map = DemoName;
 	DemoURL.AddOption( *FString::Printf( TEXT( "DemoFriendlyName=%s" ), *FriendlyName ) );
 
+	for (const FString& Option : AdditionalOptions)
+	{
+		DemoURL.AddOption(*Option);
+	}
+
 	CurrentWorld->DestroyDemoNetDriver();
 
 	const FName NAME_DemoNetDriver( TEXT( "DemoNetDriver" ) );
diff --git a/Engine/Source/Runtime/Engine/Private/UnrealEngine.cpp b/Engine/Source/Runtime/Engine/Private/UnrealEngine.cpp
--- a/Engine/Source/Runtime/Engine/Private/UnrealEngine.cpp
+++ b/Engine/Source/Runtime/Engine/Private/UnrealEngine.cpp
@@ -9961,7 +9961,7 @@ bool UEngine::LoadMap( FWorldContext& WorldContext, FURL URL, class UPendingNetG
 			const TCHAR* DemoRecName = URL.GetOption( TEXT( "DemoRec=" ), NULL );
 
 			// Record the demo, optionally with the specified custom name.
-			WorldContext.OwningGameInstance->StartRecordingReplay( FString(DemoRecName), WorldContext.World()->GetMapName() );
+			WorldContext.OwningGameInstance->StartRecordingReplay( FString(DemoRecName), WorldContext.World()->GetMapName(), URL.Op );
 		}
 	}
 
