#ue4 - Refactor particle auto detachment to a new function.

--------

diff --git a/Engine/Source/Runtime/Engine/Classes/Particles/ParticleSystemComponent.h b/Engine/Source/Runtime/Engine/Classes/Particles/ParticleSystemComponent.h
--- a/Engine/Source/Runtime/Engine/Classes/Particles/ParticleSystemComponent.h
+++ b/Engine/Source/Runtime/Engine/Classes/Particles/ParticleSystemComponent.h
@@ -310,6 +310,9 @@ private:
 	/** Did we auto attach during activation? Used to determine if we should restore the relative transform during detachment. */
 	uint32 bDidAutoAttach:1;
 
+	/** Restore relative transform from auto attachment and optionally detach from parent (regardless of whether it was an auto attachment). */
+	void CancelAutoAttachment(bool bDetachFromParent);
+
 public:
 
 	/**
diff --git a/Engine/Source/Runtime/Engine/Private/Particles/ParticleComponents.cpp b/Engine/Source/Runtime/Engine/Private/Particles/ParticleComponents.cpp
--- a/Engine/Source/Runtime/Engine/Private/Particles/ParticleComponents.cpp
+++ b/Engine/Source/Runtime/Engine/Private/Particles/ParticleComponents.cpp
@@ -3932,18 +3932,7 @@ void UParticleSystemComponent::FinalizeTickComponent()
 		}
 		else if (bAutoManageAttachment)
 		{
-			if (AttachParent)
-			{
-				if (bDidAutoAttach)
-				{
-					// Restore relative transform from before attachment. Actual transform will be updated as part of DetachFromParent().
-					RelativeLocation = SavedAutoAttachRelativeLocation;
-					RelativeRotation = SavedAutoAttachRelativeRotation;
-					RelativeScale3D = SavedAutoAttachRelativeScale3D;
-				}
-				DetachFromParent(/*bMaintainWorldPosition=*/ false);
-				bDidAutoAttach = false;
-			}
+			CancelAutoAttachment(/*bDetachFromParent=*/ true);
 		}
 	}
 	bWasCompleted = bIsCompleted;
@@ -4562,6 +4551,26 @@ void UParticleSystemComponent::DeactivateSystem()
 	LastRenderTime = GetWorld()->GetTimeSeconds();
 }
 
+void UParticleSystemComponent::CancelAutoAttachment(bool bDetachFromParent)
+{
+	if (bAutoManageAttachment)
+	{
+		if (bDidAutoAttach)
+		{
+			// Restore relative transform from before attachment. Actual transform will be updated as part of DetachFromParent().
+			RelativeLocation = SavedAutoAttachRelativeLocation;
+			RelativeRotation = SavedAutoAttachRelativeRotation;
+			RelativeScale3D = SavedAutoAttachRelativeScale3D;
+			bDidAutoAttach = false;
+		}
+
+		if (bDetachFromParent)
+		{
+			DetachFromParent(/*bMaintainWorldPosition=*/ false);
+		}
+	}
+}
+
 void UParticleSystemComponent::ComputeCanTickInAnyThread()
 {
 	check(!bIsElligibleForAsyncTickComputed);
