PR #1504: Fix UpdateAdjustVolumeMultiplier Nan (Contributed by [REDACTED_PERSON])

diff --git a/Engine/Source/Runtime/Engine/Private/ActiveSound.cpp b/Engine/Source/Runtime/Engine/Private/ActiveSound.cpp
--- a/Engine/Source/Runtime/Engine/Private/ActiveSound.cpp
+++ b/Engine/Source/Runtime/Engine/Private/ActiveSound.cpp
@@ -279,7 +279,7 @@ FWaveInstance* FActiveSound::FindWaveInstance( const UPTRINT WaveInstanceHash )
 void FActiveSound::UpdateAdjustVolumeMultiplier( const float DeltaTime )
 {
 	// keep stepping towards our target until we hit our stop time
-	if( PlaybackTime <= TargetAdjustVolumeStopTime )
+	if( PlaybackTime < TargetAdjustVolumeStopTime )
 	{
 		CurrentAdjustVolumeMultiplier += (TargetAdjustVolumeMultiplier - CurrentAdjustVolumeMultiplier) * DeltaTime / (TargetAdjustVolumeStopTime - PlaybackTime);
 	}
