Integrated PR #1351: Add some LocalPlayer->ViewportClient null pointer checks to PlayerController.cpp (Thanks bozaro!)

diff --git a/Engine/Source/Runtime/Engine/Private/PlayerController.cpp b/Engine/Source/Runtime/Engine/Private/PlayerController.cpp
--- a/Engine/Source/Runtime/Engine/Private/PlayerController.cpp
+++ b/Engine/Source/Runtime/Engine/Private/PlayerController.cpp
@@ -1324,9 +1324,13 @@ void APlayerController::ClientTeamMessage_Implementation( APlayerState* SenderPl
 	}
 
 	// since this is on the client, we can assume that if Player exists, it is a LocalPlayer
-	if( Player != NULL && CastChecked<ULocalPlayer>(Player)->ViewportClient->ViewportConsole )
+	if (Player != NULL)
 	{
-		CastChecked<ULocalPlayer>(Player)->ViewportClient->ViewportConsole->OutputText( SMod );
+		UGameViewportClient *ViewportClient = CastChecked<ULocalPlayer>(Player)->ViewportClient;
+		if ( ViewportClient && ViewportClient->ViewportConsole )
+		{
+			CastChecked<ULocalPlayer>(Player)->ViewportClient->ViewportConsole->OutputText(SMod);
+		}
 	}
 }
 
@@ -1730,7 +1734,7 @@ bool APlayerController::GetHitResultUnderCursor(ECollisionChannel TraceChannel,
 {
 	ULocalPlayer* LocalPlayer = Cast<ULocalPlayer>(Player);
 	bool bHit = false;
-	if (LocalPlayer)
+	if (LocalPlayer && LocalPlayer->ViewportClient)
 	{
 		FVector2D MousePosition;
 		if (LocalPlayer->ViewportClient->GetMousePosition(MousePosition))
@@ -1751,7 +1755,7 @@ bool APlayerController::GetHitResultUnderCursorByChannel(ETraceTypeQuery TraceCh
 {
 	ULocalPlayer* LocalPlayer = Cast<ULocalPlayer>(Player);
 	bool bHit = false;
-	if (LocalPlayer)
+	if (LocalPlayer && LocalPlayer->ViewportClient)
 	{
 		FVector2D MousePosition;
 		if (LocalPlayer->ViewportClient->GetMousePosition(MousePosition))
@@ -1772,7 +1776,7 @@ bool APlayerController::GetHitResultUnderCursorForObjects(const TArray<TEnumAsBy
 {
 	ULocalPlayer* LocalPlayer = Cast<ULocalPlayer>(Player);
 	bool bHit = false;
-	if (LocalPlayer)
+	if (LocalPlayer && LocalPlayer->ViewportClient)
 	{
 		FVector2D MousePosition;
 		if (LocalPlayer->ViewportClient->GetMousePosition(MousePosition))
@@ -2036,7 +2040,8 @@ bool APlayerController::InputKey(FKey Key, EInputEvent EventType, float AmountDe
 		if (bEnableClickEvents && Key == EKeys::LeftMouseButton)
 		{
 			FVector2D MousePosition;
-			if (CastChecked<ULocalPlayer>(Player)->ViewportClient->GetMousePosition(MousePosition))
+			UGameViewportClient* ViewportClient = CastChecked<ULocalPlayer>(Player)->ViewportClient;
+			if (ViewportClient && ViewportClient->GetMousePosition(MousePosition))
 			{
 				UPrimitiveComponent* ClickedPrimitive = NULL;
 				if (bEnableMouseOverEvents)
@@ -3125,7 +3130,7 @@ void APlayerController::ConsoleKey(FKey Key)
 #if !UE_BUILD_SHIPPING
 	if (ULocalPlayer* LocalPlayer = Cast<ULocalPlayer>(Player))
 	{
-		if (LocalPlayer->ViewportClient->ViewportConsole)
+		if (LocalPlayer->ViewportClient && LocalPlayer->ViewportClient->ViewportConsole)
 		{
 			LocalPlayer->ViewportClient->ViewportConsole->InputKey(0, Key, IE_Pressed);
 		}
@@ -3137,7 +3142,7 @@ void APlayerController::SendToConsole(const FString& Command)
 #if !UE_BUILD_SHIPPING
 	if (ULocalPlayer* LocalPlayer = Cast<ULocalPlayer>(Player))
 	{
-		if (LocalPlayer->ViewportClient->ViewportConsole)
+		if (LocalPlayer->ViewportClient && LocalPlayer->ViewportClient->ViewportConsole)
 		{
 			LocalPlayer->ViewportClient->ViewportConsole->ConsoleCommand(Command);
 		}
@@ -3217,7 +3222,7 @@ APlayerState* APlayerController::GetSplitscreenPlayerByIndex(int32 PlayerIndex)
 			UNetConnection* RemoteConnection = Cast<UNetConnection>(Player);
 			if ( LP != NULL )
 			{
-				const TArray<ULocalPlayer*>& GamePlayers = LP->ViewportClient->GetOuterUEngine()->GetGamePlayers(GetWorld());
+				const TArray<ULocalPlayer*>& GamePlayers = LP->GetOuterUEngine()->GetGamePlayers(GetWorld());
 				// this PC is a local player
 				if ( PlayerIndex >= 0 && PlayerIndex < GamePlayers.Num() )
 				{
@@ -3312,7 +3317,7 @@ int32 APlayerController::GetSplitscreenPlayerCount() const
 			UNetConnection* RemoteConnection = Cast<UNetConnection>(Player);
 			if ( LP != NULL )
 			{
-				Result = LP->ViewportClient->GetOuterUEngine()->GetNumGamePlayers(GetWorld());
+				Result = LP->GetOuterUEngine()->GetNumGamePlayers(GetWorld());
 			}
 			else if ( RemoteConnection != NULL )
 			{
@@ -3694,7 +3699,7 @@ void APlayerController::TickPlayerInput(const float DeltaSeconds, const bool bGa
 			UGameViewportClient* ViewportClient = LocalPlayer->ViewportClient;
 
 			// Only send mouse hit events if we're directly over the viewport.
-			if ( ViewportClient->GetGameViewportWidget().IsValid() && ViewportClient->GetGameViewportWidget()->IsDirectlyHovered() )
+			if ( ViewportClient && ViewportClient->GetGameViewportWidget().IsValid() && ViewportClient->GetGameViewportWidget()->IsDirectlyHovered() )
 			{
 				if ( LocalPlayer->ViewportClient->GetMousePosition(MousePosition) )
 				{
@@ -4252,7 +4257,7 @@ bool APlayerController::GetMousePosition(float& LocationX, float& LocationY) con
 	bool bGotMousePosition = false;
 	ULocalPlayer* LocalPlayer = Cast<ULocalPlayer>(Player);
 
-	if (LocalPlayer)
+	if (LocalPlayer && LocalPlayer->ViewportClient)
 	{
 		FVector2D MousePosition;
 		
