Fixed spectator pawn movement to account for per-actor time dilation.

(merged from branch)

diff --git a/Engine/Source/Runtime/Engine/Private/SpectatorPawnMovement.cpp b/Engine/Source/Runtime/Engine/Private/SpectatorPawnMovement.cpp
--- a/Engine/Source/Runtime/Engine/Private/SpectatorPawnMovement.cpp
+++ b/Engine/Source/Runtime/Engine/Private/SpectatorPawnMovement.cpp
@@ -21,7 +21,7 @@ USpectatorPawnMovement::USpectatorPawnMovement(const FObjectInitializer& ObjectI
 }
 
 
-void USpectatorPawnMovement::TickComponent(float DeltaTime, enum ELevelTick TickType, FActorComponentTickFunction *ThisTickFunction)
+void USpectatorPawnMovement::TickComponent(float DeltaTime, enum ELevelTick TickType, FActorComponentTickFunction* ThisTickFunction)
 {
 	if (!PawnOwner || !UpdatedComponent)
 	{
@@ -32,7 +32,7 @@ void USpectatorPawnMovement::TickComponent(float DeltaTime, enum ELevelTick Tick
 	float AdjustedDeltaTime = DeltaTime;
 	if (bIgnoreTimeDilation)
 	{
-		const float WorldTimeDilation = PawnOwner->GetWorldSettings()->GetEffectiveTimeDilation();
+		const float WorldTimeDilation = PawnOwner->GetActorTimeDilation();
 		if (WorldTimeDilation > KINDA_SMALL_NUMBER)
 		{
 			AdjustedDeltaTime = DeltaTime / WorldTimeDilation;
