nicer UI

diff --git a/Engine/Source/Runtime/Engine/Classes/Components/LightComponent.h b/Engine/Source/Runtime/Engine/Classes/Components/LightComponent.h
--- a/Engine/Source/Runtime/Engine/Classes/Components/LightComponent.h
+++ b/Engine/Source/Runtime/Engine/Classes/Components/LightComponent.h
@@ -141,7 +141,7 @@ class ENGINE_API ULightComponent : public ULightComponentBase
 	 * Brightness factor applied to the light when the light function is specified but disabled, for example in scene captures that use SceneCapView_LitNoShadows. 
 	 * This should be set to the average brightness of the light function material's emissive input, which should be between 0 and 1.
 	 */
-	UPROPERTY(EditAnywhere, BlueprintReadOnly, Category=LightFunction)
+	UPROPERTY(EditAnywhere, BlueprintReadOnly, Category=LightFunction, meta=(UIMin = "0.0", UIMax = "1.0"))
 	float DisabledBrightness;
 
 	/** 
