Fix cloth crash when shutting down editor. This is a temp fix while we wait for the lifetime factory fix for APEX

diff --git a/Engine/Source/Runtime/Engine/Private/SkeletalMeshComponentPhysics.cpp b/Engine/Source/Runtime/Engine/Private/SkeletalMeshComponentPhysics.cpp
--- a/Engine/Source/Runtime/Engine/Private/SkeletalMeshComponentPhysics.cpp
+++ b/Engine/Source/Runtime/Engine/Private/SkeletalMeshComponentPhysics.cpp
@@ -2089,15 +2089,17 @@ bool USkeletalMeshComponent::CreateClothingActor(int32 AssetIndex, physx::apex::
 	check(ClothingActors.Num() == AssetIndex);
 	ClothingActors.AddZeroed();
 
-	FPhysScene* PhysScene = NULL;
+	UWorld* World = GetWorld();
+	FPhysScene* PhysScene = World ? World->GetPhysicsScene() : nullptr;
 
-	if (GetWorld())
+	//this clothing actor would be crated later because this becomes invalid
+	if (!PhysScene)
 	{
-		PhysScene = GetWorld()->GetPhysicsScene();
+		return false;
 	}
 
-	//this clothing actor would be crated later because this becomes invalid
-	if (!PhysScene)
+	//TODO: this is a workaround while waiting for fix: apex cloth factory lifetime should depend on cloth asset. For now we make sure to only create asset if the scene ticks which is the only way to get refcounted
+	if(World->bShouldSimulatePhysics == false)
 	{
 		return false;
 	}
