UE4 - Critical fix for bad particle systems that have dirty diributions saved to disk.

diff --git a/Engine/Source/Runtime/Engine/Private/Distributions.cpp b/Engine/Source/Runtime/Engine/Private/Distributions.cpp
--- a/Engine/Source/Runtime/Engine/Private/Distributions.cpp
+++ b/Engine/Source/Runtime/Engine/Private/Distributions.cpp
@@ -872,12 +872,15 @@ void FRawDistributionFloat::Initialize()
 		}
 		bNeedsUpdating = true;
 	}
-
 	// only initialize if we need to
 	if (!bNeedsUpdating)
 	{
 		return;
 	}
+	if (!GIsEditor && !IsInGameThread() && !IsInAsyncLoadingThread())
+	{
+		return;
+	}
 	check(IsInGameThread() || IsInAsyncLoadingThread());
 
 	// always empty out the lookup table
