UCharacterMovementComponent no longer replicates moves to the server if its character's Player is null, which can be the case if the local player has enabled the debug camera. Prevents "no owning connection" warning spam for server RPCs. Fixes UE-20921.

diff --git a/Engine/Source/Runtime/Engine/Private/Components/CharacterMovementComponent.cpp b/Engine/Source/Runtime/Engine/Private/Components/CharacterMovementComponent.cpp
--- a/Engine/Source/Runtime/Engine/Private/Components/CharacterMovementComponent.cpp
+++ b/Engine/Source/Runtime/Engine/Private/Components/CharacterMovementComponent.cpp
@@ -6006,6 +6006,13 @@ void UCharacterMovementComponent::ReplicateMoveToServer(float DeltaTime, const F
 		return;
 	}
 
+	// Bail out if our character's controller doesn't have a Player. This may be the case when the local player
+	// has switched to another controller, such as a debug camera controller.
+	if (PC && PC->Player == nullptr)
+	{
+		return;
+	}
+
 	FNetworkPredictionData_Client_Character* ClientData = GetPredictionData_Client_Character();
 	if (!ClientData)
 	{
