#ANIM: Fix crash of undo animinstance for modifying joints

Merging

//depot/UE4-Orion/Engine/Source/Runtime/Engine/Classes/Animation/AnimNodeBase.h

to //depot/UE4/Engine/Source/Runtime/Engine/Classes/Animation/AnimNodeBase.h

diff --git a/Engine/Source/Runtime/Engine/Classes/Animation/AnimNodeBase.h b/Engine/Source/Runtime/Engine/Classes/Animation/AnimNodeBase.h
--- a/Engine/Source/Runtime/Engine/Classes/Animation/AnimNodeBase.h
+++ b/Engine/Source/Runtime/Engine/Classes/Animation/AnimNodeBase.h
@@ -397,19 +397,23 @@ struct FExposedValueHandler
 	UPROPERTY()
 	FName BoundFunction;
 
-	UPROPERTY(Transient)
-	UFunction* Function;
-
 	// Direct data access to property in anim instance
 	UPROPERTY()
 	TArray<FExposedValueCopyRecord> CopyRecords;
 
+	// funtion pointer if BoundFunction != NAME_None
+	UFunction* Function;
+
 	void Initialize(FAnimNode_Base* AnimNode, UAnimInstance* AnimInstance) 
 	{
 		if (BoundFunction != NAME_None)
 		{
 			Function = AnimInstance->FindFunction(BoundFunction);
 		}
+		else
+		{
+			Function = NULL;
+		}
 
 		// initialize copy records
 		for(auto& CopyRecord : CopyRecords)
