Added physxinfo command which shows information about the libs and version of PhysX used

diff --git a/Engine/Source/Runtime/Engine/Private/PhysicsEngine/PhysUtils.cpp b/Engine/Source/Runtime/Engine/Private/PhysicsEngine/PhysUtils.cpp
--- a/Engine/Source/Runtime/Engine/Private/PhysicsEngine/PhysUtils.cpp
+++ b/Engine/Source/Runtime/Engine/Private/PhysicsEngine/PhysUtils.cpp
@@ -591,6 +591,24 @@ bool ExecPhysCommands(const TCHAR* Cmd, FOutputDevice* Ar, UWorld* InWorld)
 		FPhysxSharedData::Get().DumpSharedMemoryUsage(Ar);
 		return 1;
 	}
+	else if(FParse::Command(&Cmd, TEXT("PHYSXINFO")))
+	{
+		UE_LOG(LogConsoleResponse, Display, TEXT("PhysX Info:"));
+		UE_LOG(LogConsoleResponse, Display, TEXT("  Version: %d.%d.%d"), PX_PHYSICS_VERSION_MAJOR, PX_PHYSICS_VERSION_MINOR, PX_PHYSICS_VERSION_BUGFIX);
+#if UE_BUILD_DEBUG && !defined(NDEBUG)
+		UE_LOG(LogConsoleResponse, Display, TEXT("  Configuration: DEBUG"));
+#elif WITH_PHYSX_RELEASE
+		UE_LOG(LogConsoleResponse, Display, TEXT("  Configuration: RELEASE"));
+#else
+		UE_LOG(LogConsoleResponse, Display, TEXT("  Configuration: PROFILE"));
+#endif
+#if WITH_PHYSICS_COOKING || WITH_RUNTIME_PHYSICS_COOKING
+	UE_LOG(LogConsoleResponse, Display, TEXT("  Cooking Module: TRUE"));
+#else
+	UE_LOG(LogConsoleResponse, Display, TEXT("  Cooking Module: FALSE"));
+#endif
+		return 1;
+	}
 
 #endif // WITH_PHYSX
 
