Use Ar instead of UE_LOG for info

diff --git a/Engine/Source/Runtime/Engine/Private/PhysicsEngine/PhysUtils.cpp b/Engine/Source/Runtime/Engine/Private/PhysicsEngine/PhysUtils.cpp
--- a/Engine/Source/Runtime/Engine/Private/PhysicsEngine/PhysUtils.cpp
+++ b/Engine/Source/Runtime/Engine/Private/PhysicsEngine/PhysUtils.cpp
@@ -593,19 +593,19 @@ bool ExecPhysCommands(const TCHAR* Cmd, FOutputDevice* Ar, UWorld* InWorld)
 	}
 	else if(FParse::Command(&Cmd, TEXT("PHYSXINFO")))
 	{
-		UE_LOG(LogConsoleResponse, Display, TEXT("PhysX Info:"));
-		UE_LOG(LogConsoleResponse, Display, TEXT("  Version: %d.%d.%d"), PX_PHYSICS_VERSION_MAJOR, PX_PHYSICS_VERSION_MINOR, PX_PHYSICS_VERSION_BUGFIX);
+		Ar->Logf(TEXT("PhysX Info:"));
+		Ar->Logf(TEXT("  Version: %d.%d.%d"), PX_PHYSICS_VERSION_MAJOR, PX_PHYSICS_VERSION_MINOR, PX_PHYSICS_VERSION_BUGFIX);
 #if UE_BUILD_DEBUG && !defined(NDEBUG)
-		UE_LOG(LogConsoleResponse, Display, TEXT("  Configuration: DEBUG"));
+		Ar->Logf(TEXT("  Configuration: DEBUG"));
 #elif WITH_PHYSX_RELEASE
-		UE_LOG(LogConsoleResponse, Display, TEXT("  Configuration: RELEASE"));
+		Ar->Logf(TEXT("  Configuration: RELEASE"));
 #else
-		UE_LOG(LogConsoleResponse, Display, TEXT("  Configuration: PROFILE"));
+		Ar->Logf(TEXT("  Configuration: PROFILE"));
 #endif
 #if WITH_PHYSICS_COOKING || WITH_RUNTIME_PHYSICS_COOKING
-	UE_LOG(LogConsoleResponse, Display, TEXT("  Cooking Module: TRUE"));
+		Ar->Logf(TEXT("  Cooking Module: TRUE"));
 #else
-	UE_LOG(LogConsoleResponse, Display, TEXT("  Cooking Module: FALSE"));
+		Ar->Logf(TEXT("  Cooking Module: FALSE"));
 #endif
 		return 1;
 	}
