Merging CL 2686647 from UE4-Orion to UE4

diff --git a/Engine/Source/Runtime/Engine/Private/AudioDevice.cpp b/Engine/Source/Runtime/Engine/Private/AudioDevice.cpp
--- a/Engine/Source/Runtime/Engine/Private/AudioDevice.cpp
+++ b/Engine/Source/Runtime/Engine/Private/AudioDevice.cpp
@@ -821,6 +821,17 @@ bool FAudioDevice::HandleClearSoloCommand(const TCHAR* Cmd, FOutputDevice& Ar)
 	return true;
 }
 
+bool FAudioDevice::HandlePlayAllPIEAudioCommand(const TCHAR* Cmd, FOutputDevice& Ar)
+{
+	FAudioDeviceManager* DeviceManager = GEngine->GetAudioDeviceManager();
+	if (DeviceManager)
+	{
+		DeviceManager->TogglePlayAllDeviceAudio();
+	}
+	return true;
+}
+
+
 bool FAudioDevice::HandleAudio3dVisualizeCommand(const TCHAR* Cmd, FOutputDevice& Ar)
 {
 	FAudioDeviceManager* DeviceManager = GEngine->GetAudioDeviceManager();
@@ -831,6 +842,15 @@ bool FAudioDevice::HandleAudio3dVisualizeCommand(const TCHAR* Cmd, FOutputDevice
 	return true;
 }
 
+bool FAudioDevice::HandlePlayAllPIEAudioCommand(const TCHAR* Cmd, FOutputDevice& Ar)
+{
+	FAudioDeviceManager* DeviceManager = GEngine->GetAudioDeviceManager();
+	if (DeviceManager)
+	{
+		DeviceManager->TogglePlayAllDeviceAudio();
+	}
+	return true;
+}
 #endif // !UE_BUILD_SHIPPING
 
 EDebugState FAudioDevice::GetMixDebugState( void )
@@ -937,6 +957,10 @@ bool FAudioDevice::Exec( UWorld* InWorld, const TCHAR* Cmd, FOutputDevice& Ar )
 	{
 		return HandleClearSoloCommand(Cmd, Ar);
 	}
+	else if (FParse::Command(&Cmd, TEXT("PlayAllPIEAudio")))
+	{
+		return HandlePlayAllPIEAudioCommand(Cmd, Ar);
+	}
 	else if (FParse::Command(&Cmd, TEXT("Audio3dVisualize")))
 	{
 		return HandleAudio3dVisualizeCommand(Cmd, Ar);
@@ -2661,6 +2685,19 @@ void FAudioDevice::StopSoundsUsingResource(USoundWave* SoundWave, TArray<UAudioC
 	}
 }
 
+bool FAudioDevice::IsAudioDeviceMuted() const
+{
+	// First check to see if the device manager has "bPlayAllPIEAudio" enabled
+	FAudioDeviceManager* DeviceManager = GEngine->GetAudioDeviceManager();
+	if (DeviceManager && DeviceManager->IsPlayAllDeviceAudio())
+	{
+		return false;
+	}
+
+	return bIsDeviceMuted;
+}
+
+
 #if WITH_EDITOR
 void FAudioDevice::OnBeginPIE(const bool bIsSimulating)
 {
diff --git a/Engine/Source/Runtime/Engine/Private/AudioDeviceManager.cpp b/Engine/Source/Runtime/Engine/Private/AudioDeviceManager.cpp
--- a/Engine/Source/Runtime/Engine/Private/AudioDeviceManager.cpp
+++ b/Engine/Source/Runtime/Engine/Private/AudioDeviceManager.cpp
@@ -37,6 +37,7 @@ FAudioDeviceManager::FAudioDeviceManager()
 	, NextResourceID(1)
 	, SoloDeviceHandle(INDEX_NONE)
 	, ActiveAudioDeviceHandle(INDEX_NONE)
+	, bPlayAllDeviceAudio(false)
 	, bVisualize3dDebug(false)
 {
 }
@@ -493,6 +494,16 @@ void FAudioDeviceManager::RemoveSoundMix(USoundMix* SoundMix)
 	}
 }
 
+void FAudioDeviceManager::TogglePlayAllDeviceAudio()
+{
+	bPlayAllDeviceAudio = !bPlayAllDeviceAudio;
+}
+
+bool FAudioDeviceManager::IsPlayAllDeviceAudio() const
+{
+	return bPlayAllDeviceAudio;
+}
+
 bool FAudioDeviceManager::IsVisualizeDebug3dEnabled() const
 {
 	return bVisualize3dDebug;
diff --git a/Engine/Source/Runtime/Engine/Public/AudioDevice.h b/Engine/Source/Runtime/Engine/Public/AudioDevice.h
--- a/Engine/Source/Runtime/Engine/Public/AudioDevice.h
+++ b/Engine/Source/Runtime/Engine/Public/AudioDevice.h
@@ -235,6 +235,7 @@ public:
 	bool HandleEnableHRTFForAllCommand(const TCHAR* Cmd, FOutputDevice& Ar);
 	bool HandleSoloCommand( const TCHAR* Cmd, FOutputDevice& Ar );
 	bool HandleClearSoloCommand( const TCHAR* Cmd, FOutputDevice& Ar );
+	bool HandlePlayAllPIEAudioCommand( const TCHAR* Cmd, FOutputDevice& Ar );
 	bool HandleAudio3dVisualizeCommand(const TCHAR* Cmd, FOutputDevice& Ar);
 #endif
 
@@ -539,6 +540,8 @@ public:
 		return bHRTFEnabledForAll && IsSpatializationPluginEnabled();
 	}
 
+	bool IsAudioDeviceMuted() const;
+
 protected:
 	friend class FSoundSource;
 
diff --git a/Engine/Source/Runtime/Engine/Public/AudioDeviceManager.h b/Engine/Source/Runtime/Engine/Public/AudioDeviceManager.h
--- a/Engine/Source/Runtime/Engine/Public/AudioDeviceManager.h
+++ b/Engine/Source/Runtime/Engine/Public/AudioDeviceManager.h
@@ -114,6 +114,12 @@ public:
 	/** Removes sound mix from all audio devices */
 	void RemoveSoundMix(USoundMix* SoundMix);
 
+	/** Toggles playing audio for all active PIE sessions (and all devices). */
+	void TogglePlayAllDeviceAudio();
+
+	/** Gets whether or not all devices should play their audio. */
+	bool IsPlayAllDeviceAudio() const;
+
 	/** Is debug visualization of 3d sounds enabled */
 	bool IsVisualizeDebug3dEnabled() const;
 
@@ -173,6 +179,9 @@ private:
 	/** Which audio device is currently active */
 	uint32 ActiveAudioDeviceHandle;
 
+	/** Whether or not to play all audio in all active audio devices. */
+	bool bPlayAllDeviceAudio;
+
 	/** Whether or not 3d debug visualization is enabled. */
 	bool bVisualize3dDebug;
 };
diff --git a/Engine/Source/Runtime/Windows/XAudio2/Private/XAudio2Source.cpp b/Engine/Source/Runtime/Windows/XAudio2/Private/XAudio2Source.cpp
--- a/Engine/Source/Runtime/Windows/XAudio2/Private/XAudio2Source.cpp
+++ b/Engine/Source/Runtime/Windows/XAudio2/Private/XAudio2Source.cpp
@@ -514,7 +514,7 @@ bool FXAudio2SoundSource::Init(FWaveInstance* InWaveInstance)
  */
 void FXAudio2SoundSource::GetChannelVolumes(float ChannelVolumes[CHANNEL_MATRIX_COUNT], float AttenuatedVolume)
 {
-	if (FApp::GetVolumeMultiplier() == 0.0f || AudioDevice->bIsDeviceMuted)
+	if (FApp::GetVolumeMultiplier() == 0.0f || AudioDevice->IsAudioDeviceMuted())
 	{
 		for( int32 i = 0; i < CHANNELOUT_COUNT; i++ )
 		{
