Fix function duplication after merge

diff --git a/Engine/Source/Runtime/Engine/Private/AudioDevice.cpp b/Engine/Source/Runtime/Engine/Private/AudioDevice.cpp
--- a/Engine/Source/Runtime/Engine/Private/AudioDevice.cpp
+++ b/Engine/Source/Runtime/Engine/Private/AudioDevice.cpp
@@ -841,16 +841,6 @@ bool FAudioDevice::HandleAudio3dVisualizeCommand(const TCHAR* Cmd, FOutputDevice
 	}
 	return true;
 }
-
-bool FAudioDevice::HandlePlayAllPIEAudioCommand(const TCHAR* Cmd, FOutputDevice& Ar)
-{
-	FAudioDeviceManager* DeviceManager = GEngine->GetAudioDeviceManager();
-	if (DeviceManager)
-	{
-		DeviceManager->TogglePlayAllDeviceAudio();
-	}
-	return true;
-}
 #endif // !UE_BUILD_SHIPPING
 
 EDebugState FAudioDevice::GetMixDebugState( void )
