Changed libraryCategoriesText to EDITORONLY_DATA as it's not needed in cooked builds and was causing deterministic cooking issues with material functions.

diff --git a/Engine/Source/Runtime/Engine/Classes/Materials/MaterialFunction.h b/Engine/Source/Runtime/Engine/Classes/Materials/MaterialFunction.h
--- a/Engine/Source/Runtime/Engine/Classes/Materials/MaterialFunction.h
+++ b/Engine/Source/Runtime/Engine/Classes/Materials/MaterialFunction.h
@@ -29,6 +29,7 @@ class UMaterialFunction : public UObject
 	UPROPERTY(EditAnywhere, Category=MaterialFunction, AssetRegistrySearchable)
 	uint32 bExposeToLibrary:1;
 
+#if WITH_EDITORONLY_DATA
 	/** 
 	 * Categories that this function belongs to in the material function library.  
 	 * Ideally categories should be chosen carefully so that there are not too many.
@@ -42,7 +43,7 @@ class UMaterialFunction : public UObject
 	 */
 	UPROPERTY(EditAnywhere, Category=MaterialFunction, AssetRegistrySearchable)
 	TArray<FText> LibraryCategoriesText;
-
+#endif
 	/** Array of material expressions, excluding Comments.  Used by the material editor. */
 	UPROPERTY()
 	TArray<class UMaterialExpression*> FunctionExpressions;
diff --git a/Engine/Source/Runtime/Engine/Private/Materials/MaterialExpressions.cpp b/Engine/Source/Runtime/Engine/Private/Materials/MaterialExpressions.cpp
--- a/Engine/Source/Runtime/Engine/Private/Materials/MaterialExpressions.cpp
+++ b/Engine/Source/Runtime/Engine/Private/Materials/MaterialExpressions.cpp
@@ -6378,8 +6378,9 @@ uint32 UMaterialExpressionCustom::GetOutputType(int32 OutputIndex)
 UMaterialFunction::UMaterialFunction(const FObjectInitializer& ObjectInitializer)
 	: Super(ObjectInitializer)
 {
+#if WITH_EDITORONLY_DATA
 	LibraryCategoriesText.Add(LOCTEXT("Misc", "Misc"));
-
+#endif
 #if WITH_EDITORONLY_DATA
 	PreviewMaterial = NULL;
 	ThumbnailInfo = NULL;
