Reversed Index Buffers. Depth only render pass now use the unified mesh when possible.
Can be disabled per staticmesh through the lod settings : "Build Reversed Index Buffer".
Can be dynamically toggled with command "ToggleReversedIndexBuffers".

diff --git a/Engine/Source/Developer/MeshUtilities/Private/MeshUtilities.cpp b/Engine/Source/Developer/MeshUtilities/Private/MeshUtilities.cpp
--- a/Engine/Source/Developer/MeshUtilities/Private/MeshUtilities.cpp
+++ b/Engine/Source/Developer/MeshUtilities/Private/MeshUtilities.cpp
@@ -3032,9 +3032,22 @@ public:
 					);
 			}
 
+			// Build the reversed index buffer.
+			if (InOutModels[0].BuildSettings.bBuildReversedIndexBuffer)
+			{
+				TArray<uint32> InversedIndices;
+				const int32 IndexCount = CombinedIndices.Num();
+				InversedIndices.AddUninitialized(IndexCount);
+				for (int32 i = 0; i < IndexCount; ++i)
+				{
+					InversedIndices[i] = CombinedIndices[IndexCount - 1 - i];
+				}
+				LODModel.ReversedIndexBuffer.SetIndices(InversedIndices, bNeeds32BitIndices ? EIndexBufferStride::Force32Bit : EIndexBufferStride::Force16Bit);
+			}
+
 			// Build the depth-only index buffer.
+			TArray<uint32> DepthOnlyIndices;
 			{
-				TArray<uint32> DepthOnlyIndices;
 				BuildDepthOnlyIndexBuffer(
 					DepthOnlyIndices,
 					Vertices,
@@ -3050,6 +3063,19 @@ public:
 				LODModel.DepthOnlyIndexBuffer.SetIndices(DepthOnlyIndices, bNeeds32BitIndices ? EIndexBufferStride::Force32Bit : EIndexBufferStride::Force16Bit);
 			}
 
+			// Build the inversed depth only index buffer.
+			if (InOutModels[0].BuildSettings.bBuildReversedIndexBuffer)
+			{
+				TArray<uint32> ReversedDepthOnlyIndices;
+				const int32 IndexCount = DepthOnlyIndices.Num();
+				ReversedDepthOnlyIndices.AddUninitialized(IndexCount);
+				for (int32 i = 0; i < IndexCount; ++i)
+				{
+					ReversedDepthOnlyIndices[i] = DepthOnlyIndices[IndexCount - 1 - i];
+				}
+				LODModel.ReversedDepthOnlyIndexBuffer.SetIndices(ReversedDepthOnlyIndices, bNeeds32BitIndices ? EIndexBufferStride::Force32Bit : EIndexBufferStride::Force16Bit);
+			}
+
 			// Build a list of wireframe edges in the static mesh.
 			{
 				TArray<FMeshEdge> Edges;
diff --git a/Engine/Source/Editor/StaticMeshEditor/Private/StaticMeshEditorTools.cpp b/Engine/Source/Editor/StaticMeshEditor/Private/StaticMeshEditorTools.cpp
--- a/Engine/Source/Editor/StaticMeshEditor/Private/StaticMeshEditorTools.cpp
+++ b/Engine/Source/Editor/StaticMeshEditor/Private/StaticMeshEditorTools.cpp
@@ -369,6 +369,22 @@ void FMeshBuildSettingsLayout::GenerateChildContent( IDetailChildrenBuilder& Chi
 		];
 	}
 
+	{
+		ChildrenBuilder.AddChildContent( LOCTEXT("BuildReversedIndexBuffer", "Build Reversed Index Buffer") )
+		.NameContent()
+		[
+			SNew(STextBlock)
+			.Font( IDetailLayoutBuilder::GetDetailFont() )
+			.Text(LOCTEXT("BuildReversedIndexBuffer", "Build Reversed Index Buffer"))
+		]
+		.ValueContent()
+		[
+			SNew(SCheckBox)
+			.IsChecked(this, &FMeshBuildSettingsLayout::ShouldBuildReversedIndexBuffer)
+			.OnCheckStateChanged(this, &FMeshBuildSettingsLayout::OnBuildReversedIndexBufferChanged)
+		];
+	}
+
 	{
 		ChildrenBuilder.AddChildContent( LOCTEXT("UseFullPrecisionUVs", "Use Full Precision UVs") )
 		.NameContent()
@@ -596,6 +612,11 @@ ECheckBoxState FMeshBuildSettingsLayout::ShouldBuildAdjacencyBuffer() const
 	return BuildSettings.bBuildAdjacencyBuffer ? ECheckBoxState::Checked : ECheckBoxState::Unchecked;
 }
 
+ECheckBoxState FMeshBuildSettingsLayout::ShouldBuildReversedIndexBuffer() const
+{
+	return BuildSettings.bBuildReversedIndexBuffer ? ECheckBoxState::Checked : ECheckBoxState::Unchecked;
+}
+
 ECheckBoxState FMeshBuildSettingsLayout::ShouldUseFullPrecisionUVs() const
 {
 	return BuildSettings.bUseFullPrecisionUVs ? ECheckBoxState::Checked : ECheckBoxState::Unchecked;
@@ -707,6 +728,19 @@ void FMeshBuildSettingsLayout::OnBuildAdjacencyBufferChanged(ECheckBoxState NewS
 	}
 }
 
+void FMeshBuildSettingsLayout::OnBuildReversedIndexBufferChanged(ECheckBoxState NewState)
+{
+	const bool bBuildReversedIndexBuffer = (NewState == ECheckBoxState::Checked) ? true : false;
+	if (BuildSettings.bBuildReversedIndexBuffer != bBuildReversedIndexBuffer)
+	{
+		if (FEngineAnalytics::IsAvailable())
+		{
+			FEngineAnalytics::GetProvider().RecordEvent(TEXT("Editor.Usage.StaticMesh.BuildSettings"), TEXT("bBuildReversedIndexBuffer"), bBuildReversedIndexBuffer ? TEXT("True") : TEXT("False"));
+		}
+		BuildSettings.bBuildReversedIndexBuffer = bBuildReversedIndexBuffer;
+	}
+}
+
 void FMeshBuildSettingsLayout::OnUseFullPrecisionUVsChanged(ECheckBoxState NewState)
 {
 	const bool bUseFullPrecisionUVs = (NewState == ECheckBoxState::Checked) ? true : false;
diff --git a/Engine/Source/Editor/StaticMeshEditor/Private/StaticMeshEditorTools.h b/Engine/Source/Editor/StaticMeshEditor/Private/StaticMeshEditorTools.h
--- a/Engine/Source/Editor/StaticMeshEditor/Private/StaticMeshEditorTools.h
+++ b/Engine/Source/Editor/StaticMeshEditor/Private/StaticMeshEditorTools.h
@@ -134,6 +134,7 @@ private:
 	ECheckBoxState ShouldUseMikkTSpace() const;
 	ECheckBoxState ShouldRemoveDegenerates() const;
 	ECheckBoxState ShouldBuildAdjacencyBuffer() const;
+	ECheckBoxState ShouldBuildReversedIndexBuffer() const;
 	ECheckBoxState ShouldUseFullPrecisionUVs() const;
 	ECheckBoxState ShouldGenerateLightmapUVs() const;
 	ECheckBoxState ShouldGenerateDistanceFieldAsIfTwoSided() const;
@@ -150,6 +151,7 @@ private:
 	void OnUseMikkTSpaceChanged(ECheckBoxState NewState);
 	void OnRemoveDegeneratesChanged(ECheckBoxState NewState);
 	void OnBuildAdjacencyBufferChanged(ECheckBoxState NewState);
+	void OnBuildReversedIndexBufferChanged(ECheckBoxState NewState);
 	void OnUseFullPrecisionUVsChanged(ECheckBoxState NewState);
 	void OnGenerateLightmapUVsChanged(ECheckBoxState NewState);
 	void OnGenerateDistanceFieldAsIfTwoSidedChanged(ECheckBoxState NewState);
diff --git a/Engine/Source/Editor/UnrealEd/Classes/Factories/FbxStaticMeshImportData.h b/Engine/Source/Editor/UnrealEd/Classes/Factories/FbxStaticMeshImportData.h
--- a/Engine/Source/Editor/UnrealEd/Classes/Factories/FbxStaticMeshImportData.h
+++ b/Engine/Source/Editor/UnrealEd/Classes/Factories/FbxStaticMeshImportData.h
@@ -45,7 +45,10 @@ class UFbxStaticMeshImportData : public UFbxMeshImportData
 	/** Required for PNT tessellation but can be slow. Recommend disabling for larger meshes. */
 	UPROPERTY(EditAnywhere, AdvancedDisplay, Category = ImportSettings, meta = (ImportType = "StaticMesh"))
 	uint32 bBuildAdjacencyBuffer:1;
-	
+
+	UPROPERTY(EditAnywhere, AdvancedDisplay, Category = ImportSettings, meta = (ImportType = "StaticMesh"))
+	uint32 bBuildReversedIndexBuffer:1;
+
 	UPROPERTY(EditAnywhere, config, AdvancedDisplay, Category=ImportSettings, meta=(ImportType="StaticMesh"))
 	uint32 bGenerateLightmapUVs:1;
 
diff --git a/Engine/Source/Editor/UnrealEd/Private/Fbx/FbxMainImport.cpp b/Engine/Source/Editor/UnrealEd/Private/Fbx/FbxMainImport.cpp
--- a/Engine/Source/Editor/UnrealEd/Private/Fbx/FbxMainImport.cpp
+++ b/Engine/Source/Editor/UnrealEd/Private/Fbx/FbxMainImport.cpp
@@ -232,6 +232,7 @@ void ApplyImportUIToImportOptions(UFbxImportUI* ImportUI, FBXImportOptions& InOu
 	InOutImportOptions.VertexOverrideColor = ImportUI->StaticMeshImportData->VertexOverrideColor;
 	InOutImportOptions.bRemoveDegenerates = ImportUI->StaticMeshImportData->bRemoveDegenerates;
 	InOutImportOptions.bBuildAdjacencyBuffer = ImportUI->StaticMeshImportData->bBuildAdjacencyBuffer;
+	InOutImportOptions.bBuildReversedIndexBuffer = ImportUI->StaticMeshImportData->bBuildReversedIndexBuffer;
 	InOutImportOptions.bGenerateLightmapUVs = ImportUI->StaticMeshImportData->bGenerateLightmapUVs;
 	InOutImportOptions.bOneConvexHullPerUCX = ImportUI->StaticMeshImportData->bOneConvexHullPerUCX;
 	InOutImportOptions.bAutoGenerateCollision = ImportUI->StaticMeshImportData->bAutoGenerateCollision;
diff --git a/Engine/Source/Editor/UnrealEd/Private/Fbx/FbxStaticMeshImport.cpp b/Engine/Source/Editor/UnrealEd/Private/Fbx/FbxStaticMeshImport.cpp
--- a/Engine/Source/Editor/UnrealEd/Private/Fbx/FbxStaticMeshImport.cpp
+++ b/Engine/Source/Editor/UnrealEd/Private/Fbx/FbxStaticMeshImport.cpp
@@ -1211,6 +1211,7 @@ UStaticMesh* UnFbx::FFbxImporter::ImportStaticMeshAsSingle(UObject* InParent, TA
 		// @todo This overrides restored values currently but we need to be able to import over the existing settings if the user chose to do so.
 		SrcModel.BuildSettings.bRemoveDegenerates = ImportOptions->bRemoveDegenerates;
 		SrcModel.BuildSettings.bBuildAdjacencyBuffer = ImportOptions->bBuildAdjacencyBuffer;
+		SrcModel.BuildSettings.bBuildReversedIndexBuffer = ImportOptions->bBuildReversedIndexBuffer;
 		SrcModel.BuildSettings.bRecomputeNormals = ImportOptions->NormalImportMethod == FBXNIM_ComputeNormals;
 		SrcModel.BuildSettings.bRecomputeTangents = ImportOptions->NormalImportMethod != FBXNIM_ImportNormalsAndTangents;
 		SrcModel.BuildSettings.bUseMikkTSpace = ImportOptions->NormalGenerationMethod == EFBXNormalGenerationMethod::MikkTSpace;
diff --git a/Engine/Source/Editor/UnrealEd/Private/Fbx/FbxStaticMeshImportData.cpp b/Engine/Source/Editor/UnrealEd/Private/Fbx/FbxStaticMeshImportData.cpp
--- a/Engine/Source/Editor/UnrealEd/Private/Fbx/FbxStaticMeshImportData.cpp
+++ b/Engine/Source/Editor/UnrealEd/Private/Fbx/FbxStaticMeshImportData.cpp
@@ -9,6 +9,7 @@ UFbxStaticMeshImportData::UFbxStaticMeshImportData(const FObjectInitializer& Obj
 	StaticMeshLODGroup = NAME_None;
 	bRemoveDegenerates = true;
 	bBuildAdjacencyBuffer = true;
+	bBuildReversedIndexBuffer = true;
 	bGenerateLightmapUVs = true;
 	bOneConvexHullPerUCX = true;
 	bAutoGenerateCollision = true;
diff --git a/Engine/Source/Editor/UnrealEd/Private/FbxImporter.h b/Engine/Source/Editor/UnrealEd/Private/FbxImporter.h
--- a/Engine/Source/Editor/UnrealEd/Private/FbxImporter.h
+++ b/Engine/Source/Editor/UnrealEd/Private/FbxImporter.h
@@ -125,6 +125,7 @@ struct FBXImportOptions
 	FColor VertexOverrideColor;
 	bool bRemoveDegenerates;
 	bool bBuildAdjacencyBuffer;
+	bool bBuildReversedIndexBuffer;
 	bool bGenerateLightmapUVs;
 	bool bOneConvexHullPerUCX;
 	bool bAutoGenerateCollision;
diff --git a/Engine/Source/Runtime/Engine/Classes/Engine/EngineTypes.h b/Engine/Source/Runtime/Engine/Classes/Engine/EngineTypes.h
--- a/Engine/Source/Runtime/Engine/Classes/Engine/EngineTypes.h
+++ b/Engine/Source/Runtime/Engine/Classes/Engine/EngineTypes.h
@@ -2168,6 +2168,10 @@ struct FMeshBuildSettings
 	UPROPERTY(EditAnywhere, Category=BuildSettings)
 	bool bBuildAdjacencyBuffer;
 
+	/** Required to optimize mesh in mirrored transform. Double index buffer size. */
+	UPROPERTY(EditAnywhere, Category=BuildSettings)
+	bool bBuildReversedIndexBuffer;
+
 	/** If true, UVs will be stored at full floating point precision. */
 	UPROPERTY(EditAnywhere, Category=BuildSettings)
 	bool bUseFullPrecisionUVs;
@@ -2215,6 +2219,7 @@ struct FMeshBuildSettings
 		, bRecomputeTangents(true)
 		, bRemoveDegenerates(true)
 		, bBuildAdjacencyBuffer(true)
+		, bBuildReversedIndexBuffer(true)
 		, bUseFullPrecisionUVs(false)
 		, bGenerateLightmapUVs(true)
 		, MinLightmapResolution(64)
@@ -2235,6 +2240,7 @@ struct FMeshBuildSettings
 			&& bUseMikkTSpace == Other.bUseMikkTSpace
 			&& bRemoveDegenerates == Other.bRemoveDegenerates
 			&& bBuildAdjacencyBuffer == Other.bBuildAdjacencyBuffer
+			&& bBuildReversedIndexBuffer == Other.bBuildReversedIndexBuffer
 			&& bUseFullPrecisionUVs == Other.bUseFullPrecisionUVs
 			&& bGenerateLightmapUVs == Other.bGenerateLightmapUVs
 			&& MinLightmapResolution == Other.MinLightmapResolution
diff --git a/Engine/Source/Runtime/Engine/Private/StaticMesh.cpp b/Engine/Source/Runtime/Engine/Private/StaticMesh.cpp
--- a/Engine/Source/Runtime/Engine/Private/StaticMesh.cpp
+++ b/Engine/Source/Runtime/Engine/Private/StaticMesh.cpp
@@ -292,6 +292,8 @@ void FStaticMeshLODResources::Serialize(FArchive& Ar, UObject* Owner, int32 Inde
 	bool bNeedsCPUAccess = !FPlatformProperties::RequiresCookedData();
 
 	bHasAdjacencyInfo = false;
+	bHasReversedIndexBuffer = false;
+	bHasReversedDepthOnlyIndexBuffer = false;
 
     // Defined class flags for possible stripping
 	const uint8 AdjacencyDataStripFlag = 1;
@@ -311,7 +313,9 @@ void FStaticMeshLODResources::Serialize(FArchive& Ar, UObject* Owner, int32 Inde
 		VertexBuffer.Serialize( Ar, bNeedsCPUAccess );
 		ColorVertexBuffer.Serialize( Ar, bNeedsCPUAccess );
 		IndexBuffer.Serialize( Ar, bNeedsCPUAccess );
+		ReversedIndexBuffer.Serialize( Ar, bNeedsCPUAccess );
 		DepthOnlyIndexBuffer.Serialize(Ar, bNeedsCPUAccess);
+		ReversedDepthOnlyIndexBuffer.Serialize( Ar, bNeedsCPUAccess );
 
 		if( !StripFlags.IsEditorDataStripped() )
 		{
@@ -323,6 +327,10 @@ void FStaticMeshLODResources::Serialize(FArchive& Ar, UObject* Owner, int32 Inde
 			AdjacencyIndexBuffer.Serialize( Ar, bNeedsCPUAccess );
 			bHasAdjacencyInfo = AdjacencyIndexBuffer.GetNumIndices() != 0;
 		}
+
+		// Needs to be done now because on cooked platform, indices are discarded after RHIInit.
+		bHasReversedIndexBuffer = ReversedIndexBuffer.GetNumIndices() != 0;
+		bHasReversedDepthOnlyIndexBuffer = ReversedDepthOnlyIndexBuffer.GetNumIndices() != 0;
 	}
 }
 
@@ -486,6 +494,8 @@ FStaticMeshLODResources::FStaticMeshLODResources()
 	: DistanceFieldData(NULL)
 	, MaxDeviation(0.0f)
 	, bHasAdjacencyInfo(false)
+	, bHasReversedIndexBuffer(false)
+	, bHasReversedDepthOnlyIndexBuffer(false)
 {
 }
 
@@ -522,11 +532,21 @@ void FStaticMeshLODResources::InitResources(UStaticMesh* Parent)
 		BeginInitResource(&ColorVertexBuffer);
 	}
 
+	if (ReversedIndexBuffer.GetNumIndices() > 0)
+	{
+		BeginInitResource(&ReversedIndexBuffer);
+	}
+
 	if (DepthOnlyIndexBuffer.GetNumIndices() > 0)
 	{
 		BeginInitResource(&DepthOnlyIndexBuffer);
 	}
 
+	if (ReversedDepthOnlyIndexBuffer.GetNumIndices() > 0)
+	{
+		BeginInitResource(&ReversedDepthOnlyIndexBuffer);
+	}
+
 	if (RHISupportsTessellation(MaxShaderPlatform))
 	{
 		BeginInitResource(&AdjacencyIndexBuffer);
@@ -585,7 +605,9 @@ void FStaticMeshLODResources::ReleaseResources()
 	BeginReleaseResource(&VertexBuffer);
 	BeginReleaseResource(&PositionVertexBuffer);
 	BeginReleaseResource(&ColorVertexBuffer);
+	BeginReleaseResource(&ReversedIndexBuffer);
 	BeginReleaseResource(&DepthOnlyIndexBuffer);
+	BeginReleaseResource(&ReversedDepthOnlyIndexBuffer);
 
 	// Release the vertex factories.
 	BeginReleaseResource(&VertexFactory);
@@ -1026,6 +1048,7 @@ FArchive& operator<<(FArchive& Ar, FMeshBuildSettings& BuildSettings)
 	Ar << BuildSettings.bUseMikkTSpace;
 	Ar << BuildSettings.bRemoveDegenerates;
 	Ar << BuildSettings.bBuildAdjacencyBuffer;
+	Ar << BuildSettings.bBuildReversedIndexBuffer;
 	Ar << BuildSettings.bUseFullPrecisionUVs;
 	Ar << BuildSettings.bGenerateLightmapUVs;
 
@@ -1060,8 +1083,7 @@ FArchive& operator<<(FArchive& Ar, FMeshBuildSettings& BuildSettings)
 // differences, etc.) replace the version GUID below with a new one.
 // In case of merge conflicts with DDC versions, you MUST generate a new GUID
 // and set this new GUID as the version.
-#define STATICMESH_DERIVEDDATA_VER TEXT("46A8778361B442A9523C54440EA1E9D")
-
+#define STATICMESH_DERIVEDDATA_VER TEXT("0651E844A9544BD83BF12B1807773C2")
 static const FString& GetStaticMeshDerivedDataVersion()
 {
 	static FString CachedVersionString;
diff --git a/Engine/Source/Runtime/Engine/Private/StaticMeshRender.cpp b/Engine/Source/Runtime/Engine/Private/StaticMeshRender.cpp
--- a/Engine/Source/Runtime/Engine/Private/StaticMeshRender.cpp
+++ b/Engine/Source/Runtime/Engine/Private/StaticMeshRender.cpp
@@ -20,6 +20,9 @@
 /** If true, optimized depth-only index buffers are used for shadow rendering. */
 static bool GUseShadowIndexBuffer = true;
 
+/** If true, reversed index buffer are used for mesh with negative transform determinants. */
+static bool GUseReversedIndexBuffer = true;
+
 static void ToggleShadowIndexBuffers()
 {
 	FlushRenderingCommands();
@@ -28,19 +31,33 @@ static void ToggleShadowIndexBuffers()
 	FGlobalComponentReregisterContext ReregisterContext;
 }
 
+static void ToggleReversedIndexBuffers()
+{
+	FlushRenderingCommands();
+	GUseReversedIndexBuffer = !GUseReversedIndexBuffer;
+	UE_LOG(LogStaticMesh,Log,TEXT("Reversed index buffers %s"),GUseReversedIndexBuffer ? TEXT("ENABLED") : TEXT("DISABLED"));
+	FGlobalComponentReregisterContext ReregisterContext;
+}
+
 static FAutoConsoleCommand GToggleShadowIndexBuffersCmd(
 	TEXT("ToggleShadowIndexBuffers"),
 	TEXT("Render static meshes with an optimized shadow index buffer that minimizes unique vertices."),
 	FConsoleCommandDelegate::CreateStatic(ToggleShadowIndexBuffers)
 	);
 
+static FAutoConsoleCommand GToggleReversedIndexBuffersCmd(
+	TEXT("ToggleReversedIndexBuffers"),
+	TEXT("Render static meshes with negative transform determinants using a reversed index buffer."),
+	FConsoleCommandDelegate::CreateStatic(ToggleReversedIndexBuffers)
+	);
+
 bool GForceDefaultMaterial = false;
 
 static void ToggleForceDefaultMaterial()
 {
 	FlushRenderingCommands();
 	GForceDefaultMaterial = !GForceDefaultMaterial;
-	UE_LOG(LogStaticMesh,Log,TEXT("Force default material %s"),GUseShadowIndexBuffer ? TEXT("ENABLED") : TEXT("DISABLED"));
+	UE_LOG(LogStaticMesh,Log,TEXT("Force default material %s"),GForceDefaultMaterial ? TEXT("ENABLED") : TEXT("DISABLED"));
 	FGlobalComponentReregisterContext ReregisterContext;
 }
 
@@ -179,10 +196,12 @@ bool FStaticMeshSceneProxy::GetShadowMeshElement(int32 LODIndex, int32 BatchInde
 	const FStaticMeshLODResources& LOD = RenderData->LODResources[LODIndex];
 	const FLODInfo& ProxyLODInfo = LODs[LODIndex];
 
+	const bool bUseReversedIndices = GUseReversedIndexBuffer && IsLocalToWorldDeterminantNegative() && LOD.bHasReversedDepthOnlyIndexBuffer;
+
 	FMeshBatchElement& OutBatchElement = OutMeshBatch.Elements[0];
 	OutMeshBatch.MaterialRenderProxy = UMaterial::GetDefaultMaterial(MD_Surface)->GetRenderProxy(false, false);
 	OutMeshBatch.VertexFactory = &LOD.VertexFactory;
-	OutBatchElement.IndexBuffer = &LOD.DepthOnlyIndexBuffer;
+	OutBatchElement.IndexBuffer = bUseReversedIndices ? &LOD.ReversedDepthOnlyIndexBuffer : &LOD.DepthOnlyIndexBuffer;
 	OutMeshBatch.Type = PT_TriangleList;
 	OutBatchElement.FirstIndex = 0;
 	OutBatchElement.NumPrimitives = LOD.DepthOnlyIndexBuffer.GetNumIndices() / 3;
@@ -190,8 +209,7 @@ bool FStaticMeshSceneProxy::GetShadowMeshElement(int32 LODIndex, int32 BatchInde
 	OutBatchElement.MinVertexIndex = 0;
 	OutBatchElement.MaxVertexIndex = LOD.PositionVertexBuffer.GetNumVertices() - 1;
 	OutMeshBatch.DepthPriorityGroup = InDepthPriorityGroup;
-	OutMeshBatch.ReverseCulling = IsLocalToWorldDeterminantNegative();
-	OutMeshBatch.CastShadow = true;
+	OutMeshBatch.ReverseCulling = IsLocalToWorldDeterminantNegative() && !bUseReversedIndices;
 	OutMeshBatch.LODIndex = LODIndex;
 #if !(UE_BUILD_SHIPPING || UE_BUILD_TEST)
 	OutMeshBatch.VisualizeLODIndex = LODIndex;
@@ -213,6 +231,10 @@ bool FStaticMeshSceneProxy::GetShadowMeshElement(int32 LODIndex, int32 BatchInde
 		}
 	}
 
+	// By default this will be a shadow only mesh.
+	OutMeshBatch.bUseAsOccluder = false;
+	OutMeshBatch.bUseForMaterial = false;
+
 	return true;
 }
 
@@ -251,8 +273,9 @@ bool FStaticMeshSceneProxy::GetMeshElement(int32 LODIndex, int32 BatchIndex, int
 
 	const bool bWireframe = false;
 	const bool bRequiresAdjacencyInformation = RequiresAdjacencyInformation( Material, OutMeshBatch.VertexFactory->GetType(), GetScene().GetFeatureLevel() );
+	const bool bUseReversedIndices = GUseReversedIndexBuffer && !bWireframe && !bRequiresAdjacencyInformation && IsLocalToWorldDeterminantNegative() && LOD.bHasReversedIndexBuffer;
 
-	SetIndexSource(LODIndex, SectionIndex, OutMeshBatch, bWireframe, bRequiresAdjacencyInformation );
+	SetIndexSource(LODIndex, SectionIndex, OutMeshBatch, bWireframe, bRequiresAdjacencyInformation, bUseReversedIndices);
 
 	FMeshBatchElement& OutBatchElement = OutMeshBatch.Elements[0];
 
@@ -268,7 +291,7 @@ bool FStaticMeshSceneProxy::GetMeshElement(int32 LODIndex, int32 BatchIndex, int
 		OutMeshBatch.VisualizeLODIndex = LODIndex;
 #endif
 		OutMeshBatch.UseDynamicData = false;
-		OutMeshBatch.ReverseCulling = IsLocalToWorldDeterminantNegative();
+		OutMeshBatch.ReverseCulling = IsLocalToWorldDeterminantNegative() && !bUseReversedIndices;
 		OutMeshBatch.CastShadow = bCastShadow && Section.bCastShadow;
 		OutMeshBatch.DepthPriorityGroup = (ESceneDepthPriorityGroup)InDepthPriorityGroup;
 		if (ForcedLodModel > 0) 
@@ -329,8 +352,9 @@ bool FStaticMeshSceneProxy::GetWireframeMeshElement(int32 LODIndex, int32 BatchI
 
 	const bool bWireframe = true;
 	const bool bRequiresAdjacencyInformation = false;
+	const bool bUseReversedIndices = false;
 
-	SetIndexSource(LODIndex, 0, OutMeshBatch, bWireframe, bRequiresAdjacencyInformation);
+	SetIndexSource(LODIndex, 0, OutMeshBatch, bWireframe, bRequiresAdjacencyInformation, bUseReversedIndices);
 
 	return OutBatchElement.NumPrimitives > 0;
 }
@@ -338,7 +362,7 @@ bool FStaticMeshSceneProxy::GetWireframeMeshElement(int32 LODIndex, int32 BatchI
 /**
  * Sets IndexBuffer, FirstIndex and NumPrimitives of OutMeshElement.
  */
-void FStaticMeshSceneProxy::SetIndexSource(int32 LODIndex, int32 SectionIndex, FMeshBatch& OutMeshElement, bool bWireframe, bool bRequiresAdjacencyInformation) const
+void FStaticMeshSceneProxy::SetIndexSource(int32 LODIndex, int32 SectionIndex, FMeshBatch& OutMeshElement, bool bWireframe, bool bRequiresAdjacencyInformation, bool bUseReversedIndices) const
 {
 	FMeshBatchElement& OutElement = OutMeshElement.Elements[0];
 	const FStaticMeshLODResources& LODModel = RenderData->LODResources[LODIndex];
@@ -367,7 +391,7 @@ void FStaticMeshSceneProxy::SetIndexSource(int32 LODIndex, int32 SectionIndex, F
 	{
 		const FStaticMeshSection& Section = LODModel.Sections[SectionIndex];
 		OutMeshElement.Type = PT_TriangleList;
-		OutElement.IndexBuffer = &LODModel.IndexBuffer;
+		OutElement.IndexBuffer = bUseReversedIndices ? &LODModel.ReversedIndexBuffer : &LODModel.IndexBuffer;
 		OutElement.FirstIndex = Section.FirstIndex;
 		OutElement.NumPrimitives = Section.NumTriangles;
 	}
@@ -501,23 +525,23 @@ void FStaticMeshSceneProxy::DrawStaticElements(FStaticPrimitiveDrawInterface* PD
 				const FStaticMeshLODResources& LODModel = RenderData->LODResources[LODIndex];
 				float ScreenSize = GetScreenSize(LODIndex);
 
-				bool bHaveShadowOnlyMesh = false;
-				if (GUseShadowIndexBuffer
-					&& bCastShadow
-					&& LODModel.DepthOnlyIndexBuffer.GetNumIndices() > 0)
+				bool bUseUnifiedMeshForShadow = false;
+				bool bUseUnifiedMeshForDepth = false;
+
+				if (GUseShadowIndexBuffer && LODModel.DepthOnlyIndexBuffer.GetNumIndices() > 0)
 				{
 					const FLODInfo& ProxyLODInfo = LODs[LODIndex];
 
 					// The shadow-only mesh can be used only if all elements cast shadows and use opaque materials with no vertex modification.
 					// In some cases (e.g. LPV) we don't want the optimization
-					bool bSafeToUseShadowOnlyMesh = AllowShadowOnlyMesh(FeatureLevel);
+					bool bSafeToUseUnifiedMesh = AllowShadowOnlyMesh(FeatureLevel);
 
-					bool bAnySectionCastsShadow = false;
 					bool bAnySectionUsesDitheredLODTransition = false;
 					bool bAllSectionsUseDitheredLODTransition = true;
 					bool bIsMovable = IsMovable();
+					bool bAllSectionsCastShadow = bCastShadow;
 
-					for (int32 SectionIndex = 0; bSafeToUseShadowOnlyMesh && SectionIndex < LODModel.Sections.Num(); SectionIndex++)
+					for (int32 SectionIndex = 0; bSafeToUseUnifiedMesh && SectionIndex < LODModel.Sections.Num(); SectionIndex++)
 					{
 						const FMaterial* Material = ProxyLODInfo.Sections[SectionIndex].Material->GetRenderProxy(false)->GetMaterial(FeatureLevel);
 						// no support for stateless dithered LOD transitions for movable meshes
@@ -525,28 +549,42 @@ void FStaticMeshSceneProxy::DrawStaticElements(FStaticPrimitiveDrawInterface* PD
 						bAllSectionsUseDitheredLODTransition = bAllSectionsUseDitheredLODTransition && (!bIsMovable && Material->IsDitheredLODTransition());
 						const FStaticMeshSection& Section = LODModel.Sections[SectionIndex];
 
-						bSafeToUseShadowOnlyMesh =
-							Section.bCastShadow
-							&& !(bAnySectionUsesDitheredLODTransition && !bAllSectionsUseDitheredLODTransition) // can't use a single section if they are not homogeneous
+						bSafeToUseUnifiedMesh =
+							!(bAnySectionUsesDitheredLODTransition && !bAllSectionsUseDitheredLODTransition) // can't use a single section if they are not homogeneous
 							&& Material->WritesEveryPixel()
 							&& !Material->IsTwoSided()
 							&& !IsTranslucentBlendMode(Material->GetBlendMode())
 							&& !Material->MaterialModifiesMeshPosition_RenderThread();
-						bAnySectionCastsShadow |= Section.bCastShadow;
+
+						bAllSectionsCastShadow &= Section.bCastShadow;
 					}
 
-					if (bAnySectionCastsShadow && bSafeToUseShadowOnlyMesh)
+					if (bSafeToUseUnifiedMesh)
 					{
-						const int32 NumBatches = GetNumMeshBatches();
+						bUseUnifiedMeshForShadow = bAllSectionsCastShadow;
+
+						// Depth pass is only used for deferred renderer. The other conditions are meant to match the logic in FStaticMesh::AddToDrawLists.
+						// Could not link to "GEarlyZPassMovable" so moveable are ignored.
+						bUseUnifiedMeshForDepth = ShouldUseAsOccluder() && GetScene().ShouldUseDeferredRenderer() && !IsMovable();
 
-						for (int32 BatchIndex = 0; BatchIndex < NumBatches; BatchIndex++)
+						if (bUseUnifiedMeshForShadow || bUseUnifiedMeshForDepth)
 						{
-							FMeshBatch MeshBatch;
+							const int32 NumBatches = GetNumMeshBatches();
 
-							if (GetShadowMeshElement(LODIndex, BatchIndex, PrimitiveDPG, MeshBatch, bAllSectionsUseDitheredLODTransition))
+							for (int32 BatchIndex = 0; BatchIndex < NumBatches; BatchIndex++)
 							{
-								bHaveShadowOnlyMesh = true;
-								PDI->DrawMesh(MeshBatch, ScreenSize, /*bShadowOnly=*/true);
+								FMeshBatch MeshBatch;
+
+								if (GetShadowMeshElement(LODIndex, BatchIndex, PrimitiveDPG, MeshBatch, bAllSectionsUseDitheredLODTransition))
+								{
+									bUseUnifiedMeshForShadow = bAllSectionsCastShadow;
+
+									MeshBatch.CastShadow = bUseUnifiedMeshForShadow;
+									MeshBatch.bUseAsOccluder = bUseUnifiedMeshForDepth;
+									MeshBatch.bUseForMaterial = false;
+
+									PDI->DrawMesh(MeshBatch, ScreenSize);
+								}
 							}
 						}
 					}
@@ -574,7 +612,9 @@ void FStaticMeshSceneProxy::DrawStaticElements(FStaticPrimitiveDrawInterface* PD
 						if (GetMeshElement(LODIndex, BatchIndex, SectionIndex, PrimitiveDPG, bUseSelectedMaterial, bUseHoveredMaterial, MeshBatch))
 						{
 							// If we have submitted an optimized shadow-only mesh, remaining mesh elements must not cast shadows.
-							MeshBatch.CastShadow = MeshBatch.CastShadow && !bHaveShadowOnlyMesh;
+							MeshBatch.CastShadow &= !bUseUnifiedMeshForShadow;
+							MeshBatch.bUseAsOccluder &= !bUseUnifiedMeshForDepth;
+
 							PDI->DrawMesh(MeshBatch, ScreenSize);
 						}
 					}
diff --git a/Engine/Source/Runtime/Engine/Public/MeshBatch.h b/Engine/Source/Runtime/Engine/Public/MeshBatch.h
--- a/Engine/Source/Runtime/Engine/Public/MeshBatch.h
+++ b/Engine/Source/Runtime/Engine/Public/MeshBatch.h
@@ -77,13 +77,14 @@ struct FMeshBatch
 	uint32 UseDynamicData : 1;
 	uint32 ReverseCulling : 1;
 	uint32 bDisableBackfaceCulling : 1;
-	uint32 CastShadow : 1;
+	uint32 CastShadow : 1;				// Wheter it can be used in shadow renderpasses.
+	uint32 bUseForMaterial : 1;	// Whether it can be used in renderpasses requiring material outputs.
+	uint32 bUseAsOccluder : 1;			// Whether it can be used in renderpasses only depending on the raw geometry (i.e. Depth Prepass).
 	uint32 bWireframe : 1;
 	// e.g. PT_TriangleList(default), PT_LineList, ..
 	uint32 Type : PT_NumBits;
 	// e.g. SDPG_World (default), SDPG_Foreground
 	uint32 DepthPriorityGroup : SDPG_NumBits;
-	uint32 bUseAsOccluder : 1;
 
 	/** Whether view mode overrides can be applied to this mesh eg unlit, wireframe. */
 	uint32 bCanApplyViewModeOverrides : 1;
@@ -202,10 +203,11 @@ struct FMeshBatch
 	,	ReverseCulling(false)
 	,	bDisableBackfaceCulling(false)
 	,	CastShadow(true)
+	,   bUseForMaterial(true)
+	,	bUseAsOccluder(true)
 	,	bWireframe(false)
 	,	Type(PT_TriangleList)
 	,	DepthPriorityGroup(SDPG_World)
-	,	bUseAsOccluder(true)
 	,	bCanApplyViewModeOverrides(false)
 	,	bUseWireframeSelectionColoring(false)
 	,	bUseSelectionOutline(true)
diff --git a/Engine/Source/Runtime/Engine/Public/SceneManagement.h b/Engine/Source/Runtime/Engine/Public/SceneManagement.h
--- a/Engine/Source/Runtime/Engine/Public/SceneManagement.h
+++ b/Engine/Source/Runtime/Engine/Public/SceneManagement.h
@@ -1229,8 +1229,7 @@ public:
 	virtual void SetHitProxy(HHitProxy* HitProxy) = 0;
 	virtual void DrawMesh(
 		const FMeshBatch& Mesh,
-		float ScreenSize,
-		bool bShadowOnly = false
+		float ScreenSize
 		) = 0;
 };
 
diff --git a/Engine/Source/Runtime/Engine/Public/StaticMeshResources.h b/Engine/Source/Runtime/Engine/Public/StaticMeshResources.h
--- a/Engine/Source/Runtime/Engine/Public/StaticMeshResources.h
+++ b/Engine/Source/Runtime/Engine/Public/StaticMeshResources.h
@@ -539,8 +539,12 @@ struct FStaticMeshLODResources
 
 	/** Index buffer resource for rendering. */
 	FRawStaticIndexBuffer IndexBuffer;
+	/** Reversed index buffer, used to prevent changing culling state between drawcalls. */
+	FRawStaticIndexBuffer ReversedIndexBuffer;
 	/** Index buffer resource for rendering in depth only passes. */
 	FRawStaticIndexBuffer DepthOnlyIndexBuffer;
+	/** Reversed depth only index buffer, used to prevent changing culling state between drawcalls. */
+	FRawStaticIndexBuffer ReversedDepthOnlyIndexBuffer;
 	/** Index buffer resource for rendering wireframe mode. */
 	FRawStaticIndexBuffer WireframeIndexBuffer;
 	/** Index buffer containing adjacency information required by tessellation. */
@@ -559,7 +563,13 @@ struct FStaticMeshLODResources
 	float MaxDeviation;
 
 	/** True if the adjacency index buffer contained data at init. Needed as it will not be available to the CPU afterwards. */
-	bool bHasAdjacencyInfo;
+	uint32 bHasAdjacencyInfo : 1;
+
+	/** True if the reversed index buffers contained data at init. Needed as it will not be available to the CPU afterwards. */
+	uint32 bHasReversedIndexBuffer : 1;
+
+	/** True if the reversed index buffers contained data at init. Needed as it will not be available to the CPU afterwards. */
+	uint32 bHasReversedDepthOnlyIndexBuffer : 1;
 
 	/** Default constructor. */
 	FStaticMeshLODResources();
@@ -761,7 +771,7 @@ protected:
 	/**
 	 * Sets IndexBuffer, FirstIndex and NumPrimitives of OutMeshElement.
 	 */
-	virtual void SetIndexSource(int32 LODIndex, int32 ElementIndex, FMeshBatch& OutMeshElement, bool bWireframe, bool bRequiresAdjacencyInformation ) const;
+	virtual void SetIndexSource(int32 LODIndex, int32 ElementIndex, FMeshBatch& OutMeshElement, bool bWireframe, bool bRequiresAdjacencyInformation, bool bUseInversedIndices) const;
 	bool IsCollisionView(const FEngineShowFlags& EngineShowFlags, bool& bDrawSimpleCollision, bool& bDrawComplexCollision) const;
 
 public:
diff --git a/Engine/Source/Runtime/Renderer/Private/ConvertToUniformMesh.cpp b/Engine/Source/Runtime/Renderer/Private/ConvertToUniformMesh.cpp
--- a/Engine/Source/Runtime/Renderer/Private/ConvertToUniformMesh.cpp
+++ b/Engine/Source/Runtime/Renderer/Private/ConvertToUniformMesh.cpp
@@ -362,7 +362,7 @@ int32 FUniformMeshConverter::Convert(
 			const FStaticMesh& StaticMesh = PrimitiveSceneInfo->StaticMeshes[StaticMeshIdx];
 
 			if (ShouldGenerateSurfelsOnMesh(StaticMesh, FeatureLevel)
-				&& !StaticMesh.bShadowOnly
+				&& StaticMesh.bUseForMaterial
 				// LOD0 only
 				&& StaticMesh.ScreenSize == LargestScreenSize)
 			{
@@ -437,7 +437,7 @@ int32 FUniformMeshConverter::Convert(
 				const FStaticMesh& StaticMesh = PrimitiveSceneInfo->StaticMeshes[StaticMeshIdx];
 
 				if (ShouldGenerateSurfelsOnMesh(StaticMesh, FeatureLevel) 
-					&& !StaticMesh.bShadowOnly
+					&& StaticMesh.bUseForMaterial
 					&& StaticMesh.ScreenSize == LargestScreenSize)
 				{
 					FConvertToUniformMeshDrawingPolicy DrawingPolicy(
diff --git a/Engine/Source/Runtime/Renderer/Private/PrimitiveSceneInfo.cpp b/Engine/Source/Runtime/Renderer/Private/PrimitiveSceneInfo.cpp
--- a/Engine/Source/Runtime/Renderer/Private/PrimitiveSceneInfo.cpp
+++ b/Engine/Source/Runtime/Renderer/Private/PrimitiveSceneInfo.cpp
@@ -36,11 +36,8 @@ public:
 			}
 		}
 	}
-	virtual void DrawMesh(
-		const FMeshBatch& Mesh,
-		float ScreenSize,
-		bool bShadowOnly
-		)
+
+	virtual void DrawMesh(const FMeshBatch& Mesh, float ScreenSize)
 	{
 		if (Mesh.GetNumPrimitives() > 0)
 		{
@@ -53,7 +50,6 @@ public:
 				PrimitiveSceneInfo,
 				Mesh,
 				ScreenSize,
-				bShadowOnly,
 				CurrentHitProxy ? CurrentHitProxy->Id : FHitProxyId()
 				);
 		}
diff --git a/Engine/Source/Runtime/Renderer/Private/SceneCore.cpp b/Engine/Source/Runtime/Renderer/Private/SceneCore.cpp
--- a/Engine/Source/Runtime/Renderer/Private/SceneCore.cpp
+++ b/Engine/Source/Runtime/Renderer/Private/SceneCore.cpp
@@ -291,18 +291,15 @@ void FStaticMesh::AddToDrawLists(FRHICommandListImmediate& RHICmdList, FScene* S
 		FShadowDepthDrawingPolicyFactory::AddStaticMesh(Scene, this);
 	}
 
-	if (!bShadowOnly && PrimitiveSceneInfo->Proxy->ShouldRenderInMainPass())
+	if (!PrimitiveSceneInfo->Proxy->ShouldRenderInMainPass())
 	{
-		const bool bRequiresHitProxies = Scene->RequiresHitProxies();
-		if (bRequiresHitProxies && PrimitiveSceneInfo->Proxy->IsSelectable())
-		{
-			// Add the static mesh to the DPG's hit proxy draw list.
-			FHitProxyDrawingPolicyFactory::AddStaticMesh(Scene, this);
-		}
+		return;
 	}
-	else
+
+	if (bUseForMaterial && Scene->RequiresHitProxies() && PrimitiveSceneInfo->Proxy->IsSelectable())
 	{
-		return;
+		// Add the static mesh to the DPG's hit proxy draw list.
+		FHitProxyDrawingPolicyFactory::AddStaticMesh(Scene, this);
 	}
 
 	if (IsTranslucent(FeatureLevel))
@@ -312,28 +309,37 @@ void FStaticMesh::AddToDrawLists(FRHICommandListImmediate& RHICmdList, FScene* S
 
 	if (Scene->ShouldUseDeferredRenderer())
 	{
-		extern TAutoConsoleVariable<int32> CVarEarlyZPass;
-		int32 EarlyZPass = CVarEarlyZPass.GetValueOnRenderThread();
-
-		extern int32 GEarlyZPassMovable;
-
-		// Render non-masked materials in the depth only pass
-		if (PrimitiveSceneInfo->Proxy->ShouldUseAsOccluder() 
-			&& (!IsMasked(FeatureLevel) || EarlyZPass == 2)
-			&& (!PrimitiveSceneInfo->Proxy->IsMovable() || GEarlyZPassMovable))
+		if (bUseAsOccluder)
 		{
-			FDepthDrawingPolicyFactory::AddStaticMesh(Scene,this);
+			// Render non-masked materials in the depth only pass
+			extern TAutoConsoleVariable<int32> CVarEarlyZPass;
+			int32 EarlyZPass = CVarEarlyZPass.GetValueOnRenderThread();
+
+			extern int32 GEarlyZPassMovable;
+
+			// WARNING : If you change this condition, also change the logic in FStaticMeshSceneProxy::DrawStaticElements.
+			if (PrimitiveSceneInfo->Proxy->ShouldUseAsOccluder() 
+				&& (!IsMasked(FeatureLevel) || EarlyZPass == 2)
+				&& (!PrimitiveSceneInfo->Proxy->IsMovable() || GEarlyZPassMovable))
+			{
+				FDepthDrawingPolicyFactory::AddStaticMesh(Scene,this);
+			}
 		}
 
-		// Add the static mesh to the DPG's base pass draw list.
-		FBasePassOpaqueDrawingPolicyFactory::AddStaticMesh(RHICmdList, Scene, this);
-
-		FVelocityDrawingPolicyFactory::AddStaticMesh(Scene, this);
+		if (bUseForMaterial)
+		{
+			// Add the static mesh to the DPG's base pass draw list.
+			FBasePassOpaqueDrawingPolicyFactory::AddStaticMesh(RHICmdList, Scene, this);
+			FVelocityDrawingPolicyFactory::AddStaticMesh(Scene, this);
+		}
 	}
 	else
 	{
-		// Add the static mesh to the DPG's base pass draw list.
-		FBasePassForwardOpaqueDrawingPolicyFactory::AddStaticMesh(RHICmdList, Scene, this);
+		if (bUseForMaterial)
+		{
+			// Add the static mesh to the DPG's base pass draw list.
+			FBasePassForwardOpaqueDrawingPolicyFactory::AddStaticMesh(RHICmdList, Scene, this);
+		}
 	}
 }
 
diff --git a/Engine/Source/Runtime/Renderer/Private/SceneCore.h b/Engine/Source/Runtime/Renderer/Private/SceneCore.h
--- a/Engine/Source/Runtime/Renderer/Private/SceneCore.h
+++ b/Engine/Source/Runtime/Renderer/Private/SceneCore.h
@@ -140,22 +140,17 @@ public:
 	/** The index of the mesh in the scene's static meshes array. */
 	int32 Id;
 
-	/** If true this static mesh should only be rendered during shadow depth passes. */
-	bool bShadowOnly;
-
 	// Constructor/destructor.
 	FStaticMesh(
 		FPrimitiveSceneInfo* InPrimitiveSceneInfo,
 		const FMeshBatch& InMesh,
 		float InScreenSize,
-		bool bInShadowOnly,
 		FHitProxyId InHitProxyId
 		):
 		FMeshBatch(InMesh),
 		ScreenSize(InScreenSize),
 		PrimitiveSceneInfo(InPrimitiveSceneInfo),
-		Id(INDEX_NONE),
-		bShadowOnly(bInShadowOnly)
+		Id(INDEX_NONE)
 	{
 		BatchHitProxyId = InHitProxyId;
 	}
diff --git a/Engine/Source/Runtime/Renderer/Private/SceneVisibility.cpp b/Engine/Source/Runtime/Renderer/Private/SceneVisibility.cpp
--- a/Engine/Source/Runtime/Renderer/Private/SceneVisibility.cpp
+++ b/Engine/Source/Runtime/Renderer/Private/SceneVisibility.cpp
@@ -1533,15 +1533,18 @@ struct FRelevancePacket
 						bNeedsBatchVisibility = true;
 					}
 
-					if(ViewRelevance.bDrawRelevance && !StaticMesh.bShadowOnly && (ViewRelevance.bRenderInMainPass || ViewRelevance.bRenderCustomDepth))
+					if(ViewRelevance.bDrawRelevance && (StaticMesh.bUseForMaterial || StaticMesh.bUseAsOccluder) && (ViewRelevance.bRenderInMainPass || ViewRelevance.bRenderCustomDepth))
 					{
 						// Mark static mesh as visible for rendering
-						MarkMask |= EMarkMaskBits::StaticMeshVisibilityMapMask;
-						if (PrimitiveSceneInfo->ShouldRenderVelocity(View, false))
+						if (StaticMesh.bUseForMaterial)
 						{
-							MarkMask |= EMarkMaskBits::StaticMeshVelocityMapMask;
+							MarkMask |= EMarkMaskBits::StaticMeshVisibilityMapMask;
+							if (PrimitiveSceneInfo->ShouldRenderVelocity(View, false))
+							{
+								MarkMask |= EMarkMaskBits::StaticMeshVelocityMapMask;
+							}
+							++NumVisibleStaticMeshElements;
 						}
-						++NumVisibleStaticMeshElements;
 
 						// If the static mesh is an occluder, check whether it covers enough of the screen to be used as an occluder.
 						if(	StaticMesh.bUseAsOccluder && bDrawDepthOnly )
