Remove string comparisons for cloth prep code.

diff --git a/Engine/Source/Runtime/Engine/Classes/Engine/SkeletalMesh.h b/Engine/Source/Runtime/Engine/Classes/Engine/SkeletalMesh.h
--- a/Engine/Source/Runtime/Engine/Classes/Engine/SkeletalMesh.h
+++ b/Engine/Source/Runtime/Engine/Classes/Engine/SkeletalMesh.h
@@ -478,6 +478,10 @@ struct FClothingAssetData
 	 * Num of this array means LOD number of clothing physical meshes 
 	 */
 	TArray<FClothVisualizationInfo> ClothVisualizationInfos;
+
+	/** Apex stores only the bones that cloth needs. We need a mapping from apex bone index to UE bone index. */
+	TArray<int32> ApexToUnrealBoneMapping;
+
 	/** currently mapped morph target name */
 	FName PreparedMorphTargetName;
 
diff --git a/Engine/Source/Runtime/Engine/Private/SkeletalMesh.cpp b/Engine/Source/Runtime/Engine/Private/SkeletalMesh.cpp
--- a/Engine/Source/Runtime/Engine/Private/SkeletalMesh.cpp
+++ b/Engine/Source/Runtime/Engine/Private/SkeletalMesh.cpp
@@ -2425,6 +2425,31 @@ bool USkeletalMesh::IsReadyForFinishDestroy()
 	return ReleaseResourcesFence.IsFenceComplete();
 }
 
+#if WITH_APEX_CLOTHING
+// convert a bone name from APEX stype to FBX style
+FName GetConvertedBoneName(NxClothingAsset* ApexClothingAsset, int32 BoneIndex)
+{
+	return FString(ApexClothingAsset->getBoneName(BoneIndex)).Replace(TEXT(" "), TEXT("-"));
+}
+
+
+void BuildApexToUnrealBoneMapping(const USkeletalMesh* SkeletalMesh, FClothingAssetData& ClothingAsset)
+{
+	NxClothingAsset* ApexClothingAsset = ClothingAsset.ApexClothingAsset;
+	uint32 NumUsedBones = ClothingAsset.ApexClothingAsset->getNumUsedBones();
+	
+	check(ClothingAsset.ApexToUnrealBoneMapping.Num() == 0);
+	ClothingAsset.ApexToUnrealBoneMapping.AddUninitialized(NumUsedBones);
+
+	for (uint32 Index = 0; Index < NumUsedBones; Index++)
+	{
+		FName BoneName = GetConvertedBoneName(ApexClothingAsset, Index);
+		int32 BoneIndex = SkeletalMesh->RefSkeleton.FindBoneIndex(BoneName);
+		ClothingAsset.ApexToUnrealBoneMapping[Index] = BoneIndex;
+	}
+}
+#endif
+
 void USkeletalMesh::Serialize( FArchive& Ar )
 {
 	DECLARE_SCOPE_CYCLE_COUNTER( TEXT("USkeletalMesh::Serialize"), STAT_SkeletalMesh_Serialize, STATGROUP_LoadTime );
@@ -2517,6 +2542,12 @@ void USkeletalMesh::Serialize( FArchive& Ar )
 		for( int32 Idx=0;Idx<ClothingAssets.Num();Idx++ )
 		{
 			Ar << ClothingAssets[Idx];
+#if WITH_APEX_CLOTHING
+			if(Ar.IsLoading())
+			{
+				BuildApexToUnrealBoneMapping(this, ClothingAssets[Idx]);
+			}
+#endif
 		}
 	}
 
diff --git a/Engine/Source/Runtime/Engine/Private/SkeletalMeshComponentPhysics.cpp b/Engine/Source/Runtime/Engine/Private/SkeletalMeshComponentPhysics.cpp
--- a/Engine/Source/Runtime/Engine/Private/SkeletalMeshComponentPhysics.cpp
+++ b/Engine/Source/Runtime/Engine/Private/SkeletalMeshComponentPhysics.cpp
@@ -1992,33 +1992,8 @@ bool USkeletalMeshComponent::ComponentOverlapMultiImpl(TArray<struct FOverlapRes
 
 #if WITH_APEX_CLOTHING
 
-// convert a bone name from APEX stype to FBX style
-static FName GetConvertedBoneName(NxClothingAsset* ApexClothingAsset, int32 BoneIndex)
-{
-	//QUICK_SCOPE_CYCLE_COUNTER(STAT_GetConvertedBoneName);
-	/*static FString Work;
-	check(IsInGameThread() && !Work.Len());
-	// FIXME, this still needs work, string conversion and hashing
-
-	const char* Bone = ApexClothingAsset->getBoneName(BoneIndex);
-	if (Bone)
-	{
-		int32 SrcLen = FCStringAnsi::Strlen(Bone) + 1;
-		int32 DestLen = FPlatformString::ConvertedLength<TCHAR>(Bone, SrcLen);
-		TArray<TCHAR>& Array = Work.GetCharArray();
-		Array.AddUninitialized(DestLen);
-		FPlatformString::Convert(Array.GetData(), DestLen, Bone, SrcLen);
-
-		Work.ReplaceInline(TEXT(" "), TEXT("-"));
-		FName Result(*Work);
-		Work.Reset();
-		//check(Result == FName(*FString(ApexClothingAsset->getBoneName(BoneIndex)).Replace(TEXT(" "), TEXT("-"))));
-		return Result;
-	}
-	return NAME_None;*/
-	//Going back to allocation approach because we now call this from many threads
-	return *FString(ApexClothingAsset->getBoneName(BoneIndex)).Replace(TEXT(" "), TEXT("-"));
-}
+
+
 
 void USkeletalMeshComponent::AddClothingBounds(FBoxSphereBounds& InOutBounds) const
 {
@@ -3579,31 +3554,25 @@ void USkeletalMeshComponent::ParallelUpdateClothState(float DeltaTime, FClothSim
 	physx::apex::ClothingTeleportMode::Enum CurTeleportMode = (physx::apex::ClothingTeleportMode::Enum)ClothSimulationContext.ClothTeleportMode;
 
 	TArray<physx::PxMat44> BoneMatrices;
-	check(!BoneMatrices.Num());
 
-	for (int32 ActorIdx = 0; ActorIdx < NumActors; ActorIdx++)
+	for (int32 Index = 0; Index < ClothingActors.Num(); ++Index)
 	{
-		// skip if ClothingActor is NULL or invalid
-		if (!IsValidClothingActor(ActorIdx))
-		{
-			continue;
-		}
+		const FClothingActor& ClothingActor = ClothingActors[Index];
+		const FClothingAssetData& ClothingAsset = SkeletalMesh->ClothingAssets[Index];
 
-		ApplyWindForCloth(ClothingActors[ActorIdx]);
+		ApplyWindForCloth(ClothingActor);
 
-		NxClothingAsset* ClothingAsset = ClothingActors[ActorIdx].ParentClothingAsset;
+		const NxClothingAsset* ApexClothingAsset = ClothingActor.ParentClothingAsset;
 
-		uint32 NumUsedBones = ClothingAsset->getNumUsedBones();
+		uint32 NumUsedBones = ApexClothingAsset->getNumUsedBones();
 
 		check(!BoneMatrices.Num());
 		BoneMatrices.AddUninitialized(NumUsedBones);
 
 		for (uint32 Index = 0; Index < NumUsedBones; Index++)
 		{
-			FName BoneName = GetConvertedBoneName(ClothingActors[ActorIdx].ParentClothingAsset, Index);
-
-			int32 BoneIndex = GetBoneIndex(BoneName);
-
+			int32 BoneIndex = ClothingAsset.ApexToUnrealBoneMapping[Index];
+			
 			if (MasterPoseComponent.IsValid())
 			{
 				int32 TempBoneIndex = BoneIndex;
@@ -3631,12 +3600,12 @@ void USkeletalMeshComponent::ParallelUpdateClothState(float DeltaTime, FClothSim
 			}
 		}
 
-		NxClothingActor* ClothingActor = ClothingActors[ActorIdx].ApexClothingActor;
+		NxClothingActor* ApexClothingActor = ClothingActor.ApexClothingActor;
 
-		check(ClothingActor);
+		check(ApexClothingActor);
 
 		// if bUseInternalboneOrder is set, "NumUsedBones" works, otherwise have to use "getNumBones" 
-		ClothingActor->updateState(
+		ApexClothingActor->updateState(
 			PxGlobalPose,
 			BoneMatrices.GetData(),
 			sizeof(physx::PxMat44),
