Fixed FStaticMeshLODResources.DepthOnlyIndexBuffer not being used in cook.

diff --git a/Engine/Source/Runtime/Engine/Private/StaticMesh.cpp b/Engine/Source/Runtime/Engine/Private/StaticMesh.cpp
--- a/Engine/Source/Runtime/Engine/Private/StaticMesh.cpp
+++ b/Engine/Source/Runtime/Engine/Private/StaticMesh.cpp
@@ -292,8 +292,10 @@ void FStaticMeshLODResources::Serialize(FArchive& Ar, UObject* Owner, int32 Inde
 	bool bNeedsCPUAccess = !FPlatformProperties::RequiresCookedData();
 
 	bHasAdjacencyInfo = false;
-	bHasReversedIndexBuffer = false;
-	bHasReversedDepthOnlyIndexBuffer = false;
+	bHasDepthOnlyIndices = false;
+	bHasReversedIndices = false;
+	bHasReversedDepthOnlyIndices = false;
+	DepthOnlyNumTriangles = 0;
 
     // Defined class flags for possible stripping
 	const uint8 AdjacencyDataStripFlag = 1;
@@ -329,8 +331,10 @@ void FStaticMeshLODResources::Serialize(FArchive& Ar, UObject* Owner, int32 Inde
 		}
 
 		// Needs to be done now because on cooked platform, indices are discarded after RHIInit.
-		bHasReversedIndexBuffer = ReversedIndexBuffer.GetNumIndices() != 0;
-		bHasReversedDepthOnlyIndexBuffer = ReversedDepthOnlyIndexBuffer.GetNumIndices() != 0;
+		bHasDepthOnlyIndices = DepthOnlyIndexBuffer.GetNumIndices() != 0;
+		bHasReversedIndices = ReversedIndexBuffer.GetNumIndices() != 0;
+		bHasReversedDepthOnlyIndices = ReversedDepthOnlyIndexBuffer.GetNumIndices() != 0;
+		DepthOnlyNumTriangles = DepthOnlyIndexBuffer.GetNumIndices() / 3;
 	}
 }
 
@@ -494,8 +498,10 @@ FStaticMeshLODResources::FStaticMeshLODResources()
 	: DistanceFieldData(NULL)
 	, MaxDeviation(0.0f)
 	, bHasAdjacencyInfo(false)
-	, bHasReversedIndexBuffer(false)
-	, bHasReversedDepthOnlyIndexBuffer(false)
+	, bHasDepthOnlyIndices(false)
+	, bHasReversedIndices(false)
+	, bHasReversedDepthOnlyIndices(false)
+	, DepthOnlyNumTriangles(0)
 {
 }
 
diff --git a/Engine/Source/Runtime/Engine/Private/StaticMeshRender.cpp b/Engine/Source/Runtime/Engine/Private/StaticMeshRender.cpp
--- a/Engine/Source/Runtime/Engine/Private/StaticMeshRender.cpp
+++ b/Engine/Source/Runtime/Engine/Private/StaticMeshRender.cpp
@@ -196,7 +196,7 @@ bool FStaticMeshSceneProxy::GetShadowMeshElement(int32 LODIndex, int32 BatchInde
 	const FStaticMeshLODResources& LOD = RenderData->LODResources[LODIndex];
 	const FLODInfo& ProxyLODInfo = LODs[LODIndex];
 
-	const bool bUseReversedIndices = GUseReversedIndexBuffer && IsLocalToWorldDeterminantNegative() && LOD.bHasReversedDepthOnlyIndexBuffer;
+	const bool bUseReversedIndices = GUseReversedIndexBuffer && IsLocalToWorldDeterminantNegative() && LOD.bHasReversedDepthOnlyIndices;
 
 	FMeshBatchElement& OutBatchElement = OutMeshBatch.Elements[0];
 	OutMeshBatch.MaterialRenderProxy = UMaterial::GetDefaultMaterial(MD_Surface)->GetRenderProxy(false, false);
@@ -204,7 +204,7 @@ bool FStaticMeshSceneProxy::GetShadowMeshElement(int32 LODIndex, int32 BatchInde
 	OutBatchElement.IndexBuffer = bUseReversedIndices ? &LOD.ReversedDepthOnlyIndexBuffer : &LOD.DepthOnlyIndexBuffer;
 	OutMeshBatch.Type = PT_TriangleList;
 	OutBatchElement.FirstIndex = 0;
-	OutBatchElement.NumPrimitives = LOD.DepthOnlyIndexBuffer.GetNumIndices() / 3;
+	OutBatchElement.NumPrimitives = LOD.DepthOnlyNumTriangles;
 	OutBatchElement.PrimitiveUniformBufferResource = &GetUniformBuffer();
 	OutBatchElement.MinVertexIndex = 0;
 	OutBatchElement.MaxVertexIndex = LOD.PositionVertexBuffer.GetNumVertices() - 1;
@@ -273,7 +273,7 @@ bool FStaticMeshSceneProxy::GetMeshElement(int32 LODIndex, int32 BatchIndex, int
 
 	const bool bWireframe = false;
 	const bool bRequiresAdjacencyInformation = RequiresAdjacencyInformation( Material, OutMeshBatch.VertexFactory->GetType(), GetScene().GetFeatureLevel() );
-	const bool bUseReversedIndices = GUseReversedIndexBuffer && !bWireframe && !bRequiresAdjacencyInformation && IsLocalToWorldDeterminantNegative() && LOD.bHasReversedIndexBuffer;
+	const bool bUseReversedIndices = GUseReversedIndexBuffer && !bWireframe && !bRequiresAdjacencyInformation && IsLocalToWorldDeterminantNegative() && LOD.bHasReversedIndices;
 
 	SetIndexSource(LODIndex, SectionIndex, OutMeshBatch, bWireframe, bRequiresAdjacencyInformation, bUseReversedIndices);
 
@@ -528,7 +528,7 @@ void FStaticMeshSceneProxy::DrawStaticElements(FStaticPrimitiveDrawInterface* PD
 				bool bUseUnifiedMeshForShadow = false;
 				bool bUseUnifiedMeshForDepth = false;
 
-				if (GUseShadowIndexBuffer && LODModel.DepthOnlyIndexBuffer.GetNumIndices() > 0)
+				if (GUseShadowIndexBuffer && LODModel.bHasDepthOnlyIndices)
 				{
 					const FLODInfo& ProxyLODInfo = LODs[LODIndex];
 
diff --git a/Engine/Source/Runtime/Engine/Public/StaticMeshResources.h b/Engine/Source/Runtime/Engine/Public/StaticMeshResources.h
--- a/Engine/Source/Runtime/Engine/Public/StaticMeshResources.h
+++ b/Engine/Source/Runtime/Engine/Public/StaticMeshResources.h
@@ -565,11 +565,16 @@ struct FStaticMeshLODResources
 	/** True if the adjacency index buffer contained data at init. Needed as it will not be available to the CPU afterwards. */
 	uint32 bHasAdjacencyInfo : 1;
 
+	/** True if the depth only index buffers contained data at init. Needed as it will not be available to the CPU afterwards. */
+	uint32 bHasDepthOnlyIndices : 1;
+
 	/** True if the reversed index buffers contained data at init. Needed as it will not be available to the CPU afterwards. */
-	uint32 bHasReversedIndexBuffer : 1;
+	uint32 bHasReversedIndices : 1;
 
 	/** True if the reversed index buffers contained data at init. Needed as it will not be available to the CPU afterwards. */
-	uint32 bHasReversedDepthOnlyIndexBuffer : 1;
+	uint32 bHasReversedDepthOnlyIndices: 1;
+
+	uint32 DepthOnlyNumTriangles;
 
 	/** Default constructor. */
 	FStaticMeshLODResources();
