Fixed disk path for vlog files on case sensitivity platforms.
Integrated change #2677112 from dev branch.

diff --git a/Engine/Source/Runtime/Engine/Private/VisualLogger/VisualLoggerBinaryFileDevice.cpp b/Engine/Source/Runtime/Engine/Private/VisualLogger/VisualLoggerBinaryFileDevice.cpp
--- a/Engine/Source/Runtime/Engine/Private/VisualLogger/VisualLoggerBinaryFileDevice.cpp
+++ b/Engine/Source/Runtime/Engine/Private/VisualLogger/VisualLoggerBinaryFileDevice.cpp
@@ -33,7 +33,7 @@ void FVisualLoggerBinaryFileDevice::StartRecordingToFile(float TimeStamp)
 	StartRecordingTime = TimeStamp;
 	LastLogTimeStamp = StartRecordingTime;
 	TempFileName = FVisualLoggerHelpers::GenerateTemporaryFilename(VISLOG_FILENAME_EXT);
-	const FString FullFilename = FString::Printf(TEXT("%slogs/%s"), *FPaths::GameSavedDir(), *TempFileName);
+	const FString FullFilename = FString::Printf(TEXT("%s%s"), *FPaths::GameLogDir(), *TempFileName);
 
 	FileArchive = IFileManager::Get().CreateFileWriter(*FullFilename);
 }
@@ -57,7 +57,7 @@ void FVisualLoggerBinaryFileDevice::StopRecordingToFile(float TimeStamp)
 	delete FileArchive;
 	FileArchive = nullptr;
 
-	const FString TempFullFilename = FString::Printf(TEXT("%slogs/%s"), *FPaths::GameSavedDir(), *TempFileName);
+	const FString TempFullFilename = FString::Printf(TEXT("%s%s"), *FPaths::GameLogDir(), *TempFileName);
 	const FString NewFileName = FVisualLoggerHelpers::GenerateFilename(TempFileName, FileName, StartRecordingTime, LastLogTimeStamp);
 
 	if (TotalSize > 0)
