#ANIM: Montage supports more than linear blending
 - It uses same blend set up as blend list
 - Because of this, now we're not testing weight <ZERO_ANIMWEIGHT, but we test Blend.IsComplete
 - Now we still advance montage since if the custom curve contains longer 0, anim is likely to pause if he comes back from 0. (I had wacky test case)

 #ANIM: Single node instance's inconsistency of montage weight has been resolved

- This also contains montage update change where now montage advance prior to node update
- Also made new function for node update, so that single node instance can use that function instead of nativeupdateanimation, which happens prior to weight update causing all inconsistency
- Enabled node weight debuging define for montage, so that we can see if any other case is causing it
- This also fixes crash when montage gets terminated after weight has been updated

Merging

//depot/UE4-Orion/...

to //depot/UE4-Orion/...

diff --git a/Engine/Documentation/Source/Shared/Types/UAnimMontage/UAnimMontage.INT.udn b/Engine/Documentation/Source/Shared/Types/UAnimMontage/UAnimMontage.INT.udn
--- a/Engine/Documentation/Source/Shared/Types/UAnimMontage/UAnimMontage.INT.udn
+++ b/Engine/Documentation/Source/Shared/Types/UAnimMontage/UAnimMontage.INT.udn
@@ -11,3 +11,13 @@ AnimMontages (or just Montages for short) are a multipurpose tool that allows fo
 Engine/Animation/AnimMontage
 [/VAR]
 [/EXCERPT:UAnimMontage]
+
+[EXCERPT:BlendIn]
+WIP - blending in option. http://easings.net/#
+
+[/EXCERPT:BlendIn]
+
+[EXCERPT:BlendOut]
+WIP - blending out option. http://easings.net/#
+
+[/EXCERPT:BlendOut]
diff --git a/Engine/Source/Editor/AnimGraph/Classes/AnimPreviewInstance.h b/Engine/Source/Editor/AnimGraph/Classes/AnimPreviewInstance.h
--- a/Engine/Source/Editor/AnimGraph/Classes/AnimPreviewInstance.h
+++ b/Engine/Source/Editor/AnimGraph/Classes/AnimPreviewInstance.h
@@ -50,7 +50,7 @@ class ANIMGRAPH_API UAnimPreviewInstance : public UAnimSingleNodeInstance
 
 	// Begin UAnimInstance interface
 	virtual void NativeInitializeAnimation() override;
-	virtual void NativeUpdateAnimation(float DeltaTimeX) override;
+	virtual void UpdateAnimationNode(float DeltaTimeX) override;
 	virtual bool NativeEvaluateAnimation(FPoseContext& Output) override;
 	// End UAnimInstance interface
 
diff --git a/Engine/Source/Editor/AnimGraph/Private/AnimPreviewInstance.cpp b/Engine/Source/Editor/AnimGraph/Private/AnimPreviewInstance.cpp
--- a/Engine/Source/Editor/AnimGraph/Private/AnimPreviewInstance.cpp
+++ b/Engine/Source/Editor/AnimGraph/Private/AnimPreviewInstance.cpp
@@ -102,7 +102,7 @@ void UAnimPreviewInstance::RemoveBoneModification(const FName& InBoneName, bool
 	);
 }
 
-void UAnimPreviewInstance::NativeUpdateAnimation(float DeltaTimeX)
+void UAnimPreviewInstance::UpdateAnimationNode(float DeltaTimeX)
 {
 #if WITH_EDITORONLY_DATA
 	if(bForceRetargetBasePose)
@@ -112,7 +112,7 @@ void UAnimPreviewInstance::NativeUpdateAnimation(float DeltaTimeX)
 	}
 #endif // #if WITH_EDITORONLY_DATA
 
-	Super::NativeUpdateAnimation(DeltaTimeX);
+	Super::UpdateAnimationNode(DeltaTimeX);
 }
 
 bool UAnimPreviewInstance::NativeEvaluateAnimation(FPoseContext& Output)
@@ -405,6 +405,9 @@ void UAnimPreviewInstance::RestartMontage(UAnimMontage* Montage, FName FromSecti
 	if (Montage == CurrentAsset)
 	{
 		MontagePreviewType = EMPT_Normal;
+		// since this is preview, we would like not to blend in
+		// just hard stop here
+		Montage_Stop(0.0f, Montage);
 		Montage_Play(Montage, PlayRate);
 		if (FromSection != NAME_None)
 		{
@@ -733,6 +736,9 @@ void UAnimPreviewInstance::MontagePreview_PreviewNormal(int32 FromSectionIdx)
 			PreviewFromSection = MontagePreviewStartSectionIdx;
 		}
 		MontagePreviewType = EMPT_Normal;
+		// since this is preview, we would like not to blend in
+		// just hard stop here
+		Montage_Stop(0.0f, Montage);
 		Montage_Play(Montage, PlayRate);
 		MontagePreview_SetLoopNormal(bLooping, FromSectionIdx);
 		Montage_JumpToSection(Montage->GetSectionName(PreviewFromSection));
@@ -747,6 +753,9 @@ void UAnimPreviewInstance::MontagePreview_PreviewAllSections()
 	if (Montage && Montage->SequenceLength > 0.f)
 	{
 		MontagePreviewType = EMPT_AllSections;
+		// since this is preview, we would like not to blend in
+		// just hard stop here
+		Montage_Stop(0.0f, Montage);
 		Montage_Play(Montage, PlayRate);
 		MontagePreview_SetLoopAllSections(bLooping);
 		MontagePreview_JumpToPreviewStart();
diff --git a/Engine/Source/Editor/Persona/Private/SMontageEditor.cpp b/Engine/Source/Editor/Persona/Private/SMontageEditor.cpp
--- a/Engine/Source/Editor/Persona/Private/SMontageEditor.cpp
+++ b/Engine/Source/Editor/Persona/Private/SMontageEditor.cpp
@@ -347,6 +347,8 @@ void SMontageEditor::OnMontageChange(class UObject *EditorAnimBaseObj, bool Rebu
 
 	if ( MontageObj != nullptr )
 	{
+		float PreviouewSeqLength = GetSequenceLength();
+
 		if(Rebuild && !bIsActiveTimerRegistered)
 		{
 			bIsActiveTimerRegistered = true;
@@ -358,6 +360,13 @@ void SMontageEditor::OnMontageChange(class UObject *EditorAnimBaseObj, bool Rebu
 		}
 		
 		MontageObj->MarkPackageDirty();
+
+		// if animation length changed, we might be out of range, let's restart
+		if (GetSequenceLength() != PreviouewSeqLength)
+		{
+			// this might not be safe
+			RestartPreview();
+		}
 	}
 }
 
diff --git a/Engine/Source/Runtime/AnimGraphRuntime/Private/AnimNodes/AnimNode_BlendListBase.cpp b/Engine/Source/Runtime/AnimGraphRuntime/Private/AnimNodes/AnimNode_BlendListBase.cpp
--- a/Engine/Source/Runtime/AnimGraphRuntime/Private/AnimNodes/AnimNode_BlendListBase.cpp
+++ b/Engine/Source/Runtime/AnimGraphRuntime/Private/AnimNodes/AnimNode_BlendListBase.cpp
@@ -115,8 +115,6 @@ void FAnimNode_BlendListBase::Update(const FAnimationUpdateContext& Context)
 				{
 					Blend.SetValueRange(BlendWeights[i], 0.0f);
 				}
-
-				Blend.Reset();
 			}
 		}
 
diff --git a/Engine/Source/Runtime/Engine/Classes/Animation/AnimInstance.h b/Engine/Source/Runtime/Engine/Classes/Animation/AnimInstance.h
--- a/Engine/Source/Runtime/Engine/Classes/Animation/AnimInstance.h
+++ b/Engine/Source/Runtime/Engine/Classes/Animation/AnimInstance.h
@@ -583,7 +583,7 @@ protected:
 	void StopAllMontages(float BlendOut);
 
 	/** Stop all active montages belonging to 'InGroupName' */
-	void StopAllMontagesByGroupName(FName InGroupName, float BlendOutTime);
+	void StopAllMontagesByGroupName(FName InGroupName, const FAlphaBlend& BlendOut);
 
 	/** Update weight of montages  **/
 	virtual void Montage_UpdateWeight(float DeltaSeconds);
@@ -1009,5 +1009,9 @@ private:
 
 	/** Native state exit bindings */
 	TArray<FNativeStateBinding> NativeStateExitBindings;
+
+protected:
+	/** Update all animation node */
+	virtual void UpdateAnimationNode(float DeltaSeconds);
 };
 
diff --git a/Engine/Source/Runtime/Engine/Classes/Animation/AnimMontage.h b/Engine/Source/Runtime/Engine/Classes/Animation/AnimMontage.h
--- a/Engine/Source/Runtime/Engine/Classes/Animation/AnimMontage.h
+++ b/Engine/Source/Runtime/Engine/Classes/Animation/AnimMontage.h
@@ -201,14 +201,8 @@ private:
 	UPROPERTY()
 	float PlayRate;
 
-	UPROPERTY()
-	float DesiredWeight;
-
-	UPROPERTY()
-	float Weight;
-
 	UPROPERTY(transient)
-	float BlendTime;
+	FAlphaBlend Blend;
 
 	// need to save if it's interrupted or not
 	// this information is crucial for gameplay
@@ -241,9 +235,9 @@ public:
 	void MontageSync_PostUpdate();
 
 	/** Get Weight */
-	float GetWeight() const { return Weight; }
-	float GetDesiredWeight() const { return DesiredWeight; }
-	float GetBlendTime() const { return BlendTime; }
+	float GetWeight() const { return Blend.GetBlendedValue(); }
+	float GetDesiredWeight() const { return Blend.GetDesiredValue(); }
+	float GetBlendTime() const { return Blend.GetBlendTime(); }
 private:
 	/** Followers this Montage will synchronize */
 	TArray<struct FAnimMontageInstance*> MontageSyncFollowers;
@@ -259,6 +253,9 @@ private:
 	/** Synchronize ourselves to our leader */
 	void MontageSync_PerformSyncToLeader();
 
+	/** Initialize Blend Setup from Montage */
+	void InitializeBlend(const FAlphaBlend& InAlphaBlend);
+
 public:
 	FAnimMontageInstance()
 		: Montage(NULL)
@@ -267,9 +264,6 @@ public:
 		, AnimInstance(NULL)
 		, Position(0.f)
 		, PlayRate(1.f)
-		, DesiredWeight(0.f)
-		, Weight(0.f)
-		, BlendTime(0.f)
 		, bInterrupted(false)
 		, PreviousWeight(0.f)
 		, MontageSyncLeader(NULL)
@@ -284,9 +278,6 @@ public:
 		, AnimInstance(InAnimInstance)
 		, Position(0.f)
 		, PlayRate(1.f)
-		, DesiredWeight(0.f)
-		, Weight(0.f)
-		, BlendTime(0.f)
 		, bInterrupted(false)
 		, PreviousWeight(0.f)	
 		, MontageSyncLeader(NULL)
@@ -296,7 +287,7 @@ public:
 
 	// montage instance interfaces
 	void Play(float InPlayRate = 1.f);
-	void Stop(float BlendOutDuration, bool bInterrupt=true);
+	void Stop(const FAlphaBlend& InBlendOut, bool bInterrupt=true);
 	void Pause();
 	void Initialize(class UAnimMontage * InMontage);
 
@@ -306,7 +297,7 @@ public:
 
 	bool IsValid() const { return (Montage!=NULL); }
 	bool IsPlaying() const { return IsValid() && bPlaying; }
-	bool IsStopped() const { return DesiredWeight == 0.f; }
+	bool IsStopped() const { return Blend.GetDesiredValue() == 0.f; }
 
 	/** Returns true if this montage is active (valid and not blending out) */
 	bool IsActive() const { return (IsValid() && !IsStopped()); }
@@ -374,18 +365,24 @@ class UAnimMontage : public UAnimCompositeBase
 {
 	GENERATED_UCLASS_BODY()
 
-	/** Default blend in time. */
-	UPROPERTY(EditAnywhere, Category=Montage)
-	float BlendInTime;
+	/** Blend in option. */
+	UPROPERTY(EditAnywhere, Category=BlendOption)
+	FAlphaBlend BlendIn;
 
-	/** Default blend out time. */
-	UPROPERTY(EditAnywhere, Category=Montage)
-	float BlendOutTime;
+	UPROPERTY()
+	float BlendInTime_DEPRECATED;
+
+	/** Blend out option. This is only used when it blends out itself. If it's interrupted by other montages, it will use new montage's BlendIn option to blend out. */
+	UPROPERTY(EditAnywhere, Category=BlendOption)
+	FAlphaBlend BlendOut;
+
+	UPROPERTY()
+	float BlendOutTime_DEPRECATED;
 
 	/** Time from Sequence End to trigger blend out.
 	 * <0 means using BlendOutTime, so BlendOut finishes as Montage ends.
 	 * >=0 means using 'SequenceEnd - BlendOutTriggerTime' to trigger blend out. */
-	UPROPERTY(EditAnywhere, Category = Montage)
+	UPROPERTY(EditAnywhere, Category = BlendOption)
 	float BlendOutTriggerTime;
 
 	// composite section. 
diff --git a/Engine/Source/Runtime/Engine/Classes/Animation/AnimSingleNodeInstance.h b/Engine/Source/Runtime/Engine/Classes/Animation/AnimSingleNodeInstance.h
--- a/Engine/Source/Runtime/Engine/Classes/Animation/AnimSingleNodeInstance.h
+++ b/Engine/Source/Runtime/Engine/Classes/Animation/AnimSingleNodeInstance.h
@@ -57,7 +57,6 @@ class ENGINE_API UAnimSingleNodeInstance : public UAnimInstance
 
 	// Begin UAnimInstance interface
 	virtual void NativeInitializeAnimation() override;
-	virtual void NativeUpdateAnimation(float DeltaTimeX) override;
 	virtual bool NativeEvaluateAnimation(FPoseContext& Output) override;
 	virtual void NativePostEvaluateAnimation() override;
 	virtual void OnMontageInstanceStopped(FAnimMontageInstance& StoppedMontageInstance) override;
@@ -115,6 +114,10 @@ public:
 private:
 	/** Slot node weight transient data */
 	FSlotNodeWeightInfo WeightInfo;
+
+protected:
+	// UAnimInstance interface to override animation node part
+	virtual void UpdateAnimationNode(float DeltaSeconds) override;
 };
 
 
diff --git a/Engine/Source/Runtime/Engine/Private/AlphaBlend.cpp b/Engine/Source/Runtime/Engine/Private/AlphaBlend.cpp
--- a/Engine/Source/Runtime/Engine/Private/AlphaBlend.cpp
+++ b/Engine/Source/Runtime/Engine/Private/AlphaBlend.cpp
@@ -3,16 +3,26 @@
 #include "EnginePrivate.h"
 #include "AlphaBlend.h"
 
-FAlphaBlend::FAlphaBlend() 
+FAlphaBlend::FAlphaBlend(float NewBlendTime) 
 	: BlendOption(EAlphaBlendOption::Linear)
 	, BeginValue(0.0f)
 	, DesiredValue(1.0f)
-	, BlendTime(0.2f)
+	, BlendTime(NewBlendTime)
 	, CustomCurve(nullptr)
 {
 	Reset();
 }
 
+FAlphaBlend::FAlphaBlend(const FAlphaBlend& Other, float NewBlendTime)
+	: BlendOption(Other.BlendOption)
+	, BeginValue(Other.BeginValue)
+	, DesiredValue(Other.DesiredValue)
+	, BlendTime(NewBlendTime)
+	, CustomCurve(Other.CustomCurve)
+{
+	Reset();
+}
+
 void FAlphaBlend::SetBlendTime(float InBlendTime)
 {
 	BlendTime = FMath::Max(InBlendTime, 0.f);
@@ -25,36 +35,10 @@ void FAlphaBlend::Reset()
 	BlendedValue = BeginValue;
 
 	// Set alpha target to full - will also handle zero blend times
-	SetTarget(1.0f);
-}
-
-bool FAlphaBlend::GetToggleStatus()
-{
-	return (AlphaTarget > 0.f);
-}
-
-void FAlphaBlend::Toggle(bool bEnable)
-{
-	ConditionalSetTarget(bEnable ? 1.f : 0.f);
-}
-
-void FAlphaBlend::ConditionalSetTarget(float InAlphaTarget)
-{
-	if( AlphaTarget != InAlphaTarget )
-	{
-		SetTarget(InAlphaTarget);
-	}
-}
-
-void FAlphaBlend::SetTarget(float InAlphaTarget)
-{
-	// Clamp parameters to valid range
-	AlphaTarget = FMath::Clamp<float>(InAlphaTarget, 0.f, 1.f);
-
 	// if blend time is zero, transition now, don't wait to call update.
 	if( BlendTime <= 0.f )
 	{
-		AlphaLerp = AlphaTarget;
+		AlphaLerp = 1.f;
 		AlphaBlend = AlphaToBlendOption();
 		BlendTimeRemaining = 0.f;
 		BlendedValue = BeginValue + (DesiredValue - BeginValue) * AlphaBlend;
@@ -62,7 +46,7 @@ void FAlphaBlend::SetTarget(float InAlphaTarget)
 	else
 	{
 		// Blend time is to go all the way, so scale that by how much we have to travel
-		BlendTimeRemaining = BlendTime * FMath::Abs(AlphaTarget - AlphaLerp);
+		BlendTimeRemaining = BlendTime * FMath::Abs(1.f - AlphaLerp);
 	}
 }
 
@@ -75,7 +59,7 @@ void FAlphaBlend::Update(float InDeltaTime)
 	{
 		if( BlendTimeRemaining > InDeltaTime )
 		{
-			const float BlendDelta = AlphaTarget - AlphaLerp; 
+			const float BlendDelta = 1.f - AlphaLerp; 
 			AlphaLerp += (BlendDelta / BlendTimeRemaining) * InDeltaTime;
 			BlendTimeRemaining -= InDeltaTime;
 
@@ -138,7 +122,7 @@ void FAlphaBlend::SetBlendOption(EAlphaBlendOption InBlendOption)
 	BlendedValue = BeginValue + (DesiredValue - BeginValue) * AlphaBlend;
 }
 
-float FAlphaBlend::GetBlendedValue()
+float FAlphaBlend::GetBlendedValue() const
 {
 	return BlendedValue;
 }
@@ -150,6 +134,14 @@ void FAlphaBlend::SetValueRange(float Begin, float Desired)
 
 	// Convert to new range
 	BlendedValue = BeginValue + (DesiredValue - BeginValue) * AlphaBlend;
+
+	Reset();
+}
+
+/** Sets the final desired value for the blended value */
+void FAlphaBlend::SetDesiredValue(float InDesired)
+{
+	SetValueRange(BlendedValue, InDesired);
 }
 
 void FAlphaBlend::SetAlpha(float InAlpha)
@@ -159,7 +151,7 @@ void FAlphaBlend::SetAlpha(float InAlpha)
 	BlendedValue = BeginValue + (DesiredValue - BeginValue) * AlphaBlend;
 }
 
-bool FAlphaBlend::IsComplete()
+bool FAlphaBlend::IsComplete() const
 {
 	return AlphaLerp == 1.0f;
 }
diff --git a/Engine/Source/Runtime/Engine/Private/Animation/AnimInstance.cpp b/Engine/Source/Runtime/Engine/Private/Animation/AnimInstance.cpp
--- a/Engine/Source/Runtime/Engine/Private/Animation/AnimInstance.cpp
+++ b/Engine/Source/Runtime/Engine/Private/Animation/AnimInstance.cpp
@@ -350,6 +350,18 @@ bool UAnimInstance::UpdateSnapshotAndSkipRemainingUpdate()
 }
 #endif
 
+// update animation node, if you have your own anim instance, you could override this to update extra nodes if you'd like
+void UAnimInstance::UpdateAnimationNode(float DeltaSeconds)
+{
+	// Update the anim graph
+	if(RootNode != NULL)
+	{
+		IncrementGraphTraversalCounter();
+		FAnimationUpdateContext UpdateContext(this, DeltaSeconds);
+		RootNode->Update(UpdateContext);
+	}
+}
+
 void UAnimInstance::UpdateAnimation(float DeltaSeconds)
 {
 	QUICK_SCOPE_CYCLE_COUNTER(STAT_UAnimInstance_UpdateAnimation);
@@ -412,13 +424,12 @@ void UAnimInstance::UpdateAnimation(float DeltaSeconds)
 
 	Montage_UpdateWeight(DeltaSeconds);
 
-	// Update the anim graph
-	if (RootNode != NULL)
-	{
-		IncrementGraphTraversalCounter();
-		FAnimationUpdateContext UpdateContext(this, DeltaSeconds);
-		RootNode->Update(UpdateContext);
-	}
+	// update montage should run in game thread
+	// if we do multi threading, make sure this stays in game thread
+	Montage_Advance(DeltaSeconds);
+
+	// update all nodes
+	UpdateAnimationNode(DeltaSeconds);
 
 	TickSyncGroupWriteIndex();
 
@@ -501,9 +512,6 @@ void UAnimInstance::UpdateAnimation(float DeltaSeconds)
 			}
 		}
 	}
-	// update montage should run in game thread
-	// if we do multi threading, make sure this stays in game thread
-		Montage_Advance(DeltaSeconds);
 
 	// We may have just partially blended root motion, so make it up to 1 by
 	// blending in identity too
@@ -1394,7 +1402,7 @@ void UAnimInstance::GetSlotWeight(FName const& SlotNodeName, float& out_SlotNode
 			}
 
 #if DEBUGMONTAGEWEIGHT			
-			TotalDesiredWeight += MontageInstance->DesiredWeight;
+			TotalDesiredWeight += MontageInstance->Blend.GetDesiredValue();
 #endif			
 		}
 	}
@@ -1432,9 +1440,10 @@ void UAnimInstance::GetSlotWeight(FName const& SlotNodeName, float& out_SlotNode
 }
 
 // for now disable becauase it will not work with single node instance
-#define DEBUG_MONTAGEINSTANCE_WEIGHT 0
 #if (UE_BUILD_SHIPPING || UE_BUILD_TEST)
 #define DEBUG_MONTAGEINSTANCE_WEIGHT 0
+#else
+#define DEBUG_MONTAGEINSTANCE_WEIGHT 1
 #endif // !(UE_BUILD_SHIPPING || UE_BUILD_TEST)
 
 void UAnimInstance::SlotEvaluatePose(FName SlotNodeName, const FCompactPose& SourcePose, const FBlendedCurve& SourceCurve, float InSourceWeight, FCompactPose& BlendedPose, FBlendedCurve& BlendedCurve, float InBlendWeight, float InTotalNodeWeight)
@@ -1940,8 +1949,8 @@ float UAnimInstance::PlaySlotAnimation(UAnimSequenceBase* Asset, FName SlotNodeN
 
 	// add new section
 	NewMontage->CompositeSections.Add(NewSection);
-	NewMontage->BlendInTime = BlendInTime;
-	NewMontage->BlendOutTime = BlendOutTime;
+	NewMontage->BlendIn.SetBlendTime(BlendInTime);
+	NewMontage->BlendOut.SetBlendTime(BlendOutTime);
 	NewMontage->SlotAnimTracks.Add(NewTrack);
 
 	return Montage_Play(NewMontage, InPlayRate);
@@ -1995,8 +2004,8 @@ UAnimMontage* UAnimInstance::PlaySlotAnimationAsDynamicMontage(UAnimSequenceBase
 
 	// add new section
 	NewMontage->CompositeSections.Add(NewSection);
-	NewMontage->BlendInTime = BlendInTime;
-	NewMontage->BlendOutTime = BlendOutTime;
+	NewMontage->BlendIn.SetBlendTime(BlendInTime);
+	NewMontage->BlendOut.SetBlendTime(BlendOutTime);
 	NewMontage->BlendOutTriggerTime = BlendOutTriggerTime;
 	NewMontage->SlotAnimTracks.Add(NewTrack);
 
@@ -2028,7 +2037,7 @@ void UAnimInstance::StopSlotAnimation(float InBlendOutTime, FName SlotNodeName)
 						if (AnimTrack && AnimTrack->SlotName == SlotNodeName)
 						{
 							// Found it
-							MontageInstance->Stop(InBlendOutTime);
+							MontageInstance->Stop(FAlphaBlend(InBlendOutTime));
 							break;
 						}
 					}
@@ -2077,7 +2086,7 @@ float UAnimInstance::Montage_Play(UAnimMontage* MontageToPlay, float InPlayRate/
 		{
 			// Enforce 'a single montage at once per group' rule
 			FName NewMontageGroupName = MontageToPlay->GetGroupName();
-			StopAllMontagesByGroupName(NewMontageGroupName, MontageToPlay->BlendInTime);
+			StopAllMontagesByGroupName(NewMontageGroupName, MontageToPlay->BlendIn);
 
 			// Enforce 'a single root motion montage at once' rule.
 			if (MontageToPlay->bEnableRootMotionTranslation || MontageToPlay->bEnableRootMotionRotation)
@@ -2085,7 +2094,7 @@ float UAnimInstance::Montage_Play(UAnimMontage* MontageToPlay, float InPlayRate/
 				FAnimMontageInstance* ActiveRootMotionMontageInstance = GetRootMotionMontageInstance();
 				if (ActiveRootMotionMontageInstance)
 				{
-					ActiveRootMotionMontageInstance->Stop(MontageToPlay->BlendInTime);
+					ActiveRootMotionMontageInstance->Stop(MontageToPlay->BlendIn);
 				}
 			}
 
@@ -2124,7 +2133,7 @@ void UAnimInstance::Montage_Stop(float InBlendOutTime, UAnimMontage* Montage)
 		FAnimMontageInstance* MontageInstance = GetActiveInstanceForMontage(*Montage);
 		if (MontageInstance)
 		{
-			MontageInstance->Stop(InBlendOutTime);
+			MontageInstance->Stop(FAlphaBlend(Montage->BlendOut, InBlendOutTime));
 		}
 	}
 	else
@@ -2135,7 +2144,7 @@ void UAnimInstance::Montage_Stop(float InBlendOutTime, UAnimMontage* Montage)
 			FAnimMontageInstance* MontageInstance = MontageInstances[InstanceIndex];
 			if (MontageInstance && MontageInstance->IsActive())
 			{
-				MontageInstance->Stop(InBlendOutTime);
+				MontageInstance->Stop(FAlphaBlend(MontageInstance->Montage->BlendOut, InBlendOutTime));
 			}
 		}
 	}
@@ -2599,18 +2608,18 @@ void UAnimInstance::StopAllMontages(float BlendOut)
 {
 	for (int32 Index = MontageInstances.Num() - 1; Index >= 0; Index--)
 	{
-		MontageInstances[Index]->Stop(BlendOut, true);
+		MontageInstances[Index]->Stop(FAlphaBlend(BlendOut), true);
 	}
 }
 
-void UAnimInstance::StopAllMontagesByGroupName(FName InGroupName, float BlendOutTime)
+void UAnimInstance::StopAllMontagesByGroupName(FName InGroupName, const FAlphaBlend& BlendOut)
 {
 	for (int32 InstanceIndex = MontageInstances.Num() - 1; InstanceIndex >= 0; InstanceIndex--)
 	{
 		FAnimMontageInstance* MontageInstance = MontageInstances[InstanceIndex];
 		if (MontageInstance && MontageInstance->IsActive() && MontageInstance->Montage && (MontageInstance->Montage->GetGroupName() == InGroupName))
 		{
-			MontageInstances[InstanceIndex]->Stop(BlendOutTime, true);
+			MontageInstances[InstanceIndex]->Stop(BlendOut, true);
 		}
 	}
 }
diff --git a/Engine/Source/Runtime/Engine/Private/Animation/AnimMontage.cpp b/Engine/Source/Runtime/Engine/Private/Animation/AnimMontage.cpp
--- a/Engine/Source/Runtime/Engine/Private/Animation/AnimMontage.cpp
+++ b/Engine/Source/Runtime/Engine/Private/Animation/AnimMontage.cpp
@@ -24,8 +24,8 @@ UAnimMontage::UAnimMontage(const FObjectInitializer& ObjectInitializer)
 	: Super(ObjectInitializer)
 {
 	bAnimBranchingPointNeedsSort = true;
-	BlendInTime = 0.25f;
-	BlendOutTime = 0.25f;
+	BlendIn.SetBlendTime(0.25f);
+	BlendOut.SetBlendTime(0.25f);
 	BlendOutTriggerTime = -1.f;
 }
 
@@ -462,6 +462,19 @@ void UAnimMontage::PostLoad()
 	{
 		ConvertBranchingPointsToAnimNotifies();
 	}
+
+	// fix up blending time deprecated variable
+	if (BlendInTime_DEPRECATED != -1.f)
+	{
+		BlendIn.SetBlendTime(BlendInTime_DEPRECATED);
+		BlendInTime_DEPRECATED = -1.f;
+	}
+
+	if(BlendOutTime_DEPRECATED != -1.f)
+	{
+		BlendOut.SetBlendTime(BlendOutTime_DEPRECATED);
+		BlendOutTime_DEPRECATED = -1.f;
+	}
 }
 
 void UAnimMontage::ConvertBranchingPointsToAnimNotifies()
@@ -838,14 +851,25 @@ void FAnimMontageInstance::Play(float InPlayRate)
 {
 	bPlaying = true;
 	PlayRate = InPlayRate;
-	if( Montage )
-	{
-		BlendTime = Montage->BlendInTime * DefaultBlendTimeMultiplier;
-	}
-	DesiredWeight = 1.f;
+
+	// if this doesn't exist, nothing works
+	ensure(Montage);
+	
+	// set blend option
+	float CurrentWeight = Blend.GetBlendedValue();
+	InitializeBlend(Montage->BlendIn);	
+	Blend.SetBlendTime(Montage->BlendIn.GetBlendTime() * DefaultBlendTimeMultiplier);
+	Blend.SetValueRange(CurrentWeight, 1.f);
 }
 
-void FAnimMontageInstance::Stop(float BlendOutDuration, bool bInterrupt)
+void FAnimMontageInstance::InitializeBlend(const FAlphaBlend& InAlphaBlend)
+{
+	Blend.SetBlendOption(InAlphaBlend.BlendOption);
+	Blend.CustomCurve = InAlphaBlend.CustomCurve;
+	Blend.SetBlendTime(InAlphaBlend.GetBlendTime());
+}
+
+void FAnimMontageInstance::Stop(const FAlphaBlend& InBlendOut, bool bInterrupt)
 {
 	// overwrite bInterrupted if it hasn't already interrupted
 	// once interrupted, you don't go back to non-interrupted
@@ -854,15 +878,15 @@ void FAnimMontageInstance::Stop(float BlendOutDuration, bool bInterrupt)
 		bInterrupted = bInterrupt;
 	}
 
-	if (DesiredWeight > 0.f)
+	// if desired weight is > 0.f, turn that off
+	if (Blend.GetDesiredValue() > 0.f)
 	{
-		DesiredWeight = 0.f;
+		Blend.SetDesiredValue(0.f);
 		if (Montage)
 		{
-			// do not use default Montage->BlendOut  Time
-			// depending on situation, the BlendOut time changes
-			// check where this function gets called and see how we calculate BlendTime
-			BlendTime = BlendOutDuration;
+			// do not use default Montage->BlendOut 
+			// depending on situation, the BlendOut time can change 
+			InitializeBlend(InBlendOut);
 
 			if (UAnimInstance* Inst = AnimInstance.Get())
 			{
@@ -877,13 +901,23 @@ void FAnimMontageInstance::Stop(float BlendOutDuration, bool bInterrupt)
 		// it is already stopped, but new montage blendtime is shorter than what 
 		// I'm blending out, that means this needs to readjust blendtime
 		// that way we don't accumulate old longer blendtime for newer montage to play
-		if (BlendOutDuration < BlendTime)
+		if (InBlendOut.GetBlendTime() < Blend.GetBlendTime())
 		{
-			BlendTime = BlendOutDuration;
+			// I don't know if also using inBlendOut is better than
+			// currently set up blend option, but it might be worse to switch between 
+			// blending out, but it is possible options in the future
+			Blend.SetBlendTime(InBlendOut.GetBlendTime());
+			// have to call this again to restart blending with new blend time
+			// we don't change blend options
+			Blend.SetDesiredValue(0.f);
 		}
 	}
 
-	if (BlendTime <= 0.0f)
+	// if blending time < 0.f
+	// set the playing to be false
+	// @todo is this better to be IsComplete? 
+	// or maybe we need this for if somebody sets blend time to be 0.f
+	if (Blend.GetBlendTime() <= 0.0f)
 	{
 		bPlaying = false;
 	}
@@ -900,7 +934,8 @@ void FAnimMontageInstance::Initialize(class UAnimMontage * InMontage)
 	{
 		Montage = InMontage;
 		Position = 0.f;
-		DesiredWeight = 1.f;
+		// initialize Blend
+		Blend.SetValueRange(0.f, 1.0f);
 
 		RefreshNextPrevSections();
 	}
@@ -970,6 +1005,10 @@ void FAnimMontageInstance::Terminate()
 	// Clear any active synchronization
 	MontageSync_StopFollowing();
 	MontageSync_StopLeading();
+
+	// clear Blend curve
+	Blend.CustomCurve = NULL;
+	Blend.BlendOption = EAlphaBlendOption::Linear;
 }
 
 bool FAnimMontageInstance::JumpToSectionName(FName const & SectionName, bool bEndOfSection)
@@ -1030,7 +1069,7 @@ bool FAnimMontageInstance::SetNextSectionID(int32 const & SectionID, int32 const
 
 void FAnimMontageInstance::OnMontagePositionChanged(FName const & ToSectionName) 
 {
-	if (bPlaying && (DesiredWeight == 0.f))
+	if (bPlaying && (Blend.GetDesiredValue() == 0.f))
 	{
 		UE_LOG(LogAnimation, Warning, TEXT("Changing section on Montage (%s) to '%s' during blend out. This can cause incorrect visuals!"),
 			*GetNameSafe(Montage), *ToSectionName.ToString());
@@ -1188,14 +1227,14 @@ void FAnimMontageInstance::UpdateWeight(float DeltaTime)
 {
 	if ( IsValid() )
 	{
-		PreviousWeight = Weight;
+		PreviousWeight = Blend.GetBlendedValue();
 
 		// update weight
-		FAnimationRuntime::TickBlendWeight(DeltaTime, DesiredWeight, Weight, BlendTime);
+		Blend.Update(DeltaTime);
 
 		// Notify weight is max of previous and current as notify could have come
 		// from any point between now and last tick
-		NotifyWeight = FMath::Max(PreviousWeight, Weight);
+		NotifyWeight = FMath::Max(PreviousWeight, Blend.GetBlendedValue());
 	}
 }
 
@@ -1304,10 +1343,10 @@ void FAnimMontageInstance::Advance(float DeltaTime, struct FRootMotionMovementPa
 		}
 #endif
 
-		// if no weight, no reason to update, and if not playing, we don't need to advance
-		// this portion is to advance position
-		// If we just reached zero weight, still tick this frame to fire end of animation events.
-		if( bPlaying && (Weight > ZERO_ANIMWEIGHT_THRESH || PreviousWeight > ZERO_ANIMWEIGHT_THRESH) )
+		// with custom curves, we can't just filter by weight
+		// also if you have custom curve with longer 0, you'll likely to pause montage during that blending time
+		// I think that is a bug. It still should move, the weight might come back later. 
+		if( bPlaying )
 		{
 			const float CombinedPlayRate = PlayRate * Montage->RateScale;
 			const bool bPlayingForward = (CombinedPlayRate > 0.f);
@@ -1364,6 +1403,8 @@ void FAnimMontageInstance::Advance(float DeltaTime, struct FRootMotionMovementPa
 
 					const bool bHaveMoved = FMath::Abs(ActualDeltaMove) > 0.f;
 
+					const float Weight = Blend.GetBlendedValue();
+
 					if (bHaveMoved)
 					{
 						// Extract Root Motion for this time slice, and accumulate it.
@@ -1396,14 +1437,14 @@ void FAnimMontageInstance::Advance(float DeltaTime, struct FRootMotionMovementPa
 						const float DeltaTimeToEnd = DeltaPosToEnd / FMath::Abs(CombinedPlayRate);
 
 						const bool bCustomBlendOutTriggerTime = (Montage->BlendOutTriggerTime >= 0);
-						const float DefaultBlendOutTime = Montage->BlendOutTime * DefaultBlendTimeMultiplier;
+						const float DefaultBlendOutTime = Montage->BlendOut.GetBlendTime() * DefaultBlendTimeMultiplier;
 						const float BlendOutTriggerTime = bCustomBlendOutTriggerTime ? Montage->BlendOutTriggerTime : DefaultBlendOutTime;
 							
 						// ... trigger blend out if within blend out time window.
 						if (DeltaTimeToEnd <= FMath::Max<float>(BlendOutTriggerTime, KINDA_SMALL_NUMBER))
 						{
 							const float BlendOutTime = bCustomBlendOutTriggerTime ? DefaultBlendOutTime : DeltaTimeToEnd;
-							Stop(BlendOutTime, false);
+							Stop(FAlphaBlend(Montage->BlendOut, BlendOutTime), false);
 						}
 					}
 
@@ -1442,7 +1483,7 @@ void FAnimMontageInstance::Advance(float DeltaTime, struct FRootMotionMovementPa
 				else
 				{
 					// stop and leave this loop
-					Stop(Montage->BlendOutTime * DefaultBlendTimeMultiplier, false);
+					Stop(FAlphaBlend(Montage->BlendOut, Montage->BlendOut.GetBlendTime() * DefaultBlendTimeMultiplier), false);
 					break;
 				}
 			}
@@ -1450,7 +1491,7 @@ void FAnimMontageInstance::Advance(float DeltaTime, struct FRootMotionMovementPa
 	}
 
 	// If this Montage has no weight, it should be terminated.
-	if ((Weight <= ZERO_ANIMWEIGHT_THRESH) && (DesiredWeight <= ZERO_ANIMWEIGHT_THRESH))
+	if ((Blend.GetDesiredValue()<= ZERO_ANIMWEIGHT_THRESH) && (Blend.IsComplete()))
 	{
 		// nothing else to do
 		Terminate();
@@ -1638,7 +1679,9 @@ void FAnimMontageInstance::SetMatineeAnimPositionInner(FName SlotName, USkeletal
 					struct FAnimMontageInstance* PrevAnimMontageInst = AnimInst->GetActiveInstanceForMontage(*CurrentlyPlayingMontage);
 					if(PrevAnimMontageInst)
 					{
-						PrevAnimMontageInst->Weight=0.f;
+						// set weight to be 0
+						PrevAnimMontageInst->Blend.SetDesiredValue(0.f);
+						PrevAnimMontageInst->Blend.SetAlpha(1.f);
 					}
 				}
 
@@ -1651,7 +1694,9 @@ void FAnimMontageInstance::SetMatineeAnimPositionInner(FName SlotName, USkeletal
 			struct FAnimMontageInstance* AnimMontageInst = AnimInst->GetActiveInstanceForMontage(*CurrentlyPlayingMontage);
 			if(AnimMontageInst)
 			{
-				AnimMontageInst->Weight = 1.f;
+				// need to set weight to be 1
+				AnimMontageInst->Blend.SetDesiredValue(1.f);
+				AnimMontageInst->Blend.SetAlpha(1.f);
 
 				float OldMontagePosition = AnimInst->Montage_GetPosition(CurrentlyPlayingMontage.Get());
 				AnimInst->Montage_SetPosition(CurrentlyPlayingMontage.Get(), InPosition);
@@ -1693,7 +1738,9 @@ void FAnimMontageInstance::PreviewMatineeSetAnimPositionInner(FName SlotName, US
 				struct FAnimMontageInstance* PrevAnimMontageInst = AnimInst->GetActiveInstanceForMontage(*CurrentlyPlayingMontage);
 				if(PrevAnimMontageInst)
 				{
-					PrevAnimMontageInst->Weight=0.f;
+					// set weight to be 0
+					PrevAnimMontageInst->Blend.SetDesiredValue(0.f);
+					PrevAnimMontageInst->Blend.SetAlpha(1.f);
 				}
 			}
 
@@ -1706,7 +1753,9 @@ void FAnimMontageInstance::PreviewMatineeSetAnimPositionInner(FName SlotName, US
 		struct FAnimMontageInstance* AnimMontageInst = AnimInst->GetActiveInstanceForMontage(*CurrentlyPlayingMontage);
 		if(AnimMontageInst)
 		{
-			AnimMontageInst->Weight = 1.f;
+			// set weight to be 1
+			AnimMontageInst->Blend.SetDesiredValue(1.f);
+			AnimMontageInst->Blend.SetAlpha(1.f);
 
 			float OldMontagePosition = AnimInst->Montage_GetPosition(CurrentlyPlayingMontage.Get());
 			AnimInst->Montage_SetPosition(CurrentlyPlayingMontage.Get(), InPosition);
diff --git a/Engine/Source/Runtime/Engine/Private/Animation/AnimNode_StateMachine.cpp b/Engine/Source/Runtime/Engine/Private/Animation/AnimNode_StateMachine.cpp
--- a/Engine/Source/Runtime/Engine/Private/Animation/AnimNode_StateMachine.cpp
+++ b/Engine/Source/Runtime/Engine/Private/Animation/AnimNode_StateMachine.cpp
@@ -48,7 +48,6 @@ FAnimationActiveTransitionEntry::FAnimationActiveTransitionEntry(int32 NextState
 	Blend.SetBlendOption(BlendOption);
 	Blend.CustomCurve = ReferenceTransitionInfo.CustomCurve;
 	Blend.SetValueRange(0.0f, 1.0f);
-	Blend.Reset();
 }
 
 float FAnimationActiveTransitionEntry::CalculateInverseAlpha(EAlphaBlendOption BlendMode, float InFraction) const
diff --git a/Engine/Source/Runtime/Engine/Private/Animation/AnimSingleNodeInstance.cpp b/Engine/Source/Runtime/Engine/Private/Animation/AnimSingleNodeInstance.cpp
--- a/Engine/Source/Runtime/Engine/Private/Animation/AnimSingleNodeInstance.cpp
+++ b/Engine/Source/Runtime/Engine/Private/Animation/AnimSingleNodeInstance.cpp
@@ -216,8 +216,7 @@ void UAnimSingleNodeInstance::NativeInitializeAnimation()
 		}
 	}
 }
-
-void UAnimSingleNodeInstance::NativeUpdateAnimation(float DeltaTimeX)
+void UAnimSingleNodeInstance::UpdateAnimationNode(float DeltaTimeX)
 {
 	float NewPlayRate = PlayRate;
 	UAnimSequence* PreviewBasePose = NULL;
diff --git a/Engine/Source/Runtime/Engine/Private/Character.cpp b/Engine/Source/Runtime/Engine/Private/Character.cpp
--- a/Engine/Source/Runtime/Engine/Private/Character.cpp
+++ b/Engine/Source/Runtime/Engine/Private/Character.cpp
@@ -1329,7 +1329,7 @@ void ACharacter::StopAnimMontage(class UAnimMontage* AnimMontage)
 
 	if ( bShouldStopMontage )
 	{
-		AnimInstance->Montage_Stop(MontageToStop->BlendOutTime, MontageToStop);
+		AnimInstance->Montage_Stop(MontageToStop->BlendOut.GetBlendTime(), MontageToStop);
 	}
 }
 
diff --git a/Engine/Source/Runtime/Engine/Public/AlphaBlend.h b/Engine/Source/Runtime/Engine/Public/AlphaBlend.h
--- a/Engine/Source/Runtime/Engine/Public/AlphaBlend.h
+++ b/Engine/Source/Runtime/Engine/Public/AlphaBlend.h
@@ -54,19 +54,19 @@ struct ENGINE_API FAlphaBlend
 	EAlphaBlendOption BlendOption;
 
 	UPROPERTY(EditAnywhere, Category = "Blend")
-	float BeginValue;
-
-	UPROPERTY(EditAnywhere, Category = "Blend")
-	float DesiredValue;
+	UCurveFloat* CustomCurve;
 
+protected:
 	UPROPERTY(EditAnywhere, Category = "Blend")
 	float	BlendTime;
 
-	UPROPERTY(EditAnywhere, Category = "Blend")
-	UCurveFloat* CustomCurve;
+public:
 
 	// Constructor
-	FAlphaBlend();
+	FAlphaBlend(float NewBlendTime = 0.2f);
+
+	// constructor
+	FAlphaBlend(const FAlphaBlend& Other, float NewBlendTime);
 
 	/** Update transition blend time */
 	void SetBlendTime(float InBlendTime);
@@ -74,38 +74,38 @@ struct ENGINE_API FAlphaBlend
 	/** Sets the method used to blend the value */
 	void SetBlendOption(EAlphaBlendOption InBlendOption);
 
-	/** Sets the range of values to map to the interpolation */
+	/** Sets the range of values to map to the interpolation
+	 *
+	 * @param Begin : this is current value
+	 * @param Desired : this is target value 
+	 *
+	 * This can be (0, 1) if you'd like to increase, or it can be (1, 0) if you'd like to get to 0
+	 */
 	void SetValueRange(float Begin, float Desired);
 
-	/** Sets the begin value for the blended value */
-	void SetBeginValue(float InBegin);
-
 	/** Sets the final desired value for the blended value */
 	void SetDesiredValue(float InDesired);
 
 	/** Sets the Lerp alpha value directly */
 	void SetAlpha(float InAlpha);
 
-	/** Reset to zero / restart the blend */
-	void Reset();
+	/** Update interpolation, has to be called once every frame */
+	void Update(float InDeltaTime);
 
-	/** Returns true if Target is > 0.f, or false otherwise */
-	bool GetToggleStatus();
+	/** Gets the current 0..1 alpha value. Changed to AlphaLerp to match with SetAlpha function */
+	float GetAlpha() const { return AlphaLerp; }
 
-	/** Enable (1.f) or Disable (0.f) */
-	void Toggle(bool bEnable);
+	/** Gets the current blended value */
+	float GetBlendedValue() const;
 
-	/** SetTarget, but check that we're actually setting a different target */
-	void ConditionalSetTarget(float InAlphaTarget);
+	/** Gets whether or not the blend is complete */
+	bool IsComplete() const;
 
-	/** Set Target for interpolation */
-	void SetTarget(float InAlphaTarget);
+	/** Gets the current blend time */
+	float GetBlendTime() const { return BlendTime; }
 
-	/** Update interpolation, has to be called once every frame */
-	void Update(float InDeltaTime);
-
-	/** Converts internal lerped alpha into the output alpha type */
-	float AlphaToBlendOption();
+	/** Get the current desired value */
+	float GetDesiredValue() const { return DesiredValue; }
 
 	/** Converts InAlpha from a linear 0...1 value into the output alpha described by InBlendOption 
 	 *  @param InAlpha In linear 0...1 alpha
@@ -114,15 +114,6 @@ struct ENGINE_API FAlphaBlend
 	 */
 	static float AlphaToBlendOption(float InAlpha, EAlphaBlendOption InBlendOption, UCurveFloat* InCustomCurve = nullptr);
 
-	/** Gets the current 0..1 alpha value */
-	float GetAlpha();
-
-	/** Gets the current blended value */
-	float GetBlendedValue();
-
-	/** Gets whether or not the blend is complete */
-	bool IsComplete();
-
 protected:
 	/** Internal Lerped value for Alpha */
 	UPROPERTY()
@@ -132,10 +123,6 @@ protected:
 	UPROPERTY()
 	float	AlphaBlend;
 
-	/** Target to reach */
-	UPROPERTY()
-	float	AlphaTarget;
-
 	/** Time left to reach target */
 	UPROPERTY()
 	float	BlendTimeRemaining;
@@ -143,4 +130,16 @@ protected:
 	/** The current blended value derived from the begin and desired values */
 	UPROPERTY()
 	float BlendedValue;
+
+	UPROPERTY()
+	float BeginValue;
+
+	UPROPERTY()
+	float DesiredValue;
+
+	/** Reset to zero / restart the blend */
+	void Reset();
+
+	/** Converts internal lerped alpha into the output alpha type */
+	float AlphaToBlendOption();
 };
\ No newline at end of file
diff --git a/Engine/Source/Runtime/GameplayAbilities/Private/AbilitySystemComponent_Abilities.cpp b/Engine/Source/Runtime/GameplayAbilities/Private/AbilitySystemComponent_Abilities.cpp
--- a/Engine/Source/Runtime/GameplayAbilities/Private/AbilitySystemComponent_Abilities.cpp
+++ b/Engine/Source/Runtime/GameplayAbilities/Private/AbilitySystemComponent_Abilities.cpp
@@ -2294,9 +2294,9 @@ void UAbilitySystemComponent::CurrentMontageStop(float OverrideBlendOutTime)
 
 	if (bShouldStopMontage)
 	{
-		const float BlendOutTime = (OverrideBlendOutTime >= 0.0f ? OverrideBlendOutTime : MontageToStop->BlendOutTime);
+		const float BlendOutTime = (OverrideBlendOutTime >= 0.0f ? OverrideBlendOutTime : MontageToStop->BlendOut.GetBlendTime());
 
-		AnimInstance->Montage_Stop(MontageToStop->BlendOutTime, MontageToStop);
+		AnimInstance->Montage_Stop(MontageToStop->BlendOut.GetBlendTime(), MontageToStop);
 
 		if (IsOwnerActorAuthoritative())
 		{
