Change 4.10 deprecation warnings to 4.11, due to branching 4.10 from 4.9.

diff --git a/Engine/Source/Runtime/Engine/Classes/Components/SkinnedMeshComponent.h b/Engine/Source/Runtime/Engine/Classes/Components/SkinnedMeshComponent.h
--- a/Engine/Source/Runtime/Engine/Classes/Components/SkinnedMeshComponent.h
+++ b/Engine/Source/Runtime/Engine/Classes/Components/SkinnedMeshComponent.h
@@ -338,7 +338,7 @@ private:
 
 public:
 
-	DEPRECATED(4.10, "bChartDistanceFactor is no longer useful, please remove references to it")
+	DEPRECATED(4.11, "bChartDistanceFactor is no longer useful, please remove references to it")
 	uint32 bChartDistanceFactor:1;
 
 	/** Whether or not we can highlight selected sections - this should really only be done in the editor */
diff --git a/Engine/Source/Runtime/Engine/Classes/Components/SplineMeshComponent.h b/Engine/Source/Runtime/Engine/Classes/Components/SplineMeshComponent.h
--- a/Engine/Source/Runtime/Engine/Classes/Components/SplineMeshComponent.h
+++ b/Engine/Source/Runtime/Engine/Classes/Components/SplineMeshComponent.h
@@ -180,7 +180,7 @@ class ENGINE_API USplineMeshComponent : public UStaticMeshComponent, public IInt
 	/** Called to notify render thread and possibly collision of a change in spline params or mesh */
 	void UpdateRenderStateAndCollision();
 
-	DEPRECATED(4.10, "This method has been renamed to UpdateRenderStateAndCollision, but use of UpdateMesh is preferred")
+	DEPRECATED(4.11, "This method has been renamed to UpdateRenderStateAndCollision, but use of UpdateMesh is preferred")
 	void MarkSplineParamsDirty() { UpdateRenderStateAndCollision(); }
 
 	/** Update the collision and render state on the spline mesh following changes to its geometry */
diff --git a/Engine/Source/Runtime/Engine/Classes/GameFramework/GameMode.h b/Engine/Source/Runtime/Engine/Classes/GameFramework/GameMode.h
--- a/Engine/Source/Runtime/Engine/Classes/GameFramework/GameMode.h
+++ b/Engine/Source/Runtime/Engine/Classes/GameFramework/GameMode.h
@@ -640,7 +640,7 @@ public:
 	virtual void GetSeamlessTravelActorList(bool bToEntry, TArray<AActor*>& ActorList);
 
 	/** Allow the game to specify a place for clients to download MapName */
-	DEPRECATED(4.10, "Override GameWelcomePlayer() instead")
+	DEPRECATED(4.11, "Override GameWelcomePlayer() instead")
 	virtual FString GetRedirectURL(const FString& MapName) const;
 
 	/** allows game to send network messages to provide more information to the client joining the game via NMT_GameSpecific
