#ue4 - Disabled RVO avoidance registration on network clients, and made a note of it only running on the server. It already did not update velocities on clients anyway.

diff --git a/Engine/Source/Runtime/Engine/Classes/GameFramework/CharacterMovementComponent.h b/Engine/Source/Runtime/Engine/Classes/GameFramework/CharacterMovementComponent.h
--- a/Engine/Source/Runtime/Engine/Classes/GameFramework/CharacterMovementComponent.h
+++ b/Engine/Source/Runtime/Engine/Classes/GameFramework/CharacterMovementComponent.h
@@ -707,7 +707,7 @@ public:
 	UPROPERTY()
 	uint32 bWantsToLeaveNavWalking:1;
 
-	/** If set, component will use RVO avoidance */
+	/** If set, component will use RVO avoidance. This only runs on the server. */
 	UPROPERTY(Category="Character Movement: Avoidance", EditAnywhere, BlueprintReadOnly)
 	uint32 bUseRVOAvoidance:1;
 
diff --git a/Engine/Source/Runtime/Engine/Private/Components/CharacterMovementComponent.cpp b/Engine/Source/Runtime/Engine/Private/Components/CharacterMovementComponent.cpp
--- a/Engine/Source/Runtime/Engine/Private/Components/CharacterMovementComponent.cpp
+++ b/Engine/Source/Runtime/Engine/Private/Components/CharacterMovementComponent.cpp
@@ -345,6 +345,11 @@ void UCharacterMovementComponent::PostEditChangeProperty(FPropertyChangedEvent&
 
 void UCharacterMovementComponent::OnRegister()
 {
+	if (bUseRVOAvoidance && GetNetMode() == NM_Client)
+	{
+		bUseRVOAvoidance = false;
+	}
+
 	Super::OnRegister();
 
 #if WITH_EDITOR
