Fix compile error on clang
Merging

//depot/UE4-Orion/Engine/Source/Runtime/Engine/Private/AlphaBlend.cpp

to //depot/UE4-Orion/Engine/Source/Runtime/Engine/Private/AlphaBlend.cpp

diff --git a/Engine/Source/Runtime/Engine/Private/AlphaBlend.cpp b/Engine/Source/Runtime/Engine/Private/AlphaBlend.cpp
--- a/Engine/Source/Runtime/Engine/Private/AlphaBlend.cpp
+++ b/Engine/Source/Runtime/Engine/Private/AlphaBlend.cpp
@@ -5,20 +5,20 @@
 
 FAlphaBlend::FAlphaBlend(float NewBlendTime) 
 	: BlendOption(EAlphaBlendOption::Linear)
+	, CustomCurve(nullptr)
+	, BlendTime(NewBlendTime)
 	, BeginValue(0.0f)
 	, DesiredValue(1.0f)
-	, BlendTime(NewBlendTime)
-	, CustomCurve(nullptr)
 {
 	Reset();
 }
 
 FAlphaBlend::FAlphaBlend(const FAlphaBlend& Other, float NewBlendTime)
 	: BlendOption(Other.BlendOption)
+	, CustomCurve(Other.CustomCurve)
+	, BlendTime(NewBlendTime)
 	, BeginValue(Other.BeginValue)
 	, DesiredValue(Other.DesiredValue)
-	, BlendTime(NewBlendTime)
-	, CustomCurve(Other.CustomCurve)
 {
 	Reset();
 }
