Transactions with the UEdGraph will preserve the graph's RF_Public flag.

diff --git a/Engine/Source/Runtime/Engine/Classes/EdGraph/EdGraph.h b/Engine/Source/Runtime/Engine/Classes/EdGraph/EdGraph.h
--- a/Engine/Source/Runtime/Engine/Classes/EdGraph/EdGraph.h
+++ b/Engine/Source/Runtime/Engine/Classes/EdGraph/EdGraph.h
@@ -105,6 +105,7 @@ public:
 
 #if WITH_EDITORONLY_DATA
 	// Begin UObject interface
+	virtual void Serialize( FArchive& Ar ) override;
 	virtual void PostInitProperties() override;
 	virtual void PostLoad() override;
 	// End UObject interface
diff --git a/Engine/Source/Runtime/Engine/Private/EdGraph/EdGraph.cpp b/Engine/Source/Runtime/Engine/Private/EdGraph/EdGraph.cpp
--- a/Engine/Source/Runtime/Engine/Private/EdGraph/EdGraph.cpp
+++ b/Engine/Source/Runtime/Engine/Private/EdGraph/EdGraph.cpp
@@ -96,6 +96,32 @@ void UEdGraph::PostInitProperties()
 	}
 }
 
+void UEdGraph::Serialize( FArchive& Ar )
+{
+	Super::Serialize(Ar);
+	// Keep track of RF_Public
+	if( Ar.IsTransacting() )
+	{
+		bool bIsPublic = HasAnyFlags(RF_Public);
+		if( Ar.IsLoading() )
+		{
+			Ar << bIsPublic;
+			if (bIsPublic)
+			{
+				SetFlags( RF_Public );
+			}
+			else
+			{
+				ClearFlags( RF_Public );
+			}
+		}
+		else if( Ar.IsSaving() )
+		{
+			Ar << bIsPublic;
+		}
+	}
+}
+
 void UEdGraph::PostLoad()
 {
 	Super::PostLoad();
