Engine: Added a number of useful particle system properties as asset registry searchable data

New columns:
- FixedBoundsSize (maximum component of the fixed bounds or None if it does not not use fixed bounds)
- NumEmitters (number of emitter modules)
- NumLODs (number of levels of detail)
- LODMethod (method of computing which LOD to use)
- LODDistanceCheckTime (interval of LOD checking when LOD mode is automatic)

Also removed unused _MAX entries in some particle system enums, and added comments on their members

Note: Unloaded particle systems will not have values for these columns until they are resaved!

[INTEGRATE] Change 2688884 by [REDACTED_PERSON][REDACTED_HANDLE] on 2015/09/11 17:01:05

diff --git a/Engine/Source/Runtime/Engine/Classes/Particles/ParticleSystem.h b/Engine/Source/Runtime/Engine/Classes/Particles/ParticleSystem.h
--- a/Engine/Source/Runtime/Engine/Classes/Particles/ParticleSystem.h
+++ b/Engine/Source/Runtime/Engine/Classes/Particles/ParticleSystem.h
@@ -14,8 +14,7 @@ enum EParticleSystemUpdateMode
 	/** RealTime	- update via the delta time passed in				*/
 	EPSUM_RealTime UMETA(DisplayName="Real-Time"),
 	/** FixedTime	- update via a fixed time step						*/
-	EPSUM_FixedTime UMETA(DisplayName="Fixed-Time"),
-	EPSUM_MAX,
+	EPSUM_FixedTime UMETA(DisplayName="Fixed-Time")
 };
 
 /**
@@ -24,10 +23,14 @@ enum EParticleSystemUpdateMode
 UENUM()
 enum ParticleSystemLODMethod
 {
+	// Automatically set the LOD level, checking every LODDistanceCheckTime seconds.
 	PARTICLESYSTEMLODMETHOD_Automatic UMETA(DisplayName="Automatic"),
-	PARTICLESYSTEMLODMETHOD_DirectSet UMETA(DisplayName="Direct Set"),
-	PARTICLESYSTEMLODMETHOD_ActivateAutomatic UMETA(DisplayName="Activate Automatic"),
-	PARTICLESYSTEMLODMETHOD_MAX,
+
+	// LOD level is directly set by the game code.
+	PARTICLESYSTEMLODMETHOD_DirectSet UMETA(DisplayName = "Direct Set"),
+
+	// LOD level is determined at Activation time, then left alone unless directly set by game code.
+	PARTICLESYSTEMLODMETHOD_ActivateAutomatic UMETA(DisplayName = "Activate Automatic")
 };
 
 /** Occlusion method enumeration */
@@ -39,8 +42,7 @@ enum EParticleSystemOcclusionBoundsMethod
 	/** Use the bounds of the particle system component when determining occlusion */
 	EPSOBM_ParticleBounds UMETA(DisplayName="Particle Bounds"),
 	/** Use the custom occlusion bounds when determining occlusion */
-	EPSOBM_CustomBounds UMETA(DisplayName="Custom Bounds"),
-	EPSOBM_MAX,
+	EPSOBM_CustomBounds UMETA(DisplayName="Custom Bounds")
 };
 
 /** Structure containing per-LOD settings that pertain to the entire UParticleSystem. */
@@ -86,7 +88,7 @@ struct FNamedEmitterMaterial
 	}
 
 	UPROPERTY(EditAnywhere, Category = NamedMaterial)
-		FName Name;
+	FName Name;
 
 	UPROPERTY(EditAnywhere, Category = NamedMaterial)
 	class UMaterialInterface* Material;
@@ -159,7 +161,7 @@ class UParticleSystem : public UObject
 	/**
 	 *	How often (in seconds) the system should perform the LOD distance check.
 	 */
-	UPROPERTY(EditAnywhere, Category=LOD)
+	UPROPERTY(EditAnywhere, Category=LOD, AssetRegistrySearchable)
 	float LODDistanceCheckTime;
 
 	/**
diff --git a/Engine/Source/Runtime/Engine/Private/Particles/ParticleComponents.cpp b/Engine/Source/Runtime/Engine/Private/Particles/ParticleComponents.cpp
--- a/Engine/Source/Runtime/Engine/Private/Particles/ParticleComponents.cpp
+++ b/Engine/Source/Runtime/Engine/Private/Particles/ParticleComponents.cpp
@@ -2043,6 +2043,33 @@ void UParticleSystem::UpdateColorModuleClampAlpha(UParticleModuleColorBase* Colo
 void UParticleSystem::GetAssetRegistryTags(TArray<FAssetRegistryTag>& OutTags) const
 {
 	OutTags.Add( FAssetRegistryTag("HasGPUEmitter", HasGPUEmitter() ? TEXT("True") : TEXT("False"), FAssetRegistryTag::TT_Alphabetical) );
+
+	const float BoundsSize = FixedRelativeBoundingBox.GetSize().GetMax();
+	OutTags.Add(FAssetRegistryTag("FixedBoundsSize", bUseFixedRelativeBoundingBox ? FString::Printf(TEXT("%.2f"), BoundsSize) : FString(TEXT("None")), FAssetRegistryTag::TT_Numerical));
+
+	OutTags.Add(FAssetRegistryTag("NumEmitters", LexicalConversion::ToString(Emitters.Num()), FAssetRegistryTag::TT_Numerical));
+
+	OutTags.Add(FAssetRegistryTag("NumLODs", LexicalConversion::ToString(LODDistances.Num()), FAssetRegistryTag::TT_Numerical));
+
+	// Done here instead of as an AssetRegistrySearchable string to avoid the long prefix on the enum value string
+	FString LODMethodString = TEXT("Unknown");
+	switch (LODMethod)
+	{
+	case PARTICLESYSTEMLODMETHOD_Automatic:
+		LODMethodString = TEXT("Automatic");
+		break;
+	case PARTICLESYSTEMLODMETHOD_DirectSet:
+		LODMethodString = TEXT("DirectSet");
+		break;
+	case PARTICLESYSTEMLODMETHOD_ActivateAutomatic:
+		LODMethodString = TEXT("Activate Automatic");
+		break;
+	default:
+		check(false); // Missing enum entry
+		break;
+	}
+	OutTags.Add(FAssetRegistryTag("LODMethod", LODMethodString, FAssetRegistryTag::TT_Alphabetical));
+
 	Super::GetAssetRegistryTags(OutTags);
 }
 
