Clean up parallel cloth code so that it's possible to use even if ragdolls use cloth.
- Moved enums out of physics settings header to avoid future engine recompile
- Remove cloth string comparison/replace code

diff --git a/Engine/Source/Runtime/Engine/Classes/Components/SkeletalMeshComponent.h b/Engine/Source/Runtime/Engine/Classes/Components/SkeletalMeshComponent.h
--- a/Engine/Source/Runtime/Engine/Classes/Components/SkeletalMeshComponent.h
+++ b/Engine/Source/Runtime/Engine/Classes/Components/SkeletalMeshComponent.h
@@ -1290,6 +1290,7 @@ private:
 
 	void ParallelTickClothing(float DeltaTime, FClothSimulationContext& ClothSimulationContext) const;
 	void ParallelUpdateClothState(float DeltaTime, FClothSimulationContext& ClothSimulationContext) const;
+	void PrepareCloth();
 
 	friend class FParallelBlendPhysicsCompletionTask;
 	void CompleteParallelBlendPhysics();
@@ -1297,6 +1298,7 @@ private:
 
 	friend class FTickClothingTask;
 	friend class FClothManager;
+	friend class FClothManagerData;
 	void PerformTickClothing(float DeltaTime);
 
 	// these are deprecated variables from removing SingleAnimSkeletalComponent
diff --git a/Engine/Source/Runtime/Engine/Classes/Engine/SkeletalMesh.h b/Engine/Source/Runtime/Engine/Classes/Engine/SkeletalMesh.h
--- a/Engine/Source/Runtime/Engine/Classes/Engine/SkeletalMesh.h
+++ b/Engine/Source/Runtime/Engine/Classes/Engine/SkeletalMesh.h
@@ -478,6 +478,10 @@ struct FClothingAssetData
 	 * Num of this array means LOD number of clothing physical meshes 
 	 */
 	TArray<FClothVisualizationInfo> ClothVisualizationInfos;
+
+	/** Apex stores only the bones that cloth needs. We need a mapping from apex bone index to UE bone index. */
+	TArray<int32> ApexToUnrealBoneMapping;
+
 	/** currently mapped morph target name */
 	FName PreparedMorphTargetName;
 
diff --git a/Engine/Source/Runtime/Engine/Classes/PhysicsEngine/PhysicsSettings.h b/Engine/Source/Runtime/Engine/Classes/PhysicsEngine/PhysicsSettings.h
--- a/Engine/Source/Runtime/Engine/Classes/PhysicsEngine/PhysicsSettings.h
+++ b/Engine/Source/Runtime/Engine/Classes/PhysicsEngine/PhysicsSettings.h
@@ -6,6 +6,7 @@
 
 #pragma once
 
+#include "PhysicsSettingsEnums.h"
 #include "Engine/DeveloperSettings.h"
 #include "PhysicsSettings.generated.h"
 
@@ -33,22 +34,6 @@ struct FPhysicalSurfaceName
 	{}
 };
 
-UENUM()
-namespace EFrictionCombineMode
-{
-	enum Type
-	{
-		/** Uses the average value of the materials touching: (a+b)/2 */
-		Average = 0,	
-		/** Uses the minimum value of the materials touching: min(a,b) */
-		Min = 1,		
-		/** Uses the product of the values of the materials touching: a*b */
-		Multiply = 2,	
-		/** Uses the maximum value of materials touching: max(a,b) */
-		Max = 3
-	};
-}
-
 UENUM()
 namespace ESettingsDOF
 {
@@ -177,11 +162,6 @@ class ENGINE_API UPhysicsSettings : public UDeveloperSettings
 	UPROPERTY(config, EditAnywhere, Category = Simulation)
 	bool bDisableCCD;
 
-	/**
-	*  If true cloth data will be updated in parallel. With this mode MorphTarget and child collision is not supported.*/
-	bool bParallelCloth;
-
-
 	/** Max Physics Delta Time to be clamped. */
 	UPROPERTY(config, EditAnywhere, meta=(ClampMin="0.0013", UIMin = "0.0013", ClampMax="1.0", UIMax="1.0"), Category=Framerate)
 	float MaxPhysicsDeltaTime;
diff --git a/Engine/Source/Runtime/Engine/Classes/PhysicsEngine/PhysicsSettingsEnums.h b/Engine/Source/Runtime/Engine/Classes/PhysicsEngine/PhysicsSettingsEnums.h
new file mode 100644
--- /dev/null
+++ b/Engine/Source/Runtime/Engine/Classes/PhysicsEngine/PhysicsSettingsEnums.h
@@ -0,0 +1,24 @@
+// Copyright 1998-2015 Epic Games, Inc. All Rights Reserved.
+
+/*=============================================================================
+	PhysicsSettingsEnums.h: Declares the PhysicsSettings enums
+=============================================================================*/
+
+#pragma once
+#include "PhysicsSettingsEnums.generated.h"
+
+UENUM()
+namespace EFrictionCombineMode
+{
+	enum Type
+	{
+		/** Uses the average value of the materials touching: (a+b)/2 */
+		Average = 0,	
+		/** Uses the minimum value of the materials touching: min(a,b) */
+		Min = 1,		
+		/** Uses the product of the values of the materials touching: a*b */
+		Multiply = 2,	
+		/** Uses the maximum value of materials touching: max(a,b) */
+		Max = 3
+	};
+}
\ No newline at end of file
diff --git a/Engine/Source/Runtime/Engine/Private/Components/SkeletalMeshComponent.cpp b/Engine/Source/Runtime/Engine/Private/Components/SkeletalMeshComponent.cpp
--- a/Engine/Source/Runtime/Engine/Private/Components/SkeletalMeshComponent.cpp
+++ b/Engine/Source/Runtime/Engine/Private/Components/SkeletalMeshComponent.cpp
@@ -172,33 +172,8 @@ void USkeletalMeshComponent::RegisterComponentTickFunctions(bool bRegister)
 	UpdatePreClothTickRegisteredState();
 }
 
-static TAutoConsoleVariable<int32> CVarAsyncCloth(
-	TEXT("p.ParallelCloth"),
-	0,
-	TEXT("Whether cloth prep and simulation is done in off the game thread."));
-
-
 void USkeletalMeshComponent::RegisterPreClothTick(bool bRegister)
 {
-	if(CVarAsyncCloth.GetValueOnGameThread())
-	{
-		UPhysicsSettings::Get()->bParallelCloth = true;
-	}else 
-	{
-		UPhysicsSettings::Get()->bParallelCloth = false;
-	}
-
-#if WITH_APEX_CLOTHING
-	if(bRegister && UPhysicsSettings::Get()->bParallelCloth)	//If we're using parallel cloth we break the cloth prep code out of the game thread. To do this we must register with the ClothManager which will do the work async
-	{
-		if(FPhysScene* PhysScene = GetWorld()->GetPhysicsScene())
-		{
-			FClothManager* ClothManager = PhysScene->GetClothManager();
-			ClothManager->RegisterForClothThisFrame(this);
-		}
-	}
-#endif
-
 	if (bRegister != PreClothTickFunction.IsTickFunctionRegistered())
 	{
 		if (bRegister)
@@ -1241,6 +1216,18 @@ void USkeletalMeshComponent::RefreshBoneTransforms(FActorComponentTickFunction*
 	}
 }
 
+void USkeletalMeshComponent::PrepareCloth()
+{
+	//TODO: logic about if to go in StartPhysics or PreCloth
+	if (FPhysScene* PhysScene = GetWorld()->GetPhysicsScene())
+	{
+		FClothManager* ClothManager = PhysScene->GetClothManager();
+		bool bClothNeedsPhysics = BodyInstance.bSimulatePhysics || IsAnySimulatingPhysics();	//TODO: this errs on the side of simulating later. We may want to optimize this so that it's only needed when cloth bodies are simulating
+		PrepareClothSchedule PrepareSchedule = bClothNeedsPhysics ? PrepareClothSchedule::WaitOnPhysics : PrepareClothSchedule::IgnorePhysics;
+		ClothManager->RegisterForPrepareCloth(this, PrepareSchedule);
+	}
+}
+
 void USkeletalMeshComponent::PostAnimEvaluation(FAnimationEvaluationContext& EvaluationContext)
 {
 	FBlendedCurve EvaluatedCurve = EvaluationContext.Curve;
@@ -1303,6 +1290,8 @@ void USkeletalMeshComponent::PostAnimEvaluation(FAnimationEvaluationContext& Eva
 
 
 	MarkRenderDynamicDataDirty();
+
+	PrepareCloth();
 }
 
 //
diff --git a/Engine/Source/Runtime/Engine/Private/PhysicsEngine/ClothManager.cpp b/Engine/Source/Runtime/Engine/Private/PhysicsEngine/ClothManager.cpp
--- a/Engine/Source/Runtime/Engine/Private/PhysicsEngine/ClothManager.cpp
+++ b/Engine/Source/Runtime/Engine/Private/PhysicsEngine/ClothManager.cpp
@@ -5,17 +5,21 @@
 #include "PhysicsPublic.h"
 #include "Components/SkeletalMeshComponent.h"
 
-FClothManager::FClothManager(UWorld* AssociatedWorld, FPhysScene* InPhysScene)
- : PhysScene(InPhysScene) 
+FClothManager::FClothManager(UWorld* AssociatedWorld)
  {
-	ClothManagerTickFunction.bCanEverTick = true;
-	ClothManagerTickFunction.Target = this;
-	ClothManagerTickFunction.TickGroup = TG_StartPhysics;
-	ClothManagerTickFunction.bRunOnAnyThread = true;	//this value seems to be ignored
-	ClothManagerTickFunction.RegisterTickFunction(AssociatedWorld->PersistentLevel);
+	PrepareClothDataArray[(int32)PrepareClothSchedule::IgnorePhysics].TickFunction.TickGroup = TG_StartPhysics;
+	PrepareClothDataArray[(int32)PrepareClothSchedule::WaitOnPhysics].TickFunction.TickGroup = TG_PreCloth;
+
+	for(FClothManagerData& PrepareClothData : PrepareClothDataArray)
+	{
+		PrepareClothData.TickFunction.bCanEverTick = true;
+		PrepareClothData.TickFunction.Target = &PrepareClothData;
+		PrepareClothData.TickFunction.bRunOnAnyThread = true;	//this value seems to be ignored
+		PrepareClothData.TickFunction.RegisterTickFunction(AssociatedWorld->PersistentLevel);
+	}
  }
 
-void FClothManager::Tick(float DeltaTime)
+void FClothManagerData::PrepareCloth(float DeltaTime)
 {
 #if WITH_APEX_CLOTHING
 	if(SkeletalMeshComponents.Num())
@@ -30,33 +34,32 @@ void FClothManager::Tick(float DeltaTime)
 		}
 
 		SkeletalMeshComponents.Empty(SkeletalMeshComponents.Num());
-		PhysScene->StartCloth();
 		IsPreparingCloth.AtomicSet(false);
 	}
 #endif
 }
 
-void FClothManager::RegisterForClothThisFrame(USkeletalMeshComponent* SkeletalMeshComponent)
+void FClothManager::RegisterForPrepareCloth(USkeletalMeshComponent* SkeletalMeshComponent, PrepareClothSchedule PrepSchedule)
 {
-	SkeletalMeshComponents.Add(SkeletalMeshComponent);
+	check(PrepSchedule != PrepareClothSchedule::MAX);
+	PrepareClothDataArray[(int32)PrepSchedule].SkeletalMeshComponents.Add(SkeletalMeshComponent);
 }
 
-
-class FClothManagerTickTask
+class FPrepareClothTickTask
 {
-	FClothManager* ClothManager;
+	FClothManagerData* PrepareClothData;
 	float DeltaTime;
 
 public:
-	FClothManagerTickTask(FClothManager* InClothManager, float InDeltaTime)
-		: ClothManager(InClothManager)
+	FPrepareClothTickTask(FClothManagerData* InPrepareClothData, float InDeltaTime)
+		: PrepareClothData(InPrepareClothData)
 		, DeltaTime(InDeltaTime)
 	{
 	}
 
 	FORCEINLINE TStatId GetStatId() const
 	{
-		RETURN_QUICK_DECLARE_CYCLE_STAT(FClothManagerTickTask, STATGROUP_TaskGraphTasks);
+		RETURN_QUICK_DECLARE_CYCLE_STAT(FPrepareClothTickTask, STATGROUP_TaskGraphTasks);
 	}
 	static ENamedThreads::Type GetDesiredThread()
 	{
@@ -69,19 +72,19 @@ public:
 
 	void DoTask(ENamedThreads::Type CurrentThread, const FGraphEventRef& MyCompletionGraphEvent)
 	{
-		ClothManager->Tick(DeltaTime);
+		PrepareClothData->PrepareCloth(DeltaTime);
 	}
 };
 
 
-void FClothManagerTickFunction::ExecuteTick(float DeltaTime, enum ELevelTick TickType, ENamedThreads::Type CurrentThread, const FGraphEventRef& MyCompletionGraphEvent)
+void FPrepareClothTickFunction::ExecuteTick(float DeltaTime, enum ELevelTick TickType, ENamedThreads::Type CurrentThread, const FGraphEventRef& MyCompletionGraphEvent)
 {
-	QUICK_SCOPE_CYCLE_COUNTER(FClothManagerTickFunction_ExecuteTick);
+	QUICK_SCOPE_CYCLE_COUNTER(FPrepareClothTickFunction_ExecuteTick);
 
-	TGraphTask<FClothManagerTickTask>::CreateTask().ConstructAndDispatchWhenReady(Target, DeltaTime);
+	TGraphTask<FPrepareClothTickTask>::CreateTask().ConstructAndDispatchWhenReady(Target, DeltaTime);
 }
 
-FString FClothManagerTickFunction::DiagnosticMessage()
+FString FPrepareClothTickFunction::DiagnosticMessage()
 {
-	return TEXT("FClothManagerTickFunction");
+	return TEXT("FPrepareClothTickFunction");
 }
diff --git a/Engine/Source/Runtime/Engine/Private/PhysicsEngine/ClothManager.h b/Engine/Source/Runtime/Engine/Private/PhysicsEngine/ClothManager.h
--- a/Engine/Source/Runtime/Engine/Private/PhysicsEngine/ClothManager.h
+++ b/Engine/Source/Runtime/Engine/Private/PhysicsEngine/ClothManager.h
@@ -10,9 +10,17 @@
 class FPhysScene;
 class USkeletalMeshComponent;
 
-struct FClothManagerTickFunction : public FTickFunction
+/** Determines when prepare cloth work will be done in relation the rigid body simulation */
+enum class PrepareClothSchedule
 {
-	class FClothManager* Target;
+	IgnorePhysics = 0,	//Do not wait for rigid body sim
+	WaitOnPhysics,	//Wait for rigid body sim (e.g. ragdolls with cloth)
+	MAX
+};
+
+struct FPrepareClothTickFunction : public FTickFunction
+{
+	class FClothManagerData* Target;
 
 	// FTickFunction interface
 	virtual void ExecuteTick(float DeltaTime, enum ELevelTick TickType, ENamedThreads::Type CurrentThread, const FGraphEventRef& MyCompletionGraphEvent) override;
@@ -20,26 +28,41 @@ struct FClothManagerTickFunction : public FTickFunction
 	// End of FTickFunction interface
 };
 
+/** Break the data into class since we have multiple passes for different scheduling purposes */
+class FClothManagerData
+{
+public:
+	TArray<USkeletalMeshComponent*> SkeletalMeshComponents;
+	FPrepareClothTickFunction TickFunction;
+	FThreadSafeBool IsPreparingCloth;
+
+	/** Go through all SkeletalMeshComponents registered for work this frame and call the solver as needed*/
+	void PrepareCloth(float DeltaTime);
+};
+
 
 class FClothManager
 {
 public:
 
-	FClothManager(UWorld* AssociatedWorld, FPhysScene* InPhysScene);
-
-	/** Go through all SkeletalMeshComponents registered for work this frame and call the solver as needed*/
-	void Tick(float DeltaTime);
+	FClothManager(UWorld* AssociatedWorld);
 
 	/** Register the SkeletalMeshComponent so that it does cloth solving work this frame */
-	void RegisterForClothThisFrame(USkeletalMeshComponent* SkeletalMeshComponent);
+	void RegisterForPrepareCloth(USkeletalMeshComponent* SkeletalMeshComponent, PrepareClothSchedule PrepSchedule);
 
 	/** Whether the cloth manager is currently doing work in an async thread */
-	bool IsPreparingClothAsync() const { return IsPreparingCloth; }
+	bool IsPreparingClothAsync() const
+	{
+		bool bIsPreparingCloth = false;
+		for(const FClothManagerData& PrepareCloth : PrepareClothDataArray)
+		{
+			return bIsPreparingCloth |= PrepareCloth.IsPreparingCloth;
+		}
+
+		return bIsPreparingCloth;
+	}
 	
 private:
-	TArray<USkeletalMeshComponent*> SkeletalMeshComponents;
-	FClothManagerTickFunction ClothManagerTickFunction;
-	FPhysScene* PhysScene;
-
-	FThreadSafeBool IsPreparingCloth;
+	
+	FClothManagerData PrepareClothDataArray[(int32)PrepareClothSchedule::MAX];
 };
\ No newline at end of file
diff --git a/Engine/Source/Runtime/Engine/Private/PhysicsEngine/PhysScene.cpp b/Engine/Source/Runtime/Engine/Private/PhysicsEngine/PhysScene.cpp
--- a/Engine/Source/Runtime/Engine/Private/PhysicsEngine/PhysScene.cpp
+++ b/Engine/Source/Runtime/Engine/Private/PhysicsEngine/PhysScene.cpp
@@ -196,7 +196,7 @@ void FPhysScene::SetOwningWorld(UWorld* InOwningWorld)
 
 	if (OwningWorld)
 	{
-		ClothManager = new FClothManager(OwningWorld, this);
+		ClothManager = new FClothManager(OwningWorld);
 	}
 }
 
diff --git a/Engine/Source/Runtime/Engine/Private/PhysicsEngine/PhysicsSettings.cpp b/Engine/Source/Runtime/Engine/Private/PhysicsEngine/PhysicsSettings.cpp
--- a/Engine/Source/Runtime/Engine/Private/PhysicsEngine/PhysicsSettings.cpp
+++ b/Engine/Source/Runtime/Engine/Private/PhysicsEngine/PhysicsSettings.cpp
@@ -22,7 +22,6 @@ UPhysicsSettings::UPhysicsSettings(const FObjectInitializer& ObjectInitializer)
 	, MaxAngularVelocity(3600)	//10 revolutions per second
 	, bSimulateSkeletalMeshOnDedicatedServer(true)
 	, bDefaultHasComplexCollision(true)
-	, bParallelCloth(false)
 	, MaxPhysicsDeltaTime(1.f / 30.f)
 	, bSubstepping(false)
 	, bSubsteppingAsync(false)
diff --git a/Engine/Source/Runtime/Engine/Private/SkeletalMesh.cpp b/Engine/Source/Runtime/Engine/Private/SkeletalMesh.cpp
--- a/Engine/Source/Runtime/Engine/Private/SkeletalMesh.cpp
+++ b/Engine/Source/Runtime/Engine/Private/SkeletalMesh.cpp
@@ -2425,6 +2425,28 @@ bool USkeletalMesh::IsReadyForFinishDestroy()
 	return ReleaseResourcesFence.IsFenceComplete();
 }
 
+#if WITH_APEX_CLOTHING
+// convert a bone name from APEX stype to FBX style
+FName GetConvertedBoneName(NxClothingAsset* ApexClothingAsset, int32 BoneIndex);
+
+
+void BuildApexToUnrealBoneMapping(const USkeletalMesh* SkeletalMesh, FClothingAssetData& ClothingAsset)
+{
+	NxClothingAsset* ApexClothingAsset = ClothingAsset.ApexClothingAsset;
+	uint32 NumUsedBones = ClothingAsset.ApexClothingAsset->getNumUsedBones();
+	
+	check(ClothingAsset.ApexToUnrealBoneMapping.Num() == 0);
+	ClothingAsset.ApexToUnrealBoneMapping.AddUninitialized(NumUsedBones);
+
+	for (uint32 Index = 0; Index < NumUsedBones; Index++)
+	{
+		FName BoneName = GetConvertedBoneName(ApexClothingAsset, Index);
+		int32 BoneIndex = SkeletalMesh->RefSkeleton.FindBoneIndex(BoneName);
+		ClothingAsset.ApexToUnrealBoneMapping[Index] = BoneIndex;
+	}
+}
+#endif
+
 void USkeletalMesh::Serialize( FArchive& Ar )
 {
 	DECLARE_SCOPE_CYCLE_COUNTER( TEXT("USkeletalMesh::Serialize"), STAT_SkeletalMesh_Serialize, STATGROUP_LoadTime );
@@ -2517,6 +2539,12 @@ void USkeletalMesh::Serialize( FArchive& Ar )
 		for( int32 Idx=0;Idx<ClothingAssets.Num();Idx++ )
 		{
 			Ar << ClothingAssets[Idx];
+#if WITH_APEX_CLOTHING
+			if(Ar.IsLoading())
+			{
+				BuildApexToUnrealBoneMapping(this, ClothingAssets[Idx]);
+			}
+#endif
 		}
 	}
 
diff --git a/Engine/Source/Runtime/Engine/Private/SkeletalMeshComponentPhysics.cpp b/Engine/Source/Runtime/Engine/Private/SkeletalMeshComponentPhysics.cpp
--- a/Engine/Source/Runtime/Engine/Private/SkeletalMeshComponentPhysics.cpp
+++ b/Engine/Source/Runtime/Engine/Private/SkeletalMeshComponentPhysics.cpp
@@ -1992,33 +1992,8 @@ bool USkeletalMeshComponent::ComponentOverlapMultiImpl(TArray<struct FOverlapRes
 
 #if WITH_APEX_CLOTHING
 
-// convert a bone name from APEX stype to FBX style
-static FName GetConvertedBoneName(NxClothingAsset* ApexClothingAsset, int32 BoneIndex)
-{
-	//QUICK_SCOPE_CYCLE_COUNTER(STAT_GetConvertedBoneName);
-	/*static FString Work;
-	check(IsInGameThread() && !Work.Len());
-	// FIXME, this still needs work, string conversion and hashing
-
-	const char* Bone = ApexClothingAsset->getBoneName(BoneIndex);
-	if (Bone)
-	{
-		int32 SrcLen = FCStringAnsi::Strlen(Bone) + 1;
-		int32 DestLen = FPlatformString::ConvertedLength<TCHAR>(Bone, SrcLen);
-		TArray<TCHAR>& Array = Work.GetCharArray();
-		Array.AddUninitialized(DestLen);
-		FPlatformString::Convert(Array.GetData(), DestLen, Bone, SrcLen);
-
-		Work.ReplaceInline(TEXT(" "), TEXT("-"));
-		FName Result(*Work);
-		Work.Reset();
-		//check(Result == FName(*FString(ApexClothingAsset->getBoneName(BoneIndex)).Replace(TEXT(" "), TEXT("-"))));
-		return Result;
-	}
-	return NAME_None;*/
-	//Going back to allocation approach because we now call this from many threads
-	return *FString(ApexClothingAsset->getBoneName(BoneIndex)).Replace(TEXT(" "), TEXT("-"));
-}
+
+
 
 void USkeletalMeshComponent::AddClothingBounds(FBoxSphereBounds& InOutBounds) const
 {
@@ -2267,7 +2242,7 @@ bool USkeletalMeshComponent::IsExecutingParallelCloth() const
 {
 	UWorld* World = GetWorld();
 	FPhysScene* PhysScene = World ? World->GetPhysicsScene() : nullptr;
-	return UPhysicsSettings::Get()->bParallelCloth && PhysScene && PhysScene->GetClothManager()->IsPreparingClothAsync();
+	return PhysScene && PhysScene->GetClothManager()->IsPreparingClothAsync();
 }
 
 void USkeletalMeshComponent::RemoveAllClothingActors()
@@ -2650,6 +2625,11 @@ bool USkeletalMeshComponent::GetClothCollisionDataFromStaticMesh(UPrimitiveCompo
 	return bSuccess;
 }
 
+FName GetConvertedBoneName(NxClothingAsset* ApexClothingAsset, int32 BoneIndex)
+{
+	return FName(*FString(ApexClothingAsset->getBoneName(BoneIndex)).Replace(TEXT(" "), TEXT("-")));
+}
+
 void USkeletalMeshComponent::FindClothCollisions(TArray<FApexClothCollisionVolumeData>& OutCollisions)
 {
 	if (!SkeletalMesh)
@@ -3188,16 +3168,6 @@ void USkeletalMeshComponent::PreClothTick(float DeltaTime, FTickFunction& ThisTi
 			BlendInPhysics(ThisTickFunction);
 		}
 	}
-
-
-
-#if WITH_APEX_CLOTHING
-	// if skeletal mesh has clothing assets, call TickClothing
-	if (SkeletalMesh && SkeletalMesh->ClothingAssets.Num() > 0 && !UPhysicsSettings::Get()->bParallelCloth)
-	{
-		TickClothing(DeltaTime, ThisTickFunction);
-	}
-#endif
 }
 
 #if WITH_APEX_CLOTHING
@@ -3579,31 +3549,25 @@ void USkeletalMeshComponent::ParallelUpdateClothState(float DeltaTime, FClothSim
 	physx::apex::ClothingTeleportMode::Enum CurTeleportMode = (physx::apex::ClothingTeleportMode::Enum)ClothSimulationContext.ClothTeleportMode;
 
 	TArray<physx::PxMat44> BoneMatrices;
-	check(!BoneMatrices.Num());
 
-	for (int32 ActorIdx = 0; ActorIdx < NumActors; ActorIdx++)
+	for (int32 Index = 0; Index < ClothingActors.Num(); ++Index)
 	{
-		// skip if ClothingActor is NULL or invalid
-		if (!IsValidClothingActor(ActorIdx))
-		{
-			continue;
-		}
+		const FClothingActor& ClothingActor = ClothingActors[Index];
+		const FClothingAssetData& ClothingAsset = SkeletalMesh->ClothingAssets[Index];
 
-		ApplyWindForCloth(ClothingActors[ActorIdx]);
+		ApplyWindForCloth(ClothingActor);
 
-		NxClothingAsset* ClothingAsset = ClothingActors[ActorIdx].ParentClothingAsset;
+		const NxClothingAsset* ApexClothingAsset = ClothingActor.ParentClothingAsset;
 
-		uint32 NumUsedBones = ClothingAsset->getNumUsedBones();
+		uint32 NumUsedBones = ApexClothingAsset->getNumUsedBones();
 
 		check(!BoneMatrices.Num());
 		BoneMatrices.AddUninitialized(NumUsedBones);
 
 		for (uint32 Index = 0; Index < NumUsedBones; Index++)
 		{
-			FName BoneName = GetConvertedBoneName(ClothingActors[ActorIdx].ParentClothingAsset, Index);
-
-			int32 BoneIndex = GetBoneIndex(BoneName);
-
+			int32 BoneIndex = ClothingAsset.ApexToUnrealBoneMapping[Index];
+			
 			if (MasterPoseComponent.IsValid())
 			{
 				int32 TempBoneIndex = BoneIndex;
@@ -3631,12 +3595,12 @@ void USkeletalMeshComponent::ParallelUpdateClothState(float DeltaTime, FClothSim
 			}
 		}
 
-		NxClothingActor* ClothingActor = ClothingActors[ActorIdx].ApexClothingActor;
+		NxClothingActor* ApexClothingActor = ClothingActor.ApexClothingActor;
 
-		check(ClothingActor);
+		check(ApexClothingActor);
 
 		// if bUseInternalboneOrder is set, "NumUsedBones" works, otherwise have to use "getNumBones" 
-		ClothingActor->updateState(
+		ApexClothingActor->updateState(
 			PxGlobalPose,
 			BoneMatrices.GetData(),
 			sizeof(physx::PxMat44),
