Missing file, fix build

diff --git a/Engine/Source/Runtime/Engine/Private/PhysicsEngine/PhysLevel.cpp b/Engine/Source/Runtime/Engine/Private/PhysicsEngine/PhysLevel.cpp
--- a/Engine/Source/Runtime/Engine/Private/PhysicsEngine/PhysLevel.cpp
+++ b/Engine/Source/Runtime/Engine/Private/PhysicsEngine/PhysLevel.cpp
@@ -212,18 +212,11 @@ void UWorld::FinishPhysicsSim()
 
 void UWorld::StartClothAndAsyncSim()
 {
-	FPhysScene* PhysScene = GetPhysicsScene();
-	if (PhysScene == NULL)
-	{
-		return;
-	}
-
-	if(!UPhysicsSettings::Get()->bParallelCloth)	//if we're doing parallel cloth, the cloth sim will be ticked from cloth manager
+	if (FPhysScene* PhysScene = GetPhysicsScene())
 	{
 		PhysScene->StartCloth();
+		PhysScene->StartAsync();
 	}
-
-	PhysScene->StartAsync();
 }
 
 // the physics tick functions
