Fix shadow variable

diff --git a/Engine/Source/Runtime/Engine/Private/SkeletalMeshComponentPhysics.cpp b/Engine/Source/Runtime/Engine/Private/SkeletalMeshComponentPhysics.cpp
--- a/Engine/Source/Runtime/Engine/Private/SkeletalMeshComponentPhysics.cpp
+++ b/Engine/Source/Runtime/Engine/Private/SkeletalMeshComponentPhysics.cpp
@@ -3550,10 +3550,10 @@ void USkeletalMeshComponent::ParallelUpdateClothState(float DeltaTime, FClothSim
 
 	TArray<physx::PxMat44> BoneMatrices;
 
-	for (int32 Index = 0; Index < ClothingActors.Num(); ++Index)
+	for (int32 ClothIndex = 0; ClothIndex < ClothingActors.Num(); ++ClothIndex)
 	{
-		const FClothingActor& ClothingActor = ClothingActors[Index];
-		const FClothingAssetData& ClothingAsset = SkeletalMesh->ClothingAssets[Index];
+		const FClothingActor& ClothingActor = ClothingActors[ClothIndex];
+		const FClothingAssetData& ClothingAsset = SkeletalMesh->ClothingAssets[ClothIndex];
 
 		ApplyWindForCloth(ClothingActor);
 
