Cleaning out programmer-only comments that were harvested by Doxygen.

diff --git a/Engine/Source/Runtime/Core/Public/IOS/IOSInputInterface.h b/Engine/Source/Runtime/Core/Public/IOS/IOSInputInterface.h
--- a/Engine/Source/Runtime/Core/Public/IOS/IOSInputInterface.h
+++ b/Engine/Source/Runtime/Core/Public/IOS/IOSInputInterface.h
@@ -44,9 +44,9 @@ public:
 
 	static void QueueTouchInput(TArray<TouchInput> InTouchEvents);
 
-	// Begin Exec Interface
+	//~ Begin Exec Interface
 	virtual bool Exec(UWorld* InWorld, const TCHAR* Cmd, FOutputDevice& Ar) override;
-	// End Exec Interface
+	//~ End Exec Interface
 
 private:
 
diff --git a/Engine/Source/Runtime/Core/Public/ProfilingDebugging/MallocProfiler.h b/Engine/Source/Runtime/Core/Public/ProfilingDebugging/MallocProfiler.h
--- a/Engine/Source/Runtime/Core/Public/ProfilingDebugging/MallocProfiler.h
+++ b/Engine/Source/Runtime/Core/Public/ProfilingDebugging/MallocProfiler.h
@@ -155,7 +155,7 @@ public:
 	 */
 	virtual ~FMallocProfilerBufferedFileWriter();
 
-	// FArchive interface.
+	//~ Begin FArchive Interface.
 	virtual void Serialize( void* V, int64 Length );
 	virtual void Seek( int64 InPos );
 	virtual bool Close();
@@ -494,9 +494,9 @@ public:
 	}
 
 	
-	// Begin Exec Interface
+	//~ Begin Exec Interface
 	virtual bool Exec( UWorld* InWorld, const TCHAR* Cmd, FOutputDevice& Ar ) override;
-	// End Exec Interface
+	//~ End Exec Interface
 
 	/** 
 	 * Exec command handlers
diff --git a/Engine/Source/Runtime/CoreUObject/Private/Serialization/AsyncLoadingThread.h b/Engine/Source/Runtime/CoreUObject/Private/Serialization/AsyncLoadingThread.h
--- a/Engine/Source/Runtime/CoreUObject/Private/Serialization/AsyncLoadingThread.h
+++ b/Engine/Source/Runtime/CoreUObject/Private/Serialization/AsyncLoadingThread.h
@@ -103,11 +103,11 @@ class FAsyncLoadingThread : public FRunnable
 
 public:
 
-	// Begin FRunnable interface.
+	//~ Begin FRunnable Interface.
 	virtual bool Init();
 	virtual uint32 Run();
 	virtual void Stop();
-	// End FRunnable interface
+	//~ End FRunnable Interface
 
 	/** Returns the async loading thread singleton */
 	static FAsyncLoadingThread& Get();
diff --git a/Engine/Source/Runtime/CoreUObject/Public/Serialization/ArchiveUObject.h b/Engine/Source/Runtime/CoreUObject/Public/Serialization/ArchiveUObject.h
--- a/Engine/Source/Runtime/CoreUObject/Public/Serialization/ArchiveUObject.h
+++ b/Engine/Source/Runtime/CoreUObject/Public/Serialization/ArchiveUObject.h
@@ -95,14 +95,14 @@ public:
 		Obj->Serialize(*this);
 	}
 
-	// FArchive interface
+	//~ Begin FArchive Interface
 	COREUOBJECT_API virtual FArchive& operator<<( class FName& N ) override;
 	COREUOBJECT_API virtual FArchive& operator<<( class UObject*& Res ) override;
 	COREUOBJECT_API virtual FArchive& operator<<( FLazyObjectPtr& LazyObjectPtr ) override;
 	COREUOBJECT_API virtual FArchive& operator<<( FAssetPtr& AssetPtr ) override;
 	COREUOBJECT_API virtual FArchive& operator<<(FStringAssetReference& AssetPtr) override;
 	COREUOBJECT_API virtual FString GetArchiveName() const override;
-	// End of FArchive interface
+	//~ End FArchive Interface
 
 protected:
 	FObjectWriter(TArray<uint8>& InBytes)
@@ -128,14 +128,14 @@ public:
 		Obj->Serialize(*this);
 	}
 
-	// FArchive interface
+	//~ Begin FArchive Interface
 	COREUOBJECT_API virtual FArchive& operator<<( class FName& N ) override;
 	COREUOBJECT_API virtual FArchive& operator<<( class UObject*& Res ) override;
 	COREUOBJECT_API virtual FArchive& operator<<( FLazyObjectPtr& LazyObjectPtr ) override;
 	COREUOBJECT_API virtual FArchive& operator<<( FAssetPtr& AssetPtr ) override;
 	COREUOBJECT_API virtual FArchive& operator<<(FStringAssetReference& AssetPtr) override;
 	COREUOBJECT_API virtual FString GetArchiveName() const override;
-	// End of FArchive interface
+	//~ End FArchive Interface
 
 protected:
 	FObjectReader(TArray<uint8>& InBytes)
@@ -856,7 +856,7 @@ private:
 	const TArray<uint8>&					ObjectData;
 	int32									Offset;
 
-	// FArchive interface.
+	//~ Begin FArchive Interface.
 
 	virtual FArchive& operator<<(FName& N);
 	virtual FArchive& operator<<(UObject*& Object);
@@ -932,7 +932,7 @@ private:
 	 */
 	struct FObjectInstancingGraph*			InstanceGraph;
 
-	// FArchive interface.
+	//~ Begin FArchive Interface.
 
 	virtual FArchive& operator<<(FName& N);
 	virtual FArchive& operator<<(UObject*& Object);
@@ -1576,12 +1576,12 @@ public:
 	*/
 	FArchiveObjectCrc32();
 
-	// Begin FArchive Interface
+	//~ Begin FArchive Interface
 	virtual void Serialize(void* Data, int64 Length);
 	virtual FArchive& operator<<(class FName& Name);
 	virtual FArchive& operator<<(class UObject*& Object);
 	virtual FString GetArchiveName() const { return TEXT("FArchiveObjectCrc32"); }
-	// End FArchive Interface
+	//~ End FArchive Interface
 
 	/**
 	* Serialize the given object, calculate and return its checksum.
diff --git a/Engine/Source/Runtime/CoreUObject/Public/Serialization/ArchiveUObjectBase.h b/Engine/Source/Runtime/CoreUObject/Public/Serialization/ArchiveUObjectBase.h
--- a/Engine/Source/Runtime/CoreUObject/Public/Serialization/ArchiveUObjectBase.h
+++ b/Engine/Source/Runtime/CoreUObject/Public/Serialization/ArchiveUObjectBase.h
@@ -13,11 +13,11 @@ class COREUOBJECT_API FArchiveUObject : public FArchive
 {
 public:
 
-	// Begin FArchive Interface
+	//~ Begin FArchive Interface
 	virtual FArchive& operator<< (class FLazyObjectPtr& Value) override;
 	virtual FArchive& operator<< (class FAssetPtr& Value) override;
 	virtual FArchive& operator<< (struct FStringAssetReference& Value) override;
-	// End FArchive Interface
+	//~ End FArchive Interface
 };
 
 class UProperty;
diff --git a/Engine/Source/Runtime/Engine/Classes/AI/Navigation/AvoidanceManager.h b/Engine/Source/Runtime/Engine/Classes/AI/Navigation/AvoidanceManager.h
--- a/Engine/Source/Runtime/Engine/Classes/AI/Navigation/AvoidanceManager.h
+++ b/Engine/Source/Runtime/Engine/Classes/AI/Navigation/AvoidanceManager.h
@@ -154,9 +154,9 @@ class ENGINE_API UAvoidanceManager : public UObject, public FSelfRegisteringExec
 	void HandleToggleAvoidance( const TCHAR* Cmd, FOutputDevice& Ar );
 #endif
 	
-	// Begin FExec Interface
+	//~ Begin FExec Interface
 	virtual bool Exec(UWorld* Inworld, const TCHAR* Cmd, FOutputDevice& Ar) override;
-	// End FExec Interface
+	//~ End FExec Interface
 
 private:
 
diff --git a/Engine/Source/Runtime/Engine/Classes/AI/Navigation/NavLinkCustomComponent.h b/Engine/Source/Runtime/Engine/Classes/AI/Navigation/NavLinkCustomComponent.h
--- a/Engine/Source/Runtime/Engine/Classes/AI/Navigation/NavLinkCustomComponent.h
+++ b/Engine/Source/Runtime/Engine/Classes/AI/Navigation/NavLinkCustomComponent.h
@@ -38,22 +38,22 @@ class ENGINE_API UNavLinkCustomComponent : public UNavRelevantComponent, public
 	virtual void OnLinkMoveFinished(UPathFollowingComponent* PathComp) override;
 	// END INavLinkCustomInterface
 
-	// BEGIN UNavRelevantComponent Interface
+	//~ Begin UNavRelevantComponent Interface
 	virtual void GetNavigationData(FNavigationRelevantData& Data) const override;
 	virtual void CalcAndCacheBounds() const override;
-	// END UNavRelevantComponent Interface
+	//~ End UNavRelevantComponent Interface
 
-	// BEGIN UActorComponent Interface
+	//~ Begin UActorComponent Interface
 	virtual void OnRegister() override;
 	virtual void OnUnregister() override;
-	// END UActorComponent Interface
+	//~ End UActorComponent Interface
 
-	// BEGIN UObject Interface
+	//~ Begin UObject Interface
 	virtual void PostLoad() override;
 #if WITH_EDITOR
 	virtual void PostEditImport() override;
 #endif
-	// END UObject Interface
+	//~ End UObject Interface
 
 	/** set basic link data: end points and direction */
 	void SetLinkData(const FVector& RelativeStart, const FVector& RelativeEnd, ENavLinkDirection::Type Direction);
diff --git a/Engine/Source/Runtime/Engine/Classes/AI/Navigation/NavLinkRenderingComponent.h b/Engine/Source/Runtime/Engine/Classes/AI/Navigation/NavLinkRenderingComponent.h
--- a/Engine/Source/Runtime/Engine/Classes/AI/Navigation/NavLinkRenderingComponent.h
+++ b/Engine/Source/Runtime/Engine/Classes/AI/Navigation/NavLinkRenderingComponent.h
@@ -11,7 +11,7 @@ class ENGINE_API UNavLinkRenderingComponent : public UPrimitiveComponent
 {
 	GENERATED_UCLASS_BODY()
 		
-	// Begin UPrimitiveComponent Interface
+	//~ Begin UPrimitiveComponent Interface
 	virtual FPrimitiveSceneProxy* CreateSceneProxy() override;
 
 	/** Should recreate proxy one very update */
@@ -20,10 +20,10 @@ class ENGINE_API UNavLinkRenderingComponent : public UPrimitiveComponent
 	virtual bool ComponentIsTouchingSelectionBox(const FBox& InSelBBox, const FEngineShowFlags& ShowFlags, const bool bConsiderOnlyBSP, const bool bMustEncompassEntireComponent) const override;
 	virtual bool ComponentIsTouchingSelectionFrustum(const FConvexVolume& InFrustum, const FEngineShowFlags& ShowFlags, const bool bConsiderOnlyBSP, const bool bMustEncompassEntireComponent) const override;
 #endif
-	// End UPrimitiveComponent Interface
+	//~ End UPrimitiveComponent Interface
 
-	// Begin USceneComponent Interface
+	//~ Begin USceneComponent Interface
 	virtual FBoxSphereBounds CalcBounds(const FTransform &LocalToWorld) const override;
-	// End USceneComponent Interface
+	//~ End USceneComponent Interface
 };
 
diff --git a/Engine/Source/Runtime/Engine/Classes/AI/Navigation/NavMeshBoundsVolume.h b/Engine/Source/Runtime/Engine/Classes/AI/Navigation/NavMeshBoundsVolume.h
--- a/Engine/Source/Runtime/Engine/Classes/AI/Navigation/NavMeshBoundsVolume.h
+++ b/Engine/Source/Runtime/Engine/Classes/AI/Navigation/NavMeshBoundsVolume.h
@@ -17,14 +17,14 @@ class ANavMeshBoundsVolume : public AVolume
 	UPROPERTY(EditAnywhere, Category = Navigation)
 	FNavAgentSelector SupportedAgents;
 
-	// Begin AActor Interface
+	//~ Begin AActor Interface
 	virtual void PostRegisterAllComponents() override;
 	virtual void PostUnregisterAllComponents() override;
-	// End AActor Interface
+	//~ End AActor Interface
 #if WITH_EDITOR
-	// Begin UObject Interface
+	//~ Begin UObject Interface
 	virtual void PostEditChangeProperty( struct FPropertyChangedEvent& PropertyChangedEvent) override;
-	// End UObject Interface
+	//~ End UObject Interface
 #endif // WITH_EDITOR
 };
 
diff --git a/Engine/Source/Runtime/Engine/Classes/AI/Navigation/NavMeshRenderingComponent.h b/Engine/Source/Runtime/Engine/Classes/AI/Navigation/NavMeshRenderingComponent.h
--- a/Engine/Source/Runtime/Engine/Classes/AI/Navigation/NavMeshRenderingComponent.h
+++ b/Engine/Source/Runtime/Engine/Classes/AI/Navigation/NavMeshRenderingComponent.h
@@ -11,24 +11,24 @@ class ENGINE_API UNavMeshRenderingComponent : public UPrimitiveComponent
 	GENERATED_UCLASS_BODY()
 
 public:
-	// Begin UObject Interface
+	//~ Begin UObject Interface
 	virtual bool NeedsLoadForServer() const override;
-	// End UObject Interface
+	//~ End UObject Interface
 	
-	// Begin UPrimitiveComponent Interface
+	//~ Begin UPrimitiveComponent Interface
 	virtual FPrimitiveSceneProxy* CreateSceneProxy() override;
 	virtual void OnRegister()  override;
 	virtual void OnUnregister()  override;
-	// End UPrimitiveComponent Interface
+	//~ End UPrimitiveComponent Interface
 
-	// Begin UActorComponent Interface
+	//~ Begin UActorComponent Interface
 	virtual void CreateRenderState_Concurrent() override;
 	virtual void DestroyRenderState_Concurrent() override;
-	// End UActorComponent Interface
+	//~ End UActorComponent Interface
 
-	// Begin USceneComponent Interface
+	//~ Begin USceneComponent Interface
 	virtual FBoxSphereBounds CalcBounds(const FTransform &LocalToWorld) const override;
-	// End USceneComponent Interface
+	//~ End USceneComponent Interface
 
 	void ForceUpdate() { bForceUpdate = true; }
 	bool IsForcingUpdate() const { return bForceUpdate; }
diff --git a/Engine/Source/Runtime/Engine/Classes/AI/Navigation/NavRelevantComponent.h b/Engine/Source/Runtime/Engine/Classes/AI/Navigation/NavRelevantComponent.h
--- a/Engine/Source/Runtime/Engine/Classes/AI/Navigation/NavRelevantComponent.h
+++ b/Engine/Source/Runtime/Engine/Classes/AI/Navigation/NavRelevantComponent.h
@@ -10,17 +10,17 @@ class ENGINE_API UNavRelevantComponent : public UActorComponent, public INavRele
 {
 	GENERATED_UCLASS_BODY()
 
-	// Begin UActorComponent Interface
+	//~ Begin UActorComponent Interface
 	virtual void OnRegister() override;
 	virtual void OnUnregister() override;
-	// End UActorComponent Interface
+	//~ End UActorComponent Interface
 
-	// Begin INavRelevantInterface Interface
+	//~ Begin INavRelevantInterface Interface
 	virtual FBox GetNavigationBounds() const override;
 	virtual bool IsNavigationRelevant() const override;
 	virtual void UpdateNavigationBounds() override;
 	virtual UObject* GetNavigationParent() const override;
-	// End INavRelevantInterface Interface
+	//~ End INavRelevantInterface Interface
 
 	DEPRECATED(4.8, "This function is deprecated. Use CalcAndCacheBounds instead")
 	virtual void CalcBounds();
diff --git a/Engine/Source/Runtime/Engine/Classes/AI/Navigation/NavigationData.h b/Engine/Source/Runtime/Engine/Classes/AI/Navigation/NavigationData.h
--- a/Engine/Source/Runtime/Engine/Classes/AI/Navigation/NavigationData.h
+++ b/Engine/Source/Runtime/Engine/Classes/AI/Navigation/NavigationData.h
@@ -444,7 +444,7 @@ public:
 	// Life cycle                                                                
 	//----------------------------------------------------------------------//
 
-	// Begin UObject/AActor Interface
+	//~ Begin UObject/AActor Interface
 	virtual void PostInitProperties() override;
 	virtual void PostInitializeComponents() override;
 	virtual void PostLoad() override;
@@ -453,7 +453,7 @@ public:
 #endif // WITH_EDITOR
 	virtual void Destroyed() override;
 	virtual void EndPlay(const EEndPlayReason::Type EndPlayReason) override;
-	// End UObject Interface
+	//~ End UObject Interface
 		
 	virtual void CleanUp();
 
diff --git a/Engine/Source/Runtime/Engine/Classes/AI/Navigation/NavigationGraphNodeComponent.h b/Engine/Source/Runtime/Engine/Classes/AI/Navigation/NavigationGraphNodeComponent.h
--- a/Engine/Source/Runtime/Engine/Classes/AI/Navigation/NavigationGraphNodeComponent.h
+++ b/Engine/Source/Runtime/Engine/Classes/AI/Navigation/NavigationGraphNodeComponent.h
@@ -18,6 +18,6 @@ class UNavigationGraphNodeComponent : public USceneComponent
 	UPROPERTY()
 	UNavigationGraphNodeComponent* PrevNodeComponent;
 
-	// Begin UObject interface.
+	//~ Begin UObject Interface.
 	virtual void BeginDestroy() override;
 };
diff --git a/Engine/Source/Runtime/Engine/Classes/AI/Navigation/NavigationSystem.h b/Engine/Source/Runtime/Engine/Classes/AI/Navigation/NavigationSystem.h
--- a/Engine/Source/Runtime/Engine/Classes/AI/Navigation/NavigationSystem.h
+++ b/Engine/Source/Runtime/Engine/Classes/AI/Navigation/NavigationSystem.h
@@ -82,9 +82,9 @@ namespace FNavigationSystem
 
 struct FNavigationSystemExec: public FSelfRegisteringExec
 {
-	// Begin FExec Interface
+	//~ Begin FExec Interface
 	virtual bool Exec(UWorld* Inworld, const TCHAR* Cmd, FOutputDevice& Ar) override;
-	// End FExec Interface
+	//~ End FExec Interface
 };
 
 namespace ENavigationBuildLock
@@ -364,13 +364,13 @@ public:
 		CleanupUnsafe,
 	};
 
-	// Begin UObject Interface
+	//~ Begin UObject Interface
 	virtual void PostInitProperties() override;
 	static void AddReferencedObjects(UObject* InThis, FReferenceCollector& Collector);
 #if WITH_EDITOR
 	virtual void PostEditChangeProperty(struct FPropertyChangedEvent& PropertyChangedEvent) override;
 #endif // WITH_EDITOR
-	// End UObject Interface
+	//~ End UObject Interface
 
 	virtual void Tick(float DeltaSeconds);	
 
@@ -389,7 +389,7 @@ protected:
 
 public:
 	//----------------------------------------------------------------------//
-	// Public querying interface                                                                
+	//~ Begin Public querying Interface                                                                
 	//----------------------------------------------------------------------//
 	/** 
 	 *	Synchronously looks for a path from @fLocation to @EndLocation for agent with properties @AgentProperties. NavData actor appropriate for specified 
diff --git a/Engine/Source/Runtime/Engine/Classes/AI/Navigation/NavigationTestingActor.h b/Engine/Source/Runtime/Engine/Classes/AI/Navigation/NavigationTestingActor.h
--- a/Engine/Source/Runtime/Engine/Classes/AI/Navigation/NavigationTestingActor.h
+++ b/Engine/Source/Runtime/Engine/Classes/AI/Navigation/NavigationTestingActor.h
@@ -160,15 +160,15 @@ public:
 	void TickMe();
 #endif // WITH_EDITOR
 
-	// Begin INavAgentInterface Interface
+	//~ Begin INavAgentInterface Interface
 	virtual const FNavAgentProperties& GetNavAgentPropertiesRef() const override { return NavAgentProps; }
 	virtual FVector GetNavAgentLocation() const override;
 	virtual void GetMoveGoalReachTest(class AActor* MovingActor, const FVector& MoveOffset, FVector& GoalOffset, float& GoalRadius, float& GoalHalfHeight) const override {}
-	// End INavAgentInterface Interface
+	//~ End INavAgentInterface Interface
 
-	// Begin INavPathObserverInterface Interface
+	//~ Begin INavPathObserverInterface Interface
 	virtual void OnPathUpdated(class INavigationPathGenerator* PathGenerator) override;
-	// End INavPathObserverInterface Interface	
+	//~ End INavPathObserverInterface Interface	
 
 	void UpdateNavData();
 	void UpdatePathfinding();
diff --git a/Engine/Source/Runtime/Engine/Classes/AI/Navigation/RecastNavMesh.h b/Engine/Source/Runtime/Engine/Classes/AI/Navigation/RecastNavMesh.h
--- a/Engine/Source/Runtime/Engine/Classes/AI/Navigation/RecastNavMesh.h
+++ b/Engine/Source/Runtime/Engine/Classes/AI/Navigation/RecastNavMesh.h
@@ -690,7 +690,7 @@ public:
 	/** Dtor */
 	virtual ~ARecastNavMesh();
 
-	// Begin UObject Interface
+	//~ Begin UObject Interface
 	virtual void PostInitProperties() override;
 	virtual void PostLoad() override;
 	virtual SIZE_T GetResourceSize(EResourceSizeMode::Type Mode) override;	
@@ -698,7 +698,7 @@ public:
 #if WITH_EDITOR
 	virtual void PostEditChangeProperty( struct FPropertyChangedEvent& PropertyChangedEvent) override;
 #endif // WITH_EDITOR
-	// End UObject Interface
+	//~ End UObject Interface
 
 #if WITH_EDITOR
 	/** RecastNavMesh instances are dynamically spawned and should not be coppied */
@@ -707,7 +707,7 @@ public:
 	
 	virtual void CleanUp() override;
 
-	// Begin ANavigationData Interface
+	//~ Begin ANavigationData Interface
 	virtual FNavLocation GetRandomPoint(TSharedPtr<const FNavigationQueryFilter> Filter = NULL, const UObject* Querier = NULL) const override;
 	virtual bool GetRandomReachablePointInRadius(const FVector& Origin, float Radius, FNavLocation& OutResult, TSharedPtr<const FNavigationQueryFilter> Filter = NULL, const UObject* Querier = NULL) const override;
 	virtual bool GetRandomPointInNavigableRadius(const FVector& Origin, float Radius, FNavLocation& OutResult, TSharedPtr<const FNavigationQueryFilter> Filter = NULL, const UObject* Querier = NULL) const override;
@@ -730,7 +730,7 @@ public:
 
 	virtual void OnStreamingLevelAdded(ULevel* InLevel, UWorld* InWorld) override;
 	virtual void OnStreamingLevelRemoved(ULevel* InLevel, UWorld* InWorld) override;
-	// End ANavigationData Interface
+	//~ End ANavigationData Interface
 
 protected:
 	/** Serialization helper. */
diff --git a/Engine/Source/Runtime/Engine/Classes/Animation/AnimComposite.h b/Engine/Source/Runtime/Engine/Classes/Animation/AnimComposite.h
--- a/Engine/Source/Runtime/Engine/Classes/Animation/AnimComposite.h
+++ b/Engine/Source/Runtime/Engine/Classes/Animation/AnimComposite.h
@@ -25,15 +25,15 @@ public:
 	UAnimSequence* PreviewBasePose;
 #endif // WITH_EDITORONLY_DATA
 
-	// Begin UAnimSequenceBase interface
+	//~ Begin UAnimSequenceBase Interface
 	ENGINE_API virtual void OnAssetPlayerTickedInternal(FAnimAssetTickContext &Context, const float PreviousTime, const float MoveDelta, const FAnimTickRecord &Instance, class UAnimInstance* InstanceOwner) const override;
 	virtual void GetAnimationPose(FCompactPose& OutPose, FBlendedCurve& OutCurve, const FAnimExtractContext& ExtractionContext) const override;	
-	// End UAnimSequenceBase interface
-	// Begin UAnimSequence interface
+	//~ End UAnimSequenceBase Interface
+	//~ Begin UAnimSequence Interface
 #if WITH_EDITOR
 	virtual bool GetAllAnimationSequencesReferred(TArray<UAnimSequence*>& AnimationSequences) override;
 	virtual void ReplaceReferredAnimations(const TMap<UAnimSequence*, UAnimSequence*>& ReplacementMap) override;
 #endif
-	// End UAnimSequence interface
+	//~ End UAnimSequence Interface
 };
 
diff --git a/Engine/Source/Runtime/Engine/Classes/Animation/AnimCompress_Automatic.h b/Engine/Source/Runtime/Engine/Classes/Animation/AnimCompress_Automatic.h
--- a/Engine/Source/Runtime/Engine/Classes/Animation/AnimCompress_Automatic.h
+++ b/Engine/Source/Runtime/Engine/Classes/Animation/AnimCompress_Automatic.h
@@ -58,9 +58,9 @@ class UAnimCompress_Automatic : public UAnimCompress
 
 
 protected:
-	// Begin UAnimCompress Interface
+	//~ Begin UAnimCompress Interface
 	virtual void DoReduction(class UAnimSequence* AnimSeq, const TArray<class FBoneData>& BoneData) override;
-	// Begin UAnimCompress Interface
+	//~ Begin UAnimCompress Interface
 };
 
 
diff --git a/Engine/Source/Runtime/Engine/Classes/Animation/AnimCompress_BitwiseCompressOnly.h b/Engine/Source/Runtime/Engine/Classes/Animation/AnimCompress_BitwiseCompressOnly.h
--- a/Engine/Source/Runtime/Engine/Classes/Animation/AnimCompress_BitwiseCompressOnly.h
+++ b/Engine/Source/Runtime/Engine/Classes/Animation/AnimCompress_BitwiseCompressOnly.h
@@ -16,9 +16,9 @@ class UAnimCompress_BitwiseCompressOnly : public UAnimCompress
 
 
 protected:
-	// Begin UAnimCompress Interface
+	//~ Begin UAnimCompress Interface
 	virtual void DoReduction(class UAnimSequence* AnimSeq, const TArray<class FBoneData>& BoneData) override;
-	// Begin UAnimCompress Interface
+	//~ Begin UAnimCompress Interface
 };
 
 
diff --git a/Engine/Source/Runtime/Engine/Classes/Animation/AnimCompress_PerTrackCompression.h b/Engine/Source/Runtime/Engine/Classes/Animation/AnimCompress_PerTrackCompression.h
--- a/Engine/Source/Runtime/Engine/Classes/Animation/AnimCompress_PerTrackCompression.h
+++ b/Engine/Source/Runtime/Engine/Classes/Animation/AnimCompress_PerTrackCompression.h
@@ -131,11 +131,11 @@ public:
 	struct FPerTrackCachedInfo* PerReductionCachedData;
 
 protected:
-	// Begin UAnimCompress Interface
+	//~ Begin UAnimCompress Interface
 	virtual void DoReduction(class UAnimSequence* AnimSeq, const TArray<class FBoneData>& BoneData) override;
-	// Begin UAnimCompress Interface
+	//~ Begin UAnimCompress Interface
 
-	// Begin UAnimCompress_RemoveLinearKeys Interface
+	//~ Begin UAnimCompress_RemoveLinearKeys Interface
 	virtual void CompressUsingUnderlyingCompressor(
 		UAnimSequence* AnimSeq,
 		const TArray<FBoneData>& BoneData, 
@@ -149,7 +149,7 @@ protected:
 		TArray<FTranslationTrack>& TranslationData,
 		TArray<FRotationTrack>& RotationData,
 		TArray<FScaleTrack>& ScaleData) override;
-	// End UAnimCompress_RemoveLinearKeys Interface
+	//~ End UAnimCompress_RemoveLinearKeys Interface
 };
 
 
diff --git a/Engine/Source/Runtime/Engine/Classes/Animation/AnimMontage.h b/Engine/Source/Runtime/Engine/Classes/Animation/AnimMontage.h
--- a/Engine/Source/Runtime/Engine/Classes/Animation/AnimMontage.h
+++ b/Engine/Source/Runtime/Engine/Classes/Animation/AnimMontage.h
@@ -285,7 +285,7 @@ public:
 	{
 	}
 
-	// montage instance interfaces
+	//~ Begin montage instance Interfaces
 	void Play(float InPlayRate = 1.f);
 	void Stop(const FAlphaBlend& InBlendOut, bool bInterrupt=true);
 	void Pause();
@@ -419,7 +419,7 @@ class UAnimMontage : public UAnimCompositeBase
 	ENGINE_API bool IsValidSlot(FName InSlotName) const;
 
 public:
-	// Begin UObject Interface
+	//~ Begin UObject Interface
 	virtual void PostLoad() override;
 
 	// Gets the sequence length of the montage by calculating it from the lengths of the segments in the montage
@@ -428,20 +428,20 @@ public:
 #if WITH_EDITOR
 	virtual void PostEditChangeProperty(FPropertyChangedEvent& PropertyChangedEvent) override;
 #endif // WITH_EDITOR
-	// End UObject Interface
+	//~ End UObject Interface
 
-	// Begin AnimSequenceBase Interface
+	//~ Begin AnimSequenceBase Interface
 	virtual bool IsValidAdditive() const override;
 #if WITH_EDITOR
 	virtual EAnimEventTriggerOffsets::Type CalculateOffsetForNotify(float NotifyDisplayTime) const override;
 #endif // WITH_EDITOR
-	// End AnimSequenceBase Interface
+	//~ End AnimSequenceBase Interface
 
 #if WITH_EDITOR
-	// Begin UAnimationAsset interface
+	//~ Begin UAnimationAsset Interface
 	virtual bool GetAllAnimationSequencesReferred(TArray<UAnimSequence*>& AnimationSequences) override;
 	virtual void ReplaceReferredAnimations(const TMap<UAnimSequence*, UAnimSequence*>& ReplacementMap) override;
-	// End UAnimationAsset interface
+	//~ End UAnimationAsset Interface
 
 	/** Update all linkable elements contained in the montage */
 	ENGINE_API void UpdateLinkableElements();
diff --git a/Engine/Source/Runtime/Engine/Classes/Animation/AnimSequence.h b/Engine/Source/Runtime/Engine/Classes/Animation/AnimSequence.h
--- a/Engine/Source/Runtime/Engine/Classes/Animation/AnimSequence.h
+++ b/Engine/Source/Runtime/Engine/Classes/Animation/AnimSequence.h
@@ -531,7 +531,7 @@ class ENGINE_API UAnimSequence : public UAnimSequenceBase
 	TArray<FName>				UniqueMarkerNames;
 
 public:
-	// Begin UObject interface
+	//~ Begin UObject Interface
 	virtual void Serialize(FArchive& Ar) override;
 	virtual void PostInitProperties() override;
 	virtual void PostLoad() override;
@@ -543,9 +543,9 @@ public:
 	virtual void BeginDestroy() override;
 	virtual SIZE_T GetResourceSize(EResourceSizeMode::Type Mode) override;
 	virtual void GetAssetRegistryTags(TArray<FAssetRegistryTag>& OutTags) const override;
-	// End of UObject interface
+	//~ End UObject Interface
 
-	// Begin UAnimationAsset interface
+	//~ Begin UAnimationAsset Interface
 	virtual bool IsValidAdditive() const override;
 	virtual TArray<FName>* GetUniqueMarkerNames() { return &UniqueMarkerNames; }
 #if WITH_EDITOR
@@ -553,13 +553,13 @@ public:
 	virtual void ReplaceReferredAnimations(const TMap<UAnimSequence*, UAnimSequence*>& ReplacementMap) override;
 	virtual int32 GetNumberOfFrames() const override { return NumFrames; }
 #endif
-	// End of UAnimationAsset interface
+	//~ End UAnimationAsset Interface
 
-	// Begin UAnimSequenceBase interface
+	//~ Begin UAnimSequenceBase Interface
 	virtual void OnAssetPlayerTickedInternal(FAnimAssetTickContext &Context, const float PreviousTime, const float MoveDelta, const FAnimTickRecord &Instance, class UAnimInstance* InstanceOwner) const override;
 	virtual bool HasRootMotion() const override { return bEnableRootMotion; }
 	virtual void RefreshCacheData() override;
-	// End UAnimSequenceBase interface
+	//~ End UAnimSequenceBase Interface
 
 	// Extract Root Motion transform from the animation
 	FTransform ExtractRootMotion(float StartTime, float DeltaTime, bool bAllowLooping) const;
diff --git a/Engine/Source/Runtime/Engine/Classes/Animation/AnimSequenceBase.h b/Engine/Source/Runtime/Engine/Classes/Animation/AnimSequenceBase.h
--- a/Engine/Source/Runtime/Engine/Classes/Animation/AnimSequenceBase.h
+++ b/Engine/Source/Runtime/Engine/Classes/Animation/AnimSequenceBase.h
@@ -53,9 +53,9 @@ class UAnimSequenceBase : public UAnimationAsset
 	TArray<FAnimNotifyTrack> AnimNotifyTracks;
 #endif // WITH_EDITORONLY_DATA
 
-	// Begin UObject interface
+	//~ Begin UObject Interface
 	virtual void PostLoad() override;
-	// End of UObject interface
+	//~ End UObject Interface
 
 	/** Returns the total play length of the montage, if played back with a speed of 1.0. */
 	UFUNCTION(BlueprintCallable, Category = "Animation")
@@ -120,7 +120,7 @@ class UAnimSequenceBase : public UAnimationAsset
 	// update cache data (notify tracks, sync markers)
 	ENGINE_API virtual void RefreshCacheData();
 
-	// Begin UAnimationAsset interface
+	//~ Begin UAnimationAsset Interface
 	virtual void TickAssetPlayerInstance(FAnimTickRecord& Instance, class UAnimInstance* InstanceOwner, FAnimAssetTickContext& Context) const override;
 
 	void TickByMarkerAsFollower(FMarkerTickRecord &Instance, FMarkerTickContext &MarkerContext, float& CurrentTime, float& OutPreviousTime, const float MoveDelta, const bool bLooping) const;
@@ -131,7 +131,7 @@ class UAnimSequenceBase : public UAnimationAsset
 	// this doesn't mean max time. In Sequence, this is SequenceLength,
 	// but for BlendSpace CurrentTime is normalized [0,1], so this is 1
 	virtual float GetMaxCurrentTime() override { return SequenceLength; }
-	// End of UAnimationAsset interface
+	//~ End UAnimationAsset Interface
 
 	/**
 	* Get Bone Transform of the Time given, relative to Parent for all RequiredBones
diff --git a/Engine/Source/Runtime/Engine/Classes/Animation/AnimSet.h b/Engine/Source/Runtime/Engine/Classes/Animation/AnimSet.h
--- a/Engine/Source/Runtime/Engine/Classes/Animation/AnimSet.h
+++ b/Engine/Source/Runtime/Engine/Classes/Animation/AnimSet.h
@@ -109,12 +109,12 @@ public:
 	// @UE4: Do change private - they will go away
 	TMap<FName,int32> SkelMesh2LinkupCache;
 
-	// Begin UObject interface
+	//~ Begin UObject Interface
 	virtual void PostLoad() override;
 	virtual SIZE_T GetResourceSize(EResourceSizeMode::Type Mode) override;
-	// End UObject interface
+	//~ End UObject Interface
 	
-	// UAnimSet interface
+	//~ Begin UAnimSet Interface
 	/**
 	 * See if we can play sequences from this AnimSet on the provided USkeletalMesh.
 	 * Returns true if there is a bone in SkelMesh for every track in the AnimSet,
diff --git a/Engine/Source/Runtime/Engine/Classes/Animation/AnimSingleNodeInstance.h b/Engine/Source/Runtime/Engine/Classes/Animation/AnimSingleNodeInstance.h
--- a/Engine/Source/Runtime/Engine/Classes/Animation/AnimSingleNodeInstance.h
+++ b/Engine/Source/Runtime/Engine/Classes/Animation/AnimSingleNodeInstance.h
@@ -55,7 +55,7 @@ class ENGINE_API UAnimSingleNodeInstance : public UAnimInstance
 	UPROPERTY(transient)
 	uint32 bReverse:1;
 
-	// Begin UAnimInstance interface
+	//~ Begin UAnimInstance Interface
 	virtual void NativeInitializeAnimation() override;
 	virtual bool NativeEvaluateAnimation(FPoseContext& Output) override;
 	virtual void NativePostEvaluateAnimation() override;
@@ -64,7 +64,7 @@ class ENGINE_API UAnimSingleNodeInstance : public UAnimInstance
 protected:
 	virtual void Montage_Advance(float DeltaTime) override;
 	void InternalBlendSpaceEvaluatePose(class UBlendSpaceBase* BlendSpace, TArray<FBlendSampleData>& BlendSampleDataCache, FPoseContext& OutContext);
-	// End UAnimInstance interface
+	//~ End UAnimInstance Interface
 public:
 
 	UFUNCTION(BlueprintCallable, Category="Animation")
diff --git a/Engine/Source/Runtime/Engine/Classes/Animation/BlendSpace.h b/Engine/Source/Runtime/Engine/Classes/Animation/BlendSpace.h
--- a/Engine/Source/Runtime/Engine/Classes/Animation/BlendSpace.h
+++ b/Engine/Source/Runtime/Engine/Classes/Animation/BlendSpace.h
@@ -44,12 +44,12 @@ public:
 	ENGINE_API void GetGridSamplesFromBlendInput(const FVector &BlendInput, FGridBlendSample & LeftBottom, FGridBlendSample & RightBottom, FGridBlendSample & LeftTop, FGridBlendSample& RightTop) const;
 
 protected:
-	// Begin UBlendSpaceBase interface
+	//~ Begin UBlendSpaceBase Interface
 	virtual void SnapToBorder(FBlendSample& Sample) const override;
 	virtual EBlendSpaceAxis GetAxisToScale() const override;
 	virtual bool IsSameSamplePoint(const FVector& SamplePointA, const FVector& SamplePointB) const override;
 	virtual void GetRawSamplesFromBlendInput(const FVector &BlendInput, TArray<FGridBlendSample, TInlineAllocator<4> > & OutBlendSamples) const override;
-	// End UBlendSpaceBase interface
+	//~ End UBlendSpaceBase Interface
 
 private:
 	/** 
diff --git a/Engine/Source/Runtime/Engine/Classes/Animation/BlendSpace1D.h b/Engine/Source/Runtime/Engine/Classes/Animation/BlendSpace1D.h
--- a/Engine/Source/Runtime/Engine/Classes/Animation/BlendSpace1D.h
+++ b/Engine/Source/Runtime/Engine/Classes/Animation/BlendSpace1D.h
@@ -30,12 +30,12 @@ public:
 	virtual bool IsValidAdditive() const override;
 
 protected:
-	// Begin UBlendSpaceBase interface
+	//~ Begin UBlendSpaceBase Interface
 	virtual void SnapToBorder(FBlendSample& Sample) const override;
 	virtual EBlendSpaceAxis GetAxisToScale() const override;
 	virtual bool IsSameSamplePoint(const FVector& SamplePointA, const FVector& SamplePointB) const override;
 	virtual void GetRawSamplesFromBlendInput(const FVector &BlendInput, TArray<FGridBlendSample, TInlineAllocator<4> > & OutBlendSamples) const override;
-	// End UBlendSpaceBase interface
+	//~ End UBlendSpaceBase Interface
 
 private:
 
diff --git a/Engine/Source/Runtime/Engine/Classes/Animation/BlendSpaceBase.h b/Engine/Source/Runtime/Engine/Classes/Animation/BlendSpaceBase.h
--- a/Engine/Source/Runtime/Engine/Classes/Animation/BlendSpaceBase.h
+++ b/Engine/Source/Runtime/Engine/Classes/Animation/BlendSpaceBase.h
@@ -256,14 +256,14 @@ public:
 	UPROPERTY(EditAnywhere, Category=SampleInterpolation)
 	TArray<FPerBoneInterpolation> PerBoneBlend;
 
-	// Begin UObject interface
+	//~ Begin UObject Interface
 	virtual void PostLoad() override;
 #if WITH_EDITOR
 	virtual void PostEditChangeProperty( struct FPropertyChangedEvent& PropertyChangedEvent) override;
 #endif // WITH_EDITOR
-	// End UObject interface
+	//~ End UObject Interface
 
-	// Begin UAnimationAsset interface
+	//~ Begin UAnimationAsset Interface
 	virtual void TickAssetPlayerInstance(FAnimTickRecord& Instance, class UAnimInstance* InstanceOwner, FAnimAssetTickContext& Context) const override;
 	// this is used in editor only when used for transition getter
 	// this doesn't mean max time. In Sequence, this is SequenceLength,
@@ -275,7 +275,7 @@ public:
 	virtual int32 GetMarkerUpdateCounter() const;
 #endif
 	virtual TArray<FName>* GetUniqueMarkerNames() override { return (bAllSequencesHaveMatchingMarkers && SampleData.Num() > 0) ? SampleData[0].Animation->GetUniqueMarkerNames() : nullptr; }
-	// End of UAnimationAsset interface
+	//~ End UAnimationAsset Interface
 
 	void TickFollowerSamples(TArray<FBlendSampleData> &SampleDataList, const int32 HighestWeightIndex, FAnimAssetTickContext &Context) const
 	{
diff --git a/Engine/Source/Runtime/Engine/Classes/Animation/SkeletalMeshActor.h b/Engine/Source/Runtime/Engine/Classes/Animation/SkeletalMeshActor.h
--- a/Engine/Source/Runtime/Engine/Classes/Animation/SkeletalMeshActor.h
+++ b/Engine/Source/Runtime/Engine/Classes/Animation/SkeletalMeshActor.h
@@ -59,13 +59,13 @@ public:
 	virtual void OnRep_ReplicatedMaterial1();
 
 
-	// Begin UObject interface
+	//~ Begin UObject Interface
 protected:
 	virtual FString GetDetailedInfoInternal() const override;
 public:
-	// End UObject interface
+	//~ End UObject Interface
 
-	// Begin AActor interface
+	//~ Begin AActor Interface
 #if WITH_EDITOR
 	virtual void CheckForErrors() override;
 	virtual bool GetReferencedContentObjects( TArray<UObject*>& Objects ) const override;
@@ -73,9 +73,9 @@ public:
 	virtual void LoadedFromAnotherClass(const FName& OldClassName) override;
 #endif
 	virtual void PostInitializeComponents() override;
-	// End AActor interface
+	//~ End AActor Interface
 
-	// Begin IMatineeAnimInterface Interface
+	//~ Begin IMatineeAnimInterface Interface
 	virtual void PreviewBeginAnimControl(class UInterpGroup* InInterpGroup) override;
 	virtual void PreviewSetAnimPosition(FName SlotName, int32 ChannelIndex, UAnimSequence* InAnimSequence, float InPosition, bool bLooping, bool bFireNotifies, float AdvanceTime) override;
 	virtual void PreviewSetAnimWeights(TArray<FAnimSlotInfo>& SlotInfos) override;
@@ -85,7 +85,7 @@ public:
 	virtual void BeginAnimControl(class UInterpGroup* InInterpGroup) override;
 	virtual void SetAnimPosition(FName SlotName, int32 ChannelIndex, class UAnimSequence* InAnimSequence, float InPosition, bool bFireNotifies, bool bLooping) override;
 	virtual void FinishAnimControl(class UInterpGroup* InInterpGroup) override;
-	// End IMatineeAnimInterface Interface
+	//~ End IMatineeAnimInterface Interface
 
 private:
 	// utility function to see if it can play animation or not
diff --git a/Engine/Source/Runtime/Engine/Classes/Animation/VertexAnim/MorphTarget.h b/Engine/Source/Runtime/Engine/Classes/Animation/VertexAnim/MorphTarget.h
--- a/Engine/Source/Runtime/Engine/Classes/Animation/VertexAnim/MorphTarget.h
+++ b/Engine/Source/Runtime/Engine/Classes/Animation/VertexAnim/MorphTarget.h
@@ -74,10 +74,10 @@ public:
 	/** morph mesh vertex data for each LOD */
 	TArray<FMorphTargetLODModel>	MorphLODModels;
 
-	// Begin UObject interface.
+	//~ Begin UObject Interface.
 	virtual void Serialize( FArchive& Ar ) override;
 	virtual SIZE_T GetResourceSize(EResourceSizeMode::Type Mode) override;
-	// Begin UObject interface.
+	//~ Begin UObject Interface.
 
 	/** Post process after importing **/
 	ENGINE_API void PostProcess( USkeletalMesh * NewMesh, const FMorphMeshRawSource& BaseSource, const FMorphMeshRawSource& TargetSource, int32 LODIndex, bool bCompareNormal );
@@ -85,10 +85,10 @@ public:
 	/** Remap vertex indices with base mesh. */
 	void RemapVertexIndices( USkeletalMesh* InBaseMesh, const TArray< TArray<uint32> > & BasedWedgePointIndices );
 
-	// Begin UVertexAnimBase interface
+	//~ Begin UVertexAnimBase Interface
 	virtual FVertexAnimDelta* GetDeltasAtTime(float Time, int32 LODIndex, FVertexAnimEvalStateBase* State, int32& OutNumDeltas) override;
 	virtual bool HasDataForLOD(int32 LODIndex) override;
-	// End UVertexAnimBase interface
+	//~ End UVertexAnimBase Interface
 
 private:
 	/**
diff --git a/Engine/Source/Runtime/Engine/Classes/Animation/VertexAnim/VertexAnimation.h b/Engine/Source/Runtime/Engine/Classes/Animation/VertexAnim/VertexAnimation.h
--- a/Engine/Source/Runtime/Engine/Classes/Animation/VertexAnim/VertexAnimation.h
+++ b/Engine/Source/Runtime/Engine/Classes/Animation/VertexAnim/VertexAnimation.h
@@ -36,17 +36,17 @@ class UVertexAnimation : public UVertexAnimBase
 	/** Raw vertex anim data */
 	TArray<FVertexAnimFrame>	VertexAnimData;
 
-	// Begin UObject interface.
+	//~ Begin UObject Interface.
 	virtual void Serialize( FArchive& Ar ) override;
 	virtual SIZE_T GetResourceSize(EResourceSizeMode::Type Mode) override;
-	// Begin UObject interface.
+	//~ Begin UObject Interface.
 
-	// Begin UVertexAnimBase interface
+	//~ Begin UVertexAnimBase Interface
 	virtual FVertexAnimEvalStateBase* InitEval() override;
 	virtual void TermEval(FVertexAnimEvalStateBase* State) override;
 	virtual FVertexAnimDelta* GetDeltasAtTime(float Time, int32 LODIndex, FVertexAnimEvalStateBase* State, int32& OutNumDeltas) override;
 	virtual bool HasDataForLOD(int32 LODIndex) override;
-	// End UVertexAnimBase interface
+	//~ End UVertexAnimBase Interface
 
 	/** Get the number of frames in the animation */
 	ENGINE_API int32 GetNumFrames();
diff --git a/Engine/Source/Runtime/Engine/Classes/Atmosphere/AtmosphericFogComponent.h b/Engine/Source/Runtime/Engine/Classes/Atmosphere/AtmosphericFogComponent.h
--- a/Engine/Source/Runtime/Engine/Classes/Atmosphere/AtmosphericFogComponent.h
+++ b/Engine/Source/Runtime/Engine/Classes/Atmosphere/AtmosphericFogComponent.h
@@ -227,11 +227,11 @@ public:
 	void StartPrecompute();	
 
 protected:
-	// Begin UActorComponent interface.
+	//~ Begin UActorComponent Interface.
 	virtual void CreateRenderState_Concurrent() override;
 	virtual void SendRenderTransform_Concurrent() override;
 	virtual void DestroyRenderState_Concurrent() override;
-	// End UActorComponent interface.
+	//~ End UActorComponent Interface.
 
 	 void AddFogIfNeeded();
 
@@ -246,7 +246,7 @@ public:
 	mutable FByteBulkData IrradianceData;
 	mutable FByteBulkData InscatterData;
 	
-	// Begin UObject interface. 
+	//~ Begin UObject Interface. 
 	virtual void PostLoad() override;
 	virtual void BeginDestroy() override;
 
@@ -256,14 +256,14 @@ public:
 #endif // WITH_EDITOR
 	virtual void PostInterpChange(UProperty* PropertyThatChanged) override;
 	virtual void Serialize(FArchive& Ar) override;
-	// End UObject Interface
+	//~ End UObject Interface
 
 	ENGINE_API void InitResource();
 	ENGINE_API void ReleaseResource();
 
-	// Begin UActorComponent interface.
+	//~ Begin UActorComponent Interface.
 	virtual FActorComponentInstanceData* GetComponentInstanceData() const override;
-	// End UActorComponent interface.
+	//~ End UActorComponent Interface.
 
 	void ApplyComponentInstanceData(class FAtmospherePrecomputeInstanceData* ComponentInstanceData);
 	const FAtmospherePrecomputeParameters& GetPrecomputeParameters() const { return PrecomputeParams;  }
diff --git a/Engine/Source/Runtime/Engine/Classes/Camera/CameraActor.h b/Engine/Source/Runtime/Engine/Classes/Camera/CameraActor.h
--- a/Engine/Source/Runtime/Engine/Classes/Camera/CameraActor.h
+++ b/Engine/Source/Runtime/Engine/Classes/Camera/CameraActor.h
@@ -54,18 +54,18 @@ private:
 	struct FPostProcessSettings PostProcessSettings_DEPRECATED;
 
 public:
-	// Begin UObject interface
+	//~ Begin UObject Interface
 	virtual void Serialize(FArchive& Ar) override;
 	virtual void PostLoadSubobjects(FObjectInstancingGraph* OuterInstanceGraph) override;
 
 #if WITH_EDITOR
 	virtual void PostEditChangeProperty(FPropertyChangedEvent& PropertyChangedEvent) override;
 #endif
-	// End UObject interface
+	//~ End UObject Interface
 	
-	// Begin AActor interface
+	//~ Begin AActor Interface
 	virtual void BeginPlay() override;
-	// End AActor interface
+	//~ End AActor Interface
 
 	/** Returns CameraComponent subobject **/
 	class UCameraComponent* GetCameraComponent() const;
diff --git a/Engine/Source/Runtime/Engine/Classes/Camera/CameraModifier_CameraShake.h b/Engine/Source/Runtime/Engine/Classes/Camera/CameraModifier_CameraShake.h
--- a/Engine/Source/Runtime/Engine/Classes/Camera/CameraModifier_CameraShake.h
+++ b/Engine/Source/Runtime/Engine/Classes/Camera/CameraModifier_CameraShake.h
@@ -9,7 +9,7 @@
 #include "Camera/CameraModifier.h"
 #include "CameraModifier_CameraShake.generated.h"
 
-//=============================================================================
+//~=============================================================================
 /**
  * A UCameraModifier_CameraShake is a camera modifier that can apply a UCameraShake to 
  * the owning camera.
@@ -45,9 +45,9 @@ public:
 	/** Stops and removes all camera shakes from the camera. */
 	virtual void RemoveAllCameraShakes();
 	
-	// Begin UCameraModifer Interface
+	//~ Begin UCameraModifer Interface
 	virtual bool ModifyCamera(float DeltaTime, struct FMinimalViewInfo& InOutPOV) override;
-	// End UCameraModifer Interface
+	//~ End UCameraModifer Interface
 
 protected:
 	/** Scaling factor applied to all camera shakes in when in splitscreen mode. Normally used to reduce shaking, since shakes feel more intense in a smaller viewport. */
diff --git a/Engine/Source/Runtime/Engine/Classes/Camera/PlayerCameraManager.h b/Engine/Source/Runtime/Engine/Classes/Camera/PlayerCameraManager.h
--- a/Engine/Source/Runtime/Engine/Classes/Camera/PlayerCameraManager.h
+++ b/Engine/Source/Runtime/Engine/Classes/Camera/PlayerCameraManager.h
@@ -141,7 +141,7 @@ public:
 };
 
 
-//=============================================================================
+//~=============================================================================
 /**
  * A PlayerCameraManager is responsible for managing the camera for a particular
  * player. It defines the final view properties used by other systems (e.g. the renderer),
@@ -434,13 +434,13 @@ public:
 	/** @return the ViewTarget if it is an APawn, or nullptr otherwise */
 	class APawn* GetViewTargetPawn() const;
 
-	// Begin AActor Interface
+	//~ Begin AActor Interface
 	virtual bool ShouldTickIfViewportsOnly() const override;
 	virtual void PostInitializeComponents() override;
 	virtual void EndPlay(const EEndPlayReason::Type EndPlayReason) override;
 	virtual void DisplayDebug(class UCanvas* Canvas, const FDebugDisplayInfo& DebugDisplay, float& YL, float& YPos) override;
 	virtual void ApplyWorldOffset(const FVector& InOffset, bool bWorldShift) override;
-	// End AActor Interface
+	//~ End AActor Interface
 
 	/** Static.  Plays an in-world camera shake that affects all nearby players, with radial distance-based attenuation.
 	 * @param InWorld - World context.
diff --git a/Engine/Source/Runtime/Engine/Classes/Components/ActorComponent.h b/Engine/Source/Runtime/Engine/Classes/Components/ActorComponent.h
--- a/Engine/Source/Runtime/Engine/Classes/Components/ActorComponent.h
+++ b/Engine/Source/Runtime/Engine/Classes/Components/ActorComponent.h
@@ -219,7 +219,7 @@ public:
 	UFUNCTION(BlueprintCallable, Category="Components")
 	bool ComponentHasTag(FName Tag) const;
 
-	// Trigger/Activation interface
+	//~ Begin Trigger/Activation Interface
 
 	/**
 	 * Activates the SceneComponent
@@ -587,7 +587,7 @@ public:
 	/** Called before we throw away components during RerunConstructionScripts, to cache any data we wish to persist across that operation */
 	virtual class FActorComponentInstanceData* GetComponentInstanceData() const;
 
-	// Begin UObject interface.
+	//~ Begin UObject Interface.
 	virtual void BeginDestroy() override;
 	virtual bool NeedsLoadForClient() const override;
 	virtual bool NeedsLoadForServer() const override;
@@ -605,13 +605,13 @@ public:
 	virtual void PreEditUndo() override;
 	virtual void PostEditUndo() override;
 #endif // WITH_EDITOR
-	// End UObject interface.
+	//~ End UObject Interface.
 
-	// Begin IInterface_AssetUserData Interface
+	//~ Begin IInterface_AssetUserData Interface
 	virtual void AddAssetUserData(UAssetUserData* InUserData) override;
 	virtual void RemoveUserDataOfClass(TSubclassOf<UAssetUserData> InUserDataClass) override;
 	virtual UAssetUserData* GetAssetUserDataOfClass(TSubclassOf<UAssetUserData> InUserDataClass) override;
-	// End IInterface_AssetUserData Interface
+	//~ End IInterface_AssetUserData Interface
 
 	/** See if the owning Actor is currently running the UCS */
 	bool IsOwnerRunningUserConstructionScript() const;
diff --git a/Engine/Source/Runtime/Engine/Classes/Components/ArrowComponent.h b/Engine/Source/Runtime/Engine/Classes/Components/ArrowComponent.h
--- a/Engine/Source/Runtime/Engine/Classes/Components/ArrowComponent.h
+++ b/Engine/Source/Runtime/Engine/Classes/Components/ArrowComponent.h
@@ -56,17 +56,17 @@ class ENGINE_API UArrowComponent : public UPrimitiveComponent
 	UFUNCTION(BlueprintCallable, DisplayName="SetArrowColor", Category="Components|Arrow")
 	virtual void SetArrowColor_New(FLinearColor NewColor);
 
-	// Begin UPrimitiveComponent interface.
+	//~ Begin UPrimitiveComponent Interface.
 	virtual FPrimitiveSceneProxy* CreateSceneProxy() override;
 #if WITH_EDITOR
 	virtual bool ComponentIsTouchingSelectionBox(const FBox& InSelBBox, const FEngineShowFlags& ShowFlags, const bool bConsiderOnlyBSP, const bool bMustEncompassEntireComponent) const override;
 	virtual bool ComponentIsTouchingSelectionFrustum(const FConvexVolume& InFrustum, const FEngineShowFlags& ShowFlags, const bool bConsiderOnlyBSP, const bool bMustEncompassEntireComponent) const override;
 #endif
-	// End UPrimitiveComponent interface.
+	//~ End UPrimitiveComponent Interface.
 
-	// Begin USceneComponent interface.
+	//~ Begin USceneComponent Interface.
 	virtual FBoxSphereBounds CalcBounds(const FTransform& LocalToWorld) const override;
-	// Begin USceneComponent interface.
+	//~ Begin USceneComponent Interface.
 
 #if WITH_EDITORONLY_DATA
 	/** Set the scale that we use when rendering in-editor */
diff --git a/Engine/Source/Runtime/Engine/Classes/Components/AudioComponent.h b/Engine/Source/Runtime/Engine/Classes/Components/AudioComponent.h
--- a/Engine/Source/Runtime/Engine/Classes/Components/AudioComponent.h
+++ b/Engine/Source/Runtime/Engine/Classes/Components/AudioComponent.h
@@ -264,28 +264,28 @@ public:
 	/** Used by the subtitle manager to prioritize subtitles wave instances spawned by this component. */
 	float SubtitlePriority;
 
-	// Begin UObject interface.
+	//~ Begin UObject Interface.
 #if WITH_EDITOR
 	virtual void PostEditChangeProperty(FPropertyChangedEvent& PropertyChangedEvent) override;
 #endif // WITH_EDITOR
 	virtual FString GetDetailedInfoInternal() const override;
-	// End UObject interface.
+	//~ End UObject Interface.
 
-	// Begin USceneComponent Interface
+	//~ Begin USceneComponent Interface
 	virtual void Activate(bool bReset=false) override;
 	virtual void Deactivate() override;
 	virtual void OnUpdateTransform(bool bSkipPhysicsMove, ETeleportType Teleport = ETeleportType::None) override;
 	// do nothing! on purpose
 	virtual void UpdateNavigationData() override {}
-	// End USceneComponent Interface
+	//~ End USceneComponent Interface
 
-	// Begin ActorComponent interface.
+	//~ Begin ActorComponent Interface.
 #if WITH_EDITORONLY_DATA
 	virtual void OnRegister() override;
 #endif
 	virtual void OnUnregister() override;
 	virtual const UObject* AdditionalStatObject() const override;
-	// End ActorComponent interface.
+	//~ End ActorComponent Interface.
 
 	/** Returns a pointer to the attenuation settings to be used (if any) for this audio component dependent on the SoundAttenuation asset or overrides set. */
 	const FAttenuationSettings* GetAttenuationSettingsToApply() const;
diff --git a/Engine/Source/Runtime/Engine/Classes/Components/BillboardComponent.h b/Engine/Source/Runtime/Engine/Classes/Components/BillboardComponent.h
--- a/Engine/Source/Runtime/Engine/Classes/Components/BillboardComponent.h
+++ b/Engine/Source/Runtime/Engine/Classes/Components/BillboardComponent.h
@@ -62,14 +62,14 @@ class ENGINE_API UBillboardComponent : public UPrimitiveComponent
 	virtual void SetSpriteAndUV(class UTexture2D* NewSprite, int32 NewU, int32 NewUL, int32 NewV, int32 NewVL);
 
 
-	// Begin UPrimitiveComponent Interface
+	//~ Begin UPrimitiveComponent Interface
 	virtual FPrimitiveSceneProxy* CreateSceneProxy() override;
 	virtual FBoxSphereBounds CalcBounds(const FTransform& LocalToWorld) const override;
 #if WITH_EDITOR
 	virtual bool ComponentIsTouchingSelectionBox(const FBox& InSelBBox, const FEngineShowFlags& ShowFlags, const bool bConsiderOnlyBSP, const bool bMustEncompassEntireComponent) const override;
 	virtual bool ComponentIsTouchingSelectionFrustum(const FConvexVolume& InFrustum, const FEngineShowFlags& ShowFlags, const bool bConsiderOnlyBSP, const bool bMustEncompassEntireComponent) const override;
 #endif
-	// End UPrimitiveComponent Interface
+	//~ End UPrimitiveComponent Interface
 
 #if WITH_EDITORONLY_DATA
 	/** Set the scale that we use when rendering in-editor */
diff --git a/Engine/Source/Runtime/Engine/Classes/Components/BoxComponent.h b/Engine/Source/Runtime/Engine/Classes/Components/BoxComponent.h
--- a/Engine/Source/Runtime/Engine/Classes/Components/BoxComponent.h
+++ b/Engine/Source/Runtime/Engine/Classes/Components/BoxComponent.h
@@ -35,19 +35,19 @@ public:
 	UFUNCTION(BlueprintCallable, Category="Components|Box")
 	FVector GetUnscaledBoxExtent() const;
 
-	// Begin UPrimitiveComponent interface.
+	//~ Begin UPrimitiveComponent Interface.
 	virtual bool IsZeroExtent() const override;
 	virtual FPrimitiveSceneProxy* CreateSceneProxy() override;
 	virtual struct FCollisionShape GetCollisionShape(float Inflation = 0.0f) const override;
-	// End UPrimitiveComponent interface.
+	//~ End UPrimitiveComponent Interface.
 
-	// Begin USceneComponent interface
+	//~ Begin USceneComponent Interface
 	virtual FBoxSphereBounds CalcBounds(const FTransform& LocalToWorld) const override;
-	// End USceneComponent interface
+	//~ End USceneComponent Interface
 
-	// Begin UShapeComponent interface
+	//~ Begin UShapeComponent Interface
 	virtual void UpdateBodySetup() override;
-	// End UShapeComponent interface
+	//~ End UShapeComponent Interface
 
 	// Sets the box extents without triggering a render or physics update.
 	FORCEINLINE void InitBoxExtent(const FVector& InBoxExtent) { BoxExtent = InBoxExtent; }
diff --git a/Engine/Source/Runtime/Engine/Classes/Components/BrushComponent.h b/Engine/Source/Runtime/Engine/Classes/Components/BrushComponent.h
--- a/Engine/Source/Runtime/Engine/Classes/Components/BrushComponent.h
+++ b/Engine/Source/Runtime/Engine/Classes/Components/BrushComponent.h
@@ -27,19 +27,19 @@ class UBrushComponent : public UPrimitiveComponent
 	UPROPERTY()
 	FVector PrePivot_DEPRECATED;
 
-	// Begin UObject interface
+	//~ Begin UObject Interface
 	virtual void PostLoad() override;
 	virtual SIZE_T GetResourceSize(EResourceSizeMode::Type Mode) override;
-	// End UObject interface
+	//~ End UObject Interface
 
-	// Begin USceneComponent interface
+	//~ Begin USceneComponent Interface
 	virtual FBoxSphereBounds CalcBounds(const FTransform& LocalToWorld) const override;
 	virtual bool ShouldCollideWhenPlacing() const override { return true; }
-	// End USceneComponent interface
+	//~ End USceneComponent Interface
 
 public:
 
-	// Begin UPrimitiveComponent interface.
+	//~ Begin UPrimitiveComponent Interface.
 	virtual FPrimitiveSceneProxy* CreateSceneProxy() override;
 	virtual class UBodySetup* GetBodySetup() override { return BrushBodySetup; };
 	virtual void GetUsedMaterials( TArray<UMaterialInterface*>& OutMaterials ) const override;
@@ -47,7 +47,7 @@ public:
 #if WITH_EDITOR
 	virtual bool ComponentIsTouchingSelectionBox(const FBox& InSelBBox, const FEngineShowFlags& ShowFlags, const bool bConsiderOnlyBSP, const bool bMustEncompassEntireComponent) const override;
 	virtual bool ComponentIsTouchingSelectionFrustum(const FConvexVolume& InFrustum, const FEngineShowFlags& ShowFlags, const bool bConsiderOnlyBSP, const bool bMustEncompassEntireComponent) const override;
-	// End UPrimitiveComponent interface.
+	//~ End UPrimitiveComponent Interface.
 
 	/** Return true if the brush appears to have inverted polys */
 	ENGINE_API bool HasInvertedPolys() const;
diff --git a/Engine/Source/Runtime/Engine/Classes/Components/CapsuleComponent.h b/Engine/Source/Runtime/Engine/Classes/Components/CapsuleComponent.h
--- a/Engine/Source/Runtime/Engine/Classes/Components/CapsuleComponent.h
+++ b/Engine/Source/Runtime/Engine/Classes/Components/CapsuleComponent.h
@@ -59,29 +59,29 @@ public:
 	UFUNCTION(BlueprintCallable, Category="Components|Capsule")
 	void SetCapsuleHalfHeight(float HalfHeight, bool bUpdateOverlaps=true);
 
-	// Begin UObject interface
+	//~ Begin UObject Interface
 	virtual void Serialize(FArchive& Ar) override;
 	virtual void PostLoad() override;
 #if WITH_EDITOR
 	virtual void PostEditChangeProperty(FPropertyChangedEvent& PropertyChangedEvent) override;
 #endif // WITH_EDITOR
-	// End UObject interface
+	//~ End UObject Interface
 
-	// Begin USceneComponent interface
+	//~ Begin USceneComponent Interface
 	virtual FBoxSphereBounds CalcBounds(const FTransform& LocalToWorld) const override;
 	virtual void CalcBoundingCylinder(float& CylinderRadius, float& CylinderHalfHeight) const override;
-	// End USceneComponent interface
+	//~ End USceneComponent Interface
 
-	// Begin UPrimitiveComponent interface.
+	//~ Begin UPrimitiveComponent Interface.
 	virtual FPrimitiveSceneProxy* CreateSceneProxy() override;
 	virtual bool IsZeroExtent() const override;
 	virtual struct FCollisionShape GetCollisionShape(float Inflation = 0.0f) const override;
 	virtual bool AreSymmetricRotations(const FQuat& A, const FQuat& B, const FVector& Scale3D) const override;
-	// End UPrimitiveComponent interface.
+	//~ End UPrimitiveComponent Interface.
 
-	// Begin UShapeComponent interface
+	//~ Begin UShapeComponent Interface
 	virtual void UpdateBodySetup() override;
-	// End UShapeComponent interface
+	//~ End UShapeComponent Interface
 
 	// @return the capsule radius scaled by the component scale.
 	UFUNCTION(BlueprintCallable, Category="Components|Capsule")
diff --git a/Engine/Source/Runtime/Engine/Classes/Components/ChildActorComponent.h b/Engine/Source/Runtime/Engine/Classes/Components/ChildActorComponent.h
--- a/Engine/Source/Runtime/Engine/Classes/Components/ChildActorComponent.h
+++ b/Engine/Source/Runtime/Engine/Classes/Components/ChildActorComponent.h
@@ -36,19 +36,19 @@ public:
 	/** Cached copy of the instance data when the ChildActor is destroyed to be available when needed */
 	mutable FChildActorComponentInstanceData* CachedInstanceData;
 
-	// Begin Object interface.
+	//~ Begin Object Interface.
 #if WITH_EDITOR
 	virtual void PostEditChangeProperty(FPropertyChangedEvent& PropertyChangedEvent) override;
 	virtual void PostEditUndo() override;
 #endif
-	// End Object interface.
+	//~ End Object Interface.
 
-	// Begin ActorComponent interface.
+	//~ Begin ActorComponent Interface.
 	virtual void OnComponentCreated() override;
 	virtual void OnComponentDestroyed() override;
 	virtual void OnRegister() override;
 	virtual FActorComponentInstanceData* GetComponentInstanceData() const override;
-	// End ActorComponent interface.
+	//~ End ActorComponent Interface.
 
 	/** Apply the component instance data to the child actor component */
 	void ApplyComponentInstanceData(class FChildActorComponentInstanceData* ComponentInstanceData);
diff --git a/Engine/Source/Runtime/Engine/Classes/Components/DecalComponent.h b/Engine/Source/Runtime/Engine/Classes/Components/DecalComponent.h
--- a/Engine/Source/Runtime/Engine/Classes/Components/DecalComponent.h
+++ b/Engine/Source/Runtime/Engine/Classes/Components/DecalComponent.h
@@ -98,18 +98,18 @@ public:
 		}
 	}
 	
-	// Begin UActorComponent Interface
+	//~ Begin UActorComponent Interface
 	virtual void CreateRenderState_Concurrent() override;
 	virtual void DestroyRenderState_Concurrent() override;
 	virtual void SendRenderTransform_Concurrent() override;
 	virtual const UObject* AdditionalStatObject() const override;
-	// End UActorComponent Interface
+	//~ End UActorComponent Interface
 
 	virtual void Serialize(FArchive& Ar) override;
 
-	// Begin USceneComponent Interface
+	//~ Begin USceneComponent Interface
 	virtual FBoxSphereBounds CalcBounds(const FTransform& LocalToWorld) const override;
-	// End USceneComponent Interface
+	//~ End USceneComponent Interface
 
 };
 
diff --git a/Engine/Source/Runtime/Engine/Classes/Components/DestructibleComponent.h b/Engine/Source/Runtime/Engine/Classes/Components/DestructibleComponent.h
--- a/Engine/Source/Runtime/Engine/Classes/Components/DestructibleComponent.h
+++ b/Engine/Source/Runtime/Engine/Classes/Components/DestructibleComponent.h
@@ -75,11 +75,11 @@ class ENGINE_API UDestructibleComponent : public USkinnedMeshComponent
 #endif // WITH_PHYSX 
 
 #if WITH_EDITOR
-	// Begin UObject interface.
+	//~ Begin UObject Interface.
 	virtual void PostEditChangeProperty( struct FPropertyChangedEvent& PropertyChangedEvent) override;
 #endif // WITH_EDITOR
 	virtual void Serialize(FArchive& Ar) override;
-	// End UObject interface.
+	//~ End UObject Interface.
 
 	// Take damage
 	UFUNCTION(BlueprintCallable, Category="Components|Destructible")
@@ -105,20 +105,20 @@ public:
 	physx::apex::NxDestructibleActor* ApexDestructibleActor;
 #endif	//WITH_APEX
 
-	// Begin USceneComponent interface.
+	//~ Begin USceneComponent Interface.
 	virtual FBoxSphereBounds CalcBounds(const FTransform& LocalToWorld) const override;
 	virtual void OnUpdateTransform(bool bSkipPhysicsMove, ETeleportType Teleport = ETeleportType::None) override;
 	virtual void Activate(bool bReset=false) override;
 	virtual void Deactivate() override;
-	// End USceneComponent interface.
+	//~ End USceneComponent Interface.
 
-	// Begin UActorComponent interface.
+	//~ Begin UActorComponent Interface.
 	virtual void CreatePhysicsState() override;
 	virtual void DestroyPhysicsState() override;
 	virtual class UBodySetup* GetBodySetup() override;
-	// End UActorComponent interface.
+	//~ End UActorComponent Interface.
 
-	// Begin UPrimitiveComponent interface.
+	//~ Begin UPrimitiveComponent Interface.
 	virtual FBodyInstance* GetBodyInstance(FName BoneName = NAME_None, bool bGetWelded = true) const override;
 	virtual bool CanEditSimulatePhysics() override;
 	virtual bool IsAnySimulatingPhysics() const override;
@@ -139,17 +139,17 @@ public:
 	virtual void SetSimulatePhysics(bool bSimulate) override;
 
 	virtual void SetMaterial(int32 ElementIndex, UMaterialInterface* Material) override;
-	// End UPrimitiveComponent interface.
+	//~ End UPrimitiveComponent Interface.
 
-	// Begin SkinnedMeshComponent interface.
+	//~ Begin SkinnedMeshComponent Interface.
 	virtual bool ShouldUpdateTransform(bool bLODHasChanged) const override;
 	virtual void RefreshBoneTransforms(FActorComponentTickFunction* TickFunction = NULL) override;
 	virtual void SetSkeletalMesh(USkeletalMesh* InSkelMesh) override;
 	virtual FTransform GetSocketTransform(FName InSocketName, ERelativeTransformSpace TransformSpace = RTS_World) const override;
-	// End SkinnedMeshComponent interface.
+	//~ End SkinnedMeshComponent Interface.
 
 
-	// Begin DestructibleComponent interface.
+	//~ Begin DestructibleComponent Interface.
 #if WITH_APEX
 	struct FFakeBodyInstanceState
 	{
@@ -204,7 +204,7 @@ public:
 	void OnVisibilityEvent(const physx::apex::NxApexChunkStateEventData & InDamageEvent);
 #endif // WITH_APEX
 
-	// End DestructibleComponent interface.
+	//~ End DestructibleComponent Interface.
 
 	virtual bool DoCustomNavigableGeometryExport(FNavigableGeometryExport& GeomExport) const override;
 
diff --git a/Engine/Source/Runtime/Engine/Classes/Components/DirectionalLightComponent.h b/Engine/Source/Runtime/Engine/Classes/Components/DirectionalLightComponent.h
--- a/Engine/Source/Runtime/Engine/Classes/Components/DirectionalLightComponent.h
+++ b/Engine/Source/Runtime/Engine/Classes/Components/DirectionalLightComponent.h
@@ -161,7 +161,7 @@ class UDirectionalLightComponent : public ULightComponent
 	UFUNCTION(BlueprintCallable, Category="Rendering|Lighting")
 	void SetLightShaftOverrideDirection(FVector NewValue);
 
-	// ULightComponent interface.
+	//~ Begin ULightComponent Interface.
 	virtual FVector4 GetLightPosition() const override;
 	virtual ELightComponentType GetLightType() const override;
 	virtual FLightmassLightSettings GetLightmassSettings() const override
@@ -177,13 +177,13 @@ class UDirectionalLightComponent : public ULightComponent
 		return bUsedAsAtmosphereSunLight;
 	}
 
-	// Begin UObject Interface
+	//~ Begin UObject Interface
 #if WITH_EDITOR
 	virtual void PostEditChangeProperty(FPropertyChangedEvent& PropertyChangedEvent) override;
 	virtual bool CanEditChange(const UProperty* InProperty) const override;
 #endif // WITH_EDITOR
 	virtual void Serialize(FArchive& Ar) override;
-	// Begin UObject Interface
+	//~ Begin UObject Interface
 
 	virtual void InvalidateLightingCacheDetailed(bool bInvalidateBuildEnqueuedLighting, bool bTranslationOnly) override;
 };
diff --git a/Engine/Source/Runtime/Engine/Classes/Components/DrawFrustumComponent.h b/Engine/Source/Runtime/Engine/Classes/Components/DrawFrustumComponent.h
--- a/Engine/Source/Runtime/Engine/Classes/Components/DrawFrustumComponent.h
+++ b/Engine/Source/Runtime/Engine/Classes/Components/DrawFrustumComponent.h
@@ -39,10 +39,10 @@ class UDrawFrustumComponent : public UPrimitiveComponent
 	class UTexture* Texture;
 
 
-	// Begin UPrimitiveComponent interface.
+	//~ Begin UPrimitiveComponent Interface.
 	virtual FPrimitiveSceneProxy* CreateSceneProxy() override;
 	virtual FBoxSphereBounds CalcBounds(const FTransform& LocalToWorld) const override;
-	// End UPrimitiveComponent interface.
+	//~ End UPrimitiveComponent Interface.
 };
 
 
diff --git a/Engine/Source/Runtime/Engine/Classes/Components/ExponentialHeightFogComponent.h b/Engine/Source/Runtime/Engine/Classes/Components/ExponentialHeightFogComponent.h
--- a/Engine/Source/Runtime/Engine/Classes/Components/ExponentialHeightFogComponent.h
+++ b/Engine/Source/Runtime/Engine/Classes/Components/ExponentialHeightFogComponent.h
@@ -86,21 +86,21 @@ public:
 	void SetStartDistance(float Value);
 
 protected:
-	// Begin UActorComponent interface.
+	//~ Begin UActorComponent Interface.
 	virtual void CreateRenderState_Concurrent() override;
 	virtual void SendRenderTransform_Concurrent() override;
 	virtual void DestroyRenderState_Concurrent() override;
-	// End UActorComponent interface.
+	//~ End UActorComponent Interface.
 
 	void AddFogIfNeeded();
 
 public:
-	// Begin UObject Interface
+	//~ Begin UObject Interface
 #if WITH_EDITOR
 	virtual void PostEditChangeProperty(FPropertyChangedEvent& PropertyChangedEvent) override;
 #endif // WITH_EDITOR
 	virtual void PostInterpChange(UProperty* PropertyThatChanged) override;
-	// End UObject Interface
+	//~ End UObject Interface
 };
 
 
diff --git a/Engine/Source/Runtime/Engine/Classes/Components/InstancedStaticMeshComponent.h b/Engine/Source/Runtime/Engine/Classes/Components/InstancedStaticMeshComponent.h
--- a/Engine/Source/Runtime/Engine/Classes/Components/InstancedStaticMeshComponent.h
+++ b/Engine/Source/Runtime/Engine/Classes/Components/InstancedStaticMeshComponent.h
@@ -167,11 +167,11 @@ public:
 	UPROPERTY()
 	class UPhysicsSerializer* PhysicsSerializer;
 
-	// Begin UActorComponent interface 
+	//~ Begin UActorComponent Interface 
 	virtual FActorComponentInstanceData* GetComponentInstanceData() const override;
-	// End UActorComponent interface 
+	//~ End UActorComponent Interface 
 
-	// Begin UPrimitiveComponent Interface
+	//~ Begin UPrimitiveComponent Interface
 	virtual FPrimitiveSceneProxy* CreateSceneProxy() override;
 	virtual void CreatePhysicsState() override;
 	virtual void DestroyPhysicsState() override;
@@ -185,11 +185,11 @@ public:
 	virtual void GetLightAndShadowMapMemoryUsage( int32& LightMapMemoryUsage, int32& ShadowMapMemoryUsage ) const override;
 	
 	virtual bool DoCustomNavigableGeometryExport(FNavigableGeometryExport& GeomExport) const override;
-	// End UPrimitiveComponent Interface
+	//~ End UPrimitiveComponent Interface
 	
-	// Begin UNavRelevantInterface interface.
+	//~ Begin UNavRelevantInterface Interface.
 	virtual void GetNavigationData(FNavigationRelevantData& Data) const override;
-	// End UPrimitiveComponent interface.
+	//~ End UPrimitiveComponent Interface.
 	
 	//Begin UObject Interface
 	virtual void Serialize(FArchive& Ar) override;
diff --git a/Engine/Source/Runtime/Engine/Classes/Components/LightComponent.h b/Engine/Source/Runtime/Engine/Classes/Components/LightComponent.h
--- a/Engine/Source/Runtime/Engine/Classes/Components/LightComponent.h
+++ b/Engine/Source/Runtime/Engine/Classes/Components/LightComponent.h
@@ -297,7 +297,7 @@ public:
 	/** Compute current light brightness based on whether there is a valid IES profile texture attached, and whether IES brightness is enabled */
 	float ComputeLightBrightness() const;
 
-	// Begin UObject interface.
+	//~ Begin UObject Interface.
 	virtual void Serialize(FArchive& Ar) override;
 	virtual void PostLoad() override;
 #if WITH_EDITOR
@@ -307,7 +307,7 @@ public:
 #endif // WITH_EDITOR
 	virtual void BeginDestroy() override;
 	virtual bool IsReadyForFinishDestroy() override;
-	// End UObject interface.
+	//~ End UObject Interface.
 
 	virtual FActorComponentInstanceData* GetComponentInstanceData() const override;
 	 void ApplyComponentInstanceData(class FPrecomputedLightInstanceData* ComponentInstanceData);
@@ -328,12 +328,12 @@ public:
 	}
 
 protected:
-	// Begin UActorComponent Interface
+	//~ Begin UActorComponent Interface
 	virtual void OnRegister() override;
 	virtual void CreateRenderState_Concurrent() override;
 	virtual void SendRenderTransform_Concurrent() override;
 	virtual void DestroyRenderState_Concurrent() override;
-	// Begin UActorComponent Interface
+	//~ Begin UActorComponent Interface
 
 public:
 	virtual void InvalidateLightingCacheDetailed(bool bInvalidateBuildEnqueuedLighting, bool bTranslationOnly) override;
diff --git a/Engine/Source/Runtime/Engine/Classes/Components/LineBatchComponent.h b/Engine/Source/Runtime/Engine/Classes/Components/LineBatchComponent.h
--- a/Engine/Source/Runtime/Engine/Classes/Components/LineBatchComponent.h
+++ b/Engine/Source/Runtime/Engine/Classes/Components/LineBatchComponent.h
@@ -155,16 +155,16 @@ class ULineBatchComponent : public UPrimitiveComponent
 	/** Draw a mesh */
 	void DrawMesh(TArray<FVector> const& Verts, TArray<int32> const& Indices, FColor const& Color, uint8 DepthPriority, float LifeTime);
 
-	// Begin UPrimitiveComponent interface.
+	//~ Begin UPrimitiveComponent Interface.
 	virtual FPrimitiveSceneProxy* CreateSceneProxy() override;
 	virtual FBoxSphereBounds CalcBounds(const FTransform& LocalToWorld) const override;
-	// End UPrimitiveComponent interface.
+	//~ End UPrimitiveComponent Interface.
 	
 	
-	// Begin UActorComponent interface.
+	//~ Begin UActorComponent Interface.
 	virtual void TickComponent(float DeltaTime, enum ELevelTick TickType, FActorComponentTickFunction *ThisTickFunction) override;
 	virtual void ApplyWorldOffset(const FVector& InOffset, bool bWorldShift) override;
-	// End UActorComponent interface.
+	//~ End UActorComponent Interface.
 
 	/** Clear all batched lines, points and meshes */
 	ENGINE_API void Flush();
diff --git a/Engine/Source/Runtime/Engine/Classes/Components/MaterialBillboardComponent.h b/Engine/Source/Runtime/Engine/Classes/Components/MaterialBillboardComponent.h
--- a/Engine/Source/Runtime/Engine/Classes/Components/MaterialBillboardComponent.h
+++ b/Engine/Source/Runtime/Engine/Classes/Components/MaterialBillboardComponent.h
@@ -69,10 +69,10 @@ class ENGINE_API UMaterialBillboardComponent : public UPrimitiveComponent
 		class UCurveFloat* DistanceToSizeCurve
 		);
 
-	// Begin UPrimitiveComponent Interface
+	//~ Begin UPrimitiveComponent Interface
 	virtual FPrimitiveSceneProxy* CreateSceneProxy() override;
 	virtual FBoxSphereBounds CalcBounds(const FTransform& LocalToWorld) const override;
 	virtual UMaterialInterface* GetMaterial(int32 Index) const override;
 	virtual void SetMaterial(int32 ElementIndex, class UMaterialInterface* Material) override;
-	// End UPrimitiveComponent Interface
+	//~ End UPrimitiveComponent Interface
 };
diff --git a/Engine/Source/Runtime/Engine/Classes/Components/MeshComponent.h b/Engine/Source/Runtime/Engine/Classes/Components/MeshComponent.h
--- a/Engine/Source/Runtime/Engine/Classes/Components/MeshComponent.h
+++ b/Engine/Source/Runtime/Engine/Classes/Components/MeshComponent.h
@@ -29,16 +29,16 @@ class ENGINE_API UMeshComponent : public UPrimitiveComponent
 	/** Returns override Materials count */
 	virtual int32 GetNumOverrideMaterials() const;
 
-	// Begin UObject Interface
+	//~ Begin UObject Interface
 	virtual void BeginDestroy() override;
-	// End UObject Interface
+	//~ End UObject Interface
 
-	// Begin UPrimitiveComponent Interface
+	//~ Begin UPrimitiveComponent Interface
 	virtual int32 GetNumMaterials() const override;
 	virtual UMaterialInterface* GetMaterial(int32 ElementIndex) const override;
 	virtual void SetMaterial(int32 ElementIndex, UMaterialInterface* Material) override;
 	virtual void GetUsedMaterials(TArray<UMaterialInterface*>& OutMaterials) const override;
-	// End UPrimitiveComponent Interface
+	//~ End UPrimitiveComponent Interface
 
 	/** Accesses the scene relevance information for the materials applied to the mesh. Valid from game thread only. */
 	FMaterialRelevance GetMaterialRelevance(ERHIFeatureLevel::Type InFeatureLevel) const;
diff --git a/Engine/Source/Runtime/Engine/Classes/Components/ModelComponent.h b/Engine/Source/Runtime/Engine/Classes/Components/ModelComponent.h
--- a/Engine/Source/Runtime/Engine/Classes/Components/ModelComponent.h
+++ b/Engine/Source/Runtime/Engine/Classes/Components/ModelComponent.h
@@ -72,7 +72,7 @@ public:
 	/** Copy model elements from the other component. This is used when duplicating components for PIE to guarantee correct rendering. */
 	ENGINE_API void CopyElementsFrom(UModelComponent* OtherModelComponent);
 
-	// Begin UPrimitiveComponent interface.
+	//~ Begin UPrimitiveComponent Interface.
 	virtual bool GetLightMapResolution( int32& Width, int32& Height ) const override;
 	virtual int32 GetStaticLightMapResolution() const override;
 	virtual void GetLightAndShadowMapMemoryUsage( int32& LightMapMemoryUsage, int32& ShadowMapMemoryUsage ) const override;
@@ -88,13 +88,13 @@ public:
 	virtual class UBodySetup* GetBodySetup() override { return ModelBodySetup; };
 	virtual int32 GetNumMaterials() const override;
 	virtual UMaterialInterface* GetMaterial(int32 MaterialIndex) const override;
-	// End UPrimitiveComponent interface.
+	//~ End UPrimitiveComponent Interface.
 
-	// Begin UActorComponent interface.
+	//~ Begin UActorComponent Interface.
 	virtual void InvalidateLightingCacheDetailed(bool bInvalidateBuildEnqueuedLighting, bool bTranslationOnly) override;
-	// End UActorComponent interface.
+	//~ End UActorComponent Interface.
 
-	// Begin UObject interface.
+	//~ Begin UObject Interface.
 	virtual void Serialize(FArchive& Ar) override;
 	virtual void PostLoad() override;
 #if WITH_EDITOR
@@ -102,13 +102,13 @@ public:
 #endif // WITH_EDITOR
 	virtual SIZE_T GetResourceSize(EResourceSizeMode::Type Mode) override;
 	static void AddReferencedObjects(UObject* InThis, FReferenceCollector& Collector);
-	// End UObject interface.
+	//~ End UObject Interface.
 
-	// Begin Interface_CollisionDataProvider Interface
+	//~ Begin Interface_CollisionDataProvider Interface
 	virtual bool GetPhysicsTriMeshData(struct FTriMeshCollisionData* CollisionData, bool InUseAllTriData) override;
 	virtual bool ContainsPhysicsTriMeshData(bool InUseAllTriData) const override;
 	virtual bool WantsNegXTriMesh() override { return false; }
-	// End Interface_CollisionDataProvider Interface
+	//~ End Interface_CollisionDataProvider Interface
 
 
 	/** Ensure ModelBodySetup is present and correctly configured. */
diff --git a/Engine/Source/Runtime/Engine/Classes/Components/PlaneReflectionCaptureComponent.h b/Engine/Source/Runtime/Engine/Classes/Components/PlaneReflectionCaptureComponent.h
--- a/Engine/Source/Runtime/Engine/Classes/Components/PlaneReflectionCaptureComponent.h
+++ b/Engine/Source/Runtime/Engine/Classes/Components/PlaneReflectionCaptureComponent.h
@@ -25,7 +25,7 @@ public:
 	virtual void UpdatePreviewShape() override;
 	virtual float GetInfluenceBoundingRadius() const override;
 
-	// Begin UObject Interface
-	// End UObject Interface
+	//~ Begin UObject Interface
+	//~ End UObject Interface
 };
 
diff --git a/Engine/Source/Runtime/Engine/Classes/Components/PointLightComponent.h b/Engine/Source/Runtime/Engine/Classes/Components/PointLightComponent.h
--- a/Engine/Source/Runtime/Engine/Classes/Components/PointLightComponent.h
+++ b/Engine/Source/Runtime/Engine/Classes/Components/PointLightComponent.h
@@ -71,13 +71,13 @@ class ENGINE_API UPointLightComponent : public ULightComponent
 	void SetSourceLength(float NewValue);
 
 protected:
-	// Begin UActorComponent Interface
+	//~ Begin UActorComponent Interface
 	virtual void SendRenderTransform_Concurrent() override;
-	// End UActorComponent Interface
+	//~ End UActorComponent Interface
 
 public:
 
-	// ULightComponent interface.
+	//~ Begin ULightComponent Interface.
 	virtual bool AffectsBounds(const FBoxSphereBounds& InBounds) const override;
 	virtual FVector4 GetLightPosition() const override;
 	virtual FBox GetBoundingBox() const override;
@@ -92,14 +92,14 @@ public:
 
 	virtual FLightSceneProxy* CreateSceneProxy() const override;
 
-	// Begin UObject Interface
+	//~ Begin UObject Interface
 	virtual void Serialize(FArchive& Ar) override;
 	virtual void PostLoad() override;
 #if WITH_EDITOR
 	virtual bool CanEditChange(const UProperty* InProperty) const override;
 	virtual void PostEditChangeProperty(FPropertyChangedEvent& PropertyChangedEvent) override;
 #endif // WITH_EDITOR
-	// End UObject Interface
+	//~ End UObject Interface
 
 	/** 
 	 * This is called when property is modified by InterpPropertyTracks
diff --git a/Engine/Source/Runtime/Engine/Classes/Components/PoseableMeshComponent.h b/Engine/Source/Runtime/Engine/Classes/Components/PoseableMeshComponent.h
--- a/Engine/Source/Runtime/Engine/Classes/Components/PoseableMeshComponent.h
+++ b/Engine/Source/Runtime/Engine/Classes/Components/PoseableMeshComponent.h
@@ -44,10 +44,10 @@ class ENGINE_API UPoseableMeshComponent : public USkinnedMeshComponent
 	UFUNCTION(BlueprintCallable, Category="Components|PoseableMesh")
 	void ResetBoneTransformByName(FName BoneName);
 
-	// Begin USkinnedMeshComponent Interface
+	//~ Begin USkinnedMeshComponent Interface
 	virtual void RefreshBoneTransforms(FActorComponentTickFunction* TickFunction = NULL) override;
 	virtual bool AllocateTransformData() override;
-	// End USkinnedMeshComponent Interface
+	//~ End USkinnedMeshComponent Interface
 
 	bool IsRunningParallelEvaluation() const { return false; }
 	/**
diff --git a/Engine/Source/Runtime/Engine/Classes/Components/PostProcessComponent.h b/Engine/Source/Runtime/Engine/Classes/Components/PostProcessComponent.h
--- a/Engine/Source/Runtime/Engine/Classes/Components/PostProcessComponent.h
+++ b/Engine/Source/Runtime/Engine/Classes/Components/PostProcessComponent.h
@@ -40,7 +40,7 @@ class UPostProcessComponent : public USceneComponent, public IInterface_PostProc
 	UPROPERTY(EditAnywhere, BlueprintReadWrite, Category = PostProcessVolume)
 	uint32 bUnbound: 1;
  
-	// Begin IInterface_PostProcessVolume Interface
+	//~ Begin IInterface_PostProcessVolume Interface
 	ENGINE_API virtual bool EncompassesPoint(FVector Point, float SphereRadius/*=0.f*/, float* OutDistanceToPoint) override;
 	ENGINE_API virtual FPostProcessVolumeProperties GetProperties() const override
 	{
@@ -53,7 +53,7 @@ class UPostProcessComponent : public USceneComponent, public IInterface_PostProc
 		Ret.Settings = &Settings;
 		return Ret;
 	}
-	// End IInterface_PostProcessVolume Interface
+	//~ End IInterface_PostProcessVolume Interface
 
 	/** Adds an Blendable (implements IBlendableInterface) to the array of Blendables (if it doesn't exist) and update the weight */
 	UFUNCTION(BlueprintCallable, Category="Rendering")
@@ -64,9 +64,9 @@ protected:
 	virtual void OnRegister() override;
 	virtual void OnUnregister() override;
 
-	// UObject interface
+	//~ Begin UObject Interface
 	virtual void Serialize(FArchive& Ar) override;
-	// End of UObject interface
+	//~ End UObject Interface
 };
 
 
diff --git a/Engine/Source/Runtime/Engine/Classes/Components/PrimitiveComponent.h b/Engine/Source/Runtime/Engine/Classes/Components/PrimitiveComponent.h
--- a/Engine/Source/Runtime/Engine/Classes/Components/PrimitiveComponent.h
+++ b/Engine/Source/Runtime/Engine/Classes/Components/PrimitiveComponent.h
@@ -1062,13 +1062,13 @@ public:
 	virtual const bool ShouldGenerateAutoLOD() const;
 #endif
 
-	// Begin UActorComponent Interface
+	//~ Begin UActorComponent Interface
 	virtual void InvalidateLightingCacheDetailed(bool bInvalidateBuildEnqueuedLighting, bool bTranslationOnly) override;
 	virtual bool IsEditorOnly() const override;
 	virtual bool ShouldCreatePhysicsState() const override;
 	virtual bool HasValidPhysicsState() const override;
 	virtual class FActorComponentInstanceData* GetComponentInstanceData() const override;
-	// End UActorComponent Interface
+	//~ End UActorComponent Interface
 
 	/** @return true if the owner is selected and this component is selectable */
 	virtual bool ShouldRenderSelected() const;
@@ -1334,7 +1334,7 @@ protected:
 	/** Give the static mesh component recreate render state context access to Create/DestroyRenderState_Concurrent(). */
 	friend class FStaticMeshComponentRecreateRenderStateContext;
 
-	// Begin USceneComponent Interface
+	//~ Begin USceneComponent Interface
 	virtual void OnUpdateTransform(bool bSkipPhysicsMove, ETeleportType Teleport = ETeleportType::None) override;
 
 	/** Event called when AttachParent changes, to allow the scene to update its attachment state. */
@@ -1358,7 +1358,7 @@ public:
 	// End USceneComponentInterface
 
 
-	// Begin UActorComponent Interface
+	//~ Begin UActorComponent Interface
 protected:
 	virtual void CreateRenderState_Concurrent() override;
 	virtual void SendRenderTransform_Concurrent() override;
@@ -1383,7 +1383,7 @@ public:
 #if WITH_EDITOR
 	virtual void CheckForErrors() override;
 #endif // WITH_EDITOR	
-	// End UActorComponent Interface
+	//~ End UActorComponent Interface
 
 protected:
 	/** Internal function that updates physics objects to match the component collision settings. */
@@ -1396,7 +1396,7 @@ protected:
 	void EnsurePhysicsStateCreated();
 public:
 
-	// Begin UObject interface.
+	//~ Begin UObject Interface.
 	virtual void Serialize(FArchive& Ar) override;
 #if WITH_EDITOR
 	virtual void PostEditChangeProperty(FPropertyChangedEvent& PropertyChangedEvent) override;
@@ -1428,7 +1428,7 @@ public:
 	virtual bool IsReadyForFinishDestroy() override;
 	virtual bool NeedsLoadForClient() const override;
 	virtual bool NeedsLoadForServer() const override;
-	// End UObject interface.
+	//~ End UObject Interface.
 
 	//Begin USceneComponent Interface
 
@@ -1817,10 +1817,10 @@ public:
 		SetCanEverAffectNavigation(false);
 	}
 
-	// Begin INavRelevantInterface Interface
+	//~ Begin INavRelevantInterface Interface
 	virtual FBox GetNavigationBounds() const override;
 	virtual bool IsNavigationRelevant() const override;
-	// End INavRelevantInterface Interface
+	//~ End INavRelevantInterface Interface
 
 	FORCEINLINE EHasCustomNavigableGeometry::Type HasCustomNavigableGeometry() const { return bHasCustomNavigableGeometry; }
 
diff --git a/Engine/Source/Runtime/Engine/Classes/Components/ReflectionCaptureComponent.h b/Engine/Source/Runtime/Engine/Classes/Components/ReflectionCaptureComponent.h
--- a/Engine/Source/Runtime/Engine/Classes/Components/ReflectionCaptureComponent.h
+++ b/Engine/Source/Runtime/Engine/Classes/Components/ReflectionCaptureComponent.h
@@ -109,13 +109,13 @@ class UReflectionCaptureComponent : public USceneComponent
 		return FullHDRDerivedData;
 	}
 
-	// Begin UActorComponent Interface
+	//~ Begin UActorComponent Interface
 	virtual void CreateRenderState_Concurrent() override;
 	virtual void DestroyRenderState_Concurrent() override;
 	virtual void SendRenderTransform_Concurrent() override;
-	// End UActorComponent Interface
+	//~ End UActorComponent Interface
 
-	// Begin UObject Interface
+	//~ Begin UObject Interface
 	virtual void PostInitProperties() override;	
 	virtual void Serialize(FArchive& Ar) override;
 	virtual void PostLoad() override;
@@ -129,7 +129,7 @@ class UReflectionCaptureComponent : public USceneComponent
 	virtual void BeginDestroy() override;
 	virtual bool IsReadyForFinishDestroy() override;
 	virtual void FinishDestroy() override;
-	// End UObject Interface
+	//~ End UObject Interface
 
 private:
 
diff --git a/Engine/Source/Runtime/Engine/Classes/Components/SceneCaptureComponent2D.h b/Engine/Source/Runtime/Engine/Classes/Components/SceneCaptureComponent2D.h
--- a/Engine/Source/Runtime/Engine/Classes/Components/SceneCaptureComponent2D.h
+++ b/Engine/Source/Runtime/Engine/Classes/Components/SceneCaptureComponent2D.h
@@ -43,7 +43,7 @@ class USceneCaptureComponent2D : public USceneCaptureComponent
 private:
 
 public:
-	// Begin UActorComponent Interface
+	//~ Begin UActorComponent Interface
 	virtual void OnRegister() override;
 	virtual void SendRenderTransform_Concurrent() override;
 	virtual bool RequiresGameThreadEndOfFrameUpdates() const override
@@ -52,16 +52,16 @@ public:
 		return true;
 	}
 	virtual void TickComponent(float DeltaTime, enum ELevelTick TickType, FActorComponentTickFunction *ThisTickFunction) override;
-	// End UActorComponent Interface
+	//~ End UActorComponent Interface
 
-	// Begin UObject Interface
+	//~ Begin UObject Interface
 #if WITH_EDITOR
 	virtual void PostEditChangeProperty(FPropertyChangedEvent& PropertyChangedEvent) override;
 #endif // WITH_EDITOR
 	
 	virtual void Serialize(FArchive& Ar);
 
-	// End UObject Interface
+	//~ End UObject Interface
 
 	/** Adds an Blendable (implements IBlendableInterface) to the array of Blendables (if it doesn't exist) and update the weight */
 	UFUNCTION(BlueprintCallable, Category="Rendering")
diff --git a/Engine/Source/Runtime/Engine/Classes/Components/SceneComponent.h b/Engine/Source/Runtime/Engine/Classes/Components/SceneComponent.h
--- a/Engine/Source/Runtime/Engine/Classes/Components/SceneComponent.h
+++ b/Engine/Source/Runtime/Engine/Classes/Components/SceneComponent.h
@@ -678,7 +678,7 @@ public:
 	UPROPERTY(BlueprintAssignable, Category=PhysicsVolume, meta=(DisplayName="Physics Volume Changed"))
 	FPhysicsVolumeChanged PhysicsVolumeChangedDelegate;
 
-	// Begin ActorComponent interface
+	//~ Begin ActorComponent Interface
 	virtual void OnRegister() override;
 	/** Return true if CreateRenderState() should be called */
 	virtual bool ShouldCreateRenderState() const override
@@ -690,12 +690,12 @@ public:
 	virtual void OnComponentDestroyed() override;
 	virtual void ApplyWorldOffset(const FVector& InOffset, bool bWorldShift) override;
 	virtual class FActorComponentInstanceData* GetComponentInstanceData() const override;
-	// End ActorComponent interface
+	//~ End ActorComponent Interface
 
 	// Call UpdateComponentToWorld if bWorldToComponentUpdated is false.
 	void ConditionalUpdateComponentToWorld();
 
-	// Begin UObject Interface
+	//~ Begin UObject Interface
 	virtual void Serialize(FArchive& Ar) override;
 	virtual void PostInterpChange(UProperty* PropertyThatChanged) override;
 	virtual void BeginDestroy() override;
@@ -706,7 +706,7 @@ public:
 	virtual void PostEditChangeChainProperty(FPropertyChangedChainEvent& PropertyChangedEvent) override;
 #endif
 
-	// End UObject Interface
+	//~ End UObject Interface
 protected:
 	/**
 	 * Internal helper, for use from MoveComponent().  Special codepath since the normal setters call MoveComponent.
diff --git a/Engine/Source/Runtime/Engine/Classes/Components/ShapeComponent.h b/Engine/Source/Runtime/Engine/Classes/Components/ShapeComponent.h
--- a/Engine/Source/Runtime/Engine/Classes/Components/ShapeComponent.h
+++ b/Engine/Source/Runtime/Engine/Classes/Components/ShapeComponent.h
@@ -36,30 +36,30 @@ class ENGINE_API UShapeComponent : public UPrimitiveComponent
 	UPROPERTY(EditAnywhere, Category = Navigation)
 	TSubclassOf<class UNavArea> AreaClass;
 
-	// Begin UPrimitiveComponent interface.
+	//~ Begin UPrimitiveComponent Interface.
 	virtual FPrimitiveSceneProxy* CreateSceneProxy() override;
 	virtual class UBodySetup* GetBodySetup() override;
 	virtual bool DoCustomNavigableGeometryExport(FNavigableGeometryExport& GeomExport) const override;
 	virtual void GetNavigationData(FNavigationRelevantData& Data) const override;
-	// End UPrimitiveComponent interface.
+	//~ End UPrimitiveComponent Interface.
 
-	// Begin INavRelevantInterface Interface
+	//~ Begin INavRelevantInterface Interface
 	virtual bool IsNavigationRelevant() const override;
-	// End INavRelevantInterface Interface
+	//~ End INavRelevantInterface Interface
 
-	// Begin USceneComponent interface
+	//~ Begin USceneComponent Interface
 	virtual FBoxSphereBounds CalcBounds(const FTransform& LocalToWorld) const override;
 	virtual bool ShouldCollideWhenPlacing() const override
 	{
 		return bShouldCollideWhenPlacing || IsCollisionEnabled();
 	}
-	// End USceneComponent interface
+	//~ End USceneComponent Interface
 
-	// Begin UObject interface.
+	//~ Begin UObject Interface.
 #if WITH_EDITOR
 	virtual void PostEditChangeProperty(FPropertyChangedEvent& PropertyChangedEvent) override;
 #endif // WITH_EDITOR
-	// End UObject interface.
+	//~ End UObject Interface.
 
 	/** Update the body setup parameters based on shape information*/
 	virtual void UpdateBodySetup();
diff --git a/Engine/Source/Runtime/Engine/Classes/Components/SkeletalMeshComponent.h b/Engine/Source/Runtime/Engine/Classes/Components/SkeletalMeshComponent.h
--- a/Engine/Source/Runtime/Engine/Classes/Components/SkeletalMeshComponent.h
+++ b/Engine/Source/Runtime/Engine/Classes/Components/SkeletalMeshComponent.h
@@ -789,7 +789,7 @@ public:
 	void RecalcRequiredBones(int32 LODIndex);
 
 public:
-	// Begin UObject interface.
+	//~ Begin UObject Interface.
 	virtual void Serialize(FArchive& Ar) override;
 #if WITH_EDITOR
 	DECLARE_MULTICAST_DELEGATE(FOnSkeletalMeshPropertyChangedMulticaster)
@@ -809,9 +809,9 @@ public:
 	virtual void UpdateCollisionProfile() override;
 #endif // WITH_EDITOR
 	virtual SIZE_T GetResourceSize(EResourceSizeMode::Type Mode) override;
-	// End UObject interface.
+	//~ End UObject Interface.
 
-	// Begin UActorComponent interface.
+	//~ Begin UActorComponent Interface.
 	virtual void OnRegister() override;
 	virtual void OnUnregister() override;
 	virtual void CreateRenderState_Concurrent() override;
@@ -824,17 +824,17 @@ public:
 	//Handle registering our pre cloth tick function
 	void RegisterPreClothTick(bool bRegister);
 
-	// End UActorComponent interface.
+	//~ End UActorComponent Interface.
 
-	// Begin USceneComponent interface.
+	//~ Begin USceneComponent Interface.
 	virtual void UpdateBounds() override;
 	virtual FBoxSphereBounds CalcBounds(const FTransform& LocalToWorld) const override;
 	virtual bool IsAnySimulatingPhysics() const override;
 	virtual void OnUpdateTransform(bool bSkipPhysicsMove, ETeleportType Teleport = ETeleportType::None) override;
 	virtual void UpdateOverlaps(TArray<FOverlapInfo> const* PendingOverlaps=NULL, bool bDoNotifies=true, const TArray<FOverlapInfo>* OverlapsAtEndLocation=NULL) override;
-	// End USceneComponent interface.
+	//~ End USceneComponent Interface.
 
-	// Begin UPrimitiveComponent interface.
+	//~ Begin UPrimitiveComponent Interface.
 protected:
 	/**
 	 *  Test the collision of the supplied component at the supplied location/rotation, and determine the set of components that it overlaps
@@ -882,10 +882,10 @@ public:
 #endif
 protected:
 	virtual FTransform GetComponentTransformFromBodyInstance(FBodyInstance* UseBI) override;
-	// End UPrimitiveComponent interface.
+	//~ End UPrimitiveComponent Interface.
 
 public:
-	// Begin USkinnedMeshComponent interface
+	//~ Begin USkinnedMeshComponent Interface
 	virtual bool UpdateLODStatus() override;
 	virtual void RefreshBoneTransforms( FActorComponentTickFunction* TickFunction = NULL ) override;
 	virtual void TickPose(float DeltaTime, bool bNeedsValidRootMotion) override;
@@ -905,7 +905,7 @@ public:
 	virtual bool IsPlayingRootMotion() override;
 	virtual bool IsPlayingRootMotionFromEverything() override;
 	virtual void FinalizeBoneTransform() override;
-	// End USkinnedMeshComponent interface
+	//~ End USkinnedMeshComponent Interface
 	/** 
 	 *	Iterate over each joint in the physics for this mesh, setting its AngularPositionTarget based on the animation information.
 	 */
diff --git a/Engine/Source/Runtime/Engine/Classes/Components/SkinnedMeshComponent.h b/Engine/Source/Runtime/Engine/Classes/Components/SkinnedMeshComponent.h
--- a/Engine/Source/Runtime/Engine/Classes/Components/SkinnedMeshComponent.h
+++ b/Engine/Source/Runtime/Engine/Classes/Components/SkinnedMeshComponent.h
@@ -443,17 +443,17 @@ public:
 	/** Gets the skeletal mesh resource used for rendering the component. */
 	FSkeletalMeshResource* GetSkeletalMeshResource() const;
 
-	// Begin UObject interface
+	//~ Begin UObject Interface
 	virtual void Serialize(FArchive& Ar) override;
 	virtual SIZE_T GetResourceSize(EResourceSizeMode::Type Mode) override;
 	virtual FString GetDetailedInfoInternal() const override;
 #if WITH_EDITOR
 	virtual void PostEditChangeProperty(FPropertyChangedEvent& PropertyChangedEvent) override;
 #endif // WITH_EDITOR
-	// End UObject interface
+	//~ End UObject Interface
 
 protected:
-	// Begin UActorComponent interface
+	//~ Begin UActorComponent Interface
 	virtual void OnRegister() override;
 	virtual void OnUnregister() override;
 	virtual void CreateRenderState_Concurrent() override;
@@ -465,25 +465,25 @@ protected:
 	}
 	virtual void TickComponent(float DeltaTime, enum ELevelTick TickType, FActorComponentTickFunction *ThisTickFunction) override;
 	virtual UObject const* AdditionalStatObject() const override;
-	// End UActorComponent interface
+	//~ End UActorComponent Interface
 
 public:
-	// Begin USceneComponent interface
+	//~ Begin USceneComponent Interface
 	virtual FBoxSphereBounds CalcBounds(const FTransform& LocalToWorld) const override;
 	virtual FTransform GetSocketTransform(FName InSocketName, ERelativeTransformSpace TransformSpace = RTS_World) const override;
 	virtual bool DoesSocketExist(FName InSocketName) const override;
 	virtual bool HasAnySockets() const override;
 	virtual void QuerySupportedSockets(TArray<FComponentSocketDescription>& OutSockets) const override;
 	virtual void UpdateOverlaps(TArray<FOverlapInfo> const* PendingOverlaps=NULL, bool bDoNotifies=true, const TArray<FOverlapInfo>* OverlapsAtEndLocation=NULL) override;
-	// End USceneComponent interface
+	//~ End USceneComponent Interface
 
-	// Begin UPrimitiveComponent interface
+	//~ Begin UPrimitiveComponent Interface
 	virtual UMaterialInterface* GetMaterial(int32 MaterialIndex) const override;
 	virtual FPrimitiveSceneProxy* CreateSceneProxy() override;
 	virtual void GetUsedMaterials(TArray<UMaterialInterface*>& OutMaterials) const override;
 	virtual void GetStreamingTextureInfo(TArray<FStreamingTexturePrimitiveInfo>& OutStreamingTextures) const override;
 	virtual int32 GetNumMaterials() const override;
-	// End UPrimitiveComponent interface
+	//~ End UPrimitiveComponent Interface
 
 	/**
 	 *	Sets the value of the bForceWireframe flag and reattaches the component as necessary.
diff --git a/Engine/Source/Runtime/Engine/Classes/Components/SkyLightComponent.h b/Engine/Source/Runtime/Engine/Classes/Components/SkyLightComponent.h
--- a/Engine/Source/Runtime/Engine/Classes/Components/SkyLightComponent.h
+++ b/Engine/Source/Runtime/Engine/Classes/Components/SkyLightComponent.h
@@ -137,7 +137,7 @@ class ENGINE_API USkyLightComponent : public ULightComponentBase
 
 	class FSkyLightSceneProxy* CreateSceneProxy() const;
 
-	// Begin UObject Interface
+	//~ Begin UObject Interface
 	virtual void PostInitProperties() override;
 	virtual void PostLoad() override;
 	virtual void PostInterpChange(UProperty* PropertyThatChanged) override;
@@ -148,7 +148,7 @@ class ENGINE_API USkyLightComponent : public ULightComponentBase
 #endif // WITH_EDITOR
 	virtual void BeginDestroy() override;
 	virtual bool IsReadyForFinishDestroy() override;
-	// End UObject Interface
+	//~ End UObject Interface
 
 	virtual FActorComponentInstanceData* GetComponentInstanceData() const override;
 	void ApplyComponentInstanceData(class FPrecomputedSkyLightInstanceData* ComponentInstanceData);
@@ -242,10 +242,10 @@ protected:
 	static TArray<USkyLightComponent*> SkyCapturesToUpdate;
 	static TArray<USkyLightComponent*> SkyCapturesToUpdateBlendDestinations;
 
-	// Begin UActorComponent Interface
+	//~ Begin UActorComponent Interface
 	virtual void CreateRenderState_Concurrent() override;
 	virtual void DestroyRenderState_Concurrent() override;
-	// Begin UActorComponent Interface
+	//~ Begin UActorComponent Interface
 
 	void UpdateLimitedRenderingStateFast();
 
diff --git a/Engine/Source/Runtime/Engine/Classes/Components/SphereComponent.h b/Engine/Source/Runtime/Engine/Classes/Components/SphereComponent.h
--- a/Engine/Source/Runtime/Engine/Classes/Components/SphereComponent.h
+++ b/Engine/Source/Runtime/Engine/Classes/Components/SphereComponent.h
@@ -35,21 +35,21 @@ public:
 	UFUNCTION(BlueprintCallable, Category="Components|Sphere")
 	float GetUnscaledSphereRadius() const;
 
-	// Begin UPrimitiveComponent interface.
+	//~ Begin UPrimitiveComponent Interface.
 	virtual FPrimitiveSceneProxy* CreateSceneProxy() override;
 	virtual bool IsZeroExtent() const override;
 	virtual struct FCollisionShape GetCollisionShape(float Inflation = 0.0f) const override;
 	virtual bool AreSymmetricRotations(const FQuat& A, const FQuat& B, const FVector& Scale3D) const override;
-	// End UPrimitiveComponent interface.
+	//~ End UPrimitiveComponent Interface.
 
-	// Begin USceneComponent interface
+	//~ Begin USceneComponent Interface
 	virtual FBoxSphereBounds CalcBounds(const FTransform& LocalToWorld) const override;
 	virtual void CalcBoundingCylinder(float& CylinderRadius, float& CylinderHalfHeight) const override;
-	// End USceneComponent interface
+	//~ End USceneComponent Interface
 
-	// Begin UShapeComponent interface
+	//~ Begin UShapeComponent Interface
 	virtual void UpdateBodySetup() override;
-	// End UShapeComponent interface
+	//~ End UShapeComponent Interface
 
 	// Get the scale used by this shape. This is a uniform scale that is the minimum of any non-uniform scaling.
 	// @return the scale used by this shape.
diff --git a/Engine/Source/Runtime/Engine/Classes/Components/SplineComponent.h b/Engine/Source/Runtime/Engine/Classes/Components/SplineComponent.h
--- a/Engine/Source/Runtime/Engine/Classes/Components/SplineComponent.h
+++ b/Engine/Source/Runtime/Engine/Classes/Components/SplineComponent.h
@@ -107,9 +107,9 @@ public:
 	float ScaleVisualizationWidth;
 #endif
 
-	// Begin UActorComponent interface.
+	//~ Begin UActorComponent Interface.
 	virtual FActorComponentInstanceData* GetComponentInstanceData() const override;
-	// End UActorComponent interface.
+	//~ End UActorComponent Interface.
 
 	void ApplyComponentInstanceData(class FSplineInstanceData* ComponentInstanceData, const bool bPostUCS);
 
@@ -450,14 +450,14 @@ public:
 	FTransform FindTransformClosestToWorldLocation(const FVector& WorldLocation, ESplineCoordinateSpace::Type CoordinateSpace, bool bUseScale = false) const;
 
 
-	// UObject interface
+	//~ Begin UObject Interface
 	virtual void PostLoad() override;
 	virtual void PostEditImport() override;
 #if WITH_EDITOR
 	virtual void Serialize(FArchive& Ar) override;
 	virtual void PostEditChangeChainProperty(FPropertyChangedChainEvent& PropertyChangedEvent) override;
 #endif
-	// End of UObject interface
+	//~ End UObject Interface
 
 private:
 	/** Returns the length of the specified spline segment up to the parametric value given */
diff --git a/Engine/Source/Runtime/Engine/Classes/Components/SplineMeshComponent.h b/Engine/Source/Runtime/Engine/Classes/Components/SplineMeshComponent.h
--- a/Engine/Source/Runtime/Engine/Classes/Components/SplineMeshComponent.h
+++ b/Engine/Source/Runtime/Engine/Classes/Components/SplineMeshComponent.h
@@ -143,9 +143,9 @@ class ENGINE_API USplineMeshComponent : public UStaticMeshComponent, public IInt
 	virtual bool Modify(bool bAlwaysMarkDirty = true) override;
 	//End UObject Interface
 
-	// Begin UActorComponent interface.
+	//~ Begin UActorComponent Interface.
 	virtual FActorComponentInstanceData* GetComponentInstanceData() const override;
-	// End UActorComponent interface.
+	//~ End UActorComponent Interface.
 
 	void ApplyComponentInstanceData(class FSplineMeshInstanceData* ComponentInstanceData);
 
@@ -170,12 +170,12 @@ class ENGINE_API USplineMeshComponent : public UStaticMeshComponent, public IInt
 	virtual class FStaticMeshStaticLightingMesh* AllocateStaticLightingMesh(int32 LODIndex, const TArray<ULightComponent*>& InRelevantLights) override;
 	//End UStaticMeshComponent Interface
 
-	// Begin Interface_CollisionDataProvider Interface
+	//~ Begin Interface_CollisionDataProvider Interface
 	virtual bool GetPhysicsTriMeshData(struct FTriMeshCollisionData* CollisionData, bool InUseAllTriData) override;
 	virtual bool ContainsPhysicsTriMeshData(bool InUseAllTriData) const override;
 	virtual bool WantsNegXTriMesh() override { return false; }
 	virtual void GetMeshId(FString& OutMeshId) override;
-	// End Interface_CollisionDataProvider Interface
+	//~ End Interface_CollisionDataProvider Interface
 
 	/** Called to notify render thread and possibly collision of a change in spline params or mesh */
 	void UpdateRenderStateAndCollision();
