Cleaning out programmer-only comments that were harvested by Doxygen.

diff --git a/Engine/Source/Runtime/Engine/Classes/Components/StaticMeshComponent.h b/Engine/Source/Runtime/Engine/Classes/Components/StaticMeshComponent.h
--- a/Engine/Source/Runtime/Engine/Classes/Components/StaticMeshComponent.h
+++ b/Engine/Source/Runtime/Engine/Classes/Components/StaticMeshComponent.h
@@ -239,7 +239,7 @@ class ENGINE_API UStaticMeshComponent : public UMeshComponent
 
 public:
 
-	// Begin UObject interface.
+	//~ Begin UObject Interface.
 	virtual void BeginDestroy() override;
 	virtual void ExportCustomProperties(FOutputDevice& Out, uint32 Indent) override;
 	virtual void ImportCustomProperties(const TCHAR* SourceText, FFeedbackContext* Warn) override;	
@@ -254,9 +254,9 @@ public:
 	virtual bool AreNativePropertiesIdenticalTo( UObject* Other ) const override;
 	virtual FString GetDetailedInfoInternal() const override;
 	static void AddReferencedObjects(UObject* InThis, FReferenceCollector& Collector);
-	// End UObject interface.
+	//~ End UObject Interface.
 
-	// Begin USceneComponent Interface
+	//~ Begin USceneComponent Interface
 	virtual FBoxSphereBounds CalcBounds(const FTransform& LocalToWorld) const override;
 	virtual bool HasAnySockets() const override;
 	virtual void QuerySupportedSockets(TArray<FComponentSocketDescription>& OutSockets) const override;
@@ -268,9 +268,9 @@ public:
 		// return IsCollisionEnabled() && (StaticMesh != NULL);
 		return false;
 	}
-	// End USceneComponent Interface
+	//~ End USceneComponent Interface
 
-	// Begin UActorComponent interface.
+	//~ Begin UActorComponent Interface.
 protected: 
 	virtual void OnRegister() override;
 	virtual void OnUnregister() override;
@@ -281,9 +281,9 @@ public:
 	virtual void CheckForErrors() override;
 #endif
 	virtual FActorComponentInstanceData* GetComponentInstanceData() const override;
-	// End UActorComponent interface.
+	//~ End UActorComponent Interface.
 
-	// Begin UPrimitiveComponent interface.
+	//~ Begin UPrimitiveComponent Interface.
 	virtual int32 GetNumMaterials() const override;
 #if WITH_EDITOR
 	virtual void GetStaticLightingInfo(FStaticLightingPrimitiveInfo& OutPrimitiveInfo,const TArray<ULightComponent*>& InRelevantLights,const FLightingBuildOptions& Options) override;
@@ -315,12 +315,12 @@ public:
 	virtual bool ComponentIsTouchingSelectionBox(const FBox& InSelBBox, const FEngineShowFlags& ShowFlags, const bool bConsiderOnlyBSP, const bool bMustEncompassEntireComponent) const override;
 	virtual bool ComponentIsTouchingSelectionFrustum(const FConvexVolume& InSelBBox, const FEngineShowFlags& ShowFlags, const bool bConsiderOnlyBSP, const bool bMustEncompassEntireComponent) const override;
 #endif
-	// End UPrimitiveComponent interface.
+	//~ End UPrimitiveComponent Interface.
 
-	// Begin INavRelevantInterface interface.
+	//~ Begin INavRelevantInterface Interface.
 	virtual bool IsNavigationRelevant() const override;
 	virtual void GetNavigationData(FNavigationRelevantData& Data) const override;
-	// End INavRelevantInterface interface.
+	//~ End INavRelevantInterface Interface.
 	/**
 	 *	Returns true if the component uses texture lightmaps
 	 *
diff --git a/Engine/Source/Runtime/Engine/Classes/Components/TextRenderComponent.h b/Engine/Source/Runtime/Engine/Classes/Components/TextRenderComponent.h
--- a/Engine/Source/Runtime/Engine/Classes/Components/TextRenderComponent.h
+++ b/Engine/Source/Runtime/Engine/Classes/Components/TextRenderComponent.h
@@ -144,7 +144,7 @@ class ENGINE_API UTextRenderComponent : public UPrimitiveComponent
 
 	// -----------------------------
 
-	// Begin UPrimitiveComponent interface.
+	//~ Begin UPrimitiveComponent Interface.
 	virtual FPrimitiveSceneProxy* CreateSceneProxy() override;
 	virtual void GetUsedMaterials( TArray<UMaterialInterface*>& OutMaterials ) const override;
 	virtual int32 GetNumMaterials() const override;
@@ -152,15 +152,15 @@ class ENGINE_API UTextRenderComponent : public UPrimitiveComponent
 	virtual bool ShouldRecreateProxyOnUpdateTransform() const override;
 	virtual void SetMaterial(int32 ElementIndex, UMaterialInterface* InMaterial) override;
 	virtual FMatrix GetRenderMatrix() const override;
-	// End UPrimitiveComponent interface.
+	//~ End UPrimitiveComponent Interface.
 
-	// Begin USceneComponent interface.
+	//~ Begin USceneComponent Interface.
 	virtual FBoxSphereBounds CalcBounds(const FTransform& LocalToWorld) const override;
-	// End USceneComponent interface.
+	//~ End USceneComponent Interface.
 
-	// Begin UObject interface.
+	//~ Begin UObject Interface.
 	virtual void PostLoad() override;
-	// End UObject interface.
+	//~ End UObject Interface.
 };
 
 
diff --git a/Engine/Source/Runtime/Engine/Classes/Components/TimelineComponent.h b/Engine/Source/Runtime/Engine/Classes/Components/TimelineComponent.h
--- a/Engine/Source/Runtime/Engine/Classes/Components/TimelineComponent.h
+++ b/Engine/Source/Runtime/Engine/Classes/Components/TimelineComponent.h
@@ -444,11 +444,11 @@ public:
 	UFUNCTION()
 	void OnRep_Timeline();
 
-	// Begin ActorComponent interface.
+	//~ Begin ActorComponent Interface.
 	virtual void TickComponent(float DeltaTime, enum ELevelTick TickType, FActorComponentTickFunction *ThisTickFunction) override;
 	virtual void Activate(bool bReset=false) override;
 	virtual void Deactivate() override;
-	// End ActorComponent interface.
+	//~ End ActorComponent Interface.
 
 	/** Get the signature function for Timeline event functions */
 	ENGINE_API static UFunction* GetTimelineEventSignature();
diff --git a/Engine/Source/Runtime/Engine/Classes/Components/VectorFieldComponent.h b/Engine/Source/Runtime/Engine/Classes/Components/VectorFieldComponent.h
--- a/Engine/Source/Runtime/Engine/Classes/Components/VectorFieldComponent.h
+++ b/Engine/Source/Runtime/Engine/Classes/Components/VectorFieldComponent.h
@@ -46,23 +46,23 @@ public:
 	class FVectorFieldInstance* VectorFieldInstance;
 
 
-	// Begin UPrimitiveComponent interface.
+	//~ Begin UPrimitiveComponent Interface.
 	virtual FPrimitiveSceneProxy* CreateSceneProxy() override;
 	virtual FBoxSphereBounds CalcBounds(const FTransform& LocalToWorld) const override;
-	// End UPrimitiveComponent interface.
+	//~ End UPrimitiveComponent Interface.
 
-	// Begin UActorComponent interface.
+	//~ Begin UActorComponent Interface.
 	virtual void OnRegister() override;
 	virtual void OnUnregister() override;
 	virtual void SendRenderTransform_Concurrent() override;
-	// End UActorComponent interface.
+	//~ End UActorComponent Interface.
 
-	// Begin UObject interface.
+	//~ Begin UObject Interface.
 	virtual void PostInterpChange(UProperty* PropertyThatChanged) override;
 #if WITH_EDITOR
 	virtual void PostEditChangeProperty(FPropertyChangedEvent& PropertyChangedEvent) override;
 #endif // WITH_EDITOR
-	// End UObject interface.
+	//~ End UObject Interface.
 };
 
 
diff --git a/Engine/Source/Runtime/Engine/Classes/Components/WindDirectionalSourceComponent.h b/Engine/Source/Runtime/Engine/Classes/Components/WindDirectionalSourceComponent.h
--- a/Engine/Source/Runtime/Engine/Classes/Components/WindDirectionalSourceComponent.h
+++ b/Engine/Source/Runtime/Engine/Classes/Components/WindDirectionalSourceComponent.h
@@ -27,11 +27,11 @@ public:
 	class FWindSourceSceneProxy* SceneProxy;
 
 protected:
-	// Begin UActorComponent interface.
+	//~ Begin UActorComponent Interface.
 	virtual void CreateRenderState_Concurrent() override;
 	virtual void SendRenderTransform_Concurrent() override;
 	virtual void DestroyRenderState_Concurrent() override;
-	// End UActorComponent interface.
+	//~ End UActorComponent Interface.
 
 public:
 	/**
diff --git a/Engine/Source/Runtime/Engine/Classes/Curves/CurveBase.h b/Engine/Source/Runtime/Engine/Classes/Curves/CurveBase.h
--- a/Engine/Source/Runtime/Engine/Classes/Curves/CurveBase.h
+++ b/Engine/Source/Runtime/Engine/Classes/Curves/CurveBase.h
@@ -421,7 +421,7 @@ public:
 };
 
 //////////////////////////////////////////////////////////////////////////
-// Curve editor interface
+//~ Begin Curve editor Interface
 
 /** Info about a curve to be edited */
 template<class T> struct FRichCurveEditInfoTemplate
@@ -542,7 +542,7 @@ public:
 	virtual bool IsValidCurve( FRichCurveEditInfo CurveInfo ) override { return false; };
 	// End FCurveOwnerInterface
 
-	// Begin UCurveBase interface
+	//~ Begin UCurveBase Interface
 
 	/** 
 	 *	Create curve from CSV style comma-separated string. 
@@ -553,9 +553,9 @@ public:
 	/** Reset all curve data */
 	void ResetCurve();
 
-	// End UCurveBase interface
+	//~ End UCurveBase Interface
 
-	// Begin UObject interface
+	//~ Begin UObject Interface
 #if WITH_EDITORONLY_DATA
 	/** Override to ensure we write out the asset import data */
 	virtual void GetAssetRegistryTags(TArray<FAssetRegistryTag>& OutTags) const override;
@@ -570,7 +570,7 @@ public:
 	FString ImportPath_DEPRECATED;
 
 #endif
-	// End UObject interface
+	//~ End UObject Interface
 };
 
 //////////////////////////////////////////////////////////////////////////
diff --git a/Engine/Source/Runtime/Engine/Classes/DeviceProfiles/DeviceProfile.h b/Engine/Source/Runtime/Engine/Classes/DeviceProfiles/DeviceProfile.h
--- a/Engine/Source/Runtime/Engine/Classes/DeviceProfiles/DeviceProfile.h
+++ b/Engine/Source/Runtime/Engine/Classes/DeviceProfiles/DeviceProfile.h
@@ -68,13 +68,13 @@ private:
 
 public:
 
-	// Begin UObject interface
+	//~ Begin UObject Interface
 	virtual void PostInitProperties() override;
-	// End UObject interface
+	//~ End UObject Interface
 
 #if WITH_EDITOR
-	// Begin UObject interface
+	//~ Begin UObject Interface
 	virtual void PostEditChangeProperty( FPropertyChangedEvent& PropertyChangedEvent ) override;
-	// End UObject interface
+	//~ End UObject Interface
 #endif
 };
diff --git a/Engine/Source/Runtime/Engine/Classes/Distributions/DistributionFloat.h b/Engine/Source/Runtime/Engine/Classes/Distributions/DistributionFloat.h
--- a/Engine/Source/Runtime/Engine/Classes/Distributions/DistributionFloat.h
+++ b/Engine/Source/Runtime/Engine/Classes/Distributions/DistributionFloat.h
@@ -134,10 +134,10 @@ class UDistributionFloat : public UDistribution
 	/** @todo document */
 	virtual float GetValue( float F = 0.f, UObject* Data = NULL, struct FRandomStream* InRandomStream = NULL ) const;
 
-	// Begin FCurveEdInterface Interface
+	//~ Begin FCurveEdInterface Interface
 	virtual void GetInRange(float& MinIn, float& MaxIn) const override;
 	virtual void GetOutRange(float& MinOut, float& MaxOut) const override;
-	// End FCurveEdInterface Interface
+	//~ End FCurveEdInterface Interface
 	
 	/** @return true of this distribution can be baked into a FRawDistribution lookup table, otherwise false */
 	virtual bool CanBeBaked() const 
diff --git a/Engine/Source/Runtime/Engine/Classes/Distributions/DistributionFloatConstant.h b/Engine/Source/Runtime/Engine/Classes/Distributions/DistributionFloatConstant.h
--- a/Engine/Source/Runtime/Engine/Classes/Distributions/DistributionFloatConstant.h
+++ b/Engine/Source/Runtime/Engine/Classes/Distributions/DistributionFloatConstant.h
@@ -14,17 +14,17 @@ class UDistributionFloatConstant : public UDistributionFloat
 	UPROPERTY(EditAnywhere, Category=DistributionFloatConstant)
 	float Constant;
 
-	// Begin UObject Interface
+	//~ Begin UObject Interface
 	virtual void PostInitProperties() override;
 	virtual void PostLoad() override;
-	// End UObject Interface
+	//~ End UObject Interface
 
-	// Begin UDistributionFloat Interface
+	//~ Begin UDistributionFloat Interface
 	virtual float GetValue( float F = 0.f, UObject* Data = NULL, struct FRandomStream* InRandomStream = NULL ) const override;
-	// End UDistributionFloat Interface
+	//~ End UDistributionFloat Interface
 
 
-	// Begin FCurveEdInterface interface
+	//~ Begin FCurveEdInterface Interface
 	virtual int32		GetNumKeys() const override;
 	virtual int32		GetNumSubCurves() const override;
 	virtual float	GetKeyIn(int32 KeyIndex) override;
@@ -41,6 +41,6 @@ class UDistributionFloatConstant : public UDistributionFloat
 	virtual void	SetKeyOut(int32 SubIndex, int32 KeyIndex, float NewOutVal) override;
 	virtual void	SetKeyInterpMode(int32 KeyIndex, EInterpCurveMode NewMode) override;
 	virtual void	SetTangents(int32 SubIndex, int32 KeyIndex, float ArriveTangent, float LeaveTangent) override;
-	// End FCurveEdInterface interface
+	//~ End FCurveEdInterface Interface
 };
 
diff --git a/Engine/Source/Runtime/Engine/Classes/Distributions/DistributionFloatConstantCurve.h b/Engine/Source/Runtime/Engine/Classes/Distributions/DistributionFloatConstantCurve.h
--- a/Engine/Source/Runtime/Engine/Classes/Distributions/DistributionFloatConstantCurve.h
+++ b/Engine/Source/Runtime/Engine/Classes/Distributions/DistributionFloatConstantCurve.h
@@ -15,11 +15,11 @@ class UDistributionFloatConstantCurve : public UDistributionFloat
 	FInterpCurveFloat ConstantCurve;
 
 
-	// Begin UDistributionFloat Interface
+	//~ Begin UDistributionFloat Interface
 	virtual float GetValue( float F = 0.f, UObject* Data = NULL, struct FRandomStream* InRandomStream = NULL ) const override;
-	// End UDistributionFloat Interface
+	//~ End UDistributionFloat Interface
 
-	// Begin FCurveEdInterface interface
+	//~ Begin FCurveEdInterface Interface
 	virtual int32		GetNumKeys() const override;
 	virtual int32		GetNumSubCurves() const override;
 	virtual float	GetKeyIn(int32 KeyIndex) override;
@@ -36,6 +36,6 @@ class UDistributionFloatConstantCurve : public UDistributionFloat
 	virtual void	SetKeyOut(int32 SubIndex, int32 KeyIndex, float NewOutVal) override;
 	virtual void	SetKeyInterpMode(int32 KeyIndex, EInterpCurveMode NewMode) override;
 	virtual void	SetTangents(int32 SubIndex, int32 KeyIndex, float ArriveTangent, float LeaveTangent) override;
-	// End FCurveEdInterface interface
+	//~ End FCurveEdInterface Interface
 };
 
diff --git a/Engine/Source/Runtime/Engine/Classes/Distributions/DistributionFloatParameterBase.h b/Engine/Source/Runtime/Engine/Classes/Distributions/DistributionFloatParameterBase.h
--- a/Engine/Source/Runtime/Engine/Classes/Distributions/DistributionFloatParameterBase.h
+++ b/Engine/Source/Runtime/Engine/Classes/Distributions/DistributionFloatParameterBase.h
@@ -35,9 +35,9 @@ class UDistributionFloatParameterBase : public UDistributionFloatConstant
 	TEnumAsByte<enum DistributionParamMode> ParamMode;
 
 
-	// Begin UDistributionFloat Interface
+	//~ Begin UDistributionFloat Interface
 	virtual float GetValue( float F = 0.f, UObject* Data = NULL, struct FRandomStream* InRandomStream = NULL ) const override;
-	// End UDistributionFloat Interface
+	//~ End UDistributionFloat Interface
 	
 	
 	/** todo document */
diff --git a/Engine/Source/Runtime/Engine/Classes/Distributions/DistributionFloatParticleParameter.h b/Engine/Source/Runtime/Engine/Classes/Distributions/DistributionFloatParticleParameter.h
--- a/Engine/Source/Runtime/Engine/Classes/Distributions/DistributionFloatParticleParameter.h
+++ b/Engine/Source/Runtime/Engine/Classes/Distributions/DistributionFloatParticleParameter.h
@@ -11,8 +11,8 @@ class UDistributionFloatParticleParameter : public UDistributionFloatParameterBa
 	GENERATED_UCLASS_BODY()
 
 
-	// Begin UDistributionFloatParameterBase Interface
+	//~ Begin UDistributionFloatParameterBase Interface
 	virtual bool GetParamValue(UObject* Data, FName ParamName, float& OutFloat) const override;
-	// End UDistributionFloatParameterBase Interface
+	//~ End UDistributionFloatParameterBase Interface
 };
 
diff --git a/Engine/Source/Runtime/Engine/Classes/Distributions/DistributionFloatUniform.h b/Engine/Source/Runtime/Engine/Classes/Distributions/DistributionFloatUniform.h
--- a/Engine/Source/Runtime/Engine/Classes/Distributions/DistributionFloatUniform.h
+++ b/Engine/Source/Runtime/Engine/Classes/Distributions/DistributionFloatUniform.h
@@ -19,20 +19,20 @@ class UDistributionFloatUniform : public UDistributionFloat
 	UPROPERTY(EditAnywhere, Category=DistributionFloatUniform)
 	float Max;
 
-	// Begin UObject Interface
+	//~ Begin UObject Interface
 	virtual void PostInitProperties() override;
 	virtual void PostLoad() override;
-	// End UObject Interface
+	//~ End UObject Interface
 
-	// Begin UDistributionFloat Interface
+	//~ Begin UDistributionFloat Interface
 	virtual float GetValue( float F = 0.f, UObject* Data = NULL, struct FRandomStream* InRandomStream = NULL ) const override;
 
 	[REDACTED_EMAIL]: Currently, consoles need this? At least until we have some sort of cooking/packaging step!
 	virtual ERawDistributionOperation GetOperation() const override;
 	virtual uint32 InitializeRawEntry(float Time, float* Values) const override;
-	// End UDistributionFloat Interface
+	//~ End UDistributionFloat Interface
 
-	// Begin FCurveEdInterface interface
+	//~ Begin FCurveEdInterface Interface
 	virtual int32		GetNumKeys() const override;
 	virtual int32		GetNumSubCurves() const override;
 	virtual FColor	GetSubCurveButtonColor(int32 SubCurveIndex, bool bIsSubCurveHidden) const override;
@@ -50,6 +50,6 @@ class UDistributionFloatUniform : public UDistributionFloat
 	virtual void	SetKeyOut(int32 SubIndex, int32 KeyIndex, float NewOutVal) override;
 	virtual void	SetKeyInterpMode(int32 KeyIndex, EInterpCurveMode NewMode) override;
 	virtual void	SetTangents(int32 SubIndex, int32 KeyIndex, float ArriveTangent, float LeaveTangent) override;
-	// End FCurveEdInterface interface
+	//~ End FCurveEdInterface Interface
 };
 
diff --git a/Engine/Source/Runtime/Engine/Classes/Distributions/DistributionFloatUniformCurve.h b/Engine/Source/Runtime/Engine/Classes/Distributions/DistributionFloatUniformCurve.h
--- a/Engine/Source/Runtime/Engine/Classes/Distributions/DistributionFloatUniformCurve.h
+++ b/Engine/Source/Runtime/Engine/Classes/Distributions/DistributionFloatUniformCurve.h
@@ -16,19 +16,19 @@ class UDistributionFloatUniformCurve : public UDistributionFloat
 	FInterpCurveVector2D ConstantCurve;
 
 
-	// Begin UDistributionFloat Interface
+	//~ Begin UDistributionFloat Interface
 	virtual float GetValue( float F = 0.f, UObject* Data = NULL, struct FRandomStream* InRandomStream = NULL ) const override;
 	[REDACTED_EMAIL]: Currently, consoles need this? At least until we have some sort of cooking/packaging step!
 	virtual ERawDistributionOperation GetOperation() const override;
 	virtual uint32 InitializeRawEntry(float Time, float* Values) const override;
-	// End UDistributionFloat Interface
+	//~ End UDistributionFloat Interface
 
 	/**
 	 *	This function will retrieve the max and min values at the given time.
 	 */
 	virtual FVector2D GetMinMaxValue(float F = 0.f, UObject* Data = NULL) const;
 
-	// Begin FCurveEdInterface interface
+	//~ Begin FCurveEdInterface Interface
 	virtual int32		GetNumKeys() const override;
 	virtual int32		GetNumSubCurves() const override;
 	virtual FColor	GetSubCurveButtonColor(int32 SubCurveIndex, bool bIsSubCurveHidden) const override;
@@ -46,6 +46,6 @@ class UDistributionFloatUniformCurve : public UDistributionFloat
 	virtual void	SetKeyOut(int32 SubIndex, int32 KeyIndex, float NewOutVal) override;
 	virtual void	SetKeyInterpMode(int32 KeyIndex, EInterpCurveMode NewMode) override;
 	virtual void	SetTangents(int32 SubIndex, int32 KeyIndex, float ArriveTangent, float LeaveTangent) override;
-	// End FCurveEdInterface interface
+	//~ End FCurveEdInterface Interface
 };
 
diff --git a/Engine/Source/Runtime/Engine/Classes/Distributions/DistributionVector.h b/Engine/Source/Runtime/Engine/Classes/Distributions/DistributionVector.h
--- a/Engine/Source/Runtime/Engine/Classes/Distributions/DistributionVector.h
+++ b/Engine/Source/Runtime/Engine/Classes/Distributions/DistributionVector.h
@@ -163,11 +163,11 @@ class ENGINE_API UDistributionVector : public UDistribution
 
 	virtual FVector	GetValue( float F = 0.f, UObject* Data = NULL, int32 LastExtreme = 0, struct FRandomStream* InRandomStream = NULL ) const;
 
-	// Begin FCurveEdInterface Interface
+	//~ Begin FCurveEdInterface Interface
 	virtual void GetInRange(float& MinIn, float& MaxIn) const override;
 	virtual void GetOutRange(float& MinOut, float& MaxOut) const override;
 	virtual	void GetRange(FVector& OutMin, FVector& OutMax) const;
-	// End FCurveEdInterface Interface
+	//~ End FCurveEdInterface Interface
 
 	/** @return true of this distribution can be baked into a FRawDistribution lookup table, otherwise false */
 	virtual bool CanBeBaked() const 
diff --git a/Engine/Source/Runtime/Engine/Classes/Distributions/DistributionVectorConstant.h b/Engine/Source/Runtime/Engine/Classes/Distributions/DistributionVectorConstant.h
--- a/Engine/Source/Runtime/Engine/Classes/Distributions/DistributionVectorConstant.h
+++ b/Engine/Source/Runtime/Engine/Classes/Distributions/DistributionVectorConstant.h
@@ -27,12 +27,12 @@ class ENGINE_API UDistributionVectorConstant : public UDistributionVector
 	virtual	void	GetRange(FVector& OutMin, FVector& OutMax) const override;
 	//End UDistributionVector Interface
 
-	// Begin UObject interface
+	//~ Begin UObject Interface
 	virtual void PostInitProperties() override;
 	virtual void PostLoad() override;
-	// End UObject interface
+	//~ End UObject Interface
 	
-	// Begin FCurveEdInterface interface
+	//~ Begin FCurveEdInterface Interface
 	virtual int32		GetNumKeys() const override;
 	virtual int32		GetNumSubCurves() const override;
 	virtual FColor	GetSubCurveButtonColor(int32 SubCurveIndex, bool bIsSubCurveHidden) const override;
@@ -50,7 +50,7 @@ class ENGINE_API UDistributionVectorConstant : public UDistributionVector
 	virtual void	SetKeyOut(int32 SubIndex, int32 KeyIndex, float NewOutVal) override;
 	virtual void	SetKeyInterpMode(int32 KeyIndex, EInterpCurveMode NewMode) override;
 	virtual void	SetTangents(int32 SubIndex, int32 KeyIndex, float ArriveTangent, float LeaveTangent) override;
-	// End FCurveEdInterface interface
+	//~ End FCurveEdInterface Interface
 
 };
 
diff --git a/Engine/Source/Runtime/Engine/Classes/Distributions/DistributionVectorConstantCurve.h b/Engine/Source/Runtime/Engine/Classes/Distributions/DistributionVectorConstantCurve.h
--- a/Engine/Source/Runtime/Engine/Classes/Distributions/DistributionVectorConstantCurve.h
+++ b/Engine/Source/Runtime/Engine/Classes/Distributions/DistributionVectorConstantCurve.h
@@ -27,7 +27,7 @@ class UDistributionVectorConstantCurve : public UDistributionVector
 	virtual	void	GetRange(FVector& OutMin, FVector& OutMax) const override;
 	//End UDistributionVector Interface
 
-	// Begin FCurveEdInterface interface
+	//~ Begin FCurveEdInterface Interface
 	virtual int32		GetNumKeys() const override;
 	virtual int32		GetNumSubCurves() const override;
 	virtual FColor	GetSubCurveButtonColor(int32 SubCurveIndex, bool bIsSubCurveHidden) const override;
@@ -45,7 +45,7 @@ class UDistributionVectorConstantCurve : public UDistributionVector
 	virtual void	SetKeyOut(int32 SubIndex, int32 KeyIndex, float NewOutVal) override;
 	virtual void	SetKeyInterpMode(int32 KeyIndex, EInterpCurveMode NewMode) override;
 	virtual void	SetTangents(int32 SubIndex, int32 KeyIndex, float ArriveTangent, float LeaveTangent) override;
-	// End FCurveEdInterface interface
+	//~ End FCurveEdInterface Interface
 
 };
 
diff --git a/Engine/Source/Runtime/Engine/Classes/Distributions/DistributionVectorParticleParameter.h b/Engine/Source/Runtime/Engine/Classes/Distributions/DistributionVectorParticleParameter.h
--- a/Engine/Source/Runtime/Engine/Classes/Distributions/DistributionVectorParticleParameter.h
+++ b/Engine/Source/Runtime/Engine/Classes/Distributions/DistributionVectorParticleParameter.h
@@ -11,8 +11,8 @@ class UDistributionVectorParticleParameter : public UDistributionVectorParameter
 	GENERATED_UCLASS_BODY()
 
 
-	// Begin UDistributionVectorParameterBase Interface
+	//~ Begin UDistributionVectorParameterBase Interface
 	virtual bool GetParamValue(UObject* Data, FName ParamName, FVector& OutVector) const override;
-	// End UDistributionVectorParameterBase Interface
+	//~ End UDistributionVectorParameterBase Interface
 };
 
diff --git a/Engine/Source/Runtime/Engine/Classes/Distributions/DistributionVectorUniform.h b/Engine/Source/Runtime/Engine/Classes/Distributions/DistributionVectorUniform.h
--- a/Engine/Source/Runtime/Engine/Classes/Distributions/DistributionVectorUniform.h
+++ b/Engine/Source/Runtime/Engine/Classes/Distributions/DistributionVectorUniform.h
@@ -32,10 +32,10 @@ class UDistributionVectorUniform : public UDistributionVector
 	uint32 bUseExtremes:1;
 
 
-	// Begin UObject Interface
+	//~ Begin UObject Interface
 	virtual void PostInitProperties() override;
 	virtual void PostLoad() override;
-	// End UObject Interface
+	//~ End UObject Interface
 
 	virtual FVector	GetValue( float F = 0.f, UObject* Data = NULL, int32 LastExtreme = 0, struct FRandomStream* InRandomStream = NULL ) const override;
 
@@ -57,7 +57,7 @@ class UDistributionVectorUniform : public UDistributionVector
 	// 2,3 = min/max y
 	// 4,5 = min/max z
 
-	// Begin FCurveEdInterface interface
+	//~ Begin FCurveEdInterface Interface
 	virtual int32		GetNumKeys() const override;
 	virtual int32		GetNumSubCurves() const override;
 	virtual FColor	GetSubCurveButtonColor(int32 SubCurveIndex, bool bIsSubCurveHidden) const override;
@@ -75,7 +75,7 @@ class UDistributionVectorUniform : public UDistributionVector
 	virtual void	SetKeyOut(int32 SubIndex, int32 KeyIndex, float NewOutVal) override;
 	virtual void	SetKeyInterpMode(int32 KeyIndex, EInterpCurveMode NewMode) override;
 	virtual void	SetTangents(int32 SubIndex, int32 KeyIndex, float ArriveTangent, float LeaveTangent) override;
-	// Begin FCurveEdInterface interface
+	//~ Begin FCurveEdInterface Interface
 
 };
 
diff --git a/Engine/Source/Runtime/Engine/Classes/Distributions/DistributionVectorUniformCurve.h b/Engine/Source/Runtime/Engine/Classes/Distributions/DistributionVectorUniformCurve.h
--- a/Engine/Source/Runtime/Engine/Classes/Distributions/DistributionVectorUniformCurve.h
+++ b/Engine/Source/Runtime/Engine/Classes/Distributions/DistributionVectorUniformCurve.h
@@ -48,7 +48,7 @@ class UDistributionVectorUniformCurve : public UDistributionVector
 	 */
 	virtual FTwoVectors GetMinMaxValue(float F = 0.f, UObject* Data = NULL) const;
 
-	// Begin FCurveEdInterface interface
+	//~ Begin FCurveEdInterface Interface
 	virtual int32		GetNumKeys() const override;
 	virtual int32		GetNumSubCurves() const override;
 	virtual FColor	GetSubCurveButtonColor(int32 SubCurveIndex, bool bIsSubCurveHidden) const override;
@@ -67,7 +67,7 @@ class UDistributionVectorUniformCurve : public UDistributionVector
 	virtual void	SetKeyInterpMode(int32 KeyIndex, EInterpCurveMode NewMode) override;
 	virtual void	SetTangents(int32 SubIndex, int32 KeyIndex, float ArriveTangent, float LeaveTangent) override;
 	virtual void	LockAndMirror(FTwoVectors& Val) const;
-	// Begin FCurveEdInterface interface
+	//~ Begin FCurveEdInterface Interface
 
 };
 
diff --git a/Engine/Source/Runtime/Engine/Classes/EdGraph/EdGraph.h b/Engine/Source/Runtime/Engine/Classes/EdGraph/EdGraph.h
--- a/Engine/Source/Runtime/Engine/Classes/EdGraph/EdGraph.h
+++ b/Engine/Source/Runtime/Engine/Classes/EdGraph/EdGraph.h
@@ -104,11 +104,11 @@ public:
 	void RemoveOnGraphChangedHandler( FDelegateHandle Handle );
 
 #if WITH_EDITORONLY_DATA
-	// Begin UObject interface
+	//~ Begin UObject interface
 	virtual void Serialize( FArchive& Ar ) override;
 	virtual void PostInitProperties() override;
 	virtual void PostLoad() override;
-	// End UObject interface
+	//~ End UObject Interface
 #endif
 
 public:
diff --git a/Engine/Source/Runtime/Engine/Classes/EdGraph/EdGraphNode_Documentation.h b/Engine/Source/Runtime/Engine/Classes/EdGraph/EdGraphNode_Documentation.h
--- a/Engine/Source/Runtime/Engine/Classes/EdGraph/EdGraphNode_Documentation.h
+++ b/Engine/Source/Runtime/Engine/Classes/EdGraph/EdGraphNode_Documentation.h
@@ -19,7 +19,7 @@ class UEdGraphNode_Documentation : public UEdGraphNode
 public:
 
 #if WITH_EDITOR
-	// Begin UEdGraphNode interface
+	//~ Begin UEdGraphNode Interface
 	virtual void AllocateDefaultPins() override {}
 	virtual bool ShouldOverridePinNames() const override { return true; }
 	virtual FLinearColor GetNodeTitleColor() const override { return FLinearColor( 0.2f, 1.f, 0.2f ); }
@@ -33,7 +33,7 @@ public:
 	ENGINE_API virtual FString GetDocumentationLink() const override { return Link; }
 	ENGINE_API virtual FString GetDocumentationExcerptName() const override { return Excerpt; }
 	ENGINE_API virtual FName GetPaletteIcon(FLinearColor& OutColor) const override { return TEXT("GraphEditor.Documentation_16x");  }
-	// End UEdGraphNode interface
+	//~ End UEdGraphNode Interface
 
 	/** Set the Bounds for the comment node */
 	ENGINE_API void SetBounds(const class FSlateRect& Rect);
diff --git a/Engine/Source/Runtime/Engine/Classes/Engine/Blueprint.h b/Engine/Source/Runtime/Engine/Classes/Engine/Blueprint.h
--- a/Engine/Source/Runtime/Engine/Classes/Engine/Blueprint.h
+++ b/Engine/Source/Runtime/Engine/Classes/Engine/Blueprint.h
@@ -570,7 +570,7 @@ public:
 	/** Renames only the generated classes. Should only be used internally or when testing for rename. */
 	virtual bool RenameGeneratedClasses(const TCHAR* NewName = nullptr, UObject* NewOuter = nullptr, ERenameFlags Flags = REN_None);
 
-	// Begin UObject interface (WITH_EDITOR)
+	//~ Begin UObject Interface (WITH_EDITOR)
 	virtual void PostDuplicate(bool bDuplicateForPIE) override;
 	virtual bool Rename(const TCHAR* NewName = nullptr, UObject* NewOuter = nullptr, ERenameFlags Flags = REN_None) override;
 	virtual UClass* RegenerateClass(UClass* ClassToRegenerate, UObject* PreviousCDO, TArray<UObject*>& ObjLoaded) override;
@@ -578,7 +578,7 @@ public:
 	virtual void PostLoadSubobjects( FObjectInstancingGraph* OuterInstanceGraph ) override;
 	virtual bool Modify(bool bAlwaysMarkDirty = true) override;
 	virtual void GetAssetRegistryTags(TArray<FAssetRegistryTag>& OutTags) const override;
-	// End of UObject interface
+	//~ End UObject Interface
 
 	/** Consigns the GeneratedClass and the SkeletonGeneratedClass to oblivion, and nulls their references */
 	void RemoveGeneratedClasses();
@@ -601,14 +601,14 @@ public:
 
 #endif	//#if WITH_EDITOR
 
-	// Begin UObject interface
+	//~ Begin UObject Interface
 	virtual void Serialize(FArchive& Ar) override;
 	virtual FString GetDesc(void) override;
 	virtual void TagSubobjects(EObjectFlags NewFlags) override;
 	virtual bool NeedsLoadForClient() const override;
 	virtual bool NeedsLoadForServer() const override;
 	virtual bool NeedsLoadForEditorGame() const override;
-	// End of UObject interface
+	//~ End UObject Interface
 
 	/** Get the Blueprint object that generated the supplied class */
 	static UBlueprint* GetBlueprintFromClass(const UClass* InClass);
diff --git a/Engine/Source/Runtime/Engine/Classes/Engine/ComponentDelegateBinding.h b/Engine/Source/Runtime/Engine/Classes/Engine/ComponentDelegateBinding.h
--- a/Engine/Source/Runtime/Engine/Classes/Engine/ComponentDelegateBinding.h
+++ b/Engine/Source/Runtime/Engine/Classes/Engine/ComponentDelegateBinding.h
@@ -41,10 +41,10 @@ class ENGINE_API UComponentDelegateBinding
 	UPROPERTY()
 	TArray<FBlueprintComponentDelegateBinding> ComponentDelegateBindings;
 
-	// Begin DynamicBlueprintBinding interface
+	//~ Begin DynamicBlueprintBinding Interface
 	virtual void BindDynamicDelegates(UObject* InInstance) const override;
 	virtual void UnbindDynamicDelegatesForProperty(UObject* InInstance, const UObjectProperty* InObjectProperty) const override;
-	// End DynamicBlueprintBinding interface
+	//~ End DynamicBlueprintBinding Interface
 
 private:
 	// Utility method used to find the target delegate given an instance and a binding descriptor
diff --git a/Engine/Source/Runtime/Engine/Classes/Engine/ControlChannel.h b/Engine/Source/Runtime/Engine/Classes/Engine/ControlChannel.h
--- a/Engine/Source/Runtime/Engine/Classes/Engine/ControlChannel.h
+++ b/Engine/Source/Runtime/Engine/Classes/Engine/ControlChannel.h
@@ -58,11 +58,11 @@ class ENGINE_API UControlChannel
 
 	virtual void Init( UNetConnection* InConnection, int32 InChIndex, bool InOpenedLocally ) override;
 
-	// Begin UChannel interface.
+	//~ Begin UChannel Interface.
 	virtual FPacketIdRange SendBunch(FOutBunch* Bunch, bool Merge) override;
 
 	virtual void Tick() override;
-	// End UChannel interface.
+	//~ End UChannel Interface.
 
 
 	/** Handle an incoming bunch. */
diff --git a/Engine/Source/Runtime/Engine/Classes/Engine/CoreSettings.h b/Engine/Source/Runtime/Engine/Classes/Engine/CoreSettings.h
--- a/Engine/Source/Runtime/Engine/Classes/Engine/CoreSettings.h
+++ b/Engine/Source/Runtime/Engine/Classes/Engine/CoreSettings.h
@@ -85,13 +85,13 @@ protected:
 		ToolTip = "Batching granularity used to register actor components during level streaming."))
 	int32 LevelStreamingComponentsRegistrationGranularity;
 
-	// Begin UObject interface
+	//~ Begin UObject Interface
 	virtual void PostInitProperties() override;
 
 #if WITH_EDITOR
 	virtual void PostEditChangeProperty(FPropertyChangedEvent& PropertyChangedEvent) override;
 #endif
-	// End UObject interface
+	//~ End UObject Interface
 };
 
 /** Whether to allow background level streaming. */
@@ -151,11 +151,11 @@ protected:
 		ToolTip = "Size Of Permanent Object Pool (bytes). Works only in cooked builds."))
 	int32 SizeOfPermanentObjectPool;
 
-	// Begin UObject interface
+	//~ Begin UObject Interface
 	virtual void PostInitProperties() override;
 
 #if WITH_EDITOR
 	virtual void PostEditChangeProperty(FPropertyChangedEvent& PropertyChangedEvent) override;
 #endif
-	// End UObject interface
+	//~ End UObject Interface
 };
diff --git a/Engine/Source/Runtime/Engine/Classes/Engine/CurveTable.h b/Engine/Source/Runtime/Engine/Classes/Engine/CurveTable.h
--- a/Engine/Source/Runtime/Engine/Classes/Engine/CurveTable.h
+++ b/Engine/Source/Runtime/Engine/Classes/Engine/CurveTable.h
@@ -27,7 +27,7 @@ class UCurveTable
 	/** Map of name of row to row data structure. */
 	TMap<FName, FRichCurve*>	RowMap;
 
-	// Begin UObject interface.
+	//~ Begin UObject Interface.
 	virtual void FinishDestroy() override;
 	virtual void Serialize( FArchive& Ar ) override;
 
@@ -45,9 +45,9 @@ class UCurveTable
 	
 #endif	// WITH_EDITORONLY_DATA
 
-	// End  UObject interface
+	//~ End  UObject Interface
 
-	// Begin UCurveTable interface
+	//~ Begin UCurveTable Interface
 
 	/** Function to find the row of a table given its name. */
 	FRichCurve* FindCurve(FName RowName, const FString& ContextString, bool WarnIfNotFound=true) const
diff --git a/Engine/Source/Runtime/Engine/Classes/Engine/DataTable.h b/Engine/Source/Runtime/Engine/Classes/Engine/DataTable.h
--- a/Engine/Source/Runtime/Engine/Classes/Engine/DataTable.h
+++ b/Engine/Source/Runtime/Engine/Classes/Engine/DataTable.h
@@ -41,7 +41,7 @@ class UDataTable
 	/** Map of name of row to row data structure. */
 	TMap<FName, uint8*>		RowMap;
 
-	// Begin UObject interface.
+	//~ Begin UObject Interface.
 	ENGINE_API virtual void FinishDestroy() override;
 	ENGINE_API virtual void Serialize(FArchive& Ar) override;
 	ENGINE_API static void AddReferencedObjects(UObject* InThis, FReferenceCollector& Collector);
@@ -58,9 +58,9 @@ class UDataTable
 	FString ImportPath_DEPRECATED;
 #endif	// WITH_EDITORONLY_DATA
 
-	// End  UObject interface
+	//~ End  UObject Interface
 
-	// Begin UDataTable interface
+	//~ Begin UDataTable Interface
 
 	/** Function to find the row of a table given its name. */
 	template <class T>
@@ -189,7 +189,7 @@ public:
 	ENGINE_API TArray< TArray<FString> > GetTableData() const;
 #endif //WITH_EDITOR
 
-	// End UDataTable interface
+	//~ End UDataTable Interface
 
 private:
 	void SaveStructData(FArchive& Ar);
diff --git a/Engine/Source/Runtime/Engine/Classes/Engine/DecalActor.h b/Engine/Source/Runtime/Engine/Classes/Engine/DecalActor.h
--- a/Engine/Source/Runtime/Engine/Classes/Engine/DecalActor.h
+++ b/Engine/Source/Runtime/Engine/Classes/Engine/DecalActor.h
@@ -57,15 +57,15 @@ public:
 
 	
 #if WITH_EDITOR
-	// Begin UObject Interface
+	//~ Begin UObject Interface
 	virtual void PostEditMove(bool bFinished) override;
 	virtual void PostEditChangeProperty(FPropertyChangedEvent& PropertyChangedEvent) override;
-	// End UObject Interface
+	//~ End UObject Interface
 
-	// Begin AActor interface.
+	//~ Begin AActor Interface.
 	virtual void EditorApplyScale(const FVector& DeltaScale, const FVector* PivotLocation, bool bAltDown, bool bShiftDown, bool bCtrlDown) override;
 	virtual bool GetReferencedContentObjects(TArray<UObject*>& Objects) const override;
-	// End AActor interface.
+	//~ End AActor Interface.
 #endif // WITH_EDITOR
 
 	
diff --git a/Engine/Source/Runtime/Engine/Classes/Engine/DestructibleFractureSettings.h b/Engine/Source/Runtime/Engine/Classes/Engine/DestructibleFractureSettings.h
--- a/Engine/Source/Runtime/Engine/Classes/Engine/DestructibleFractureSettings.h
+++ b/Engine/Source/Runtime/Engine/Classes/Engine/DestructibleFractureSettings.h
@@ -90,11 +90,11 @@ struct FFractureMaterial
 		, InteriorElementIndex(-1)
 	{ }
 
-	// Begin FFractureMaterial interface
+	//~ Begin FFractureMaterial Interface
 #if WITH_APEX
 	void	FillNxFractureMaterialDesc(physx::apex::NxFractureMaterialDesc& PFractureMaterialDesc);
 #endif // WITH_APEX
-	// End FFractureMaterial interface
+	//~ End FFractureMaterial Interface
 };
 
 
@@ -199,13 +199,13 @@ class UDestructibleFractureSettings
 	TArray<physx::apex::NxDestructibleGeometryDesc>	GeometryDescs;
 #endif // WITH_APEX
 
-	// Begin UObject interface.
+	//~ Begin UObject Interface.
 	virtual void PostInitProperties() override;
 	virtual void BeginDestroy() override;
 	virtual void Serialize( FArchive& Ar ) override;
-	// End  UObject interface
+	//~ End  UObject Interface
 
-	// Begin UDestructibleFractureSettings interface
+	//~ Begin UDestructibleFractureSettings Interface
 #if WITH_APEX && WITH_EDITOR
 	/**
 	 * Creates an NxDestructibleAssetCookingDesc suitable for the authoring mesh, which can be passed to CreateApexDestructibleAsset().
@@ -260,5 +260,5 @@ class UDestructibleFractureSettings
 	 */
 	ENGINE_API	physx::apex::NxDestructibleAsset*	CreateApexDestructibleAsset(const physx::apex::NxDestructibleAssetCookingDesc& DestructibleAssetCookingDesc);
 #endif // WITH_APEX
-	// End UDestructibleFractureSettings interface
+	//~ End UDestructibleFractureSettings Interface
 };
diff --git a/Engine/Source/Runtime/Engine/Classes/Engine/Engine.h b/Engine/Source/Runtime/Engine/Classes/Engine/Engine.h
--- a/Engine/Source/Runtime/Engine/Classes/Engine/Engine.h
+++ b/Engine/Source/Runtime/Engine/Classes/Engine/Engine.h
@@ -1665,11 +1665,11 @@ public:
 		NetworkFailureEvent.Broadcast(World, NetDriver, FailureType, ErrorString);
 	}
 
-	// Begin UObject interface.
+	//~ Begin UObject Interface.
 	virtual void FinishDestroy() override;
 	virtual void Serialize(FArchive& Ar) override;
 	static void AddReferencedObjects(UObject* InThis, FReferenceCollector& Collector);
-	// End UObject interface.
+	//~ End UObject Interface.
 
 	/** Initialize the game engine. */
 	virtual void Init(IEngineLoop* InEngineLoop);
@@ -1681,9 +1681,9 @@ public:
 	/** Called at startup, in the middle of FEngineLoop::Init.	 */
 	void ParseCommandline();
 
-	// Begin FExec Interface
+	//~ Begin FExec Interface
 	virtual bool Exec( UWorld* InWorld, const TCHAR* Cmd, FOutputDevice& Out=*GLog ) override;
-	// End FExec Interface
+	//~ End FExec Interface
 
 	/** 
 	 * Exec command handlers
@@ -2705,7 +2705,7 @@ protected:
 
 public:
 
-	// Public interface for async map change functions
+	//~ Begin Public Interface for async map change functions
 
 	bool CommitMapChange(UWorld* InWorld) { return CommitMapChange(GetWorldContextFromWorldChecked(InWorld)); }
 	bool IsReadyForMapChange(UWorld* InWorld) { return IsReadyForMapChange(GetWorldContextFromWorldChecked(InWorld)); }
diff --git a/Engine/Source/Runtime/Engine/Classes/Engine/Font.h b/Engine/Source/Runtime/Engine/Classes/Engine/Font.h
--- a/Engine/Source/Runtime/Engine/Classes/Engine/Font.h
+++ b/Engine/Source/Runtime/Engine/Classes/Engine/Font.h
@@ -196,7 +196,7 @@ public:
 	 */
 	virtual SIZE_T GetResourceSize(EResourceSizeMode::Type Mode) override;
 
-	// UFont interface
+	//~ Begin UFont Interface
 	ENGINE_API TCHAR RemapChar(TCHAR CharCode) const;
 
 	/**
@@ -245,14 +245,14 @@ public:
 	 */
 	ENGINE_API int32 GetStringHeightSize( const TCHAR *Text ) const;
 
-	// Begin UObject interface
+	//~ Begin UObject Interface
 	virtual void Serialize( FArchive& Ar ) override;
 	virtual void PostLoad() override;
 	virtual bool IsLocalizedResource() override;
 #if WITH_EDITORONLY_DATA
 	virtual void GetAssetRegistryTags(TArray<FAssetRegistryTag>& OutTags) const override;
 #endif
-	// End UObject interface
+	//~ End UObject Interface
 
 	/**
 	 * Caches the character count and maximum character height for this font (as well as sub-fonts, in the multi-font case)
diff --git a/Engine/Source/Runtime/Engine/Classes/Engine/GameInstance.h b/Engine/Source/Runtime/Engine/Classes/Engine/GameInstance.h
--- a/Engine/Source/Runtime/Engine/Classes/Engine/GameInstance.h
+++ b/Engine/Source/Runtime/Engine/Classes/Engine/GameInstance.h
@@ -57,14 +57,14 @@ public:
 
 	FString PIEMapName;
 
-	// Begin FExec Interface
+	//~ Begin FExec Interface
 	virtual bool Exec(UWorld* InWorld, const TCHAR* Cmd, FOutputDevice& Out = *GLog) override;
-	// End FExec Interface
+	//~ End FExec Interface
 
-	// Begin UObject Interface
+	//~ Begin UObject Interface
 	virtual class UWorld* GetWorld() const override;
 	virtual void FinishDestroy() override;
-	// End UObject Interface
+	//~ End UObject Interface
 
 	/** virtual function to allow custom GameInstances an opportunity to set up what it needs */
 	virtual void Init();
diff --git a/Engine/Source/Runtime/Engine/Classes/Engine/GameViewportClient.h b/Engine/Source/Runtime/Engine/Classes/Engine/GameViewportClient.h
--- a/Engine/Source/Runtime/Engine/Classes/Engine/GameViewportClient.h
+++ b/Engine/Source/Runtime/Engine/Classes/Engine/GameViewportClient.h
@@ -135,12 +135,12 @@ public:
 	virtual void Init(struct FWorldContext& WorldContext, UGameInstance* OwningGameInstance, bool bCreateNewAudioDevice = true);
 
 public:
-	// Begin UObject Interface
+	//~ Begin UObject Interface
 	virtual void PostInitProperties() override;
 	virtual void BeginDestroy() override;
-	// End UObject Interface
+	//~ End UObject Interface
 
-	// FViewportClient interface.
+	//~ Begin FViewportClient Interface.
 	virtual void RedrawRequested(FViewport* InViewport) override {}
 	virtual bool InputKey(FViewport* Viewport, int32 ControllerId, FKey Key, EInputEvent EventType, float AmountDepressed=1.f, bool bGamepad=false) override;
 	virtual bool InputAxis(FViewport* Viewport, int32 ControllerId, FKey Key, float Delta, float DeltaTime, int32 NumSamples=1, bool bGamepad=false) override;
@@ -162,15 +162,15 @@ public:
 	virtual void MouseEnter(FViewport* Viewport, int32 x, int32 y) override;
 	virtual void MouseLeave(FViewport* Viewport) override;
 	virtual void SetIsSimulateInEditorViewport(bool bInIsSimulateInEditorViewport) override;
-	// End of FViewportClient interface.
+	//~ End FViewportClient Interface.
 
 	/** make any adjustments to the views after they've been completely set up */
 	virtual void FinalizeViews(class FSceneViewFamily* ViewFamily, const TMap<ULocalPlayer*, FSceneView*>& PlayerViewMap)
 	{}
 
-	// Begin FExec interface.
+	//~ Begin FExec Interface.
 	virtual bool Exec( UWorld* InWorld, const TCHAR* Cmd,FOutputDevice& Ar) override;
-	// End of FExec interface.
+	//~ End FExec Interface.
 
 	/**
 	 * Exec command handlers
diff --git a/Engine/Source/Runtime/Engine/Classes/Engine/InputActionDelegateBinding.h b/Engine/Source/Runtime/Engine/Classes/Engine/InputActionDelegateBinding.h
--- a/Engine/Source/Runtime/Engine/Classes/Engine/InputActionDelegateBinding.h
+++ b/Engine/Source/Runtime/Engine/Classes/Engine/InputActionDelegateBinding.h
@@ -35,7 +35,7 @@ class ENGINE_API UInputActionDelegateBinding : public UInputDelegateBinding
 	UPROPERTY()
 	TArray<FBlueprintInputActionDelegateBinding> InputActionDelegateBindings;
 
-	// Begin UInputDelegateBinding interface
+	//~ Begin UInputDelegateBinding Interface
 	virtual void BindToInputComponent(UInputComponent* InputComponent) const override;
-	// End UInputDelegateBinding interface
+	//~ End UInputDelegateBinding Interface
 };
diff --git a/Engine/Source/Runtime/Engine/Classes/Engine/InputAxisDelegateBinding.h b/Engine/Source/Runtime/Engine/Classes/Engine/InputAxisDelegateBinding.h
--- a/Engine/Source/Runtime/Engine/Classes/Engine/InputAxisDelegateBinding.h
+++ b/Engine/Source/Runtime/Engine/Classes/Engine/InputAxisDelegateBinding.h
@@ -31,7 +31,7 @@ class ENGINE_API UInputAxisDelegateBinding : public UInputDelegateBinding
 	UPROPERTY()
 	TArray<FBlueprintInputAxisDelegateBinding> InputAxisDelegateBindings;
 
-	// Begin UInputDelegateBinding interface
+	//~ Begin UInputDelegateBinding Interface
 	virtual void BindToInputComponent(UInputComponent* InputComponent) const override;
-	// End UInputDelegateBinding interface
+	//~ End UInputDelegateBinding Interface
 };
diff --git a/Engine/Source/Runtime/Engine/Classes/Engine/InterpCurveEdSetup.h b/Engine/Source/Runtime/Engine/Classes/Engine/InterpCurveEdSetup.h
--- a/Engine/Source/Runtime/Engine/Classes/Engine/InterpCurveEdSetup.h
+++ b/Engine/Source/Runtime/Engine/Classes/Engine/InterpCurveEdSetup.h
@@ -108,7 +108,7 @@ class UInterpCurveEdSetup : public UObject
 	int32 ActiveTab;
 
 
-	// Begin UObject interface
+	//~ Begin UObject Interface
 	virtual void PostLoad() override;
 	virtual bool NeedsLoadForClient() const override
 	{
@@ -118,9 +118,9 @@ class UInterpCurveEdSetup : public UObject
 	{
 		return false;
 	}
-	// End UObject interface
+	//~ End UObject Interface
 	
-	// InterpCurveEdSetup interface
+	//~ Begin InterpCurveEdSetup Interface
 	ENGINE_API static FCurveEdInterface* GetCurveEdInterfacePointer(const FCurveEdEntry& Entry);
 
 	/** 
diff --git a/Engine/Source/Runtime/Engine/Classes/Engine/LODActor.h b/Engine/Source/Runtime/Engine/Classes/Engine/LODActor.h
--- a/Engine/Source/Runtime/Engine/Classes/Engine/LODActor.h
+++ b/Engine/Source/Runtime/Engine/Classes/Engine/LODActor.h
@@ -37,10 +37,10 @@ public:
 	TArray<class UObject*> SubObjects;
 
 #if WITH_EDITOR
-	// Begin AActor Interface
+	//~ Begin AActor Interface
 	virtual void CheckForErrors() override;
 	virtual bool GetReferencedContentObjects( TArray<UObject*>& Objects ) const override;
-	// End AActor Interface
+	//~ End AActor Interface
 #endif // WITH_EDITOR	
 
 	/** Sets StaticMesh and IsPreviewActor to true if InStaticMesh equals nullptr */
@@ -100,7 +100,7 @@ public:
 	void CleanSubActorArray();
 #endif // WITH_EDITOR
 	
-	// Begin UObject interface.
+	//~ Begin UObject Interface.
 	virtual FString GetDetailedInfoInternal() const override;
 	virtual FBox GetComponentsBoundingBox(bool bNonColliding = false) const override;	
 #if WITH_EDITOR
@@ -111,7 +111,7 @@ public:
 	virtual void EditorApplyMirror(const FVector& MirrorScale, const FVector& PivotLocation) override;
 #endif // WITH_EDITOR	
 	virtual void PostRegisterAllComponents() override;
-	// End UObject interface.		
+	//~ End UObject Interface.		
 protected:
 #if WITH_EDITORONLY_DATA
 	/** Whether or not this LODActor is not build or needs rebuilding */
diff --git a/Engine/Source/Runtime/Engine/Classes/Engine/Level.h b/Engine/Source/Runtime/Engine/Classes/Engine/Level.h
--- a/Engine/Source/Runtime/Engine/Classes/Engine/Level.h
+++ b/Engine/Source/Runtime/Engine/Classes/Engine/Level.h
@@ -536,7 +536,7 @@ public:
 
 	~ULevel();
 
-	// Begin UObject interface.
+	//~ Begin UObject Interface.
 	virtual void Serialize( FArchive& Ar ) override;
 	virtual void BeginDestroy() override;
 	virtual bool IsReadyForFinishDestroy() override;
@@ -551,7 +551,7 @@ public:
 	virtual void PreSave() override;
 	virtual void PostDuplicate(bool bDuplicateForPIE) override;
 	static void AddReferencedObjects(UObject* InThis, FReferenceCollector& Collector);
-	// End UObject interface.
+	//~ End UObject Interface.
 
 	/**
 	 * Adds a newly-created RuntimeMovieScenePlayer to this level.  The level assumes ownership of this object.
@@ -760,11 +760,11 @@ public:
 	/** Push any pending auto receive input actor's input components on to the player controller's input stack */
 	void PushPendingAutoReceiveInput(APlayerController* PC);
 	
-	// Begin IInterface_AssetUserData Interface
+	//~ Begin IInterface_AssetUserData Interface
 	virtual void AddAssetUserData(UAssetUserData* InUserData) override;
 	virtual void RemoveUserDataOfClass(TSubclassOf<UAssetUserData> InUserDataClass) override;
 	virtual UAssetUserData* GetAssetUserDataOfClass(TSubclassOf<UAssetUserData> InUserDataClass) override;
-	// End IInterface_AssetUserData Interface
+	//~ End IInterface_AssetUserData Interface
 
 #if WITH_EDITOR
 	/** meant to be called only from editor, calculating and storing static geometry to be used with off-line and/or on-line navigation building */
diff --git a/Engine/Source/Runtime/Engine/Classes/Engine/LevelBounds.h b/Engine/Source/Runtime/Engine/Classes/Engine/LevelBounds.h
--- a/Engine/Source/Runtime/Engine/Classes/Engine/LevelBounds.h
+++ b/Engine/Source/Runtime/Engine/Classes/Engine/LevelBounds.h
@@ -32,14 +32,14 @@ class ALevelBounds
 	UPROPERTY(EditAnywhere, Category=LevelBounds)
 	bool bAutoUpdateBounds;
 
-	// Begin UObject Interface
+	//~ Begin UObject Interface
 	virtual void PostLoad() override;
-	// End UObject Interface
+	//~ End UObject Interface
 	
-	// Begin AActor interface.
+	//~ Begin AActor Interface.
 	virtual FBox GetComponentsBoundingBox(bool bNonColliding = false) const override;
 	virtual bool IsLevelBoundsRelevant() const override { return false; }
-	// End AActor interface.
+	//~ End AActor Interface.
 
 	/** @return Bounding box which includes all relevant actors bounding boxes belonging to specified level */
 	ENGINE_API static FBox CalculateLevelBounds(ULevel* InLevel);
diff --git a/Engine/Source/Runtime/Engine/Classes/Engine/LevelScriptActor.h b/Engine/Source/Runtime/Engine/Classes/Engine/LevelScriptActor.h
--- a/Engine/Source/Runtime/Engine/Classes/Engine/LevelScriptActor.h
+++ b/Engine/Source/Runtime/Engine/Classes/Engine/LevelScriptActor.h
@@ -56,17 +56,17 @@ class ENGINE_API ALevelScriptActor : public AActor
 	void WorldOriginLocationChanged(FIntVector OldOriginLocation, FIntVector NewOriginLocation);
 	
 #if WITH_EDITOR
-	// Begin UObject Interface
+	//~ Begin UObject Interface
 	virtual void PostDuplicate(bool bDuplicateForPIE) override;
 	virtual void BeginDestroy() override;
 #endif
 	virtual void PreInitializeComponents() override;
-	// End UObject Interface
+	//~ End UObject Interface
 
-	// Begin AActor Interface
+	//~ Begin AActor Interface
 	virtual void EnableInput(class APlayerController* PlayerController) override;
 	virtual void DisableInput(class APlayerController* PlayerController) override;
-	// End AActor Interface
+	//~ End AActor Interface
 
 	bool InputEnabled() const { return bInputEnabled; }
 
diff --git a/Engine/Source/Runtime/Engine/Classes/Engine/LevelScriptBlueprint.h b/Engine/Source/Runtime/Engine/Classes/Engine/LevelScriptBlueprint.h
--- a/Engine/Source/Runtime/Engine/Classes/Engine/LevelScriptBlueprint.h
+++ b/Engine/Source/Runtime/Engine/Classes/Engine/LevelScriptBlueprint.h
@@ -34,11 +34,11 @@ class ULevelScriptBlueprint : public UBlueprint
 		return Cast<ULevel>(GetOuter());
 	}
 
-	// Begin UBlueprint interface
+	//~ Begin UBlueprint Interface
 	ENGINE_API virtual UObject* GetObjectBeingDebugged() override;
 	ENGINE_API virtual void SetObjectBeingDebugged(UObject* NewObject) override;
 	ENGINE_API virtual FString GetFriendlyName() const override;
-	// End UBlueprint interface
+	//~ End UBlueprint Interface
 
 	/** Generate a name for a level script blueprint from the current level */
 	static FString CreateLevelScriptNameFromLevel (const class ULevel* Level );
diff --git a/Engine/Source/Runtime/Engine/Classes/Engine/LevelStreaming.h b/Engine/Source/Runtime/Engine/Classes/Engine/LevelStreaming.h
--- a/Engine/Source/Runtime/Engine/Classes/Engine/LevelStreaming.h
+++ b/Engine/Source/Runtime/Engine/Classes/Engine/LevelStreaming.h
@@ -161,7 +161,7 @@ class ULevelStreaming : public UObject
 	UPROPERTY()
 	TArray<FString> Keywords;
 
-	// Begin UObject Interface
+	//~ Begin UObject Interface
 	virtual void PostLoad() override;
 	virtual void Serialize( FArchive& Ar ) override;
 #if WITH_EDITOR
@@ -170,7 +170,7 @@ class ULevelStreaming : public UObject
 	/** Remove duplicates in EditorStreamingVolumes list*/
 	void RemoveStreamingVolumeDuplicates();
 #endif
-	// End UObject Interface
+	//~ End UObject Interface
 
 	bool HasLoadedLevel() const
 	{
@@ -266,7 +266,7 @@ class ULevelStreaming : public UObject
 	UFUNCTION(BlueprintPure, meta = (BlueprintInternalUseOnly = "true"))
 	ALevelScriptActor* GetLevelScriptActor();
 
-	//==============================================================================================
+	//~==============================================================================================
 	// Delegates
 	
 	/** Called when level is streamed in  */
diff --git a/Engine/Source/Runtime/Engine/Classes/Engine/LevelStreamingAlwaysLoaded.h b/Engine/Source/Runtime/Engine/Classes/Engine/LevelStreamingAlwaysLoaded.h
--- a/Engine/Source/Runtime/Engine/Classes/Engine/LevelStreamingAlwaysLoaded.h
+++ b/Engine/Source/Runtime/Engine/Classes/Engine/LevelStreamingAlwaysLoaded.h
@@ -15,9 +15,9 @@ class ULevelStreamingAlwaysLoaded : public ULevelStreaming
 {
 	GENERATED_UCLASS_BODY()
 
-	// Begin ULevelStreaming Interface
+	//~ Begin ULevelStreaming Interface
 	virtual bool ShouldBeLoaded() const override;
 	virtual bool ShouldBeAlwaysLoaded() const override { return true; } 
-	// End ULevelStreaming Interface
+	//~ End ULevelStreaming Interface
 };
 
diff --git a/Engine/Source/Runtime/Engine/Classes/Engine/LevelStreamingKismet.h b/Engine/Source/Runtime/Engine/Classes/Engine/LevelStreamingKismet.h
--- a/Engine/Source/Runtime/Engine/Classes/Engine/LevelStreamingKismet.h
+++ b/Engine/Source/Runtime/Engine/Classes/Engine/LevelStreamingKismet.h
@@ -24,13 +24,13 @@ class ULevelStreamingKismet : public ULevelStreaming
 	UPROPERTY(Category=LevelStreaming, EditAnywhere)
 	uint32 bInitiallyVisible:1;
 	
-	// Begin UObject Interface
+	//~ Begin UObject Interface
 	virtual void PostLoad() override;
-	// End UObject Interface
+	//~ End UObject Interface
 
-	// Begin ULevelStreaming Interface
+	//~ Begin ULevelStreaming Interface
 	virtual bool ShouldBeLoaded() const override;
 	virtual bool ShouldBeVisible() const override;
-	// End ULevelStreaming Interface
+	//~ End ULevelStreaming Interface
 };
 
diff --git a/Engine/Source/Runtime/Engine/Classes/Engine/LevelStreamingPersistent.h b/Engine/Source/Runtime/Engine/Classes/Engine/LevelStreamingPersistent.h
--- a/Engine/Source/Runtime/Engine/Classes/Engine/LevelStreamingPersistent.h
+++ b/Engine/Source/Runtime/Engine/Classes/Engine/LevelStreamingPersistent.h
@@ -16,7 +16,7 @@ class ULevelStreamingPersistent : public ULevelStreaming
 	GENERATED_UCLASS_BODY()
 
 
-	// Begin ULevelStreaming Interface
+	//~ Begin ULevelStreaming Interface
 	virtual bool ShouldBeLoaded() const override
 	{
 		return true;
@@ -25,6 +25,6 @@ class ULevelStreamingPersistent : public ULevelStreaming
 	{
 		return true;
 	}
-	// End ULevelStreaming Interface
+	//~ End ULevelStreaming Interface
 };
 
diff --git a/Engine/Source/Runtime/Engine/Classes/Engine/LevelStreamingVolume.h b/Engine/Source/Runtime/Engine/Classes/Engine/LevelStreamingVolume.h
--- a/Engine/Source/Runtime/Engine/Classes/Engine/LevelStreamingVolume.h
+++ b/Engine/Source/Runtime/Engine/Classes/Engine/LevelStreamingVolume.h
@@ -47,14 +47,14 @@ class ALevelStreamingVolume : public AVolume
 	UPROPERTY(EditAnywhere, BlueprintReadWrite, Category=LevelStreamingVolume)
 	TEnumAsByte<enum EStreamingVolumeUsage> StreamingUsage;
 
-	// Begin UObject interface.
+	//~ Begin UObject Interface.
 	void PostLoad() override;
-	// End UObject interface.
+	//~ End UObject Interface.
 	
 #if WITH_EDITOR
-	// Begin AActor interface.
+	//~ Begin AActor Interface.
 	virtual void CheckForErrors() override;
-	// End AActor interface.
+	//~ End AActor Interface.
 	
 	/** Updates list of streaming levels that are referenced by this streaming volume */
 	void UpdateStreamingLevelsRefs();
diff --git a/Engine/Source/Runtime/Engine/Classes/Engine/Light.h b/Engine/Source/Runtime/Engine/Classes/Engine/Light.h
--- a/Engine/Source/Runtime/Engine/Classes/Engine/Light.h
+++ b/Engine/Source/Runtime/Engine/Classes/Engine/Light.h
@@ -57,9 +57,9 @@ public:
 
 public:
 #if WITH_EDITOR
-	// Begin AActor interface.
+	//~ Begin AActor Interface.
 	virtual void CheckForErrors() override;
-	// End AActor interface.
+	//~ End AActor Interface.
 #endif
 
 	/**
@@ -67,10 +67,10 @@ public:
 	 */
 	bool IsToggleable() const;
 
-	// Begin AActor interface.
+	//~ Begin AActor Interface.
 	void Destroyed() override;
 	virtual bool IsLevelBoundsRelevant() const override { return false; }
-	// End AActor interface.
+	//~ End AActor Interface.
 
 	/** Returns LightComponent subobject **/
 	class ULightComponent* GetLightComponent() const;
diff --git a/Engine/Source/Runtime/Engine/Classes/Engine/NetConnection.h b/Engine/Source/Runtime/Engine/Classes/Engine/NetConnection.h
--- a/Engine/Source/Runtime/Engine/Classes/Engine/NetConnection.h
+++ b/Engine/Source/Runtime/Engine/Classes/Engine/NetConnection.h
@@ -361,7 +361,7 @@ public:
 	// Constructors and destructors.
 	ENGINE_API UNetConnection(const FObjectInitializer& ObjectInitializer = FObjectInitializer::Get());
 
-	// Begin UObject interface.
+	//~ Begin UObject Interface.
 
 	ENGINE_API virtual void Serialize( FArchive& Ar ) override;
 
@@ -376,14 +376,14 @@ public:
 	 */
 	ENGINE_API virtual UWorld* GetWorld() const override;
 
-	// End UObject interface.
+	//~ End UObject Interface.
 
 
-	// Begin FExec interface.
+	//~ Begin FExec Interface.
 
 	ENGINE_API virtual bool Exec( UWorld* InWorld, const TCHAR* Cmd, FOutputDevice& Ar=*GLog ) override;
 
-	// End FExec interface.
+	//~ End FExec Interface.
 
 	/** read input */
 	void ReadInput( float DeltaSeconds );
diff --git a/Engine/Source/Runtime/Engine/Classes/Engine/NetDriver.h b/Engine/Source/Runtime/Engine/Classes/Engine/NetDriver.h
--- a/Engine/Source/Runtime/Engine/Classes/Engine/NetDriver.h
+++ b/Engine/Source/Runtime/Engine/Classes/Engine/NetDriver.h
@@ -329,14 +329,14 @@ public:
 	ENGINE_API UNetDriver(const FObjectInitializer& ObjectInitializer = FObjectInitializer::Get());
 
 
-	// Begin UObject interface.
+	//~ Begin UObject Interface.
 	ENGINE_API virtual void PostInitProperties() override;
 	ENGINE_API virtual void FinishDestroy() override;
 	ENGINE_API virtual void Serialize( FArchive& Ar ) override;
 	ENGINE_API static void AddReferencedObjects(UObject* InThis, FReferenceCollector& Collector);
-	// End UObject interface.
+	//~ End UObject Interface.
 
-	// Begin FExec interface
+	//~ Begin FExec Interface
 
 	/**
 	 * Handle exec commands
@@ -351,7 +351,7 @@ public:
 
 	ENGINE_API ENetMode	GetNetMode() const;
 
-	// End FExec interface.
+	//~ End FExec Interface.
 
 	/** 
 	 * Returns true if this net driver is valid for the current configuration.
diff --git a/Engine/Source/Runtime/Engine/Classes/Engine/PendingNetGame.h b/Engine/Source/Runtime/Engine/Classes/Engine/PendingNetGame.h
--- a/Engine/Source/Runtime/Engine/Classes/Engine/PendingNetGame.h
+++ b/Engine/Source/Runtime/Engine/Classes/Engine/PendingNetGame.h
@@ -103,12 +103,12 @@ public:
 
 	void	InitNetDriver();
 
-	// Begin FNetworkNotify interface.
+	//~ Begin FNetworkNotify Interface.
 	virtual EAcceptConnection::Type NotifyAcceptingConnection() override;
 	virtual void NotifyAcceptedConnection( class UNetConnection* Connection ) override;
 	virtual bool NotifyAcceptingChannel( class UChannel* Channel ) override;
 	virtual void NotifyControlMessage(UNetConnection* Connection, uint8 MessageType, class FInBunch& Bunch) override;
-	// End FNetworkNotify interface.
+	//~ End FNetworkNotify Interface.
 
 	/**  Update the pending level's status. */
 	virtual void Tick( float DeltaTime );
@@ -119,7 +119,7 @@ public:
 	/** Send JOIN to other end */
 	virtual void SendJoin();
 
-	// Begin UObject interface.
+	//~ Begin UObject Interface.
 	virtual void Serialize( FArchive& Ar ) override;
 
 	virtual void FinishDestroy() override
@@ -130,7 +130,7 @@ public:
 	}
 	
 	static void AddReferencedObjects(UObject* InThis, FReferenceCollector& Collector);
-	// End UObject interface.
+	//~ End UObject Interface.
 	
 
 	/** Create the peer net driver and a socket to listen for new client peer connections. */
diff --git a/Engine/Source/Runtime/Engine/Classes/Engine/PlaneReflectionCapture.h b/Engine/Source/Runtime/Engine/Classes/Engine/PlaneReflectionCapture.h
--- a/Engine/Source/Runtime/Engine/Classes/Engine/PlaneReflectionCapture.h
+++ b/Engine/Source/Runtime/Engine/Classes/Engine/PlaneReflectionCapture.h
@@ -15,9 +15,9 @@ class APlaneReflectionCapture : public AReflectionCapture
 public:
 
 #if WITH_EDITOR
-	// Begin AActor interface.
+	//~ Begin AActor Interface.
 	virtual void EditorApplyScale(const FVector& DeltaScale, const FVector* PivotLocation, bool bAltDown, bool bShiftDown, bool bCtrlDown) override;
-	// End AActor interface.
+	//~ End AActor Interface.
 #endif
 
 };
diff --git a/Engine/Source/Runtime/Engine/Classes/Engine/Player.h b/Engine/Source/Runtime/Engine/Classes/Engine/Player.h
--- a/Engine/Source/Runtime/Engine/Classes/Engine/Player.h
+++ b/Engine/Source/Runtime/Engine/Classes/Engine/Player.h
@@ -30,9 +30,9 @@ class UPlayer : public UObject, public FExec
 	int32 ConfiguredLanSpeed;
 
 public:
-	// Begin FExec interface.
+	//~ Begin FExec Interface.
 	ENGINE_API virtual bool Exec( UWorld* InWorld, const TCHAR* Cmd,FOutputDevice& Ar) override;
-	// End FExec interface.
+	//~ End FExec Interface.
 
 	/**
 	 * Dynamically assign Controller to Player and set viewport.
diff --git a/Engine/Source/Runtime/Engine/Classes/Engine/PointLight.h b/Engine/Source/Runtime/Engine/Classes/Engine/PointLight.h
--- a/Engine/Source/Runtime/Engine/Classes/Engine/PointLight.h
+++ b/Engine/Source/Runtime/Engine/Classes/Engine/PointLight.h
@@ -20,17 +20,17 @@ class APointLight : public ALight
 	// END DEPRECATED
 
 #if WITH_EDITOR
-	// Begin AActor interface.
+	//~ Begin AActor Interface.
 	virtual void EditorApplyScale(const FVector& DeltaScale, const FVector* PivotLocation, bool bAltDown, bool bShiftDown, bool bCtrlDown) override;
-	// End AActor interface.
+	//~ End AActor Interface.
 #endif
 
-	// Begin UObject interface.
+	//~ Begin UObject Interface.
 	virtual void PostLoad() override;
 #if WITH_EDITOR
 	virtual void LoadedFromAnotherClass(const FName& OldClassName) override;
 #endif
-	// End UObject interface.
+	//~ End UObject Interface.
 };
 
 
diff --git a/Engine/Source/Runtime/Engine/Classes/Engine/PostProcessVolume.h b/Engine/Source/Runtime/Engine/Classes/Engine/PostProcessVolume.h
--- a/Engine/Source/Runtime/Engine/Classes/Engine/PostProcessVolume.h
+++ b/Engine/Source/Runtime/Engine/Classes/Engine/PostProcessVolume.h
@@ -46,7 +46,7 @@ class APostProcessVolume : public AVolume, public IInterface_PostProcessVolume
 	UPROPERTY(EditAnywhere, BlueprintReadWrite, Category=PostProcessVolume)
 	uint32 bUnbound:1;
 
-	// Begin IInterface_PostProcessVolume Interface
+	//~ Begin IInterface_PostProcessVolume Interface
 	ENGINE_API virtual bool EncompassesPoint(FVector Point, float SphereRadius/*=0.f*/, float* OutDistanceToPoint) override;
 	ENGINE_API virtual FPostProcessVolumeProperties GetProperties() const override
 	{
@@ -59,24 +59,24 @@ class APostProcessVolume : public AVolume, public IInterface_PostProcessVolume
 		Ret.Settings = &Settings;
 		return Ret;
 	}
-	// End IInterface_PostProcessVolume Interface
+	//~ End IInterface_PostProcessVolume Interface
 
 
-	// Begin AActor Interface
+	//~ Begin AActor Interface
 	virtual void PostUnregisterAllComponents( void ) override;
 
 protected:
 	virtual void PostRegisterAllComponents() override;
-	// End AActor Interface
+	//~ End AActor Interface
 public:
 	
-	// Begin UObject interface
+	//~ Begin UObject Interface
 #if WITH_EDITOR
 	virtual void PostEditChangeProperty(FPropertyChangedEvent& PropertyChangedEvent) override;
 	virtual bool CanEditChange(const UProperty* InProperty) const override;
 	virtual void Serialize(FArchive& Ar);
 #endif // WITH_EDITOR
-	// End UObject interface
+	//~ End UObject Interface
 
 	/** Adds an Blendable (implements IBlendableInterface) to the array of Blendables (if it doesn't exist) and update the weight */
 	UFUNCTION(BlueprintCallable, Category="Rendering")
diff --git a/Engine/Source/Runtime/Engine/Classes/Engine/SceneCapture2D.h b/Engine/Source/Runtime/Engine/Classes/Engine/SceneCapture2D.h
--- a/Engine/Source/Runtime/Engine/Classes/Engine/SceneCapture2D.h
+++ b/Engine/Source/Runtime/Engine/Classes/Engine/SceneCapture2D.h
@@ -26,9 +26,9 @@ private_subobject:
 
 public:
 
-	// Begin AActor interface
+	//~ Begin AActor Interface
 	ENGINE_API virtual void PostActorCreated() override;
-	// End AActor interface.
+	//~ End AActor Interface.
 
 	/** Used to synchronize the DrawFrustumComponent with the SceneCaptureComponent2D settings. */
 	void UpdateDrawFrustum();
diff --git a/Engine/Source/Runtime/Engine/Classes/Engine/SceneCaptureCube.h b/Engine/Source/Runtime/Engine/Classes/Engine/SceneCaptureCube.h
--- a/Engine/Source/Runtime/Engine/Classes/Engine/SceneCaptureCube.h
+++ b/Engine/Source/Runtime/Engine/Classes/Engine/SceneCaptureCube.h
@@ -27,12 +27,12 @@ private_subobject:
 
 public:
 
-	// Begin AActor interface
+	//~ Begin AActor Interface
 	virtual void PostActorCreated() override;
 #if WITH_EDITOR
 	ENGINE_API virtual void PostEditMove(bool bFinished) override;
 #endif
-	// End AActor interface.
+	//~ End AActor Interface.
 
 	/** Used to synchronize the DrawFrustumComponent with the SceneCaptureComponentCube settings. */
 	void UpdateDrawFrustum();
diff --git a/Engine/Source/Runtime/Engine/Classes/Engine/Selection.h b/Engine/Source/Runtime/Engine/Classes/Engine/Selection.h
--- a/Engine/Source/Runtime/Engine/Classes/Engine/Selection.h
+++ b/Engine/Source/Runtime/Engine/Classes/Engine/Selection.h
@@ -172,7 +172,7 @@ public:
 
 				if ( InClass->HasAnyClassFlags(CLASS_Interface) )
 				{
-					// InClass is an interface, and we want the top object that implements it
+					//~ Begin InClass is an Interface, and we want the top object that implements it
 					if ( SelectedObject->GetClass()->ImplementsInterface(InClass) )
 					{
 						return SelectedObject;
@@ -180,7 +180,7 @@ public:
 				}
 				else if ( SelectedObject->IsA(InClass) )
 				{
-					// InClass is a class, so we want the top object of that class that implements the required interface, if specified
+					//~ Begin InClass is a class, so we want the top object of that class that implements the required Interface, if specified
 					if ( !RequiredInterface || SelectedObject->GetClass()->ImplementsInterface(RequiredInterface) )
 					{
 						return SelectedObject;
@@ -290,10 +290,10 @@ public:
 	}
 
 
-	// Begin UObject Interface
+	//~ Begin UObject Interface
 	virtual void Serialize(FArchive& Ar) override;
 	virtual bool Modify( bool bAlwaysMarkDirty=true) override;
-	// End UObject Interface
+	//~ End UObject Interface
 
 
 	/**
diff --git a/Engine/Source/Runtime/Engine/Classes/Engine/SkeletalMesh.h b/Engine/Source/Runtime/Engine/Classes/Engine/SkeletalMesh.h
--- a/Engine/Source/Runtime/Engine/Classes/Engine/SkeletalMesh.h
+++ b/Engine/Source/Runtime/Engine/Classes/Engine/SkeletalMesh.h
@@ -498,7 +498,7 @@ struct FClothingAssetData
 	SIZE_T GetResourceSize() const;
 };
 
-// Material interface for USkeletalMesh - contains a material and a shadow casting flag
+//~ Begin Material Interface for USkeletalMesh - contains a material and a shadow casting flag
 USTRUCT()
 struct FSkeletalMaterial
 {
@@ -734,7 +734,7 @@ public:
 	 */
 	ENGINE_API uint32 GetVertexBufferFlags() const;
 
-	// Begin UObject interface.
+	//~ Begin UObject Interface.
 #if WITH_EDITOR
 	virtual void PreEditChange(UProperty* PropertyAboutToChange) override;
 	virtual void PostEditChangeProperty(FPropertyChangedEvent& PropertyChangedEvent) override;
@@ -752,7 +752,7 @@ public:
 	virtual FString GetDetailedInfoInternal() const override;
 	virtual SIZE_T GetResourceSize(EResourceSizeMode::Type Mode) override;
 	static void AddReferencedObjects(UObject* InThis, FReferenceCollector& Collector);
-	// End UObject interface.
+	//~ End UObject Interface.
 
 	/** Setup-only routines - not concerned with the instance. */
 
@@ -878,21 +878,21 @@ public:
 #endif
 	
 
-	// Begin Interface_CollisionDataProvider Interface
+	//~ Begin Interface_CollisionDataProvider Interface
 	virtual bool GetPhysicsTriMeshData(struct FTriMeshCollisionData* CollisionData, bool InUseAllTriData) override;
 	virtual bool ContainsPhysicsTriMeshData(bool InUseAllTriData) const override;
 	virtual bool WantsNegXTriMesh() override
 	{
 		return true;
 	}
-	// End Interface_CollisionDataProvider Interface
+	//~ End Interface_CollisionDataProvider Interface
 
-	// Begin IInterface_AssetUserData Interface
+	//~ Begin IInterface_AssetUserData Interface
 	virtual void AddAssetUserData(UAssetUserData* InUserData) override;
 	virtual void RemoveUserDataOfClass(TSubclassOf<UAssetUserData> InUserDataClass) override;
 	virtual UAssetUserData* GetAssetUserDataOfClass(TSubclassOf<UAssetUserData> InUserDataClass) override;
 	virtual const TArray<UAssetUserData*>* GetAssetUserDataArray() const override;
-	// End IInterface_AssetUserData Interface
+	//~ End IInterface_AssetUserData Interface
 
 private:
 
diff --git a/Engine/Source/Runtime/Engine/Classes/Engine/SphereReflectionCapture.h b/Engine/Source/Runtime/Engine/Classes/Engine/SphereReflectionCapture.h
--- a/Engine/Source/Runtime/Engine/Classes/Engine/SphereReflectionCapture.h
+++ b/Engine/Source/Runtime/Engine/Classes/Engine/SphereReflectionCapture.h
@@ -24,9 +24,9 @@ private_subobject:
 public:
 
 #if WITH_EDITOR
-	// Begin AActor interface.
+	//~ Begin AActor Interface.
 	virtual void EditorApplyScale(const FVector& DeltaScale, const FVector* PivotLocation, bool bAltDown, bool bShiftDown, bool bCtrlDown) override;
-	// End AActor interface.
+	//~ End AActor Interface.
 #endif
 
 	/** Returns DrawCaptureRadius subobject **/
diff --git a/Engine/Source/Runtime/Engine/Classes/Engine/SplineMeshActor.h b/Engine/Source/Runtime/Engine/Classes/Engine/SplineMeshActor.h
--- a/Engine/Source/Runtime/Engine/Classes/Engine/SplineMeshActor.h
+++ b/Engine/Source/Runtime/Engine/Classes/Engine/SplineMeshActor.h
@@ -23,16 +23,16 @@ public:
 	void SetMobility(EComponentMobility::Type InMobility);
 
 #if WITH_EDITOR
-	// Begin AActor Interface
+	//~ Begin AActor Interface
 	virtual void CheckForErrors() override;
 	virtual bool GetReferencedContentObjects(TArray<UObject*>& Objects) const override;
-	// End AActor Interface
+	//~ End AActor Interface
 #endif // WITH_EDITOR	
 
 protected:
-	// Begin UObject interface.
+	//~ Begin UObject Interface.
 	virtual FString GetDetailedInfoInternal() const override;
-	// End UObject interface.
+	//~ End UObject Interface.
 
 public:
 	/** Returns SplineMeshComponent subobject **/
diff --git a/Engine/Source/Runtime/Engine/Classes/Engine/SpotLight.h b/Engine/Source/Runtime/Engine/Classes/Engine/SpotLight.h
--- a/Engine/Source/Runtime/Engine/Classes/Engine/SpotLight.h
+++ b/Engine/Source/Runtime/Engine/Classes/Engine/SpotLight.h
@@ -34,18 +34,18 @@ public:
 	//void SetLightShaftConeAngle(float NewLightShaftConeAngle);
 
 #if WITH_EDITOR
-	// Begin AActor Interface.
+	//~ Begin AActor Interface.
 	virtual void EditorApplyScale(const FVector& DeltaScale, const FVector* PivotLocation, bool bAltDown, bool bShiftDown, bool bCtrlDown) override;
-	// End AActor Interface.
+	//~ End AActor Interface.
 #endif
 
-	// Begin UObject Interface
+	//~ Begin UObject Interface
 	virtual void PostLoad() override;
 #if WITH_EDITOR
 	virtual void LoadedFromAnotherClass(const FName& OldClassName) override;
 	virtual void PostEditChangeProperty(FPropertyChangedEvent& PropertyChangedEvent) override;
 #endif
-	// End UObject Interface
+	//~ End UObject Interface
 
 
 public:
diff --git a/Engine/Source/Runtime/Engine/Classes/Engine/StaticMesh.h b/Engine/Source/Runtime/Engine/Classes/Engine/StaticMesh.h
--- a/Engine/Source/Runtime/Engine/Classes/Engine/StaticMesh.h
+++ b/Engine/Source/Runtime/Engine/Classes/Engine/StaticMesh.h
@@ -451,7 +451,7 @@ public:
 	 */
 	ENGINE_API UStaticMesh(const FObjectInitializer& ObjectInitializer = FObjectInitializer::Get());
 
-	// Begin UObject interface.
+	//~ Begin UObject Interface.
 #if WITH_EDITOR
 	ENGINE_API virtual void PreEditChange(UProperty* PropertyAboutToChange) override;
 	ENGINE_API virtual void PostEditChangeProperty(FPropertyChangedEvent& PropertyChangedEvent) override;
@@ -466,7 +466,7 @@ public:
 	ENGINE_API virtual FString GetDesc() override;
 	ENGINE_API virtual SIZE_T GetResourceSize(EResourceSizeMode::Type Mode) override;
 	static void AddReferencedObjects(UObject* InThis, FReferenceCollector& Collector);
-	// End UObject interface.
+	//~ End UObject Interface.
 
 	/**
 	 * Rebuilds renderable data for this static mesh.
@@ -540,7 +540,7 @@ public:
 	 */
 	ENGINE_API static void CheckLightMapUVs( UStaticMesh* InStaticMesh, TArray< FString >& InOutAssetsWithMissingUVSets, TArray< FString >& InOutAssetsWithBadUVSets, TArray< FString >& InOutAssetsWithValidUVSets, bool bInVerbose = true );
 
-	// Begin Interface_CollisionDataProvider Interface
+	//~ Begin Interface_CollisionDataProvider Interface
 	ENGINE_API virtual bool GetPhysicsTriMeshData(struct FTriMeshCollisionData* CollisionData, bool InUseAllTriData) override;
 	ENGINE_API virtual bool ContainsPhysicsTriMeshData(bool InUseAllTriData) const override;
 	virtual bool WantsNegXTriMesh() override
@@ -548,14 +548,14 @@ public:
 		return true;
 	}
 	ENGINE_API virtual void GetMeshId(FString& OutMeshId) override;
-	// End Interface_CollisionDataProvider Interface
+	//~ End Interface_CollisionDataProvider Interface
 
-	// Begin IInterface_AssetUserData Interface
+	//~ Begin IInterface_AssetUserData Interface
 	virtual void AddAssetUserData(UAssetUserData* InUserData) override;
 	virtual void RemoveUserDataOfClass(TSubclassOf<UAssetUserData> InUserDataClass) override;
 	virtual UAssetUserData* GetAssetUserDataOfClass(TSubclassOf<UAssetUserData> InUserDataClass) override;
 	virtual const TArray<UAssetUserData*>* GetAssetUserDataArray() const override;
-	// End IInterface_AssetUserData Interface
+	//~ End IInterface_AssetUserData Interface
 
 
 	/**
diff --git a/Engine/Source/Runtime/Engine/Classes/Engine/SubsurfaceProfile.h b/Engine/Source/Runtime/Engine/Classes/Engine/SubsurfaceProfile.h
--- a/Engine/Source/Runtime/Engine/Classes/Engine/SubsurfaceProfile.h
+++ b/Engine/Source/Runtime/Engine/Classes/Engine/SubsurfaceProfile.h
@@ -60,10 +60,10 @@ class USubsurfaceProfile : public UObject
 	UPROPERTY(Category = USubsurfaceProfile, EditAnywhere, meta = (ShowOnlyInnerProperties))
 	struct FSubsurfaceProfileStruct Settings;
 
-	// Begin UObject interface
+	//~ Begin UObject Interface
 	virtual void BeginDestroy();
 	virtual void PostEditChangeProperty(struct FPropertyChangedEvent& PropertyChangedEvent);
-	// End UObject interface
+	//~ End UObject Interface
 };
 
 // render thread
@@ -114,7 +114,7 @@ public:
 	const struct IPooledRenderTarget* GetTexture(FRHICommandListImmediate& RHICmdList);
 
 
-	// FRenderResource interface.
+	//~ Begin FRenderResource Interface.
 	/**
 	* Release textures when device is lost/destroyed.
 	*/
diff --git a/Engine/Source/Runtime/Engine/Classes/Engine/Texture2D.h b/Engine/Source/Runtime/Engine/Classes/Engine/Texture2D.h
--- a/Engine/Source/Runtime/Engine/Classes/Engine/Texture2D.h
+++ b/Engine/Source/Runtime/Engine/Classes/Engine/Texture2D.h
@@ -132,7 +132,7 @@ public:
 
 public:
 
-	// Begin UObject interface.
+	//~ Begin UObject Interface.
 	virtual void Serialize(FArchive& Ar) override;	
 #if WITH_EDITOR
 	virtual void PostLinkerChange() override;
@@ -144,9 +144,9 @@ public:
 	virtual void PreSave() override;
 	virtual void GetAssetRegistryTags(TArray<FAssetRegistryTag>& OutTags) const override;
 	virtual FString GetDesc() override;
-	// End UObject interface.
+	//~ End UObject Interface.
 
-	// Begin UTexture interface.
+	//~ Begin UTexture Interface.
 	virtual float GetSurfaceWidth() const override { return GetSizeX(); }
 	virtual float GetSurfaceHeight() const override { return GetSizeY(); }
 	virtual FTextureResource* CreateResource() override;
@@ -157,7 +157,7 @@ public:
 #if WITH_EDITOR
 	virtual TMap<FString,FTexturePlatformData*>* GetCookedPlatformData() override { return &CookedPlatformData; }
 #endif
-	// End UTexture interface.
+	//~ End UTexture Interface.
 
 	/** Trivial accessors. */
 	FORCEINLINE int32 GetSizeX() const
diff --git a/Engine/Source/Runtime/Engine/Classes/Engine/Texture2DDynamic.h b/Engine/Source/Runtime/Engine/Classes/Engine/Texture2DDynamic.h
--- a/Engine/Source/Runtime/Engine/Classes/Engine/Texture2DDynamic.h
+++ b/Engine/Source/Runtime/Engine/Classes/Engine/Texture2DDynamic.h
@@ -27,12 +27,12 @@ class UTexture2DDynamic : public UTexture
 
 
 public:
-	// Begin UTexture Interface.
+	//~ Begin UTexture Interface.
 	ENGINE_API virtual FTextureResource* CreateResource() override;
 	virtual EMaterialValueType GetMaterialType() override { return MCT_Texture2D; }
 	ENGINE_API virtual float GetSurfaceWidth() const override;
 	ENGINE_API virtual float GetSurfaceHeight() const override;
-	// End UTexture Interface.
+	//~ End UTexture Interface.
 	
 	/**
 	 * Initializes the texture with 1 mip-level and creates the render resource.
diff --git a/Engine/Source/Runtime/Engine/Classes/Engine/TextureCube.h b/Engine/Source/Runtime/Engine/Classes/Engine/TextureCube.h
--- a/Engine/Source/Runtime/Engine/Classes/Engine/TextureCube.h
+++ b/Engine/Source/Runtime/Engine/Classes/Engine/TextureCube.h
@@ -14,13 +14,13 @@ public:
 	FTexturePlatformData *PlatformData;
 	TMap<FString, FTexturePlatformData*> CookedPlatformData;
 
-	// Begin UObject interface.
+	//~ Begin UObject Interface.
 	virtual void Serialize(FArchive& Ar) override;
 	virtual void PostLoad() override;
 	virtual void GetAssetRegistryTags(TArray<FAssetRegistryTag>& OutTags) const override;
 	virtual FString GetDesc() override;
 	virtual SIZE_T GetResourceSize(EResourceSizeMode::Type Mode) override;
-	// End UObject interface.
+	//~ End UObject Interface.
 
 	/** Trivial accessors. */
 	FORCEINLINE int32 GetSizeX() const
@@ -56,7 +56,7 @@ public:
 		return PF_Unknown;
 	}
 
-	// Begin UTexture interface
+	//~ Begin UTexture Interface
 	virtual float GetSurfaceWidth() const override { return GetSizeX(); }
 	virtual float GetSurfaceHeight() const override { return GetSizeY(); }
 	virtual FTextureResource* CreateResource() override;
@@ -64,7 +64,7 @@ public:
 	virtual EMaterialValueType GetMaterialType() override { return MCT_TextureCube; }
 	virtual FTexturePlatformData** GetRunningPlatformData() override { return &PlatformData; }
 	virtual TMap<FString, FTexturePlatformData*> *GetCookedPlatformData() override { return &CookedPlatformData; }
-	// End UTexture interface
+	//~ End UTexture Interface
 	
 	/**
 	 * Calculates the size of this texture in bytes if it had MipCount miplevels streamed in.
diff --git a/Engine/Source/Runtime/Engine/Classes/Engine/TextureRenderTarget.h b/Engine/Source/Runtime/Engine/Classes/Engine/TextureRenderTarget.h
--- a/Engine/Source/Runtime/Engine/Classes/Engine/TextureRenderTarget.h
+++ b/Engine/Source/Runtime/Engine/Classes/Engine/TextureRenderTarget.h
@@ -36,10 +36,10 @@ class UTextureRenderTarget : public UTexture
 	ENGINE_API class FTextureRenderTargetResource* GameThread_GetRenderTargetResource();
 
 
-	// Begin UTexture Interface
+	//~ Begin UTexture Interface
 	ENGINE_API virtual class FTextureResource* CreateResource() override;
 	ENGINE_API virtual EMaterialValueType GetMaterialType() override;
-	// End UTexture Interface
+	//~ End UTexture Interface
 };
 
 
diff --git a/Engine/Source/Runtime/Engine/Classes/Engine/TextureRenderTarget2D.h b/Engine/Source/Runtime/Engine/Classes/Engine/TextureRenderTarget2D.h
--- a/Engine/Source/Runtime/Engine/Classes/Engine/TextureRenderTarget2D.h
+++ b/Engine/Source/Runtime/Engine/Classes/Engine/TextureRenderTarget2D.h
@@ -84,21 +84,21 @@ class UTextureRenderTarget2D : public UTextureRenderTarget
 	 */
 	ENGINE_API void UpdateResourceImmediate(bool bClearRenderTarget=true);
 
-	// Begin UTexture interface.
+	//~ Begin UTexture Interface.
 	virtual float GetSurfaceWidth() const override { return SizeX; }
 	virtual float GetSurfaceHeight() const override { return SizeY; }
 	ENGINE_API virtual FTextureResource* CreateResource() override;
 	ENGINE_API virtual EMaterialValueType GetMaterialType() override;
-	// End UTexture interface.
+	//~ End UTexture Interface.
 
-	// Begin UObject interface
+	//~ Begin UObject Interface
 #if WITH_EDITOR
 	ENGINE_API virtual void PostEditChangeProperty(FPropertyChangedEvent& PropertyChangedEvent) override;
 #endif // WITH_EDITOR
 	ENGINE_API virtual void PostLoad() override;
 	ENGINE_API virtual SIZE_T GetResourceSize(EResourceSizeMode::Type Mode) override;
 	ENGINE_API virtual FString GetDesc() override;
-	// End UObject interface
+	//~ End UObject Interface
 
 	FORCEINLINE int32 GetNumMips() const
 	{
diff --git a/Engine/Source/Runtime/Engine/Classes/Engine/TextureRenderTargetCube.h b/Engine/Source/Runtime/Engine/Classes/Engine/TextureRenderTargetCube.h
--- a/Engine/Source/Runtime/Engine/Classes/Engine/TextureRenderTargetCube.h
+++ b/Engine/Source/Runtime/Engine/Classes/Engine/TextureRenderTargetCube.h
@@ -58,14 +58,14 @@ class UTextureRenderTargetCube : public UTextureRenderTarget
 	*/
 	ENGINE_API class UTextureCube* ConstructTextureCube(UObject* InOuter, const FString& NewTexName, EObjectFlags InFlags);
 
-	// Begin UTexture interface.
+	//~ Begin UTexture Interface.
 	virtual float GetSurfaceWidth() const  override { return SizeX; }
 	// PVS-Studio notices that the implementation of GetSurfaceWidth is identical to this one
 	// and warns us. In this case, it is intentional, so we disable the warning:
 	virtual float GetSurfaceHeight()const  override { return SizeX; }	 //-V524
 	virtual FTextureResource* CreateResource() override;
 	virtual EMaterialValueType GetMaterialType() override;
-	// End UTexture interface.
+	//~ End UTexture Interface.
 
 	EPixelFormat GetFormat() const
 	{
@@ -83,14 +83,14 @@ class UTextureRenderTargetCube : public UTextureRenderTarget
 	{
 		return 1;
 	}
-	// Begin UObject interface
+	//~ Begin UObject Interface
 #if WITH_EDITOR
 	virtual void PostEditChangeProperty(FPropertyChangedEvent& PropertyChangedEvent) override;
 #endif // WITH_EDITOR
 	virtual void PostLoad() override;
 	virtual SIZE_T GetResourceSize(EResourceSizeMode::Type Mode) override;
 	virtual FString GetDesc() override;
-	// Begin UObject interface
+	//~ Begin UObject Interface
 };
 
 
diff --git a/Engine/Source/Runtime/Engine/Classes/Engine/UserDefinedEnum.h b/Engine/Source/Runtime/Engine/Classes/Engine/UserDefinedEnum.h
--- a/Engine/Source/Runtime/Engine/Classes/Engine/UserDefinedEnum.h
+++ b/Engine/Source/Runtime/Engine/Classes/Engine/UserDefinedEnum.h
@@ -34,9 +34,9 @@ class ENGINE_API UUserDefinedEnum : public UEnum
 	TArray<FText> DisplayNames;
 
 public:
-	// UObject interface.
+	//~ Begin UObject Interface.
 	virtual void Serialize(FArchive& Ar) override;
-	// End of UObject interface.
+	//~ End UObject Interface.
 	/**
 	 * Generates full enum name give enum name.
 	 * For UUserDefinedEnum full enumerator name has form: '<enumeration path>::<short, user defined enumerator name>'
@@ -62,12 +62,12 @@ public:
 	virtual bool SetEnums(TArray<TPair<FName, uint8>>& InNames, ECppForm InCppForm) override;
 
 #if WITH_EDITOR
-	// Begin UObject interface
+	//~ Begin UObject Interface
 	virtual bool Rename(const TCHAR* NewName = NULL, UObject* NewOuter = NULL, ERenameFlags Flags = REN_None) override;
 	virtual void PostDuplicate(bool bDuplicateForPIE) override;
 	virtual void PostLoad() override;
 	virtual void PostEditUndo() override;
-	// End of UObject interface
+	//~ End UObject Interface
 
 	FString GenerateNewEnumeratorName();
 #endif	// WITH_EDITOR
diff --git a/Engine/Source/Runtime/Engine/Classes/Engine/World.h b/Engine/Source/Runtime/Engine/Classes/Engine/World.h
--- a/Engine/Source/Runtime/Engine/Classes/Engine/World.h
+++ b/Engine/Source/Runtime/Engine/Classes/Engine/World.h
@@ -2312,7 +2312,7 @@ public:
 	 */
 	virtual bool AllowAudioPlayback();
 
-	// Begin UObject Interface
+	//~ Begin UObject Interface
 	virtual void Serialize( FArchive& Ar ) override;
 	virtual void FinishDestroy() override;
 	virtual void PostLoad() override;
@@ -2325,7 +2325,7 @@ public:
 	virtual void GetAssetRegistryTags(TArray<FAssetRegistryTag>& OutTags) const override;
 #endif
 	virtual void PostDuplicate(bool bDuplicateForPIE) override;
-	// End UObject Interface
+	//~ End UObject Interface
 	
 	/**
 	 * Clears all level components and world components like e.g. line batcher.
@@ -3012,12 +3012,12 @@ public:
 	 */
 	bool DestroySwappedPC(UNetConnection* Connection);
 
-	// Begin FNetworkNotify interface
+	//~ Begin FNetworkNotify Interface
 	virtual EAcceptConnection::Type NotifyAcceptingConnection() override;
 	virtual void NotifyAcceptedConnection( class UNetConnection* Connection ) override;
 	virtual bool NotifyAcceptingChannel( class UChannel* Channel ) override;
 	virtual void NotifyControlMessage(UNetConnection* Connection, uint8 MessageType, class FInBunch& Bunch) override;
-	// End FNetworkNotify interface
+	//~ End FNetworkNotify Interface
 
 	/** Welcome a new player joining this server. */
 	void WelcomePlayer(UNetConnection* Connection);
diff --git a/Engine/Source/Runtime/Engine/Classes/GameFramework/Actor.h b/Engine/Source/Runtime/Engine/Classes/GameFramework/Actor.h
--- a/Engine/Source/Runtime/Engine/Classes/GameFramework/Actor.h
+++ b/Engine/Source/Runtime/Engine/Classes/GameFramework/Actor.h
@@ -537,7 +537,7 @@ public:
 	UPROPERTY(transient)
 	uint64 HiddenEditorViews;
 
-	//==============================================================================================
+	//~==============================================================================================
 	// Delegates
 	
 	/** Called when the actor is damaged in any way. */
@@ -645,7 +645,7 @@ public:
 	AController* GetInstigatorController() const;
 
 
-	//=============================================================================
+	//~=============================================================================
 	// General functions.
 
 	/**
@@ -1065,7 +1065,7 @@ public:
 	 */
 	void DetachSceneComponentsFromParent(class USceneComponent* InParentComponent, bool bMaintainWorldPosition = true);
 
-	//==============================================================================
+	//~==============================================================================
 	// Tags
 
 	/** See if this actor contains the supplied tag */
@@ -1073,7 +1073,7 @@ public:
 	bool ActorHasTag(FName Tag) const;
 
 
-	//==============================================================================
+	//~==============================================================================
 	// Misc Blueprint support
 
 	/** 
@@ -1114,7 +1114,7 @@ public:
 	UFUNCTION(BlueprintCallable, meta=(DeprecatedFunction, DeprecationMessage="Use PrimitiveComponent.CreateAndSetMaterialInstanceDynamic instead.", BlueprintProtected = "true"), Category="Rendering|Material")
 	class UMaterialInstanceDynamic* MakeMIDForMaterial(class UMaterialInterface* Parent);
 
-	//=============================================================================
+	//~=============================================================================
 	// AI functions.
 	
 	/**
@@ -1131,7 +1131,7 @@ public:
 	UFUNCTION(BlueprintCallable, BlueprintAuthorityOnly, Category="AI", meta=(BlueprintProtected = "true"))
 	void MakeNoise(float Loudness=1.f, APawn* NoiseInstigator=NULL, FVector NoiseLocation=FVector::ZeroVector, float MaxRange = 0.f, FName Tag = NAME_None);
 
-	//=============================================================================
+	//~=============================================================================
 	// Blueprint
 	
 	/** Event when play begins for this actor. */
@@ -1323,7 +1323,7 @@ public:
 	UPROPERTY(BlueprintAssignable, Category="Game")
 	FActorEndPlaySignature OnEndPlay;
 	
-	// Begin UObject Interface
+	//~ Begin UObject Interface
 	virtual bool CheckDefaultSubobjectsInternal() override;
 	virtual void PostInitProperties() override;
 	virtual bool Modify( bool bAlwaysMarkDirty=true ) override;
@@ -1391,7 +1391,7 @@ public:
 	/** @return true if the component is allowed to re-register its components when modified.  False for CDOs or PIE instances. */
 	bool ReregisterComponentsWhenModified() const;
 #endif // WITH_EDITOR
-	// End UObject Interface
+	//~ End UObject Interface
 
 #if WITH_EDITOR
 	virtual void PostEditMove(bool bFinished);
@@ -2245,7 +2245,7 @@ public:
 	 */
 	virtual void TornOff();
 
-	//=============================================================================
+	//~=============================================================================
 	// Collision functions.
  
 	/** 
@@ -2257,7 +2257,7 @@ public:
 	 */
 	virtual ECollisionResponse GetComponentsCollisionResponseToChannel(ECollisionChannel Channel) const;
 
-	//=============================================================================
+	//~=============================================================================
 	// Physics
 
 	/** Stop all simulation from all components in this actor */
@@ -2545,7 +2545,7 @@ public:
 	const TArray<UActorComponent*>& GetInstanceComponents() const;
 
 public:
-	//=============================================================================
+	//~=============================================================================
 	// Navigation related functions
 	// 
 
@@ -2554,7 +2554,7 @@ public:
 	 */
 	virtual bool IsComponentRelevantForNavigation(UActorComponent* Component) const { return true; }
 
-	//=============================================================================
+	//~=============================================================================
 	// Debugging functions
 public:
 	/**
diff --git a/Engine/Source/Runtime/Engine/Classes/GameFramework/Character.h b/Engine/Source/Runtime/Engine/Classes/GameFramework/Character.h
--- a/Engine/Source/Runtime/Engine/Classes/GameFramework/Character.h
+++ b/Engine/Source/Runtime/Engine/Classes/GameFramework/Character.h
@@ -303,9 +303,9 @@ public:
 	/** @return Saved rotation offset of mesh. */
 	const FQuat& GetBaseRotationOffset() const { return BaseRotationOffset; }
 
-	// Begin INavAgentInterface Interface
+	//~ Begin INavAgentInterface Interface
 	virtual FVector GetNavAgentLocation() const override;
-	// End INavAgentInterface Interface
+	//~ End INavAgentInterface Interface
 
 	/** Default crouched eye height */
 	UPROPERTY(EditAnywhere, BlueprintReadWrite, Category=Camera)
@@ -359,7 +359,7 @@ public:
 	UPROPERTY(EditAnywhere, BlueprintReadWrite, Replicated, Category=Character)
 	float JumpMaxHoldTime;
 
-	// Begin AActor Interface.
+	//~ Begin AActor Interface.
 	virtual void ClearCrossLevelReferences() override;
 	virtual void PreNetReceive() override;
 	virtual void PostNetReceive() override;
@@ -369,7 +369,7 @@ public:
 	virtual void ApplyWorldOffset(const FVector& InOffset, bool bWorldShift) override;
 	virtual UActorComponent* FindComponentByClass(const TSubclassOf<UActorComponent> ComponentClass) const override;
 	virtual void TornOff() override;
-	// End AActor Interface
+	//~ End AActor Interface
 
 	template<class T>
 	T* FindComponentByClass() const
@@ -377,7 +377,7 @@ public:
 		return AActor::FindComponentByClass<T>();
 	}
 
-	// Begin APawn Interface.
+	//~ Begin APawn Interface.
 	virtual void PostInitializeComponents() override;
 	virtual UPawnMovementComponent* GetMovementComponent() const override;
 	virtual UPrimitiveComponent* GetMovementBase() const override final { return BasedMovement.MovementBase; }
@@ -391,7 +391,7 @@ public:
 	virtual void DisplayDebug(class UCanvas* Canvas, const FDebugDisplayInfo& DebugDisplay, float& YL, float& YPos) override;
 	virtual void RecalculateBaseEyeHeight() override;
 	virtual void UpdateNavigationRelevance() override;
-	// End APawn Interface
+	//~ End APawn Interface
 
 	/** Apply momentum caused by damage. */
 	virtual void ApplyDamageMomentum(float DamageTaken, FDamageEvent const& DamageEvent, APawn* PawnInstigator, AActor* DamageCauser);
diff --git a/Engine/Source/Runtime/Engine/Classes/GameFramework/Controller.h b/Engine/Source/Runtime/Engine/Classes/GameFramework/Controller.h
--- a/Engine/Source/Runtime/Engine/Classes/GameFramework/Controller.h
+++ b/Engine/Source/Runtime/Engine/Classes/GameFramework/Controller.h
@@ -10,7 +10,7 @@ class APawn;
 UDELEGATE(BlueprintAuthorityOnly)
 DECLARE_DYNAMIC_MULTICAST_DELEGATE_FourParams( FInstigatedAnyDamageSignature, float, Damage, const class UDamageType*, DamageType, class AActor*, DamagedActor, class AActor*, DamageCauser );
 
-//=============================================================================
+//~=============================================================================
 /**
  * Controllers are non-physical actors that can possess a Pawn to control
  * its actions.  PlayerControllers are used by human players to control pawns, while
@@ -113,7 +113,7 @@ public:
 	/** Actor marking where this controller spawned in. */
 	TWeakObjectPtr<class AActor> StartSpot;
 
-	//=============================================================================
+	//~=============================================================================
 	// CONTROLLER STATE PROPERTIES
 
 	UPROPERTY()
@@ -172,7 +172,7 @@ public:
 
 	virtual void GetLifetimeReplicatedProps(TArray<FLifetimeProperty>& OutLifetimeProps) const override;
 
-	// Begin AActor interface
+	//~ Begin AActor Interface
 	virtual void TickActor( float DeltaTime, enum ELevelTick TickType, FActorTickFunction& ThisTickFunction ) override;
 	virtual void K2_DestroyActor() override;
 	virtual void DisplayDebug(class UCanvas* Canvas, const FDebugDisplayInfo& DebugDisplay, float& YL, float& YPos) override;
@@ -184,7 +184,7 @@ public:
 
 	virtual void Reset() override;
 	virtual void Destroyed() override;
-	// End AActor interface
+	//~ End AActor Interface
 
 	/** Getter for Pawn */
 	FORCEINLINE APawn* GetPawn() const { return Pawn; }
@@ -267,13 +267,13 @@ public:
 	/** GameMode failed to spawn pawn for me. */
 	virtual void FailedToSpawnPawn();
 
-	// Begin INavAgentInterface Interface
+	//~ Begin INavAgentInterface Interface
 	virtual const struct FNavAgentProperties& GetNavAgentPropertiesRef() const override;
 	virtual FVector GetNavAgentLocation() const override;
 	virtual void GetMoveGoalReachTest(class AActor* MovingActor, const FVector& MoveOffset, FVector& GoalOffset, float& GoalRadius, float& GoalHalfHeight) const override;
 	virtual bool ShouldPostponePathUpdates() const override;
 	virtual bool IsFollowingAPath() const override;
-	// End INavAgentInterface Interface
+	//~ End INavAgentInterface Interface
 
 	/** prepares path following component */
 	virtual void InitNavigationControl(UPathFollowingComponent*& PathFollowingComp);
diff --git a/Engine/Source/Runtime/Engine/Classes/GameFramework/EngineMessage.h b/Engine/Source/Runtime/Engine/Classes/GameFramework/EngineMessage.h
--- a/Engine/Source/Runtime/Engine/Classes/GameFramework/EngineMessage.h
+++ b/Engine/Source/Runtime/Engine/Classes/GameFramework/EngineMessage.h
@@ -41,7 +41,7 @@ class  UEngineMessage : public ULocalMessage
 	UPROPERTY()
 	FString NewSpecMessage;
 
-	// Begin ULocalMessage interface
+	//~ Begin ULocalMessage Interface
 	virtual void ClientReceive(const FClientReceiveData& ClientData) const override;
-	// End ULocalMessage interface
+	//~ End ULocalMessage Interface
 };
diff --git a/Engine/Source/Runtime/Engine/Classes/GameFramework/GameMode.h b/Engine/Source/Runtime/Engine/Classes/GameFramework/GameMode.h
--- a/Engine/Source/Runtime/Engine/Classes/GameFramework/GameMode.h
+++ b/Engine/Source/Runtime/Engine/Classes/GameFramework/GameMode.h
@@ -9,10 +9,10 @@
 class FDebugDisplayInfo;
 class FUniqueNetId;
 
-//=============================================================================
+//~=============================================================================
 //  GameMode defines the rules and mechanics of the game.  It is only 
 //  instanced on the server and will never exist on the client.
-//=============================================================================
+//~=============================================================================
 
 /** Possible state of the current match, where a match is all the gameplay that happens on a single map */
 namespace MatchState
@@ -280,12 +280,12 @@ public:
 	UFUNCTION(exec)
 	virtual void SetBandwidthLimit(float AsyncIOBandwidthLimit);
 
-	// Begin AActor interface
+	//~ Begin AActor Interface
 	virtual void PreInitializeComponents() override;
 	virtual void DisplayDebug(class UCanvas* Canvas, const FDebugDisplayInfo& DebugDisplay, float& YL, float& YPos) override;
 	virtual void Reset() override;
 	virtual void Tick(float DeltaSeconds) override;
-	// End AActor interface
+	//~ End AActor Interface
 
 	/** 
 	 * Overridable function to determine whether an Actor should have Reset called when the game has Reset called on it.
@@ -370,7 +370,7 @@ public:
 	 */
 	void ForceClearUnpauseDelegates( AActor* PauseActor );
 
-	//=========================================================================
+	//~=========================================================================
 	// URL Parsing
 	/** Grab the next option from a string. */
 	DEPRECATED(4.9, "Use UGameplayStatics::GrabOption instead")
@@ -391,7 +391,7 @@ public:
 	/** Search array of options for ParseString and return the index or CurrentValue if not found*/
 	DEPRECATED(4.9, "Use UGameplayStatics::GetIntOption instead")
 	static int32 GetIntOption( const FString& Options, const FString& ParseString, int32 CurrentValue);
-	//=========================================================================
+	//~=========================================================================
 
 	/** 
 	 * @return the full path to the optimal GameMode class to use for the specified map and options
diff --git a/Engine/Source/Runtime/Engine/Classes/GameFramework/GameState.h b/Engine/Source/Runtime/Engine/Classes/GameFramework/GameState.h
--- a/Engine/Source/Runtime/Engine/Classes/GameFramework/GameState.h
+++ b/Engine/Source/Runtime/Engine/Classes/GameFramework/GameState.h
@@ -4,9 +4,9 @@
 #include "TimerManager.h"
 #include "GameState.generated.h"
 
-//=============================================================================
+//~=============================================================================
 // GameState is replicated and is valid on servers and clients.
-//=============================================================================
+//~=============================================================================
 
 UCLASS(config=Game, notplaceable, BlueprintType, Blueprintable)
 class ENGINE_API AGameState : public AInfo
@@ -90,9 +90,9 @@ public:
 	UFUNCTION()
 	virtual void OnRep_ElapsedTime();
 
-	// Begin AActor interface
+	//~ Begin AActor Interface
 	virtual void PostInitializeComponents() override;
-	// End AActor interface
+	//~ End AActor Interface
 
 	/** Helper to return the default object of the GameMode class corresponding to this GameState */
 	AGameMode* GetDefaultGameMode() const;
