Cleaning out programmer-only comments that were harvested by Doxygen.

diff --git a/Engine/Source/Runtime/Engine/Classes/GameFramework/MovementComponent.h b/Engine/Source/Runtime/Engine/Classes/GameFramework/MovementComponent.h
--- a/Engine/Source/Runtime/Engine/Classes/GameFramework/MovementComponent.h
+++ b/Engine/Source/Runtime/Engine/Classes/GameFramework/MovementComponent.h
@@ -151,7 +151,7 @@ public:
 	UPROPERTY(EditAnywhere, BlueprintReadOnly, Category=MovementComponent)
 	uint32 bAutoRegisterUpdatedComponent:1;
 
-	// Begin ActorComponent interface 
+	//~ Begin ActorComponent Interface 
 	virtual void TickComponent(float DeltaTime, enum ELevelTick TickType, FActorComponentTickFunction *ThisTickFunction) override;
 	virtual void RegisterComponentTickFunctions(bool bRegister) override;
 	virtual void PostLoad() override;
@@ -168,7 +168,7 @@ public:
 	static void PhysicsLockedAxisSettingChanged();
 #endif // WITH_EDITOR
 
-	// End ActorComponent interface
+	//~ End ActorComponent Interface
 
 	/** @return gravity that affects this component */
 	UFUNCTION(BlueprintCallable, Category="Components|Movement")
diff --git a/Engine/Source/Runtime/Engine/Classes/GameFramework/Pawn.h b/Engine/Source/Runtime/Engine/Classes/GameFramework/Pawn.h
--- a/Engine/Source/Runtime/Engine/Classes/GameFramework/Pawn.h
+++ b/Engine/Source/Runtime/Engine/Classes/GameFramework/Pawn.h
@@ -193,7 +193,7 @@ public:
 	/** used to prevent re-entry of OutsideWorldBounds event. */
 	uint32 bProcessingOutsideWorldBounds:1;
 
-	// Begin AActor Interface.
+	//~ Begin AActor Interface.
 	virtual FVector GetVelocity() const override;
 	virtual void Reset() override;
 	virtual FString GetHumanReadableName() const override;
@@ -221,7 +221,7 @@ public:
 
 	/** Overridden to defer to the RootComponent's CanCharacterStepUpOn setting if it is explicitly Yes or No. If set to Owner, will return Super::CanBeBaseForCharacter(). */
 	virtual bool CanBeBaseForCharacter(APawn* APawn) const override;
-	// End AActor Interface
+	//~ End AActor Interface
 
 	/** Use SetCanAffectNavigationGeneration to change this value at runtime.
 	 *	Note that calling this function at runtime will result in any navigation change only if runtime navigation generation is enabled. */
@@ -231,13 +231,13 @@ public:
 	/** update all components relevant for navigation generators to match bCanAffectNavigationGeneration flag */
 	virtual void UpdateNavigationRelevance() {}
 
-	// Begin INavAgentInterface Interface
+	//~ Begin INavAgentInterface Interface
 	virtual const FNavAgentProperties& GetNavAgentPropertiesRef() const override;
 	/** Basically retrieved pawn's location on navmesh */
 	UFUNCTION(BlueprintCallable, Category="Pawn")
 	virtual FVector GetNavAgentLocation() const override { return GetActorLocation() - FVector(0.f, 0.f, BaseEyeHeight); }
 	virtual void GetMoveGoalReachTest(AActor* MovingActor, const FVector& MoveOffset, FVector& GoalOffset, float& GoalRadius, float& GoalHalfHeight) const override;
-	// End INavAgentInterface Interface
+	//~ End INavAgentInterface Interface
 
 	/** updates MovementComponent's parameters used by navigation system */
 	void UpdateNavAgent();
diff --git a/Engine/Source/Runtime/Engine/Classes/GameFramework/PhysicsVolume.h b/Engine/Source/Runtime/Engine/Classes/GameFramework/PhysicsVolume.h
--- a/Engine/Source/Runtime/Engine/Classes/GameFramework/PhysicsVolume.h
+++ b/Engine/Source/Runtime/Engine/Classes/GameFramework/PhysicsVolume.h
@@ -14,17 +14,17 @@ class ENGINE_API APhysicsVolume : public AVolume
 {
 	GENERATED_UCLASS_BODY()
 
-	// Begin UObject interface.
+	//~ Begin UObject Interface.
 #if WITH_EDITOR
 	virtual void LoadedFromAnotherClass(const FName& OldClassName) override;
 #endif // WITH_EDITOR	
-	// End UObject interface.
+	//~ End UObject Interface.
 
 	virtual void PostInitializeComponents() override;
 	virtual void Destroyed() override;
 	virtual void EndPlay(const EEndPlayReason::Type EndPlayReason) override;
 
-	//======================================================================================
+	//~======================================================================================
 	// Character Movement related properties
 
 	/** Terminal velocity of pawns using CharacterMovement when falling. */
@@ -43,7 +43,7 @@ class ENGINE_API APhysicsVolume : public AVolume
 	UPROPERTY(EditAnywhere, BlueprintReadWrite, Category=CharacterMovement)
 	uint32 bWaterVolume:1;
 
-	//======================================================================================
+	//~======================================================================================
 	// Physics related properties
 	
 	/**	By default, the origin of an AActor must be inside a PhysicsVolume for it to affect the actor. However if this flag is true, the other actor only has to touch the volume to be affected by it. */
diff --git a/Engine/Source/Runtime/Engine/Classes/GameFramework/PlayerController.h b/Engine/Source/Runtime/Engine/Classes/GameFramework/PlayerController.h
--- a/Engine/Source/Runtime/Engine/Classes/GameFramework/PlayerController.h
+++ b/Engine/Source/Runtime/Engine/Classes/GameFramework/PlayerController.h
@@ -129,7 +129,7 @@ protected:
 };
 
 
-//=============================================================================
+//~=============================================================================
 /**
  * PlayerControllers are used by human players to control Pawns.
  *
@@ -1200,11 +1200,11 @@ public:
 	 */
 	virtual FString ConsoleCommand(const FString& Command, bool bWriteToLog = true);
 
-	// Begin UObject Interface
+	//~ Begin UObject Interface
 	virtual void PostLoad() override;
-	// End of UObject Interface
+	//~ End UObject Interface
 
-	// Begin AActor Interface
+	//~ Begin AActor Interface
 	virtual void GetActorEyesViewPoint(FVector& Location, FRotator& Rotation) const override;
 	virtual void CalcCamera(float DeltaTime, struct FMinimalViewInfo& OutResult) override;
 	virtual void TickActor(float DeltaTime, enum ELevelTick TickType, FActorTickFunction& ThisTickFunction) override;
@@ -1229,9 +1229,9 @@ public:
 	virtual void PostInitializeComponents() override;
 	virtual void EnableInput(class APlayerController* PlayerController) override;
 	virtual void DisableInput(class APlayerController* PlayerController) override;
-	// End AActor interface
+	//~ End AActor Interface
 
-	// Begin AController interface
+	//~ Begin AController Interface
 	virtual void GameHasEnded(class AActor* EndGameFocus = NULL, bool bIsWinner = false) override;
 	virtual bool IsLocalPlayerController() const override;
 	virtual bool IsLocalController() const override;
@@ -1243,7 +1243,7 @@ public:
 	virtual void EndInactiveState() override;
 	virtual void FailedToSpawnPawn() override;
 	virtual void SetPawn(APawn* InPawn) override;
-	// End AController interface
+	//~ End AController Interface
 
 	/** called on the server when the client sends a message indicating it was unable to initialize an Actor channel,
 	 * most commonly because the desired Actor's archetype couldn't be serialized
diff --git a/Engine/Source/Runtime/Engine/Classes/GameFramework/PlayerInput.h b/Engine/Source/Runtime/Engine/Classes/GameFramework/PlayerInput.h
--- a/Engine/Source/Runtime/Engine/Classes/GameFramework/PlayerInput.h
+++ b/Engine/Source/Runtime/Engine/Classes/GameFramework/PlayerInput.h
@@ -1,10 +1,10 @@
 // Copyright 1998-2015 Epic Games, Inc. All Rights Reserved.
 
-//=============================================================================
+//~=============================================================================
 // PlayerInput
 // Object within PlayerController that manages player input.
 // Only spawned on client.
-//=============================================================================
+//~=============================================================================
 
 #pragma once
 #include "Components/InputComponent.h"
@@ -389,10 +389,10 @@ private:
 
 public:
 	
-	// Begin UObject Interface
+	//~ Begin UObject Interface
 	virtual void PostInitProperties() override;
 	virtual UWorld* GetWorld() const override;
-	// End UObject Interface
+	//~ End UObject Interface
 
 	/** Flushes the current key state. */
 	void FlushPressedKeys();
diff --git a/Engine/Source/Runtime/Engine/Classes/GameFramework/Volume.h b/Engine/Source/Runtime/Engine/Classes/GameFramework/Volume.h
--- a/Engine/Source/Runtime/Engine/Classes/GameFramework/Volume.h
+++ b/Engine/Source/Runtime/Engine/Classes/GameFramework/Volume.h
@@ -28,7 +28,7 @@ private:
 	/** Called during posteditchange after the volume's initial shape has changed */
 	static FOnVolumeShapeChanged OnVolumeShapeChanged;
 public:
-	// Begin AActor Interface
+	//~ Begin AActor Interface
 	/**
 	* Function that gets called from within Map_Check to allow this actor to check itself
 	* for any potential errors and register them with map check dialog.
@@ -37,12 +37,12 @@ public:
 #endif // WITH_EDITOR
 	
 	virtual bool IsLevelBoundsRelevant() const override;
-	// End AActor Interface
+	//~ End AActor Interface
 
-	// Begin Brush Interface
+	//~ Begin Brush Interface
 	virtual bool IsStaticBrush() const override;
 	virtual bool IsVolumeBrush() const override;
-	// End Brush Interface
+	//~ End Brush Interface
 
 	/** @returns true if a sphere/point (with optional radius CheckRadius) overlaps this volume */
 	bool EncompassesPoint(FVector Point, float SphereRadius=0.f, float* OutDistanceToPoint = 0);
diff --git a/Engine/Source/Runtime/Engine/Classes/GameFramework/WheeledVehicle.h b/Engine/Source/Runtime/Engine/Classes/GameFramework/WheeledVehicle.h
--- a/Engine/Source/Runtime/Engine/Classes/GameFramework/WheeledVehicle.h
+++ b/Engine/Source/Runtime/Engine/Classes/GameFramework/WheeledVehicle.h
@@ -42,9 +42,9 @@ public:
 	/** Util to get the wheeled vehicle movement component */
 	class UWheeledVehicleMovementComponent* GetVehicleMovementComponent() const;
 
-	// Begin AActor interface
+	//~ Begin AActor Interface
 	virtual void DisplayDebug(class UCanvas* Canvas, const FDebugDisplayInfo& DebugDisplay, float& YL, float& YPos) override;
-	// End Actor interface
+	//~ End Actor Interface
 
 	/** Returns Mesh subobject **/
 	class USkeletalMeshComponent* GetMesh() const;
diff --git a/Engine/Source/Runtime/Engine/Classes/GameFramework/WorldSettings.h b/Engine/Source/Runtime/Engine/Classes/GameFramework/WorldSettings.h
--- a/Engine/Source/Runtime/Engine/Classes/GameFramework/WorldSettings.h
+++ b/Engine/Source/Runtime/Engine/Classes/GameFramework/WorldSettings.h
@@ -485,21 +485,21 @@ class ENGINE_API AWorldSettings : public AInfo, public IInterface_AssetUserData
 	TArray<UAssetUserData*> AssetUserData;
 
 public:
-	// Begin UObject interface.
+	//~ Begin UObject Interface.
 	virtual void PostLoad() override;
 #if WITH_EDITOR
 	virtual void PostEditChangeProperty(FPropertyChangedEvent& PropertyChangedEvent) override;
 #endif // WITH_EDITOR
-	// End UObject interface.
+	//~ End UObject Interface.
 
 
-	// Begin AActor interface.
+	//~ Begin AActor Interface.
 #if WITH_EDITOR
 	virtual void CheckForErrors() override;
 #endif
 	virtual void PreInitializeComponents() override;
 	virtual void PostInitializeComponents() override;
-	// End AActor interface.
+	//~ End AActor Interface.
 
 	/**
 	 * Returns the Z component of the current world gravity and initializes it to the default
@@ -529,11 +529,11 @@ public:
 	 */	
 	virtual void NotifyMatchStarted();
 
-	// Begin IInterface_AssetUserData Interface
+	//~ Begin IInterface_AssetUserData Interface
 	virtual void AddAssetUserData(UAssetUserData* InUserData) override;
 	virtual void RemoveUserDataOfClass(TSubclassOf<UAssetUserData> InUserDataClass) override;
 	virtual UAssetUserData* GetAssetUserDataOfClass(TSubclassOf<UAssetUserData> InUserDataClass) override;
-	// End IInterface_AssetUserData Interface
+	//~ End IInterface_AssetUserData Interface
 
 
 private:
diff --git a/Engine/Source/Runtime/Engine/Classes/Materials/Material.h b/Engine/Source/Runtime/Engine/Classes/Materials/Material.h
--- a/Engine/Source/Runtime/Engine/Classes/Materials/Material.h
+++ b/Engine/Source/Runtime/Engine/Classes/Materials/Material.h
@@ -766,7 +766,7 @@ private:
 #endif // WITH_EDITORONLY_DATA
 public:
 
-	// Begin UMaterialInterface interface.
+	//~ Begin UMaterialInterface Interface.
 	ENGINE_API virtual UMaterial* GetMaterial() override;
 	ENGINE_API virtual const UMaterial* GetMaterial() const override;
 	ENGINE_API virtual const UMaterial* GetMaterial_Concurrent(TMicRecursionGuard& RecursionGuard) const override;
@@ -813,9 +813,9 @@ public:
 	/** Allows material properties to be compiled with the option of being overridden by the material attributes input. */
 	ENGINE_API virtual int32 CompilePropertyEx( class FMaterialCompiler* Compiler, EMaterialProperty Property ) override;
 	ENGINE_API virtual void ForceRecompileForRendering() override;
-	// End UMaterialInterface interface.
+	//~ End UMaterialInterface Interface.
 
-	// Begin UObject Interface
+	//~ Begin UObject Interface
 	ENGINE_API virtual void PreSave() override;
 	ENGINE_API virtual void PostInitProperties() override;
 	ENGINE_API virtual void Serialize(FArchive& Ar) override;
@@ -837,7 +837,7 @@ public:
 	ENGINE_API virtual void FinishDestroy() override;
 	ENGINE_API virtual SIZE_T GetResourceSize(EResourceSizeMode::Type Mode) override;
 	ENGINE_API static void AddReferencedObjects(UObject* InThis, FReferenceCollector& Collector);
-	// End UObject Interface
+	//~ End UObject Interface
 
 #if WITH_EDITOR
 	/** Cancels any currently outstanding compilation jobs for this material. Useful in the material editor when some edits superceds existing, in flight compilation jobs.*/
diff --git a/Engine/Source/Runtime/Engine/Classes/Materials/MaterialExpressionAbs.h b/Engine/Source/Runtime/Engine/Classes/Materials/MaterialExpressionAbs.h
--- a/Engine/Source/Runtime/Engine/Classes/Materials/MaterialExpressionAbs.h
+++ b/Engine/Source/Runtime/Engine/Classes/Materials/MaterialExpressionAbs.h
@@ -19,10 +19,10 @@ class UMaterialExpressionAbs : public UMaterialExpression
 	FExpressionInput Input;
 
 
-	// Begin UMaterialExpression Interface
+	//~ Begin UMaterialExpression Interface
 	virtual int32 Compile(class FMaterialCompiler* Compiler, int32 OutputIndex, int32 MultiplexIndex) override;
 	virtual void GetCaption(TArray<FString>& OutCaptions) const override;
-	// End UMaterialExpression Interface
+	//~ End UMaterialExpression Interface
 
 };
 
diff --git a/Engine/Source/Runtime/Engine/Classes/Materials/MaterialExpressionActorPositionWS.h b/Engine/Source/Runtime/Engine/Classes/Materials/MaterialExpressionActorPositionWS.h
--- a/Engine/Source/Runtime/Engine/Classes/Materials/MaterialExpressionActorPositionWS.h
+++ b/Engine/Source/Runtime/Engine/Classes/Materials/MaterialExpressionActorPositionWS.h
@@ -11,10 +11,10 @@ class UMaterialExpressionActorPositionWS : public UMaterialExpression
 	GENERATED_UCLASS_BODY()
 
 
-	// Begin UMaterialExpression Interface
+	//~ Begin UMaterialExpression Interface
 	virtual int32 Compile(class FMaterialCompiler* Compiler, int32 OutputIndex, int32 MultiplexIndex) override;
 	virtual void GetCaption(TArray<FString>& OutCaptions) const override;
-	// End UMaterialExpression Interface
+	//~ End UMaterialExpression Interface
 };
 
 
diff --git a/Engine/Source/Runtime/Engine/Classes/Materials/MaterialExpressionAntialiasedTextureMask.h b/Engine/Source/Runtime/Engine/Classes/Materials/MaterialExpressionAntialiasedTextureMask.h
--- a/Engine/Source/Runtime/Engine/Classes/Materials/MaterialExpressionAntialiasedTextureMask.h
+++ b/Engine/Source/Runtime/Engine/Classes/Materials/MaterialExpressionAntialiasedTextureMask.h
@@ -27,16 +27,16 @@ class UMaterialExpressionAntialiasedTextureMask : public UMaterialExpressionText
 	TEnumAsByte<enum ETextureColorChannel> Channel;
 
 
-	// Begin UMaterialExpressionTextureSampleParameter Interface
+	//~ Begin UMaterialExpressionTextureSampleParameter Interface
 	virtual bool TextureIsValid( UTexture* InTexture ) override;
 	virtual const TCHAR* GetRequirements() override;
 	virtual void SetDefaultTexture() override;
-	// End UMaterialExpressionTextureSampleParameter Interface
+	//~ End UMaterialExpressionTextureSampleParameter Interface
 
-	// Begin UMaterialExpression Interface
+	//~ Begin UMaterialExpression Interface
 	virtual int32 Compile(class FMaterialCompiler* Compiler, int32 OutputIndex, int32 MultiplexIndex) override;
 	virtual void GetCaption(TArray<FString>& OutCaptions) const override;
-	// End UMaterialExpression Interface
+	//~ End UMaterialExpression Interface
 };
 
 
diff --git a/Engine/Source/Runtime/Engine/Classes/Materials/MaterialExpressionClamp.h b/Engine/Source/Runtime/Engine/Classes/Materials/MaterialExpressionClamp.h
--- a/Engine/Source/Runtime/Engine/Classes/Materials/MaterialExpressionClamp.h
+++ b/Engine/Source/Runtime/Engine/Classes/Materials/MaterialExpressionClamp.h
@@ -37,14 +37,14 @@ class UMaterialExpressionClamp : public UMaterialExpression
 	float MaxDefault;
 
 
-	// Begin UObject Interface
+	//~ Begin UObject Interface
 	virtual void Serialize( FArchive& Ar ) override;
-	// End UObject Interface
+	//~ End UObject Interface
 
-	// Begin UMaterialExpression Interface
+	//~ Begin UMaterialExpression Interface
 	virtual int32 Compile(class FMaterialCompiler* Compiler, int32 OutputIndex, int32 MultiplexIndex) override;
 	virtual void GetCaption(TArray<FString>& OutCaptions) const override;
-	// End UMaterialExpression Interface
+	//~ End UMaterialExpression Interface
 };
 
 
diff --git a/Engine/Source/Runtime/Engine/Classes/Materials/MaterialExpressionCollectionParameter.h b/Engine/Source/Runtime/Engine/Classes/Materials/MaterialExpressionCollectionParameter.h
--- a/Engine/Source/Runtime/Engine/Classes/Materials/MaterialExpressionCollectionParameter.h
+++ b/Engine/Source/Runtime/Engine/Classes/Materials/MaterialExpressionCollectionParameter.h
@@ -25,18 +25,18 @@ class UMaterialExpressionCollectionParameter : public UMaterialExpression
 	UPROPERTY()
 	FGuid ParameterId;
 
-	// Begin UObject interface.
+	//~ Begin UObject Interface.
 #if WITH_EDITOR
 	virtual void PostEditChangeProperty(FPropertyChangedEvent& PropertyChangedEvent) override;
 #endif // WITH_EDITOR
 	virtual void PostLoad() override;
-	// End UObject interface.
+	//~ End UObject Interface.
 
-	// Begin UMaterialExpression Interface
+	//~ Begin UMaterialExpression Interface
 	virtual int32 Compile(class FMaterialCompiler* Compiler, int32 OutputIndex, int32 MultiplexIndex) override;
 	virtual void GetCaption(TArray<FString>& OutCaptions) const override;
 	virtual bool MatchesSearchQuery( const TCHAR* SearchQuery ) override;
-	// End UMaterialExpression Interface
+	//~ End UMaterialExpression Interface
 };
 
 
diff --git a/Engine/Source/Runtime/Engine/Classes/Materials/MaterialExpressionComment.h b/Engine/Source/Runtime/Engine/Classes/Materials/MaterialExpressionComment.h
--- a/Engine/Source/Runtime/Engine/Classes/Materials/MaterialExpressionComment.h
+++ b/Engine/Source/Runtime/Engine/Classes/Materials/MaterialExpressionComment.h
@@ -23,17 +23,17 @@ class UMaterialExpressionComment : public UMaterialExpression
 	UPROPERTY(EditAnywhere, Category=MaterialExpressionComment)
 	FLinearColor CommentColor;
 
-	// Begin UObject Interface
+	//~ Begin UObject Interface
 #if WITH_EDITOR
 	virtual void PostEditChangeProperty(FPropertyChangedEvent& PropertyChangedEvent) override;
 #endif // WITH_EDITOR
 	virtual bool Modify( bool bAlwaysMarkDirty=true ) override;
-	// End UObject Interface
+	//~ End UObject Interface
 
-	// Begin UMaterialExpression Interface
+	//~ Begin UMaterialExpression Interface
 	virtual void GetCaption(TArray<FString>& OutCaptions) const override;
 	virtual bool MatchesSearchQuery( const TCHAR* SearchQuery ) override;
-	// End UMaterialExpression Interface
+	//~ End UMaterialExpression Interface
 };
 
 
diff --git a/Engine/Source/Runtime/Engine/Classes/Materials/MaterialExpressionCustom.h b/Engine/Source/Runtime/Engine/Classes/Materials/MaterialExpressionCustom.h
--- a/Engine/Source/Runtime/Engine/Classes/Materials/MaterialExpressionCustom.h
+++ b/Engine/Source/Runtime/Engine/Classes/Materials/MaterialExpressionCustom.h
@@ -46,13 +46,13 @@ class UMaterialExpressionCustom : public UMaterialExpression
 	TArray<struct FCustomInput> Inputs;
 
 
-	// Begin UObject interface.
+	//~ Begin UObject Interface.
 #if WITH_EDITOR
 	virtual void PostEditChangeProperty(FPropertyChangedEvent& PropertyChangedEvent) override;
 #endif // WITH_EDITOR
-	// End UObject interface.
+	//~ End UObject Interface.
 	
-	// Begin UMaterialExpression Interface
+	//~ Begin UMaterialExpression Interface
 	virtual int32 Compile(class FMaterialCompiler* Compiler, int32 OutputIndex, int32 MultiplexIndex) override;
 	virtual void GetCaption(TArray<FString>& OutCaptions) const override;
 	virtual const TArray<FExpressionInput*> GetInputs() override;
@@ -62,7 +62,7 @@ class UMaterialExpressionCustom : public UMaterialExpression
 	virtual uint32 GetInputType(int32 InputIndex) override {return MCT_Unknown;}
 	virtual uint32 GetOutputType(int32 OutputIndex) override;
 #endif // WITH_EDITOR
-	// End UMaterialExpression Interface
+	//~ End UMaterialExpression Interface
 };
 
 
diff --git a/Engine/Source/Runtime/Engine/Classes/Materials/MaterialExpressionDepthFade.h b/Engine/Source/Runtime/Engine/Classes/Materials/MaterialExpressionDepthFade.h
--- a/Engine/Source/Runtime/Engine/Classes/Materials/MaterialExpressionDepthFade.h
+++ b/Engine/Source/Runtime/Engine/Classes/Materials/MaterialExpressionDepthFade.h
@@ -30,7 +30,7 @@ class UMaterialExpressionDepthFade : public UMaterialExpression
 	float FadeDistanceDefault;
 
 
-	// Begin UMaterialExpression Interface
+	//~ Begin UMaterialExpression Interface
 	virtual FString GetInputName(int32 InputIndex) const override
 	{
 		if (InputIndex == 0)
@@ -47,7 +47,7 @@ class UMaterialExpressionDepthFade : public UMaterialExpression
 	{
 		OutCaptions.Add(TEXT("Depth Fade"));
 	}
-	// End UMaterialExpression Interface#endif
+	//~ End UMaterialExpression Interface#endif
 };
 
 
diff --git a/Engine/Source/Runtime/Engine/Classes/Materials/MaterialExpressionTextureObjectParameter.h b/Engine/Source/Runtime/Engine/Classes/Materials/MaterialExpressionTextureObjectParameter.h
--- a/Engine/Source/Runtime/Engine/Classes/Materials/MaterialExpressionTextureObjectParameter.h
+++ b/Engine/Source/Runtime/Engine/Classes/Materials/MaterialExpressionTextureObjectParameter.h
@@ -15,7 +15,7 @@ class UMaterialExpressionTextureObjectParameter : public UMaterialExpressionText
 	GENERATED_UCLASS_BODY()
 
 
-	// Begin UMaterialExpression Interface
+	//~ Begin UMaterialExpression Interface
 	virtual void GetCaption(TArray<FString>& OutCaptions) const override;
 	virtual const TArray<FExpressionInput*> GetInputs() override;
 	virtual int32 Compile(class FMaterialCompiler* Compiler, int32 OutputIndex, int32 MultiplexIndex) override;
@@ -23,11 +23,11 @@ class UMaterialExpressionTextureObjectParameter : public UMaterialExpressionText
 #if WITH_EDITOR
 	virtual uint32 GetOutputType(int32 InputIndex) override {return MCT_Texture;}
 #endif
-	// End UMaterialExpression Interface
+	//~ End UMaterialExpression Interface
 
-	// Begin UMaterialExpressionTextureSampleParameter Interface
+	//~ Begin UMaterialExpressionTextureSampleParameter Interface
 	virtual const TCHAR* GetRequirements() override;
-	// End UMaterialExpressionTextureSampleParameter Interface
+	//~ End UMaterialExpressionTextureSampleParameter Interface
 
 };
 
diff --git a/Engine/Source/Runtime/Engine/Classes/Materials/MaterialExpressionTextureSampleParameter2D.h b/Engine/Source/Runtime/Engine/Classes/Materials/MaterialExpressionTextureSampleParameter2D.h
--- a/Engine/Source/Runtime/Engine/Classes/Materials/MaterialExpressionTextureSampleParameter2D.h
+++ b/Engine/Source/Runtime/Engine/Classes/Materials/MaterialExpressionTextureSampleParameter2D.h
@@ -11,15 +11,15 @@ class ENGINE_API UMaterialExpressionTextureSampleParameter2D : public UMaterialE
 	GENERATED_UCLASS_BODY()
 
 
-	// Begin UMaterialExpression Interface
+	//~ Begin UMaterialExpression Interface
 	virtual void GetCaption(TArray<FString>& OutCaptions) const override;
-	// End UMaterialExpression Interface
+	//~ End UMaterialExpression Interface
 	
-	// Begin UMaterialExpressionTextureSampleParameter Interface
+	//~ Begin UMaterialExpressionTextureSampleParameter Interface
 	virtual bool TextureIsValid( UTexture* InTexture ) override;
 	virtual const TCHAR* GetRequirements() override;
 	virtual void SetDefaultTexture() override;
-	// End UMaterialExpressionTextureSampleParameter Interface
+	//~ End UMaterialExpressionTextureSampleParameter Interface
 };
 
 
diff --git a/Engine/Source/Runtime/Engine/Classes/Materials/MaterialExpressionTextureSampleParameterSubUV.h b/Engine/Source/Runtime/Engine/Classes/Materials/MaterialExpressionTextureSampleParameterSubUV.h
--- a/Engine/Source/Runtime/Engine/Classes/Materials/MaterialExpressionTextureSampleParameterSubUV.h
+++ b/Engine/Source/Runtime/Engine/Classes/Materials/MaterialExpressionTextureSampleParameterSubUV.h
@@ -14,15 +14,15 @@ class UMaterialExpressionTextureSampleParameterSubUV : public UMaterialExpressio
 	uint32 bBlend:1;
 
 
-	// Begin UMaterialExpressionTextureSampleParameter Interface
+	//~ Begin UMaterialExpressionTextureSampleParameter Interface
 	virtual bool TextureIsValid( UTexture* InTexture ) override;
 	virtual const TCHAR* GetRequirements() override;
-	// End UMaterialExpressionTextureSampleParameter Interface
+	//~ End UMaterialExpressionTextureSampleParameter Interface
 
-	// Begin UMaterialExpression Interface
+	//~ Begin UMaterialExpression Interface
 	virtual int32 Compile(class FMaterialCompiler* Compiler, int32 OutputIndex, int32 MultiplexIndex) override;
 	virtual void GetCaption(TArray<FString>& OutCaptions) const override;
-	// End UMaterialExpression Interface
+	//~ End UMaterialExpression Interface
 };
 
 
diff --git a/Engine/Source/Runtime/Engine/Classes/Materials/MaterialInstance.h b/Engine/Source/Runtime/Engine/Classes/Materials/MaterialInstance.h
--- a/Engine/Source/Runtime/Engine/Classes/Materials/MaterialInstance.h
+++ b/Engine/Source/Runtime/Engine/Classes/Materials/MaterialInstance.h
@@ -239,11 +239,11 @@ private:
 	FRenderCommandFence ReleaseFence;
 
 public:
-	// Begin interface IBlendableInterface
+	//~ Begin Begin Interface IBlendableInterface
 	virtual ENGINE_API void OverrideBlendableSettings(class FSceneView& View, float Weight) const override;
-	// End interface IBlendableInterface
+	//~ Begin End Interface IBlendableInterface
 
-	// Begin UMaterialInterface interface.
+	//~ Begin UMaterialInterface Interface.
 	virtual ENGINE_API UMaterial* GetMaterial() override;
 	virtual ENGINE_API const UMaterial* GetMaterial() const override;
 	virtual ENGINE_API const UMaterial* GetMaterial_Concurrent(TMicRecursionGuard& RecursionGuard) const override;
@@ -298,9 +298,9 @@ public:
 	ENGINE_API virtual bool IsPropertyActive(EMaterialProperty InProperty) const override;
 	/** Allows material properties to be compiled with the option of being overridden by the material attributes input. */
 	ENGINE_API virtual int32 CompilePropertyEx(class FMaterialCompiler* Compiler, EMaterialProperty Property) override;
-	// End UMaterialInterface interface.
+	//~ End UMaterialInterface Interface.
 
-	// Begin UObject interface.
+	//~ Begin UObject Interface.
 	virtual ENGINE_API SIZE_T GetResourceSize(EResourceSizeMode::Type Mode) override;
 	virtual ENGINE_API void PostInitProperties() override;	
 #if WITH_EDITOR
@@ -326,7 +326,7 @@ public:
 
 #endif // WITH_EDITOR
 
-	// End UObject interface.
+	//~ End UObject Interface.
 
 	/**
 	 * Recompiles static permutations if necessary.
diff --git a/Engine/Source/Runtime/Engine/Classes/Materials/MaterialInterface.h b/Engine/Source/Runtime/Engine/Classes/Materials/MaterialInterface.h
--- a/Engine/Source/Runtime/Engine/Classes/Materials/MaterialInterface.h
+++ b/Engine/Source/Runtime/Engine/Classes/Materials/MaterialInterface.h
@@ -208,7 +208,7 @@ public:
 	/** Set which feature levels _all_ materials should compile to. GMaxRHIFeatureLevel is always compiled. */
 	ENGINE_API static void SetGlobalRequiredFeatureLevel(ERHIFeatureLevel::Type FeatureLevel, bool bShouldCompile);
 
-	// Begin UObject interface.
+	//~ Begin UObject Interface.
 	ENGINE_API virtual void BeginDestroy() override;
 	ENGINE_API virtual bool IsReadyForFinishDestroy() override;
 	ENGINE_API virtual void PostLoad() override;
@@ -216,11 +216,11 @@ public:
 #if WITH_EDITOR
 	ENGINE_API virtual void PostEditChangeProperty(FPropertyChangedEvent& PropertyChangedEvent) override;
 #endif // WITH_EDITOR
-	// End UObject interface.
+	//~ End UObject Interface.
 
-	// Begin interface IBlendableInterface
+	//~ Begin Begin Interface IBlendableInterface
 	ENGINE_API virtual void OverrideBlendableSettings(class FSceneView& View, float Weight) const override;
-	// End interface IBlendableInterface
+	//~ Begin End Interface IBlendableInterface
 
 	/** Walks up parent chain and finds the base Material that this is an instance of. Just calls the virtual GetMaterial() */
 	UFUNCTION(BlueprintCallable, Category="Rendering|Material")
diff --git a/Engine/Source/Runtime/Engine/Classes/Matinee/InterpFilter_Classes.h b/Engine/Source/Runtime/Engine/Classes/Matinee/InterpFilter_Classes.h
--- a/Engine/Source/Runtime/Engine/Classes/Matinee/InterpFilter_Classes.h
+++ b/Engine/Source/Runtime/Engine/Classes/Matinee/InterpFilter_Classes.h
@@ -27,8 +27,8 @@ class UInterpFilter_Classes : public UInterpFilter
 
 #endif // WITH_EDITORONLY_DATA
 
-	// Begin UInterpFilter Interface
+	//~ Begin UInterpFilter Interface
 	virtual void FilterData(class AMatineeActor* InMatineeActor) override;
-	// End UInterpFilter Interface
+	//~ End UInterpFilter Interface
 };
 
diff --git a/Engine/Source/Runtime/Engine/Classes/Matinee/InterpFilter_Custom.h b/Engine/Source/Runtime/Engine/Classes/Matinee/InterpFilter_Custom.h
--- a/Engine/Source/Runtime/Engine/Classes/Matinee/InterpFilter_Custom.h
+++ b/Engine/Source/Runtime/Engine/Classes/Matinee/InterpFilter_Custom.h
@@ -22,8 +22,8 @@ class UInterpFilter_Custom : public UInterpFilter
 
 #endif // WITH_EDITORONLY_DATA
 
-	// Begin UInterpFilter Interface
+	//~ Begin UInterpFilter Interface
 	virtual void FilterData(class AMatineeActor* InMatineeActor) override;
-	// End UInterpFilter Interface
+	//~ End UInterpFilter Interface
 };
 
diff --git a/Engine/Source/Runtime/Engine/Classes/Matinee/InterpGroupDirector.h b/Engine/Source/Runtime/Engine/Classes/Matinee/InterpGroupDirector.h
--- a/Engine/Source/Runtime/Engine/Classes/Matinee/InterpGroupDirector.h
+++ b/Engine/Source/Runtime/Engine/Classes/Matinee/InterpGroupDirector.h
@@ -14,10 +14,10 @@ class UInterpGroupDirector : public UInterpGroup
 {
 	GENERATED_UCLASS_BODY()
 
-	// Begin UInterpGroup Interface
+	//~ Begin UInterpGroup Interface
 	virtual AActor* SelectGroupActor( class UInterpGroupInst* GrInst, bool bDeselectActors ) override;
 	virtual AActor* DeselectGroupActor( class UInterpGroupInst* GrInst ) override;
-	// End UInterpGroup Interface
+	//~ End UInterpGroup Interface
 
 	/** @return the director track inside this Director group - if present. */
 	ENGINE_API class UInterpTrackDirector* GetDirectorTrack();
diff --git a/Engine/Source/Runtime/Engine/Classes/Matinee/InterpTrack.h b/Engine/Source/Runtime/Engine/Classes/Matinee/InterpTrack.h
--- a/Engine/Source/Runtime/Engine/Classes/Matinee/InterpTrack.h
+++ b/Engine/Source/Runtime/Engine/Classes/Matinee/InterpTrack.h
@@ -164,9 +164,9 @@ public:
 
 
 public:
-	// Begin FInterpEdInputInterface Interface
+	//~ Begin FInterpEdInputInterface Interface
 	virtual UObject* GetUObject() override {return this;}
-	// End FInterpEdInputInterface Interface
+	//~ End FInterpEdInputInterface Interface
 
 	/** @return  The total number of keyframes currently in this track. */
 	virtual int32 GetNumKeyframes() const { return 0; }
diff --git a/Engine/Source/Runtime/Engine/Classes/Matinee/InterpTrackAnimControl.h b/Engine/Source/Runtime/Engine/Classes/Matinee/InterpTrackAnimControl.h
--- a/Engine/Source/Runtime/Engine/Classes/Matinee/InterpTrackAnimControl.h
+++ b/Engine/Source/Runtime/Engine/Classes/Matinee/InterpTrackAnimControl.h
@@ -74,11 +74,11 @@ class UInterpTrackAnimControl : public UInterpTrackFloatBase
 
 
 
-	// Begin UObject interface.
+	//~ Begin UObject Interface.
 	virtual void PostLoad() override;
-	// End UObject interface.
+	//~ End UObject Interface.
 
-	// Begin UInterpTrack interface.
+	//~ Begin UInterpTrack Interface.
 	virtual int32 GetNumKeyframes() const override;
 	virtual void GetTimeRange(float& StartTime, float& EndTime) const override;
 	virtual float GetTrackEndTime() const override;
@@ -98,14 +98,14 @@ class UInterpTrackAnimControl : public UInterpTrackFloatBase
 	virtual class UTexture2D* GetTrackIcon() const override;
 #endif // WITH_EDITORONLY_DATA
 	virtual void DrawTrack( FCanvas* Canvas, UInterpGroup* Group, const FInterpTrackDrawParams& Params ) override;
-	// End UInterpTrack interface.
+	//~ End UInterpTrack Interface.
 
-	// Begin FInterpEdInputInterface Interface
+	//~ Begin FInterpEdInputInterface Interface
 	virtual void BeginDrag(FInterpEdInputData &InputData) override;
 	virtual void EndDrag(FInterpEdInputData &InputData) override;
 	virtual EMouseCursor::Type GetMouseCursor(FInterpEdInputData &InputData) override;
 	virtual void ObjectDragged(FInterpEdInputData& InputData) override;
-	// End FInterpEdInputInterface Interface
+	//~ End FInterpEdInputInterface Interface
 
 	/** 
 	 * Calculates the reversed time for a sequence key, if the key has bReverse set.
diff --git a/Engine/Source/Runtime/Engine/Classes/Matinee/InterpTrackAudioMaster.h b/Engine/Source/Runtime/Engine/Classes/Matinee/InterpTrackAudioMaster.h
--- a/Engine/Source/Runtime/Engine/Classes/Matinee/InterpTrackAudioMaster.h
+++ b/Engine/Source/Runtime/Engine/Classes/Matinee/InterpTrackAudioMaster.h
@@ -11,7 +11,7 @@ class UInterpTrackAudioMaster : public UInterpTrackVectorBase
 	GENERATED_UCLASS_BODY()
 
 
-	// Begin UInterpTrack interface.
+	//~ Begin UInterpTrack Interface.
 	virtual int32 AddKeyframe(float Time, UInterpTrackInst* TrInst, EInterpCurveMode InitInterpMode) override;
 	virtual void UpdateKeyframe(int32 KeyIndex, UInterpTrackInst* TrInst) override;
 	virtual void PreviewUpdateTrack(float NewPosition, UInterpTrackInst* TrInst) override;
@@ -20,7 +20,7 @@ class UInterpTrackAudioMaster : public UInterpTrackVectorBase
 #if WITH_EDITORONLY_DATA
 	virtual class UTexture2D* GetTrackIcon() const override;
 #endif // WITH_EDITORONLY_DATA
-	// End UInterpTrack interface.
+	//~ End UInterpTrack Interface.
 
 	/** Return the sound volume scale for the specified time */
 	float GetVolumeScaleForTime( float Time ) const;
diff --git a/Engine/Source/Runtime/Engine/Classes/Matinee/InterpTrackBoolProp.h b/Engine/Source/Runtime/Engine/Classes/Matinee/InterpTrackBoolProp.h
--- a/Engine/Source/Runtime/Engine/Classes/Matinee/InterpTrackBoolProp.h
+++ b/Engine/Source/Runtime/Engine/Classes/Matinee/InterpTrackBoolProp.h
@@ -39,7 +39,7 @@ class UInterpTrackBoolProp : public UInterpTrack
 	FName PropertyName;
 
 
-	// Begin UInterpTrack Interface
+	//~ Begin UInterpTrack Interface
 	virtual int32 GetNumKeyframes() const override;
 	virtual float GetTrackEndTime() const override;
 	virtual float GetKeyframeTime( int32 KeyIndex ) const override;
@@ -60,7 +60,7 @@ class UInterpTrackBoolProp : public UInterpTrack
 #if WITH_EDITORONLY_DATA
 	virtual class UTexture2D* GetTrackIcon() const override;
 #endif // WITH_EDITORONLY_DATA
-	// End UInterpTrack Interface
+	//~ End UInterpTrack Interface
 };
 
 
diff --git a/Engine/Source/Runtime/Engine/Classes/Matinee/InterpTrackColorProp.h b/Engine/Source/Runtime/Engine/Classes/Matinee/InterpTrackColorProp.h
--- a/Engine/Source/Runtime/Engine/Classes/Matinee/InterpTrackColorProp.h
+++ b/Engine/Source/Runtime/Engine/Classes/Matinee/InterpTrackColorProp.h
@@ -15,7 +15,7 @@ class UInterpTrackColorProp : public UInterpTrackVectorBase
 	FName PropertyName;
 
 
-	// Begin UInterpTrack interface
+	//~ Begin UInterpTrack Interface
 	virtual int32 AddKeyframe(float Time, UInterpTrackInst* TrInst, EInterpCurveMode InitInterpMode) override;
 	virtual bool CanAddKeyframe( UInterpTrackInst* TrackInst ) override;
 	virtual void UpdateKeyframe(int32 KeyIndex, UInterpTrackInst* TrInst) override;
@@ -26,7 +26,7 @@ class UInterpTrackColorProp : public UInterpTrackVectorBase
 #if WITH_EDITORONLY_DATA
 	virtual class UTexture2D* GetTrackIcon() const override;
 #endif // WITH_EDITORONLY_DATA
-	// End UInterpTrack interface
+	//~ End UInterpTrack Interface
 };
 
 
diff --git a/Engine/Source/Runtime/Engine/Classes/Matinee/InterpTrackColorScale.h b/Engine/Source/Runtime/Engine/Classes/Matinee/InterpTrackColorScale.h
--- a/Engine/Source/Runtime/Engine/Classes/Matinee/InterpTrackColorScale.h
+++ b/Engine/Source/Runtime/Engine/Classes/Matinee/InterpTrackColorScale.h
@@ -15,7 +15,7 @@ class UInterpTrackColorScale : public UInterpTrackVectorBase
 	GENERATED_UCLASS_BODY()
 
 
-	// Begin UInterpTrack interface.
+	//~ Begin UInterpTrack Interface.
 	virtual int32 AddKeyframe(float Time, UInterpTrackInst* TrInst, EInterpCurveMode InitInterpMode) override;
 	virtual void UpdateKeyframe(int32 KeyIndex, UInterpTrackInst* TrInst) override;
 	virtual void PreviewUpdateTrack(float NewPosition, UInterpTrackInst* TrInst) override;
@@ -24,7 +24,7 @@ class UInterpTrackColorScale : public UInterpTrackVectorBase
 #if WITH_EDITORONLY_DATA
 	virtual class UTexture2D* GetTrackIcon() const override;
 #endif // WITH_EDITORONLY_DATA
-	// Begin UInterpTrack interface.
+	//~ Begin UInterpTrack Interface.
 
 	/** Return the blur alpha we want at the given time. */
 	ENGINE_API FVector GetColorScaleAtTime(float Time);
diff --git a/Engine/Source/Runtime/Engine/Classes/Matinee/InterpTrackDirector.h b/Engine/Source/Runtime/Engine/Classes/Matinee/InterpTrackDirector.h
--- a/Engine/Source/Runtime/Engine/Classes/Matinee/InterpTrackDirector.h
+++ b/Engine/Source/Runtime/Engine/Classes/Matinee/InterpTrackDirector.h
@@ -64,11 +64,11 @@ private:
 #endif // WITH_EDITORONLY_DATA
 
 public:
-	// Begin UObject interface
+	//~ Begin UObject Interface
 	virtual void PostLoad() override;
-	// End UObject interface
+	//~ End UObject Interface
 	
-	// Begin UInterpTrack Interface
+	//~ Begin UInterpTrack Interface
 	virtual int32 GetNumKeyframes() const override;
 	virtual void GetTimeRange( float& StartTime, float& EndTime ) const override;
 	virtual float GetTrackEndTime() const override;
@@ -87,7 +87,7 @@ public:
 	virtual class UTexture2D* GetTrackIcon() const override;
 #endif // WITH_EDITORONLY_DATA
 	virtual void DrawTrack( FCanvas* Canvas, UInterpGroup* Group, const FInterpTrackDrawParams& Params ) override;
-	// End UInterpTrack Interface
+	//~ End UInterpTrack Interface
 
 	/** Get the keyframe index nearest the time, without going over */
 	int32 GetNearestKeyframeIndex( float KeyTime ) const;
diff --git a/Engine/Source/Runtime/Engine/Classes/Matinee/InterpTrackEvent.h b/Engine/Source/Runtime/Engine/Classes/Matinee/InterpTrackEvent.h
--- a/Engine/Source/Runtime/Engine/Classes/Matinee/InterpTrackEvent.h
+++ b/Engine/Source/Runtime/Engine/Classes/Matinee/InterpTrackEvent.h
@@ -56,7 +56,7 @@ class UInterpTrackEvent : public UInterpTrack
 	UPROPERTY(EditAnywhere, Category=InterpTrackEvent)
 	uint32 bUseCustomEventName:1;
 
-	// Begin UInterpTrack Interface
+	//~ Begin UInterpTrack Interface
 	virtual int32 GetNumKeyframes() const override;
 	virtual void GetTimeRange( float& StartTime, float& EndTime ) const override;
 	virtual float GetTrackEndTime() const override;
@@ -76,7 +76,7 @@ class UInterpTrackEvent : public UInterpTrack
 #endif // WITH_EDITORONLY_DATA
 	virtual bool AllowStaticActors() override { return true; }
 	virtual void DrawTrack( FCanvas* Canvas, UInterpGroup* Group, const FInterpTrackDrawParams& Params ) override;
-	// End UInterpTrack Interface
+	//~ End UInterpTrack Interface
 };
 
 
diff --git a/Engine/Source/Runtime/Engine/Classes/Matinee/InterpTrackFade.h b/Engine/Source/Runtime/Engine/Classes/Matinee/InterpTrackFade.h
--- a/Engine/Source/Runtime/Engine/Classes/Matinee/InterpTrackFade.h
+++ b/Engine/Source/Runtime/Engine/Classes/Matinee/InterpTrackFade.h
@@ -27,7 +27,7 @@ class UInterpTrackFade : public UInterpTrackFloatBase
 	UPROPERTY(EditAnywhere, Category = InterpTrackFade)
 	FLinearColor FadeColor;
 
-	// Begin UInterpTrack interface.
+	//~ Begin UInterpTrack Interface.
 	virtual int32 AddKeyframe(float Time, UInterpTrackInst* TrInst, EInterpCurveMode InitInterpMode) override;
 	virtual void UpdateKeyframe(int32 KeyIndex, UInterpTrackInst* TrInst) override;
 	virtual void PreviewUpdateTrack(float NewPosition, UInterpTrackInst* TrInst) override;
@@ -35,7 +35,7 @@ class UInterpTrackFade : public UInterpTrackFloatBase
 #if WITH_EDITORONLY_DATA
 	virtual class UTexture2D* GetTrackIcon() const override;
 #endif // WITH_EDITORONLY_DATA
-	// End UInterpTrack interface.
+	//~ End UInterpTrack Interface.
 
 	/** @return the amount of fading we want at the given time. */
 	ENGINE_API float GetFadeAmountAtTime(float Time);
diff --git a/Engine/Source/Runtime/Engine/Classes/Matinee/InterpTrackFloatAnimBPParam.h b/Engine/Source/Runtime/Engine/Classes/Matinee/InterpTrackFloatAnimBPParam.h
--- a/Engine/Source/Runtime/Engine/Classes/Matinee/InterpTrackFloatAnimBPParam.h
+++ b/Engine/Source/Runtime/Engine/Classes/Matinee/InterpTrackFloatAnimBPParam.h
@@ -16,17 +16,17 @@ class UInterpTrackFloatAnimBPParam : public UInterpTrackFloatBase
 	UPROPERTY(EditAnywhere, Category=InterpTrackFloatAnimBPParam)
 	FName ParamName;
 
-	// Begin UObject Interface
+	//~ Begin UObject Interface
 #if WITH_EDITOR
 	virtual void PostEditChangeProperty(FPropertyChangedEvent& PropertyChangedEvent) override;
 #endif // WITH_EDITOR
-	// End UObject Interface
+	//~ End UObject Interface
 
-	// Begin UInterpTrack interface.
+	//~ Begin UInterpTrack Interface.
 	virtual int32 AddKeyframe(float Time, UInterpTrackInst* TrInst, EInterpCurveMode InitInterpMode) override;
 	virtual void PreviewUpdateTrack(float NewPosition, UInterpTrackInst* TrInst) override;
 	virtual void UpdateTrack(float NewPosition, UInterpTrackInst* TrInst, bool bJump) override;
-	// End UInterpTrack interface.
+	//~ End UInterpTrack Interface.
 
 private:
 	bool bRefreshParamter;
diff --git a/Engine/Source/Runtime/Engine/Classes/Matinee/InterpTrackFloatBase.h b/Engine/Source/Runtime/Engine/Classes/Matinee/InterpTrackFloatBase.h
--- a/Engine/Source/Runtime/Engine/Classes/Matinee/InterpTrackFloatBase.h
+++ b/Engine/Source/Runtime/Engine/Classes/Matinee/InterpTrackFloatBase.h
@@ -18,13 +18,13 @@ class ENGINE_API UInterpTrackFloatBase : public UInterpTrack
 	float CurveTension;
 
 
-	// Begin UObject interface.
+	//~ Begin UObject Interface.
 #if WITH_EDITOR
 	virtual void PostEditChangeProperty(FPropertyChangedEvent& PropertyChangedEvent) override;
 #endif // WITH_EDITOR
-	// End UObject interface.
+	//~ End UObject Interface.
 
-	// Begin UInterpTrack interface.
+	//~ Begin UInterpTrack Interface.
 	virtual int32 GetNumKeyframes() const override;
 	virtual void GetTimeRange(float& StartTime, float& EndTime) const override;
 	virtual float GetTrackEndTime() const override;
@@ -35,9 +35,9 @@ class ENGINE_API UInterpTrackFloatBase : public UInterpTrack
 	virtual int32 DuplicateKeyframe(int32 KeyIndex, float NewKeyTime, UInterpTrack* ToTrack = NULL) override;
 	virtual bool GetClosestSnapPosition(float InPosition, TArray<int32> &IgnoreKeys, float& OutPosition) override;
 	virtual FColor GetKeyframeColor(int32 KeyIndex) const override;
-	// End UInterpTrack interface.
+	//~ End UInterpTrack Interface.
 
-	// Begin FCurveEdInterface interface.
+	//~ Begin FCurveEdInterface Interface.
 	virtual int32		GetNumKeys() const override;
 	virtual int32		GetNumSubCurves() const override;
 	virtual float	GetKeyIn(int32 KeyIndex) override;
@@ -53,7 +53,7 @@ class ENGINE_API UInterpTrackFloatBase : public UInterpTrack
 	virtual void	SetKeyOut(int32 SubIndex, int32 KeyIndex, float NewOutVal) override;
 	virtual void	SetKeyInterpMode(int32 KeyIndex, EInterpCurveMode NewMode) override;
 	virtual void	SetTangents(int32 SubIndex, int32 KeyIndex, float ArriveTangent, float LeaveTangent) override;
-	// End FCurveEdInterface interface.
+	//~ End FCurveEdInterface Interface.
 };
 
 
diff --git a/Engine/Source/Runtime/Engine/Classes/Matinee/InterpTrackFloatMaterialParam.h b/Engine/Source/Runtime/Engine/Classes/Matinee/InterpTrackFloatMaterialParam.h
--- a/Engine/Source/Runtime/Engine/Classes/Matinee/InterpTrackFloatMaterialParam.h
+++ b/Engine/Source/Runtime/Engine/Classes/Matinee/InterpTrackFloatMaterialParam.h
@@ -17,18 +17,18 @@ class UInterpTrackFloatMaterialParam : public UInterpTrackFloatBase
 	FName ParamName;
 
 
-	// Begin UObject Interface
+	//~ Begin UObject Interface
 #if WITH_EDITOR
 	virtual void PreEditChange(UProperty* PropertyThatWillChange) override;
 	virtual void PostEditChangeProperty(FPropertyChangedEvent& PropertyChangedEvent) override;
 #endif // WITH_EDITOR
-	// End UObject Interface
+	//~ End UObject Interface
 
-	// Begin UInterpTrack interface.
+	//~ Begin UInterpTrack Interface.
 	virtual int32 AddKeyframe(float Time, UInterpTrackInst* TrInst, EInterpCurveMode InitInterpMode) override;
 	virtual void PreviewUpdateTrack(float NewPosition, UInterpTrackInst* TrInst) override;
 	virtual void UpdateTrack(float NewPosition, UInterpTrackInst* TrInst, bool bJump) override;
-	// End UInterpTrack interface.
+	//~ End UInterpTrack Interface.
 };
 
 
diff --git a/Engine/Source/Runtime/Engine/Classes/Matinee/InterpTrackFloatProp.h b/Engine/Source/Runtime/Engine/Classes/Matinee/InterpTrackFloatProp.h
--- a/Engine/Source/Runtime/Engine/Classes/Matinee/InterpTrackFloatProp.h
+++ b/Engine/Source/Runtime/Engine/Classes/Matinee/InterpTrackFloatProp.h
@@ -14,7 +14,7 @@ class UInterpTrackFloatProp : public UInterpTrackFloatBase
 	FName PropertyName;
 
 
-	// Begin InterpTrack interface.
+	//~ Begin InterpTrack Interface.
 	virtual int32 AddKeyframe(float Time, UInterpTrackInst* TrInst, EInterpCurveMode InitInterpMode) override;
 	virtual bool CanAddKeyframe( UInterpTrackInst* TrackInst ) override;
 	virtual void UpdateKeyframe(int32 KeyIndex, UInterpTrackInst* TrInst) override;
@@ -26,7 +26,7 @@ class UInterpTrackFloatProp : public UInterpTrackFloatBase
 	virtual class UTexture2D* GetTrackIcon() const override;
 #endif // WITH_EDITORONLY_DATA
 	virtual void ReduceKeys( float IntervalStart, float IntervalEnd, float Tolerance ) override;
-	// End InterpTrack interface.
+	//~ End InterpTrack Interface.
 };
 
 
diff --git a/Engine/Source/Runtime/Engine/Classes/Matinee/InterpTrackLinearColorBase.h b/Engine/Source/Runtime/Engine/Classes/Matinee/InterpTrackLinearColorBase.h
--- a/Engine/Source/Runtime/Engine/Classes/Matinee/InterpTrackLinearColorBase.h
+++ b/Engine/Source/Runtime/Engine/Classes/Matinee/InterpTrackLinearColorBase.h
@@ -18,7 +18,7 @@ class UInterpTrackLinearColorBase : public UInterpTrack
 	float CurveTension;
 
 
-	// Begin UObject interface.
+	//~ Begin UObject Interface.
 #if WITH_EDITOR
 	virtual void PostEditChangeProperty(FPropertyChangedEvent& PropertyChangedEvent) override;
 #endif // WITH_EDITOR
@@ -26,9 +26,9 @@ class UInterpTrackLinearColorBase : public UInterpTrack
 	{
 		Super::Serialize(Ar);
 	}
-	// End UObject interface.
+	//~ End UObject Interface.
 
-	// Begin UInterpTrack interface.
+	//~ Begin UInterpTrack Interface.
 	virtual int32 GetNumKeyframes() const override;
 	virtual void GetTimeRange(float& StartTime, float& EndTime) const override; 
 	virtual float GetTrackEndTime() const override;
@@ -39,9 +39,9 @@ class UInterpTrackLinearColorBase : public UInterpTrack
 	virtual int32 DuplicateKeyframe(int32 KeyIndex, float NewKeyTime, UInterpTrack* ToTrack = NULL) override;
 	virtual bool GetClosestSnapPosition(float InPosition, TArray<int32> &IgnoreKeys, float& OutPosition) override;
 	virtual FColor GetKeyframeColor(int32 KeyIndex) const override;
-	// End UInterpTrack interface.
+	//~ End UInterpTrack Interface.
 
-	// Begin FCurveEdInterface interface
+	//~ Begin FCurveEdInterface Interface
 	virtual int32		GetNumKeys() const override;
 	virtual int32		GetNumSubCurves() const override;
 	virtual FColor	GetSubCurveButtonColor(int32 SubCurveIndex, bool bIsSubCurveHidden) const override;
@@ -59,7 +59,7 @@ class UInterpTrackLinearColorBase : public UInterpTrack
 	virtual void	SetKeyOut(int32 SubIndex, int32 KeyIndex, float NewOutVal) override;
 	virtual void	SetKeyInterpMode(int32 KeyIndex, EInterpCurveMode NewMode) override;
 	virtual void	SetTangents(int32 SubIndex, int32 KeyIndex, float ArriveTangent, float LeaveTangent) override;
-	// End FCurveEdInterface interface
+	//~ End FCurveEdInterface Interface
 };
 
 
diff --git a/Engine/Source/Runtime/Engine/Classes/Matinee/InterpTrackMove.h b/Engine/Source/Runtime/Engine/Classes/Matinee/InterpTrackMove.h
--- a/Engine/Source/Runtime/Engine/Classes/Matinee/InterpTrackMove.h
+++ b/Engine/Source/Runtime/Engine/Classes/Matinee/InterpTrackMove.h
@@ -148,14 +148,14 @@ class UInterpTrackMove : public UInterpTrack
 	TEnumAsByte<enum EInterpTrackMoveRotMode> RotMode;
 
 
-	// Begin UObject interface.
+	//~ Begin UObject Interface.
 #if WITH_EDITOR
 	virtual void PostEditChangeProperty(FPropertyChangedEvent& PropertyChangedEvent) override;
 	virtual void PostEditImport() override;
 #endif
-	// End UObject interface.
+	//~ End UObject Interface.
 
-	// Begin UInterpTrack interface.
+	//~ Begin UInterpTrack Interface.
 	virtual int32 GetNumKeyframes() const override;
 	virtual void GetTimeRange(float& StartTime, float& EndTime) const override;
 	virtual float GetTrackEndTime() const override;
@@ -182,9 +182,9 @@ class UInterpTrackMove : public UInterpTrack
 	virtual void Render3DTrack(UInterpTrackInst* TrInst, const FSceneView* View, class FPrimitiveDrawInterface* PDI, int32 TrackIndex, const FColor& TrackColor, TArray<struct FInterpEdSelKey>& SelectedKeys) override;
 	virtual void SetTrackToSensibleDefault() override;
 	virtual void ApplyWorldOffset(const FVector& InOffset, bool bWorldShift) override;
-	// End UInterpTrack interface.
+	//~ End UInterpTrack Interface.
 
-	// Begin FCurveEdInterface interface.
+	//~ Begin FCurveEdInterface Interface.
 	virtual int32		GetNumKeys() const override;
 	virtual int32		GetNumSubCurves() const override;
 	virtual FColor	GetSubCurveButtonColor(int32 SubCurveIndex, bool bIsSubCurveHidden) const override;
@@ -202,7 +202,7 @@ class UInterpTrackMove : public UInterpTrack
 	virtual void	SetKeyOut(int32 SubIndex, int32 KeyIndex, float NewOutVal) override;
 	virtual void	SetKeyInterpMode(int32 KeyIndex, EInterpCurveMode NewMode) override;
 	virtual void	SetTangents(int32 SubIndex, int32 KeyIndex, float ArriveTangent, float LeaveTangent) override;
-	// End FCurveEdInterface interface.
+	//~ End FCurveEdInterface Interface.
 
 
 	/**
diff --git a/Engine/Source/Runtime/Engine/Classes/Matinee/InterpTrackMoveAxis.h b/Engine/Source/Runtime/Engine/Classes/Matinee/InterpTrackMoveAxis.h
--- a/Engine/Source/Runtime/Engine/Classes/Matinee/InterpTrackMoveAxis.h
+++ b/Engine/Source/Runtime/Engine/Classes/Matinee/InterpTrackMoveAxis.h
@@ -37,7 +37,7 @@ class UInterpTrackMoveAxis : public UInterpTrackFloatBase
 	struct FInterpLookupTrack LookupTrack;
 
 
-	// Begin UInterpTrack Interface
+	//~ Begin UInterpTrack Interface
 	virtual int32 AddKeyframe( float Time, UInterpTrackInst* TrInst, EInterpCurveMode InitInterpMode ) override;
 	virtual void UpdateKeyframe(int32 KeyIndex, UInterpTrackInst* TrInst) override;
 	virtual int32 SetKeyframeTime( int32 KeyIndex, float NewKeyTime, bool bUpdateOrder ) override;
@@ -47,17 +47,17 @@ class UInterpTrackMoveAxis : public UInterpTrackFloatBase
 	virtual class UTexture2D* GetTrackIcon() const override;
 #endif // WITH_EDITORONLY_DATA
 	virtual void ReduceKeys( float IntervalStart, float IntervalEnd, float Tolerance ) override;
-	// End UInterpTrack Interface
+	//~ End UInterpTrack Interface
 	
 
 
-	// Begin FCurveEdInterface interface.
+	//~ Begin FCurveEdInterface Interface.
 	virtual FColor GetSubCurveButtonColor( int32 SubCurveIndex, bool bIsSubCurveHidden ) const override;
 	virtual int32 CreateNewKey( float KeyIn ) override;
 	virtual void DeleteKey( int32 KeyIndex ) override;
 	virtual int32 SetKeyIn( int32 KeyIndex, float NewInVal ) override;
 	virtual FColor GetKeyColor(int32 SubIndex, int32 KeyIndex, const FColor& CurveColor) override;
-	// End FCurveEdInterface interface.
+	//~ End FCurveEdInterface Interface.
 
 	/** @todo document */
 	void GetKeyframeValue( UInterpTrackInst* TrInst, int32 KeyIndex, float& OutTime, float &OutValue, float* OutArriveTangent, float* OutLeaveTangent );
diff --git a/Engine/Source/Runtime/Engine/Classes/Matinee/InterpTrackParticleReplay.h b/Engine/Source/Runtime/Engine/Classes/Matinee/InterpTrackParticleReplay.h
--- a/Engine/Source/Runtime/Engine/Classes/Matinee/InterpTrackParticleReplay.h
+++ b/Engine/Source/Runtime/Engine/Classes/Matinee/InterpTrackParticleReplay.h
@@ -56,7 +56,7 @@ class UInterpTrackParticleReplay : public UInterpTrack
 
 #endif // WITH_EDITORONLY_DATA
 
-	// Begin UInterpTrack interface.
+	//~ Begin UInterpTrack Interface.
 	virtual int32 GetNumKeyframes() const override;
 	virtual void GetTimeRange(float& StartTime, float& EndTime) const override;
 	virtual float GetTrackEndTime() const override;
@@ -76,14 +76,14 @@ class UInterpTrackParticleReplay : public UInterpTrack
 #endif // WITH_EDITORONLY_DATA
 	virtual bool AllowStaticActors() override { return true; }
 	virtual void DrawTrack( FCanvas* Canvas, UInterpGroup* Group, const FInterpTrackDrawParams& Params ) override;
-	// End UInterpTrack interface.
+	//~ End UInterpTrack Interface.
 
-	// Begin FInterpEdInputInterface Interface
+	//~ Begin FInterpEdInputInterface Interface
 	virtual void BeginDrag(FInterpEdInputData &InputData) override;
 	virtual void EndDrag(FInterpEdInputData &InputData) override;
 	virtual EMouseCursor::Type GetMouseCursor(FInterpEdInputData &InputData) override;
 	virtual void ObjectDragged(FInterpEdInputData& InputData) override;
-	// End FInterpEdInputInterface Interface
+	//~ End FInterpEdInputInterface Interface
 
 };
 
diff --git a/Engine/Source/Runtime/Engine/Classes/Matinee/InterpTrackSlomo.h b/Engine/Source/Runtime/Engine/Classes/Matinee/InterpTrackSlomo.h
--- a/Engine/Source/Runtime/Engine/Classes/Matinee/InterpTrackSlomo.h
+++ b/Engine/Source/Runtime/Engine/Classes/Matinee/InterpTrackSlomo.h
@@ -10,7 +10,7 @@ class UInterpTrackSlomo : public UInterpTrackFloatBase
 	GENERATED_UCLASS_BODY()
 
 
-	// Begin InterpTrack interface.
+	//~ Begin InterpTrack Interface.
 	virtual int32 AddKeyframe(float Time, UInterpTrackInst* TrInst, EInterpCurveMode InitInterpMode) override;
 	virtual void UpdateKeyframe(int32 KeyIndex, UInterpTrackInst* TrInst) override;
 	virtual void PreviewUpdateTrack(float NewPosition, UInterpTrackInst* TrInst) override;
@@ -19,7 +19,7 @@ class UInterpTrackSlomo : public UInterpTrackFloatBase
 #if WITH_EDITORONLY_DATA
 	virtual class UTexture2D* GetTrackIcon() const override;
 #endif // WITH_EDITORONLY_DATA
-	// End InterpTrack interface.
+	//~ End InterpTrack Interface.
 
 	/** @return the slomo factor we want at the given time. */
 	ENGINE_API float GetSlomoFactorAtTime(float Time);
diff --git a/Engine/Source/Runtime/Engine/Classes/Matinee/InterpTrackSound.h b/Engine/Source/Runtime/Engine/Classes/Matinee/InterpTrackSound.h
--- a/Engine/Source/Runtime/Engine/Classes/Matinee/InterpTrackSound.h
+++ b/Engine/Source/Runtime/Engine/Classes/Matinee/InterpTrackSound.h
@@ -70,11 +70,11 @@ class UInterpTrackSound : public UInterpTrackVectorBase
 	// True if the sound should have been playing at any point
 	uint32 bPlaying:1;
 
-	// Begin UObject Interface
+	//~ Begin UObject Interface
 	virtual void PostLoad() override;
-	// End UObject Interface
+	//~ End UObject Interface
 
-	// Begin UInterpTrack interface
+	//~ Begin UInterpTrack Interface
 	virtual int32 GetNumKeyframes() const override;
 	virtual void GetTimeRange(float& StartTime, float& EndTime) const override;
 	virtual float GetTrackEndTime() const override;
@@ -96,7 +96,7 @@ class UInterpTrackSound : public UInterpTrackVectorBase
 	virtual void DrawTrack( FCanvas* Canvas, UInterpGroup* Group, const FInterpTrackDrawParams& Params ) override;
 	virtual bool AllowStaticActors() override { return true; }
 	virtual void SetTrackToSensibleDefault() override;
-	// End UInterpTrack interface
+	//~ End UInterpTrack Interface
 
 	/**
 	 * Return the key at the specified position in the track.
diff --git a/Engine/Source/Runtime/Engine/Classes/Matinee/InterpTrackToggle.h b/Engine/Source/Runtime/Engine/Classes/Matinee/InterpTrackToggle.h
--- a/Engine/Source/Runtime/Engine/Classes/Matinee/InterpTrackToggle.h
+++ b/Engine/Source/Runtime/Engine/Classes/Matinee/InterpTrackToggle.h
@@ -78,7 +78,7 @@ class UInterpTrackToggle : public UInterpTrack
 	uint32 bFireEventsWhenJumpingForwards:1;
 
 
-	// Begin UInterpTrack interface.
+	//~ Begin UInterpTrack Interface.
 	virtual int32 GetNumKeyframes() const override;
 	virtual void GetTimeRange(float& StartTime, float& EndTime) const override;
 	virtual float GetTrackEndTime() const override;
@@ -98,7 +98,7 @@ class UInterpTrackToggle : public UInterpTrack
 #endif // WITH_EDITORONLY_DATA
 	virtual bool AllowStaticActors() override { return true; }
 	virtual void DrawTrack( FCanvas* Canvas, UInterpGroup* Group, const FInterpTrackDrawParams& Params ) override;
-	// Begin UInterpTrack interface.
+	//~ Begin UInterpTrack Interface.
 };
 
 
diff --git a/Engine/Source/Runtime/Engine/Classes/Matinee/InterpTrackVectorBase.h b/Engine/Source/Runtime/Engine/Classes/Matinee/InterpTrackVectorBase.h
--- a/Engine/Source/Runtime/Engine/Classes/Matinee/InterpTrackVectorBase.h
+++ b/Engine/Source/Runtime/Engine/Classes/Matinee/InterpTrackVectorBase.h
@@ -18,13 +18,13 @@ class ENGINE_API UInterpTrackVectorBase : public UInterpTrack
 	float CurveTension;
 
 
-	// Begin UObject interface.
+	//~ Begin UObject Interface.
 #if WITH_EDITOR
 	virtual void PostEditChangeProperty(FPropertyChangedEvent& PropertyChangedEvent) override;
 #endif // WITH_EDITOR
-	// End UObject interface.
+	//~ End UObject Interface.
 
-	// Begin InterpTrack interface.
+	//~ Begin InterpTrack Interface.
 	virtual int32 GetNumKeyframes() const override;
 	virtual void GetTimeRange(float& StartTime, float& EndTime) const override;
 	virtual float GetTrackEndTime() const override;
@@ -35,9 +35,9 @@ class ENGINE_API UInterpTrackVectorBase : public UInterpTrack
 	virtual int32 DuplicateKeyframe(int32 KeyIndex, float NewKeyTime, UInterpTrack* ToTrack = NULL) override;
 	virtual bool GetClosestSnapPosition(float InPosition, TArray<int32> &IgnoreKeys, float& OutPosition) override;
 	virtual FColor GetKeyframeColor(int32 KeyIndex) const override;
-	// End InterpTrack interface.
+	//~ End InterpTrack Interface.
 
-	// Begin FCurveEdInterface interface.
+	//~ Begin FCurveEdInterface Interface.
 	virtual int32		GetNumKeys() const override;
 	virtual int32		GetNumSubCurves() const override;
 	virtual FColor	GetSubCurveButtonColor(int32 SubCurveIndex, bool bIsSubCurveHidden) const override;
@@ -55,7 +55,7 @@ class ENGINE_API UInterpTrackVectorBase : public UInterpTrack
 	virtual void	SetKeyOut(int32 SubIndex, int32 KeyIndex, float NewOutVal) override;
 	virtual void	SetKeyInterpMode(int32 KeyIndex, EInterpCurveMode NewMode) override;
 	virtual void	SetTangents(int32 SubIndex, int32 KeyIndex, float ArriveTangent, float LeaveTangent) override;
-	// End FCurveEdInterface interface.
+	//~ End FCurveEdInterface Interface.
 };
 
 
diff --git a/Engine/Source/Runtime/Engine/Classes/Matinee/MatineeActor.h b/Engine/Source/Runtime/Engine/Classes/Matinee/MatineeActor.h
--- a/Engine/Source/Runtime/Engine/Classes/Matinee/MatineeActor.h
+++ b/Engine/Source/Runtime/Engine/Classes/Matinee/MatineeActor.h
@@ -313,7 +313,7 @@ public:
 	};
 #endif
 
-	// Begin AActor Interface
+	//~ Begin AActor Interface
 	virtual float GetNetPriority(const FVector& ViewPos, const FVector& ViewDir, AActor* Viewer, AActor* ViewTarget, UActorChannel* InChannel, float Time, bool bLowBandwidth) override;
 	virtual void Tick(float DeltaTime) override;
 	virtual void PreNetReceive() override;
@@ -325,14 +325,14 @@ public:
 #if WITH_EDITOR
 	virtual bool GetReferencedContentObjects( TArray<UObject*>& Objects ) const override;
 #endif
-	// Begin AActor Interface
+	//~ Begin AActor Interface
 
-	// Begin UObject Interface
+	//~ Begin UObject Interface
 #if WITH_EDITOR
 	virtual void PostEditChangeProperty( FPropertyChangedEvent& PropertyChangedEvent ) override; 
 	virtual bool CanEditChange( const UProperty* Property ) const override;
 #endif // WITH_EDITOR
-	// End UObject Interface
+	//~ End UObject Interface
 
 	/** Increment track forwards by given timestep and iterate over each track updating any properties. */
 	ENGINE_API void StepInterp(float DeltaTime, bool bPreview=false);
diff --git a/Engine/Source/Runtime/Engine/Classes/Particles/Acceleration/ParticleModuleAccelerationBase.h b/Engine/Source/Runtime/Engine/Classes/Particles/Acceleration/ParticleModuleAccelerationBase.h
--- a/Engine/Source/Runtime/Engine/Classes/Particles/Acceleration/ParticleModuleAccelerationBase.h
+++ b/Engine/Source/Runtime/Engine/Classes/Particles/Acceleration/ParticleModuleAccelerationBase.h
@@ -17,8 +17,8 @@ class UParticleModuleAccelerationBase : public UParticleModule
 	uint32 bAlwaysInWorldSpace:1;
 
 
-	// Begin UParticleModule Interface
+	//~ Begin UParticleModule Interface
 	virtual void SetToSensibleDefaults(UParticleEmitter* Owner) override;
-	// End UParticleModule Interface
+	//~ End UParticleModule Interface
 };
 
diff --git a/Engine/Source/Runtime/Engine/Classes/Particles/Emitter.h b/Engine/Source/Runtime/Engine/Classes/Particles/Emitter.h
--- a/Engine/Source/Runtime/Engine/Classes/Particles/Emitter.h
+++ b/Engine/Source/Runtime/Engine/Classes/Particles/Emitter.h
@@ -1,8 +1,8 @@
 // Copyright 1998-2015 Epic Games, Inc. All Rights Reserved.
 
-//=============================================================================
+//~=============================================================================
 // Emitter actor class.
-//=============================================================================
+//~=============================================================================
 
 #pragma once
 #include "Emitter.generated.h"
@@ -124,18 +124,18 @@ public:
 
 	void AutoPopulateInstanceProperties();
 
-	// Begin UObject Interface
+	//~ Begin UObject Interface
 	virtual FString GetDetailedInfoInternal() const override;
-	// End UObject Interface
+	//~ End UObject Interface
 
 
-	// Begin AActor interface
+	//~ Begin AActor Interface
 	virtual void PostActorCreated() override;
 	virtual void PostInitializeComponents() override;
 #if WITH_EDITOR
 	virtual void CheckForErrors() override;
 	virtual bool GetReferencedContentObjects( TArray<UObject*>& Objects ) const override;
-	// End AActor interface
+	//~ End AActor Interface
 
 	/**
 	 *	Called to reset the emitter actor in the level.
diff --git a/Engine/Source/Runtime/Engine/Classes/Particles/Event/ParticleModuleEventBase.h b/Engine/Source/Runtime/Engine/Classes/Particles/Event/ParticleModuleEventBase.h
--- a/Engine/Source/Runtime/Engine/Classes/Particles/Event/ParticleModuleEventBase.h
+++ b/Engine/Source/Runtime/Engine/Classes/Particles/Event/ParticleModuleEventBase.h
@@ -11,8 +11,8 @@ class UParticleModuleEventBase : public UParticleModule
 	GENERATED_UCLASS_BODY()
 
 
-	// Begin UParticleModule Interface
+	//~ Begin UParticleModule Interface
 	virtual EModuleType	GetModuleType() const override {	return EPMT_Event;	}
-	// End UParticleModule Interface
+	//~ End UParticleModule Interface
 };
 
diff --git a/Engine/Source/Runtime/Engine/Classes/Particles/Event/ParticleModuleEventReceiverKillParticles.h b/Engine/Source/Runtime/Engine/Classes/Particles/Event/ParticleModuleEventReceiverKillParticles.h
--- a/Engine/Source/Runtime/Engine/Classes/Particles/Event/ParticleModuleEventReceiverKillParticles.h
+++ b/Engine/Source/Runtime/Engine/Classes/Particles/Event/ParticleModuleEventReceiverKillParticles.h
@@ -15,9 +15,9 @@ class UParticleModuleEventReceiverKillParticles : public UParticleModuleEventRec
 	uint32 bStopSpawning:1;
 
 
-	// Begin UParticleModuleEventBase Interface
+	//~ Begin UParticleModuleEventBase Interface
 	virtual bool ProcessParticleEvent(FParticleEmitterInstance* Owner, FParticleEventData& InEvent, float DeltaTime) override;
-	// End UParticleModuleEventBase Interface
+	//~ End UParticleModuleEventBase Interface
 };
 
 
diff --git a/Engine/Source/Runtime/Engine/Classes/Particles/Event/ParticleModuleEventReceiverSpawn.h b/Engine/Source/Runtime/Engine/Classes/Particles/Event/ParticleModuleEventReceiverSpawn.h
--- a/Engine/Source/Runtime/Engine/Classes/Particles/Event/ParticleModuleEventReceiverSpawn.h
+++ b/Engine/Source/Runtime/Engine/Classes/Particles/Event/ParticleModuleEventReceiverSpawn.h
@@ -53,9 +53,9 @@ class UParticleModuleEventReceiverSpawn : public UParticleModuleEventReceiverBas
 	virtual void PostInitProperties() override;
 	//End UObject Interface
 
-	// Begin UParticleModuleEventBase Interface
+	//~ Begin UParticleModuleEventBase Interface
 	virtual bool ProcessParticleEvent(FParticleEmitterInstance* Owner, FParticleEventData& InEvent, float DeltaTime) override;
-	// End UParticleModuleEventBase Interface
+	//~ End UParticleModuleEventBase Interface
 };
 
 
diff --git a/Engine/Source/Runtime/Engine/Classes/Particles/Lifetime/ParticleModuleLifetime.h b/Engine/Source/Runtime/Engine/Classes/Particles/Lifetime/ParticleModuleLifetime.h
--- a/Engine/Source/Runtime/Engine/Classes/Particles/Lifetime/ParticleModuleLifetime.h
+++ b/Engine/Source/Runtime/Engine/Classes/Particles/Lifetime/ParticleModuleLifetime.h
@@ -30,10 +30,10 @@ class UParticleModuleLifetime : public UParticleModuleLifetimeBase
 	virtual void SetToSensibleDefaults(UParticleEmitter* Owner) override;
 	//End UParticleModule Interface
 
-	// Begin UParticleModuleLifetimeBase Interface
+	//~ Begin UParticleModuleLifetimeBase Interface
 	virtual float GetMaxLifetime() override;
 	virtual float GetLifetimeValue(FParticleEmitterInstance* Owner, float InTime, UObject* Data = NULL) override;
-	// End UParticleModuleLifetimeBase Interface
+	//~ End UParticleModuleLifetimeBase Interface
 
 	/**
 	 *	Extended version of spawn, allows for using a random stream for distribution value retrieval
diff --git a/Engine/Source/Runtime/Engine/Classes/Particles/Lifetime/ParticleModuleLifetime_Seeded.h b/Engine/Source/Runtime/Engine/Classes/Particles/Lifetime/ParticleModuleLifetime_Seeded.h
--- a/Engine/Source/Runtime/Engine/Classes/Particles/Lifetime/ParticleModuleLifetime_Seeded.h
+++ b/Engine/Source/Runtime/Engine/Classes/Particles/Lifetime/ParticleModuleLifetime_Seeded.h
@@ -26,9 +26,9 @@ class UParticleModuleLifetime_Seeded : public UParticleModuleLifetime
 	virtual void EmitterLoopingNotify(FParticleEmitterInstance* Owner) override;
 	//End UParticleModule Interface
 
-	// Begin UParticleModuleLifetimeBase Interface
+	//~ Begin UParticleModuleLifetimeBase Interface
 	virtual float	GetLifetimeValue(FParticleEmitterInstance* Owner, float InTime, UObject* Data = NULL) override;
-	// End UParticleModuleLifetimeBase Interface
+	//~ End UParticleModuleLifetimeBase Interface
 };
 
 
diff --git a/Engine/Source/Runtime/Engine/Classes/Particles/Location/ParticleModuleLocationDirect.h b/Engine/Source/Runtime/Engine/Classes/Particles/Location/ParticleModuleLocationDirect.h
--- a/Engine/Source/Runtime/Engine/Classes/Particles/Location/ParticleModuleLocationDirect.h
+++ b/Engine/Source/Runtime/Engine/Classes/Particles/Location/ParticleModuleLocationDirect.h
@@ -44,18 +44,18 @@ class UParticleModuleLocationDirect : public UParticleModuleLocationBase
 	/** Initializes the default values for this property */
 	void InitializeDefaults();
 
-	// Begin UObject Interface
+	//~ Begin UObject Interface
 #if WITH_EDITOR
 	virtual void PostEditChangeProperty(FPropertyChangedEvent& PropertyChangedEvent) override;
 #endif // WITH_EDITOR
 	virtual void	PostInitProperties() override;
-	// End UObject Interface
+	//~ End UObject Interface
 
-	// Begin UParticleModule Interface
+	//~ Begin UParticleModule Interface
 	virtual void	Spawn(FParticleEmitterInstance* Owner, int32 Offset, float SpawnTime, FBaseParticle* ParticleBase) override;
 	virtual void	Update(FParticleEmitterInstance* Owner, int32 Offset, float DeltaTime) override;
 	virtual uint32	RequiredBytes(FParticleEmitterInstance* Owner = NULL) override;
-	// End UParticleModule Interface
+	//~ End UParticleModule Interface
 };
 
 
diff --git a/Engine/Source/Runtime/Engine/Classes/Particles/Location/ParticleModuleLocationEmitter.h b/Engine/Source/Runtime/Engine/Classes/Particles/Location/ParticleModuleLocationEmitter.h
--- a/Engine/Source/Runtime/Engine/Classes/Particles/Location/ParticleModuleLocationEmitter.h
+++ b/Engine/Source/Runtime/Engine/Classes/Particles/Location/ParticleModuleLocationEmitter.h
@@ -1,12 +1,12 @@
 // Copyright 1998-2015 Epic Games, Inc. All Rights Reserved.
 
-//=============================================================================
+//~=============================================================================
 // ParticleModuleLocationEmitter
 //
 // A location module that uses particles from another emitters particles as
 // spawn points for its particles.
 //
-//=============================================================================
+//~=============================================================================
 
 #pragma once
 #include "Particles/Location/ParticleModuleLocationBase.h"
@@ -25,9 +25,9 @@ class UParticleModuleLocationEmitter : public UParticleModuleLocationBase
 {
 	GENERATED_UCLASS_BODY()
 
-	//=============================================================================
+	//~=============================================================================
 	// Variables
-	//=============================================================================
+	//~=============================================================================
 	// LocationEmitter
 	
 	/** The name of the emitter to use that the source location for particle. */
@@ -60,10 +60,10 @@ class UParticleModuleLocationEmitter : public UParticleModuleLocationBase
 	float InheritSourceRotationScale;
 
 
-	// Begin UParticleModule Interface
+	//~ Begin UParticleModule Interface
 	virtual void	Spawn(FParticleEmitterInstance* Owner, int32 Offset, float SpawnTime, FBaseParticle* ParticleBase) override;
 	virtual uint32	RequiredBytesPerInstance(FParticleEmitterInstance* Owner = NULL) override;
-	// End UParticleModule Interface
+	//~ End UParticleModule Interface
 };
 
 
diff --git a/Engine/Source/Runtime/Engine/Classes/Particles/Location/ParticleModuleLocationEmitterDirect.h b/Engine/Source/Runtime/Engine/Classes/Particles/Location/ParticleModuleLocationEmitterDirect.h
--- a/Engine/Source/Runtime/Engine/Classes/Particles/Location/ParticleModuleLocationEmitterDirect.h
+++ b/Engine/Source/Runtime/Engine/Classes/Particles/Location/ParticleModuleLocationEmitterDirect.h
@@ -1,12 +1,12 @@
 // Copyright 1998-2015 Epic Games, Inc. All Rights Reserved.
 
-//=============================================================================
+//~=============================================================================
 // ParticleModuleLocationEmitterDirect
 //
 // A location module that uses particles from another emitters particles as
 // position for its particles.
 //
-//=============================================================================
+//~=============================================================================
 
 #pragma once
 #include "Particles/Location/ParticleModuleLocationBase.h"
@@ -22,10 +22,10 @@ class UParticleModuleLocationEmitterDirect : public UParticleModuleLocationBase
 	FName EmitterName;
 
 
-	// Begin UParticleModule Interface
+	//~ Begin UParticleModule Interface
 	virtual void	Spawn(FParticleEmitterInstance* Owner, int32 Offset, float SpawnTime, FBaseParticle* ParticleBase) override;
 	virtual void	Update(FParticleEmitterInstance* Owner, int32 Offset, float DeltaTime) override;
-	// End UParticleModule Interface
+	//~ End UParticleModule Interface
 
 };
 
diff --git a/Engine/Source/Runtime/Engine/Classes/Particles/Location/ParticleModuleLocationPrimitiveCylinder.h b/Engine/Source/Runtime/Engine/Classes/Particles/Location/ParticleModuleLocationPrimitiveCylinder.h
--- a/Engine/Source/Runtime/Engine/Classes/Particles/Location/ParticleModuleLocationPrimitiveCylinder.h
+++ b/Engine/Source/Runtime/Engine/Classes/Particles/Location/ParticleModuleLocationPrimitiveCylinder.h
@@ -1,9 +1,9 @@
 // Copyright 1998-2015 Epic Games, Inc. All Rights Reserved.
 
-//=============================================================================
+//~=============================================================================
 // ParticleModuleLocationPrimitiveCylinder
 // Location primitive spawning within a cylinder.
-//=============================================================================
+//~=============================================================================
 
 #pragma once
 #include "Particles/Location/ParticleModuleLocationPrimitiveBase.h"
@@ -48,17 +48,17 @@ class UParticleModuleLocationPrimitiveCylinder : public UParticleModuleLocationP
 	/** Initializes the default values for this property */
 	void InitializeDefaults();
 
-	// Begin UObject Interface
+	//~ Begin UObject Interface
 #if WITH_EDITOR
 	virtual void PostEditChangeProperty(FPropertyChangedEvent& PropertyChangedEvent) override;
 #endif // WITH_EDITOR
 	virtual void PostInitProperties() override;
-	// End UObject Interface
+	//~ End UObject Interface
 
-	// Begin UParticleModule Interface
+	//~ Begin UParticleModule Interface
 	virtual void	Spawn(FParticleEmitterInstance* Owner, int32 Offset, float SpawnTime, FBaseParticle* ParticleBase) override;
 	virtual void	Render3DPreview(FParticleEmitterInstance* Owner, const FSceneView* View,FPrimitiveDrawInterface* PDI) override;
-	// End UParticleModule Interface
+	//~ End UParticleModule Interface
 
 	/**
 	 *	Extended version of spawn, allows for using a random stream for distribution value retrieval
diff --git a/Engine/Source/Runtime/Engine/Classes/Particles/Location/ParticleModuleLocationPrimitiveSphere.h b/Engine/Source/Runtime/Engine/Classes/Particles/Location/ParticleModuleLocationPrimitiveSphere.h
--- a/Engine/Source/Runtime/Engine/Classes/Particles/Location/ParticleModuleLocationPrimitiveSphere.h
+++ b/Engine/Source/Runtime/Engine/Classes/Particles/Location/ParticleModuleLocationPrimitiveSphere.h
@@ -1,9 +1,9 @@
 // Copyright 1998-2015 Epic Games, Inc. All Rights Reserved.
 
-//=============================================================================
+//~=============================================================================
 // ParticleModuleLocationPrimitiveSphere
 // Location primitive spawning within a Sphere.
-//=============================================================================
+//~=============================================================================
 
 #pragma once
 #include "Particles/Location/ParticleModuleLocationPrimitiveBase.h"
@@ -21,17 +21,17 @@ class UParticleModuleLocationPrimitiveSphere : public UParticleModuleLocationPri
 	/** Initializes the default values for this property */
 	void InitializeDefaults();
 
-	// Begin UObject Interface
+	//~ Begin UObject Interface
 #if WITH_EDITOR
 	virtual void PostEditChangeProperty(FPropertyChangedEvent& PropertyChangedEvent) override;
 #endif // WITH_EDITOR
 	virtual void PostInitProperties() override;
-	// End UObject Interface
+	//~ End UObject Interface
 
-	// Begin UParticleModule Interface
+	//~ Begin UParticleModule Interface
 	virtual void	Spawn(FParticleEmitterInstance* Owner, int32 Offset, float SpawnTime, FBaseParticle* ParticleBase) override;
 	virtual void	Render3DPreview(FParticleEmitterInstance* Owner, const FSceneView* View,FPrimitiveDrawInterface* PDI) override;
-	// End UParticleModule Interface
+	//~ End UParticleModule Interface
 
 	/**
 	 *	Extended version of spawn, allows for using a random stream for distribution value retrieval
diff --git a/Engine/Source/Runtime/Engine/Classes/Particles/Orientation/ParticleModuleOrientationAxisLock.h b/Engine/Source/Runtime/Engine/Classes/Particles/Orientation/ParticleModuleOrientationAxisLock.h
--- a/Engine/Source/Runtime/Engine/Classes/Particles/Orientation/ParticleModuleOrientationAxisLock.h
+++ b/Engine/Source/Runtime/Engine/Classes/Particles/Orientation/ParticleModuleOrientationAxisLock.h
@@ -55,16 +55,16 @@ class UParticleModuleOrientationAxisLock : public UParticleModuleOrientationBase
 	TEnumAsByte<EParticleAxisLock> LockAxisFlags;
 
 
-	// Begin UObject Interface
+	//~ Begin UObject Interface
 #if WITH_EDITOR
 	virtual void	PostEditChangeProperty(FPropertyChangedEvent& PropertyChangedEvent) override;
 #endif // WITH_EDITOR
-	// End UObject Interface
+	//~ End UObject Interface
 
-	// Begin UParticleModule Interface
+	//~ Begin UParticleModule Interface
 	virtual void	Spawn(FParticleEmitterInstance* Owner, int32 Offset, float SpawnTime, FBaseParticle* ParticleBase) override;
 	virtual void	Update(FParticleEmitterInstance* Owner, int32 Offset, float DeltaTime) override;
-	// End UParticleModule Interface
+	//~ End UParticleModule Interface
 
 	//@todo document
 	virtual void	SetLockAxis(EParticleAxisLock eLockFlags);
diff --git a/Engine/Source/Runtime/Engine/Classes/Particles/Parameter/ParticleModuleParameterDynamic.h b/Engine/Source/Runtime/Engine/Classes/Particles/Parameter/ParticleModuleParameterDynamic.h
--- a/Engine/Source/Runtime/Engine/Classes/Particles/Parameter/ParticleModuleParameterDynamic.h
+++ b/Engine/Source/Runtime/Engine/Classes/Particles/Parameter/ParticleModuleParameterDynamic.h
@@ -108,9 +108,9 @@ class UParticleModuleParameterDynamic : public UParticleModuleParameterBase
 	virtual void	PostEditChangeProperty(FPropertyChangedEvent& PropertyChangedEvent) override;
 #endif // WITH_EDITOR
 	virtual void	PostInitProperties() override;
-	// Begin UObject Interface
+	//~ Begin UObject Interface
 
-	// Begin UParticleModule Interface
+	//~ Begin UParticleModule Interface
 	virtual void	Spawn(FParticleEmitterInstance* Owner, int32 Offset, float SpawnTime, FBaseParticle* ParticleBase) override;
 	virtual void	Update(FParticleEmitterInstance* Owner, int32 Offset, float DeltaTime) override;
 	virtual uint32	RequiredBytes(FParticleEmitterInstance* Owner = NULL) override;
@@ -125,7 +125,7 @@ class UParticleModuleParameterDynamic : public UParticleModuleParameterBase
 	virtual void GetParticleSysParamsUtilized(TArray<FString>& ParticleSysParamList) override;
 	virtual void GetParticleParametersUtilized(TArray<FString>& ParticleParameterList) override;
 	virtual void RefreshModule(UInterpCurveEdSetup* EdSetup, UParticleEmitter* InEmitter, int32 InLODLevel) override;
-	// End UParticleModule Interface
+	//~ End UParticleModule Interface
 
 	/**
 	 *	Extended version of spawn, allows for using a random stream for distribution value retrieval
diff --git a/Engine/Source/Runtime/Engine/Classes/Particles/ParticleEmitter.h b/Engine/Source/Runtime/Engine/Classes/Particles/ParticleEmitter.h
--- a/Engine/Source/Runtime/Engine/Classes/Particles/ParticleEmitter.h
+++ b/Engine/Source/Runtime/Engine/Classes/Particles/ParticleEmitter.h
@@ -1,9 +1,9 @@
 // Copyright 1998-2015 Epic Games, Inc. All Rights Reserved.
 
-//=============================================================================
+//~=============================================================================
 // ParticleEmitter
 // The base class for any particle emitter objects.
-//=============================================================================
+//~=============================================================================
 
 #pragma once
 #include "Components/SceneComponent.h"
@@ -13,9 +13,9 @@
 
 class UInterpCurveEdSetup;
 
-//=============================================================================
+//~=============================================================================
 //	Burst emissions
-//=============================================================================
+//~=============================================================================
 UENUM()
 enum EParticleBurstMethod
 {
@@ -24,9 +24,9 @@ enum EParticleBurstMethod
 	EPBM_MAX,
 };
 
-//=============================================================================
+//~=============================================================================
 //	SubUV-related
-//=============================================================================
+//~=============================================================================
 UENUM()
 enum EParticleSubUVInterpMethod
 {
@@ -38,9 +38,9 @@ enum EParticleSubUVInterpMethod
 	PSUVIM_MAX,
 };
 
-//=============================================================================
+//~=============================================================================
 //	Cascade-related
-//=============================================================================
+//~=============================================================================
 UENUM()
 enum EEmitterRenderMode
 {
@@ -85,9 +85,9 @@ class UParticleEmitter : public UObject
 {
 	GENERATED_UCLASS_BODY()
 
-	//=============================================================================
+	//~=============================================================================
 	//	General variables
-	//=============================================================================
+	//~=============================================================================
 	/** The name of the emitter. */
 	UPROPERTY(EditAnywhere, Category=Particle)
 	FName EmitterName;
@@ -113,9 +113,9 @@ class UParticleEmitter : public UObject
 	FColor EmitterEditorColor;
 
 #endif // WITH_EDITORONLY_DATA
-	//=============================================================================
+	//~=============================================================================
 	//	'Private' data - not required by the editor
-	//=============================================================================
+	//~=============================================================================
 	UPROPERTY(instanced)
 	TArray<class UParticleLODLevel*> LODLevels;
 
@@ -125,9 +125,9 @@ class UParticleEmitter : public UObject
 	UPROPERTY()
 	int32 PeakActiveParticles;
 
-	//=============================================================================
+	//~=============================================================================
 	//	Performance/LOD Data
-	//=============================================================================
+	//~=============================================================================
 	
 	/**
 	 *	Initial allocation count - overrides calculated peak count if > 0
@@ -173,12 +173,12 @@ class UParticleEmitter : public UObject
 	UPROPERTY(EditAnywhere, Category = Particle)
 	uint32 bDisabledLODsKeepEmitterAlive : 1;
 
-	// Begin UObject Interface
+	//~ Begin UObject Interface
 #if WITH_EDITOR
 	virtual void PostEditChangeProperty(FPropertyChangedEvent& PropertyChangedEvent) override;
 #endif // WITH_EDITOR
 	virtual void PostLoad() override;
-	// End UObject Interface
+	//~ End UObject Interface
 
 	// @todo document
 	virtual FParticleEmitterInstance* CreateInstance(UParticleSystemComponent* InComponent);
diff --git a/Engine/Source/Runtime/Engine/Classes/Particles/ParticleSpriteEmitter.h b/Engine/Source/Runtime/Engine/Classes/Particles/ParticleSpriteEmitter.h
--- a/Engine/Source/Runtime/Engine/Classes/Particles/ParticleSpriteEmitter.h
+++ b/Engine/Source/Runtime/Engine/Classes/Particles/ParticleSpriteEmitter.h
@@ -23,13 +23,13 @@ class UParticleSpriteEmitter : public UParticleEmitter
 	GENERATED_UCLASS_BODY()
 
 
-	// Begin UObject Interface
+	//~ Begin UObject Interface
 	virtual void PostLoad() override;
-	// End UObject Interface
+	//~ End UObject Interface
 
-	// Begin UParticleEmitter Interface
+	//~ Begin UParticleEmitter Interface
 	virtual FParticleEmitterInstance* CreateInstance(UParticleSystemComponent* InComponent) override;
 	virtual void SetToSensibleDefaults() override;
-	// End UParticleEmitter Interface
+	//~ End UParticleEmitter Interface
 };
 
diff --git a/Engine/Source/Runtime/Engine/Classes/Particles/ParticleSystemComponent.h b/Engine/Source/Runtime/Engine/Classes/Particles/ParticleSystemComponent.h
--- a/Engine/Source/Runtime/Engine/Classes/Particles/ParticleSystemComponent.h
+++ b/Engine/Source/Runtime/Engine/Classes/Particles/ParticleSystemComponent.h
@@ -784,7 +784,7 @@ private:
 	bool bNeedsFinalize;
 public:
 
-	// Begin UActorComponent interface.
+	//~ Begin UActorComponent Interface.
 #if WITH_EDITOR
 	virtual void CheckForErrors() override;
 #endif
@@ -801,7 +801,7 @@ protected:
 	virtual void OnUnregister()  override;
 	virtual void SendRenderDynamicData_Concurrent() override;
 public:
-	// End UActorComponent interface.
+	//~ End UActorComponent Interface.
 
 	enum EForceAsyncWorkCompletion
 	{
@@ -839,7 +839,7 @@ private:
 
 public:
 
-	// Begin UObject interface.
+	//~ Begin UObject Interface.
 	virtual void PostLoad() override;
 	virtual void BeginDestroy() override;
 	virtual void FinishDestroy() override;
@@ -850,7 +850,7 @@ public:
 	virtual void Serialize(FArchive& Ar) override;
 	virtual SIZE_T GetResourceSize(EResourceSizeMode::Type Mode) override;
 	virtual FString GetDetailedInfoInternal() const override;
-	// End UObject interface.
+	//~ End UObject Interface.
 
 	//Begin UPrimitiveComponent Interface
 	virtual int32 GetNumMaterials() const override; 
@@ -861,7 +861,7 @@ public:
 	virtual void GetUsedMaterials( TArray<UMaterialInterface*>& OutMaterials ) const override;
 	//End UPrimitiveComponent Interface
 
-	// Begin USceneComonent Interface
+	//~ Begin USceneComonent Interface
 protected:
 	virtual bool ShouldActivate() const override;
 
@@ -869,7 +869,7 @@ public:
 	virtual void Activate(bool bReset=false) override;
 	virtual void Deactivate() override;
 	virtual void ApplyWorldOffset(const FVector& InOffset, bool bWorldShift) override;
-	// End USceneComponent Interface
+	//~ End USceneComponent Interface
 
 	/** Activate the system */
 	void ActivateSystem(bool bFlagAsJustAttached = false);
diff --git a/Engine/Source/Runtime/Engine/Classes/Particles/RotationRate/ParticleModuleRotationRate.h b/Engine/Source/Runtime/Engine/Classes/Particles/RotationRate/ParticleModuleRotationRate.h
--- a/Engine/Source/Runtime/Engine/Classes/Particles/RotationRate/ParticleModuleRotationRate.h
+++ b/Engine/Source/Runtime/Engine/Classes/Particles/RotationRate/ParticleModuleRotationRate.h
@@ -27,11 +27,11 @@ class UParticleModuleRotationRate : public UParticleModuleRotationRateBase
 	virtual void	PostInitProperties() override;
 	//End UObject Interface
 
-	// Begin UParticleModule Interface
+	//~ Begin UParticleModule Interface
 	virtual void CompileModule( struct FParticleEmitterBuildInfo& EmitterInfo ) override;
 	virtual void Spawn(FParticleEmitterInstance* Owner, int32 Offset, float SpawnTime, FBaseParticle* ParticleBase) override;
 	virtual void SetToSensibleDefaults(UParticleEmitter* Owner) override;
-	// Begin UParticleModule Interface
+	//~ Begin UParticleModule Interface
 
 	/**
 	 *	Extended version of spawn, allows for using a random stream for distribution value retrieval
diff --git a/Engine/Source/Runtime/Engine/Classes/Particles/RotationRate/ParticleModuleRotationRateMultiplyLife.h b/Engine/Source/Runtime/Engine/Classes/Particles/RotationRate/ParticleModuleRotationRateMultiplyLife.h
--- a/Engine/Source/Runtime/Engine/Classes/Particles/RotationRate/ParticleModuleRotationRateMultiplyLife.h
+++ b/Engine/Source/Runtime/Engine/Classes/Particles/RotationRate/ParticleModuleRotationRateMultiplyLife.h
@@ -27,7 +27,7 @@ class UParticleModuleRotationRateMultiplyLife : public UParticleModuleRotationRa
 	virtual void PostInitProperties() override;
 	//End UObject Interface
 
-	// Begin UParticleModule Interface
+	//~ Begin UParticleModule Interface
 	virtual void Spawn(FParticleEmitterInstance* Owner, int32 Offset, float SpawnTime, FBaseParticle* ParticleBase) override;
 	virtual void Update(FParticleEmitterInstance* Owner, int32 Offset, float DeltaTime) override;
 	virtual void SetToSensibleDefaults(UParticleEmitter* Owner) override;
diff --git a/Engine/Source/Runtime/Engine/Classes/Particles/Size/ParticleModuleSize.h b/Engine/Source/Runtime/Engine/Classes/Particles/Size/ParticleModuleSize.h
--- a/Engine/Source/Runtime/Engine/Classes/Particles/Size/ParticleModuleSize.h
+++ b/Engine/Source/Runtime/Engine/Classes/Particles/Size/ParticleModuleSize.h
@@ -26,12 +26,12 @@ class UParticleModuleSize : public UParticleModuleSizeBase
 	virtual void PostEditChangeProperty(FPropertyChangedEvent& PropertyChangedEvent) override;
 #endif // WITH_EDITOR
 	virtual void PostInitProperties() override;
-	// End UObject Interface
+	//~ End UObject Interface
 
-	// Begin UParticleModule Interface
+	//~ Begin UParticleModule Interface
 	virtual void CompileModule( struct FParticleEmitterBuildInfo& EmitterInfo ) override;
 	virtual void Spawn(FParticleEmitterInstance* Owner, int32 Offset, float SpawnTime, FBaseParticle* ParticleBase) override;
-	// End UParticleModule Interface
+	//~ End UParticleModule Interface
 
 	/**
 	 *	Extended version of spawn, allows for using a random stream for distribution value retrieval
diff --git a/Engine/Source/Runtime/Engine/Classes/Particles/Size/ParticleModuleSizeMultiplyLife.h b/Engine/Source/Runtime/Engine/Classes/Particles/Size/ParticleModuleSizeMultiplyLife.h
--- a/Engine/Source/Runtime/Engine/Classes/Particles/Size/ParticleModuleSizeMultiplyLife.h
+++ b/Engine/Source/Runtime/Engine/Classes/Particles/Size/ParticleModuleSizeMultiplyLife.h
@@ -46,9 +46,9 @@ class UParticleModuleSizeMultiplyLife : public UParticleModuleSizeBase
 	virtual void PostEditChangeProperty(FPropertyChangedEvent& PropertyChangedEvent) override;
 #endif // WITH_EDITOR
 	virtual void PostInitProperties() override;
-	// End UObject Interface
+	//~ End UObject Interface
 
-	// Begin UParticleModule Interface
+	//~ Begin UParticleModule Interface
 	virtual void CompileModule( struct FParticleEmitterBuildInfo& EmitterInfo ) override;
 	virtual void Spawn(FParticleEmitterInstance* Owner, int32 Offset, float SpawnTime, FBaseParticle* ParticleBase) override;
 	virtual void	Update(FParticleEmitterInstance* Owner, int32 Offset, float DeltaTime) override;
@@ -59,7 +59,7 @@ class UParticleModuleSizeMultiplyLife : public UParticleModuleSizeBase
 	virtual bool IsValidForLODLevel(UParticleLODLevel* LODLevel, FString& OutErrorString) override;
 #endif
 
-	// End UParticleModule Interface
+	//~ End UParticleModule Interface
 
 protected:
 	friend class FParticleModuleSizeMultiplyLifeDetails;
diff --git a/Engine/Source/Runtime/Engine/Classes/Particles/Size/ParticleModuleSizeScaleBySpeed.h b/Engine/Source/Runtime/Engine/Classes/Particles/Size/ParticleModuleSizeScaleBySpeed.h
--- a/Engine/Source/Runtime/Engine/Classes/Particles/Size/ParticleModuleSizeScaleBySpeed.h
+++ b/Engine/Source/Runtime/Engine/Classes/Particles/Size/ParticleModuleSizeScaleBySpeed.h
@@ -22,7 +22,7 @@ class UParticleModuleSizeScaleBySpeed : public UParticleModuleSizeBase
 	FVector2D MaxScale;
 
 
-	// Begin UParticleModule Interface
+	//~ Begin UParticleModule Interface
 	virtual void Update(FParticleEmitterInstance* Owner, int32 Offset, float DeltaTime) override;
 	virtual void CompileModule(struct FParticleEmitterBuildInfo& EmitterInfo) override;
 	// End UParticleModule Interface
diff --git a/Engine/Source/Runtime/Engine/Classes/Particles/Spawn/ParticleModuleSpawn.h b/Engine/Source/Runtime/Engine/Classes/Particles/Spawn/ParticleModuleSpawn.h
--- a/Engine/Source/Runtime/Engine/Classes/Particles/Spawn/ParticleModuleSpawn.h
+++ b/Engine/Source/Runtime/Engine/Classes/Particles/Spawn/ParticleModuleSpawn.h
@@ -40,25 +40,25 @@ class UParticleModuleSpawn : public UParticleModuleSpawnBase
 	/** Initializes the default values for this property */
 	void InitializeDefaults();
 
-	// Begin UObject Interface
+	//~ Begin UObject Interface
 #if WITH_EDITOR
 	virtual void	PostEditChangeProperty(FPropertyChangedEvent& PropertyChangedEvent) override;
 #endif // WITH_EDITOR
 	virtual void	PostInitProperties() override;
 	virtual void	PostLoad() override;
-	// End UObject Interface
+	//~ End UObject Interface
 
-	// Begin UParticleModule Interface
+	//~ Begin UParticleModule Interface
 	virtual bool	GenerateLODModuleValues(UParticleModule* SourceModule, float Percentage, UParticleLODLevel* LODLevel) override;
-	// End UParticleModule Interface
+	//~ End UParticleModule Interface
 
-	// Begin UParticleModuleSpawnBase Interface
+	//~ Begin UParticleModuleSpawnBase Interface
 	virtual bool GetSpawnAmount(FParticleEmitterInstance* Owner, int32 Offset, float OldLeftover, 
 		float DeltaTime, int32& Number, float& Rate) override;
 	virtual float GetMaximumSpawnRate() override;
 	virtual float GetEstimatedSpawnRate() override;
 	virtual int32 GetMaximumBurstCount() override;
-	// End UParticleModuleSpawnBase Interface
+	//~ End UParticleModuleSpawnBase Interface
 
 	float GetGlobalRateScale()const;
 };
diff --git a/Engine/Source/Runtime/Engine/Classes/Particles/Spawn/ParticleModuleSpawnPerUnit.h b/Engine/Source/Runtime/Engine/Classes/Particles/Spawn/ParticleModuleSpawnPerUnit.h
--- a/Engine/Source/Runtime/Engine/Classes/Particles/Spawn/ParticleModuleSpawnPerUnit.h
+++ b/Engine/Source/Runtime/Engine/Classes/Particles/Spawn/ParticleModuleSpawnPerUnit.h
@@ -66,22 +66,22 @@ class UParticleModuleSpawnPerUnit : public UParticleModuleSpawnBase
 	/** Initializes the default values for this property */
 	void InitializeDefaults();
 
-	// Begin UObject Interface
+	//~ Begin UObject Interface
 #if WITH_EDITOR
 	virtual void PostEditChangeProperty(FPropertyChangedEvent& PropertyChangedEvent) override;
 #endif // WITH_EDITOR
 	virtual void PostInitProperties() override;
-	// End UObject Interface
+	//~ End UObject Interface
 
-	// Begin UParticleModule Interface
+	//~ Begin UParticleModule Interface
 	virtual void CompileModule( struct FParticleEmitterBuildInfo& EmitterInfo ) override;
 	virtual uint32	RequiredBytesPerInstance(FParticleEmitterInstance* Owner = NULL) override;
-	// End UParticleModule Interface
+	//~ End UParticleModule Interface
 
-	// Begin UParticleModuleSpawnBase Interface
+	//~ Begin UParticleModuleSpawnBase Interface
 	virtual bool GetSpawnAmount(FParticleEmitterInstance* Owner, int32 Offset, float OldLeftover, 
 		float DeltaTime, int32& Number, float& Rate) override;
-	// End UParticleModuleSpawnBase Interface
+	//~ End UParticleModuleSpawnBase Interface
 };
 
 
diff --git a/Engine/Source/Runtime/Engine/Classes/Particles/SubUV/ParticleModuleSubUVMovie.h b/Engine/Source/Runtime/Engine/Classes/Particles/SubUV/ParticleModuleSubUVMovie.h
--- a/Engine/Source/Runtime/Engine/Classes/Particles/SubUV/ParticleModuleSubUVMovie.h
+++ b/Engine/Source/Runtime/Engine/Classes/Particles/SubUV/ParticleModuleSubUVMovie.h
@@ -43,17 +43,17 @@ class UParticleModuleSubUVMovie : public UParticleModuleSubUV
 	virtual void PostInitProperties() override;
 	//End UObject Interface
 
-	// Begin UParticleModule Interface
+	//~ Begin UParticleModule Interface
 	virtual void Spawn(FParticleEmitterInstance* Owner, int32 Offset, float SpawnTime, FBaseParticle* ParticleBase) override;
 	virtual uint32 RequiredBytes(FParticleEmitterInstance* Owner = NULL) override;
 	virtual void SetToSensibleDefaults(UParticleEmitter* Owner) override;
 	virtual void GetCurveObjects(TArray<FParticleCurvePair>& OutCurves) override;
-	// End UParticleModule Interface
+	//~ End UParticleModule Interface
 	
-	// Begin UParticleModuleSubUV Interface
+	//~ Begin UParticleModuleSubUV Interface
 	virtual float DetermineImageIndex(FParticleEmitterInstance* Owner, int32 Offset, FBaseParticle* Particle, 
 		EParticleSubUVInterpMethod InterpMethod, FFullSubUVPayload& SubUVPayload, float DeltaTime) override;
-	// End UParticleModuleSubUV Interface
+	//~ End UParticleModuleSubUV Interface
 	
 };
 
diff --git a/Engine/Source/Runtime/Engine/Classes/Particles/TypeData/ParticleModuleTypeDataAnimTrail.h b/Engine/Source/Runtime/Engine/Classes/Particles/TypeData/ParticleModuleTypeDataAnimTrail.h
--- a/Engine/Source/Runtime/Engine/Classes/Particles/TypeData/ParticleModuleTypeDataAnimTrail.h
+++ b/Engine/Source/Runtime/Engine/Classes/Particles/TypeData/ParticleModuleTypeDataAnimTrail.h
@@ -64,17 +64,17 @@ class UParticleModuleTypeDataAnimTrail : public UParticleModuleTypeDataBase
 	float WidthTessellationStepSize;
 
 
-	// Begin UParticleModule Interface
+	//~ Begin UParticleModule Interface
 	virtual uint32 RequiredBytes(FParticleEmitterInstance* Owner = NULL) override;
 	virtual bool CanTickInAnyThread() override
 	{
 		return false;
 	}
-	// End UParticleModule Interface
+	//~ End UParticleModule Interface
 
-	// Begin UParticleModuleTypeDataBase Interface
+	//~ Begin UParticleModuleTypeDataBase Interface
 	virtual FParticleEmitterInstance* CreateInstance(UParticleEmitter* InEmitterParent, UParticleSystemComponent* InComponent) override;
-	// End UParticleModuleTypeDataBase Interface
+	//~ End UParticleModuleTypeDataBase Interface
 };
 
 
diff --git a/Engine/Source/Runtime/Engine/Classes/Particles/TypeData/ParticleModuleTypeDataBeam2.h b/Engine/Source/Runtime/Engine/Classes/Particles/TypeData/ParticleModuleTypeDataBeam2.h
--- a/Engine/Source/Runtime/Engine/Classes/Particles/TypeData/ParticleModuleTypeDataBeam2.h
+++ b/Engine/Source/Runtime/Engine/Classes/Particles/TypeData/ParticleModuleTypeDataBeam2.h
@@ -181,15 +181,15 @@ class UParticleModuleTypeDataBeam2 : public UParticleModuleTypeDataBase
 	/** Initializes the default values for this property */
 	void InitializeDefaults();
 
-	// Begin UObject Interface
+	//~ Begin UObject Interface
 #if WITH_EDITOR
 	virtual void	PostEditChangeProperty(FPropertyChangedEvent& PropertyChangedEvent) override;
 #endif // WITH_EDITOR
 	virtual void	PostInitProperties() override;
-	// End UObject Interface
+	//~ End UObject Interface
 
 
-	// Begin UParticleModule Interface
+	//~ Begin UParticleModule Interface
 	virtual void	Spawn(FParticleEmitterInstance* Owner, int32 Offset, float SpawnTime, FBaseParticle* ParticleBase) override;
 	virtual void	Update(FParticleEmitterInstance* Owner, int32 Offset, float DeltaTime) override;
 	virtual uint32	RequiredBytes(FParticleEmitterInstance* Owner = NULL) override;
@@ -198,12 +198,12 @@ class UParticleModuleTypeDataBeam2 : public UParticleModuleTypeDataBase
 	{
 		return false;
 	}
-	// End UParticleModule Interface
+	//~ End UParticleModule Interface
 
 
-	// Begin UParticleModuleTypeDataBase Interface
+	//~ Begin UParticleModuleTypeDataBase Interface
 	virtual FParticleEmitterInstance* CreateInstance(UParticleEmitter* InEmitterParent, UParticleSystemComponent* InComponent) override;
-	// End UParticleModuleTypeDataBase Interface
+	//~ End UParticleModuleTypeDataBase Interface
 
 
 	/**
