Comment cleanup.

diff --git a/Engine/Source/Runtime/Engine/Classes/Kismet/GameplayStatics.h b/Engine/Source/Runtime/Engine/Classes/Kismet/GameplayStatics.h
--- a/Engine/Source/Runtime/Engine/Classes/Kismet/GameplayStatics.h
+++ b/Engine/Source/Runtime/Engine/Classes/Kismet/GameplayStatics.h
@@ -172,14 +172,14 @@ class ENGINE_API UGameplayStatics : public UBlueprintFunctionLibrary
 	static class UClass *GetObjectClass(const UObject *Object);
 
 	/**
-	 * Sets the global time dilation
-	 * @param	TimeDilation	value to set the global time dilation to
+	 * Gets the current global time dilation.
+	 * @return Current time dilation.
 	 */
 	UFUNCTION(BlueprintPure, Category="Utilities|Time", meta=(WorldContext="WorldContextObject") )
 	static float GetGlobalTimeDilation(UObject* WorldContextObject);
 
 	/**
-	 * Sets the global time dilation
+	 * Sets the global time dilation.
 	 * @param	TimeDilation	value to set the global time dilation to
 	 */
 	UFUNCTION(BlueprintCallable, Category="Utilities|Time", meta=(WorldContext="WorldContextObject") )
@@ -628,7 +628,7 @@ class ENGINE_API UGameplayStatics : public UBlueprintFunctionLibrary
 	UFUNCTION(BlueprintCallable, Category = "Game")
 	static bool DeleteGameInSlot(const FString& SlotName, const int32 UserIndex);
 
-	/** Returns the frame delta time in seconds adjusted by e.g. time dilation. */
+	/** Returns the frame delta time in seconds, adjusted by time dilation. */
 	UFUNCTION(BlueprintPure, Category = "Utilities|Time", meta = (WorldContext="WorldContextObject"))
 	static float GetWorldDeltaSeconds(UObject* WorldContextObject);
 
@@ -643,7 +643,7 @@ class ENGINE_API UGameplayStatics : public UBlueprintFunctionLibrary
 	UFUNCTION(BlueprintPure, Category="Utilities|Time", meta=(WorldContext="WorldContextObject"))
 	static void GetAccurateRealTime(UObject* WorldContextObject, int32& Seconds, float& PartialSeconds);
 
-	/** DVRStreaming API */
+	/*~ DVRStreaming API */
 	
 	/**
 	 * Toggle live DVR streaming.
@@ -677,32 +677,32 @@ class ENGINE_API UGameplayStatics : public UBlueprintFunctionLibrary
 	/** Native version, has more options than the Blueprint version. */
 	static bool SuggestProjectileVelocity(UObject* WorldContextObject, FVector& TossVelocity, FVector StartLocation, FVector EndLocation, float TossSpeed, bool bHighArc = false, float CollisionRadius = 0.f, float OverrideGravityZ = 0, ESuggestProjVelocityTraceOption::Type TraceOption = ESuggestProjVelocityTraceOption::TraceFullPath, const FCollisionResponseParams& ResponseParam = FCollisionResponseParams::DefaultResponseParam, const TArray<AActor*>& ActorsToIgnore = TArray<AActor*>(), bool bDrawDebug = false);
 
-	/** Returns world origin current location */
+	/** Returns world origin current location. */
 	UFUNCTION(BlueprintPure, Category="Game", meta=(WorldContext="WorldContextObject") )
 	static FIntVector GetWorldOriginLocation(UObject* WorldContextObject);
 	
-	/** Requests a new location for a world origin */
+	/** Requests a new location for a world origin. */
 	UFUNCTION(BlueprintCallable, Category="Game", meta=(WorldContext="WorldContextObject"))
 	static void SetWorldOriginLocation(UObject* WorldContextObject, FIntVector NewLocation);
 
 	/**
 	* Counts how many grass foliage instances overlap a given sphere.
 	*
-	* @param	Mesh			The static mesh we are interested in counting
-	* @param	CenterPosition	The center position of the sphere
+	* @param	Mesh			The static mesh we are interested in counting.
+	* @param	CenterPosition	The center position of the sphere.
 	* @param	Radius			The radius of the sphere.
 	*
-	* @return number of foliage instances with their mesh set to Mesh that overlap the sphere
+	* @return Number of foliage instances with their mesh set to Mesh that overlap the sphere.
 	*/
 	UFUNCTION(BlueprintCallable, Category = "Foliage", meta = (WorldContext = "WorldContextObject", UnsafeDuringActorConstruction = "true"))
 	static int32 GrassOverlappingSphereCount(UObject* WorldContextObject, const UStaticMesh* StaticMesh, FVector CenterPosition, float Radius);
 
 	/** 
-	 * Transforms the given 2D screen space coordinate into a 3D world-space point and direction
+	 * Transforms the given 2D screen space coordinate into a 3D world-space point and direction.
 	 * @param Player			Deproject using this player's view.
-	 * @param ScreenPosition	2D screen space to deproject
-	 * @param WorldPosition		(out) Corresponding 3D position in world space
-	 * @param WorldDirection	(out) World space direction vector away from the camera at the given 2d poiunt
+	 * @param ScreenPosition	2D screen space to deproject.
+	 * @param WorldPosition		(out) Corresponding 3D position in world space.
+	 * @param WorldDirection	(out) World space direction vector away from the camera at the given 2d poiunt.
 	 */
 	UFUNCTION(BlueprintPure, Category = "Utilities", meta = (Keywords = "unproject"))
 	static bool DeprojectScreenToWorld(APlayerController const* Player, const FVector2D& ScreenPosition, FVector& WorldPosition, FVector& WorldDirection);
@@ -716,10 +716,10 @@ class ENGINE_API UGameplayStatics : public UBlueprintFunctionLibrary
 	UFUNCTION(BlueprintPure, Category = "Utilities")
 	static bool ProjectWorldToScreen(APlayerController const* Player, const FVector& WorldPosition, FVector2D& ScreenPosition);
 
-	// Utility functions for interacting with Options strings
+	//~ Utility functions for interacting with Options strings
 
-	//=========================================================================
-	// URL Parsing
+	//~=========================================================================
+	//~ URL Parsing
 
 	static bool GrabOption( FString& Options, FString& ResultString );
 
