Merging CL 2690191: Reduce memory churn in animation systems

Use GetNumNames instead of filling an array and discarding it only to use its size.
Re-order PostAnimEvaluation to avoid use of temporary FBlendedCurve variable. This actually exposed a bug whereby curve evaluation was not being interpolated correctly using URO, now fixed.

diff --git a/Engine/Source/Runtime/Engine/Classes/Animation/SmartName.h b/Engine/Source/Runtime/Engine/Classes/Animation/SmartName.h
--- a/Engine/Source/Runtime/Engine/Classes/Animation/SmartName.h
+++ b/Engine/Source/Runtime/Engine/Classes/Animation/SmartName.h
@@ -66,7 +66,7 @@ struct ENGINE_API FSmartNameMapping
 	bool Exists(const FName& Name);
 
 	// Get the number of names currently stored in this container
-	int32 GetNumNames();
+	int32 GetNumNames() const;
 
 	// Serialize this to the provided archive; required for TMap serialization
 	void Serialize(FArchive& Ar);
diff --git a/Engine/Source/Runtime/Engine/Private/Animation/SmartName.cpp b/Engine/Source/Runtime/Engine/Private/Animation/SmartName.cpp
--- a/Engine/Source/Runtime/Engine/Private/Animation/SmartName.cpp
+++ b/Engine/Source/Runtime/Engine/Private/Animation/SmartName.cpp
@@ -124,7 +124,7 @@ void FSmartNameMapping::Serialize(FArchive& Ar)
 	}
 }
 
-int32 FSmartNameMapping::GetNumNames()
+int32 FSmartNameMapping::GetNumNames() const
 {
 	return UidMap.Num();
 }
diff --git a/Engine/Source/Runtime/Engine/Private/Components/SkeletalMeshComponent.cpp b/Engine/Source/Runtime/Engine/Private/Components/SkeletalMeshComponent.cpp
--- a/Engine/Source/Runtime/Engine/Private/Components/SkeletalMeshComponent.cpp
+++ b/Engine/Source/Runtime/Engine/Private/Components/SkeletalMeshComponent.cpp
@@ -1050,10 +1050,7 @@ int32 GetCurveNumber(USkeleton* Skeleton)
 	// get all curve list
 	if(const FSmartNameMapping* Mapping = Skeleton->SmartNames.GetContainer(USkeleton::AnimCurveMappingName))
 	{
-		TArray<FSmartNameMapping::UID> UIDList;
-		Mapping->FillUidArray(UIDList);
-
-		return UIDList.Num();
+		return Mapping->GetNumNames();
 	}
 
 	return 0;
@@ -1230,17 +1227,15 @@ void USkeletalMeshComponent::PrepareCloth()
 
 void USkeletalMeshComponent::PostAnimEvaluation(FAnimationEvaluationContext& EvaluationContext)
 {
-	FBlendedCurve EvaluatedCurve = EvaluationContext.Curve;
 	AnimEvaluationContext.Clear();
 
 	SCOPE_CYCLE_COUNTER(STAT_PostAnimEvaluation);
 	
-	if(AnimScriptInstance)
+	if (EvaluationContext.bDuplicateToCacheCurve)
 	{
-		// curve update happens first
-		AnimScriptInstance->UpdateCurves(EvaluatedCurve);
+		CachedCurve.InitFrom(EvaluationContext.Curve);
 	}
-	
+
 	if (EvaluationContext.bDuplicateToCacheBones)
 	{
 		CachedSpaceBases.Reset();
@@ -1249,11 +1244,6 @@ void USkeletalMeshComponent::PostAnimEvaluation(FAnimationEvaluationContext& Eva
 		CachedLocalAtoms.Append(LocalAtoms);
 	}
 
-	if (EvaluationContext.bDuplicateToCacheCurve)
-	{
-		CachedCurve = EvaluatedCurve;
-	}
-
 	if (EvaluationContext.bDoInterpolation)
 	{
 		SCOPE_CYCLE_COUNTER(STAT_InterpolateSkippedFrames);
@@ -1271,7 +1261,13 @@ void USkeletalMeshComponent::PostAnimEvaluation(FAnimationEvaluationContext& Eva
 		FAnimationRuntime::LerpBoneTransforms(GetEditableSpaceBases(), CachedSpaceBases, Alpha, RequiredBones);
 
 		// interpolate curve
-		EvaluatedCurve.BlendWith(CachedCurve, Alpha);
+		EvaluationContext.Curve.BlendWith(CachedCurve, Alpha);
+	}
+
+	if(AnimScriptInstance)
+	{
+		// curve update happens first
+		AnimScriptInstance->UpdateCurves(EvaluationContext.Curve);
 	}
 
 	bNeedToFlipSpaceBaseBuffers = true;
