Fix crash with cloth when no clothing assets are used. Also added CVar to turn parallel cloth off by default.

diff --git a/Engine/Source/Runtime/Engine/Private/Components/SkeletalMeshComponent.cpp b/Engine/Source/Runtime/Engine/Private/Components/SkeletalMeshComponent.cpp
--- a/Engine/Source/Runtime/Engine/Private/Components/SkeletalMeshComponent.cpp
+++ b/Engine/Source/Runtime/Engine/Private/Components/SkeletalMeshComponent.cpp
@@ -1215,13 +1215,15 @@ void USkeletalMeshComponent::RefreshBoneTransforms(FActorComponentTickFunction*
 
 void USkeletalMeshComponent::PrepareCloth()
 {
-	//TODO: logic about if to go in StartPhysics or PreCloth
-	if (FPhysScene* PhysScene = GetWorld()->GetPhysicsScene())
+	if(ClothingActors.Num())
 	{
-		FClothManager* ClothManager = PhysScene->GetClothManager();
-		bool bClothNeedsPhysics = BodyInstance.bSimulatePhysics || IsAnySimulatingPhysics();	//TODO: this errs on the side of simulating later. We may want to optimize this so that it's only needed when cloth bodies are simulating
-		PrepareClothSchedule PrepareSchedule = bClothNeedsPhysics ? PrepareClothSchedule::WaitOnPhysics : PrepareClothSchedule::IgnorePhysics;
-		ClothManager->RegisterForPrepareCloth(this, PrepareSchedule);
+		if (FPhysScene* PhysScene = GetWorld()->GetPhysicsScene())
+		{
+			FClothManager* ClothManager = PhysScene->GetClothManager();
+			bool bClothNeedsPhysics = BodyInstance.bSimulatePhysics || IsAnySimulatingPhysics();	//TODO: this errs on the side of simulating later. We may want to optimize this so that it's only needed when cloth bodies are simulating
+			PrepareClothSchedule PrepareSchedule = bClothNeedsPhysics ? PrepareClothSchedule::WaitOnPhysics : PrepareClothSchedule::IgnorePhysics;
+			ClothManager->RegisterForPrepareCloth(this, PrepareSchedule);
+		}
 	}
 }
 
diff --git a/Engine/Source/Runtime/Engine/Private/PhysicsEngine/ClothManager.cpp b/Engine/Source/Runtime/Engine/Private/PhysicsEngine/ClothManager.cpp
--- a/Engine/Source/Runtime/Engine/Private/PhysicsEngine/ClothManager.cpp
+++ b/Engine/Source/Runtime/Engine/Private/PhysicsEngine/ClothManager.cpp
@@ -5,8 +5,16 @@
 #include "PhysicsPublic.h"
 #include "Components/SkeletalMeshComponent.h"
 
-FClothManager::FClothManager(UWorld* AssociatedWorld)
+FClothManager::FClothManager(UWorld* InAssociatedWorld)
+: AssociatedWorld(InAssociatedWorld)
  {
+	//simulate cloth tick
+	StartClothTickFunction.bCanEverTick = true;
+	StartClothTickFunction.Target = this;
+	StartClothTickFunction.TickGroup = TG_StartCloth;
+	StartClothTickFunction.RegisterTickFunction(AssociatedWorld->PersistentLevel);
+
+	//prepare cloth data ticks
 	PrepareClothDataArray[(int32)PrepareClothSchedule::IgnorePhysics].TickFunction.TickGroup = TG_StartPhysics;
 	PrepareClothDataArray[(int32)PrepareClothSchedule::WaitOnPhysics].TickFunction.TickGroup = TG_PreCloth;
 
@@ -45,6 +53,42 @@ void FClothManager::RegisterForPrepareCloth(USkeletalMeshComponent* SkeletalMesh
 	PrepareClothDataArray[(int32)PrepSchedule].SkeletalMeshComponents.Add(SkeletalMeshComponent);
 }
 
+void FClothManager::StartCloth()
+{
+	FGraphEventArray ThingsToComplete;
+	for (const FClothManagerData& PrepareData : PrepareClothDataArray)
+	{
+		if(PrepareData.PrepareCompletion.IsValid())
+		{
+			ThingsToComplete.Add(PrepareData.PrepareCompletion);
+		}
+	}
+
+	if(ThingsToComplete.Num())
+	{
+		QUICK_SCOPE_CYCLE_COUNTER(STAT_FClothManager_WaitPrepareCloth);
+		FTaskGraphInterface::Get().WaitUntilTasksComplete(ThingsToComplete, ENamedThreads::GameThread);
+	}
+	
+	if (FPhysScene* PhysScene = AssociatedWorld->GetPhysicsScene())
+	{
+		PhysScene->StartCloth();
+	}
+}
+
+
+void FStartClothTickFunction::ExecuteTick(float DeltaTime, enum ELevelTick TickType, ENamedThreads::Type CurrentThread, const FGraphEventRef& MyCompletionGraphEvent)
+{
+	Target->StartCloth();
+}
+
+FString FStartClothTickFunction::DiagnosticMessage()
+{
+	return TEXT("FStartClothTickFunction");
+}
+
+TAutoConsoleVariable<int32> CVarParallelCloth(TEXT("p.ParallelCloth"), 0, TEXT("If turned on, preparing cloth will happen off the game thread."));
+
 class FPrepareClothTickTask
 {
 	FClothManagerData* PrepareClothData;
@@ -63,7 +107,7 @@ public:
 	}
 	static ENamedThreads::Type GetDesiredThread()
 	{
-		return ENamedThreads::AnyThread;
+		return CVarParallelCloth.GetValueOnGameThread() ? ENamedThreads::AnyThread : ENamedThreads::GameThread;
 	}
 	static ESubsequentsMode::Type GetSubsequentsMode()
 	{
@@ -81,7 +125,8 @@ void FPrepareClothTickFunction::ExecuteTick(float DeltaTime, enum ELevelTick Tic
 {
 	QUICK_SCOPE_CYCLE_COUNTER(FPrepareClothTickFunction_ExecuteTick);
 
-	TGraphTask<FPrepareClothTickTask>::CreateTask().ConstructAndDispatchWhenReady(Target, DeltaTime);
+	check(!Target->PrepareCompletion.IsValid())
+	Target->PrepareCompletion = TGraphTask<FPrepareClothTickTask>::CreateTask().ConstructAndDispatchWhenReady(Target, DeltaTime);
 }
 
 FString FPrepareClothTickFunction::DiagnosticMessage()
diff --git a/Engine/Source/Runtime/Engine/Private/PhysicsEngine/ClothManager.h b/Engine/Source/Runtime/Engine/Private/PhysicsEngine/ClothManager.h
--- a/Engine/Source/Runtime/Engine/Private/PhysicsEngine/ClothManager.h
+++ b/Engine/Source/Runtime/Engine/Private/PhysicsEngine/ClothManager.h
@@ -35,11 +35,22 @@ public:
 	TArray<USkeletalMeshComponent*> SkeletalMeshComponents;
 	FPrepareClothTickFunction TickFunction;
 	FThreadSafeBool IsPreparingCloth;
+	FGraphEventRef PrepareCompletion;
 
 	/** Go through all SkeletalMeshComponents registered for work this frame and call the solver as needed*/
 	void PrepareCloth(float DeltaTime);
 };
 
+struct FStartClothTickFunction : public FTickFunction
+{
+	class FClothManager* Target;
+
+	// FTickFunction interface
+	virtual void ExecuteTick(float DeltaTime, enum ELevelTick TickType, ENamedThreads::Type CurrentThread, const FGraphEventRef& MyCompletionGraphEvent) override;
+	virtual FString DiagnosticMessage() override;
+	// End of FTickFunction interface
+};
+
 
 class FClothManager
 {
@@ -63,6 +74,11 @@ public:
 	}
 	
 private:
-	
+	UWorld* AssociatedWorld;
 	FClothManagerData PrepareClothDataArray[(int32)PrepareClothSchedule::MAX];
+	FStartClothTickFunction StartClothTickFunction;
+
+	friend FStartClothTickFunction;
+
+	void StartCloth();
 };
\ No newline at end of file
diff --git a/Engine/Source/Runtime/Engine/Private/PhysicsEngine/PhysLevel.cpp b/Engine/Source/Runtime/Engine/Private/PhysicsEngine/PhysLevel.cpp
--- a/Engine/Source/Runtime/Engine/Private/PhysicsEngine/PhysLevel.cpp
+++ b/Engine/Source/Runtime/Engine/Private/PhysicsEngine/PhysLevel.cpp
@@ -214,7 +214,6 @@ void UWorld::StartClothAndAsyncSim()
 {
 	if (FPhysScene* PhysScene = GetPhysicsScene())
 	{
-		PhysScene->StartCloth();
 		PhysScene->StartAsync();
 	}
 }
