Cloth manager only ticks cloth scene if there is cloth registered. Should help games/cases where cloth is not used

diff --git a/Engine/Source/Runtime/Engine/Private/PhysicsEngine/ClothManager.cpp b/Engine/Source/Runtime/Engine/Private/PhysicsEngine/ClothManager.cpp
--- a/Engine/Source/Runtime/Engine/Private/PhysicsEngine/ClothManager.cpp
+++ b/Engine/Source/Runtime/Engine/Private/PhysicsEngine/ClothManager.cpp
@@ -41,7 +41,6 @@ void FClothManagerData::PrepareCloth(float DeltaTime)
 			SkeletalMeshComponent->ParallelTickClothing(DeltaTime, ClothSimulationContext);
 		}
 
-		SkeletalMeshComponents.Empty(SkeletalMeshComponents.Num());
 		IsPreparingCloth.AtomicSet(false);
 	}
 #endif
@@ -55,12 +54,20 @@ void FClothManager::RegisterForPrepareCloth(USkeletalMeshComponent* SkeletalMesh
 
 void FClothManager::StartCloth()
 {
+	bool bNeedSimulateCloth = false;
+
 	FGraphEventArray ThingsToComplete;
-	for (const FClothManagerData& PrepareData : PrepareClothDataArray)
+	for (FClothManagerData& PrepareData : PrepareClothDataArray)
 	{
-		if(PrepareData.PrepareCompletion.IsValid())
+		if(PrepareData.SkeletalMeshComponents.Num())
 		{
-			ThingsToComplete.Add(PrepareData.PrepareCompletion);
+			if (PrepareData.PrepareCompletion.IsValid())
+			{
+				ThingsToComplete.Add(PrepareData.PrepareCompletion);
+			}
+
+			PrepareData.SkeletalMeshComponents.Reset();
+			bNeedSimulateCloth = true;
 		}
 	}
 
@@ -70,9 +77,12 @@ void FClothManager::StartCloth()
 		FTaskGraphInterface::Get().WaitUntilTasksComplete(ThingsToComplete, ENamedThreads::GameThread);
 	}
 	
-	if (FPhysScene* PhysScene = AssociatedWorld->GetPhysicsScene())
+	if(bNeedSimulateCloth)
 	{
-		PhysScene->StartCloth();
+		if (FPhysScene* PhysScene = AssociatedWorld->GetPhysicsScene())
+		{
+			PhysScene->StartCloth();
+		}
 	}
 }
 
