Fix non apex build

diff --git a/Engine/Source/Runtime/Engine/Private/Components/SkeletalMeshComponent.cpp b/Engine/Source/Runtime/Engine/Private/Components/SkeletalMeshComponent.cpp
--- a/Engine/Source/Runtime/Engine/Private/Components/SkeletalMeshComponent.cpp
+++ b/Engine/Source/Runtime/Engine/Private/Components/SkeletalMeshComponent.cpp
@@ -1215,6 +1215,7 @@ void USkeletalMeshComponent::RefreshBoneTransforms(FActorComponentTickFunction*
 
 void USkeletalMeshComponent::PrepareCloth()
 {
+#if WITH_APEX_CLOTHING
 	if(ClothingActors.Num())
 	{
 		if (FPhysScene* PhysScene = GetWorld()->GetPhysicsScene())
@@ -1225,6 +1226,7 @@ void USkeletalMeshComponent::PrepareCloth()
 			ClothManager->RegisterForPrepareCloth(this, PrepareSchedule);
 		}
 	}
+#endif
 }
 
 void USkeletalMeshComponent::PostAnimEvaluation(FAnimationEvaluationContext& EvaluationContext)
