Fix race condition with parallel cloth where the skeletal mesh components could potentially empty while running on another thread

diff --git a/Engine/Source/Runtime/Engine/Private/PhysicsEngine/ClothManager.cpp b/Engine/Source/Runtime/Engine/Private/PhysicsEngine/ClothManager.cpp
--- a/Engine/Source/Runtime/Engine/Private/PhysicsEngine/ClothManager.cpp
+++ b/Engine/Source/Runtime/Engine/Private/PhysicsEngine/ClothManager.cpp
@@ -54,30 +54,34 @@ void FClothManager::RegisterForPrepareCloth(USkeletalMeshComponent* SkeletalMesh
 
 void FClothManager::StartCloth()
 {
-	bool bNeedSimulateCloth = false;
-
+	//Make sure prepare tasks are done
 	FGraphEventArray ThingsToComplete;
-	for (FClothManagerData& PrepareData : PrepareClothDataArray)
+	for (const FClothManagerData& PrepareData : PrepareClothDataArray)
 	{
-		if(PrepareData.SkeletalMeshComponents.Num())
+		if (PrepareData.PrepareCompletion.IsValid())
 		{
-			if (PrepareData.PrepareCompletion.IsValid())
-			{
-				ThingsToComplete.Add(PrepareData.PrepareCompletion);
-			}
-
-			PrepareData.SkeletalMeshComponents.Reset();
-			bNeedSimulateCloth = true;
+			ThingsToComplete.Add(PrepareData.PrepareCompletion);
 		}
 	}
 
-	if(ThingsToComplete.Num())
+	if (ThingsToComplete.Num())
 	{
 		QUICK_SCOPE_CYCLE_COUNTER(STAT_FClothManager_WaitPrepareCloth);
 		FTaskGraphInterface::Get().WaitUntilTasksComplete(ThingsToComplete, ENamedThreads::GameThread);
 	}
-	
-	if(bNeedSimulateCloth)
+
+	//Reset skeletal mesh components
+	bool bNeedSimulateCloth = false;
+	for (FClothManagerData& PrepareData : PrepareClothDataArray)
+	{
+		if (PrepareData.SkeletalMeshComponents.Num())
+		{
+			bNeedSimulateCloth = true;
+			PrepareData.SkeletalMeshComponents.Reset();
+		}
+	}
+
+	if (bNeedSimulateCloth)
 	{
 		if (FPhysScene* PhysScene = AssociatedWorld->GetPhysicsScene())
 		{
