Fixed possible divide by zero in animation transitions with zero crossfade duration.

diff --git a/Engine/Source/Runtime/Engine/Private/Animation/AnimNode_StateMachine.cpp b/Engine/Source/Runtime/Engine/Private/Animation/AnimNode_StateMachine.cpp
--- a/Engine/Source/Runtime/Engine/Private/Animation/AnimNode_StateMachine.cpp
+++ b/Engine/Source/Runtime/Engine/Private/Animation/AnimNode_StateMachine.cpp
@@ -104,7 +104,16 @@ void FAnimationActiveTransitionEntry::Update(const FAnimationUpdateContext& Cont
 	{
 		ElapsedTime += Context.GetDeltaTime();
 		Blend.Update(Context.GetDeltaTime());
-		Alpha = FAlphaBlend::AlphaToBlendOption(ElapsedTime / CrossfadeDuration, Blend.BlendOption, Blend.CustomCurve); //Blend.GetBlendedValue();
+
+		float QueryAlpha = 1.0f;
+
+		// If non-zero, calculate the query alpha
+		if (CrossfadeDuration > 0.0f)
+		{
+			QueryAlpha = ElapsedTime / CrossfadeDuration;
+		}
+
+		Alpha = FAlphaBlend::AlphaToBlendOption(QueryAlpha, Blend.BlendOption, Blend.CustomCurve);
 
 		if(Blend.IsComplete())
 		{
