Merging using UE4->Orion

diff --git a/Engine/Source/Runtime/Core/Private/HAL/ConsoleManager.cpp b/Engine/Source/Runtime/Core/Private/HAL/ConsoleManager.cpp
--- a/Engine/Source/Runtime/Core/Private/HAL/ConsoleManager.cpp
+++ b/Engine/Source/Runtime/Core/Private/HAL/ConsoleManager.cpp
@@ -122,7 +122,7 @@ public:
 			FConsoleManager& ConsoleManager = (FConsoleManager&)IConsoleManager::Get();
 			FString CVarName = ConsoleManager.FindConsoleObjectName(this);
 			UE_LOG(LogConsoleManager, Warning,
-				TEXT("Console variable '%s' wasn't set ('%s' has a lower priority than '%s')"),
+				TEXT("A console variable Set() for '%s' failed. The access with 'SetBy%s' failed because 'SetBy%s' did set it before and and has a higher priority (this might be an intentional override)"),
 				CVarName.IsEmpty() ? TEXT("unknown?") : *CVarName,
 				GetSetByTCHAR((EConsoleVariableFlags)NewPri),
 				GetSetByTCHAR((EConsoleVariableFlags)OldPri)
diff --git a/Engine/Source/Runtime/Engine/Private/SkeletalRenderGPUSkin.cpp b/Engine/Source/Runtime/Engine/Private/SkeletalRenderGPUSkin.cpp
--- a/Engine/Source/Runtime/Engine/Private/SkeletalRenderGPUSkin.cpp
+++ b/Engine/Source/Runtime/Engine/Private/SkeletalRenderGPUSkin.cpp
@@ -220,7 +220,8 @@ void FSkeletalMeshObjectGPUSkin::Update(int32 LODIndex,USkinnedMeshComponent* In
 static TAutoConsoleVariable<int32> CVarDeferSkeletalDynamicDataUpdateUntilGDME(
 	TEXT("r.DeferSkeletalDynamicDataUpdateUntilGDME"),
 	0,
-	TEXT("If > 0, then do skeletal mesh dynamic data updates will be deferred until GDME. Experimental option."));
+	TEXT("If > 0, then do skeletal mesh dynamic data updates will be deferred until GDME. Experimental option."),
+	ECVF_Cheat);
 
 void FSkeletalMeshObjectGPUSkin::UpdateDynamicData_RenderThread(FRHICommandListImmediate& RHICmdList, FDynamicSkelMeshObjectDataGPUSkin* InDynamicData)
 {
diff --git a/Engine/Source/Runtime/Sockets/Private/NetworkMessage.cpp b/Engine/Source/Runtime/Sockets/Private/NetworkMessage.cpp
--- a/Engine/Source/Runtime/Sockets/Private/NetworkMessage.cpp
+++ b/Engine/Source/Runtime/Sockets/Private/NetworkMessage.cpp
@@ -85,7 +85,7 @@ bool FNFSMessageHeader::ReceivePayload(FArrayReader& OutPayload, const FSimpleAb
 		
 	if (!Socket.Receive(HeaderBytes.GetData(), Size))
 	{
-		UE_LOG(LogSockets, Error, TEXT("Unable to read full header"));
+		UE_LOG(LogSockets, Error, TEXT("Unable to read full network header"));
 		return false;
 	}
 
