scale crash fix

I found a way to reproduce the bug easily:  drag a box shape into the level in the editor. make it movable. ctrl-e to edit. remove collisions. add about 6 box collisions to it. simulate. then select the box and try to scale it.

essentially elements of SomeClass in a TArray<FKShapeElem> was storing addresses of themselves.  when the TArray grew, those pointers became no longer valid. these self pointers are used as 'userdata' on physx objects. When creating a physx shape from the FKShapeElem, it was getting an out of date pointer to use as userdata. Its probably the only place that uses this particualr physx shape userdata is the scale code as it iterating over physx shapes and reaches back into what was turning out to be the wrong place to get initial shape information.

given physics is recreated on each change to collision already, we can fix this issue simply by computing this pointer/userdata on the fly, rather then the first time the shape is added to the TArray

future todo:  still would want a better from-physx-back-pointing system to safeguard to prevent this these sorts of thing from happening in the first place.  what if  this array somehow got  modified without physx scene getting re-created.

crash reports:  http://crashreporter/Buggs/Show/73744
https://jira.ol.epicgames.net/browse/UE-20961

diff --git a/Engine/Source/Runtime/Engine/Classes/PhysicsEngine/ShapeElem.h b/Engine/Source/Runtime/Engine/Classes/PhysicsEngine/ShapeElem.h
--- a/Engine/Source/Runtime/Engine/Classes/PhysicsEngine/ShapeElem.h
+++ b/Engine/Source/Runtime/Engine/Classes/PhysicsEngine/ShapeElem.h
@@ -55,7 +55,7 @@ struct FKShapeElem
 		return (T*)this;
 	}
 
-	const FPhysxUserData* GetUserData() const { return &UserData; }
+	const FPhysxUserData* GetUserData() const { FPhysxUserData::Set<FKShapeElem>((void*)&UserData, const_cast<FKShapeElem*>(this));  return &UserData; }
 
 	ENGINE_API static EAggCollisionShape::Type StaticShapeType;
 
