Engine: Allow memreport to work in Test configuration without requiring UEBuildConfiguration.bUseLoggingInShipping to be true (still requires ALLOW_DEBUG_FILES, and is disabled in the true Shipping configuration)
- Added a new FOutputDeviceArchiveWrapper which doesn't try to do any special work and merely adapts the FOutputDevice API to write to a FArchive
- Switched UEngine::HandleMemReportDeferredCommand over to use FOutputDeviceArchiveWrapper

[INTEGRATE] Change 2691200 by [REDACTED_PERSON][REDACTED_HANDLE] on 2015/09/14 20:17:53

diff --git a/Engine/Source/Runtime/Core/Private/HAL/OutputDevices.cpp b/Engine/Source/Runtime/Core/Private/HAL/OutputDevices.cpp
--- a/Engine/Source/Runtime/Core/Private/HAL/OutputDevices.cpp
+++ b/Engine/Source/Runtime/Core/Private/HAL/OutputDevices.cpp
@@ -1245,3 +1245,31 @@ void FOutputDeviceAnsiError::HandleError()
 
 	FCoreDelegates::OnShutdownAfterError.Broadcast();
 }
+
+/*-----------------------------------------------------------------------------
+	FOutputDeviceArchiveWrapper subclasses.
+-----------------------------------------------------------------------------*/
+
+void FOutputDeviceArchiveWrapper::Flush()
+{
+	LogAr->Flush();
+}
+
+void FOutputDeviceArchiveWrapper::Serialize(const TCHAR* Data, ELogVerbosity::Type Verbosity, const class FName& Category)
+{
+	if (Verbosity != ELogVerbosity::SetColor)
+	{
+		CastAndSerializeData(Data);
+
+		if (bAutoEmitLineTerminator)
+		{
+			CastAndSerializeData(LINE_TERMINATOR);
+		}
+	}
+}
+
+void FOutputDeviceArchiveWrapper::CastAndSerializeData(const TCHAR* Data)
+{
+	FTCHARToUTF8 ConvertedData(Data);
+	LogAr->Serialize((ANSICHAR*)(ConvertedData.Get()), ConvertedData.Length());
+}
diff --git a/Engine/Source/Runtime/Core/Public/HAL/OutputDevices.h b/Engine/Source/Runtime/Core/Public/HAL/OutputDevices.h
--- a/Engine/Source/Runtime/Core/Public/HAL/OutputDevices.h
+++ b/Engine/Source/Runtime/Core/Public/HAL/OutputDevices.h
@@ -379,7 +379,7 @@ public:
 #define BACKUP_LOG_FILENAME_POSTFIX TEXT("-backup-")
 
 /**
- * File output device.
+ * File output device (Note: Only works if ALLOW_LOG_FILE && !NO_LOGGING is true, otherwise Serialize does nothing).
  */
 class CORE_API FOutputDeviceFile : public FOutputDevice
 {
@@ -521,6 +521,36 @@ private:
 	int32		ErrorPos;
 };
 
+
+/**
+ * Output device wrapping any kind of FArchive.  Note: Works in any build configuration.
+ */
+class CORE_API FOutputDeviceArchiveWrapper : public FOutputDevice
+{
+public:
+	/**
+	 * Constructor, initializing member variables.
+	 *
+	 * @param InArchive	Archive to use, must not be nullptr.  Does not take ownership of the archive, clean up or delete the archive independently!
+	 */
+	FOutputDeviceArchiveWrapper(FArchive* InArchive)
+		: LogAr(InArchive)
+	{
+		check(InArchive);
+	}
+
+	// FOutputDevice interface
+	virtual void Flush() override;
+	virtual void Serialize(const TCHAR* Data, ELogVerbosity::Type Verbosity, const class FName& Category) override;
+	// End of FOutputDevice interface
+
+private:
+	FArchive* LogAr;
+
+private:
+	void CastAndSerializeData(const TCHAR* Data);
+};
+
 CORE_API DECLARE_LOG_CATEGORY_EXTERN(LogHAL, Log, All);
 CORE_API DECLARE_LOG_CATEGORY_EXTERN(LogMac, Log, All);
 CORE_API DECLARE_LOG_CATEGORY_EXTERN(LogLinux, Log, All);
diff --git a/Engine/Source/Runtime/Engine/Private/UnrealEngine.cpp b/Engine/Source/Runtime/Engine/Private/UnrealEngine.cpp
--- a/Engine/Source/Runtime/Engine/Private/UnrealEngine.cpp
+++ b/Engine/Source/Runtime/Engine/Private/UnrealEngine.cpp
@@ -4200,6 +4200,7 @@ bool UEngine::HandleMemReportCommand( const TCHAR* Cmd, FOutputDevice& Ar, UWorl
 
 bool UEngine::HandleMemReportDeferredCommand( const TCHAR* Cmd, FOutputDevice& Ar, UWorld* InWorld )
 {
+#if ALLOW_DEBUG_FILES
 	const bool bPerformSlowCommands = FParse::Param( Cmd, TEXT("FULL") );
 	const bool bLogOutputToFile = !FParse::Param( Cmd, TEXT("LOG") );
 
@@ -4212,7 +4213,8 @@ bool UEngine::HandleMemReportDeferredCommand( const TCHAR* Cmd, FOutputDevice& A
 	FlushRenderingCommands();
 
 	FOutputDevice* ReportAr = &Ar;
-	FOutputDeviceFile* FileAr = NULL;
+	FArchive* FileAr = nullptr;
+	FOutputDeviceArchiveWrapper* FileArWrapper = nullptr;
 	FString FilenameFull;
 	
 	if (bLogOutputToFile)
@@ -4222,8 +4224,10 @@ bool UEngine::HandleMemReportDeferredCommand( const TCHAR* Cmd, FOutputDevice& A
 
 		const FString Filename = CreateProfileFilename( TEXT(".memreport"), true );
 		FilenameFull = PathName + Filename;
-		FileAr = new FOutputDeviceFile(*FilenameFull);
-		ReportAr = FileAr;
+	
+		FileAr = IFileManager::Get().CreateDebugFileWriter(*FilenameFull);
+		FileArWrapper = new FOutputDeviceArchiveWrapper(FileAr);
+		ReportAr = FileArWrapper;
 
 		UE_LOG(LogEngine, Log, TEXT("MemReportDeferred: saving to %s"), *FilenameFull);		
 	}
@@ -4256,14 +4260,13 @@ bool UEngine::HandleMemReportDeferredCommand( const TCHAR* Cmd, FOutputDevice& A
 		}
 	}
 
-	if (FileAr)
+	if (FileArWrapper != nullptr)
 	{
-		FileAr->TearDown();
-		// This no longer seems to work in UE4
-		// SendDataToPCViaUnrealConsole( TEXT("UE_PROFILER!MEMREPORT:"), *(FilenameFull) );
-
+		FileArWrapper->TearDown();
+		delete FileArWrapper;
 		delete FileAr;
 	}
+#endif
 
 	return true;
 }
