Rendering: Prevent particle system bounding boxes from being typed in the wrong way around (with Min > Max)

[INTEGRATE] Change 2691197 by [REDACTED_PERSON][REDACTED_HANDLE] on 2015/09/14 20:15:49

diff --git a/Engine/Source/Runtime/Engine/Private/Particles/ParticleComponents.cpp b/Engine/Source/Runtime/Engine/Private/Particles/ParticleComponents.cpp
--- a/Engine/Source/Runtime/Engine/Private/Particles/ParticleComponents.cpp
+++ b/Engine/Source/Runtime/Engine/Private/Particles/ParticleComponents.cpp
@@ -1791,6 +1791,23 @@ void UParticleSystem::PostEditChangeProperty(FPropertyChangedEvent& PropertyChan
 		}
 	}
 
+	// Ensure the bounds have a positive size
+	if (FixedRelativeBoundingBox.IsValid)
+	{
+		if (FixedRelativeBoundingBox.Min.X > FixedRelativeBoundingBox.Max.X)
+		{
+			Swap(FixedRelativeBoundingBox.Min.X, FixedRelativeBoundingBox.Max.X);
+		}
+		if (FixedRelativeBoundingBox.Min.Y > FixedRelativeBoundingBox.Max.Y)
+		{
+			Swap(FixedRelativeBoundingBox.Min.Y, FixedRelativeBoundingBox.Max.Y);
+		}
+		if (FixedRelativeBoundingBox.Min.Z > FixedRelativeBoundingBox.Max.Z)
+		{
+			Swap(FixedRelativeBoundingBox.Min.Z, FixedRelativeBoundingBox.Max.Z);
+		}
+	}
+
 	//cap the WarmupTickRate to realistic values
 	if (WarmupTickRate <= 0)
 	{
