AB testing for a given scope.

diff --git a/Engine/Source/Runtime/Core/Private/ProfilingDebugging/ABTesting.cpp b/Engine/Source/Runtime/Core/Private/ProfilingDebugging/ABTesting.cpp
new file mode 100644
--- /dev/null
+++ b/Engine/Source/Runtime/Core/Private/ProfilingDebugging/ABTesting.cpp
@@ -0,0 +1,379 @@
+// Copyright 1998-2015 Epic Games, Inc. All Rights Reserved.
+
+/**
+* Here are a number of profiling helper functions so we do not have to duplicate a lot of the glue
+* code everywhere.  And we can have consistent naming for all our files.
+*
+*/
+
+// Core includes.
+#include "CorePrivatePCH.h"
+#include "ABTesting.h"
+
+#if ENABLE_ABTEST
+
+static TAutoConsoleVariable<int32> CVarABTestHistory(
+	TEXT("abtest.HistoryNum"),
+	1000,
+	TEXT("Number of history frames to use for stats."));
+
+static TAutoConsoleVariable<int32> CVarABTestReportNum(
+	TEXT("abtest.ReportNum"),
+	100,
+	TEXT("Number of frames between reports."));
+
+static TAutoConsoleVariable<int32> CVarABTestCoolDown(
+	TEXT("abtest.CoolDown"),
+	3,
+	TEXT("Number of frames to discard data after each command to cover threading."));
+
+static TAutoConsoleVariable<int32> CVarABTestMinFramesPerTrial(
+	TEXT("abtest.MinFramesPerTrial"),
+	5,
+	TEXT("The number of frames to run a given command before switching; this is randomized."));
+
+static TAutoConsoleVariable<int32> CVarABTestNumResamples(
+	TEXT("abtest.NumResamples"),
+	256,
+	TEXT("The number of resamples to use to determine confidence."));
+
+
+
+FABTest::FABTest()
+		: Stream(9567)
+		, bABTestActive(false)
+		, bABScopeTestActive(false)
+		, bFrameLog(false)
+		, ABTestNumSamples(0)
+		, RemainingCoolDown(0)
+		, CurrentTest(0)
+		, RemainingTrial(0)
+		, RemainingPrint(0)
+		, SampleIndex(0)
+		, HistoryNum(0)
+		, ReportNum(0)
+		, CoolDown(0)
+		, MinFramesPerTrial(0)
+		, NumResamples(0)
+		, TotalScopeTimeInFrame(0)
+	{
+	}
+
+void FABTest::StartFrameLog()
+{
+	TotalTime = 0.0;
+	TotalFrames = 0;
+	Spikes = 0;
+	bFrameLog = true;
+	UE_LOG(LogConsoleResponse, Display, TEXT("Starting frame log."));
+}
+
+void FABTest::FrameLogTick(double Delta)
+{
+	if (Delta > .034)
+	{
+		Spikes++;
+	}
+	else
+	{
+		TotalFrames++;
+		TotalTime += Delta;
+	}
+	if (TotalFrames > 0 && TotalFrames % 1000 == 0)
+	{
+		FPlatformMisc::LowLevelOutputDebugStringf(TEXT("%8d frames   %6.3fms/f    %8d spikes rejected "), TotalFrames, float(1000.0 * TotalTime / TotalFrames), Spikes);
+	}
+}
+
+const TCHAR* FABTest::TickAndGetCommand()
+{
+	const TCHAR* OutCommand = nullptr;
+
+	static double LastTime = 0;
+	if (bABTestActive && RemainingCoolDown)
+	{
+		RemainingCoolDown--;
+	}
+	else if (bABTestActive)
+	{
+		float Delta;
+		if (bABScopeTestActive)
+		{
+			Delta = TotalScopeTimeInFrame;
+		}
+		else
+		{
+			Delta = FPlatformTime::Seconds() - LastTime;
+			check(Delta > 0);
+		}
+
+		FSample* Sample = nullptr;
+		if (ABTestNumSamples < HistoryNum)
+		{
+			check(ABTestNumSamples == Samples.Num());
+			Sample = new (Samples)FSample();
+		}
+		else
+		{
+			Sample = &Samples[ABTestNumSamples % HistoryNum];
+
+			// remove the old data
+			check(Totals[Sample->TestIndex] > Sample->Micros || bABScopeTestActive);
+			Totals[Sample->TestIndex] -= Sample->Micros;
+			check(Counts[Sample->TestIndex]);
+			Counts[Sample->TestIndex]--;
+			check(Sample->InResamples.Num() == NumResamples);
+			for (TConstSetBitIterator<> CurIndexIt(Sample->InResamples); CurIndexIt; ++CurIndexIt)
+			{
+				int32 ResampleIndex = CurIndexIt.GetIndex();
+				check(ResampleAccumulators[ResampleIndex] > Sample->Micros || bABScopeTestActive);
+				ResampleAccumulators[ResampleIndex] -= Sample->Micros;
+				check(ResampleCount[ResampleIndex]);
+				ResampleCount[ResampleIndex]--;
+			}
+		}
+		ABTestNumSamples++;
+		Sample->Micros = uint32(Delta * 1000000.0f);
+		check(Sample->Micros || bABScopeTestActive);
+		Sample->TestIndex = CurrentTest;
+		Totals[Sample->TestIndex] += Sample->Micros;
+		Counts[Sample->TestIndex]++;
+		Sample->InResamples.Empty(NumResamples);
+		for (int32 ResampleIndex = 0; ResampleIndex < NumResamples; ResampleIndex++)
+		{
+			bool In = Stream.FRand() > 0.5f;
+			Sample->InResamples.Add(In);
+			if (In)
+			{
+				ResampleAccumulators[ResampleIndex] += Sample->Micros;
+				ResampleCount[ResampleIndex]++;
+			}
+		}
+		check(RemainingPrint);
+		if (!--RemainingPrint)
+		{
+			if (Counts[0] && Counts[1])
+			{
+				FPlatformMisc::LowLevelOutputDebugStringf(TEXT("      %7.4fms  (%4d samples)  A = '%s'\n"), float(Totals[0]) / float(Counts[0]) / 1000.0f, Counts[0], *ABTestCmds[0]);
+				FPlatformMisc::LowLevelOutputDebugStringf(TEXT("      %7.4fms  (%4d samples)  B = '%s'\n"), float(Totals[1]) / float(Counts[1]) / 1000.0f, Counts[1], *ABTestCmds[1]);
+
+				float Diff = (float(Totals[0]) / float(Counts[0]) / 1000.0f) - (float(Totals[1]) / float(Counts[1]) / 1000.0f);
+				bool bAIsFaster = false;
+				if (Diff < 0.0f)
+				{
+					Diff *= -1.0f;
+					bAIsFaster = true;
+				}
+				static TArray<float> Scores;
+				Scores.Reset();
+				uint32 TotalFrame = Totals[0] + Totals[1];
+				uint32 TotalCount = Counts[0] + Counts[1];
+				for (int32 ResampleIndex = 0; ResampleIndex < NumResamples; ResampleIndex++)
+				{
+					check(TotalFrame > ResampleAccumulators[ResampleIndex] || bABScopeTestActive);
+					uint32 Opposite = TotalFrame - ResampleAccumulators[ResampleIndex];
+					check(TotalCount > ResampleCount[ResampleIndex] || bABScopeTestActive);
+					uint32 OppositeCount = TotalCount - ResampleCount[ResampleIndex];
+
+					// this ABS here gives us better data but makes it one sided, so the 90% confidence interval is at index 80/100 etc
+					// it is legit because we could have easily done the inverse permutation
+					float SampleDiff = FMath::Abs(((ResampleAccumulators[ResampleIndex]) / float(ResampleCount[ResampleIndex]) / 1000.0f) - (float(Opposite) / float(OppositeCount) / 1000.0f));
+					Scores.Add(SampleDiff);
+
+				}
+				Scores.Sort();
+				int32 Conf = NumResamples;
+				for (int32 Trial = 0; Trial < NumResamples; Trial++)
+				{
+					if (Scores[Trial] > Diff)
+					{
+						Conf = Trial;
+						break;
+					}
+				}
+
+				float fConf = 1.0f - (0.5f + float(Conf) / float(NumResamples) / 2.0f);
+
+				if (bAIsFaster)
+				{
+					FPlatformMisc::LowLevelOutputDebugStringf(TEXT("      A is %7.4fms faster than B;  %3.0f%% chance this is noise.\n"), Diff, fConf * 100.0f);
+				}
+				else
+				{
+					FPlatformMisc::LowLevelOutputDebugStringf(TEXT("      B is %7.4fms faster than A;  %3.0f%% chance this is noise.\n"), Diff, fConf * 100.0f);
+				}
+
+				UE_LOG(LogConsoleResponse, Display, TEXT("----------------"));
+			}
+			else
+			{
+				UE_LOG(LogConsoleResponse, Display, TEXT("No Samples?"));
+			}
+			RemainingPrint = ReportNum;
+		}
+		check(RemainingTrial);
+		if (!--RemainingTrial)
+		{
+			OutCommand = SwitchTest(1 - CurrentTest);
+		}
+	}
+	else if (bFrameLog)
+	{
+		FrameLogTick(FPlatformTime::Seconds() - LastTime);
+	}
+	LastTime = FPlatformTime::Seconds();
+
+	TotalScopeTimeInFrame = 0.0;
+	return OutCommand;
+}
+
+FABTest& FABTest::Get()
+{
+	static FABTest Singleton;
+	return Singleton;
+}
+
+void FABTest::ABTestCmdFunc(const TArray<FString>& Args)
+{
+	FString ABTestCmds[2];
+	if (Args.Num() == 1 && Args[0].Compare(FString(TEXT("framelog")), ESearchCase::IgnoreCase) == 0)
+	{
+		Get().Stop();
+		Get().StartFrameLog();
+		return;
+	}
+	if (Args.Num() == 1 && Args[0].Compare(FString(TEXT("stop")), ESearchCase::IgnoreCase) == 0)
+	{
+		Get().Stop();		
+		return;
+	}
+	if (Args.Num() == 1 && Args[0].Compare(FString(TEXT("scope")), ESearchCase::IgnoreCase) == 0)
+	{
+		ABTestCmds[0] = FString("ScopeA");
+		ABTestCmds[1] = FString("ScopeB");
+		Get().Start(ABTestCmds, true);
+		return;
+	}
+	if (Args.Num() == 3 && !Args[0].StartsWith(TEXT("\"")))
+	{
+		ABTestCmds[0] = Args[0].TrimQuotes() + TEXT(" ") + Args[1].TrimQuotes();
+		ABTestCmds[1] = Args[0].TrimQuotes() + TEXT(" ") + Args[2].TrimQuotes();
+	}
+	else if (Args.Num() > 2 && Args[0].StartsWith(TEXT("\"")))
+	{
+		FString Work;
+		int32 Test = 0;
+		for (int32 Index = 0; Index < Args.Num(); Index++)
+		{
+			Work += Args[Index];
+			if (Work.Len() > 2 && Work.StartsWith(TEXT("\"")) && Work.EndsWith(TEXT("\"")))
+			{
+				ABTestCmds[Test++] = Work.TrimQuotes();
+				Work.Empty();
+				if (Test > 1)
+				{
+					break;
+				}
+			}
+			else if (Index + 1 < Args.Num())
+			{
+				Work += TEXT(" ");
+			}
+		}
+	}
+	else if (Args.Num() == 2)
+	{
+		ABTestCmds[0] = Args[0].TrimQuotes();
+		ABTestCmds[1] = Args[1].TrimQuotes();
+	}
+	else
+	{
+		UE_LOG(LogConsoleResponse, Display, TEXT("abtest command requires two (quoted) arguments or three args or 'stop' or 'scope'."));
+		UE_LOG(LogConsoleResponse, Display, TEXT("Example: abtest \"r.MyCVar 0\" \"r.MyCVar 1\""));
+		UE_LOG(LogConsoleResponse, Display, TEXT("Example: abtest r.MyCVar 0 1"));
+		return;
+	}
+	Get().Start(ABTestCmds, false);
+}
+
+void FABTest::Stop()
+{
+	if (bABTestActive)
+	{
+		UE_LOG(LogConsoleResponse, Display, TEXT("Running 'A' console command and stopping test."));
+		SwitchTest(0);
+		bABTestActive = false;
+	}
+	else if (bFrameLog)
+	{
+		UE_LOG(LogConsoleResponse, Display, TEXT("Stopping frame log."));
+		bFrameLog = false;
+	}
+	bABScopeTestActive = false;
+}
+
+void FABTest::Start(FString* InABTestCmds, bool bScopeTest)
+{
+	if (bABTestActive)
+	{
+		Stop();
+	}
+
+	if (InABTestCmds)
+	{
+		ABTestCmds[0] = InABTestCmds[0];
+		ABTestCmds[1] = InABTestCmds[1];
+	}
+	else
+	{
+		ABTestCmds[0] = FString("");
+		ABTestCmds[1] = FString("");
+	}
+	
+	bABScopeTestActive = bScopeTest;
+
+	HistoryNum = CVarABTestHistory.GetValueOnGameThread();
+	ReportNum = CVarABTestReportNum.GetValueOnGameThread();
+	CoolDown = CVarABTestCoolDown.GetValueOnGameThread();
+	MinFramesPerTrial = CVarABTestMinFramesPerTrial.GetValueOnGameThread();
+	NumResamples = CVarABTestNumResamples.GetValueOnGameThread();
+
+	Samples.Empty(HistoryNum);
+	ResampleAccumulators.Empty(NumResamples);
+	ResampleAccumulators.AddZeroed(NumResamples);
+	ResampleCount.Empty(NumResamples);
+	ResampleCount.AddZeroed(NumResamples);
+	ABTestNumSamples = 0;
+	Totals[0] = 0;
+	Totals[1] = 0;
+	Counts[0] = 0;
+	Counts[1] = 0;
+	RemainingPrint = ReportNum;
+
+	bABTestActive = true;
+	SwitchTest(0);
+	UE_LOG(LogConsoleResponse, Display, TEXT("abtest started with A = '%s' and B = '%s'"), *ABTestCmds[0], *ABTestCmds[1]);
+}
+
+const TCHAR* FABTest::SwitchTest(int32 Index)
+{
+	RemainingCoolDown = CoolDown;
+	CurrentTest = Index;
+	RemainingTrial = Stream.RandRange(MinFramesPerTrial, MinFramesPerTrial * 3);
+	check(RemainingTrial);
+
+	if (!bABScopeTestActive)
+	{
+		return *(ABTestCmds[CurrentTest]);
+	}
+	return nullptr;
+}
+
+static FAutoConsoleCommand ABTestCmd(
+	TEXT("abtest"),
+	TEXT("Provide two console commands or 'stop' to stop the abtest. Frames are timed with the two options, logging results over time."),
+	FConsoleCommandWithArgsDelegate::CreateStatic(&FABTest::ABTestCmdFunc)
+	);
+
+#endif
+
diff --git a/Engine/Source/Runtime/Core/Public/ProfilingDebugging/ABTesting.h b/Engine/Source/Runtime/Core/Public/ProfilingDebugging/ABTesting.h
new file mode 100644
--- /dev/null
+++ b/Engine/Source/Runtime/Core/Public/ProfilingDebugging/ABTesting.h
@@ -0,0 +1,137 @@
+// Copyright 1998-2015 Epic Games, Inc. All Rights Reserved.
+
+/**
+* Declarations for ABTesting framework.
+*/
+
+#pragma once
+
+#if !UE_BUILD_SHIPPING
+#define ENABLE_ABTEST 1
+#else
+#define ENABLE_ABTEST 0
+#endif
+
+#include "ScopedTimers.h"
+
+#if ENABLE_ABTEST
+class CORE_API FABTest
+{
+public:
+
+	FABTest();
+	
+	//returns a command to execute, if any.
+	const TCHAR* TickAndGetCommand();
+
+	bool IsActive()
+	{
+		return bABTestActive;
+	}
+
+	void ReportScopeTime(double ScopeTime)
+	{
+		TotalScopeTimeInFrame += ScopeTime;
+	}
+
+	bool GetDoFirstScopeTest() const
+	{
+		return CurrentTest == 0;
+	}
+
+	static FABTest& Get();
+
+	static void ABTestCmdFunc(const TArray<FString>& Args);
+
+	static bool StaticIsActive()
+	{
+		return Get().bABTestActive;
+	}
+
+private:
+
+	void StartFrameLog();
+	void FrameLogTick(double Delta);
+	
+	void Start(FString* InABTestCmds, bool bScopeTest);
+	void Stop();
+
+	const TCHAR* SwitchTest(int32 Index);
+	
+	
+
+	FRandomStream Stream;
+	bool bABTestActive;
+	bool bABScopeTestActive; //whether we are doing abtesting on the scope macro rather than a CVAR.	
+	bool bFrameLog;
+	FString ABTestCmds[2];
+	int32 ABTestNumSamples;
+	int32 RemainingCoolDown;
+	int32 CurrentTest;
+	int32 RemainingTrial;
+	int32 RemainingPrint;
+	int32 SampleIndex;
+
+	int32 HistoryNum;
+	int32 ReportNum;
+	int32 CoolDown;
+	int32 MinFramesPerTrial;
+	int32 NumResamples;
+
+	struct FSample
+	{
+		uint32 Micros;
+		int32 TestIndex;
+		TBitArray<> InResamples;
+	};
+
+	TArray<FSample> Samples;
+	TArray<uint32> ResampleAccumulators;
+	TArray<uint32> ResampleCount;
+	uint32 Totals[2];
+	uint32 Counts[2];
+
+	double TotalTime;
+	uint32 TotalFrames;
+	uint32 Spikes;
+
+	double TotalScopeTimeInFrame;
+
+};
+
+class CORE_API FScopedABTimer : public FDurationTimer
+{
+public:
+	explicit FScopedABTimer() 
+		: FDurationTimer(TimerData)
+		, TimerData(0)
+	{
+		Start();
+	}
+
+	/** Dtor, updating seconds with time delta. */
+	~FScopedABTimer()
+	{
+		Stop();
+		FABTest::Get().ReportScopeTime(GetAccumulatedTime());
+	}
+private:
+	double TimerData;
+};
+
+#define SCOPED_ABTEST() FScopedABTimer ABSCOPETIMER;
+#define SCOPED_ABTEST_DOFIRSTTEST() FABTest::Get().GetDoFirstScopeTest()
+#else
+
+class CORE_API FABTest
+{
+public:
+	static bool StaticIsActive()
+	{
+		return false;
+	}
+};
+
+#define SCOPED_ABTEST(TimerName)
+#define SCOPED_ABTEST_DOFIRSTTEST() true
+#endif
\ No newline at end of file
diff --git a/Engine/Source/Runtime/Core/Public/ProfilingDebugging/ScopedTimers.h b/Engine/Source/Runtime/Core/Public/ProfilingDebugging/ScopedTimers.h
--- a/Engine/Source/Runtime/Core/Public/ProfilingDebugging/ScopedTimers.h
+++ b/Engine/Source/Runtime/Core/Public/ProfilingDebugging/ScopedTimers.h
@@ -29,6 +29,11 @@ public:
 		return StopTime;
 	}
 
+	double GetAccumulatedTime() const
+	{
+		return Accumulator;
+	}
+
 private:
 	/** Start time, captured in ctor. */
 	double StartTime;
diff --git a/Engine/Source/Runtime/Engine/Private/UnrealEngine.cpp b/Engine/Source/Runtime/Engine/Private/UnrealEngine.cpp
--- a/Engine/Source/Runtime/Engine/Private/UnrealEngine.cpp
+++ b/Engine/Source/Runtime/Engine/Private/UnrealEngine.cpp
@@ -116,6 +116,7 @@
 #endif
 
 #include "GameFramework/OnlineSession.h"
+#include "ABTesting.h"
 
 #if WITH_EDITOR
 #include "Editor/UnrealEd/Public/Animation/AnimationRecorder.h"
@@ -1061,393 +1062,21 @@ void UEngine::TickDeferredCommands()
 	DeferredCommands.RemoveAt(0, DeferredCommandsCount);
 }
 
-
-#if !UE_BUILD_SHIPPING
-
-static TAutoConsoleVariable<int32> CVarABTestHistory(
-	TEXT("abtest.HistoryNum"),
-	1000,
-	TEXT("Number of history frames to use for stats."));
-
-static TAutoConsoleVariable<int32> CVarABTestReportNum(
-	TEXT("abtest.ReportNum"),
-	100,
-	TEXT("Number of frames between reports."));
-
-static TAutoConsoleVariable<int32> CVarABTestCoolDown(
-	TEXT("abtest.CoolDown"),
-	3,
-	TEXT("Number of frames to discard data after each command to cover threading."));
-
-static TAutoConsoleVariable<int32> CVarABTestMinFramesPerTrial(
-	TEXT("abtest.MinFramesPerTrial"),
-	5,
-	TEXT("The number of frames to run a given command before switching; this is randomized."));
-
-static TAutoConsoleVariable<int32> CVarABTestNumResamples(
-	TEXT("abtest.NumResamples"),
-	256,
-	TEXT("The number of resamples to use to determine confidence."));
-
-struct FABTest
-{
-	FRandomStream Stream;
-	bool bABTestActive;
-	bool bFrameLog;
-	FString ABTestCmds[2];
-	int32 ABTestNumSamples;
-	int32 RemainingCoolDown;
-	int32 CurrentTest;
-	int32 RemainingTrial;
-	int32 RemainingPrint;
-	int32 SampleIndex;
-
-	int32 HistoryNum;
-	int32 ReportNum;
-	int32 CoolDown;
-	int32 MinFramesPerTrial;
-	int32 NumResamples;
-
-	struct FSample
-	{
-		uint32 Micros;
-		int32 TestIndex;
-		TBitArray<> InResamples;
-	};
-
-	TArray<FSample> Samples;
-	TArray<uint32> ResampleAccumulators;
-	TArray<uint32> ResampleCount;
-	uint32 Totals[2];
-	uint32 Counts[2];
-
-	double TotalTime;
-	uint32 TotalFrames;
-	uint32 Spikes;
-
-	FABTest()
-		: Stream(9567)
-		, bABTestActive(false)
-		, bFrameLog(false)
-		, ABTestNumSamples(0)
-		, RemainingCoolDown(0)
-		, CurrentTest(0)
-		, RemainingTrial(0)
-		, RemainingPrint(0)
-		, SampleIndex(0)
-		, HistoryNum(0)
-		, ReportNum(0)
-		, CoolDown(0)
-		, MinFramesPerTrial(0)
-		, NumResamples(0)
-	{
-	}
-
-	void StartFrameLog()
-	{
-		TotalTime = 0.0;
-		TotalFrames = 0;
-		Spikes = 0;
-		bFrameLog = true;
-		UE_LOG(LogConsoleResponse, Display, TEXT("Starting frame log."));
-	}
-
-	void FrameLogTick(double Delta)
-	{
-		if (Delta > .034)
-		{
-			Spikes++;
-		}
-		else
-		{
-			TotalFrames++;
-			TotalTime += Delta;
-		}
-		if (TotalFrames > 0 && TotalFrames % 1000 == 0)
-		{
-			FPlatformMisc::LowLevelOutputDebugStringf(TEXT("%8d frames   %6.3fms/f    %8d spikes rejected "), TotalFrames, float(1000.0 * TotalTime / TotalFrames), Spikes);
-		}
-	}
-	void Tick()
-	{
-		static double LastTime = 0;
-		if (bABTestActive && RemainingCoolDown)
-		{
-			RemainingCoolDown--;
-		}
-		else if (bABTestActive)
-		{
-			float Delta = FPlatformTime::Seconds() - LastTime;
-			check(Delta > 0);
-			FSample* Sample = nullptr;
-			if (ABTestNumSamples < HistoryNum)
-			{
-				check(ABTestNumSamples == Samples.Num());
-				Sample = new (Samples) FSample();
-			}
-			else
-			{
-				Sample = &Samples[ABTestNumSamples % HistoryNum];
-
-				// remove the old data
-				check(Totals[Sample->TestIndex] > Sample->Micros);
-				Totals[Sample->TestIndex] -= Sample->Micros;
-				check(Counts[Sample->TestIndex]);
-				Counts[Sample->TestIndex]--;
-				check(Sample->InResamples.Num() == NumResamples);
-				for (TConstSetBitIterator<> CurIndexIt(Sample->InResamples); CurIndexIt; ++CurIndexIt)
-				{
-					int32 ResampleIndex = CurIndexIt.GetIndex();
-					check(ResampleAccumulators[ResampleIndex] > Sample->Micros);
-					ResampleAccumulators[ResampleIndex] -= Sample->Micros;
-					check(ResampleCount[ResampleIndex]);
-					ResampleCount[ResampleIndex]--;
-				}
-			}
-			ABTestNumSamples++;
-			Sample->Micros = uint32(Delta * 1000000.0f);
-			check(Sample->Micros);
-			Sample->TestIndex = CurrentTest;
-			Totals[Sample->TestIndex] += Sample->Micros;
-			Counts[Sample->TestIndex]++;
-			Sample->InResamples.Empty(NumResamples);
-			for (int32 ResampleIndex = 0; ResampleIndex < NumResamples; ResampleIndex++)
-			{
-				bool In = Stream.FRand() > 0.5f;
-				Sample->InResamples.Add(In);
-				if (In)
-				{
-					ResampleAccumulators[ResampleIndex] += Sample->Micros;
-					ResampleCount[ResampleIndex]++;
-				}
-			}
-			check(RemainingPrint);
-			if (!--RemainingPrint)
-			{
-				if (Counts[0] && Counts[1])
-				{
-					UE_LOG(LogConsoleResponse, Display, TEXT("      %7.4fms  (%4d samples)  A = '%s'"), float(Totals[0]) / float(Counts[0]) / 1000.0f, Counts[0], *ABTestCmds[0]);
-					UE_LOG(LogConsoleResponse, Display, TEXT("      %7.4fms  (%4d samples)  B = '%s'"), float(Totals[1]) / float(Counts[1]) / 1000.0f, Counts[1], *ABTestCmds[1]);
-
-					float Diff = (float(Totals[0]) / float(Counts[0]) / 1000.0f) - (float(Totals[1]) / float(Counts[1]) / 1000.0f);
-					bool bAIsFaster = false;
-					if (Diff < 0.0f)
-					{
-						Diff *= -1.0f;
-						bAIsFaster = true;
-					}
-					static TArray<float> Scores;
-					Scores.Reset();
-					uint32 TotalFrame = Totals[0] + Totals[1];
-					uint32 TotalCount = Counts[0] + Counts[1];
-					for (int32 ResampleIndex = 0; ResampleIndex < NumResamples; ResampleIndex++)
-					{
-						check(TotalFrame > ResampleAccumulators[ResampleIndex]);
-						uint32 Opposite = TotalFrame - ResampleAccumulators[ResampleIndex];
-						check(TotalCount > ResampleCount[ResampleIndex]);
-						uint32 OppositeCount = TotalCount - ResampleCount[ResampleIndex];
-
-						// this ABS here gives us better data but makes it one sided, so the 90% confidence interval is at index 80/100 etc
-						// it is legit because we could have easily done the inverse permutation
-						float SampleDiff = FMath::Abs(((ResampleAccumulators[ResampleIndex]) / float(ResampleCount[ResampleIndex]) / 1000.0f) - (float(Opposite) / float(OppositeCount) / 1000.0f));
-						Scores.Add(SampleDiff);
-
-					}
-					Scores.Sort();
-					int32 Conf = NumResamples;
-					for (int32 Trial = 0; Trial < NumResamples; Trial++)
-					{
-						if (Scores[Trial] > Diff)
-						{
-							Conf = Trial;
-							break;
-						}
-					}
-
-					float fConf = 1.0f - (0.5f + float(Conf) / float(NumResamples) / 2.0f);
-
-					if (bAIsFaster)
-					{
-						UE_LOG(LogConsoleResponse, Display, TEXT("      A is %7.4fms faster than B;  %3.0f%% chance this is noise."), Diff, fConf * 100.0f);
-					}
-					else
-					{
-						UE_LOG(LogConsoleResponse, Display, TEXT("      B is %7.4fms faster than A;  %3.0f%% chance this is noise."), Diff, fConf * 100.0f);
-					}
-
-					UE_LOG(LogConsoleResponse, Display, TEXT("----------------"));
-				}
-				else
-				{
-					UE_LOG(LogConsoleResponse, Display, TEXT("No Samples?"));
-				}
-				RemainingPrint = ReportNum;
-			}
-			check(RemainingTrial);
-			if (!--RemainingTrial)
-			{
-				SwitchTest(1 - CurrentTest);
-			}
-		}
-		else if (bFrameLog)
-		{
-			FrameLogTick(FPlatformTime::Seconds() - LastTime);
-		}
-		LastTime = FPlatformTime::Seconds();
-	}
-
-	static FABTest& Get()
-	{
-		static FABTest Singleton;
-		return Singleton;
-	}
-	static void ABTestCmdFunc(const TArray<FString>& Args)
-	{
-		FString ABTestCmds[2];
-		if (Args.Num() == 1 && Args[0].Compare(FString(TEXT("framelog")), ESearchCase::IgnoreCase) == 0)
-		{
-			Get().Stop();
-			Get().StartFrameLog();
-			return;
-		}
-		if (Args.Num() == 1 && Args[0].Compare(FString(TEXT("stop")), ESearchCase::IgnoreCase) == 0)
-		{
-			Get().Stop();
-			return;
-		}
-		if (Args.Num() == 3 && !Args[0].StartsWith(TEXT("\"")))
-		{
-			ABTestCmds[0] = Args[0].TrimQuotes() + TEXT(" ") + Args[1].TrimQuotes();
-			ABTestCmds[1] = Args[0].TrimQuotes() + TEXT(" ") + Args[2].TrimQuotes();
-		}
-		else if (Args.Num() > 2 && Args[0].StartsWith(TEXT("\"")))
-		{
-			FString Work;
-			int32 Test = 0;
-			for (int32 Index = 0; Index < Args.Num(); Index++)
-			{
-				Work += Args[Index];
-				if (Work.Len() > 2 && Work.StartsWith(TEXT("\"")) && Work.EndsWith(TEXT("\"")))
-				{
-					ABTestCmds[Test++] = Work.TrimQuotes();
-					Work.Empty();
-					if (Test > 1)
-					{
-						break;
-					}
-				}
-				else if (Index + 1 < Args.Num())
-				{
-					Work += TEXT(" ");
-				}
-			}
-		}
-		else if (Args.Num() == 2)
-		{
-			ABTestCmds[0] = Args[0].TrimQuotes();
-			ABTestCmds[1] = Args[1].TrimQuotes();
-		}
-		else
-		{
-			UE_LOG(LogConsoleResponse, Display, TEXT("abtest command requires two (quoted) arguments or three args or 'stop'."));
-			UE_LOG(LogConsoleResponse, Display, TEXT("Example: abtest \"r.MyCVar 0\" \"r.MyCVar 1\""));
-			UE_LOG(LogConsoleResponse, Display, TEXT("Example: abtest r.MyCVar 0 1"));
-			return;
-		}
-		Get().Start(ABTestCmds);
-	}
-
-	void Stop()
-	{
-		if (bABTestActive)
-		{
-			UE_LOG(LogConsoleResponse, Display, TEXT("Running 'A' console command and stopping test."));
-			SwitchTest(0);
-			bABTestActive = false;
-		}
-		else if (bFrameLog)
-		{
-			UE_LOG(LogConsoleResponse, Display, TEXT("Stopping frame log."));
-			bFrameLog = false;
-		}
-	}
-	
-	void Start(FString* InABTestCmds)
-	{
-		if (bABTestActive)
-		{
-			Stop();
-		}
-		ABTestCmds[0] = InABTestCmds[0];
-		ABTestCmds[1] = InABTestCmds[1];
-
-		HistoryNum = CVarABTestHistory.GetValueOnGameThread();
-		ReportNum = CVarABTestReportNum.GetValueOnGameThread();
-		CoolDown = CVarABTestCoolDown.GetValueOnGameThread();
-		MinFramesPerTrial = CVarABTestMinFramesPerTrial.GetValueOnGameThread();
-		NumResamples = CVarABTestNumResamples.GetValueOnGameThread();
-
-		Samples.Empty(HistoryNum);
-		ResampleAccumulators.Empty(NumResamples);
-		ResampleAccumulators.AddZeroed(NumResamples);
-		ResampleCount.Empty(NumResamples);
-		ResampleCount.AddZeroed(NumResamples);
-		ABTestNumSamples = 0;
-		Totals[0] = 0;
-		Totals[1] = 0;
-		Counts[0] = 0;
-		Counts[1] = 0;
-		RemainingPrint = ReportNum;
-
-		bABTestActive = true;
-		SwitchTest(0);
-		UE_LOG(LogConsoleResponse, Display, TEXT("abtest started with A = '%s' and B = '%s'"), *ABTestCmds[0], *ABTestCmds[1]);
-	}
-
-	void SwitchTest(int32 Index)
-	{
-		RemainingCoolDown = CoolDown;
-		CurrentTest = Index;
-		RemainingTrial = Stream.RandRange(MinFramesPerTrial, MinFramesPerTrial * 3);
-		check(RemainingTrial);
-		GEngine->Exec(nullptr, *ABTestCmds[CurrentTest], *GLog);
-	}
-
-	static void StaticTick()
-	{
-		Get().Tick();
-	}
-	static bool IsActive()
-	{
-		return Get().bABTestActive;
-	}
-};
-
-static FAutoConsoleCommand ABTestCmd(
-	TEXT("abtest"),
-	TEXT("Provide two console commands or 'stop' to stop the abtest. Frames are timed with the two options, logging results over time."),
-	FConsoleCommandWithArgsDelegate::CreateStatic(&FABTest::ABTestCmdFunc)
-	);
-
-#else
-
-struct FABTest
+void PumpABTest()
 {
-	FORCEINLINE static void StaticTick()
+#if ENABLE_ABTEST
+	const TCHAR* Command = FABTest::Get().TickAndGetCommand();
+	if (Command)
 	{
+		GEngine->Exec(nullptr, Command, *GLog);
 	}
-	FORCEINLINE static bool IsActive()
-	{
-		return false;
-	}
-};
-
 #endif
-
+}
 
 void UEngine::UpdateTimeAndHandleMaxTickRate()
 {
-	FABTest::StaticTick();
+	PumpABTest();
+
 	// start at now minus a bit so we don't get a zero delta.
 	static double LastTime = FPlatformTime::Seconds() - 0.0001;
 	static bool bTimeWasManipulated = false;
@@ -1496,7 +1125,7 @@ void UEngine::UpdateTimeAndHandleMaxTickRate()
 		UpdateRunningAverageDeltaTime(DeltaTime);
 
 		// Get max tick rate based on network settings and current delta time.
-		const float MaxTickRate	= FABTest::IsActive() ? 0.0f : GetMaxTickRate( DeltaTime );
+		const float MaxTickRate = FABTest::StaticIsActive() ? 0.0f : GetMaxTickRate(DeltaTime);
 		float WaitTime		= 0;
 		// Convert from max FPS to wait time.
 		if( MaxTickRate > 0 )
