Merging CL 2691446 to Main

Expose async load priority through streamable manager

diff --git a/Engine/Source/Runtime/Engine/Classes/Engine/StreamableManager.h b/Engine/Source/Runtime/Engine/Classes/Engine/StreamableManager.h
--- a/Engine/Source/Runtime/Engine/Classes/Engine/StreamableManager.h
+++ b/Engine/Source/Runtime/Engine/Classes/Engine/StreamableManager.h
@@ -11,6 +11,9 @@ struct ENGINE_API FStreamableManager
 {
 	GENERATED_USTRUCT_BODY();
 
+	// Default priority for all async loads
+	static const TAsyncLoadPriority DefaultAsyncLoadPriority = 0;
+
 	FStreamableManager();
 	~FStreamableManager();
 
@@ -30,7 +33,7 @@ struct ENGINE_API FStreamableManager
 	 * Perform a simple asynchronous load of a single object.
 	 * Object will be strongly referenced until Unload is manually called on it
 	 */
-	void SimpleAsyncLoad(FStringAssetReference const& Target);
+	void SimpleAsyncLoad(FStringAssetReference const& Target, TAsyncLoadPriority Priority = DefaultAsyncLoadPriority);
 
 	/** 
 	 * Releases a reference to an asynchonously loaded object, that was loaded by SimpleAsyncLoad
@@ -46,8 +49,8 @@ struct ENGINE_API FStreamableManager
 	 * Request streaming of one or more target objects, and call a delegate on completion. 
 	 * Objects will be strongly referenced until the delegate is called, then Unload is called automatically 
 	 */
-	void RequestAsyncLoad(const TArray<FStringAssetReference>& TargetsToStream, FStreamableDelegate DelegateToCall);
-	void RequestAsyncLoad( const FStringAssetReference& TargetToStream, FStreamableDelegate DelegateToCall );
+	void RequestAsyncLoad(const TArray<FStringAssetReference>& TargetsToStream, FStreamableDelegate DelegateToCall, TAsyncLoadPriority Priority = DefaultAsyncLoadPriority);
+	void RequestAsyncLoad(const FStringAssetReference& TargetToStream, FStreamableDelegate DelegateToCall, TAsyncLoadPriority Priority = DefaultAsyncLoadPriority);
 
 	/** Exposes references to GC system */
 	void AddStructReferencedObjects(class FReferenceCollector& Collector) const;
@@ -55,7 +58,7 @@ struct ENGINE_API FStreamableManager
 private:
 	FStringAssetReference ResolveRedirects(FStringAssetReference const& Target) const;
 	void FindInMemory(FStringAssetReference& InOutTarget, struct FStreamable* Existing);
-	struct FStreamable* StreamInternal(FStringAssetReference const& Target);
+	struct FStreamable* StreamInternal(FStringAssetReference const& Target, TAsyncLoadPriority Priority);
 	UObject* GetStreamed(FStringAssetReference const& Target);
 	void CheckCompletedRequests(FStringAssetReference const& Target, struct FStreamable* Existing);
 
diff --git a/Engine/Source/Runtime/Engine/Private/StreamableManager.cpp b/Engine/Source/Runtime/Engine/Private/StreamableManager.cpp
--- a/Engine/Source/Runtime/Engine/Private/StreamableManager.cpp
+++ b/Engine/Source/Runtime/Engine/Private/StreamableManager.cpp
@@ -179,7 +179,7 @@ UObject* FStreamableManager::SynchronousLoad(FStringAssetReference const& InTarg
 	return Existing->Target;
 }
 
-struct FStreamable* FStreamableManager::StreamInternal(FStringAssetReference const& InTargetName)
+struct FStreamable* FStreamableManager::StreamInternal(FStringAssetReference const& InTargetName, TAsyncLoadPriority Priority)
 {
 	check(IsInGameThread());
 	UE_LOG(LogStreamableManager, Verbose, TEXT("Asynchronous load %s"), *InTargetName.ToString());
@@ -233,17 +233,17 @@ struct FStreamable* FStreamableManager::StreamInternal(FStringAssetReference con
 		}
 
 		Existing->bAsyncLoadRequestOutstanding = true;
-		LoadPackageAsync(Package, FLoadPackageAsyncDelegate::CreateStatic(&AsyncLoadCallbackWrapper, new FCallback(TargetName, this)));
+		LoadPackageAsync(Package, FLoadPackageAsyncDelegate::CreateStatic(&AsyncLoadCallbackWrapper, new FCallback(TargetName, this)), Priority);
 	}
 	return Existing;
 }
 
-void FStreamableManager::SimpleAsyncLoad(FStringAssetReference const& InTargetName)
+void FStreamableManager::SimpleAsyncLoad(FStringAssetReference const& InTargetName, TAsyncLoadPriority Priority)
 {
-	StreamInternal(InTargetName);
+	StreamInternal(InTargetName, Priority);
 }
 
-void FStreamableManager::RequestAsyncLoad(const TArray<FStringAssetReference>& TargetsToStream, FStreamableDelegate DelegateToCall)
+void FStreamableManager::RequestAsyncLoad(const TArray<FStringAssetReference>& TargetsToStream, FStreamableDelegate DelegateToCall, TAsyncLoadPriority Priority)
 {
 	// Schedule a new callback, this will get called when all related async loads are completed
 	TSharedRef<FStreamableRequest> NewRequest = MakeShareable(new FStreamableRequest());
@@ -254,7 +254,7 @@ void FStreamableManager::RequestAsyncLoad(const TArray<FStringAssetReference>& T
 
 	for (int32 i = 0; i < TargetsToStream.Num(); i++)
 	{
-		FStreamable* Existing = StreamInternal(TargetsToStream[i]);
+		FStreamable* Existing = StreamInternal(TargetsToStream[i], Priority);
 		ExistingStreamables[i] = Existing;
 
 		if (Existing)
@@ -275,12 +275,12 @@ void FStreamableManager::RequestAsyncLoad(const TArray<FStringAssetReference>& T
 	}
 }
 
-void FStreamableManager::RequestAsyncLoad( const FStringAssetReference& TargetToStream, FStreamableDelegate DelegateToCall )
+void FStreamableManager::RequestAsyncLoad(const FStringAssetReference& TargetToStream, FStreamableDelegate DelegateToCall, TAsyncLoadPriority Priority)
 {
 	TSharedRef< FStreamableRequest > NewRequest = MakeShareable( new FStreamableRequest() );
 	NewRequest->CompletionDelegate = DelegateToCall;
 
-	if ( FStreamable* Streamable = StreamInternal( TargetToStream ) )
+	if ( FStreamable* Streamable = StreamInternal( TargetToStream, Priority ) )
 	{
 		Streamable->AddRelatedRequest( NewRequest );
 
