Fix StartCloth not waiting on cloth prepare tasks

diff --git a/Engine/Source/Runtime/Engine/Private/PhysicsEngine/ClothManager.cpp b/Engine/Source/Runtime/Engine/Private/PhysicsEngine/ClothManager.cpp
--- a/Engine/Source/Runtime/Engine/Private/PhysicsEngine/ClothManager.cpp
+++ b/Engine/Source/Runtime/Engine/Private/PhysicsEngine/ClothManager.cpp
@@ -79,6 +79,7 @@ void FClothManager::StartCloth()
 			bNeedSimulateCloth = true;
 			PrepareData.SkeletalMeshComponents.Reset();
 		}
+		PrepareData.PrepareCompletion.SafeRelease();
 	}
 
 	if (bNeedSimulateCloth)
@@ -125,7 +126,7 @@ public:
 	}
 	static ESubsequentsMode::Type GetSubsequentsMode()
 	{
-		return ESubsequentsMode::FireAndForget;
+		return ESubsequentsMode::TrackSubsequents;
 	}
 
 	void DoTask(ENamedThreads::Type CurrentThread, const FGraphEventRef& MyCompletionGraphEvent)
