Removed unnecessary code from UBlueprintGeneratedClass::CreateComponentsForActor

diff --git a/Engine/Source/Runtime/Engine/Private/BlueprintGeneratedClass.cpp b/Engine/Source/Runtime/Engine/Private/BlueprintGeneratedClass.cpp
--- a/Engine/Source/Runtime/Engine/Private/BlueprintGeneratedClass.cpp
+++ b/Engine/Source/Runtime/Engine/Private/BlueprintGeneratedClass.cpp
@@ -568,28 +568,15 @@ void UBlueprintGeneratedClass::CreateComponentsForActor(const UClass* ThisClass,
 			}
 		}
 	}
-	else
+	else if (auto DynamicClass = Cast<UDynamicClass>(ThisClass))
 	{
-		TInlineComponentArray<UActorComponent*> Components;
-		Actor->GetComponents(Components);
-		for (auto Component : Components)
-		{
-			if (Component->IsCreatedByConstructionScript())
-			{
-				Actor->BlueprintCreatedComponents.Add(Component);
-			}
-		}
-
-		if (auto DynamicClass = Cast<UDynamicClass>(ThisClass))
+		for (auto MiscObj : DynamicClass->Timelines)
 		{
-			for (auto MiscObj : DynamicClass->Timelines)
+			auto TimelineTemplate = Cast<const UTimelineTemplate>(MiscObj);
+			// Not fatal if NULL, but shouldn't happen and ignored if not wired up in graph
+			if (TimelineTemplate && TimelineTemplate->bValidatedAsWired)
 			{
-				auto TimelineTemplate = Cast<const UTimelineTemplate>(MiscObj);
-				// Not fatal if NULL, but shouldn't happen and ignored if not wired up in graph
-				if (TimelineTemplate && TimelineTemplate->bValidatedAsWired)
-				{
-					CreateTimelineComponent(Actor, TimelineTemplate);
-				}
+				CreateTimelineComponent(Actor, TimelineTemplate);
 			}
 		}
 	}
