#ue4 - AnimNotifies for trail effects can recycle ParticleSystemComponents, to avoid spawning them each time and avoid the template setup cost. This is enabled by default.

- Trails are not limited to PSCs spawned by the trails (though those are preferred), so a persistent effect can be spawned for an actor and reused forever.
- Optimized some internal functions to avoid dynamic memory allocations.
- Optimized editor code for trails to avoid creating FText for every update.

--------

diff --git a/Engine/Source/Runtime/Engine/Classes/Animation/AnimNotifies/AnimNotifyState_Trail.h b/Engine/Source/Runtime/Engine/Classes/Animation/AnimNotifies/AnimNotifyState_Trail.h
--- a/Engine/Source/Runtime/Engine/Classes/Animation/AnimNotifies/AnimNotifyState_Trail.h
+++ b/Engine/Source/Runtime/Engine/Classes/Animation/AnimNotifies/AnimNotifyState_Trail.h
@@ -39,6 +39,9 @@ class UAnimNotifyState_Trail : public UAnimNotifyState
 	UPROPERTY(EditAnywhere, BlueprintReadWrite, Category = Trail)
 	FName WidthScaleCurve;
 
+	UPROPERTY(EditAnywhere, BlueprintReadOnly, Category = Trail)
+	uint32 bRecycleSpawnedSystems:1;
+
 #if WITH_EDITORONLY_DATA
 	/** If true, render the trail geometry (this should typically be on) */
 	UPROPERTY(transient, EditAnywhere, Category = Rendering)
diff --git a/Engine/Source/Runtime/Engine/Classes/Particles/ParticleSystemComponent.h b/Engine/Source/Runtime/Engine/Classes/Particles/ParticleSystemComponent.h
--- a/Engine/Source/Runtime/Engine/Classes/Particles/ParticleSystemComponent.h
+++ b/Engine/Source/Runtime/Engine/Classes/Particles/ParticleSystemComponent.h
@@ -296,6 +296,14 @@ public:
 	/** This flag will be set the first time the PSysComp is activated... used to prevent auto activated PSysComps from calling InitParticles twice on level load */
 	uint32 bHasBeenActivated:1;
 
+	/**
+	 * If true, this Particle System will be available for recycling after it has completed. Auto-destroyed systems cannot be recycled.
+	 * Some systems (currently particle trail effects) can recycle components to avoid respawning them to play new effects.
+	 * This is only an optimization and does not change particle system behavior, aside from not triggering normal component initialization events more than once.
+	 */
+	UPROPERTY(EditAnywhere, BlueprintReadWrite, AdvancedDisplay, Category=Particles)
+	uint32 bAllowRecycling:1;
+
 	/**
 	 * True if we should automatically attach to AutoAttachParent when activated, and detach from our parent when completed.
 	 * This overrides any current attachment that may be present at the time of activation (deferring initial attachment until activation, if AutoAttachParent is null).
@@ -729,13 +737,17 @@ public:
 	UFUNCTION(BlueprintCallable, Category="Effects|Components|ParticleSystem")
 	int32 GetNumActiveParticles() const;
 
+
+	/** Array of trail emitters. */
+	typedef TArray< struct FParticleAnimTrailEmitterInstance*, TInlineAllocator<8> > TrailEmitterArray;
+
 	/**
 	* Fills the passed array with all trail emitters associated with a particular object.
 	* @param OutTrailEmitters	The array to fill with pointers to the trail emitters.
 	* @param InOwner			The object that triggered this trail. Can be NULL if no assosiation was set by the owner. Not to be confused with the result of GetOwner().
 	* @param bSetOwner			If true, all trail emitters will be set as owned by InOwner.
 	*/
-	virtual void GetOwnedTrailEmitters(TArray< struct FParticleAnimTrailEmitterInstance* >& OutTrailEmitters, const void* InOwner, bool bSetOwner = false);
+	virtual void GetOwnedTrailEmitters(TrailEmitterArray& OutTrailEmitters, const void* InOwner, bool bSetOwner = false);
 	
 	/**
 	* Begins all trail emitters in this component.
diff --git a/Engine/Source/Runtime/Engine/Private/Animation/AnimNotifyState_Trail.cpp b/Engine/Source/Runtime/Engine/Private/Animation/AnimNotifyState_Trail.cpp
--- a/Engine/Source/Runtime/Engine/Private/Animation/AnimNotifyState_Trail.cpp
+++ b/Engine/Source/Runtime/Engine/Private/Animation/AnimNotifyState_Trail.cpp
@@ -14,6 +14,29 @@
 
 DEFINE_LOG_CATEGORY(LogAnimTrails);
 
+typedef TInlineComponentArray<UParticleSystemComponent*, 8> ParticleSystemComponentArray;
+
+static void GetCandidateSystems(USkeletalMeshComponent& MeshComp, ParticleSystemComponentArray& Components)
+{
+	if (AActor* Owner = MeshComp.GetOwner())
+	{
+		Owner->GetComponents(Components);
+	}
+	else
+	{
+		// No actor owner in some editor windows. Get PSCs spawned by the MeshComp.
+		TArray<UObject*> Children;
+		GetObjectsWithOuter(&MeshComp, Children, false, RF_PendingKill);
+		for (UObject* Child : Children)
+		{
+			if (UParticleSystemComponent* ChildPSC = Cast<UParticleSystemComponent>(Child))
+			{
+				Components.Add(ChildPSC);
+			}
+		}
+	}
+}
+
 /////////////////////////////////////////////////////
 // UAnimNotifyState_Trail
 
@@ -26,6 +49,8 @@ UAnimNotifyState_Trail::UAnimNotifyState_Trail(const FObjectInitializer& ObjectI
 	WidthScaleMode = ETrailWidthMode_FromCentre;
 	WidthScaleCurve = NAME_None;
 
+	bRecycleSpawnedSystems = true;
+
 #if WITH_EDITORONLY_DATA
 	bRenderGeometry = true;
 	bRenderSpawnPoints = false;
@@ -43,9 +68,9 @@ void UAnimNotifyState_Trail::NotifyBegin(class USkeletalMeshComponent * MeshComp
 		return;
 	}
 
-	TArray<USceneComponent*> Children;
-	MeshComp->GetChildrenComponents(false, Children);
-	
+	ParticleSystemComponentArray Children;
+	GetCandidateSystems(*MeshComp, Children);
+
 	UAnimInstance* AnimInst = MeshComp->GetAnimInstance();
 	float Width = 1.0f;
 	if (WidthScaleCurve != NAME_None && AnimInst)
@@ -59,61 +84,87 @@ void UAnimNotifyState_Trail::NotifyBegin(class USkeletalMeshComponent * MeshComp
 		PSTemplate = ParticleSystemTemplate;
 	}
 
+	UParticleSystemComponent* RecycleCandidates[3] = {nullptr, nullptr, nullptr}; // in order of priority
 	bool bFoundExistingTrail = false;
-	for (USceneComponent* Comp : Children)
+	for (UParticleSystemComponent* ParticleComp : Children)
 	{
-		UParticleSystemComponent* ParticleComp = Cast<UParticleSystemComponent>(Comp);
-		if (ParticleComp)
+		if (ParticleComp->IsActive())
 		{
-			TArray< FParticleAnimTrailEmitterInstance* > TrailEmitters;
+			UParticleSystemComponent::TrailEmitterArray TrailEmitters;
 			ParticleComp->GetOwnedTrailEmitters(TrailEmitters, this, false);
 
-			bFoundExistingTrail = TrailEmitters.Num() > 0;
-
-			//If there are any trails, ensure the template hasn't been changed. Also destroy the component if there are errors.
-			if (TrailEmitters.Num() > 0 && (bError || PSTemplate != ParticleComp->Template))
+			if (TrailEmitters.Num() > 0)
 			{
-				//The PSTemplate was changed so we need to destroy this system and create it again with the new template. May be able to just change the template?
-				ParticleComp->DestroyComponent();
-			}
-			else
-			{
-				for (FParticleAnimTrailEmitterInstance* Trail : TrailEmitters)
-				{
-					Trail->BeginTrail();
-					Trail->SetTrailSourceData(FirstSocketName, SecondSocketName, WidthScaleMode, Width);
+				// This has active emitters, we'll just restart this one.
+				bFoundExistingTrail = true;
 
-#if WITH_EDITORONLY_DATA
-					Trail->SetTrailDebugData( bRenderGeometry, bRenderSpawnPoints, bRenderTessellation, bRenderTangents);
-#endif
+				//If there are any trails, ensure the template hasn't been changed. Also destroy the component if there are errors.
+				if (bError || (PSTemplate != ParticleComp->Template && ParticleComp->GetOuter() == MeshComp))
+				{
+					//The PSTemplate was changed so we need to destroy this system and create it again with the new template. May be able to just change the template?
+					ParticleComp->DestroyComponent();
+				}
+				else
+				{
+					for (FParticleAnimTrailEmitterInstance* Trail : TrailEmitters)
+					{
+						Trail->BeginTrail();
+						Trail->SetTrailSourceData(FirstSocketName, SecondSocketName, WidthScaleMode, Width);
+
+	#if WITH_EDITORONLY_DATA
+						Trail->SetTrailDebugData(bRenderGeometry, bRenderSpawnPoints, bRenderTessellation, bRenderTangents);
+	#endif
+					}
 				}
-			}
 
-			if (bFoundExistingTrail)
+				break;
+			}
+		}
+		else if (ParticleComp->bAllowRecycling && !ParticleComp->IsActive())
+		{
+			// We prefer to recycle one with a matching template, and prefer one created by us.
+			// 0: matching template, owned by mesh
+			// 1: matching template, owned by actor
+			// 2: non-matching template, owned by actor or mesh
+			int32 RecycleIndex = 2;
+			if (ParticleComp->Template == PSTemplate)
 			{
-				break;//Found and re-used the existing trail so can exit now.
+				RecycleIndex = (ParticleComp->GetOuter() == MeshComp ? 0 : 1);
 			}
+			RecycleCandidates[RecycleIndex] = ParticleComp;
 		}
 	}
 
 	if (!bFoundExistingTrail && !bError)
 	{
-		//Spawn a new component from PSTemplate. This notify is made the outer so that the component can be identified later.
-		UParticleSystemComponent* NewParticleComp = NewObject<UParticleSystemComponent>(MeshComp);
-		NewParticleComp->bAutoDestroy = true;
+		// Spawn a new component from PSTemplate, or recycle an old one.
+		UParticleSystemComponent* RecycleComponent = (RecycleCandidates[0] ? RecycleCandidates[0] : (RecycleCandidates[1] ? RecycleCandidates[1] : RecycleCandidates[2]));
+		UParticleSystemComponent* NewParticleComp = (RecycleComponent ? RecycleComponent : NewObject<UParticleSystemComponent>(MeshComp));
+		NewParticleComp->bAutoDestroy = (RecycleComponent ? false : !bRecycleSpawnedSystems);
+		NewParticleComp->bAllowRecycling = true;
 		NewParticleComp->SecondsBeforeInactive = 0.0f;
 		NewParticleComp->bAutoActivate = false;
-		NewParticleComp->SetTemplate(PSTemplate);
 		NewParticleComp->bOverrideLODMethod = false;
+		NewParticleComp->RelativeScale3D = FVector(1.f);
+		NewParticleComp->bAutoManageAttachment = true; // Let it detach when finished (only happens if not auto-destroying)
+		NewParticleComp->SetAutoAttachParams(MeshComp, NAME_None);
 
-		NewParticleComp->RegisterComponentWithWorld(MeshComp->GetWorld());
-		NewParticleComp->AttachTo(MeshComp, NAME_None);
-		//NewParticleComp->SetWorldLocationAndRotation(Location, Rotation);
+		// When recycling we can avoid setting the template if set already.
+		if (NewParticleComp->Template != PSTemplate)
+		{
+			NewParticleComp->SetTemplate(PSTemplate);
+		}
+
+		// Recycled components are usually already registered
+		if (!NewParticleComp->IsRegistered())
+		{
+			NewParticleComp->RegisterComponentWithWorld(MeshComp->GetWorld());
+		}
 
-		NewParticleComp->SetRelativeScale3D(FVector(1.f));
+		NewParticleComp->AttachTo(MeshComp, NAME_None);
 		NewParticleComp->ActivateSystem(true);
 
-		TArray< FParticleAnimTrailEmitterInstance* > TrailEmitters;
+		UParticleSystemComponent::TrailEmitterArray TrailEmitters;
 		NewParticleComp->GetOwnedTrailEmitters(TrailEmitters, this, true);
 
 		for (FParticleAnimTrailEmitterInstance* Trail : TrailEmitters)
@@ -139,8 +190,8 @@ void UAnimNotifyState_Trail::NotifyTick(class USkeletalMeshComponent * MeshComp,
 		return;
 	}
 
-	TArray<USceneComponent*> Children;
-	MeshComp->GetChildrenComponents(false, Children);
+	ParticleSystemComponentArray Children;
+	GetCandidateSystems(*MeshComp, Children);
 
 	UAnimInstance* AnimInst = MeshComp->GetAnimInstance();
 	float Width = 1.0f;
@@ -149,16 +200,15 @@ void UAnimNotifyState_Trail::NotifyTick(class USkeletalMeshComponent * MeshComp,
 		Width = AnimInst->GetCurveValue(WidthScaleCurve);
 	}
 
-	for (USceneComponent* Comp : Children)
+	for (UParticleSystemComponent* ParticleComp : Children)
 	{
-		UParticleSystemComponent* ParticleComp = Cast<UParticleSystemComponent>(Comp);
-		if ( ParticleComp )
+		if (ParticleComp->IsActive())
 		{
-			TArray< FParticleAnimTrailEmitterInstance* > TrailEmitters;
+			UParticleSystemComponent::TrailEmitterArray TrailEmitters;
 			ParticleComp->GetOwnedTrailEmitters(TrailEmitters, this, false);
 			if (bError && TrailEmitters.Num() > 0)
 			{
-				Comp->DestroyComponent();
+				ParticleComp->DestroyComponent();
 			}
 			else
 			{
@@ -166,9 +216,9 @@ void UAnimNotifyState_Trail::NotifyTick(class USkeletalMeshComponent * MeshComp,
 				{
 					Trail->SetTrailSourceData(FirstSocketName, SecondSocketName, WidthScaleMode, Width);
 
-#if WITH_EDITORONLY_DATA
+	#if WITH_EDITORONLY_DATA
 					Trail->SetTrailDebugData(bRenderGeometry, bRenderSpawnPoints, bRenderTessellation, bRenderTangents);
-#endif
+	#endif
 				}
 			}
 		}
@@ -184,15 +234,14 @@ void UAnimNotifyState_Trail::NotifyEnd(class USkeletalMeshComponent * MeshComp,
 		return;
 	}
 
-	TArray<USceneComponent*> Children;
-	MeshComp->GetChildrenComponents(false, Children);
+	ParticleSystemComponentArray Children;
+	GetCandidateSystems(*MeshComp, Children);
 
-	for (USceneComponent* Comp : Children)
+	for (UParticleSystemComponent* ParticleComp : Children)
 	{
-		UParticleSystemComponent* ParticleComp = Cast<UParticleSystemComponent>(Comp);
-		if (ParticleComp)
+		if (ParticleComp->IsActive())
 		{
-			TArray< FParticleAnimTrailEmitterInstance* > TrailEmitters;
+			UParticleSystemComponent::TrailEmitterArray TrailEmitters;
 			ParticleComp->GetOwnedTrailEmitters(TrailEmitters, this, false);
 			for (FParticleAnimTrailEmitterInstance* Trail : TrailEmitters)
 			{
@@ -208,11 +257,6 @@ bool UAnimNotifyState_Trail::ValidateInput(class USkeletalMeshComponent * MeshCo
 {
 #if WITH_EDITOR
 	bool bError = false;
-	FText FirstSocketEqualsNoneErrorText = FText::Format( LOCTEXT("NoneFirstSocket", "{0}: Must set First Socket Name."), FText::FromString(GetName()));
-	FText SecondSocketEqualsNoneErrorText = FText::Format( LOCTEXT("NoneSecondSocket", "{0}: Must set Second Socket Name."), FText::FromString(GetName()));
-	FText PSTemplateEqualsNoneErrorText = FText::Format( LOCTEXT("NonePSTemplate", "{0}: Trail must have a PSTemplate."), FText::FromString(GetName()));
-	FString PSTemplateName = PSTemplate ? PSTemplate->GetName() : "";
-	FText PSTemplateInvalidErrorText = FText::Format(LOCTEXT("InvalidPSTemplateFmt", "{0}: {1} does not contain any trail emittter."), FText::FromString(GetName()), FText::FromString(PSTemplateName));
 
 	MeshComp->ClearAnimNotifyErrors(this);
 
@@ -221,6 +265,7 @@ bool UAnimNotifyState_Trail::ValidateInput(class USkeletalMeshComponent * MeshCo
 	{
 		if (bReportErrors)
 		{
+			const FText FirstSocketEqualsNoneErrorText = FText::Format( LOCTEXT("NoneFirstSocket", "{0}: Must set First Socket Name."), FText::FromString(GetName()));
 			MeshComp->ReportAnimNotifyError(FirstSocketEqualsNoneErrorText, this);
 		}
 		bError = true;
@@ -230,6 +275,7 @@ bool UAnimNotifyState_Trail::ValidateInput(class USkeletalMeshComponent * MeshCo
 	{
 		if (bReportErrors)
 		{
+			const FText SecondSocketEqualsNoneErrorText = FText::Format( LOCTEXT("NoneSecondSocket", "{0}: Must set Second Socket Name."), FText::FromString(GetName()));
 			MeshComp->ReportAnimNotifyError(SecondSocketEqualsNoneErrorText, this);
 		}
 		bError = true;
@@ -239,6 +285,7 @@ bool UAnimNotifyState_Trail::ValidateInput(class USkeletalMeshComponent * MeshCo
 	{
 		if (bReportErrors)
 		{
+			const FText PSTemplateEqualsNoneErrorText = FText::Format( LOCTEXT("NonePSTemplate", "{0}: Trail must have a PSTemplate."), FText::FromString(GetName()));
 			MeshComp->ReportAnimNotifyError(PSTemplateEqualsNoneErrorText, this);
 		}
 		bError = true;
@@ -249,6 +296,8 @@ bool UAnimNotifyState_Trail::ValidateInput(class USkeletalMeshComponent * MeshCo
 		{
 			if (bReportErrors)
 			{
+				const FString PSTemplateName = PSTemplate ? PSTemplate->GetName() : "";
+				const FText PSTemplateInvalidErrorText = FText::Format(LOCTEXT("InvalidPSTemplateFmt", "{0}: {1} does not contain any trail emittter."), FText::FromString(GetName()), FText::FromString(PSTemplateName));
 				MeshComp->ReportAnimNotifyError(PSTemplateInvalidErrorText, this);
 			}
 
diff --git a/Engine/Source/Runtime/Engine/Private/Particles/ParticleComponents.cpp b/Engine/Source/Runtime/Engine/Private/Particles/ParticleComponents.cpp
--- a/Engine/Source/Runtime/Engine/Private/Particles/ParticleComponents.cpp
+++ b/Engine/Source/Runtime/Engine/Private/Particles/ParticleComponents.cpp
@@ -4737,7 +4737,7 @@ int32 UParticleSystemComponent::GetNumActiveParticles() const
 	return NumParticles;
 }
 
-void UParticleSystemComponent::GetOwnedTrailEmitters(TArray< struct FParticleAnimTrailEmitterInstance* >& OutTrailEmitters, const void* InOwner, bool bSetOwner)
+void UParticleSystemComponent::GetOwnedTrailEmitters(UParticleSystemComponent::TrailEmitterArray& OutTrailEmitters, const void* InOwner, bool bSetOwner)
 {
 	for (FParticleEmitterInstance* Inst : EmitterInstances)
 	{
