temp fix merge

- Merging all of it to make sure it doesn't confuse P4.

Merging

//depot/UE4-Orion/Engine/Source/Runtime/Engine/...

to //depot/UE4/Engine/Source/Runtime/Engine/...

diff --git a/Engine/Source/Runtime/Engine/Classes/Animation/AnimInstance.h b/Engine/Source/Runtime/Engine/Classes/Animation/AnimInstance.h
--- a/Engine/Source/Runtime/Engine/Classes/Animation/AnimInstance.h
+++ b/Engine/Source/Runtime/Engine/Classes/Animation/AnimInstance.h
@@ -1010,6 +1010,9 @@ private:
 	/** Native state exit bindings */
 	TArray<FNativeStateBinding> NativeStateExitBindings;
 
+	// update montage
+	void UpdateMontage(float DeltaSeconds);
+
 protected:
 	/** Update all animation node */
 	virtual void UpdateAnimationNode(float DeltaSeconds);
diff --git a/Engine/Source/Runtime/Engine/Private/Animation/AnimInstance.cpp b/Engine/Source/Runtime/Engine/Private/Animation/AnimInstance.cpp
--- a/Engine/Source/Runtime/Engine/Private/Animation/AnimInstance.cpp
+++ b/Engine/Source/Runtime/Engine/Private/Animation/AnimInstance.cpp
@@ -362,6 +362,17 @@ void UAnimInstance::UpdateAnimationNode(float DeltaSeconds)
 	}
 }
 
+void UAnimInstance::UpdateMontage(float DeltaSeconds)
+{
+	// update montage weight
+	Montage_UpdateWeight(DeltaSeconds);
+
+	// update montage should run in game thread
+	// if we do multi threading, make sure this stays in game thread
+	Montage_Advance(DeltaSeconds);
+
+}
+
 void UAnimInstance::UpdateAnimation(float DeltaSeconds)
 {
 	QUICK_SCOPE_CYCLE_COUNTER(STAT_UAnimInstance_UpdateAnimation);
@@ -399,6 +410,7 @@ void UAnimInstance::UpdateAnimation(float DeltaSeconds)
 
 	const TArray<FAnimGroupInstance>& PreviousSyncGroups = SyncGroupArrays[GetSyncGroupReadIndex()];
 
+	// clear previous frame data
 	{
 		QUICK_SCOPE_CYCLE_COUNTER(STAT_UAnimInstance_UpdateAnimation_MiscClear);
 		AnimNotifies.Reset();
@@ -413,6 +425,7 @@ void UAnimInstance::UpdateAnimation(float DeltaSeconds)
 		}
 	}
 
+	// update blueprint and native update
 	{
 		SCOPE_CYCLE_COUNTER(STAT_NativeUpdateAnimation);
 		NativeUpdateAnimation(DeltaSeconds);
@@ -422,11 +435,9 @@ void UAnimInstance::UpdateAnimation(float DeltaSeconds)
 		BlueprintUpdateAnimation(DeltaSeconds);
 	}
 
-	Montage_UpdateWeight(DeltaSeconds);
-
-	// update montage should run in game thread
-	// if we do multi threading, make sure this stays in game thread
-	Montage_Advance(DeltaSeconds);
+	// need to update montage BEFORE node update
+	// so that node knows where montage is
+	UpdateMontage(DeltaSeconds);
 
 	// update all nodes
 	UpdateAnimationNode(DeltaSeconds);
@@ -1531,7 +1542,13 @@ void UAnimInstance::SlotEvaluatePose(FName SlotNodeName, const FCompactPose& Sou
 	}
 
 	// Make sure we have at least one montage here.
-	check((AdditivePoses.Num() > 0) || (NonAdditivePoses.Num() > 0));
+	if (ensure((AdditivePoses.Num() > 0) || (NonAdditivePoses.Num() > 0)) == false)
+	{
+		// failed to get poses, @wip investigating
+		BlendedPose = SourcePose;
+		BlendedCurve = SourceCurve;
+		return;
+	}
 
 	// Second pass, blend non additive poses together
 	{
