Merging

//depot/UE4-Orion/Engine/Source/Runtime/Engine/Private/Animation/AnimInstance.cpp

to //depot/UE4/Engine/Source/Runtime/Engine/Private/Animation/AnimInstance.cpp

diff --git a/Engine/Source/Runtime/Engine/Private/Animation/AnimInstance.cpp b/Engine/Source/Runtime/Engine/Private/Animation/AnimInstance.cpp
--- a/Engine/Source/Runtime/Engine/Private/Animation/AnimInstance.cpp
+++ b/Engine/Source/Runtime/Engine/Private/Animation/AnimInstance.cpp
@@ -1545,8 +1545,7 @@ void UAnimInstance::SlotEvaluatePose(FName SlotNodeName, const FCompactPose& Sou
 	if (ensure((AdditivePoses.Num() > 0) || (NonAdditivePoses.Num() > 0)) == false)
 	{
 		// failed to get poses, @wip investigating
-		BlendedPose = SourcePose;
-		BlendedCurve = SourceCurve;
+		BlendedPose.ResetToRefPose();
 		return;
 	}
 
