Make sure Alpha Blend does not change weight unless, it's Update to guarantee that weight transition is applied properly when parameter has changed.

- Usually this class is used by update node or so, and modifying weight directly often results in undesired behavior because it changes weight outside of update, which can cause issues - in this case Montages.

- You can modify weight by setting alpha value, but you won't apply weight directly otherwise, Update will make sure your blended value gets applied.

- Moved Morphtarget eval data gathering to update Animation.

- Added verbose logs on weight tracking

Merging

//depot/UE4-Orion/Engine/Source/Runtime/Engine/...

to //depot/UE4/Engine/Source/Runtime/Engine/...

diff --git a/Engine/Source/Runtime/Engine/Private/AlphaBlend.cpp b/Engine/Source/Runtime/Engine/Private/AlphaBlend.cpp
--- a/Engine/Source/Runtime/Engine/Private/AlphaBlend.cpp
+++ b/Engine/Source/Runtime/Engine/Private/AlphaBlend.cpp
@@ -23,31 +23,64 @@ FAlphaBlend::FAlphaBlend(const FAlphaBlend& Other, float NewBlendTime)
 	Reset();
 }
 
-void FAlphaBlend::SetBlendTime(float InBlendTime)
+void FAlphaBlend::ResetBlendTime()
 {
-	BlendTime = FMath::Max(InBlendTime, 0.f);
+	// if blend time is <= 0, then blending is done and complete
+	if(BlendTime <= 0.f)
+	{
+		BlendTimeRemaining = 0.f;
+		SetAlpha(1.f);
+	}
+	else
+	{
+		// Blend time is to go all the way, so scale that by how much we have to travel
+		BlendTimeRemaining = BlendTime * FMath::Abs(1.f - AlphaLerp);
+	}
+	
+	bNeedsToResetBlendTime = false;
 }
 
-void FAlphaBlend::Reset()
+void FAlphaBlend::ResetAlpha()
 {
-	AlphaLerp = 0.0f;
-	AlphaBlend = 0.0f;
-	BlendedValue = BeginValue;
+	float SmallerValue = FMath::Min(BeginValue, DesiredValue);
+	float BiggerValue = FMath::Max(BeginValue, DesiredValue);
+	// make sure it's within the range
+	float NewBlendedValue = FMath::Clamp(BlendedValue, SmallerValue, BiggerValue);
+
+	// if blend time is <= 0, or begin == end is same, there is nothing to be done
+	// blending is done and complete
+	if (BeginValue == DesiredValue)
+	{
+		SetAlpha(1.f);
+	}
+	else
+	{
+		AlphaLerp = (BlendedValue - BeginValue)/(DesiredValue - BeginValue);
+		SetAlpha(AlphaLerp);
+	}
+
+	// reset the flag
+	bNeedsToResetAlpha = false;
+}
 
+void FAlphaBlend::Reset()
+{
 	// Set alpha target to full - will also handle zero blend times
 	// if blend time is zero, transition now, don't wait to call update.
 	if( BlendTime <= 0.f )
 	{
-		AlphaLerp = 1.f;
-		AlphaBlend = AlphaToBlendOption();
+		SetAlpha(1.f);
 		BlendTimeRemaining = 0.f;
-		BlendedValue = BeginValue + (DesiredValue - BeginValue) * AlphaBlend;
 	}
 	else
 	{
+		SetAlpha(0.f);
 		// Blend time is to go all the way, so scale that by how much we have to travel
 		BlendTimeRemaining = BlendTime * FMath::Abs(1.f - AlphaLerp);
 	}
+
+	bNeedsToResetAlpha = false;
+	bNeedsToResetBlendTime = false;
 }
 
 void FAlphaBlend::Update(float InDeltaTime)
@@ -55,6 +88,19 @@ void FAlphaBlend::Update(float InDeltaTime)
 	// Make sure passed in delta time is positive
 	check(InDeltaTime >= 0.f);
 
+	// check if we should reset alpha
+	if (bNeedsToResetAlpha)
+	{
+		ResetAlpha();
+	}
+
+	// or should re calc blend time remaining
+	if (bNeedsToResetBlendTime)
+	{
+		ResetBlendTime();
+	}
+
+	// if not complete, 
 	if( !IsComplete() )
 	{
 		if( BlendTimeRemaining > InDeltaTime )
@@ -62,16 +108,12 @@ void FAlphaBlend::Update(float InDeltaTime)
 			const float BlendDelta = 1.f - AlphaLerp; 
 			AlphaLerp += (BlendDelta / BlendTimeRemaining) * InDeltaTime;
 			BlendTimeRemaining -= InDeltaTime;
-
-			AlphaBlend = AlphaToBlendOption();
-			BlendedValue = BeginValue + (DesiredValue - BeginValue) * AlphaBlend;
+			SetAlpha(AlphaLerp);
 		}
 		else
 		{
 			BlendTimeRemaining = 0.f; 
-			AlphaLerp = 1.0f;
-			AlphaBlend = 1.0f;
-			BlendedValue = DesiredValue;
+			SetAlpha(1.f);
 		}
 	}
 }
@@ -113,29 +155,12 @@ float FAlphaBlend::AlphaToBlendOption(float InAlpha, EAlphaBlendOption InBlendOp
 	return InAlpha;
 }
 
-void FAlphaBlend::SetBlendOption(EAlphaBlendOption InBlendOption)
-{
-	BlendOption = InBlendOption;
-
-	// Recalculate the current blended value
-	AlphaBlend = AlphaToBlendOption();
-	BlendedValue = BeginValue + (DesiredValue - BeginValue) * AlphaBlend;
-}
-
-float FAlphaBlend::GetBlendedValue() const
-{
-	return BlendedValue;
-}
-
 void FAlphaBlend::SetValueRange(float Begin, float Desired)
 {
 	BeginValue = Begin;
 	DesiredValue = Desired;
 
-	// Convert to new range
-	BlendedValue = BeginValue + (DesiredValue - BeginValue) * AlphaBlend;
-
-	Reset();
+	bNeedsToResetAlpha = true;
 }
 
 /** Sets the final desired value for the blended value */
@@ -144,6 +169,7 @@ void FAlphaBlend::SetDesiredValue(float InDesired)
 	SetValueRange(BlendedValue, InDesired);
 }
 
+/** note this function can modify BlendedValue right away */
 void FAlphaBlend::SetAlpha(float InAlpha)
 {
 	AlphaLerp = FMath::Clamp(InAlpha, 0.0f, 1.0f);
@@ -151,7 +177,15 @@ void FAlphaBlend::SetAlpha(float InAlpha)
 	BlendedValue = BeginValue + (DesiredValue - BeginValue) * AlphaBlend;
 }
 
+void FAlphaBlend::SetBlendTime(float InBlendTime)
+{
+	BlendTime = FMath::Max(InBlendTime, 0.f);
+	// when blend time changes, we have to restart alpha
+	bNeedsToResetBlendTime = true;
+}
+
 bool FAlphaBlend::IsComplete() const
 {
 	return AlphaLerp == 1.0f;
 }
+
diff --git a/Engine/Source/Runtime/Engine/Private/Animation/AnimInstance.cpp b/Engine/Source/Runtime/Engine/Private/Animation/AnimInstance.cpp
--- a/Engine/Source/Runtime/Engine/Private/Animation/AnimInstance.cpp
+++ b/Engine/Source/Runtime/Engine/Private/Animation/AnimInstance.cpp
@@ -371,6 +371,8 @@ void UAnimInstance::UpdateMontage(float DeltaSeconds)
 	// if we do multi threading, make sure this stays in game thread
 	Montage_Advance(DeltaSeconds);
 
+	// update montage eval data
+	UpdateMontageEvaluationData();
 }
 
 void UAnimInstance::UpdateAnimation(float DeltaSeconds)
@@ -1415,6 +1417,8 @@ void UAnimInstance::GetSlotWeight(FName const& SlotNodeName, float& out_SlotNode
 #if DEBUGMONTAGEWEIGHT			
 			TotalDesiredWeight += MontageInstance->Blend.GetDesiredValue();
 #endif			
+			UE_LOG(LogAnimation, Verbose, TEXT("GetSlotWeight : Owner: %s, AnimMontage: %s,  (DesiredWeight:%0.2f, Weight:%0.2f)"),
+						*GetNameSafe(GetOwningActor()), *MontageInstance->Montage->GetName(), MontageInstance->GetDesiredWeight(), MontageInstance->GetWeight());
 		}
 	}
 
@@ -1542,12 +1546,7 @@ void UAnimInstance::SlotEvaluatePose(FName SlotNodeName, const FCompactPose& Sou
 	}
 
 	// Make sure we have at least one montage here.
-	if (ensure((AdditivePoses.Num() > 0) || (NonAdditivePoses.Num() > 0)) == false)
-	{
-		// failed to get poses, @wip investigating
-		BlendedPose.ResetToRefPose();
-		return;
-	}
+	check((AdditivePoses.Num() > 0) || (NonAdditivePoses.Num() > 0));
 
 	// Second pass, blend non additive poses together
 	{
@@ -2130,6 +2129,8 @@ float UAnimInstance::Montage_Play(UAnimMontage* MontageToPlay, float InPlayRate/
 
 			OnMontageStarted.Broadcast(MontageToPlay);
 
+			UE_LOG(LogAnimation, Verbose, TEXT("Montage_Play: AnimMontage: %s,  (DesiredWeight:%0.2f, Weight:%0.2f)"),
+						*NewInstance->Montage->GetName(), NewInstance->GetDesiredWeight(), NewInstance->GetWeight());
 			return NewInstance->Montage->SequenceLength;
 		}
 		else
@@ -2804,8 +2805,12 @@ bool UAnimInstance::HasMarkerBeenHitThisFrame(FName SyncGroup, FName MarkerName)
 void UAnimInstance::UpdateMontageEvaluationData()
 {
 	MontageEvaluationData.Reset(MontageInstances.Num());
+	UE_LOG(LogAnimation, Verbose, TEXT("UpdateMontageEvaluationData Strting: Owner: %s"),	*GetNameSafe(GetOwningActor()));
+
 	for (FAnimMontageInstance* MontageInstance : MontageInstances)
 	{
+		// although montage can advance with 0.f weight, it is fine to filter by weight here
+		// because we don't want to evaluate them if 0 weight
 		if (MontageInstance->Montage && MontageInstance->GetWeight() > ZERO_ANIMWEIGHT_THRESH)
 		{
 			UE_LOG(LogAnimation, Verbose, TEXT("UpdateMontageEvaluationData : AnimMontage: %s,  (DesiredWeight:%0.2f, Weight:%0.2f)"),
diff --git a/Engine/Source/Runtime/Engine/Private/Animation/AnimMontage.cpp b/Engine/Source/Runtime/Engine/Private/Animation/AnimMontage.cpp
--- a/Engine/Source/Runtime/Engine/Private/Animation/AnimMontage.cpp
+++ b/Engine/Source/Runtime/Engine/Private/Animation/AnimMontage.cpp
@@ -864,13 +864,16 @@ void FAnimMontageInstance::Play(float InPlayRate)
 
 void FAnimMontageInstance::InitializeBlend(const FAlphaBlend& InAlphaBlend)
 {
-	Blend.SetBlendOption(InAlphaBlend.BlendOption);
+	Blend.BlendOption = InAlphaBlend.BlendOption;
 	Blend.CustomCurve = InAlphaBlend.CustomCurve;
 	Blend.SetBlendTime(InAlphaBlend.GetBlendTime());
 }
 
 void FAnimMontageInstance::Stop(const FAlphaBlend& InBlendOut, bool bInterrupt)
 {
+	UE_LOG(LogAnimation, Verbose, TEXT("Montage.Stop Before: AnimMontage: %s,  (DesiredWeight:%0.2f, Weight:%0.2f)"),
+			*Montage->GetName(), GetDesiredWeight(), GetWeight());
+
 	// overwrite bInterrupted if it hasn't already interrupted
 	// once interrupted, you don't go back to non-interrupted
 	if (!bInterrupted && bInterrupt)
@@ -921,6 +924,9 @@ void FAnimMontageInstance::Stop(const FAlphaBlend& InBlendOut, bool bInterrupt)
 	{
 		bPlaying = false;
 	}
+
+	UE_LOG(LogAnimation, Verbose, TEXT("Montage.Stop After: AnimMontage: %s,  (DesiredWeight:%0.2f, Weight:%0.2f)"),
+			*Montage->GetName(), GetDesiredWeight(), GetWeight());
 }
 
 void FAnimMontageInstance::Pause()
@@ -1009,6 +1015,8 @@ void FAnimMontageInstance::Terminate()
 	// clear Blend curve
 	Blend.CustomCurve = NULL;
 	Blend.BlendOption = EAlphaBlendOption::Linear;
+
+	UE_LOG(LogAnimation, Verbose, TEXT("Terminating: AnimMontage: %s"), *GetNameSafe(OldMontage));
 }
 
 bool FAnimMontageInstance::JumpToSectionName(FName const & SectionName, bool bEndOfSection)
@@ -1235,6 +1243,9 @@ void FAnimMontageInstance::UpdateWeight(float DeltaTime)
 		// Notify weight is max of previous and current as notify could have come
 		// from any point between now and last tick
 		NotifyWeight = FMath::Max(PreviousWeight, Blend.GetBlendedValue());
+
+		UE_LOG(LogAnimation, Verbose, TEXT("UpdateWeight: AnimMontage: %s,  (DesiredWeight:%0.2f, Weight:%0.2f, PreviousWeight: %0.2f)"),
+			*Montage->GetName(), GetDesiredWeight(), GetWeight(), PreviousWeight);
 	}
 }
 
diff --git a/Engine/Source/Runtime/Engine/Private/Animation/AnimNode_StateMachine.cpp b/Engine/Source/Runtime/Engine/Private/Animation/AnimNode_StateMachine.cpp
--- a/Engine/Source/Runtime/Engine/Private/Animation/AnimNode_StateMachine.cpp
+++ b/Engine/Source/Runtime/Engine/Private/Animation/AnimNode_StateMachine.cpp
@@ -45,7 +45,7 @@ FAnimationActiveTransitionEntry::FAnimationActiveTransitionEntry(int32 NextState
 	CrossfadeDuration = ReferenceTransitionInfo.CrossfadeDuration * CalculateInverseAlpha(BlendOption, Scaler);
 
 	Blend.SetBlendTime(CrossfadeDuration);
-	Blend.SetBlendOption(BlendOption);
+	Blend.BlendOption = BlendOption;
 	Blend.CustomCurve = ReferenceTransitionInfo.CustomCurve;
 	Blend.SetValueRange(0.0f, 1.0f);
 }
diff --git a/Engine/Source/Runtime/Engine/Private/Components/SkeletalMeshComponent.cpp b/Engine/Source/Runtime/Engine/Private/Components/SkeletalMeshComponent.cpp
--- a/Engine/Source/Runtime/Engine/Private/Components/SkeletalMeshComponent.cpp
+++ b/Engine/Source/Runtime/Engine/Private/Components/SkeletalMeshComponent.cpp
@@ -1089,11 +1089,6 @@ void USkeletalMeshComponent::RefreshBoneTransforms(FActorComponentTickFunction*
 		RecalcRequiredBones(PredictedLODLevel);
 	}
 
-	if (AnimScriptInstance)
-	{
-		QUICK_SCOPE_CYCLE_COUNTER(STAT_USkeletalMeshComponent_RefreshBoneTransforms_UpdateMontageEvaluationData);
-		AnimScriptInstance->UpdateMontageEvaluationData();
-	}
 	const bool bDoEvaluationRateOptimization = ShouldUseUpdateRateOptimizations() && AnimUpdateRateParams->DoEvaluationRateOptimizations();
 
 	//Handle update rate optimization setup
diff --git a/Engine/Source/Runtime/Engine/Public/AlphaBlend.h b/Engine/Source/Runtime/Engine/Public/AlphaBlend.h
--- a/Engine/Source/Runtime/Engine/Public/AlphaBlend.h
+++ b/Engine/Source/Runtime/Engine/Public/AlphaBlend.h
@@ -42,41 +42,38 @@ enum class EAlphaBlendOption : uint8
 };
 
 /**
- * Alpha Blend Type
+ * Alpha Blend class that supports different blend options as well as custom curves
  */
 USTRUCT()
 struct ENGINE_API FAlphaBlend
 {
 	GENERATED_BODY()
+public:
+	/** 
+	 * Please note that changing below two variables would get applied in the NEXT UPDATE. 
+	 * This does not change the Alpha and BlendedValue right away and it is intentional
+	 */
 
 	/** Type of blending used (Linear, Cubic, etc.) */
 	UPROPERTY(EditAnywhere, Category = "Blend")
 	EAlphaBlendOption BlendOption;
 
+	/** If you're using Custom BlendOption, you can specify curve */
 	UPROPERTY(EditAnywhere, Category = "Blend")
 	UCurveFloat* CustomCurve;
 
-protected:
-	UPROPERTY(EditAnywhere, Category = "Blend")
-	float	BlendTime;
-
-public:
-
-	// Constructor
+	/* Constructor */
 	FAlphaBlend(float NewBlendTime = 0.2f);
 
-	// constructor
+	/* Constructor */
 	FAlphaBlend(const FAlphaBlend& Other, float NewBlendTime);
 
-	/** Update transition blend time */
+	/** Update transition blend time. This new value will be applied in the next Update. */
 	void SetBlendTime(float InBlendTime);
 
-	/** Sets the method used to blend the value */
-	void SetBlendOption(EAlphaBlendOption InBlendOption);
-
 	/** Sets the range of values to map to the interpolation
 	 *
-	 * @param Begin : this is current value
+	 * @param Begin : this is start value
 	 * @param Desired : this is target value 
 	 *
 	 * This can be (0, 1) if you'd like to increase, or it can be (1, 0) if you'd like to get to 0
@@ -86,20 +83,20 @@ public:
 	/** Sets the final desired value for the blended value */
 	void SetDesiredValue(float InDesired);
 
-	/** Sets the Lerp alpha value directly */
+	/** Sets the Lerp alpha value directly. PLEASE NOTE that this modifies the Blended Value right away.  */
 	void SetAlpha(float InAlpha);
 
 	/** Update interpolation, has to be called once every frame */
 	void Update(float InDeltaTime);
 
+	/** Gets whether or not the blend is complete */
+	bool IsComplete() const;
+
 	/** Gets the current 0..1 alpha value. Changed to AlphaLerp to match with SetAlpha function */
 	float GetAlpha() const { return AlphaLerp; }
 
 	/** Gets the current blended value */
-	float GetBlendedValue() const;
-
-	/** Gets whether or not the blend is complete */
-	bool IsComplete() const;
+	float GetBlendedValue() const { return BlendedValue; }
 
 	/** Gets the current blend time */
 	float GetBlendTime() const { return BlendTime; }
@@ -114,32 +111,69 @@ public:
 	 */
 	static float AlphaToBlendOption(float InAlpha, EAlphaBlendOption InBlendOption, UCurveFloat* InCustomCurve = nullptr);
 
-protected:
+private:
+	/** Blend Time */
+	UPROPERTY(EditAnywhere, Category = "Blend")
+	float	BlendTime;
+
 	/** Internal Lerped value for Alpha */
-	UPROPERTY()
 	float	AlphaLerp;
 
 	/** Resulting Alpha value, between 0.f and 1.f */
-	UPROPERTY()
 	float	AlphaBlend;
 
 	/** Time left to reach target */
-	UPROPERTY()
 	float	BlendTimeRemaining;
 
-	/** The current blended value derived from the begin and desired values */
-	UPROPERTY()
+	/** The current blended value derived from the begin and desired values. 
+	  * This value should not change unless Update. 
+	  */
 	float BlendedValue;
 
-	UPROPERTY()
+	/** The Start value. It is 'from' range */
 	float BeginValue;
 
-	UPROPERTY()
+	/**  The Target value. It is 'to' range */
 	float DesiredValue;
 
-	/** Reset to zero / restart the blend */
+	/** 
+	 * Reset functions - 
+	 *
+	 * The distinction is important. 
+	 * 
+	 * We have 3 different reset functions depending on usability
+	 * 
+	 * Reset function : will change Blended Value to BeginValue, so that it can start blending
+	 * That is only supposed to be used when you want to clear it up and restart
+	 *
+	 * ResetAlpha will change AlphaLerp/AlphaBlend to match with current Blend Value, 
+	 * that way we keep the current blended value and move to target using the current value
+	 * That will make sure this doesn't pop when you change direction of Desired
+	 *
+	 * ResetBlendTime will change BlendTimeRemaining as well as possibly weight because if BlendTimeRemaining <= 0.f, 
+	 * We'll arrive to the destination
+	 *
+	 * The reason we need ResetAlpha and ResetBlendTime is that we don't want to modify blend time if direction changes or we don't want to reset alpha if blend time changes
+	 * Those two have to work independently
+	 * */
+
+public:
+	/** Reset to zero / restart the blend. This resets whole thing.  */
 	void Reset();
 
+private:
+	/** Reset alpha, this keeps current BlendedValue but modify Alpha to keep the blending state.  */
+	void ResetAlpha();
+
+	/* Reset Blend Time, this modifies BlendTimeRemaining and possibly Weight when BlendTimeRemaining <= 0.f */
+	void ResetBlendTime();
+
 	/** Converts internal lerped alpha into the output alpha type */
 	float AlphaToBlendOption();
+
+	/** internal flag to reset the alpha value */
+	bool bNeedsToResetAlpha;
+
+	/** internal flat to reset blend time. The two distinction is required. Otherwise, you'll modify blend time when you just change range */
+	bool bNeedsToResetBlendTime;
 };
\ No newline at end of file
