Turn on async component tick by default.

diff --git a/Engine/Source/Runtime/Engine/Private/TickTaskManager.cpp b/Engine/Source/Runtime/Engine/Private/TickTaskManager.cpp
--- a/Engine/Source/Runtime/Engine/Private/TickTaskManager.cpp
+++ b/Engine/Source/Runtime/Engine/Private/TickTaskManager.cpp
@@ -38,7 +38,7 @@ static TAutoConsoleVariable<int32> CVarLogTicksShowPrerequistes(
 
 static TAutoConsoleVariable<int32> CVarAllowAsyncComponentTicks(
 	TEXT("tick.AllowAsyncComponentTicks"),
-	0,
+	1,
 	TEXT("Used to control async component ticks."));
 
 static TAutoConsoleVariable<int32> CVarAllowConcurrentQueue(
