Parallel cloth now uses async tick function instead of tasks. Moved end and start cloth into cloth manager.

diff --git a/Engine/Source/Runtime/Engine/Classes/Engine/World.h b/Engine/Source/Runtime/Engine/Classes/Engine/World.h
--- a/Engine/Source/Runtime/Engine/Classes/Engine/World.h
+++ b/Engine/Source/Runtime/Engine/Classes/Engine/World.h
@@ -374,7 +374,7 @@ struct FEndPhysicsTickFunction : public FTickFunction
 * Tick function that starts the cloth tick
 **/
 USTRUCT()
-struct FStartClothAndAsyncSimulationFunction : public FTickFunction
+struct FStartAsyncSimulationFunction : public FTickFunction
 {
 	GENERATED_USTRUCT_BODY()
 
@@ -870,9 +870,7 @@ public:
 	FEndPhysicsTickFunction EndPhysicsTickFunction;
 
 	/** Tick function for starting cloth simulation																				*/
-	FStartClothAndAsyncSimulationFunction StartClothTickFunction;
-	/** Tick function for ending cloth simulation																				*/
-	FEndClothSimulationFunction EndClothTickFunction;
+	FStartAsyncSimulationFunction StartAsyncTickFunction;
 
 	/** 
 	 * Indicates that during world ticking we are doing the final component update of dirty components 
@@ -2113,10 +2111,10 @@ public:
 	UGameViewportClient* GetGameViewport() const;
 
 private:
-	/** Begin cloth and async simulation */
-	void StartClothAndAsyncSim();
+	/** Begin async simulation */
+	void StartAsyncSim();
 
-	friend FStartClothAndAsyncSimulationFunction;
+	friend FStartAsyncSimulationFunction;
 
 public:
 
diff --git a/Engine/Source/Runtime/Engine/Private/Components/PrimitiveComponent.cpp b/Engine/Source/Runtime/Engine/Private/Components/PrimitiveComponent.cpp
--- a/Engine/Source/Runtime/Engine/Private/Components/PrimitiveComponent.cpp
+++ b/Engine/Source/Runtime/Engine/Private/Components/PrimitiveComponent.cpp
@@ -313,10 +313,10 @@ void UPrimitiveComponent::RegisterComponentTickFunctions(bool bRegister)
 				PostPhysicsComponentTick.AddPrerequisite(this,PrimaryComponentTick); 
 			}
 
-			// Set a prereq for the post physics tick to happen after physics and cloth is finished
-			if(World != NULL)
+			// Set a prereq for the post physics tick to happen after physics is finished
+			if (World != NULL)
 			{
-				PostPhysicsComponentTick.AddPrerequisite(World, World->EndClothTickFunction);
+				PostPhysicsComponentTick.AddPrerequisite(World, World->EndPhysicsTickFunction);
 			}
 		}
 	}
diff --git a/Engine/Source/Runtime/Engine/Private/PhysicsEngine/ClothManager.cpp b/Engine/Source/Runtime/Engine/Private/PhysicsEngine/ClothManager.cpp
--- a/Engine/Source/Runtime/Engine/Private/PhysicsEngine/ClothManager.cpp
+++ b/Engine/Source/Runtime/Engine/Private/PhysicsEngine/ClothManager.cpp
@@ -7,34 +7,18 @@
 
 FClothManager::FClothManager(UWorld* InAssociatedWorld)
 : AssociatedWorld(InAssociatedWorld)
- {
-	//simulate cloth tick
-	StartClothTickFunction.bCanEverTick = true;
-	StartClothTickFunction.Target = this;
-	StartClothTickFunction.TickGroup = TG_StartCloth;
-	StartClothTickFunction.RegisterTickFunction(AssociatedWorld->PersistentLevel);
-
-	//prepare cloth data ticks
-	PrepareClothDataArray[(int32)PrepareClothSchedule::IgnorePhysics].TickFunction.TickGroup = TG_StartPhysics;
-	PrepareClothDataArray[(int32)PrepareClothSchedule::WaitOnPhysics].TickFunction.TickGroup = TG_PreCloth;
-
-	for(FClothManagerData& PrepareClothData : PrepareClothDataArray)
-	{
-		PrepareClothData.TickFunction.bCanEverTick = true;
-		PrepareClothData.TickFunction.Target = &PrepareClothData;
-		PrepareClothData.TickFunction.bRunOnAnyThread = true;	//this value seems to be ignored
-		PrepareClothData.TickFunction.RegisterTickFunction(AssociatedWorld->PersistentLevel);
-	}
- }
+{
+}
 
 void FClothManagerData::PrepareCloth(float DeltaTime)
 {
 #if WITH_APEX_CLOTHING
-	if(SkeletalMeshComponents.Num())
+	QUICK_SCOPE_CYCLE_COUNTER(STAT_FClothManager_PrepareCloth);
+	if (SkeletalMeshComponents.Num())
 	{
 		IsPreparingCloth.AtomicSet(true);
 		for (USkeletalMeshComponent* SkeletalMeshComponent : SkeletalMeshComponents)
-		{	
+		{
 			SkeletalMeshComponent->SubmitClothSimulationContext();	//make sure user params are passed internally
 
 			FClothSimulationContext& ClothSimulationContext = SkeletalMeshComponent->InternalClothSimulationContext;
@@ -52,23 +36,10 @@ void FClothManager::RegisterForPrepareCloth(USkeletalMeshComponent* SkeletalMesh
 	PrepareClothDataArray[(int32)PrepSchedule].SkeletalMeshComponents.Add(SkeletalMeshComponent);
 }
 
-void FClothManager::StartCloth()
+void FClothManager::StartCloth(float DeltaTime)
 {
-	//Make sure prepare tasks are done
-	FGraphEventArray ThingsToComplete;
-	for (const FClothManagerData& PrepareData : PrepareClothDataArray)
-	{
-		if (PrepareData.PrepareCompletion.IsValid())
-		{
-			ThingsToComplete.Add(PrepareData.PrepareCompletion);
-		}
-	}
-
-	if (ThingsToComplete.Num())
-	{
-		QUICK_SCOPE_CYCLE_COUNTER(STAT_FClothManager_WaitPrepareCloth);
-		FTaskGraphInterface::Get().WaitUntilTasksComplete(ThingsToComplete, ENamedThreads::GameThread);
-	}
+	//First we must prepare any cloth that waits on physics sim
+	PrepareClothDataArray[(int32)PrepareClothSchedule::WaitOnPhysics].PrepareCloth(DeltaTime);
 
 	//Reset skeletal mesh components
 	bool bNeedSimulateCloth = false;
@@ -91,60 +62,86 @@ void FClothManager::StartCloth()
 	}
 }
 
+static TAutoConsoleVariable<int32> CVarParallelCloth(TEXT("p.ParallelCloth"), 0, TEXT("Whether to prepare and start cloth sim off the game thread"));
 
-void FStartClothTickFunction::ExecuteTick(float DeltaTime, enum ELevelTick TickType, ENamedThreads::Type CurrentThread, const FGraphEventRef& MyCompletionGraphEvent)
+void FClothManager::SetupClothTickFunction(bool bTickClothSim)
 {
-	Target->StartCloth();
-}
+	check(IsInGameThread());
+	check(!IsPreparingClothAsync());
 
-FString FStartClothTickFunction::DiagnosticMessage()
-{
-	return TEXT("FStartClothTickFunction");
-}
-
-TAutoConsoleVariable<int32> CVarParallelCloth(TEXT("p.ParallelCloth"), 0, TEXT("If turned on, preparing cloth will happen off the game thread."));
+	bool bTickOnAny = !!CVarParallelCloth.GetValueOnGameThread();
+	StartIgnorePhysicsClothTickFunction.bRunOnAnyThread = bTickOnAny;
+	StartClothTickFunction.bRunOnAnyThread = bTickOnAny;
 
-class FPrepareClothTickTask
-{
-	FClothManagerData* PrepareClothData;
-	float DeltaTime;
-
-public:
-	FPrepareClothTickTask(FClothManagerData* InPrepareClothData, float InDeltaTime)
-		: PrepareClothData(InPrepareClothData)
-		, DeltaTime(InDeltaTime)
+	/** Need to register tick function and we are not currently registered */
+	if (bTickClothSim && !StartClothTickFunction.IsTickFunctionRegistered())
 	{
+		//prepare cloth data ticks
+		StartIgnorePhysicsClothTickFunction.TickGroup = TG_StartPhysics;
+		StartIgnorePhysicsClothTickFunction.Target = this;
+		StartIgnorePhysicsClothTickFunction.bCanEverTick = true;
+		StartIgnorePhysicsClothTickFunction.RegisterTickFunction(AssociatedWorld->PersistentLevel);
+
+		//simulate cloth tick
+		StartClothTickFunction.TickGroup = TG_StartCloth;
+		StartClothTickFunction.Target = this;
+		StartClothTickFunction.bCanEverTick = true;
+		StartClothTickFunction.AddPrerequisite(AssociatedWorld, StartIgnorePhysicsClothTickFunction);
+		StartClothTickFunction.RegisterTickFunction(AssociatedWorld->PersistentLevel);
+
+		//wait for cloth sim to finish
+		EndClothTickFunction.TickGroup = TG_EndCloth;
+		EndClothTickFunction.Target = this;
+		EndClothTickFunction.bCanEverTick = true;
+		EndClothTickFunction.bRunOnAnyThread = false;	//this HAS to run on game thread because we need to block if cloth sim is not done
+		EndClothTickFunction.AddPrerequisite(AssociatedWorld, StartClothTickFunction);
+		EndClothTickFunction.RegisterTickFunction(AssociatedWorld->PersistentLevel);
 	}
-
-	FORCEINLINE TStatId GetStatId() const
+	else if (!bTickClothSim && StartClothTickFunction.IsTickFunctionRegistered())	//Need to unregister tick function and we current are registered
 	{
-		RETURN_QUICK_DECLARE_CYCLE_STAT(FPrepareClothTickTask, STATGROUP_TaskGraphTasks);
-	}
-	static ENamedThreads::Type GetDesiredThread()
-	{
-		return CVarParallelCloth.GetValueOnGameThread() ? ENamedThreads::AnyThread : ENamedThreads::GameThread;
-	}
-	static ESubsequentsMode::Type GetSubsequentsMode()
-	{
-		return ESubsequentsMode::TrackSubsequents;
+		StartIgnorePhysicsClothTickFunction.UnRegisterTickFunction();
+		StartClothTickFunction.UnRegisterTickFunction();
+		EndClothTickFunction.UnRegisterTickFunction();
 	}
+}
+
 
-	void DoTask(ENamedThreads::Type CurrentThread, const FGraphEventRef& MyCompletionGraphEvent)
+void FClothManager::EndCloth()
+{
+	QUICK_SCOPE_CYCLE_COUNTER(FClothManager_ExecuteTick);
+	if (FPhysScene* PhysScene = AssociatedWorld->GetPhysicsScene())
 	{
-		PrepareClothData->PrepareCloth(DeltaTime);
+		PhysScene->WaitClothScene();
 	}
-};
+}
 
+void FStartIgnorePhysicsClothTickFunction::ExecuteTick(float DeltaTime, enum ELevelTick TickType, ENamedThreads::Type CurrentThread, const FGraphEventRef& MyCompletionGraphEvent)
+{
+	Target->PrepareClothDataArray[(int32)PrepareClothSchedule::IgnorePhysics].PrepareCloth(DeltaTime);
+}
 
-void FPrepareClothTickFunction::ExecuteTick(float DeltaTime, enum ELevelTick TickType, ENamedThreads::Type CurrentThread, const FGraphEventRef& MyCompletionGraphEvent)
+FString FStartIgnorePhysicsClothTickFunction::DiagnosticMessage()
 {
-	QUICK_SCOPE_CYCLE_COUNTER(FPrepareClothTickFunction_ExecuteTick);
+	return TEXT("FStartIgnorePhysicsClothTickFunction");
+}
 
-	check(!Target->PrepareCompletion.IsValid())
-	Target->PrepareCompletion = TGraphTask<FPrepareClothTickTask>::CreateTask().ConstructAndDispatchWhenReady(Target, DeltaTime);
+void FStartClothTickFunction::ExecuteTick(float DeltaTime, enum ELevelTick TickType, ENamedThreads::Type CurrentThread, const FGraphEventRef& MyCompletionGraphEvent)
+{
+	Target->StartCloth(DeltaTime);
+}
+
+FString FStartClothTickFunction::DiagnosticMessage()
+{
+	return TEXT("FStartClothTickFunction");
 }
 
-FString FPrepareClothTickFunction::DiagnosticMessage()
+
+void FEndClothTickFunction::ExecuteTick(float DeltaTime, enum ELevelTick TickType, ENamedThreads::Type CurrentThread, const FGraphEventRef& MyCompletionGraphEvent)
 {
-	return TEXT("FPrepareClothTickFunction");
+	Target->EndCloth();
 }
+
+FString FEndClothTickFunction::DiagnosticMessage()
+{
+	return TEXT("FEndClothTickFunction");
+}
\ No newline at end of file
diff --git a/Engine/Source/Runtime/Engine/Private/PhysicsEngine/ClothManager.h b/Engine/Source/Runtime/Engine/Private/PhysicsEngine/ClothManager.h
--- a/Engine/Source/Runtime/Engine/Private/PhysicsEngine/ClothManager.h
+++ b/Engine/Source/Runtime/Engine/Private/PhysicsEngine/ClothManager.h
@@ -18,22 +18,11 @@ enum class PrepareClothSchedule
 	MAX
 };
 
-struct FPrepareClothTickFunction : public FTickFunction
-{
-	class FClothManagerData* Target;
-
-	// FTickFunction interface
-	virtual void ExecuteTick(float DeltaTime, enum ELevelTick TickType, ENamedThreads::Type CurrentThread, const FGraphEventRef& MyCompletionGraphEvent) override;
-	virtual FString DiagnosticMessage() override;
-	// End of FTickFunction interface
-};
-
 /** Break the data into class since we have multiple passes for different scheduling purposes */
 class FClothManagerData
 {
 public:
 	TArray<USkeletalMeshComponent*> SkeletalMeshComponents;
-	FPrepareClothTickFunction TickFunction;
 	FThreadSafeBool IsPreparingCloth;
 	FGraphEventRef PrepareCompletion;
 
@@ -41,6 +30,16 @@ public:
 	void PrepareCloth(float DeltaTime);
 };
 
+struct FStartIgnorePhysicsClothTickFunction : public FTickFunction
+{
+	class FClothManager* Target;
+
+	// FTickFunction interface
+	virtual void ExecuteTick(float DeltaTime, enum ELevelTick TickType, ENamedThreads::Type CurrentThread, const FGraphEventRef& MyCompletionGraphEvent) override;
+	virtual FString DiagnosticMessage() override;
+	// End of FTickFunction interface
+};
+
 struct FStartClothTickFunction : public FTickFunction
 {
 	class FClothManager* Target;
@@ -51,6 +50,16 @@ struct FStartClothTickFunction : public FTickFunction
 	// End of FTickFunction interface
 };
 
+struct FEndClothTickFunction : public FTickFunction
+{
+	class FClothManager* Target;
+
+	// FTickFunction interface
+	virtual void ExecuteTick(float DeltaTime, enum ELevelTick TickType, ENamedThreads::Type CurrentThread, const FGraphEventRef& MyCompletionGraphEvent) override;
+	virtual FString DiagnosticMessage() override;
+	// End of FTickFunction interface
+};
+
 
 class FClothManager
 {
@@ -72,13 +81,21 @@ public:
 
 		return bIsPreparingCloth;
 	}
-	
+
+	/** Whether or not to tick the cloth sim. Note this function should not be called during parallel evaluation*/
+	void SetupClothTickFunction(bool bTickClothSim);
+
 private:
 	UWorld* AssociatedWorld;
 	FClothManagerData PrepareClothDataArray[(int32)PrepareClothSchedule::MAX];
+	FStartIgnorePhysicsClothTickFunction StartIgnorePhysicsClothTickFunction;
 	FStartClothTickFunction StartClothTickFunction;
+	FEndClothTickFunction EndClothTickFunction;
 
+	friend FStartIgnorePhysicsClothTickFunction;
 	friend FStartClothTickFunction;
+	friend FEndClothTickFunction;
 
-	void StartCloth();
+	void StartCloth(float DeltaTime);
+	void EndCloth();
 };
\ No newline at end of file
diff --git a/Engine/Source/Runtime/Engine/Private/PhysicsEngine/PhysLevel.cpp b/Engine/Source/Runtime/Engine/Private/PhysicsEngine/PhysLevel.cpp
--- a/Engine/Source/Runtime/Engine/Private/PhysicsEngine/PhysLevel.cpp
+++ b/Engine/Source/Runtime/Engine/Private/PhysicsEngine/PhysLevel.cpp
@@ -108,18 +108,14 @@ void UWorld::SetupPhysicsTickFunctions(float DeltaSeconds)
 	EndPhysicsTickFunction.bCanEverTick = true;
 	EndPhysicsTickFunction.Target = this;
 	
-	StartClothTickFunction.bCanEverTick = true;
-	StartClothTickFunction.Target = this;
-	
-	EndClothTickFunction.bCanEverTick = true;
-	EndClothTickFunction.Target = this;
+	StartAsyncTickFunction.bCanEverTick = true;
+	StartAsyncTickFunction.Target = this;
 	
 	
 	// see if we need to update tick registration
 	bool bNeedToUpdateTickRegistration = (bShouldSimulatePhysics != StartPhysicsTickFunction.IsTickFunctionRegistered())
 		|| (bShouldSimulatePhysics != EndPhysicsTickFunction.IsTickFunctionRegistered())
-		|| (bShouldSimulatePhysics != StartClothTickFunction.IsTickFunctionRegistered())
-		|| (bShouldSimulatePhysics != EndClothTickFunction.IsTickFunctionRegistered());
+		|| (bShouldSimulatePhysics != StartAsyncTickFunction.IsTickFunctionRegistered());
 
 	if (bNeedToUpdateTickRegistration && PersistentLevel)
 	{
@@ -145,27 +141,15 @@ void UWorld::SetupPhysicsTickFunctions(float DeltaSeconds)
 			EndPhysicsTickFunction.UnRegisterTickFunction();
 		}
 
-		//cloth
-		if (bShouldSimulatePhysics && !StartClothTickFunction.IsTickFunctionRegistered())
-		{
-			StartClothTickFunction.TickGroup = TG_StartCloth;
-			StartClothTickFunction.RegisterTickFunction(PersistentLevel);
-		}
-		else if (!bShouldSimulatePhysics && StartClothTickFunction.IsTickFunctionRegistered())
-		{
-			StartClothTickFunction.UnRegisterTickFunction();
-		}
-
-		if (bShouldSimulatePhysics && !EndClothTickFunction.IsTickFunctionRegistered())
+		//async scene
+		if (bShouldSimulatePhysics && !StartAsyncTickFunction.IsTickFunctionRegistered())
 		{
-			EndClothTickFunction.TickGroup = TG_EndCloth;
-			EndClothTickFunction.RegisterTickFunction(PersistentLevel);
-			EndClothTickFunction.AddPrerequisite(this, StartClothTickFunction);
+			StartAsyncTickFunction.TickGroup = TG_StartCloth;
+			StartAsyncTickFunction.RegisterTickFunction(PersistentLevel);
 		}
-		else if (!bShouldSimulatePhysics && EndClothTickFunction.IsTickFunctionRegistered())
+		else if (!bShouldSimulatePhysics && StartAsyncTickFunction.IsTickFunctionRegistered())
 		{
-			EndClothTickFunction.RemovePrerequisite(this, StartClothTickFunction);
-			EndClothTickFunction.UnRegisterTickFunction();
+			StartAsyncTickFunction.UnRegisterTickFunction();
 		}
 	}
 
@@ -175,6 +159,11 @@ void UWorld::SetupPhysicsTickFunctions(float DeltaSeconds)
 		return;
 	}
 
+	if (FClothManager* ClothManager = PhysicsScene->GetClothManager())
+	{
+		ClothManager->SetupClothTickFunction(bShouldSimulatePhysics);
+	}
+
 #if WITH_PHYSX
 	
 	// When ticking the main scene, clean up any physics engine resources (once a frame)
@@ -210,7 +199,7 @@ void UWorld::FinishPhysicsSim()
 	PhysScene->EndFrame(LineBatcher);
 }
 
-void UWorld::StartClothAndAsyncSim()
+void UWorld::StartAsyncSim()
 {
 	if (FPhysScene* PhysScene = GetPhysicsScene())
 	{
@@ -278,17 +267,17 @@ FString FEndPhysicsTickFunction::DiagnosticMessage()
 	return TEXT("FEndPhysicsTickFunction");
 }
 
-void FStartClothAndAsyncSimulationFunction::ExecuteTick(float DeltaTime, enum ELevelTick TickType, ENamedThreads::Type CurrentThread, const FGraphEventRef& MyCompletionGraphEvent)
+void FStartAsyncSimulationFunction::ExecuteTick(float DeltaTime, enum ELevelTick TickType, ENamedThreads::Type CurrentThread, const FGraphEventRef& MyCompletionGraphEvent)
 {
-	QUICK_SCOPE_CYCLE_COUNTER(FStartClothSimulationFunction_ExecuteTick);
+	QUICK_SCOPE_CYCLE_COUNTER(FStartAsyncSimulationFunction_ExecuteTick);
 
 	check(Target);
-	Target->StartClothAndAsyncSim();
+	Target->StartAsyncSim();
 }
 
-FString FStartClothAndAsyncSimulationFunction::DiagnosticMessage()
+FString FStartAsyncSimulationFunction::DiagnosticMessage()
 {
-	return TEXT("FStartClothSimulationFunction");
+	return TEXT("FStartAsyncSimulationFunction");
 }
 
 void FEndClothSimulationFunction::ExecuteTick(float DeltaTime, enum ELevelTick TickType, ENamedThreads::Type CurrentThread, const FGraphEventRef& MyCompletionGraphEvent)
