attach weldness of being crash

Some sort of change in collision status, attachment, welding configuration was causing a recreate of physx objects.  because of the way things were configured, it was attempting to initialize a body twice into physx.   the first time due to an unweld would instantiate the "root' which turned out to be self, and then again later it detected that it was actor was already created but not properly dealt with.  now we remove things from the to-be-created-and-added-to-scene list if they happen to magically appear during the welding steps.

http://crashreporter/Buggs/Show/70202

diff --git a/Engine/Source/Runtime/Engine/Private/PhysicsEngine/BodyInstance.cpp b/Engine/Source/Runtime/Engine/Private/PhysicsEngine/BodyInstance.cpp
--- a/Engine/Source/Runtime/Engine/Private/PhysicsEngine/BodyInstance.cpp
+++ b/Engine/Source/Runtime/Engine/Private/PhysicsEngine/BodyInstance.cpp
@@ -979,7 +979,7 @@ void FBodyInstance::UpdatePhysicsFilterData()
 template <bool bCompileStatic>
 struct FInitBodiesHelper
 {
-	FInitBodiesHelper(const TArray<FBodyInstance*>& InBodies, const TArray<FTransform>& InTransforms, class UBodySetup* InBodySetup, class UPrimitiveComponent* InPrimitiveComp, class FPhysScene* InInRBScene, FBodyInstance::PhysXAggregateType InInAggregate = NULL, UPhysicsSerializer* InPhysicsSerializer = nullptr)
+	FInitBodiesHelper( TArray<FBodyInstance*> InBodies, TArray<FTransform> InTransforms, class UBodySetup* InBodySetup, class UPrimitiveComponent* InPrimitiveComp, class FPhysScene* InInRBScene, FBodyInstance::PhysXAggregateType InInAggregate = NULL, UPhysicsSerializer* InPhysicsSerializer = nullptr)
 	: Bodies(InBodies)
 	, Transforms(InTransforms)
 	, BodySetup(InBodySetup)
@@ -1011,7 +1011,7 @@ struct FInitBodiesHelper
 
 	}
 
-	void InitBodies() const
+	void InitBodies() 
 	{
 #if WITH_PHYSX
 		InitBodies_PhysX();
@@ -1023,8 +1023,8 @@ struct FInitBodiesHelper
 	}
 
 	//The arguments passed into InitBodies
-	const TArray<FBodyInstance*>& Bodies;
-	const TArray<FTransform>& Transforms;
+	TArray<FBodyInstance*> Bodies;   
+	TArray<FTransform> Transforms;
 	class UBodySetup* BodySetup;
 	class UPrimitiveComponent* PrimitiveComp;
 	class FPhysScene* PhysScene;
@@ -1154,7 +1154,7 @@ struct FInitBodiesHelper
 		return bInitFail;
 	}
 
-	bool CreateShapesAndActors_PhysX(TArray<PxActor*>& PSyncActors, TArray<PxActor*>& PAsyncActors, TArray<PxActor*>& PDynamicActors, const bool bCanDefer, bool& bDynamicsUseAsyncScene) const
+	bool CreateShapesAndActors_PhysX(TArray<PxActor*>& PSyncActors, TArray<PxActor*>& PAsyncActors, TArray<PxActor*>& PDynamicActors, const bool bCanDefer, bool& bDynamicsUseAsyncScene) 
 	{
 		const int32 NumBodies = Bodies.Num();
 		PSyncActors.Reserve(NumBodies);
@@ -1167,7 +1167,8 @@ struct FInitBodiesHelper
 
 		// Ensure we have the AggGeom inside the body setup so we can calculate the number of shapes
 		BodySetup->CreatePhysicsMeshes();
-		for (int32 BodyIdx = 0; BodyIdx < NumBodies; ++BodyIdx)
+
+		for (int32 BodyIdx = NumBodies - 1; BodyIdx >= 0; BodyIdx--)   // iterate in reverse since list might shrink
 		{
 			FBodyInstance* Instance = Bodies[BodyIdx];
 			const FTransform& Transform = Transforms[BodyIdx];
@@ -1199,8 +1200,9 @@ struct FInitBodiesHelper
 			if (Instance->GetPxRigidActor_AssumesLocked())
 			{
 				Instance->OwnerComponent = nullptr;
-				Instance->BodySetup = nullptr;
-
+				Instance->BodySetup      = nullptr;
+				Bodies.RemoveAt(BodyIdx);  // so we wont add it to the physx scene again later.
+				Transforms.RemoveAt(BodyIdx);
 				continue;
 			}
 
@@ -1373,7 +1375,7 @@ struct FInitBodiesHelper
 		}
 	}
 
-	void InitBodies_PhysX() const
+	void InitBodies_PhysX() 
 	{
 		TArray<PxActor*> PSyncActors;
 		TArray<PxActor*> PAsyncActors;
@@ -1584,22 +1586,10 @@ void FBodyInstance::InitBody(class UBodySetup* Setup, const FTransform& Transfor
 	check(Setup);
 	check(InRBScene);
 	
-	static TArray<FBodyInstance*> Bodies;
-	static TArray<FTransform> Transforms;
-
-	if(Bodies.Num() == 1 && Transforms.Num() == 1)
-	{
-		Bodies[0] = this;
-		Transforms[0] = Transform;
-
-	}
-	else
-	{
-		Bodies.Add(this);
-		Transforms.Add(Transform);
-	}
-	
-	check(Bodies.Num() == Transforms.Num() == 1);
+	TArray<FBodyInstance*> Bodies;
+	TArray<FTransform> Transforms;
+	Bodies.Add(this);
+	Transforms.Add(Transform);
 
 	FInitBodiesHelper<false> InitBodiesHelper(Bodies, Transforms, Setup, PrimComp, InRBScene, InAggregate);
 	InitBodiesHelper.InitBodies();
