Fix non-unity build

diff --git a/Engine/Source/Runtime/Engine/Private/PhysicsEngine/PhysLevel.cpp b/Engine/Source/Runtime/Engine/Private/PhysicsEngine/PhysLevel.cpp
--- a/Engine/Source/Runtime/Engine/Private/PhysicsEngine/PhysLevel.cpp
+++ b/Engine/Source/Runtime/Engine/Private/PhysicsEngine/PhysLevel.cpp
@@ -4,6 +4,7 @@
 #include "PhysicsPublic.h"
 #include "ParticleDefinitions.h"
 #include "PrecomputedLightVolume.h"
+#include "PhysicsEngine/ClothManager.h"
 
 #if WITH_PHYSX
 	#include "PhysXSupport.h"
