Change the format for error messages when compiling shaders on the build machines, so that the post processor can parse them correctly. Now formats them all on one line, and outputs errors with a verbosity of error.

diff --git a/Engine/Source/Runtime/Engine/Private/ShaderCompiler/ShaderCompiler.cpp b/Engine/Source/Runtime/Engine/Private/ShaderCompiler/ShaderCompiler.cpp
--- a/Engine/Source/Runtime/Engine/Private/ShaderCompiler/ShaderCompiler.cpp
+++ b/Engine/Source/Runtime/Engine/Private/ShaderCompiler/ShaderCompiler.cpp
@@ -1440,7 +1440,7 @@ bool FShaderCompilingManager::HandlePotentialRetryOnError(TMap<int32, FShaderMap
 									CurrentError.ErrorFile = FString(CurrentJob.ShaderType->GetShaderFilename()) + TEXT(".usf");
 								}
 
-								FString UniqueErrorString;
+								FString UniqueErrorPrefix;
 
 								if ( CurrentJob.ShaderType )
 								{
@@ -1449,21 +1449,26 @@ bool FShaderCompilingManager::HandlePotentialRetryOnError(TMap<int32, FShaderMap
 									FString ShaderPath = FPlatformProcess::ShaderDir();
 									FPaths::MakePathRelativeTo(ShaderPath, *SolutionPath);
 									CurrentError.ErrorFile = ShaderPath / CurrentError.ErrorFile;
-									UniqueErrorString = FString::Printf(TEXT("%s(%s): Shader %s, VF %s:\n\t%s\n"), 
+									UniqueErrorPrefix = FString::Printf(TEXT("%s(%s): Shader %s, VF %s:\n\t"), 
 										*CurrentError.ErrorFile, 
 										*CurrentError.ErrorLineString, 
 										CurrentJob.ShaderType->GetName(), 
-										CurrentJob.VFType ? CurrentJob.VFType->GetName() : TEXT("None"), 
-										*CurrentError.StrippedErrorMessage);
+										CurrentJob.VFType ? CurrentJob.VFType->GetName() : TEXT("None"));
 								}
 								else
 								{
-									UniqueErrorString = FString::Printf(TEXT("%s(0): %s\n"), 
-										*CurrentJob.Input.SourceFilename, 
-										*CurrentError.StrippedErrorMessage);
+									UniqueErrorPrefix = FString::Printf(TEXT("%s(0): "), 
+										*CurrentJob.Input.SourceFilename);
 								}
 
-								if (FPlatformMisc::IsDebuggerPresent() && !GIsBuildMachine)
+								FString UniqueErrorString = UniqueErrorPrefix + CurrentError.StrippedErrorMessage + TEXT("\n");
+
+								if(GIsBuildMachine)
+								{
+									// Format everything on one line, and with the correct verbosity, so we can display proper errors in the failure logs.
+									UE_LOG(LogShaderCompilers, Error, TEXT("%s%s"), *UniqueErrorPrefix.Replace(TEXT("\n"), TEXT("")), *CurrentError.StrippedErrorMessage);
+								}
+								else if (FPlatformMisc::IsDebuggerPresent() && !GIsBuildMachine)
 								{
 									// Using OutputDebugString to avoid any text getting added before the filename,
 									// Which will throw off VS.NET's ability to take you directly to the file and line of the error when double clicking it in the output window.
