Added a safety feature that should address navmesh-generation related classes when BP-implemented actors get their construction script re-run just after navmesh generation ending #UE4

UE-20646

diff --git a/Engine/Source/Runtime/Engine/Classes/AI/Navigation/RecastNavMesh.h b/Engine/Source/Runtime/Engine/Classes/AI/Navigation/RecastNavMesh.h
--- a/Engine/Source/Runtime/Engine/Classes/AI/Navigation/RecastNavMesh.h
+++ b/Engine/Source/Runtime/Engine/Classes/AI/Navigation/RecastNavMesh.h
@@ -786,6 +786,8 @@ public:
 	/** Event from generator that navmesh build has finished */
 	virtual void OnNavMeshGenerationFinished();
 
+	virtual void EnsureBuildCompletion() override;
+
 	virtual void SetConfig(const FNavDataConfig& Src) override;
 protected:
 	virtual void FillConfig(FNavDataConfig& Dest) override;
@@ -830,6 +832,8 @@ public:
 	/** Area sort function */
 	virtual void SortAreasForGenerator(TArray<FRecastAreaNavModifierElement>& Areas) const;
 
+	void RecreateDefaultFilter();
+
 	int32 GetMaxSimultaneousTileGenerationJobsCount() const { return MaxSimultaneousTileGenerationJobsCount; }
 	void SetMaxSimultaneousTileGenerationJobsCount(int32 NewJobsCountLimit);
 
diff --git a/Engine/Source/Runtime/Engine/Private/AI/Navigation/RecastNavMesh.cpp b/Engine/Source/Runtime/Engine/Private/AI/Navigation/RecastNavMesh.cpp
--- a/Engine/Source/Runtime/Engine/Private/AI/Navigation/RecastNavMesh.cpp
+++ b/Engine/Source/Runtime/Engine/Private/AI/Navigation/RecastNavMesh.cpp
@@ -458,11 +458,7 @@ void ARecastNavMesh::PostInitProperties()
 
 	if (HasAnyFlags(RF_ClassDefaultObject) == false)
 	{
-		DefaultQueryFilter->SetFilterType<FRecastQueryFilter>();
-		FRecastQueryFilter* DetourFilter = static_cast<FRecastQueryFilter*>(DefaultQueryFilter->GetImplementation());
-		DetourFilter->SetIsVirtual(bUseVirtualFilters);
-		DetourFilter->setHeuristicScale(HeuristicScale);
-		DefaultQueryFilter->SetMaxSearchNodes(DefaultMaxSearchNodes);		
+		RecreateDefaultFilter();
 	}
 
 	// voxel cache requires the same rasterization setup for all navmeshes, as it's stored in octree
@@ -514,6 +510,33 @@ void ARecastNavMesh::PostInitProperties()
 	UpdatePolyRefBitsPreview();
 }
 
+void ARecastNavMesh::RecreateDefaultFilter()
+{
+	DefaultQueryFilter->SetFilterType<FRecastQueryFilter>();
+	DefaultQueryFilter->SetMaxSearchNodes(DefaultMaxSearchNodes);
+
+	FRecastQueryFilter* DetourFilter = static_cast<FRecastQueryFilter*>(DefaultQueryFilter->GetImplementation());
+	DetourFilter->SetIsVirtual(bUseVirtualFilters);
+	DetourFilter->setHeuristicScale(HeuristicScale);
+
+	for (int32 Idx = 0; Idx < SupportedAreas.Num(); Idx++)
+	{
+		const FSupportedAreaData& AreaData = SupportedAreas[Idx];
+		
+		UNavArea* DefArea = nullptr;
+		if (AreaData.AreaClass)
+		{
+			DefArea = ((UClass*)AreaData.AreaClass)->GetDefaultObject<UNavArea>();
+		}
+
+		if (DefArea)
+		{
+			DetourFilter->SetAreaCost(AreaData.AreaID, DefArea->DefaultCost);
+			DetourFilter->SetFixedAreaEnteringCost(AreaData.AreaID, DefArea->GetFixedAreaEnteringCost());
+		}
+	}
+}
+
 FVector ARecastNavMesh::GetModifiedQueryExtent(const FVector& QueryExtent) const
 {
 	return FVector(QueryExtent.X, QueryExtent.Y, QueryExtent.Z + FMath::Max(0.0f, VerticalDeviationFromGroundCompensation));
@@ -1474,6 +1497,16 @@ URecastNavMeshDataChunk* ARecastNavMesh::GetNavigationDataChunk(ULevel* InLevel)
 	return RcNavDataChunk;
 }
 
+void ARecastNavMesh::EnsureBuildCompletion()
+{
+	Super::EnsureBuildCompletion();
+
+	// Doing this as a safety net solution due to UE-20646, which was basically a result of random 
+	// over-releasing of default filter's shared pointer (it seemed). We might have time to get 
+	// back to this time some time in next 3 years :D
+	RecreateDefaultFilter();
+}
+
 void ARecastNavMesh::OnNavMeshGenerationFinished()
 {
 	UWorld* World = GetWorld();
