UE4 - Moved churn tracker to launch instead of game ticks.

diff --git a/Engine/Source/Runtime/Engine/Private/GameEngine.cpp b/Engine/Source/Runtime/Engine/Private/GameEngine.cpp
--- a/Engine/Source/Runtime/Engine/Private/GameEngine.cpp
+++ b/Engine/Source/Runtime/Engine/Private/GameEngine.cpp
@@ -203,10 +203,10 @@ void UGameEngine::ConditionallyOverrideSettings(int32& ResolutionX, int32& Resol
 	FParse::Value(FCommandLine::Get(), TEXT("ResX="), ResolutionX);
 	FParse::Value(FCommandLine::Get(), TEXT("ResY="), ResolutionY);
 
-	// consume available desktop area
-	FDisplayMetrics DisplayMetrics;
+		// consume available desktop area
+		FDisplayMetrics DisplayMetrics;
 	FDisplayMetrics::GetDisplayMetrics(DisplayMetrics);
-
+		
 	int32 DesktopResolutionX = DisplayMetrics.PrimaryDisplayWidth;
 	int32 DesktopResolutionY = DisplayMetrics.PrimaryDisplayHeight;
 
@@ -548,7 +548,7 @@ void UGameEngine::PreExit()
 
 			if (World->GetGameInstance() != nullptr)
 			{
-				World->GetGameInstance()->Shutdown();
+			World->GetGameInstance()->Shutdown();
 			}
 
 			World->FlushLevelStreaming(EFlushLevelStreamingType::Visibility);
@@ -811,128 +811,9 @@ float UGameEngine::GetMaxTickRate(float DeltaTime, bool bAllowFrameRateSmoothing
 	return MaxTickRate;
 }
 
-#if !UE_BUILD_SHIPPING && !UE_BUILD_TEST
-
-#include "StackTracker.h"
-static TAutoConsoleVariable<int32> CVarLogGameThreadMallocChurn(
-	TEXT("LogGameThreadMallocChurn.Enable"),
-	0,
-	TEXT("If > 0, then collect sample game thread malloc, realloc and free, periodically print a report of the worst offenders."));
-
-static TAutoConsoleVariable<int32> CVarLogGameThreadMallocChurn_PrintFrequency(
-	TEXT("LogGameThreadMallocChurn.PrintFrequency"),
-	300,
-	TEXT("Number of frames between churn reports."));
-
-static TAutoConsoleVariable<int32> CVarLogGameThreadMallocChurn_Threshhold(
-	TEXT("LogGameThreadMallocChurn.Threshhold"),
-	10,
-	TEXT("Minimum average number of allocs per frame to include in the report."));
-
-static TAutoConsoleVariable<int32> CVarLogGameThreadMallocChurn_SampleFrequency(
-	TEXT("LogGameThreadMallocChurn.SampleFrequency"),
-	100,
-	TEXT("Number of allocs to skip between samples. This is used to prevent churn sampling from slowing the game down too much."));
-
-static TAutoConsoleVariable<int32> CVarLogGameThreadMallocChurn_StackIgnore(
-	TEXT("LogGameThreadMallocChurn.StackIgnore"),
-	2,
-	TEXT("Number of items to discard from the top of a stack frame."));
-
-static TAutoConsoleVariable<int32> CVarLogGameThreadMallocChurn_StackLen(
-	TEXT("LogGameThreadMallocChurn.StackLen"),
-	3,
-	TEXT("Maximum number of stack frame items to keep. This improves aggregation because calls that originate from multiple places but end up in the same place will be accounted together."));
-
-
-extern CORE_API TFunction<void(int32)>* GGameThreadMallocHook;
-
-struct FScopedSampleMallocChurn
-{
-	static FStackTracker GGameThreadMallocChurnTracker;
-	static uint64 DumpFrame;
-
-	bool bEnabled;
-	int32 CountDown;
-	TFunction<void(int32)> Hook;
-
-	FScopedSampleMallocChurn()
-		: bEnabled(CVarLogGameThreadMallocChurn.GetValueOnGameThread() > 0)
-		, CountDown(CVarLogGameThreadMallocChurn_SampleFrequency.GetValueOnGameThread())
-		, Hook(
-		[this](int32 Index)
-	{
-		if (--CountDown <= 0)
-		{
-			CountDown = CVarLogGameThreadMallocChurn_SampleFrequency.GetValueOnGameThread();
-			CollectSample();
-		}
-	}
-	)
-	{
-		if (bEnabled)
-		{
-			check(IsInGameThread());
-			check(!GGameThreadMallocHook);
-			if (!DumpFrame)
-			{
-				DumpFrame = GFrameCounter + CVarLogGameThreadMallocChurn_PrintFrequency.GetValueOnGameThread();
-				GGameThreadMallocChurnTracker.ResetTracking();
-			}
-			GGameThreadMallocChurnTracker.ToggleTracking(true, true);
-			GGameThreadMallocHook = &Hook;
-		}
-		else
-		{
-			check(IsInGameThread());
-			GGameThreadMallocChurnTracker.ToggleTracking(false, true);
-			if (DumpFrame)
-			{
-				DumpFrame = 0;
-				GGameThreadMallocChurnTracker.ResetTracking();
-			}
-		}
-	}
-	~FScopedSampleMallocChurn()
-	{
-		if (bEnabled)
-		{
-			check(IsInGameThread());
-			check(GGameThreadMallocHook == &Hook);
-			GGameThreadMallocHook = nullptr;
-			GGameThreadMallocChurnTracker.ToggleTracking(false, true);
-			check(DumpFrame);
-			if (GFrameCounter > DumpFrame)
-			{
-				PrintResultsAndReset();
-			}
-		}
-	}
-
-	void CollectSample()
-	{
-		check(IsInGameThread());
-		GGameThreadMallocChurnTracker.CaptureStackTrace(CVarLogGameThreadMallocChurn_StackIgnore.GetValueOnGameThread(), nullptr, CVarLogGameThreadMallocChurn_StackLen.GetValueOnGameThread());
-	}
-	void PrintResultsAndReset()
-	{
-		DumpFrame = GFrameCounter + CVarLogGameThreadMallocChurn_PrintFrequency.GetValueOnGameThread();
-		FOutputDeviceRedirector* Log = FOutputDeviceRedirector::Get();
-		float SampleAndFrameCorrection = float(CVarLogGameThreadMallocChurn_SampleFrequency.GetValueOnGameThread()) / float(CVarLogGameThreadMallocChurn_PrintFrequency.GetValueOnGameThread());
-		GGameThreadMallocChurnTracker.DumpStackTraces(CVarLogGameThreadMallocChurn_Threshhold.GetValueOnGameThread(), *Log, SampleAndFrameCorrection);
-		GGameThreadMallocChurnTracker.ResetTracking();
-	}
-};
-FStackTracker FScopedSampleMallocChurn::GGameThreadMallocChurnTracker;
-uint64 FScopedSampleMallocChurn::DumpFrame = 0;
-
-#endif
 
 void UGameEngine::Tick( float DeltaSeconds, bool bIdleMode )
 {
-#if !UE_BUILD_SHIPPING && !UE_BUILD_TEST
-	FScopedSampleMallocChurn ChurnTracker;
-#endif
 	SCOPE_CYCLE_COUNTER(STAT_GameEngineTick);
 	NETWORK_PROFILER(GNetworkProfiler.TrackFrameBegin());
 
@@ -1115,11 +996,11 @@ void UGameEngine::Tick( float DeltaSeconds, bool bIdleMode )
 
 		if (Context.WorldType != EWorldType::Preview)
 		{
-			// Update Audio. This needs to occur after rendering as the rendering code updates the listener position.
-			if (FAudioDevice* AudioDevice = Context.World()->GetAudioDevice())
-			{
-				AudioDevice->Update(!Context.World()->IsPaused());
-			}
+		// Update Audio. This needs to occur after rendering as the rendering code updates the listener position.
+		if (FAudioDevice* AudioDevice = Context.World()->GetAudioDevice())
+		{
+			AudioDevice->Update(!Context.World()->IsPaused());
+		}
 		}
 
 		if( GIsClient )
diff --git a/Engine/Source/Runtime/Launch/Private/LaunchEngineLoop.cpp b/Engine/Source/Runtime/Launch/Private/LaunchEngineLoop.cpp
--- a/Engine/Source/Runtime/Launch/Private/LaunchEngineLoop.cpp
+++ b/Engine/Source/Runtime/Launch/Private/LaunchEngineLoop.cpp
@@ -2350,8 +2350,128 @@ bool FEngineLoop::ShouldUseIdleMode() const
 	return bIdleMode;
 }
 
+#if !UE_BUILD_SHIPPING && !UE_BUILD_TEST
+
+#include "StackTracker.h"
+static TAutoConsoleVariable<int32> CVarLogGameThreadMallocChurn(
+	TEXT("LogGameThreadMallocChurn.Enable"),
+	0,
+	TEXT("If > 0, then collect sample game thread malloc, realloc and free, periodically print a report of the worst offenders."));
+
+static TAutoConsoleVariable<int32> CVarLogGameThreadMallocChurn_PrintFrequency(
+	TEXT("LogGameThreadMallocChurn.PrintFrequency"),
+	300,
+	TEXT("Number of frames between churn reports."));
+
+static TAutoConsoleVariable<int32> CVarLogGameThreadMallocChurn_Threshhold(
+	TEXT("LogGameThreadMallocChurn.Threshhold"),
+	10,
+	TEXT("Minimum average number of allocs per frame to include in the report."));
+
+static TAutoConsoleVariable<int32> CVarLogGameThreadMallocChurn_SampleFrequency(
+	TEXT("LogGameThreadMallocChurn.SampleFrequency"),
+	100,
+	TEXT("Number of allocs to skip between samples. This is used to prevent churn sampling from slowing the game down too much."));
+
+static TAutoConsoleVariable<int32> CVarLogGameThreadMallocChurn_StackIgnore(
+	TEXT("LogGameThreadMallocChurn.StackIgnore"),
+	2,
+	TEXT("Number of items to discard from the top of a stack frame."));
+
+static TAutoConsoleVariable<int32> CVarLogGameThreadMallocChurn_StackLen(
+	TEXT("LogGameThreadMallocChurn.StackLen"),
+	3,
+	TEXT("Maximum number of stack frame items to keep. This improves aggregation because calls that originate from multiple places but end up in the same place will be accounted together."));
+
+
+extern CORE_API TFunction<void(int32)>* GGameThreadMallocHook;
+
+struct FScopedSampleMallocChurn
+{
+	static FStackTracker GGameThreadMallocChurnTracker;
+	static uint64 DumpFrame;
+
+	bool bEnabled;
+	int32 CountDown;
+	TFunction<void(int32)> Hook;
+
+	FScopedSampleMallocChurn()
+		: bEnabled(CVarLogGameThreadMallocChurn.GetValueOnGameThread() > 0)
+		, CountDown(CVarLogGameThreadMallocChurn_SampleFrequency.GetValueOnGameThread())
+		, Hook(
+		[this](int32 Index)
+	{
+		if (--CountDown <= 0)
+		{
+			CountDown = CVarLogGameThreadMallocChurn_SampleFrequency.GetValueOnGameThread();
+			CollectSample();
+		}
+	}
+	)
+	{
+		if (bEnabled)
+		{
+			check(IsInGameThread());
+			check(!GGameThreadMallocHook);
+			if (!DumpFrame)
+			{
+				DumpFrame = GFrameCounter + CVarLogGameThreadMallocChurn_PrintFrequency.GetValueOnGameThread();
+				GGameThreadMallocChurnTracker.ResetTracking();
+			}
+			GGameThreadMallocChurnTracker.ToggleTracking(true, true);
+			GGameThreadMallocHook = &Hook;
+		}
+		else
+		{
+			check(IsInGameThread());
+			GGameThreadMallocChurnTracker.ToggleTracking(false, true);
+			if (DumpFrame)
+			{
+				DumpFrame = 0;
+				GGameThreadMallocChurnTracker.ResetTracking();
+			}
+		}
+	}
+	~FScopedSampleMallocChurn()
+	{
+		if (bEnabled)
+		{
+			check(IsInGameThread());
+			check(GGameThreadMallocHook == &Hook);
+			GGameThreadMallocHook = nullptr;
+			GGameThreadMallocChurnTracker.ToggleTracking(false, true);
+			check(DumpFrame);
+			if (GFrameCounter > DumpFrame)
+			{
+				PrintResultsAndReset();
+			}
+		}
+	}
+
+	void CollectSample()
+	{
+		check(IsInGameThread());
+		GGameThreadMallocChurnTracker.CaptureStackTrace(CVarLogGameThreadMallocChurn_StackIgnore.GetValueOnGameThread(), nullptr, CVarLogGameThreadMallocChurn_StackLen.GetValueOnGameThread());
+	}
+	void PrintResultsAndReset()
+	{
+		DumpFrame = GFrameCounter + CVarLogGameThreadMallocChurn_PrintFrequency.GetValueOnGameThread();
+		FOutputDeviceRedirector* Log = FOutputDeviceRedirector::Get();
+		float SampleAndFrameCorrection = float(CVarLogGameThreadMallocChurn_SampleFrequency.GetValueOnGameThread()) / float(CVarLogGameThreadMallocChurn_PrintFrequency.GetValueOnGameThread());
+		GGameThreadMallocChurnTracker.DumpStackTraces(CVarLogGameThreadMallocChurn_Threshhold.GetValueOnGameThread(), *Log, SampleAndFrameCorrection);
+		GGameThreadMallocChurnTracker.ResetTracking();
+	}
+};
+FStackTracker FScopedSampleMallocChurn::GGameThreadMallocChurnTracker;
+uint64 FScopedSampleMallocChurn::DumpFrame = 0;
+
+#endif
+
 void FEngineLoop::Tick()
 {
+#if !UE_BUILD_SHIPPING && !UE_BUILD_TEST
+	FScopedSampleMallocChurn ChurnTracker;
+#endif
 	// Ensure we aren't starting a frame while loading or playing a loading movie
 	ensure(GetMoviePlayer()->IsLoadingFinished() && !GetMoviePlayer()->IsMovieCurrentlyPlaying());
 
