Merging CL #2693296: Fix particle vertex declaration mismatch on iOS Metal.

diff --git a/Engine/Source/Runtime/Engine/Public/ParticleHelper.h b/Engine/Source/Runtime/Engine/Public/ParticleHelper.h
--- a/Engine/Source/Runtime/Engine/Public/ParticleHelper.h
+++ b/Engine/Source/Runtime/Engine/Public/ParticleHelper.h
@@ -1645,7 +1645,7 @@ struct FDynamicSpriteEmitterData : public FDynamicSpriteEmitterDataBase
 	 */
 	virtual int32 GetDynamicVertexStride(ERHIFeatureLevel::Type InFeatureLevel) const override
 	{
-		const bool bInstanced = InFeatureLevel >= ERHIFeatureLevel::SM4;
+		const bool bInstanced = RHISupportsInstancing(GetFeatureLevelShaderPlatform(InFeatureLevel));
 		return bInstanced ? sizeof(FParticleSpriteVertex) : sizeof(FParticleSpriteVertexNonInstanced);
 	}
 
