Merging CL #2690470: Change the FVectorFieldUniformParameters structure alignment again to resolve another alignment problem with the FVector/float3 arrays. It isn't clear to me whether this mis-alignment is in our TStaticArray which should be aligning each 12-byte FVector to a 16-byte boundary, or the Metal SL where float3 should be 16-byte aligned, but they demonstrably don't match in ElementalDemo.

diff --git a/Engine/Source/Runtime/Engine/Private/Particles/ParticleGpuSimulation.cpp b/Engine/Source/Runtime/Engine/Private/Particles/ParticleGpuSimulation.cpp
--- a/Engine/Source/Runtime/Engine/Private/Particles/ParticleGpuSimulation.cpp
+++ b/Engine/Source/Runtime/Engine/Private/Particles/ParticleGpuSimulation.cpp
@@ -765,8 +765,8 @@ BEGIN_UNIFORM_BUFFER_STRUCT( FVectorFieldUniformParameters,)
 	DECLARE_UNIFORM_BUFFER_STRUCT_MEMBER_ARRAY( FMatrix, WorldToVolume, [MAX_VECTOR_FIELDS] )
 	DECLARE_UNIFORM_BUFFER_STRUCT_MEMBER_ARRAY( FMatrix, VolumeToWorld, [MAX_VECTOR_FIELDS] )
 	DECLARE_UNIFORM_BUFFER_STRUCT_MEMBER_ARRAY( FVector4, IntensityAndTightness, [MAX_VECTOR_FIELDS] )
-	DECLARE_UNIFORM_BUFFER_STRUCT_MEMBER_ARRAY( FVector, VolumeSize, [MAX_VECTOR_FIELDS] )
-	DECLARE_UNIFORM_BUFFER_STRUCT_MEMBER_ARRAY( FVector, TilingAxes, [MAX_VECTOR_FIELDS] )
+	DECLARE_UNIFORM_BUFFER_STRUCT_MEMBER_ARRAY( FVector4, VolumeSize, [MAX_VECTOR_FIELDS] )
+	DECLARE_UNIFORM_BUFFER_STRUCT_MEMBER_ARRAY( FVector4, TilingAxes, [MAX_VECTOR_FIELDS] )
 END_UNIFORM_BUFFER_STRUCT( FVectorFieldUniformParameters )
 
 IMPLEMENT_UNIFORM_BUFFER_STRUCT(FVectorFieldUniformParameters,TEXT("VectorFields"));
@@ -4161,7 +4161,7 @@ static void SetParametersForVectorField(FVectorFieldUniformParameters& OutParame
 
 	OutParameters.WorldToVolume[Index] = VectorFieldInstance->WorldToVolume;
 	OutParameters.VolumeToWorld[Index] = VectorFieldInstance->VolumeToWorldNoScale;
-	OutParameters.VolumeSize[Index] = FVector(Resource->SizeX, Resource->SizeY, Resource->SizeZ);
+	OutParameters.VolumeSize[Index] = FVector4(Resource->SizeX, Resource->SizeY, Resource->SizeZ, 0);
 	OutParameters.IntensityAndTightness[Index] = FVector4(Intensity, Tightness, 0, 0 );
 	OutParameters.TilingAxes[Index].X = VectorFieldInstance->bTileX ? 1.0f : 0.0f;
 	OutParameters.TilingAxes[Index].Y = VectorFieldInstance->bTileY ? 1.0f : 0.0f;
@@ -4243,7 +4243,7 @@ void FFXSystem::SimulateGPUParticles(
 		{
 			VectorFieldParameters.WorldToVolume[Index] = FMatrix::Identity;
 			VectorFieldParameters.VolumeToWorld[Index] = FMatrix::Identity;
-			VectorFieldParameters.VolumeSize[Index] = FVector(1.0f);
+			VectorFieldParameters.VolumeSize[Index] = FVector4(1.0f);
 			VectorFieldParameters.IntensityAndTightness[Index] = FVector4(0.0f);
 		}
 		VectorFieldParameters.Count = 0;
@@ -4317,7 +4317,7 @@ void FFXSystem::SimulateGPUParticles(
 						const int32 Index = PadCount++;
 						VectorFieldParameters.WorldToVolume[Index] = FMatrix::Identity;
 						VectorFieldParameters.VolumeToWorld[Index] = FMatrix::Identity;
-						VectorFieldParameters.VolumeSize[Index] = FVector(1.0f);
+						VectorFieldParameters.VolumeSize[Index] = FVector4(1.0f);
 						VectorFieldParameters.IntensityAndTightness[Index] = FVector4(0.0f);
 					}
 					SimulationCommand->VectorFieldsUniformBuffer = FVectorFieldUniformBufferRef::CreateUniformBufferImmediate(VectorFieldParameters, UniformBuffer_SingleFrame);
