[UE-12263] Fix a situation where a previously-disabled node could be left enabled after serialization during compile-on-load.

diff --git a/Engine/Source/Runtime/Engine/Classes/EdGraph/EdGraphNode.h b/Engine/Source/Runtime/Engine/Classes/EdGraph/EdGraphNode.h
--- a/Engine/Source/Runtime/Engine/Classes/EdGraph/EdGraphNode.h
+++ b/Engine/Source/Runtime/Engine/Classes/EdGraph/EdGraphNode.h
@@ -209,6 +209,7 @@ private:
 public:
 	// UObject interface
 	static void AddReferencedObjects(UObject* InThis, FReferenceCollector& Collector);
+	virtual void Serialize(FArchive& Ar) override;
 	virtual void PostLoad() override;
 	virtual void PostEditChangeProperty(FPropertyChangedEvent& PropertyChangedEvent) override;
 	// End of UObject interface
diff --git a/Engine/Source/Runtime/Engine/Private/EdGraph/EdGraphNode.cpp b/Engine/Source/Runtime/Engine/Private/EdGraph/EdGraphNode.cpp
--- a/Engine/Source/Runtime/Engine/Private/EdGraph/EdGraphNode.cpp
+++ b/Engine/Source/Runtime/Engine/Private/EdGraph/EdGraphNode.cpp
@@ -257,6 +257,21 @@ void UEdGraphNode::AddReferencedObjects(UObject* InThis, FReferenceCollector& Co
 	Super::AddReferencedObjects(This, Collector);
 }
 
+void UEdGraphNode::Serialize(FArchive& Ar)
+{
+	Super::Serialize(Ar);
+
+	if (Ar.IsLoading())
+	{
+		// If this was an older version, ensure that we update the enabled state for already-disabled nodes.
+		// Note: We need to do this here and not in PostLoad() as it must be assigned prior to compile-on-load.
+		if(!bIsNodeEnabled_DEPRECATED && !bUserSetEnabledState && EnabledState == ENodeEnabledState::Enabled)
+		{
+			EnabledState = ENodeEnabledState::Disabled;
+		}
+	}
+}
+
 void UEdGraphNode::PostLoad()
 {
 	Super::PostLoad();
@@ -281,12 +296,6 @@ void UEdGraphNode::PostLoad()
 	{
 		bCommentBubbleVisible = !NodeComment.IsEmpty();
 	}
-
-	// If this was an older version, ensure that we update the enabled state to match
-	if(!bIsNodeEnabled_DEPRECATED)
-	{
-		EnabledState = ENodeEnabledState::Disabled;
-	}
 }
 
 void UEdGraphNode::PostEditChangeProperty(FPropertyChangedEvent& PropertyChangedEvent)
