integrated from Orion

CL 2694266

	optimized BasePass Pixel shader to not have to combine two bits to create a new value from the PrimitiveBuffer

diff --git a/Engine/Shaders/BasePassPixelShader.usf b/Engine/Shaders/BasePassPixelShader.usf
--- a/Engine/Shaders/BasePassPixelShader.usf
+++ b/Engine/Shaders/BasePassPixelShader.usf
@@ -845,8 +845,7 @@ void Main(
 		GBuffer.IndirectIrradiance = IndirectIrradiance;
 		GBuffer.PrecomputedShadowFactors = PrecomputedShadowFactors;
 		GBuffer.GBufferAO = MaterialAO;
-		GBuffer.HasDistanceFieldRepresentation = Primitive.HasDistanceFieldRepresentation;
-		GBuffer.HasHeightfieldRepresentation = Primitive.HasHeightfieldRepresentation;
+		GBuffer.PerObjectGBufferData = Primitive.PerObjectGBufferData;
 
 		#if MATERIAL_SHADINGMODEL_UNLIT
 			GBuffer.ShadingModelID = SHADINGMODELID_UNLIT;
diff --git a/Engine/Shaders/DeferredDecal.usf b/Engine/Shaders/DeferredDecal.usf
--- a/Engine/Shaders/DeferredDecal.usf
+++ b/Engine/Shaders/DeferredDecal.usf
@@ -355,8 +355,7 @@ void MainPS(
 	Data.PrecomputedShadowFactors = 0;
 	Data.GBufferAO = 1;
 	Data.ShadingModelID = SHADINGMODELID_DEFAULT_LIT;
-	Data.HasDistanceFieldRepresentation = 1;
-	Data.HasHeightfieldRepresentation = 1;
+	Data.PerObjectGBufferData = 1;
 
 	#if DECAL_RENDERSTAGE == 0
 	{
diff --git a/Engine/Shaders/DeferredShadingCommon.usf b/Engine/Shaders/DeferredShadingCommon.usf
--- a/Engine/Shaders/DeferredShadingCommon.usf
+++ b/Engine/Shaders/DeferredShadingCommon.usf
@@ -261,10 +261,8 @@ struct FGBufferData
 	float GBufferAO;
 	// 0..255 
 	uint ShadingModelID;
-	// 0..1
-	uint HasDistanceFieldRepresentation; 
-	// 0..1
-	uint HasHeightfieldRepresentation; 
+	// 0..1, 2 bits, use HasDistanceFieldRepresentation(GBuffer) or HasHeightfieldRepresentation(GBuffer) to extract
+	float PerObjectGBufferData;
 	// in world units
 	float CustomDepth;
 	// Custom depth stencil value
@@ -281,6 +279,20 @@ struct FGBufferData
 	float3 StoredSpecular;
 };
 
+bool HasDistanceFieldRepresentation(FGBufferData GBufferData)
+{
+	// can be optimized
+	uint PackedAlpha = (uint)(GBufferData.PerObjectGBufferData * 3.999f);
+	return PackedAlpha & 1;
+}
+
+bool HasHeightfieldRepresentation(FGBufferData GBufferData)
+{
+	// can be optimized
+	uint PackedAlpha = (uint)(GBufferData.PerObjectGBufferData * 3.999f);
+	return (PackedAlpha & 2) != 0;
+}
+
 
 // all values that are output by the forward rendering pass
 struct FDBufferData
@@ -324,11 +336,7 @@ void EncodeGBuffer(
 	)
 {
 	OutGBufferA.rgb = EncodeNormal( GBuffer.WorldNormal );
-
-	// compress in 2 bits
-	// Note: distance field lighting dependencies (normal, representation bits) packed for minimal texture fetches
-	uint PackedAlpha = GBuffer.HasHeightfieldRepresentation * 2 + (GBuffer.HasDistanceFieldRepresentation);
-	OutGBufferA.a = PackedAlpha * (1.0 / 3.0);
+	OutGBufferA.a = GBuffer.PerObjectGBufferData;
 
 	if (GBuffer.ShadingModelID == SHADINGMODELID_UNLIT)
 	{
@@ -414,9 +422,7 @@ FGBufferData DecodeGBufferData(
 		GBuffer.WorldNormal = normalize(GBuffer.WorldNormal);
 	}
 
-	uint PackedAlpha = (uint)( InGBufferA.a * 3.999 );
-	GBuffer.HasHeightfieldRepresentation = (PackedAlpha & 2) != 0;
-	GBuffer.HasDistanceFieldRepresentation = PackedAlpha & 1;
+	GBuffer.PerObjectGBufferData = InGBufferA.a;
 	GBuffer.ShadingModelID = DecodeShadingModelId(InGBufferB.a);
 	
 	GBuffer.StoredBaseColor = DecodeBaseColor(InGBufferC.rgb);
diff --git a/Engine/Shaders/DistanceFieldAOShared.usf b/Engine/Shaders/DistanceFieldAOShared.usf
--- a/Engine/Shaders/DistanceFieldAOShared.usf
+++ b/Engine/Shaders/DistanceFieldAOShared.usf
@@ -92,8 +92,8 @@ SamplerState DistanceFieldNormalSampler;
 
 float4 EncodeDownsampledGBuffer(FGBufferData GBufferData, float SceneDepth)
 {
-	float DepthSign = GBufferData.HasDistanceFieldRepresentation > 0 ? 1 : -1;
-	float ZSign = GBufferData.HasHeightfieldRepresentation > 0 ? 1 : -1;
+	float DepthSign = HasDistanceFieldRepresentation(GBufferData) > 0 ? 1 : -1;
+	float ZSign = HasHeightfieldRepresentation(GBufferData)> 0 ? 1 : -1;
 	float ZChannel = GBufferData.WorldNormal.z > 0 ? .5f : 1;
 	float4 OutValue = float4(GBufferData.WorldNormal.xy, ZChannel * ZSign, DepthSign * SceneDepth);
 	return OutValue;
diff --git a/Engine/Source/Runtime/Engine/Public/PrimitiveUniformShaderParameters.h b/Engine/Source/Runtime/Engine/Public/PrimitiveUniformShaderParameters.h
--- a/Engine/Source/Runtime/Engine/Public/PrimitiveUniformShaderParameters.h
+++ b/Engine/Source/Runtime/Engine/Public/PrimitiveUniformShaderParameters.h
@@ -9,15 +9,14 @@
 
 /** The uniform shader parameters associated with a primitive. */
 BEGIN_UNIFORM_BUFFER_STRUCT(FPrimitiveUniformShaderParameters,ENGINE_API)
-	DECLARE_UNIFORM_BUFFER_STRUCT_MEMBER(FMatrix,LocalToWorld)
-	DECLARE_UNIFORM_BUFFER_STRUCT_MEMBER(FMatrix,WorldToLocal)
-	DECLARE_UNIFORM_BUFFER_STRUCT_MEMBER(FVector4,ObjectWorldPositionAndRadius)
-	DECLARE_UNIFORM_BUFFER_STRUCT_MEMBER(FVector,ObjectBounds)
-	DECLARE_UNIFORM_BUFFER_STRUCT_MEMBER_EX(float,LocalToWorldDeterminantSign,EShaderPrecisionModifier::Half)
+	DECLARE_UNIFORM_BUFFER_STRUCT_MEMBER(FMatrix,LocalToWorld)		// always needed
+	DECLARE_UNIFORM_BUFFER_STRUCT_MEMBER(FMatrix,WorldToLocal)		// rarely needed
+	DECLARE_UNIFORM_BUFFER_STRUCT_MEMBER(FVector4,ObjectWorldPositionAndRadius)	// needed by some materials
+	DECLARE_UNIFORM_BUFFER_STRUCT_MEMBER(FVector,ObjectBounds)		// only needed for editor/development
+	DECLARE_UNIFORM_BUFFER_STRUCT_MEMBER_EX(float,LocalToWorldDeterminantSign,EShaderPrecisionModifier::Half)	// could be stored in the sign bit of the object radius
 	DECLARE_UNIFORM_BUFFER_STRUCT_MEMBER(FVector,ActorWorldPosition)
 	DECLARE_UNIFORM_BUFFER_STRUCT_MEMBER_EX(float,DecalReceiverMask,EShaderPrecisionModifier::Half)
-	DECLARE_UNIFORM_BUFFER_STRUCT_MEMBER_EX(float,HasDistanceFieldRepresentation,EShaderPrecisionModifier::Half)
-	DECLARE_UNIFORM_BUFFER_STRUCT_MEMBER_EX(float,HasHeightfieldRepresentation,EShaderPrecisionModifier::Half)
+	DECLARE_UNIFORM_BUFFER_STRUCT_MEMBER_EX(float,PerObjectGBufferData,EShaderPrecisionModifier::Half)		// 0..1, 2 bits, bDistanceFieldRepresentation, bHeightfieldRepresentation
 	DECLARE_UNIFORM_BUFFER_STRUCT_MEMBER_EX(float,UseEditorDepthTest,EShaderPrecisionModifier::Half)
 	DECLARE_UNIFORM_BUFFER_STRUCT_MEMBER_EX(FVector4,ObjectOrientation,EShaderPrecisionModifier::Half)
 	DECLARE_UNIFORM_BUFFER_STRUCT_MEMBER_EX(FVector4,NonUniformScale,EShaderPrecisionModifier::Half)
@@ -70,8 +69,7 @@ inline FPrimitiveUniformShaderParameters GetPrimitiveUniformShaderParameters(
 
 	Result.LocalToWorldDeterminantSign = FMath::FloatSelect(LocalToWorld.RotDeterminant(),1,-1);
 	Result.DecalReceiverMask = bReceivesDecals ? 1 : 0;
-	Result.HasDistanceFieldRepresentation = bHasDistanceFieldRepresentation ? 1 : 0;
-	Result.HasHeightfieldRepresentation = bHasHeightfieldRepresentation ? 1 : 0;
+	Result.PerObjectGBufferData = (2 * (int32)bHasHeightfieldRepresentation + (int32)bHasDistanceFieldRepresentation) / 3.0f;
 	Result.UseEditorDepthTest = bUseEditorDepthTest ? 1 : 0;
 	return Result;
 }
