#ue4 - ParticleSystems: Avoid useless detach/attach during registration when bAutoManageAttachment is enabled and the particle system auto-activates.

--------

diff --git a/Engine/Source/Runtime/Engine/Private/Particles/ParticleComponents.cpp b/Engine/Source/Runtime/Engine/Private/Particles/ParticleComponents.cpp
--- a/Engine/Source/Runtime/Engine/Private/Particles/ParticleComponents.cpp
+++ b/Engine/Source/Runtime/Engine/Private/Particles/ParticleComponents.cpp
@@ -2880,7 +2880,11 @@ void UParticleSystemComponent::OnRegister()
 			// Prevent attachment before Super::OnRegister() tries to attach us, since we only attach when activated.
 			if (AttachParent->AttachChildren.Contains(this))
 			{
-				DetachFromParent(/*bMaintainWorldPosition=*/ false, /*bCallModify=*/ false);
+				// Only detach if we are not about to auto attach to the same target, that would be wasteful.
+				if (!bAutoActivate || (AutoAttachLocationType != EAttachLocation::KeepRelativeOffset) || (AutoAttachSocketName != AttachSocketName) || (AutoAttachParent != AttachParent))
+				{
+					DetachFromParent(/*bMaintainWorldPosition=*/ false, /*bCallModify=*/ false);
+				}
 			}
 			else
 			{
