Merging CL 2694625
Added stats for Slate vertex shaders in the material editor

diff --git a/Engine/Source/Runtime/Engine/Private/Materials/MaterialShared.cpp b/Engine/Source/Runtime/Engine/Private/Materials/MaterialShared.cpp
--- a/Engine/Source/Runtime/Engine/Private/Materials/MaterialShared.cpp
+++ b/Engine/Source/Runtime/Engine/Private/Materials/MaterialShared.cpp
@@ -978,8 +978,14 @@ void FMaterialResource::GetRepresentativeShaderTypesAndDescriptions(TArray<FStri
 
 	if( IsUIMaterial() )
 	{
-		new(ShaderTypeNames) FString(TEXT("TSlateMaterialShaderPSDefaultfalse"));
-		new(ShaderTypeDescriptions) FString(TEXT("Default UI Shader"));
+		new(ShaderTypeNames)FString(TEXT("TSlateMaterialShaderPSDefaultfalse"));
+		new(ShaderTypeDescriptions)FString(TEXT("Default UI Pixel Shader"));
+
+		new(ShaderTypeNames) FString(TEXT("TSlateMaterialShaderVSfalse"));
+		new(ShaderTypeDescriptions) FString(TEXT("Default UI Vertex Shader"));
+
+		new(ShaderTypeNames)FString(TEXT("TSlateMaterialShaderVStrue"));
+		new(ShaderTypeDescriptions)FString(TEXT("Instanced UI Vertex Shader"));
 	}
 	else if (GetFeatureLevel() >= ERHIFeatureLevel::SM4)
 	{
