Comment nodes added in the Material Editor now start with focus on the editable text.

Added methods to create a 'user invoked' node on UEdGraph and FGraphNodeCreator, which propagates this property through the FEdGraphEditAction to the SGraphNode. This allows us to give comment nodes focus when created.

diff --git a/Engine/Source/Editor/MaterialEditor/Private/MaterialEditor.cpp b/Engine/Source/Editor/MaterialEditor/Private/MaterialEditor.cpp
--- a/Engine/Source/Editor/MaterialEditor/Private/MaterialEditor.cpp
+++ b/Engine/Source/Editor/MaterialEditor/Private/MaterialEditor.cpp
@@ -3066,7 +3066,7 @@ UMaterialExpressionComment* FMaterialEditor::CreateNewMaterialExpressionComment(
 
 	if (NewComment)
 	{
-		Material->MaterialGraph->AddComment(NewComment);
+		Material->MaterialGraph->AddComment(NewComment, true);
 
 		// Select the new comment.
 		GraphEditor->ClearSelectionSet();
diff --git a/Engine/Source/Editor/UnrealEd/Classes/MaterialGraph/MaterialGraph.h b/Engine/Source/Editor/UnrealEd/Classes/MaterialGraph/MaterialGraph.h
--- a/Engine/Source/Editor/UnrealEd/Classes/MaterialGraph/MaterialGraph.h
+++ b/Engine/Source/Editor/UnrealEd/Classes/MaterialGraph/MaterialGraph.h
@@ -138,11 +138,12 @@ public:
 	/**
 	 * Add a Comment to the Graph
 	 *
-	 * @param	Comment	Comment to add
+	 * @param	Comment			Comment to add
+	 * @param	bIsUserInvoked	Whether the comment has just been created by the user via the UI
 	 *
 	 * @return	UMaterialGraphNode_Comment*	Newly created Graph node to represent comment
 	 */
-	class UMaterialGraphNode_Comment*	AddComment(UMaterialExpressionComment* Comment);
+	class UMaterialGraphNode_Comment*	AddComment(UMaterialExpressionComment* Comment, bool bIsUserInvoked = false);
 
 	/** Link all of the Graph nodes using the Material's connections */
 	void LinkGraphNodesFromMaterial();
diff --git a/Engine/Source/Editor/UnrealEd/Private/MaterialGraph.cpp b/Engine/Source/Editor/UnrealEd/Private/MaterialGraph.cpp
--- a/Engine/Source/Editor/UnrealEd/Private/MaterialGraph.cpp
+++ b/Engine/Source/Editor/UnrealEd/Private/MaterialGraph.cpp
@@ -97,14 +97,21 @@ UMaterialGraphNode* UMaterialGraph::AddExpression(UMaterialExpression* Expressio
 	return NewNode;
 }
 
-UMaterialGraphNode_Comment* UMaterialGraph::AddComment(UMaterialExpressionComment* Comment)
+UMaterialGraphNode_Comment* UMaterialGraph::AddComment(UMaterialExpressionComment* Comment, bool bIsUserInvoked)
 {
 	UMaterialGraphNode_Comment* NewComment = NULL;
 	if (Comment)
 	{
 		Modify();
 		FGraphNodeCreator<UMaterialGraphNode_Comment> NodeCreator(*this);
-		NewComment = NodeCreator.CreateNode(false);
+		if (bIsUserInvoked)
+		{
+			NewComment = NodeCreator.CreateUserInvokedNode(true);
+		}
+		else
+		{
+			NewComment = NodeCreator.CreateNode(false);
+		}
 		NewComment->MaterialExpressionComment = Comment;
 		NewComment->MaterialDirtyDelegate = MaterialDirtyDelegate;
 		Comment->GraphNode = NewComment;
diff --git a/Engine/Source/Runtime/Engine/Classes/EdGraph/EdGraph.h b/Engine/Source/Runtime/Engine/Classes/EdGraph/EdGraph.h
--- a/Engine/Source/Runtime/Engine/Classes/EdGraph/EdGraph.h
+++ b/Engine/Source/Runtime/Engine/Classes/EdGraph/EdGraph.h
@@ -201,11 +201,22 @@ protected:
 	 * Creates an empty node in this graph. Use FGraphNodeCreator above
 	 *
 	 * @param NewNodeClass		The node class to create
+	 * @param bFromUI			Whether or not the node was created by the user via the UI
 	 * @param bSelectNewNode	Whether or not to select the new node being created
 	 *
 	 * @return A new graph node of the given type
 	 */
-	UEdGraphNode* CreateNode ( TSubclassOf<UEdGraphNode> NewNodeClass, bool bSelectNewNode = true );
+	UEdGraphNode* CreateNode( TSubclassOf<UEdGraphNode> NewNodeClass, bool bFromUI, bool bSelectNewNode );
+
+	UEdGraphNode* CreateNode(TSubclassOf<UEdGraphNode> NewNodeClass, bool bSelectNewNode = true)
+	{
+		return CreateNode( NewNodeClass, false, bSelectNewNode );
+	}
+
+	UEdGraphNode* CreateUserInvokedNode(TSubclassOf<UEdGraphNode> NewNodeClass, bool bSelectNewNode = true)
+	{
+		return CreateNode(NewNodeClass, true, bSelectNewNode);
+	}
 
 private:
 	/** A delegate that broadcasts a notification whenever the graph has changed. */
@@ -245,6 +256,13 @@ public:
 		return Node;
 	} 
 
+	/** Create an empty placeable graph node */
+	NodeType* CreateUserInvokedNode(bool bSelectNewNode = true)
+	{
+		Node = (NodeType*)Graph.CreateUserInvokedNode(NodeType::StaticClass(), bSelectNewNode);
+		return Node;
+	}
+
 	/** Call to finalize the node's construction */
 	void Finalize()
 	{
diff --git a/Engine/Source/Runtime/Engine/Private/EdGraph/EdGraph.cpp b/Engine/Source/Runtime/Engine/Private/EdGraph/EdGraph.cpp
--- a/Engine/Source/Runtime/Engine/Private/EdGraph/EdGraph.cpp
+++ b/Engine/Source/Runtime/Engine/Private/EdGraph/EdGraph.cpp
@@ -157,7 +157,7 @@ void UEdGraph::RemoveOnGraphChangedHandler( FDelegateHandle Handle )
 	OnGraphChanged.Remove( Handle );
 }
 
-UEdGraphNode* UEdGraph::CreateNode( TSubclassOf<UEdGraphNode> NewNodeClass, bool bSelectNewNode/* = true*/ )
+UEdGraphNode* UEdGraph::CreateNode( TSubclassOf<UEdGraphNode> NewNodeClass, bool bFromUI, bool bSelectNewNode )
 {
 	UEdGraphNode* NewNode = NewObject<UEdGraphNode>(this, NewNodeClass, NAME_None, RF_Transactional);
 
@@ -166,7 +166,7 @@ UEdGraphNode* UEdGraph::CreateNode( TSubclassOf<UEdGraphNode> NewNodeClass, bool
 		NewNode->SetFlags(RF_Transient);
 	}
 
-	AddNode(NewNode, false, bSelectNewNode );
+	AddNode(NewNode, bFromUI, bSelectNewNode );
 	return NewNode;
 }
 
