UE4 - Tweaked the updating of component transforms for performance.

diff --git a/Engine/Source/Runtime/Core/Public/Math/Transform.h b/Engine/Source/Runtime/Core/Public/Math/Transform.h
--- a/Engine/Source/Runtime/Core/Public/Math/Transform.h
+++ b/Engine/Source/Runtime/Core/Public/Math/Transform.h
@@ -102,6 +102,13 @@ public:
 		DiagnosticCheckNaN_All();
 	}
 
+	/**
+	 * Constructor with leaving uninitialized memory
+	 */
+	FORCEINLINE explicit FTransform(ENoInit) 
+	{
+	}
+
 	/**
 	 * Constructor with an initial rotation
 	 *
diff --git a/Engine/Source/Runtime/Core/Public/Math/TransformVectorized.h b/Engine/Source/Runtime/Core/Public/Math/TransformVectorized.h
--- a/Engine/Source/Runtime/Core/Public/Math/TransformVectorized.h
+++ b/Engine/Source/Runtime/Core/Public/Math/TransformVectorized.h
@@ -117,7 +117,7 @@ public:
 	FORCEINLINE explicit FTransform(const FQuat& InRotation) 
 	{
 		// Rotation = InRotation
-		Rotation =  MakeVectorRegister( InRotation.X, InRotation.Y, InRotation.Z, InRotation.W );
+		Rotation =  VectorLoadAligned( &InRotation.X );
 		// Translation = {0,0,0,0)
 		Translation = VectorZero();
 		// Scale3D = {1,1,1,0);
@@ -135,7 +135,7 @@ public:
 	{
 		FQuat InQuatRotation = InRotation.Quaternion();
 		// Rotation = InRotation
-		Rotation =  MakeVectorRegister( InQuatRotation.X, InQuatRotation.Y, InQuatRotation.Z, InQuatRotation.W );
+		Rotation =  VectorLoadAligned( &InQuatRotation.X );
 		// Translation = {0,0,0,0)
 		Translation = VectorZero();
 		// Scale3D = {1,1,1,0);
@@ -154,7 +154,7 @@ public:
 	FORCEINLINE FTransform(const FQuat& InRotation, const FVector& InTranslation, const FVector& InScale3D = FVector(1.f)) 
 	{
 		// Rotation = InRotation
-		Rotation =  MakeVectorRegister( InRotation.X, InRotation.Y, InRotation.Z, InRotation.W );
+		Rotation =  VectorLoadAligned( &InRotation.X );
 		// Translation = InTranslation
 		Translation = MakeVectorRegister(InTranslation.X, InTranslation.Y, InTranslation.Z, 0.0f );
 		// Scale3D = InScale3D
@@ -189,7 +189,7 @@ public:
 	{
 		FQuat InQuatRotation = InRotation.Quaternion();
 		// Rotation = InRotation
-		Rotation =  MakeVectorRegister( InQuatRotation.X, InQuatRotation.Y, InQuatRotation.Z, InQuatRotation.W );
+		Rotation =  VectorLoadAligned( &InQuatRotation.X );
 		// Translation = InTranslation
 		Translation = MakeVectorRegister(InTranslation.X, InTranslation.Y, InTranslation.Z, 0.0f );
 		// Scale3D = InScale3D
@@ -198,6 +198,13 @@ public:
 		DiagnosticCheckNaN_All();
 	}
 
+	/**
+	 * Constructor with leaving uninitialized memory
+	 */
+	FORCEINLINE explicit FTransform(ENoInit) 
+	{
+	}
+
 	/**
 	 * Copy-constructor
 	 *
@@ -541,7 +548,6 @@ public:
 	FORCEINLINE FVector		GetSafeScaleReciprocal(const FVector& InScale, float Tolerance=0.0f) const;
 
 
-	// temp function for easy conversion
 	FORCEINLINE FVector GetLocation() const
 	{
 		FVector OutTranslation;
diff --git a/Engine/Source/Runtime/Engine/Classes/Components/SceneComponent.h b/Engine/Source/Runtime/Engine/Classes/Components/SceneComponent.h
--- a/Engine/Source/Runtime/Engine/Classes/Components/SceneComponent.h
+++ b/Engine/Source/Runtime/Engine/Classes/Components/SceneComponent.h
@@ -685,7 +685,10 @@ public:
 	{
 		return true;
 	}
-	virtual void UpdateComponentToWorld(bool bSkipPhysicsMove = false, ETeleportType Teleport = ETeleportType::None) override final;
+	virtual void UpdateComponentToWorld(bool bSkipPhysicsMove = false, ETeleportType Teleport = ETeleportType::None) override final
+	{
+		UpdateComponentToWorldWithParent(AttachParent, AttachSocketName, bSkipPhysicsMove, RelativeRotationCache.RotatorToQuat(RelativeRotation), Teleport);
+	}
 	virtual void DestroyComponent(bool bPromoteChildren = false) override;
 	virtual void OnComponentDestroyed() override;
 	virtual void ApplyWorldOffset(const FVector& InOffset, bool bWorldShift) override;
@@ -918,8 +921,20 @@ protected:
 
 	/** Calculate the new ComponentToWorld transform for this component.
 		Parent is optional and can be used for computing ComponentToWorld based on arbitrary USceneComponent.
-		If Parent is not passed in or is NULL then we use the component's existing AttachParent and AttachSocket */
-	virtual FTransform CalcNewComponentToWorld(const FTransform& NewRelativeTransform, const USceneComponent* Parent = NULL, FName SocketName = NAME_None) const;
+		If Parent is not passed in we use the component's AttachParent*/
+	FORCEINLINE FTransform CalcNewComponentToWorld(const FTransform& NewRelativeTransform, const USceneComponent* Parent = NULL, FName SocketName = NAME_None) const
+	{
+		SocketName = Parent ? SocketName : AttachSocketName;
+		Parent = Parent ? Parent : AttachParent;
+		bool bGeneral = !Parent || bAbsoluteLocation || bAbsoluteRotation || bAbsoluteScale;
+		if (!bGeneral)
+		{
+			return NewRelativeTransform * Parent->GetSocketTransform(AttachSocketName);
+		}
+		return CalcNewComponentToWorld_GeneralCase(NewRelativeTransform, Parent, SocketName);
+	}
+
+	FTransform CalcNewComponentToWorld_GeneralCase(const FTransform& NewRelativeTransform, const USceneComponent* Parent, FName SocketName) const;
 
 	
 public:
diff --git a/Engine/Source/Runtime/Engine/Classes/GameFramework/Actor.h b/Engine/Source/Runtime/Engine/Classes/GameFramework/Actor.h
--- a/Engine/Source/Runtime/Engine/Classes/GameFramework/Actor.h
+++ b/Engine/Source/Runtime/Engine/Classes/GameFramework/Actor.h
@@ -653,10 +653,21 @@ public:
 	 * @return The transform that transforms from actor space to world space.
 	 */
 	UFUNCTION(BlueprintCallable, meta=(DisplayName = "GetActorTransform"), Category="Utilities|Transformation")
-	FTransform GetTransform() const;
+	FTransform GetTransform() const
+	{
+		return ActorToWorld();
+	}
 
 	/** Get the local-to-world transform of the RootComponent. Identical to GetTransform(). */
-	FTransform ActorToWorld() const;
+	FORCEINLINE FTransform ActorToWorld() const
+	{
+		if( RootComponent != NULL )
+		{
+			return RootComponent->ComponentToWorld;
+		}
+		return FTransform::Identity;
+	}
+
 
 	/** Returns the location of the RootComponent of this Actor */
 	UFUNCTION(BlueprintCallable, meta=(DisplayName = "GetActorLocation", Keywords="position"), Category="Utilities|Transformation")
diff --git a/Engine/Source/Runtime/Engine/Private/Actor.cpp b/Engine/Source/Runtime/Engine/Private/Actor.cpp
--- a/Engine/Source/Runtime/Engine/Private/Actor.cpp
+++ b/Engine/Source/Runtime/Engine/Private/Actor.cpp
@@ -1760,29 +1760,10 @@ FVector AActor::GetPlacementExtent() const
 	return Extent;
 }
 
-FTransform AActor::ActorToWorld() const
-{
-	FTransform Result = FTransform::Identity;
-	if( RootComponent != NULL )
-	{
-		Result = RootComponent->ComponentToWorld;
-	}
-	else
-	{
-		UE_LOG(LogActor, Log, TEXT("AActor::ActorToWorld (%s) No RootComponent!"), *GetPathName());
-	}
-
-	return Result;
-}
-
 class UClass* AActor::GetActorClass() const
 {
 	return GetClass();
 }
-FTransform AActor::GetTransform() const
-{
-	return ActorToWorld();
-}
 
 void AActor::Destroyed()
 {
diff --git a/Engine/Source/Runtime/Engine/Private/Components/SceneComponent.cpp b/Engine/Source/Runtime/Engine/Private/Components/SceneComponent.cpp
--- a/Engine/Source/Runtime/Engine/Private/Components/SceneComponent.cpp
+++ b/Engine/Source/Runtime/Engine/Private/Components/SceneComponent.cpp
@@ -298,15 +298,12 @@ void USceneComponent::PostEditChangeChainProperty(FPropertyChangedChainEvent& Pr
 
 #endif
 
-FTransform USceneComponent::CalcNewComponentToWorld(const FTransform& NewRelativeTransform, const USceneComponent* Parent, FName SocketName) const
+FTransform USceneComponent::CalcNewComponentToWorld_GeneralCase(const FTransform& NewRelativeTransform, const USceneComponent* Parent, FName SocketName) const
 {
-	SocketName = Parent ? SocketName : AttachSocketName;
-	Parent = Parent ? Parent : AttachParent;
 	if (Parent != NULL)
 	{
 		const FTransform ParentToWorld = Parent->GetSocketTransform(SocketName);
 		FTransform NewCompToWorld = NewRelativeTransform * ParentToWorld;
-
 		if(bAbsoluteLocation)
 		{
 			NewCompToWorld.CopyTranslation(NewRelativeTransform);
@@ -336,9 +333,11 @@ void USceneComponent::OnUpdateTransform(bool bSkipPhysicsMove, ETeleportType Tel
 
 void USceneComponent::UpdateComponentToWorldWithParent(USceneComponent* Parent,FName SocketName, bool bSkipPhysicsMove, const FQuat& RelativeRotationQuat, ETeleportType Teleport)
 {
+	//QUICK_SCOPE_CYCLE_COUNTER(STAT_USceneComponent_UpdateComponentToWorldWithParent);
 	// If our parent hasn't been updated before, we'll need walk up our parent attach hierarchy
 	if (Parent && !Parent->bWorldToComponentUpdated)
 	{
+		//QUICK_SCOPE_CYCLE_COUNTER(STAT_USceneComponent_UpdateComponentToWorldWithParent_Parent);
 		Parent->UpdateComponentToWorld();
 
 		// Updating the parent may (depending on if we were already attached to parent) result in our being updated, so just return
@@ -350,23 +349,33 @@ void USceneComponent::UpdateComponentToWorldWithParent(USceneComponent* Parent,F
 
 	bWorldToComponentUpdated = true;
 
-	// Calculate the new ComponentToWorld transform
-	const FTransform RelativeTransform(RelativeRotationQuat, RelativeLocation, RelativeScale3D);
-	FTransform NewTransform = CalcNewComponentToWorld(RelativeTransform, Parent, SocketName);
+	FTransform NewTransform(NoInit);
+
+	{
+		//QUICK_SCOPE_CYCLE_COUNTER(STAT_USceneComponent_UpdateComponentToWorldWithParent_XForm);
+		// Calculate the new ComponentToWorld transform
+		const FTransform RelativeTransform(RelativeRotationQuat, RelativeLocation, RelativeScale3D);
+		NewTransform = CalcNewComponentToWorld(RelativeTransform, Parent, SocketName);
+	}
 
-#if !(UE_BUILD_SHIPPING || UE_BUILD_TEST)
 	ensure(NewTransform.IsValid());
-#endif
 
 	// If transform has changed..
-	if (!ComponentToWorld.Equals(NewTransform, SMALL_NUMBER))
+	bool bHasChanged;
 	{
+		//QUICK_SCOPE_CYCLE_COUNTER(STAT_USceneComponent_UpdateComponentToWorldWithParent_HasChanged);
+		bHasChanged = ComponentToWorld.Equals(NewTransform, SMALL_NUMBER);
+	}
+	if (!bHasChanged)
+	{
+		//QUICK_SCOPE_CYCLE_COUNTER(STAT_USceneComponent_UpdateComponentToWorldWithParent_Changed);
 		// Update transform
 		ComponentToWorld = NewTransform;
 		PropagateTransformUpdate(true, bSkipPhysicsMove, Teleport);
 	}
 	else
 	{
+		//QUICK_SCOPE_CYCLE_COUNTER(STAT_USceneComponent_UpdateComponentToWorldWithParent_NotChanged);
 		PropagateTransformUpdate(false);
 	}
 }
@@ -411,38 +420,44 @@ void USceneComponent::OnRegister()
 #endif
 }
 
-void USceneComponent::UpdateComponentToWorld(bool bSkipPhysicsMove, ETeleportType Teleport)
-{
-	UpdateComponentToWorldWithParent(AttachParent,AttachSocketName, bSkipPhysicsMove, RelativeRotationCache.RotatorToQuat(RelativeRotation), Teleport);
-}
-
-
 void USceneComponent::PropagateTransformUpdate(bool bTransformChanged, bool bSkipPhysicsMove, ETeleportType Teleport)
 {
+	//QUICK_SCOPE_CYCLE_COUNTER(STAT_SceneComponent_PropagateTransformUpdate);
 	if (IsDeferringMovementUpdates())
 	{
 		// We are deferring these updates until later.
 		return;
 	}
-
+	FPlatformMisc::Prefetch(AttachChildren.GetData());
 	if (bTransformChanged)
 	{
+		//QUICK_SCOPE_CYCLE_COUNTER(STAT_SceneComponent_PropagateTransformUpdate_TransformChanged);
+		{
 		// Then update bounds
+			//QUICK_SCOPE_CYCLE_COUNTER(STAT_SceneComponent_PropagateTransformUpdate_UpdateBounds);
 		UpdateBounds();
+		}
 
+		{
+			//QUICK_SCOPE_CYCLE_COUNTER(STAT_SceneComponent_PropagateTransformUpdate_OnUpdateTransform);
 		// Always send new transform to physics
 		OnUpdateTransform(bSkipPhysicsMove, Teleport);
+		}
 
 		// Flag render transform as dirty
 		MarkRenderTransformDirty();
 		
+		{
+			//QUICK_SCOPE_CYCLE_COUNTER(STAT_SceneComponent_PropagateTransformUpdate_UpdateChildTransforms);
 		// Now go and update children
 		//Do not pass skip physics to children. This is only used when physics updates us, but in that case we really do need to update the attached children since they are kinematic
 		UpdateChildTransforms(false, Teleport);
+		}
 
 		// Refresh navigation
 		UpdateNavigationData();
 
+		//QUICK_SCOPE_CYCLE_COUNTER(STAT_SceneComponent_PropagateTransformUpdate_AttachmentReplication);
 
 		AActor* Owner = GetOwner();
 		if (Owner && Owner->HasAuthority() && Owner->GetRootComponent() == this)
@@ -464,16 +479,26 @@ void USceneComponent::PropagateTransformUpdate(bool bTransformChanged, bool bSki
 	}
 	else
 	{
+		//QUICK_SCOPE_CYCLE_COUNTER(STAT_SceneComponent_PropagateTransformUpdate_NOT_TransformChanged);
+		{
+			//QUICK_SCOPE_CYCLE_COUNTER(STAT_SceneComponent_PropagateTransformUpdate_UpdateBounds);
 		// We update bounds even if transform doesn't change, as shape/mesh etc might have done
 		UpdateBounds();
+		}
 
+		{
+			//QUICK_SCOPE_CYCLE_COUNTER(STAT_SceneComponent_PropagateTransformUpdate_UpdateChildTransforms);
 		// Now go and update children
 		UpdateChildTransforms();
+		}
 
+		{
+			QUICK_SCOPE_CYCLE_COUNTER(STAT_SceneComponent_PropagateTransformUpdate_MarkRenderTransformDirty);
 		// Need to flag as dirty so new bounds are sent to render thread
 		MarkRenderTransformDirty();
 	}
 }
+}
 
 
 void USceneComponent::EndScopedMovementUpdate(class FScopedMovementUpdate& CompletedScope)
@@ -2269,7 +2294,7 @@ void USceneComponent::PostNetReceive()
 		}
 	}
 }
-
+	
 void USceneComponent::PostRepNotifies()
 {
 	if (bNetUpdateAttachment)
