Expose bSplineHasBeenEdited as an editable property (exposed as 'Override Construction Script').

diff --git a/Engine/Source/Runtime/Engine/Classes/Components/SplineComponent.h b/Engine/Source/Runtime/Engine/Classes/Components/SplineComponent.h
--- a/Engine/Source/Runtime/Engine/Classes/Components/SplineComponent.h
+++ b/Engine/Source/Runtime/Engine/Classes/Components/SplineComponent.h
@@ -73,7 +73,7 @@ class ENGINE_API USplineComponent : public UPrimitiveComponent
 	bool bStationaryEndpoints;
 
 	/** Whether the spline has been edited from its default by the spline component visualizer */
-	UPROPERTY(VisibleAnywhere, Category = Spline)
+	UPROPERTY(EditAnywhere, Category = Spline, meta=(DisplayName="Override Construction Script"))
 	bool bSplineHasBeenEdited;
 
 private:
