Alpha Blend change

- Fixed so that you don't complete right after you switch desired direction
 - this fixes pop happening with Stop of Montage

Merging

//depot/UE4-Orion/Engine/Source/Runtime/Engine/...

to //depot/UE4-Orion/Engine/Source/Runtime/Engine/...

diff --git a/Engine/Source/Runtime/Engine/Private/AlphaBlend.cpp b/Engine/Source/Runtime/Engine/Private/AlphaBlend.cpp
--- a/Engine/Source/Runtime/Engine/Private/AlphaBlend.cpp
+++ b/Engine/Source/Runtime/Engine/Private/AlphaBlend.cpp
@@ -79,6 +79,7 @@ void FAlphaBlend::Reset()
 		BlendTimeRemaining = BlendTime * FMath::Abs(1.f - AlphaLerp);
 	}
 
+	RecacheDesiredBlendedValue();
 	bNeedsToResetAlpha = false;
 	bNeedsToResetBlendTime = false;
 }
@@ -161,6 +162,7 @@ void FAlphaBlend::SetValueRange(float Begin, float Desired)
 	DesiredValue = Desired;
 
 	bNeedsToResetAlpha = true;
+	RecacheDesiredBlendedValue();
 }
 
 /** Sets the final desired value for the blended value */
@@ -177,6 +179,11 @@ void FAlphaBlend::SetAlpha(float InAlpha)
 	BlendedValue = BeginValue + (DesiredValue - BeginValue) * AlphaBlend;
 }
 
+void FAlphaBlend::RecacheDesiredBlendedValue()
+{
+	CachedDesiredBlendedValue = BeginValue + (DesiredValue - BeginValue) * AlphaToBlendOption(1.f, BlendOption, CustomCurve);
+}
+
 void FAlphaBlend::SetBlendTime(float InBlendTime)
 {
 	BlendTime = FMath::Max(InBlendTime, 0.f);
@@ -184,8 +191,20 @@ void FAlphaBlend::SetBlendTime(float InBlendTime)
 	bNeedsToResetBlendTime = true;
 }
 
+void FAlphaBlend::SetBlendOption(EAlphaBlendOption InBlendOption)
+{
+	BlendOption = InBlendOption;
+	RecacheDesiredBlendedValue();
+}
+
+void FAlphaBlend::SetCustomCurve(UCurveFloat* InCustomCurve)
+{
+	CustomCurve = InCustomCurve;
+	RecacheDesiredBlendedValue();
+}
+
 bool FAlphaBlend::IsComplete() const
 {
-	return AlphaLerp == 1.0f;
+	return CachedDesiredBlendedValue == BlendedValue;
 }
 
diff --git a/Engine/Source/Runtime/Engine/Public/AlphaBlend.h b/Engine/Source/Runtime/Engine/Public/AlphaBlend.h
--- a/Engine/Source/Runtime/Engine/Public/AlphaBlend.h
+++ b/Engine/Source/Runtime/Engine/Public/AlphaBlend.h
@@ -48,7 +48,7 @@ USTRUCT()
 struct ENGINE_API FAlphaBlend
 {
 	GENERATED_BODY()
-public:
+private:
 	/** 
 	 * Please note that changing below two variables would get applied in the NEXT UPDATE. 
 	 * This does not change the Alpha and BlendedValue right away and it is intentional
@@ -62,12 +62,16 @@ public:
 	UPROPERTY(EditAnywhere, Category = "Blend")
 	UCurveFloat* CustomCurve;
 
+public:
 	/* Constructor */
 	FAlphaBlend(float NewBlendTime = 0.2f);
 
 	/* Constructor */
 	FAlphaBlend(const FAlphaBlend& Other, float NewBlendTime);
 
+	/** Setters - need to refresh cached value */
+	void SetBlendOption(EAlphaBlendOption InBlendOption);
+	void SetCustomCurve(UCurveFloat* InCustomCurve);
 	/** Update transition blend time. This new value will be applied in the next Update. */
 	void SetBlendTime(float InBlendTime);
 
@@ -98,8 +102,10 @@ public:
 	/** Gets the current blended value */
 	float GetBlendedValue() const { return BlendedValue; }
 
-	/** Gets the current blend time */
+	/** Getters */
 	float GetBlendTime() const { return BlendTime; }
+	EAlphaBlendOption GetBlendOption() const { return BlendOption; }
+	UCurveFloat* GetCustomCurve() const { return CustomCurve; }
 
 	/** Get the current desired value */
 	float GetDesiredValue() const { return DesiredValue; }
@@ -176,4 +182,11 @@ private:
 
 	/** internal flat to reset blend time. The two distinction is required. Otherwise, you'll modify blend time when you just change range */
 	bool bNeedsToResetBlendTime;
+
+	/** Cached Desired Value with Alpha 1 so that we can check if reached or not */
+	float CachedDesiredBlendedValue;
+	
+	/** This function refreshes desired blended value, so that we can check if we reached there or not. We make sure this value gets updated whenever any data changes. 
+	 *	IsComplete is called by Update, and recalc in every frame might not be worth it, so caching that data here*/
+	void RecacheDesiredBlendedValue();
 };
\ No newline at end of file
