Fix compile error

Merging

//depot/UE4-Orion/Engine/Source/Runtime/...

to //depot/UE4-Orion/Engine/Source/Runtime/...

diff --git a/Engine/Source/Runtime/AnimGraphRuntime/Private/AnimNodes/AnimNode_BlendListBase.cpp b/Engine/Source/Runtime/AnimGraphRuntime/Private/AnimNodes/AnimNode_BlendListBase.cpp
--- a/Engine/Source/Runtime/AnimGraphRuntime/Private/AnimNodes/AnimNode_BlendListBase.cpp
+++ b/Engine/Source/Runtime/AnimGraphRuntime/Private/AnimNodes/AnimNode_BlendListBase.cpp
@@ -40,8 +40,8 @@ void FAnimNode_BlendListBase::Initialize(const FAnimationInitializeContext& Cont
 		FAlphaBlend& Blend = Blends[i];
 
 		Blend.SetBlendTime(0.0f);
-		Blend.BlendOption = BlendType;
-		Blend.CustomCurve = CustomBlendCurve;
+		Blend.SetBlendOption(BlendType);
+		Blend.SetCustomCurve(CustomBlendCurve);
 	}
 	Blends[0].SetAlpha(1.0f);
 
diff --git a/Engine/Source/Runtime/Engine/Private/Animation/AnimMontage.cpp b/Engine/Source/Runtime/Engine/Private/Animation/AnimMontage.cpp
--- a/Engine/Source/Runtime/Engine/Private/Animation/AnimMontage.cpp
+++ b/Engine/Source/Runtime/Engine/Private/Animation/AnimMontage.cpp
@@ -864,8 +864,8 @@ void FAnimMontageInstance::Play(float InPlayRate)
 
 void FAnimMontageInstance::InitializeBlend(const FAlphaBlend& InAlphaBlend)
 {
-	Blend.BlendOption = InAlphaBlend.BlendOption;
-	Blend.CustomCurve = InAlphaBlend.CustomCurve;
+	Blend.SetBlendOption(InAlphaBlend.GetBlendOption());
+	Blend.SetCustomCurve(InAlphaBlend.GetCustomCurve());
 	Blend.SetBlendTime(InAlphaBlend.GetBlendTime());
 }
 
@@ -884,13 +884,13 @@ void FAnimMontageInstance::Stop(const FAlphaBlend& InBlendOut, bool bInterrupt)
 	// if desired weight is > 0.f, turn that off
 	if (Blend.GetDesiredValue() > 0.f)
 	{
+		// do not use default Montage->BlendOut 
+		// depending on situation, the BlendOut time can change 
+		InitializeBlend(InBlendOut);
 		Blend.SetDesiredValue(0.f);
-		if (Montage)
-		{
-			// do not use default Montage->BlendOut 
-			// depending on situation, the BlendOut time can change 
-			InitializeBlend(InBlendOut);
 
+		if(Montage)
+		{
 			if (UAnimInstance* Inst = AnimInstance.Get())
 			{
 				// Let AnimInstance know we are being stopped.
@@ -1013,8 +1013,8 @@ void FAnimMontageInstance::Terminate()
 	MontageSync_StopLeading();
 
 	// clear Blend curve
-	Blend.CustomCurve = NULL;
-	Blend.BlendOption = EAlphaBlendOption::Linear;
+	Blend.SetCustomCurve(NULL);
+	Blend.SetBlendOption(EAlphaBlendOption::Linear);
 
 	UE_LOG(LogAnimation, Verbose, TEXT("Terminating: AnimMontage: %s"), *GetNameSafe(OldMontage));
 }
diff --git a/Engine/Source/Runtime/Engine/Private/Animation/AnimNode_StateMachine.cpp b/Engine/Source/Runtime/Engine/Private/Animation/AnimNode_StateMachine.cpp
--- a/Engine/Source/Runtime/Engine/Private/Animation/AnimNode_StateMachine.cpp
+++ b/Engine/Source/Runtime/Engine/Private/Animation/AnimNode_StateMachine.cpp
@@ -45,8 +45,8 @@ FAnimationActiveTransitionEntry::FAnimationActiveTransitionEntry(int32 NextState
 	CrossfadeDuration = ReferenceTransitionInfo.CrossfadeDuration * CalculateInverseAlpha(BlendOption, Scaler);
 
 	Blend.SetBlendTime(CrossfadeDuration);
-	Blend.BlendOption = BlendOption;
-	Blend.CustomCurve = ReferenceTransitionInfo.CustomCurve;
+	Blend.SetBlendOption(BlendOption);
+	Blend.SetCustomCurve(ReferenceTransitionInfo.CustomCurve);
 	Blend.SetValueRange(0.0f, 1.0f);
 }
 
@@ -113,7 +113,7 @@ void FAnimationActiveTransitionEntry::Update(const FAnimationUpdateContext& Cont
 			QueryAlpha = ElapsedTime / CrossfadeDuration;
 		}
 
-		Alpha = FAlphaBlend::AlphaToBlendOption(QueryAlpha, Blend.BlendOption, Blend.CustomCurve);
+		Alpha = FAlphaBlend::AlphaToBlendOption(QueryAlpha, Blend.GetBlendOption(), Blend.GetCustomCurve());
 
 		if(Blend.IsComplete())
 		{
