Code cleanup and addition of a missing NULL check in GameplayStatics.

diff --git a/Engine/Source/Runtime/Engine/Private/GameplayStatics.cpp b/Engine/Source/Runtime/Engine/Private/GameplayStatics.cpp
--- a/Engine/Source/Runtime/Engine/Private/GameplayStatics.cpp
+++ b/Engine/Source/Runtime/Engine/Private/GameplayStatics.cpp
@@ -37,54 +37,43 @@ UGameplayStatics::UGameplayStatics(const FObjectInitializer& ObjectInitializer)
 UGameInstance* UGameplayStatics::GetGameInstance(UObject* WorldContextObject)
 {
 	UWorld* World = GEngine->GetWorldFromContextObject(WorldContextObject);
-	return World->GetGameInstance();
+	return World ? World->GetGameInstance() : nullptr;
 }
 
 APlayerController* UGameplayStatics::GetPlayerController(UObject* WorldContextObject, int32 PlayerIndex ) 
 {
-	if (WorldContextObject)
+	if (UWorld* World = GEngine->GetWorldFromContextObject(WorldContextObject))
 	{
-		UWorld* World = GEngine->GetWorldFromContextObject(WorldContextObject);
-		if(World != nullptr)
+		uint32 Index = 0;
+		for (FConstPlayerControllerIterator Iterator = World->GetPlayerControllerIterator(); Iterator; ++Iterator)
 		{
-			uint32 Index = 0;
-			for (FConstPlayerControllerIterator Iterator = World->GetPlayerControllerIterator(); Iterator; ++Iterator)
+			APlayerController* PlayerController = *Iterator;
+			if (Index == PlayerIndex)
 			{
-				APlayerController* PlayerController = *Iterator;
-				if (Index == PlayerIndex)
-				{
-					return PlayerController;
-				}
-				Index++;
+				return PlayerController;
 			}
+			Index++;
 		}
 	}
-
 	return nullptr;
 }
 
 ACharacter* UGameplayStatics::GetPlayerCharacter(UObject* WorldContextObject, int32 PlayerIndex)
 {
-	ACharacter* Character = nullptr;
-	APlayerController* PC = GetPlayerController(WorldContextObject, PlayerIndex );
-	if (PC != nullptr)
-	{
-		Character = Cast<ACharacter>(PC->GetPawn());
-	}
-
-	return Character;
+	APlayerController* PC = GetPlayerController(WorldContextObject, PlayerIndex);
+	return PC ? Cast<ACharacter>(PC->GetPawn()) : nullptr;
 }
 
 APawn* UGameplayStatics::GetPlayerPawn(UObject* WorldContextObject, int32 PlayerIndex)
 {
-	APlayerController* PlayerController = GetPlayerController(WorldContextObject, PlayerIndex);
-	return PlayerController ? PlayerController->GetPawnOrSpectator() : NULL;
+	APlayerController* PC = GetPlayerController(WorldContextObject, PlayerIndex);
+	return PC ? PC->GetPawnOrSpectator() : nullptr;
 }
 
 APlayerCameraManager* UGameplayStatics::GetPlayerCameraManager(UObject* WorldContextObject, int32 PlayerIndex)
 {
-	APlayerController* const PC = GetPlayerController( WorldContextObject, PlayerIndex );
-	return PC ? PC->PlayerCameraManager : NULL;
+	APlayerController* const PC = GetPlayerController(WorldContextObject, PlayerIndex);
+	return PC ? PC->PlayerCameraManager : nullptr;
 }
 
 APlayerController* UGameplayStatics::CreatePlayer(UObject* WorldContextObject, int32 ControllerId, bool bSpawnPawn)
@@ -99,7 +88,7 @@ APlayerController* UGameplayStatics::CreatePlayer(UObject* WorldContextObject, i
 		UE_LOG(LogPlayerManagement, Error, TEXT("Failed to Create Player: %s"), *Error);
 	}
 
-	return (LocalPlayer ? LocalPlayer->PlayerController : NULL);
+	return (LocalPlayer ? LocalPlayer->PlayerController : nullptr);
 }
 
 void UGameplayStatics::RemovePlayer(APlayerController* PlayerController, bool bDestroyPawn)
@@ -129,12 +118,12 @@ AGameMode* UGameplayStatics::GetGameMode(UObject* WorldContextObject)
 AGameState* UGameplayStatics::GetGameState(UObject* WorldContextObject)
 {
 	UWorld* const World = GEngine->GetWorldFromContextObject(WorldContextObject);
-	return World ? World->GameState : NULL;
+	return World ? World->GameState : nullptr;
 }
 
 class UClass* UGameplayStatics::GetObjectClass(const UObject* Object)
 {
-	return Object ? Object->GetClass() : NULL;
+	return Object ? Object->GetClass() : nullptr;
 }
 
 float UGameplayStatics::GetGlobalTimeDilation(UObject* WorldContextObject)
@@ -255,7 +244,7 @@ bool UGameplayStatics::ApplyRadialDamageWithFalloff(UObject* WorldContextObject,
 	}
 
 	// make sure we have a good damage type
-	TSubclassOf<UDamageType> const ValidDamageTypeClass = (DamageTypeClass == NULL) ? TSubclassOf<UDamageType>(UDamageType::StaticClass()) : DamageTypeClass;
+	TSubclassOf<UDamageType> const ValidDamageTypeClass = DamageTypeClass ? DamageTypeClass : TSubclassOf<UDamageType>(UDamageType::StaticClass());
 
 	bool bAppliedDamage = false;
 
@@ -284,7 +273,7 @@ void UGameplayStatics::ApplyPointDamage(AActor* DamagedActor, float BaseDamage,
 	if (DamagedActor && BaseDamage != 0.f)
 	{
 		// make sure we have a good damage type
-		TSubclassOf<UDamageType> const ValidDamageTypeClass = DamageTypeClass == NULL ? TSubclassOf<UDamageType>(UDamageType::StaticClass()) : DamageTypeClass;
+		TSubclassOf<UDamageType> const ValidDamageTypeClass = DamageTypeClass ? DamageTypeClass : TSubclassOf<UDamageType>(UDamageType::StaticClass());
 		FPointDamageEvent PointDamageEvent(BaseDamage, HitInfo, HitFromDirection, ValidDamageTypeClass);
 
 		DamagedActor->TakeDamage(BaseDamage, PointDamageEvent, EventInstigator, DamageCauser);
@@ -296,7 +285,7 @@ void UGameplayStatics::ApplyDamage(AActor* DamagedActor, float BaseDamage, ACont
 	if ( DamagedActor && (BaseDamage != 0.f) )
 	{
 		// make sure we have a good damage type
-		TSubclassOf<UDamageType> const ValidDamageTypeClass = DamageTypeClass == NULL ? TSubclassOf<UDamageType>(UDamageType::StaticClass()) : DamageTypeClass;
+		TSubclassOf<UDamageType> const ValidDamageTypeClass = DamageTypeClass ? DamageTypeClass : TSubclassOf<UDamageType>(UDamageType::StaticClass());
 		FDamageEvent DamageEvent(ValidDamageTypeClass);
 
 		DamagedActor->TakeDamage(BaseDamage, DamageEvent, EventInstigator, DamageCauser);
@@ -381,20 +370,19 @@ UObject* UGameplayStatics::SpawnObject(TSubclassOf<UObject> ObjectClass, UObject
 
 AActor* UGameplayStatics::BeginSpawningActorFromBlueprint(UObject* WorldContextObject, UBlueprint const* Blueprint, const FTransform& SpawnTransform, bool bNoCollisionFail)
 {
-	AActor* NewActor = NULL;
-	if(Blueprint != NULL && Blueprint->GeneratedClass != NULL)
+	if (Blueprint && Blueprint->GeneratedClass)
 	{
 		if( Blueprint->GeneratedClass->IsChildOf(AActor::StaticClass()) )
 		{
 			ESpawnActorCollisionHandlingMethod const CollisionHandlingOverride = bNoCollisionFail ? ESpawnActorCollisionHandlingMethod::AdjustIfPossibleButDontSpawnIfColliding : ESpawnActorCollisionHandlingMethod::AlwaysSpawn;
-			NewActor = BeginDeferredActorSpawnFromClass(WorldContextObject, *Blueprint->GeneratedClass, SpawnTransform, CollisionHandlingOverride);
+			return BeginDeferredActorSpawnFromClass(WorldContextObject, *Blueprint->GeneratedClass, SpawnTransform, CollisionHandlingOverride);
 		}
 		else
 		{
 			UE_LOG(LogScript, Warning, TEXT("UGameplayStatics::BeginSpawningActorFromBlueprint: %s is not an actor class"), *Blueprint->GeneratedClass->GetName() );
 		}
 	}
-	return NewActor;
+	return nullptr;
 }
 
 // deprecated
@@ -406,28 +394,23 @@ AActor* UGameplayStatics::BeginSpawningActorFromClass(UObject* WorldContextObjec
 
 AActor* UGameplayStatics::BeginDeferredActorSpawnFromClass(UObject* WorldContextObject, TSubclassOf<AActor> ActorClass, const FTransform& SpawnTransform, ESpawnActorCollisionHandlingMethod CollisionHandlingMethod, AActor* Owner)
 {
-	AActor* NewActor = NULL;
-
-	UClass* Class = *ActorClass;
-	if (Class != NULL)
+	if (UClass* Class = *ActorClass)
 	{
 		// If the WorldContextObject is a Pawn we will use that as the instigator.
 		// Otherwise if the WorldContextObject is an Actor we will share its instigator.
 		// If the value is set via the exposed parameter on SpawnNode it will be overwritten anyways, so this is safe to specify here
 		APawn* AutoInstigator = Cast<APawn>(WorldContextObject);
-		if (AutoInstigator == nullptr)
+		if (AutoInstigator)
 		{
-			AActor* ContextActor = Cast<AActor>(WorldContextObject);
-			if (ContextActor)
+			if (AActor* ContextActor = Cast<AActor>(WorldContextObject))
 			{
 				AutoInstigator = ContextActor->Instigator;
 			}
 		}
 
-		UWorld* World = GEngine->GetWorldFromContextObject(WorldContextObject);
-		if (World)
+		if (UWorld* World = GEngine->GetWorldFromContextObject(WorldContextObject))
 		{
-			NewActor = World->SpawnActorDeferred<AActor>(Class, SpawnTransform, Owner, AutoInstigator, CollisionHandlingMethod);
+			return World->SpawnActorDeferred<AActor>(Class, SpawnTransform, Owner, AutoInstigator, CollisionHandlingMethod);
 		}
 		else
 		{
@@ -438,8 +421,7 @@ AActor* UGameplayStatics::BeginDeferredActorSpawnFromClass(UObject* WorldContext
 	{
 		UE_LOG(LogScript, Warning, TEXT("UGameplayStatics::BeginSpawningActorFromClass: can not spawn an actor from a NULL class"));
 	}
-
-	return NewActor;
+	return nullptr;
 }
 
 AActor* UGameplayStatics::FinishSpawningActor(AActor* Actor, const FTransform& SpawnTransform)
@@ -454,11 +436,10 @@ AActor* UGameplayStatics::FinishSpawningActor(AActor* Actor, const FTransform& S
 
 void UGameplayStatics::LoadStreamLevel(UObject* WorldContextObject, FName LevelName,bool bMakeVisibleAfterLoad,bool bShouldBlockOnLoad,FLatentActionInfo LatentInfo)
 {
-	UWorld* World = GEngine->GetWorldFromContextObject(WorldContextObject);
-	if(World != nullptr)
+	if (UWorld* World = GEngine->GetWorldFromContextObject(WorldContextObject))
 	{
 		FLatentActionManager& LatentManager = World->GetLatentActionManager();
-		if (LatentManager.FindExistingAction<FStreamLevelAction>(LatentInfo.CallbackTarget, LatentInfo.UUID) == NULL)
+		if (LatentManager.FindExistingAction<FStreamLevelAction>(LatentInfo.CallbackTarget, LatentInfo.UUID) == nullptr)
 		{
 			FStreamLevelAction* NewAction = new FStreamLevelAction(true, LevelName, bMakeVisibleAfterLoad, bShouldBlockOnLoad, LatentInfo, World);
 			LatentManager.AddNewAction(LatentInfo.CallbackTarget, LatentInfo.UUID, NewAction);
@@ -468,11 +449,10 @@ void UGameplayStatics::LoadStreamLevel(UObject* WorldContextObject, FName LevelN
 
 void UGameplayStatics::UnloadStreamLevel(UObject* WorldContextObject, FName LevelName,FLatentActionInfo LatentInfo)
 {
-	UWorld* World = GEngine->GetWorldFromContextObject(WorldContextObject);
-	if (World != nullptr)
+	if (UWorld* World = GEngine->GetWorldFromContextObject(WorldContextObject))
 	{
 		FLatentActionManager& LatentManager = World->GetLatentActionManager();
-		if (LatentManager.FindExistingAction<FStreamLevelAction>(LatentInfo.CallbackTarget, LatentInfo.UUID) == NULL)
+		if (LatentManager.FindExistingAction<FStreamLevelAction>(LatentInfo.CallbackTarget, LatentInfo.UUID) == nullptr)
 		{
 			FStreamLevelAction* NewAction = new FStreamLevelAction(false, LevelName, false, false, LatentInfo, World );
 			LatentManager.AddNewAction(LatentInfo.CallbackTarget, LatentInfo.UUID, NewAction );
@@ -484,8 +464,7 @@ ULevelStreaming* UGameplayStatics::GetStreamingLevel(UObject* WorldContextObject
 {
 	if (InPackageName != NAME_None)
 	{
-		UWorld* World = GEngine->GetWorldFromContextObject(WorldContextObject);
-		if (World != nullptr)
+		if (UWorld* World = GEngine->GetWorldFromContextObject(WorldContextObject))
 		{
 			const FString SearchPackageName = FStreamLevelAction::MakeSafeLevelName(InPackageName, World);
 
@@ -493,8 +472,7 @@ ULevelStreaming* UGameplayStatics::GetStreamingLevel(UObject* WorldContextObject
 			{
 				// We check only suffix of package name, to handle situations when packages were saved for play into a temporary folder
 				// Like Saved/Autosaves/PackageName
-				if (LevelStreaming && 
-					LevelStreaming->GetWorldAssetPackageName().EndsWith(SearchPackageName, ESearchCase::IgnoreCase))
+				if (LevelStreaming && LevelStreaming->GetWorldAssetPackageName().EndsWith(SearchPackageName, ESearchCase::IgnoreCase))
 				{
 					return LevelStreaming;
 				}
@@ -507,8 +485,7 @@ ULevelStreaming* UGameplayStatics::GetStreamingLevel(UObject* WorldContextObject
 
 void UGameplayStatics::FlushLevelStreaming(UObject* WorldContextObject)
 {
-	UWorld* World = GEngine->GetWorldFromContextObject(WorldContextObject);
-	if (World != nullptr)
+	if (UWorld* World = GEngine->GetWorldFromContextObject(WorldContextObject))
 	{
 		World->FlushLevelStreaming();
 	}
@@ -565,7 +542,7 @@ FVector UGameplayStatics::GetActorArrayAverageLocation(const TArray<AActor*>& Ac
 	{
 		AActor* A = Actors[ActorIdx];
 		// Check actor is non-null, not deleted, and has a root component
-		if(A != NULL && !A->IsPendingKill() && A->GetRootComponent() != NULL)
+		if (A && !A->IsPendingKill() && A->GetRootComponent())
 		{
 			LocationSum += A->GetActorLocation();
 			ActorCount++;
@@ -589,7 +566,7 @@ void UGameplayStatics::GetActorArrayBounds(const TArray<AActor*>& Actors, bool b
 	{
 		AActor* A = Actors[ActorIdx];
 		// Check actor is non-null, not deleted
-		if(A != NULL && !A->IsPendingKill())
+		if(A && !A->IsPendingKill())
 		{
 			ActorBounds += A->GetComponentsBoundingBox(!bOnlyCollidingComponents);
 		}
@@ -611,7 +588,7 @@ void UGameplayStatics::GetAllActorsOfClass(UObject* WorldContextObject, TSubclas
 	UWorld* World = GEngine->GetWorldFromContextObject(WorldContextObject);
 
 	// We do nothing if not class provided, rather than giving ALL actors!
-	if(ActorClass != NULL && World != nullptr)
+	if (ActorClass && World)
 	{
 		for(TActorIterator<AActor> It(World, ActorClass); It; ++It)
 		{
@@ -630,12 +607,12 @@ void UGameplayStatics::GetAllActorsWithInterface(UObject* WorldContextObject, TS
 
 	UWorld* World = GEngine->GetWorldFromContextObject(WorldContextObject);
 	// We do nothing if not class provided, rather than giving ALL actors!
-	if(Interface != NULL && World != nullptr)
+	if (Interface && World)
 	{
 		for(FActorIterator It(World); It; ++It)
 		{
 			AActor* Actor = *It;
-			if(Actor != NULL && !Actor->IsPendingKill() && Actor->GetClass()->ImplementsInterface(Interface))
+			if (Actor && !Actor->IsPendingKill() && Actor->GetClass()->ImplementsInterface(Interface))
 			{
 				OutActors.Add(Actor);
 			}
@@ -670,36 +647,32 @@ UParticleSystemComponent* CreateParticleSystem(UParticleSystem* EmitterTemplate,
 
 UParticleSystemComponent* UGameplayStatics::SpawnEmitterAtLocation(UObject* WorldContextObject, UParticleSystem* EmitterTemplate, FVector SpawnLocation, FRotator SpawnRotation, bool bAutoDestroy)
 {
-	UParticleSystemComponent* PSC = NULL;
 	if (EmitterTemplate)
 	{
-		UWorld* World = GEngine->GetWorldFromContextObject(WorldContextObject);
-		if(World != nullptr)
+		if (UWorld* World = GEngine->GetWorldFromContextObject(WorldContextObject))
 		{
-			PSC = CreateParticleSystem(EmitterTemplate, World, World->GetWorldSettings(), bAutoDestroy);
-
+			UParticleSystemComponent* PSC = CreateParticleSystem(EmitterTemplate, World, World->GetWorldSettings(), bAutoDestroy);
 			PSC->SetAbsolute(true, true, true);
 			PSC->SetWorldLocationAndRotation(SpawnLocation, SpawnRotation);
 			PSC->SetRelativeScale3D(FVector(1.f));
 			PSC->ActivateSystem(true);
+			return PSC;
 		}
 	}
-	return PSC;
+	return nullptr;
 }
 
 UParticleSystemComponent* UGameplayStatics::SpawnEmitterAttached(UParticleSystem* EmitterTemplate, USceneComponent* AttachToComponent, FName AttachPointName, FVector Location, FRotator Rotation, EAttachLocation::Type LocationType, bool bAutoDestroy)
 {
-	UParticleSystemComponent* PSC = NULL;
 	if (EmitterTemplate)
 	{
-		if (AttachToComponent == NULL)
+		if (!AttachToComponent)
 		{
 			UE_LOG(LogScript, Warning, TEXT("UGameplayStatics::SpawnEmitterAttached: NULL AttachComponent specified!"));
 		}
-		else if (AttachToComponent->GetWorld() != NULL && AttachToComponent->GetWorld()->GetNetMode() != NM_DedicatedServer)
+		else if (AttachToComponent->GetWorld() && AttachToComponent->GetWorld()->GetNetMode() != NM_DedicatedServer)
 		{
-			PSC = CreateParticleSystem(EmitterTemplate, AttachToComponent->GetWorld(), AttachToComponent->GetOwner(), bAutoDestroy);
-
+			UParticleSystemComponent* PSC = CreateParticleSystem(EmitterTemplate, AttachToComponent->GetWorld(), AttachToComponent->GetOwner(), bAutoDestroy);
 			PSC->AttachTo(AttachToComponent, AttachPointName);
 			if (LocationType == EAttachLocation::KeepWorldPosition)
 			{
@@ -711,9 +684,10 @@ UParticleSystemComponent* UGameplayStatics::SpawnEmitterAttached(UParticleSystem
 			}
 			PSC->SetRelativeScale3D(FVector(1.f));
 			PSC->ActivateSystem(true);
+			return PSC;
 		}
 	}
-	return PSC;
+	return nullptr;
 }
 
 void UGameplayStatics::BreakHitResult(const FHitResult& Hit, bool& bBlockingHit, bool& bInitialOverlap, float& Time, FVector& Location, FVector& ImpactPoint, FVector& Normal, FVector& ImpactNormal, UPhysicalMaterial*& PhysMat, AActor*& HitActor, UPrimitiveComponent*& HitComponent, FName& HitBoneName, int32& HitItem, FVector& TraceStart, FVector& TraceEnd)
@@ -1152,47 +1126,40 @@ UDecalComponent* CreateDecalComponent(class UMaterialInterface* DecalMaterial, F
 
 UDecalComponent* UGameplayStatics::SpawnDecalAtLocation(UObject* WorldContextObject, class UMaterialInterface* DecalMaterial, FVector DecalSize, FVector Location, FRotator Rotation, float LifeSpan)
 {
-	UDecalComponent* DecalComp = NULL;
-
 	if (DecalMaterial)
 	{
-		UWorld* World = GEngine->GetWorldFromContextObject(WorldContextObject);
-		if(World != nullptr)
+		if (UWorld* World = GEngine->GetWorldFromContextObject(WorldContextObject))
 		{
-			DecalComp = CreateDecalComponent(DecalMaterial, DecalSize, World, World->GetWorldSettings(), LifeSpan);
+			UDecalComponent* DecalComp = CreateDecalComponent(DecalMaterial, DecalSize, World, World->GetWorldSettings(), LifeSpan);
 			DecalComp->SetWorldLocationAndRotation(Location, Rotation);
+			return DecalComp;
 		}
 	}
-
-	return DecalComp;
+	return nullptr;
 }
 
 UDecalComponent* UGameplayStatics::SpawnDecalAttached(class UMaterialInterface* DecalMaterial, FVector DecalSize, class USceneComponent* AttachToComponent, FName AttachPointName, FVector Location, FRotator Rotation, EAttachLocation::Type LocationType, float LifeSpan)
 {
-	UDecalComponent* DecalComp = nullptr;
-
 	if (DecalMaterial)
 	{
-		if (AttachToComponent == nullptr)
+		if (!AttachToComponent)
 		{
 			UE_LOG(LogScript, Warning, TEXT("UGameplayStatics::SpawnDecalAttached: NULL AttachComponent specified!"));
 		}
 		else
 		{
 			UPrimitiveComponent* AttachToPrimitive = Cast<UPrimitiveComponent>(AttachToComponent);
-			if (AttachToPrimitive == nullptr || AttachToPrimitive->bReceivesDecals)
+			if (!AttachToPrimitive || AttachToPrimitive->bReceivesDecals)
 			{
-				const bool bOnBSPBrush = AttachToPrimitive && (Cast<AWorldSettings>(AttachToPrimitive->GetOwner()) != nullptr);
-				if (bOnBSPBrush)
+				if (AttachToPrimitive && Cast<AWorldSettings>(AttachToPrimitive->GetOwner()))
 				{
 					// special case: don't attach to component when it's owned by invisible WorldSettings (decals on BSP brush)
-					DecalComp = SpawnDecalAtLocation(AttachToPrimitive->GetOwner(), DecalMaterial, DecalSize, Location, Rotation, LifeSpan);
+					return SpawnDecalAtLocation(AttachToPrimitive->GetOwner(), DecalMaterial, DecalSize, Location, Rotation, LifeSpan);
 				}
 				else
 				{
-					DecalComp = CreateDecalComponent(DecalMaterial, DecalSize, AttachToComponent->GetWorld(), AttachToComponent->GetOwner(), LifeSpan);
+					UDecalComponent* DecalComp = CreateDecalComponent(DecalMaterial, DecalSize, AttachToComponent->GetWorld(), AttachToComponent->GetOwner(), LifeSpan);
 					DecalComp->AttachTo(AttachToComponent, AttachPointName);
-
 					if (LocationType == EAttachLocation::KeepWorldPosition)
 					{
 						DecalComp->SetWorldLocationAndRotation(Location, Rotation);
@@ -1201,51 +1168,41 @@ UDecalComponent* UGameplayStatics::SpawnDecalAttached(class UMaterialInterface*
 					{
 						DecalComp->SetRelativeLocationAndRotation(Location, Rotation);
 					}
+					return DecalComp;
 				}
 			}
 		}
 	}
-
-	return DecalComp;
+	return nullptr;
 }
 
 //////////////////////////////////////////////////////////////////////////
 
 USaveGame* UGameplayStatics::CreateSaveGameObject(TSubclassOf<USaveGame> SaveGameClass)
 {
-	USaveGame* OutSaveGameObject = NULL;
-
-	// don't save if no class or if class is the abstract base class
-	if ( (*SaveGameClass != NULL) && (*SaveGameClass != USaveGame::StaticClass()) )
+	// Don't save if no class or if class is the abstract base class.
+	if (*SaveGameClass && (*SaveGameClass != USaveGame::StaticClass()))
 	{
-		OutSaveGameObject = NewObject<USaveGame>(GetTransientPackage(), SaveGameClass);
+		return NewObject<USaveGame>(GetTransientPackage(), SaveGameClass);
 	}
-
-	return OutSaveGameObject;
+	return nullptr;
 }
 
 USaveGame* UGameplayStatics::CreateSaveGameObjectFromBlueprint(UBlueprint* SaveGameBlueprint)
 {
-	USaveGame* OutSaveGameObject = NULL;
-
-	if( SaveGameBlueprint != NULL &&
-		SaveGameBlueprint->GeneratedClass != NULL &&
-		SaveGameBlueprint->GeneratedClass->IsChildOf(USaveGame::StaticClass()) )
+	if (SaveGameBlueprint && SaveGameBlueprint->GeneratedClass && SaveGameBlueprint->GeneratedClass->IsChildOf(USaveGame::StaticClass()))
 	{
-		OutSaveGameObject = NewObject<USaveGame>(GetTransientPackage(), SaveGameBlueprint->GeneratedClass);
+		return NewObject<USaveGame>(GetTransientPackage(), SaveGameBlueprint->GeneratedClass);
 	}
-
-	return OutSaveGameObject;
+	return nullptr;
 }
 
 
 bool UGameplayStatics::SaveGameToSlot(USaveGame* SaveGameObject, const FString& SlotName, const int32 UserIndex)
 {
-	bool bSuccess = false;
-
 	ISaveGameSystem* SaveSystem = IPlatformFeaturesModule::Get().GetSaveGameSystem();
 	// If we have a system and an object to save and a save name...
-	if(SaveSystem && (SaveGameObject != NULL) && (SlotName.Len() > 0))
+	if(SaveSystem && SaveGameObject && (SlotName.Len() > 0))
 	{
 		TArray<uint8> ObjectBytes;
 		FMemoryWriter MemoryWriter(ObjectBytes, true);
@@ -1274,30 +1231,23 @@ bool UGameplayStatics::SaveGameToSlot(USaveGame* SaveGameObject, const FString&
 		SaveGameObject->Serialize(Ar);
 
 		// Stuff that data into the save system with the desired file name
-		bSuccess = SaveSystem->SaveGame(false, *SlotName, UserIndex, ObjectBytes);
+		return SaveSystem->SaveGame(false, *SlotName, UserIndex, ObjectBytes);
 	}
-
-	return bSuccess;
+	return false;
 }
 
 bool UGameplayStatics::DoesSaveGameExist(const FString& SlotName, const int32 UserIndex)
 {
-	bool bExists = false;
-
-	ISaveGameSystem* SaveSystem = IPlatformFeaturesModule::Get().GetSaveGameSystem();
-	if(SaveSystem != NULL)
+	if (ISaveGameSystem* SaveSystem = IPlatformFeaturesModule::Get().GetSaveGameSystem())
 	{
-		bExists = SaveSystem->DoesSaveGameExist(*SlotName, UserIndex);
+		return SaveSystem->DoesSaveGameExist(*SlotName, UserIndex);
 	}
-
-	return bExists;
+	return false;
 }
 
-//static 
 bool UGameplayStatics::DeleteGameInSlot(const FString& SlotName, const int32 UserIndex)
 {
-	ISaveGameSystem* SaveSystem = IPlatformFeaturesModule::Get().GetSaveGameSystem();
-	if (SaveSystem != NULL)
+	if (ISaveGameSystem* SaveSystem = IPlatformFeaturesModule::Get().GetSaveGameSystem())
 	{
 		return SaveSystem->DeleteGame(false, *SlotName, UserIndex);
 	}
@@ -1403,8 +1353,7 @@ void UGameplayStatics::GetAccurateRealTime(UObject* WorldContextObject, int32& S
 
 void UGameplayStatics::EnableLiveStreaming(bool Enable)
 {
-	IDVRStreamingSystem* StreamingSystem = IPlatformFeaturesModule::Get().GetStreamingSystem();
-	if (StreamingSystem)
+	if (IDVRStreamingSystem* StreamingSystem = IPlatformFeaturesModule::Get().GetStreamingSystem())
 	{
 		StreamingSystem->EnableStreaming(Enable);
 	}
