Added GetIni command, for retrieving .ini file values through the console (required to access from the console, values not accessible through config UPROPERTY members).

diff --git a/Engine/Source/Runtime/Engine/Classes/Engine/Engine.h b/Engine/Source/Runtime/Engine/Classes/Engine/Engine.h
--- a/Engine/Source/Runtime/Engine/Classes/Engine/Engine.h
+++ b/Engine/Source/Runtime/Engine/Classes/Engine/Engine.h
@@ -1771,7 +1771,8 @@ public:
 	bool HandleEnableAllScreenMessagesCommand( const TCHAR* Cmd, FOutputDevice& Ar );			
 	bool HandleToggleAllScreenMessagesCommand( const TCHAR* Cmd, FOutputDevice& Ar );			
 	bool HandleConfigHashCommand( const TCHAR* Cmd, FOutputDevice& Ar );						
-	bool HandleConfigMemCommand( const TCHAR* Cmd, FOutputDevice& Ar );							
+	bool HandleConfigMemCommand( const TCHAR* Cmd, FOutputDevice& Ar );	
+	bool HandleGetIniCommand(const TCHAR* Cmd, FOutputDevice& Ar);
 #endif // !UE_BUILD_SHIPPING
 
 	/** Update everything. */
diff --git a/Engine/Source/Runtime/Engine/Private/UnrealEngine.cpp b/Engine/Source/Runtime/Engine/Private/UnrealEngine.cpp
--- a/Engine/Source/Runtime/Engine/Private/UnrealEngine.cpp
+++ b/Engine/Source/Runtime/Engine/Private/UnrealEngine.cpp
@@ -2671,6 +2671,10 @@ bool UEngine::Exec( UWorld* InWorld, const TCHAR* Cmd, FOutputDevice& Ar )
 	{
 		return HandleConfigMemCommand( Cmd, Ar );
 	}
+	else if (FParse::Command(&Cmd, TEXT("GETINI")))
+	{
+		return HandleGetIniCommand(Cmd, Ar);
+	}
 #endif // !UE_BUILD_SHIPPING
 	
 	else if ( FParse::Command(&Cmd,TEXT("SCALABILITY")) )
@@ -5750,6 +5754,105 @@ bool UEngine::HandleConfigMemCommand( const TCHAR* Cmd, FOutputDevice& Ar )
 	GConfig->ShowMemoryUsage(Ar);
 	return true;
 }
+
+bool UEngine::HandleGetIniCommand(const TCHAR* Cmd, FOutputDevice& Ar)
+{
+	// Format: GetIni IniFile:Section.SubSection Key
+	TCHAR IniPlusSectionName[256];
+	TCHAR KeyName[256];
+
+	if (FParse::Token(Cmd, IniPlusSectionName, ARRAY_COUNT(IniPlusSectionName), true))
+	{
+		FString IniPlusSection = IniPlusSectionName;
+		int32 IniDelim = IniPlusSection.Find(TEXT(":"));
+		FString IniName;
+		FString SectionName = (IniDelim != INDEX_NONE ? IniPlusSection.Mid(IniDelim+1) : IniPlusSection);
+
+		if (IniDelim != INDEX_NONE)
+		{
+			// Check for hardcoded engine-ini:, game-ini: etc. parsing, and if not found fallback to raw string
+			const FString* HardcodedIni = GetIniFilenameFromObjectsReference(IniPlusSection);
+
+			if (HardcodedIni != NULL)
+			{
+				IniName = *HardcodedIni;
+			}
+			else
+			{
+				TArray<FString> ConfigList;
+				FString SearchStr = IniPlusSection.Left(IniDelim) + TEXT(".ini");
+
+				GConfig->GetConfigFilenames(ConfigList);
+
+				for (auto CurConfig : ConfigList)
+				{
+					if (CurConfig.Contains(SearchStr, ESearchCase::IgnoreCase, ESearchDir::FromEnd))
+					{
+						IniName = CurConfig;
+						break;
+					}
+				}
+
+				if (IniName.IsEmpty())
+				{
+					UE_SUPPRESS(LogExec, Warning,
+								Ar.Logf(TEXT("Failed to find IniFile '%s' (note: can only search loaded ini files)."), *SearchStr));
+				}
+			}
+		}
+		else
+		{
+			IniName = GEngineIni;
+		}
+
+		if (!IniName.IsEmpty() && !SectionName.IsEmpty())
+		{
+			if (FParse::Token(Cmd, KeyName, ARRAY_COUNT(KeyName), true))
+			{
+				TArray<FString> Values;
+
+				bool bSuccess = !!GConfig->GetArray(*SectionName, KeyName, Values, IniName);
+
+				if (bSuccess)
+				{
+					for (auto CurValue : Values)
+					{
+						Ar.Log(*CurValue);
+					}
+				}
+				else
+				{
+					UE_SUPPRESS(LogExec, Warning, Ar.Logf(TEXT("Failed to get config key '%s', from section '%s', in ini file '%s'."),
+								KeyName, *SectionName, *IniName));
+				}
+			}
+			else
+			{
+				UE_SUPPRESS(LogExec, Warning,
+							Ar.Logf(TEXT("No Key specified. Command format: GetIni IniFile:Section.SubSection Key")));
+			}
+		}
+		else if (IniName.IsEmpty())
+		{
+			UE_SUPPRESS(LogExec, Warning,
+						Ar.Logf(TEXT("IniFile parsing failed (%s). Command format: GetIni IniFile:Section.SubSection Key"),
+								IniPlusSectionName));
+		}
+		else // if (SectionName.IsEmpty())
+		{
+			UE_SUPPRESS(LogExec, Warning,
+						Ar.Logf(TEXT("Section parsing failed (%s). Command format: GetIni IniFile:Section.SubSection Key"),
+								IniPlusSectionName));
+		}
+	}
+	else
+	{
+		UE_SUPPRESS(LogExec, Warning,
+					Ar.Logf(TEXT("No Section specified. Command format: GetIni IniFile:Section.SubSection Key")))
+	}
+
+	return true;
+}
 #endif // !UE_BUILD_SHIPPING
 
 bool UEngine::HandleFlushLogCommand( const TCHAR* Cmd, FOutputDevice& Ar )
