Fixing SoundGenerators sample. Added a small function to ActiveSound.h in the engine.

diff --git a/Engine/Source/Runtime/Engine/Public/ActiveSound.h b/Engine/Source/Runtime/Engine/Public/ActiveSound.h
--- a/Engine/Source/Runtime/Engine/Public/ActiveSound.h
+++ b/Engine/Source/Runtime/Engine/Public/ActiveSound.h
@@ -82,6 +82,7 @@ private:
 
 public:
 
+	bool IsAudioComponentValid() const { return AudioComponent.IsValid(); }
 	UAudioComponent* GetAudioComponent() const { return AudioComponent.Get(); }
 	UPTRINT GetAudioComponentIndex() const { return AudioComponentIndex; }
 	void SetAudioComponent(UAudioComponent* Component)
