fixed comment

diff --git a/Engine/Source/Runtime/Engine/Private/Materials/HLSLMaterialTranslator.h b/Engine/Source/Runtime/Engine/Private/Materials/HLSLMaterialTranslator.h
--- a/Engine/Source/Runtime/Engine/Private/Materials/HLSLMaterialTranslator.h
+++ b/Engine/Source/Runtime/Engine/Private/Materials/HLSLMaterialTranslator.h
@@ -2516,7 +2516,7 @@ protected:
 		return AddInlinedCodeChunk(MCT_Float, TEXT("GetScreenPosition(Parameters).w"));		
 	}
 
-	/** Calculate screen aligned UV coodinates from an offset fraction or texture coordinate */
+	/** Calculate screen aligned UV coordinates from an offset fraction or texture coordinate */
 	int32 GetScreenAlignedUV(int32 Offset, int32 UV, bool bUseOffset)
 	{
 		if(bUseOffset)
