UE4 - Fixed malloc churn on particle async tasks and queue the gamethread task from the async task for non-editor builds.

diff --git a/Engine/Source/Runtime/Engine/Classes/Particles/ParticleSystemComponent.h b/Engine/Source/Runtime/Engine/Classes/Particles/ParticleSystemComponent.h
--- a/Engine/Source/Runtime/Engine/Classes/Particles/ParticleSystemComponent.h
+++ b/Engine/Source/Runtime/Engine/Classes/Particles/ParticleSystemComponent.h
@@ -523,7 +523,7 @@ private:
 	bool bParallelRenderThreadUpdate;
 public:
 
-	const FTransform& GetAsyncComponentToWorld()
+	FORCEINLINE const FTransform& GetAsyncComponentToWorld()
 	{
 		if (!bParallelRenderThreadUpdate && !IsInGameThread())
 		{
@@ -533,7 +533,7 @@ public:
 		return ComponentToWorld;
 	}
 
-	const TArray<struct FParticleSysParam>& GetAsyncInstanceParameters()
+	FORCEINLINE const TArray<struct FParticleSysParam>& GetAsyncInstanceParameters()
 	{
 		if (!bParallelRenderThreadUpdate && !IsInGameThread())
 		{
diff --git a/Engine/Source/Runtime/Engine/Private/Particles/ParticleComponents.cpp b/Engine/Source/Runtime/Engine/Private/Particles/ParticleComponents.cpp
--- a/Engine/Source/Runtime/Engine/Private/Particles/ParticleComponents.cpp
+++ b/Engine/Source/Runtime/Engine/Private/Particles/ParticleComponents.cpp
@@ -3659,6 +3659,13 @@ public:
 	void DoTask(ENamedThreads::Type CurrentThread, const FGraphEventRef& MyCompletionGraphEvent)
 	{
 		Target->ComputeTickComponent_Concurrent();
+#if !WITH_EDITOR  // otherwise this is queued by the calling code because we need to be able to block and wait on it
+		{
+			QUICK_SCOPE_CYCLE_COUNTER(STAT_UParticleSystemComponent_QueueFinalize);
+			FGraphEventRef Finalize = TGraphTask<FParticleFinalizeTask>::CreateTask(nullptr, CurrentThread).ConstructAndDispatchWhenReady(Target);
+			MyCompletionGraphEvent->DontCompleteUntil(Finalize);
+		}
+#endif
 	}
 };
 
@@ -3842,13 +3849,17 @@ void UParticleSystemComponent::TickComponent(float DeltaTime, enum ELevelTick Ti
 		bAsyncDataCopyIsValid = true;
 		check(!bParallelRenderThreadUpdate);
 		AsyncComponentToWorld = ComponentToWorld;
-		AsyncInstanceParameters = InstanceParameters;
+		AsyncInstanceParameters.Reset();
+		AsyncInstanceParameters.Append(InstanceParameters);
 
 		{
 			QUICK_SCOPE_CYCLE_COUNTER(STAT_UParticleSystemComponent_QueueAsync);
 			AsyncWork = TGraphTask<FParticleAsyncTask>::CreateTask(nullptr, ENamedThreads::GameThread).ConstructAndDispatchWhenReady(this);
+#if !WITH_EDITOR  // we need to not complete until this is done because the game thread finalize task has not beed queued yet
+			ThisTickFunction->GetCompletionHandle()->DontCompleteUntil(AsyncWork);
+#endif
 		}
-
+#if WITH_EDITOR  // we need to queue this here because we need to be able to block and wait on it
 		{
 			QUICK_SCOPE_CYCLE_COUNTER(STAT_UParticleSystemComponent_QueueFinalize);
 			FGraphEventArray Prereqs;
@@ -3856,6 +3867,7 @@ void UParticleSystemComponent::TickComponent(float DeltaTime, enum ELevelTick Ti
 			FGraphEventRef Finalize = TGraphTask<FParticleFinalizeTask>::CreateTask(&Prereqs, ENamedThreads::GameThread).ConstructAndDispatchWhenReady(this);
 			ThisTickFunction->GetCompletionHandle()->DontCompleteUntil(Finalize);
 		}
+#endif
 	}
 }
 
@@ -4052,8 +4064,9 @@ void UParticleSystemComponent::FinalizeTickComponent()
 
 void UParticleSystemComponent::WaitForAsyncAndFinalize(EForceAsyncWorkCompletion Behavior) const
 {
-	if (!AsyncWork->IsComplete())
+	if (AsyncWork.GetReference() && !AsyncWork->IsComplete())
 	{
+#if WITH_EDITOR
 		check(IsInGameThread());
 		SCOPE_CYCLE_COUNTER(STAT_GTSTallTime);
 		double StartTime = FPlatformTime::Seconds();
@@ -4069,6 +4082,9 @@ void UParticleSystemComponent::WaitForAsyncAndFinalize(EForceAsyncWorkCompletion
 			}
 		}
 		const_cast<UParticleSystemComponent*>(this)->FinalizeTickComponent();
+#else
+		check(0);
+#endif
 	}
 }
 
