grey physical material UI if not surface

diff --git a/Engine/Source/Runtime/Engine/Private/Materials/Material.cpp b/Engine/Source/Runtime/Engine/Private/Materials/Material.cpp
--- a/Engine/Source/Runtime/Engine/Private/Materials/Material.cpp
+++ b/Engine/Source/Runtime/Engine/Private/Materials/Material.cpp
@@ -2487,6 +2487,11 @@ bool UMaterial::CanEditChange(const UProperty* InProperty) const
 	{
 		FString PropertyName = InProperty->GetName();
 
+		if (PropertyName == GET_MEMBER_NAME_STRING_CHECKED(UMaterial, PhysMaterial))
+		{
+			return MaterialDomain == MD_Surface;
+		}
+
 		if (PropertyName == GET_MEMBER_NAME_STRING_CHECKED(UMaterial, OpacityMaskClipValue))
 		{
 			return BlendMode == BLEND_Masked;
