Allow systems to register to provide new options for show debug

diff --git a/Engine/Source/Runtime/Engine/Classes/GameFramework/HUD.h b/Engine/Source/Runtime/Engine/Classes/GameFramework/HUD.h
--- a/Engine/Source/Runtime/Engine/Classes/GameFramework/HUD.h
+++ b/Engine/Source/Runtime/Engine/Classes/GameFramework/HUD.h
@@ -13,6 +13,8 @@ class UCanvas;
 class APlayerController;
 class UFont;
 
+DECLARE_MULTICAST_DELEGATE_FiveParams(FOnShowDebugInfo, AHUD* /* HUD */, UCanvas* /* Canvas */, const FDebugDisplayInfo& /* DisplayInfo */, float& /* YL */, float& /* YPos */);
+
 /** 
  * Base class of the heads-up display. This has a canvas and a debug canvas on which primitives can be drawn.
  * It also contains a list of simple hit boxes that can be used for simple item click detection.
@@ -454,6 +456,9 @@ public:
 	 */
 	virtual void ShowDebugInfo(float& YL, float& YPos);
 
+	// Callback allowing external systems to register to show debug info
+	static FOnShowDebugInfo OnShowDebugInfo;
+
 	/** PostRender is the main draw loop. */
 	virtual void PostRender();
 
diff --git a/Engine/Source/Runtime/Engine/Private/HUD.cpp b/Engine/Source/Runtime/Engine/Private/HUD.cpp
--- a/Engine/Source/Runtime/Engine/Private/HUD.cpp
+++ b/Engine/Source/Runtime/Engine/Private/HUD.cpp
@@ -19,6 +19,8 @@ DEFINE_LOG_CATEGORY_STATIC(LogHUD, Log, All);
 #define LOCTEXT_NAMESPACE "HUD"
 
 bool AHUD::bShowDebugForReticleTarget = false;
+FOnShowDebugInfo AHUD::OnShowDebugInfo;
+
 
 AHUD::AHUD(const FObjectInitializer& ObjectInitializer)
 	: Super(ObjectInitializer)
@@ -354,6 +356,11 @@ void AHUD::ShowDebugInfo(float& YL, float& YPos)
 		{
 			GetWorld()->GetAuthGameMode()->DisplayDebug(DebugCanvas, DisplayInfo, YL, YPos);
 		}
+
+		if (bShowDebugInfo)
+		{
+			OnShowDebugInfo.Broadcast(this, DebugCanvas, DisplayInfo, YL, YPos);
+		}
 	}
 }
 
