Added TouchEventLocations to pair with touch events since we want the actual location for that event.  After a release, the finger id may be reused for another touch far from the location of the release so the wrong location would be sent for all the events.
#ue4
#android
#ios

diff --git a/Engine/Source/Runtime/Engine/Classes/GameFramework/PlayerInput.h b/Engine/Source/Runtime/Engine/Classes/GameFramework/PlayerInput.h
--- a/Engine/Source/Runtime/Engine/Classes/GameFramework/PlayerInput.h
+++ b/Engine/Source/Runtime/Engine/Classes/GameFramework/PlayerInput.h
@@ -256,6 +256,9 @@ public:
 	// @todo: We have 10 touches to match the number of Touch* entries in EKeys (not easy to make this an enum or define or anything)
 	FVector Touches[EKeys::NUM_TOUCH_KEYS];
 
+	/** Used to store paired touch locations for event ids during the frame and flushed when processed. */
+	TMap<uint32, FVector> TouchEventLocations;
+
 	// Mouse smoothing sample data
 	float ZeroTime[2];    /** How long received mouse movement has been zero. */
 	float SmoothedMouse[2];    /** Current average mouse movement/sample */
diff --git a/Engine/Source/Runtime/Engine/Private/UserInterface/PlayerInput.cpp b/Engine/Source/Runtime/Engine/Private/UserInterface/PlayerInput.cpp
--- a/Engine/Source/Runtime/Engine/Private/UserInterface/PlayerInput.cpp
+++ b/Engine/Source/Runtime/Engine/Private/UserInterface/PlayerInput.cpp
@@ -254,6 +254,9 @@ bool UPlayerInput::InputTouch(uint32 Handle, ETouchType::Type Type, const FVecto
 		break;
 	}
 
+	// store current touch location paired with event id
+	TouchEventLocations.Add(EventCount, Touches[Handle]);
+
 	// accumulate deltas until processed next
 	KeyState.SampleCountAccumulator++;
 	KeyState.RawValueAccumulator = KeyState.Value = KeyState.RawValue = FVector(TouchLocation.X, TouchLocation.Y, 0);
@@ -976,11 +979,12 @@ void UPlayerInput::ProcessInputStack(const TArray<UInputComponent*>& InputCompon
 						if (TB.bExecuteWhenPaused || !bGamePaused)
 						{
 							check(EventIndices.Num() > 0);
-							FDelegateDispatchDetails TouchInfo(EventIndices[0], NonAxisDelegates.Num(), TB.TouchDelegate, Touches[TouchIndex], TouchIndex);
+							FDelegateDispatchDetails TouchInfo(EventIndices[0], NonAxisDelegates.Num(), TB.TouchDelegate, TouchEventLocations[EventIndices[0]], TouchIndex);
 							NonAxisDelegates.Add(TouchInfo);
 							for (int32 EventsIndex = 1; EventsIndex < EventIndices.Num(); ++EventsIndex)
 							{
 								TouchInfo.EventIndex = EventIndices[EventsIndex];
+								TouchInfo.TouchLocation = TouchEventLocations[TouchInfo.EventIndex];
 								NonAxisDelegates.Add(TouchInfo);
 							}
 						}
@@ -1151,6 +1155,7 @@ void UPlayerInput::ProcessInputStack(const TArray<UInputComponent*>& InputCompon
 	AxisDelegates.Reset();
 	VectorAxisDelegates.Reset();
 	NonAxisDelegates.Reset();
+	TouchEventLocations.Reset();
 }
 
 void UPlayerInput::DiscardPlayerInput()
