ActionKeyMapping & AxisKeyMapping are now exposed to Blueprint

diff --git a/Engine/Source/Runtime/Engine/Classes/GameFramework/PlayerInput.h b/Engine/Source/Runtime/Engine/Classes/GameFramework/PlayerInput.h
--- a/Engine/Source/Runtime/Engine/Classes/GameFramework/PlayerInput.h
+++ b/Engine/Source/Runtime/Engine/Classes/GameFramework/PlayerInput.h
@@ -115,7 +115,7 @@ struct FInputAxisConfigEntry
  *
  * @see https://docs.unrealengine.com/latest/INT/Gameplay/Input/index.html
  */
-USTRUCT()
+USTRUCT( BlueprintType )
 struct FInputActionKeyMapping
 {
 	GENERATED_USTRUCT_BODY()
@@ -183,7 +183,7 @@ struct FInputActionKeyMapping
  * 
  * @see https://docs.unrealengine.com/latest/INT/Gameplay/Input/index.html
 **/
-USTRUCT()
+USTRUCT( BlueprintType )
 struct FInputAxisKeyMapping
 {
 	GENERATED_USTRUCT_BODY()
