#ue4 - Fix OnComponentCreated() not always being called before RegisterComponent().

--------

diff --git a/Engine/Source/Runtime/Engine/Private/Components/ActorComponent.cpp b/Engine/Source/Runtime/Engine/Private/Components/ActorComponent.cpp
--- a/Engine/Source/Runtime/Engine/Private/Components/ActorComponent.cpp
+++ b/Engine/Source/Runtime/Engine/Private/Components/ActorComponent.cpp
@@ -825,6 +825,11 @@ void UActorComponent::RegisterComponentWithWorld(UWorld* InWorld)
 	}
 #endif // !(UE_BUILD_SHIPPING || UE_BUILD_TEST)
 
+	if (!bHasBeenCreated)
+	{
+		OnComponentCreated();
+	}
+
 	World = InWorld;
 
 	ExecuteRegisterEvents();
