#ue4 - Move AttachmentReplication update out of SceneComponent and to ActorReplication. It only needs to be done on the network authority and only for the root component, and replication updates is a better place to check this.

- AActor::AttachmentReplication is now private. Added a read-only accessor.

--------

diff --git a/Engine/Source/Runtime/Engine/Classes/GameFramework/Actor.h b/Engine/Source/Runtime/Engine/Classes/GameFramework/Actor.h
--- a/Engine/Source/Runtime/Engine/Classes/GameFramework/Actor.h
+++ b/Engine/Source/Runtime/Engine/Classes/GameFramework/Actor.h
@@ -264,10 +264,15 @@ public:
 	UPROPERTY(EditDefaultsOnly, ReplicatedUsing=OnRep_ReplicatedMovement, Category=Replication, AdvancedDisplay)
 	struct FRepMovement ReplicatedMovement;
 
+private:
 	/** Used for replicating attachment of this actor's RootComponent to another actor. */
 	UPROPERTY(Transient, replicatedUsing=OnRep_AttachmentReplication)
 	struct FRepAttachment AttachmentReplication;
 
+public:
+	/** Get read-only access to current AttachmentReplication. */
+	const struct FRepAttachment& GetAttachmentReplication() const { return AttachmentReplication; }
+
 	/** Called on client when updated AttachmentReplication value is received for this actor. */
 	UFUNCTION()
 	virtual void OnRep_AttachmentReplication();
diff --git a/Engine/Source/Runtime/Engine/Private/Actor.cpp b/Engine/Source/Runtime/Engine/Private/Actor.cpp
--- a/Engine/Source/Runtime/Engine/Private/Actor.cpp
+++ b/Engine/Source/Runtime/Engine/Private/Actor.cpp
@@ -817,7 +817,10 @@ bool AActor::ShouldTickIfViewportsOnly() const
 
 void AActor::PreReplication( IRepChangedPropertyTracker & ChangedPropertyTracker )
 {
-	if ( bReplicateMovement || AttachmentReplication.AttachParent )
+	// Attachment replication gets filled in by GatherCurrentMovement(), but in the case of a detached root we need to trigger remote detachment.
+	AttachmentReplication.AttachParent = nullptr;
+
+	if ( bReplicateMovement || (RootComponent && RootComponent->AttachParent) )
 	{
 		GatherCurrentMovement();
 	}
diff --git a/Engine/Source/Runtime/Engine/Private/ActorReplication.cpp b/Engine/Source/Runtime/Engine/Private/ActorReplication.cpp
--- a/Engine/Source/Runtime/Engine/Private/ActorReplication.cpp
+++ b/Engine/Source/Runtime/Engine/Private/ActorReplication.cpp
@@ -231,6 +231,8 @@ bool AActor::IsNetRelevantFor(const AActor* RealViewer, const AActor* ViewTarget
 
 void AActor::GatherCurrentMovement()
 {
+	AttachmentReplication.AttachParent = nullptr;
+
 	UPrimitiveComponent* RootPrimComp = Cast<UPrimitiveComponent>(GetRootComponent());
 	if (RootPrimComp && RootPrimComp->IsSimulatingPhysics())
 	{
@@ -240,18 +242,21 @@ void AActor::GatherCurrentMovement()
 		ReplicatedMovement.FillFrom(RBState);
 		ReplicatedMovement.bRepPhysics = true;
 	}
-	else if(RootComponent != NULL)
+	else if (RootComponent != nullptr)
 	{
-		// If we are attached, don't replicate absolute position
-		if( RootComponent->AttachParent != NULL )
+		// If we are attached, don't replicate absolute position, use AttachmentReplication instead.
+		if (RootComponent->AttachParent != nullptr)
 		{
 			// Networking for attachments assumes the RootComponent of the AttachParent actor. 
 			// If that's not the case, we can't update this, as the client wouldn't be able to resolve the Component and would detach as a result.
-			if( AttachmentReplication.AttachParent != NULL )
+			AttachmentReplication.AttachParent = RootComponent->AttachParent->GetAttachmentRootActor();
+			if (AttachmentReplication.AttachParent != nullptr)
 			{
 				AttachmentReplication.LocationOffset = RootComponent->RelativeLocation;
 				AttachmentReplication.RotationOffset = RootComponent->RelativeRotation;
 				AttachmentReplication.RelativeScale3D = RootComponent->RelativeScale3D;
+				AttachmentReplication.AttachComponent = RootComponent->AttachParent;
+				AttachmentReplication.AttachSocket = RootComponent->AttachSocketName;
 			}
 		}
 		else
diff --git a/Engine/Source/Runtime/Engine/Private/Components/SceneComponent.cpp b/Engine/Source/Runtime/Engine/Private/Components/SceneComponent.cpp
--- a/Engine/Source/Runtime/Engine/Private/Components/SceneComponent.cpp
+++ b/Engine/Source/Runtime/Engine/Private/Components/SceneComponent.cpp
@@ -456,26 +456,6 @@ void USceneComponent::PropagateTransformUpdate(bool bTransformChanged, bool bSki
 
 		// Refresh navigation
 		UpdateNavigationData();
-
-		//QUICK_SCOPE_CYCLE_COUNTER(STAT_SceneComponent_PropagateTransformUpdate_AttachmentReplication);
-
-		AActor* Owner = GetOwner();
-		if (Owner && Owner->HasAuthority() && Owner->GetRootComponent() == this)
-		{
-			if (AttachParent)
-			{
-				Owner->AttachmentReplication.AttachParent = AttachParent->GetAttachmentRootActor();
-				Owner->AttachmentReplication.LocationOffset = RelativeLocation;
-				Owner->AttachmentReplication.RotationOffset = RelativeRotation;
-				Owner->AttachmentReplication.RelativeScale3D = RelativeScale3D;
-				Owner->AttachmentReplication.AttachSocket = AttachSocketName;
-				Owner->AttachmentReplication.AttachComponent = AttachParent;
-			}
-			else
-			{
-				Owner->AttachmentReplication.AttachParent = nullptr;
-			}
-		}
 	}
 	else
 	{
