#ue4 - Clean up some stats in SceneComponent.

--------

diff --git a/Engine/Source/Runtime/Engine/Private/Components/SceneComponent.cpp b/Engine/Source/Runtime/Engine/Private/Components/SceneComponent.cpp
--- a/Engine/Source/Runtime/Engine/Private/Components/SceneComponent.cpp
+++ b/Engine/Source/Runtime/Engine/Private/Components/SceneComponent.cpp
@@ -422,7 +422,7 @@ void USceneComponent::OnRegister()
 
 void USceneComponent::PropagateTransformUpdate(bool bTransformChanged, bool bSkipPhysicsMove, ETeleportType Teleport)
 {
-	//QUICK_SCOPE_CYCLE_COUNTER(STAT_SceneComponent_PropagateTransformUpdate);
+	//QUICK_SCOPE_CYCLE_COUNTER(STAT_USceneComponent_PropagateTransformUpdate);
 	if (IsDeferringMovementUpdates())
 	{
 		// We are deferring these updates until later.
@@ -431,15 +431,15 @@ void USceneComponent::PropagateTransformUpdate(bool bTransformChanged, bool bSki
 	FPlatformMisc::Prefetch(AttachChildren.GetData());
 	if (bTransformChanged)
 	{
-		//QUICK_SCOPE_CYCLE_COUNTER(STAT_SceneComponent_PropagateTransformUpdate_TransformChanged);
+		//QUICK_SCOPE_CYCLE_COUNTER(STAT_USceneComponent_PropagateTransformUpdate_TransformChanged);
 		{
 		// Then update bounds
-			//QUICK_SCOPE_CYCLE_COUNTER(STAT_SceneComponent_PropagateTransformUpdate_UpdateBounds);
+			//QUICK_SCOPE_CYCLE_COUNTER(STAT_USceneComponent_PropagateTransformUpdate_UpdateBounds);
 		UpdateBounds();
 		}
 
 		{
-			//QUICK_SCOPE_CYCLE_COUNTER(STAT_SceneComponent_PropagateTransformUpdate_OnUpdateTransform);
+			//QUICK_SCOPE_CYCLE_COUNTER(STAT_USceneComponent_PropagateTransformUpdate_OnUpdateTransform);
 		// Always send new transform to physics
 		OnUpdateTransform(bSkipPhysicsMove, Teleport);
 		}
@@ -448,7 +448,7 @@ void USceneComponent::PropagateTransformUpdate(bool bTransformChanged, bool bSki
 		MarkRenderTransformDirty();
 		
 		{
-			//QUICK_SCOPE_CYCLE_COUNTER(STAT_SceneComponent_PropagateTransformUpdate_UpdateChildTransforms);
+			//QUICK_SCOPE_CYCLE_COUNTER(STAT_USceneComponent_PropagateTransformUpdate_UpdateChildTransforms);
 		// Now go and update children
 		//Do not pass skip physics to children. This is only used when physics updates us, but in that case we really do need to update the attached children since they are kinematic
 		UpdateChildTransforms(false, Teleport);
@@ -459,21 +459,21 @@ void USceneComponent::PropagateTransformUpdate(bool bTransformChanged, bool bSki
 	}
 	else
 	{
-		//QUICK_SCOPE_CYCLE_COUNTER(STAT_SceneComponent_PropagateTransformUpdate_NOT_TransformChanged);
+		//QUICK_SCOPE_CYCLE_COUNTER(STAT_USceneComponent_PropagateTransformUpdate_NOT_TransformChanged);
 		{
-			//QUICK_SCOPE_CYCLE_COUNTER(STAT_SceneComponent_PropagateTransformUpdate_UpdateBounds);
+			//QUICK_SCOPE_CYCLE_COUNTER(STAT_USceneComponent_PropagateTransformUpdate_UpdateBounds);
 		// We update bounds even if transform doesn't change, as shape/mesh etc might have done
 		UpdateBounds();
 		}
 
 		{
-			//QUICK_SCOPE_CYCLE_COUNTER(STAT_SceneComponent_PropagateTransformUpdate_UpdateChildTransforms);
+			//QUICK_SCOPE_CYCLE_COUNTER(STAT_USceneComponent_PropagateTransformUpdate_UpdateChildTransforms);
 		// Now go and update children
 		UpdateChildTransforms();
 		}
 
 		{
-			QUICK_SCOPE_CYCLE_COUNTER(STAT_SceneComponent_PropagateTransformUpdate_MarkRenderTransformDirty);
+			//QUICK_SCOPE_CYCLE_COUNTER(STAT_USceneComponent_PropagateTransformUpdate_MarkRenderTransformDirty);
 		// Need to flag as dirty so new bounds are sent to render thread
 		MarkRenderTransformDirty();
 	}
