Remove logspam for HISMC concurrent removal

diff --git a/Engine/Source/Runtime/Engine/Private/HierarchicalInstancedStaticMesh.cpp b/Engine/Source/Runtime/Engine/Private/HierarchicalInstancedStaticMesh.cpp
--- a/Engine/Source/Runtime/Engine/Private/HierarchicalInstancedStaticMesh.cpp
+++ b/Engine/Source/Runtime/Engine/Private/HierarchicalInstancedStaticMesh.cpp
@@ -2232,7 +2232,7 @@ void UHierarchicalInstancedStaticMeshComponent::ApplyBuildTreeAsync(ENamedThread
 	{
 		bConcurrentRemoval = false;
 
-		UE_LOG(LogStaticMesh, Display, TEXT("Discarded foliage hierarchy of %d elements build due to concurrent removal (%.1fs)"), Builder->Result->InstanceReorderTable.Num(), float(FPlatformTime::Seconds() - StartTime));
+		UE_LOG(LogStaticMesh, Verbose, TEXT("Discarded foliage hierarchy of %d elements build due to concurrent removal (%.1fs)"), Builder->Result->InstanceReorderTable.Num(), float(FPlatformTime::Seconds() - StartTime));
 
 		// There were removes or updates while we were building, it's too slow to fix up the result now, so build async again.
 		BuildTreeAsync();
