Exposed the ESplineMeshAxis type to Blueprints.

diff --git a/Engine/Source/Runtime/Engine/Classes/Components/SplineMeshComponent.h b/Engine/Source/Runtime/Engine/Classes/Components/SplineMeshComponent.h
--- a/Engine/Source/Runtime/Engine/Classes/Components/SplineMeshComponent.h
+++ b/Engine/Source/Runtime/Engine/Classes/Components/SplineMeshComponent.h
@@ -7,7 +7,7 @@
 
 #include "SplineMeshComponent.generated.h"
 
-UENUM()
+UENUM(BlueprintType)
 namespace ESplineMeshAxis
 {
 	enum Type
