SplineComponent::GetSegmentLength now takes into account the world scale of the component when rebuilding the remap table. Also exposed UpdateSpline to Blueprints so the spline can be updated in code after applying a new scale transform.

diff --git a/Engine/Source/Runtime/Engine/Classes/Components/SplineComponent.h b/Engine/Source/Runtime/Engine/Classes/Components/SplineComponent.h
--- a/Engine/Source/Runtime/Engine/Classes/Components/SplineComponent.h
+++ b/Engine/Source/Runtime/Engine/Classes/Components/SplineComponent.h
@@ -114,6 +114,7 @@ public:
 	void ApplyComponentInstanceData(class FSplineInstanceData* ComponentInstanceData, const bool bPostUCS);
 
 	/** Update the spline tangents and SplineReparamTable */
+	UFUNCTION(BlueprintCallable, Category = Spline)
 	void UpdateSpline();
 
 	/** Get location along spline at the provided input key value */
diff --git a/Engine/Source/Runtime/Engine/Private/Components/SplineComponent.cpp b/Engine/Source/Runtime/Engine/Private/Components/SplineComponent.cpp
--- a/Engine/Source/Runtime/Engine/Private/Components/SplineComponent.cpp
+++ b/Engine/Source/Runtime/Engine/Private/Components/SplineComponent.cpp
@@ -231,7 +231,7 @@ float USplineComponent::GetSegmentLength(const int32 Index, const float Param) c
 	{
 		// Calculate derivative at each Legendre-Gauss sample, and perform a weighted sum
 		const float Alpha = HalfParam * (1.0f + LegendreGaussCoefficient.Abscissa);
-		const FVector Derivative = (Coeff1 * Alpha + Coeff2) * Alpha + Coeff3;
+		const FVector Derivative = ((Coeff1 * Alpha + Coeff2) * Alpha + Coeff3) * ComponentToWorld.GetScale3D();
 		Length += Derivative.Size() * LegendreGaussCoefficient.Weight;
 	}
 	Length *= HalfParam;
