Remove names fixup from UUserDefinedEnum::PostLoad

diff --git a/Engine/Source/Runtime/Engine/Private/UserDefinedEnum.cpp b/Engine/Source/Runtime/Engine/Private/UserDefinedEnum.cpp
--- a/Engine/Source/Runtime/Engine/Private/UserDefinedEnum.cpp
+++ b/Engine/Source/Runtime/Engine/Private/UserDefinedEnum.cpp
@@ -69,11 +69,6 @@ void UUserDefinedEnum::PostLoad()
 	Super::PostLoad();
 	FEnumEditorUtils::UpdateAfterPathChanged(this);
 	FEnumEditorUtils::EnsureAllDisplayNamesExist(this);
-
-	for (int32 i = 0; i < Names.Num(); ++i)
-	{
-		Names[i].Value = i;
-	}
 }
 
 void UUserDefinedEnum::PostEditUndo()
