Movie Scene Rendering improvements

 - All movie scene capture responsiblities are now encapsulated inside the MovieSceneCapture module. Corresponding logic has been moved out of engine classes.
 - Fixed time-step capture is fullly supported on PC and mac, real-time capture is still experimental.
 - Matinee now goes through the same MovieSceneCapture module logic as actor animations
 - AVIWriter can now be instantiated (used to be a single static w/ globals) to give greater control of capturing.
 - This mostly addresses UETOOL-446

diff --git a/Engine/Source/Editor/DetailCustomizations/DetailCustomizations.Build.cs b/Engine/Source/Editor/DetailCustomizations/DetailCustomizations.Build.cs
--- a/Engine/Source/Editor/DetailCustomizations/DetailCustomizations.Build.cs
+++ b/Engine/Source/Editor/DetailCustomizations/DetailCustomizations.Build.cs
@@ -23,6 +23,7 @@ public class DetailCustomizations : ModuleRules
 				"UnrealEd",
 				"EditorWidgets",
 				"KismetWidgets",
+				"MovieSceneCapture",
 				"SharedSettingsWidgets",
                 "ContentBrowser",
 				"BlueprintGraph",
diff --git a/Engine/Source/Editor/DetailCustomizations/Private/CaptureResolutionCustomization.cpp b/Engine/Source/Editor/DetailCustomizations/Private/CaptureResolutionCustomization.cpp
new file mode 100644
--- /dev/null
+++ b/Engine/Source/Editor/DetailCustomizations/Private/CaptureResolutionCustomization.cpp
@@ -0,0 +1,160 @@
+﻿// Copyright 1998-2015 Epic Games, Inc. All Rights Reserved.
+
+#include "DetailCustomizationsPrivatePCH.h"
+#include "CaptureResolutionCustomization.h"
+#include "MovieSceneCaptureSettings.h"
+#include "SEditableText.h"
+
+#define LOCTEXT_NAMESPACE "CaptureResolutionCustomization"
+
+TSharedRef<IPropertyTypeCustomization> FCaptureResolutionCustomization::MakeInstance()
+{
+	return MakeShareable( new FCaptureResolutionCustomization );
+}
+
+void FCaptureResolutionCustomization::CustomizeHeader( TSharedRef<IPropertyHandle> InPropertyHandle, FDetailWidgetRow& HeaderRow, IPropertyTypeCustomizationUtils& CustomizationUtils )
+{
+	PropertyHandle = InPropertyHandle;
+
+	ResXHandle = PropertyHandle->GetChildHandle(GET_MEMBER_NAME_CHECKED(FCaptureResolution, ResX));
+	ResYHandle = PropertyHandle->GetChildHandle(GET_MEMBER_NAME_CHECKED(FCaptureResolution, ResY));
+
+	// Resolution
+	Resolutions.Add(MakeShareable( new FPredefinedResolution{LOCTEXT("ResolutionA", "320 x 240 (4:3)"), 320, 240} ));
+	Resolutions.Add(MakeShareable( new FPredefinedResolution{LOCTEXT("ResolutionB", "640 x 480 (4:3)"), 640, 480} ));
+	Resolutions.Add(MakeShareable( new FPredefinedResolution{LOCTEXT("ResolutionC", "640 x 360 (16:9)"), 640, 360} ));
+	Resolutions.Add(MakeShareable( new FPredefinedResolution{LOCTEXT("ResolutionD", "1280 x 720 (16:9)"), 1280, 720} ));
+	Resolutions.Add(MakeShareable( new FPredefinedResolution{LOCTEXT("ResolutionE", "1920 x 1080 (16:9)"), 1920, 1080} ));
+	Resolutions.Add(MakeShareable( new FPredefinedResolution{LOCTEXT("ResolutionF", "Custom"), 1920, 1080} ));
+
+	uint32 CurrentResX = 0, CurrentResY = 0;
+	ResXHandle->GetValue(CurrentResX);
+	ResYHandle->GetValue(CurrentResY);
+
+	CurrentIndex = Resolutions.IndexOfByPredicate([&](const TSharedPtr<FPredefinedResolution>& In){
+		return In->ResX == CurrentResX && In->ResY == CurrentResY;
+	});
+
+	if (CurrentIndex == INDEX_NONE)
+	{
+		CurrentIndex = Resolutions.Num() - 1;
+	}
+
+	HeaderRow
+	.NameContent()
+	[
+		PropertyHandle->CreatePropertyNameWidget()
+	]
+	.ValueContent()
+	.HAlign(HAlign_Fill)
+	.MaxDesiredWidth(TOptional<float>())
+	[
+		SNew(SVerticalBox)
+
+		+ SVerticalBox::Slot()
+		.HAlign(HAlign_Left)
+		.AutoHeight()
+		[
+			SNew(SComboBox<TSharedPtr<FPredefinedResolution>>)
+			.OptionsSource(&Resolutions)
+			.OnSelectionChanged_Lambda([&](TSharedPtr<FPredefinedResolution> Resolution, ESelectInfo::Type){
+				
+				CurrentIndex = Resolutions.IndexOfByPredicate([&](const TSharedPtr<FPredefinedResolution>& In){
+					return In == Resolution;
+				});
+				if (CurrentIndex == INDEX_NONE)
+				{
+					CurrentIndex = Resolutions.Num() - 1;
+				}
+
+				UpdateProperty();
+			})
+			.OnGenerateWidget_Lambda([&](TSharedPtr<FPredefinedResolution> Resolution){
+				return SNew(STextBlock)
+				.Font(IDetailLayoutBuilder::GetDetailFont())
+				.Text(Resolution->DisplayName);
+			})
+			.InitiallySelectedItem(Resolutions[CurrentIndex])
+			[
+				SAssignNew(CurrentText, STextBlock)
+				.Font(IDetailLayoutBuilder::GetDetailFont())
+				.Text(Resolutions[CurrentIndex]->DisplayName)
+			]
+		]
+
+		+ SVerticalBox::Slot()
+		.Padding(FMargin(0,4.f))
+		.AutoHeight()
+		[
+			SAssignNew(CustomSliders, SHorizontalBox)
+			.Visibility(CurrentIndex == Resolutions.Num() - 1 ? EVisibility::Visible : EVisibility::Collapsed)
+
+			+ SHorizontalBox::Slot()
+			[
+				SNew(SHorizontalBox)
+
+				+ SHorizontalBox::Slot()
+				.VAlign(VAlign_Center)
+				.Padding(FMargin(0,0,4.f,0))
+				.AutoWidth()
+				[
+					SNew(STextBlock)
+					.Font(IDetailLayoutBuilder::GetDetailFont())
+					.Text(LOCTEXT("Width", "Width"))
+				]
+				+ SHorizontalBox::Slot()
+				.Padding(FMargin(0,0,4.f,0))
+				[
+					ResXHandle->CreatePropertyValueWidget()
+				]
+			]
+
+			+ SHorizontalBox::Slot()
+			[
+				SNew(SHorizontalBox)
+
+				+ SHorizontalBox::Slot()
+				.VAlign(VAlign_Center)
+				.Padding(FMargin(0,0,4.f,0))
+				.AutoWidth()
+				[
+					SNew(STextBlock)
+					.Font(IDetailLayoutBuilder::GetDetailFont())
+					.Text(LOCTEXT("Height", "Height"))
+				]
+				+ SHorizontalBox::Slot()
+				.Padding(FMargin(0,0,4.f,0))
+				[
+					ResYHandle->CreatePropertyValueWidget()
+				]
+			]
+		]
+	];
+}
+
+void FCaptureResolutionCustomization::CustomizeChildren( TSharedRef<IPropertyHandle> InPropertyHandle, IDetailChildrenBuilder& ChildBuilder, IPropertyTypeCustomizationUtils& CustomizationUtils )
+{
+
+}
+
+void FCaptureResolutionCustomization::UpdateProperty()
+{
+	const TSharedPtr<FPredefinedResolution>& Resolution = Resolutions[CurrentIndex];
+
+	if (CurrentIndex == Resolutions.Num() - 1)
+	{
+		// Show custom stuff
+		CustomSliders->SetVisibility(EVisibility::Visible);
+	}
+	else
+	{
+		// Hide custom Stuff
+		CustomSliders->SetVisibility(EVisibility::Collapsed);
+		ResXHandle->SetValue(Resolution->ResX);
+		ResYHandle->SetValue(Resolution->ResY);
+	}
+
+	CurrentText->SetText(Resolution->DisplayName);
+}
+
+#undef LOCTEXT_NAMESPACE
diff --git a/Engine/Source/Editor/DetailCustomizations/Private/CaptureResolutionCustomization.h b/Engine/Source/Editor/DetailCustomizations/Private/CaptureResolutionCustomization.h
new file mode 100644
--- /dev/null
+++ b/Engine/Source/Editor/DetailCustomizations/Private/CaptureResolutionCustomization.h
@@ -0,0 +1,40 @@
+// Copyright 1998-2015 Epic Games, Inc. All Rights Reserved.
+
+#pragma once
+
+class FCaptureResolutionCustomization : public IPropertyTypeCustomization
+{
+public:
+	/** Makes a new instance of this detail layout class for a specific detail view requesting it */
+	static TSharedRef<IPropertyTypeCustomization> MakeInstance();
+
+	/** IDetailCustomization interface */
+	virtual void CustomizeHeader( TSharedRef<IPropertyHandle> PropertyHandle, FDetailWidgetRow& HeaderRow, IPropertyTypeCustomizationUtils& CustomizationUtils ) override;
+	virtual void CustomizeChildren( TSharedRef<IPropertyHandle> PropertyHandle, IDetailChildrenBuilder& ChildBuilder, IPropertyTypeCustomizationUtils& CustomizationUtils ) override;
+
+private:
+
+	/** Respond to a selection change event from the combo box */
+	void UpdateProperty();
+	
+private:
+
+	/** Array of predefined resolutions */
+	struct FPredefinedResolution
+	{
+		FText DisplayName;
+		uint32 ResX, ResY;
+	};
+	TArray<TSharedPtr<FPredefinedResolution>> Resolutions;
+
+	/** The index of the resolution we're currently displaying */
+	int32 CurrentIndex;
+	/** The text of the current selection */
+	TSharedPtr<STextBlock> CurrentText;
+	/** The custom sliders to be hidden and shown based on combo box selection */
+	TSharedPtr<SWidget> CustomSliders;
+
+	/** Property handles of the properties we're editing */
+	TSharedPtr<IPropertyHandle> PropertyHandle;
+	TSharedPtr<IPropertyHandle> ResXHandle, ResYHandle;
+};
diff --git a/Engine/Source/Editor/DetailCustomizations/Private/DetailCustomizations.cpp b/Engine/Source/Editor/DetailCustomizations/Private/DetailCustomizations.cpp
--- a/Engine/Source/Editor/DetailCustomizations/Private/DetailCustomizations.cpp
+++ b/Engine/Source/Editor/DetailCustomizations/Private/DetailCustomizations.cpp
@@ -103,6 +103,7 @@
 #include "PostProcessSettingsCustomization.h"
 #include "ConfigEditorPropertyDetails.h"
 #include "AssetImportDataCustomization.h"
+#include "CaptureResolutionCustomization.h"
 #include "TextCustomization.h"
 #include "AnimTrailNodeDetails.h"
 
@@ -211,6 +212,7 @@ void FDetailCustomizationsModule::RegisterPropertyTypeCustomizations()
 	RegisterCustomPropertyTypeLayout("HierarchicalSimplification", FOnGetPropertyTypeCustomizationInstance::CreateStatic(&FHierarchicalSimplificationCustomizations::MakeInstance));
 	RegisterCustomPropertyTypeLayout("PostProcessSettings", FOnGetPropertyTypeCustomizationInstance::CreateStatic(&FPostProcessSettingsCustomization::MakeInstance));
 	RegisterCustomPropertyTypeLayout("AssetImportInfo", FOnGetPropertyTypeCustomizationInstance::CreateStatic(&FAssetImportDataCustomization::MakeInstance));
+	RegisterCustomPropertyTypeLayout("CaptureResolution", FOnGetPropertyTypeCustomizationInstance::CreateStatic(&FCaptureResolutionCustomization::MakeInstance));
 	RegisterCustomPropertyTypeLayout("WeightedBlendable", FOnGetPropertyTypeCustomizationInstance::CreateStatic(&FWeightedBlendableCustomization::MakeInstance));
 }
 
diff --git a/Engine/Source/Editor/EditorStyle/Private/SlateEditorStyle.cpp b/Engine/Source/Editor/EditorStyle/Private/SlateEditorStyle.cpp
--- a/Engine/Source/Editor/EditorStyle/Private/SlateEditorStyle.cpp
+++ b/Engine/Source/Editor/EditorStyle/Private/SlateEditorStyle.cpp
@@ -1618,6 +1618,7 @@ void FSlateEditorStyle::FStyle::SetupGeneralStyles()
 		Set( "Sequencer.MoveTool.Small", new IMAGE_BRUSH( "Icons/SequencerIcons/icon_Sequencer_Move_24x", Icon16x16 ) );
 		Set( "Sequencer.MarqueeTool_Keys.Small", new IMAGE_BRUSH( "Sequencer/MarqueeTool_Keys_16px", Icon16x16 ) ); //@todo
 		Set( "Sequencer.MarqueeTool_Sections.Small", new IMAGE_BRUSH( "Sequencer/MarqueeTool_Sections_16px", Icon16x16 ) ); //@todo
+		Set( "Sequencer.RenderMovie.Small", new IMAGE_BRUSH( "Sequencer/RenderMovie_16px", Icon16x16 ) );
 		Set( "Sequencer.LockCamera", new IMAGE_BRUSH( "Icons/icon_ToggleActorPilotCameraView_16x", Icon16x16 ) );
 		Set( "Sequencer.UnlockCamera", new IMAGE_BRUSH( "Icons/icon_ToggleActorPilotCameraView_16x", Icon16x16, FLinearColor(1.f, 1.f, 1.f, 0.5f) ) );
 		Set( "Sequencer.Section.SelectionBorder", new BORDER_BRUSH( TEXT("Sequencer/SectionHighlight"), FMargin(7.f/16.f) ) );
diff --git a/Engine/Source/Editor/Matinee/Matinee.Build.cs b/Engine/Source/Editor/Matinee/Matinee.Build.cs
--- a/Engine/Source/Editor/Matinee/Matinee.Build.cs
+++ b/Engine/Source/Editor/Matinee/Matinee.Build.cs
@@ -40,6 +40,7 @@ public class Matinee : ModuleRules
 				"RenderCore",
 				"AssetRegistry",
 				"ContentBrowser",
+				"MovieSceneCapture",
 				"BlueprintGraph"
 			}
 			);
diff --git a/Engine/Source/Editor/Matinee/Private/Matinee.cpp b/Engine/Source/Editor/Matinee/Private/Matinee.cpp
--- a/Engine/Source/Editor/Matinee/Private/Matinee.cpp
+++ b/Engine/Source/Editor/Matinee/Private/Matinee.cpp
@@ -30,6 +30,7 @@
 #include "SubtitleManager.h"
 #include "InterpolationHitProxy.h"
 
+#include "LevelEditor.h"
 #include "LevelEditorActions.h"
 #include "EditorSupportDelegates.h"
 
@@ -50,6 +51,8 @@
 #include "Camera/CameraActor.h"
 #include "Camera/CameraAnim.h"
 
+#include "MovieSceneCaptureModule.h"
+#include "LevelCapture.h"
 
 DEFINE_LOG_CATEGORY(LogSlateMatinee);
 
@@ -1483,12 +1486,6 @@ void FMatinee::StopPlaying()
 	UpdateFixedTimeStepPlayback();
 }
 
-/** Starts recording the current sequence */
-void FMatinee::StartRecordingMovie()
-{
-	GUnrealEd->PlayMap(NULL, NULL, 0, -1, false, true);
-}
-
 void FMatinee::OnPostUndoRedo(FUndoSessionContext SessionContext, bool Succeeded)
 {
 	InvalidateTrackWindowViewports();
@@ -2972,4 +2969,22 @@ bool FMatinee::IsCameraAnim() const
 	return MatineeActor->IsA(AMatineeActorCameraAnim::StaticClass());
 }
 
+void FMatinee::OnMenuCreateMovie()
+{
+	FLevelEditorModule& LevelEditorModule = FModuleManager::GetModuleChecked<FLevelEditorModule>(TEXT("LevelEditor"));
+	IMovieSceneCaptureModule& MovieSceneCaptureModule = IMovieSceneCaptureModule::Get();
+
+	// Create a new movie scene capture object for a generic level capture
+	ULevelCapture* MovieSceneCapture = NewObject<ULevelCapture>(GetTransientPackage(), ULevelCapture::StaticClass(), NAME_None, RF_Transient);
+	MovieSceneCapture->LoadConfig();
+
+	FEdModeInterpEdit* Mode = (FEdModeInterpEdit*)GLevelEditorModeTools().GetActiveMode( FBuiltinEditorModes::EM_InterpEdit );
+	if (Mode && Mode->InterpEd)
+	{
+		MovieSceneCapture->Level = Mode->InterpEd->GetMatineeActor()->GetOutermost()->GetName();
+	}
+
+	MovieSceneCaptureModule.OpenCaptureSettings(LevelEditorModule.GetLevelEditorTabManager().ToSharedRef(), MovieSceneCapture);
+}
+
 #undef LOCTEXT_NAMESPACE
diff --git a/Engine/Source/Editor/Matinee/Private/Matinee.h b/Engine/Source/Editor/Matinee/Private/Matinee.h
--- a/Engine/Source/Editor/Matinee/Private/Matinee.h
+++ b/Engine/Source/Editor/Matinee/Private/Matinee.h
@@ -97,9 +97,6 @@ public:
 	/** Stops playing the current sequence. */
 	void StopPlaying() override;
 
-	/** Starts recording the current sequence */
-	void StartRecordingMovie() override;
-
 	/** Handle undo redo events */
 	void OnPostUndoRedo(FUndoSessionContext SessionContext, bool Succeeded);
 
diff --git a/Engine/Source/Editor/Matinee/Private/MatineeRecordMovie.cpp b/Engine/Source/Editor/Matinee/Private/MatineeRecordMovie.cpp
deleted file mode 100644
diff --git a/Engine/Source/Editor/Matinee/Public/IMatinee.h b/Engine/Source/Editor/Matinee/Public/IMatinee.h
--- a/Engine/Source/Editor/Matinee/Public/IMatinee.h
+++ b/Engine/Source/Editor/Matinee/Public/IMatinee.h
@@ -55,7 +55,6 @@ public:
 	virtual void StopPlaying() = 0;
 	virtual void ResumePlaying() = 0;
 
-	virtual void StartRecordingMovie() = 0;
 	virtual bool Hide3DTrackView() = 0;
 	virtual void ToggleRecordMenuDisplay() = 0;
 	virtual void ToggleRecordInterpValues() = 0;
diff --git a/Engine/Source/Editor/Sequencer/Private/SSequencer.cpp b/Engine/Source/Editor/Sequencer/Private/SSequencer.cpp
--- a/Engine/Source/Editor/Sequencer/Private/SSequencer.cpp
+++ b/Engine/Source/Editor/Sequencer/Private/SSequencer.cpp
@@ -748,6 +748,17 @@ TSharedRef<SWidget> SSequencer::MakeToolBar()
 	{
 		ToolBarBuilder.AddToolBarButton( FSequencerCommands::Get().ToggleShowCurveEditor );
 	}
+	ToolBarBuilder.EndSection();
+
+	if( Sequencer.Pin()->IsLevelEditorSequencer() )
+	{
+		ToolBarBuilder.BeginSection("Render Movie");
+
+		ToolBarBuilder.AddToolBarButton( FSequencerCommands::Get().RenderMovie );
+
+		ToolBarBuilder.EndSection();
+	}
+	
 
 	return ToolBarBuilder.MakeWidget();
 }
diff --git a/Engine/Source/Editor/Sequencer/Private/Sequencer.cpp b/Engine/Source/Editor/Sequencer/Private/Sequencer.cpp
--- a/Engine/Source/Editor/Sequencer/Private/Sequencer.cpp
+++ b/Engine/Source/Editor/Sequencer/Private/Sequencer.cpp
@@ -39,8 +39,11 @@
 #include "GenericCommands.h"
 #include "Engine/BlueprintGeneratedClass.h"
 #include "Engine/Selection.h"
+#include "LevelEditor.h"
 #include "IMenu.h"
 
+#include "MovieSceneCaptureModule.h"
+#include "AutomatedActorAnimationCapture.h"
 
 #define LOCTEXT_NAMESPACE "Sequencer"
 
@@ -2279,6 +2282,31 @@ void FSequencer::BindSequencerCommands()
 		FCanExecuteAction::CreateLambda( []{ return true; } ),
 		FIsActionChecked::CreateLambda( [this]{ return Settings->GetShowCurveEditor(); } ) );
 
+	SequencerCommandBindings->MapAction( Commands.RenderMovie,
+		FExecuteAction::CreateLambda([this]{
+			FLevelEditorModule& LevelEditorModule = FModuleManager::GetModuleChecked<FLevelEditorModule>(TEXT("LevelEditor"));
+			IMovieSceneCaptureModule& MovieSceneCaptureModule = IMovieSceneCaptureModule::Get();
+
+			// Create a new movie scene capture object for an automated actor animation, and open the tab
+			UAutomatedActorAnimationCapture* MovieSceneCapture = NewObject<UAutomatedActorAnimationCapture>(GetTransientPackage(), UAutomatedActorAnimationCapture::StaticClass(), NAME_None, RF_Transient);
+			MovieSceneCapture->LoadConfig();
+
+			// Set the actor animation asset we want to render
+			MovieSceneCapture->ActorAnimation = GetCurrentAsset()->GetPathName();
+
+			// Set the level we want to render
+			for (const FWorldContext& Context : GEngine->GetWorldContexts())
+			{
+				if (Context.WorldType == EWorldType::Editor)
+				{
+					MovieSceneCapture->Level = Context.World()->GetOutermost()->GetPathName();
+					break;
+				}
+			}
+			MovieSceneCaptureModule.OpenCaptureSettings(LevelEditorModule.GetLevelEditorTabManager().ToSharedRef(), MovieSceneCapture);
+		})
+	);
+
 	SequencerWidget->BindCommands(SequencerCommandBindings);
 
 	for (int32 i = 0; i < TrackEditors.Num(); ++i)
diff --git a/Engine/Source/Editor/Sequencer/Private/SequencerCommands.cpp b/Engine/Source/Editor/Sequencer/Private/SequencerCommands.cpp
--- a/Engine/Source/Editor/Sequencer/Private/SequencerCommands.cpp
+++ b/Engine/Source/Editor/Sequencer/Private/SequencerCommands.cpp
@@ -62,6 +62,8 @@ void FSequencerCommands::RegisterCommands()
 
 	UI_COMMAND( MoveTool, "Move", "Activates the move tool", EUserInterfaceActionType::ToggleButton, FInputChord( EKeys::V ) );
 	UI_COMMAND( MarqueeTool, "Marquee", "Activates the marquee selection tool", EUserInterfaceActionType::ToggleButton, FInputChord( EKeys::M ) );
+
+	UI_COMMAND( RenderMovie, "Render Movie", "Render this movie to a video, or image frame sequence", EUserInterfaceActionType::Button, FInputChord() );
 }
 
 #undef LOCTEXT_NAMESPACE
\ No newline at end of file
diff --git a/Engine/Source/Editor/Sequencer/Private/SequencerCommands.h b/Engine/Source/Editor/Sequencer/Private/SequencerCommands.h
--- a/Engine/Source/Editor/Sequencer/Private/SequencerCommands.h
+++ b/Engine/Source/Editor/Sequencer/Private/SequencerCommands.h
@@ -135,9 +135,15 @@ public:
 	/** Toggles whether or not the curve editor should be shown. */
 	TSharedPtr< FUICommandInfo > ToggleShowCurveEditor;
 
+	/** Enable the move tool */
 	TSharedPtr< FUICommandInfo > MoveTool;
+
+	/** Enable the marquee selection tool */
 	TSharedPtr< FUICommandInfo > MarqueeTool;
 
+	/** Open a panel that enables exporting the sequence to a movie */
+	TSharedPtr< FUICommandInfo > RenderMovie;
+
 	/**
 	 * Initialize commands
 	 */
diff --git a/Engine/Source/Editor/Sequencer/Sequencer.Build.cs b/Engine/Source/Editor/Sequencer/Sequencer.Build.cs
--- a/Engine/Source/Editor/Sequencer/Sequencer.Build.cs
+++ b/Engine/Source/Editor/Sequencer/Sequencer.Build.cs
@@ -21,6 +21,7 @@ public class Sequencer : ModuleRules
 				"MovieScene", 
 				"MovieSceneTracks", 
 				"MovieSceneTools", 
+				"MovieSceneCapture", 
 				"EditorWidgets", 
 				"SequencerWidgets",
 				"BlueprintGraph",
@@ -37,6 +38,7 @@ public class Sequencer : ModuleRules
 		DynamicallyLoadedModuleNames.AddRange(
 			new string[] {
 				"PropertyEditor",
+				"LevelEditor",
 				"WorkspaceMenuStructure",
 			}
 		);
diff --git a/Engine/Source/Editor/UnrealEd/Classes/Editor/EditorEngine.h b/Engine/Source/Editor/UnrealEd/Classes/Editor/EditorEngine.h
--- a/Engine/Source/Editor/UnrealEd/Classes/Editor/EditorEngine.h
+++ b/Engine/Source/Editor/UnrealEd/Classes/Editor/EditorEngine.h
@@ -469,10 +469,6 @@ public:
 	UPROPERTY()
 	TArray<TWeakObjectPtr<class AActor> > ActorsThatWereSelected;
 
-	/** True to start movie capturing right away when launching the game from the editor on PC platform */
-	UPROPERTY()
-	uint32 bStartMovieCapture:1;
-
 	/** Where did the person want to play? Where to play the game - -1 means in editor, 0 or more is an index into the GConsoleSupportContainer	*/
 	UPROPERTY()
 	int32 PlayWorldDestination;
@@ -507,14 +503,6 @@ public:
 	UPROPERTY()
 	int32 PlayInEditorViewportIndex;
 
-	/** The width resolution that we want to use for the matinee capture */
-	UPROPERTY()
-	int32 MatineeCaptureResolutionX;
-
-	/** The height resolution that we want to use for the matinee capture */
-	UPROPERTY()
-	int32 MatineeCaptureResolutionY;
-
 	/** Play world url string edited by a user. */
 	UPROPERTY()
 	FString UserEditedPlayWorldURL;
@@ -1521,9 +1509,8 @@ public:
 	 * @param	DestinationConsole		Where to play the game - -1 means in editor, 0 or more is an index into the GConsoleSupportContainer
 	 * @param	InPlayInViewportIndex	Viewport index to play the game in, or -1 to spawn a standalone PIE window
 	 * @param	bUseMobilePreview		True to enable mobile preview mode (PC platform only)
-	 * @param	bMovieCapture			True to start with movie capture recording (PC platform only)
 	 */
-	virtual void PlayMap( const FVector* StartLocation = NULL, const FRotator* StartRotation = NULL, int32 DestinationConsole = -1, int32 InPlayInViewportIndex = -1, bool bUseMobilePreview = false, bool bMovieCapture = false );
+	virtual void PlayMap( const FVector* StartLocation = NULL, const FRotator* StartRotation = NULL, int32 DestinationConsole = -1, int32 InPlayInViewportIndex = -1, bool bUseMobilePreview = false );
 
 
 
@@ -1639,11 +1626,6 @@ public:
 
 	virtual UGameInstance* CreatePIEGameInstance(int32 PIEInstance, bool bInSimulateInEditor, bool bAnyBlueprintErrors, bool bStartInSpectatorMode, bool bPlayNetDedicated, float PIEStartTime);
 
-	/**
-	 * Launches the game in movie capture mode
-	 */
-	virtual void PlayForMovieCapture();
-
 	/**
 	 * Kills the Play From Here session
 	 */
diff --git a/Engine/Source/Editor/UnrealEd/Private/PlayLevel.cpp b/Engine/Source/Editor/UnrealEd/Private/PlayLevel.cpp
--- a/Engine/Source/Editor/UnrealEd/Private/PlayLevel.cpp
+++ b/Engine/Source/Editor/UnrealEd/Private/PlayLevel.cpp
@@ -642,7 +642,7 @@ void UEditorEngine::TeardownPlaySession(FWorldContext &PieWorldContext)
 	}
 }
 
-void UEditorEngine::PlayMap( const FVector* StartLocation, const FRotator* StartRotation, int32 Destination, int32 InPlayInViewportIndex, bool bUseMobilePreview, bool bMovieCapture )
+void UEditorEngine::PlayMap( const FVector* StartLocation, const FRotator* StartRotation, int32 Destination, int32 InPlayInViewportIndex, bool bUseMobilePreview )
 {
 	// queue up a Play From Here request, this way the load/save won't conflict with the TransBuffer, which doesn't like 
 	// loading and saving to happen during a transaction
@@ -666,9 +666,6 @@ void UEditorEngine::PlayMap( const FVector* StartLocation, const FRotator* Start
 	bUseMobilePreviewForPlayWorld = bUseMobilePreview;
 	bUseVRPreviewForPlayWorld = false;
 
-	// Set whether or not we want to start movie capturing immediately (PC platform only)
-	bStartMovieCapture = bMovieCapture;
-
 	// tell the editor to kick it off next Tick()
 	bIsPlayWorldQueued = true;
 
@@ -710,9 +707,6 @@ void UEditorEngine::RequestPlaySession( bool bAtPlayerStart, TSharedPtr<class IL
 
 	bUseVRPreviewForPlayWorld = bUseVRPreview;
 
-	// Not capturing a movie
-	bStartMovieCapture = false;
-
 	// tell the editor to kick it off next Tick()
 	bIsPlayWorldQueued = true;
 
@@ -1018,18 +1012,10 @@ void UEditorEngine::StartQueuedPlayMapRequest()
 		// If we're playing in the editor
 		if (!bPlayOnLocalPcSession)
 		{
-			if (bStartMovieCapture)
-			{
-				// @todo Fix for UE4.  This is a temp workaround. 
-				PlayForMovieCapture();
-			}
-			else
-			{
-				PlayInEditor(GetEditorWorldContext().World(), bWantSimulateInEditor);
+			PlayInEditor(GetEditorWorldContext().World(), bWantSimulateInEditor);
 
-				// Editor counts as a client
-				NumClients++;
-			}
+			// Editor counts as a client
+			NumClients++;
 		}
 
 		// Spawn number of clients
@@ -1050,11 +1036,6 @@ void UEditorEngine::StartQueuedPlayMapRequest()
 		{
 			PlayUsingLauncher();
 		}
-		else if (bStartMovieCapture)
-		{
-			// @todo Fix for UE4.  This is a temp workaround. 
-			PlayForMovieCapture();
-		}
 		else
 		{
 			PlayInEditor( GetEditorWorldContext().World(), bWantSimulateInEditor );
@@ -1922,139 +1903,6 @@ void UEditorEngine::PlayUsingLauncher()
 	}
 }
 
-void UEditorEngine::PlayForMovieCapture()
-{
-	TArray<FString> SavedMapNames;
-	SaveWorldForPlay(SavedMapNames);
-
-	if (SavedMapNames.Num() == 0)
-	{
-		return;
-	}
-
-	// this parameter tells UE4Editor to run in game mode
-	FString EditorCommandLine = SavedMapNames[0];
-	EditorCommandLine += TEXT(" -game");
-
-	// renderer overrides - hack
-	EditorCommandLine += FParse::Param(FCommandLine::Get(), TEXT("d3d11"))		?	TEXT(" -d3d11")		: TEXT("");
-	EditorCommandLine += FParse::Param(FCommandLine::Get(), TEXT("sm5"))		?	TEXT(" -sm5")		: TEXT("");
-	EditorCommandLine += FParse::Param(FCommandLine::Get(), TEXT("dx11"))		?	TEXT(" -dx11")		: TEXT("");
-	EditorCommandLine += FParse::Param(FCommandLine::Get(), TEXT("d3d10"))		?	TEXT(" -d3d10")		: TEXT("");
-	EditorCommandLine += FParse::Param(FCommandLine::Get(), TEXT("sm4"))		?	TEXT(" -sm4")		: TEXT("");
-	EditorCommandLine += FParse::Param(FCommandLine::Get(), TEXT("dx10"))		?	TEXT(" -dx10")		: TEXT("");
-	EditorCommandLine += FParse::Param(FCommandLine::Get(), TEXT("opengl"))		?	TEXT(" -opengl")	: TEXT("");
-	EditorCommandLine += FParse::Param(FCommandLine::Get(), TEXT("opengl3"))	?	TEXT(" -opengl3")	: TEXT("");
-	EditorCommandLine += FParse::Param(FCommandLine::Get(), TEXT("opengl4"))	?	TEXT(" -opengl4")	: TEXT("");
-
-
-	// this parameter tells UGameEngine to add the auto-save dir to the paths array and repopulate the package file cache
-	// this is needed in order to support streaming levels as the streaming level packages will be loaded only when needed (thus
-	// their package names need to be findable by the package file caching system)
-	// (we add to EditorCommandLine because the URL is ignored by WindowsTools)
-	EditorCommandLine += TEXT(" -PIEVIACONSOLE");
-	
-	// if we want to start movie capturing right away, then append the argument for that
-	if (bStartMovieCapture)
-	{
-		//disable movies
-		EditorCommandLine += FString::Printf(TEXT(" -nomovie"));
-	
-		//set res options
-		EditorCommandLine += FString::Printf(TEXT(" -ResX=%d"), GEditor->MatineeCaptureResolutionX);
-		EditorCommandLine += FString::Printf(TEXT(" -ResY=%d"), GEditor->MatineeCaptureResolutionY);
-					
-		if( GUnrealEd->MatineeScreenshotOptions.bNoTextureStreaming )
-		{
-			EditorCommandLine += TEXT(" -NoTextureStreaming");
-		}
-
-		//set fps
-		EditorCommandLine += FString::Printf(TEXT(" -BENCHMARK -FPS=%d"), GEditor->MatineeScreenshotOptions.MatineeCaptureFPS);
-	
-		if (GEditor->MatineeScreenshotOptions.MatineeCaptureType.GetValue() != EMatineeCaptureType::AVI)
-		{
-			EditorCommandLine += FString::Printf(TEXT(" -MATINEESSCAPTURE=%s"), *GEngine->MatineeScreenshotOptions.MatineeCaptureName);//*GEditor->MatineeNameForRecording);
-
-			switch(GEditor->MatineeScreenshotOptions.MatineeCaptureType.GetValue())
-			{
-			case EMatineeCaptureType::BMP:
-				EditorCommandLine += TEXT(" -MATINEESSFORMAT=BMP");
-				break;
-			case EMatineeCaptureType::PNG:
-				EditorCommandLine += TEXT(" -MATINEESSFORMAT=PNG");
-				break;
-			case EMatineeCaptureType::JPEG:
-				EditorCommandLine += TEXT(" -MATINEESSFORMAT=JPEG");
-				break;
-			default: break;
-			}
-
-			// If buffer visualization dumping is enabled, we need to tell capture process to enable it too
-			static const auto CVarDumpFrames = IConsoleManager::Get().FindTConsoleVariableDataInt(TEXT("r.BufferVisualizationDumpFrames"));
-
-			if (CVarDumpFrames && CVarDumpFrames->GetValueOnGameThread())
-			{
-				EditorCommandLine += TEXT(" -MATINEEBUFFERVISUALIZATIONDUMP");
-			}
-		}
-		else
-		{
-			EditorCommandLine += FString::Printf(TEXT(" -MATINEEAVICAPTURE=%s"), *GEngine->MatineeScreenshotOptions.MatineeCaptureName);//*GEditor->MatineeNameForRecording);
-		}
-					
-		EditorCommandLine += FString::Printf(TEXT(" -MATINEEPACKAGE=%s"), *GEngine->MatineeScreenshotOptions.MatineePackageCaptureName);//*GEditor->MatineePackageNameForRecording);
-	
-		if (GEditor->MatineeScreenshotOptions.bCompressMatineeCapture == 1)
-		{
-			EditorCommandLine += TEXT(" -CompressCapture");
-		}
-	}
-
-	FString GamePath = FPlatformProcess::GenerateApplicationPath(FApp::GetName(), FApp::GetBuildConfiguration());
-	FString Params;
-
-	if (FPaths::IsProjectFilePathSet())
-	{
-		Params = FString::Printf(TEXT("\"%s\" %s %s"), *FPaths::GetProjectFilePath(), *EditorCommandLine, *FCommandLine::GetSubprocessCommandline());
-	}
-	else
-	{
-		Params = FString::Printf(TEXT("%s %s %s"), FApp::GetGameName(), *EditorCommandLine, *FCommandLine::GetSubprocessCommandline());
-	}
-
-	if ( FRocketSupport::IsRocket() )
-	{
-		Params += TEXT(" -rocket");
-	}
-
-	bool bRunningDebug = FParse::Param(FCommandLine::Get(), TEXT("debug"));
-	if (bRunningDebug)
-	{
-		Params += TEXT(" -debug");
-	}
-
-	FProcHandle ProcessHandle = FPlatformProcess::CreateProc(*GamePath, *Params, true, false, false, NULL, 0, NULL, NULL);
-
-	if (ProcessHandle.IsValid())
-	{
-		bool bCloseEditor = false;
-
-		GConfig->GetBool(TEXT("MatineeCreateMovieOptions"), TEXT("CloseEditor"), bCloseEditor, GEditorPerProjectIni);
-
-		if (bCloseEditor)
-		{
-			IMainFrameModule& MainFrameModule = FModuleManager::LoadModuleChecked<IMainFrameModule>(TEXT("MainFrame"));
-			MainFrameModule.RequestCloseEditor();
-		}
-	}
-	else
-	{
-		UE_LOG(LogPlayLevel, Error,  TEXT("Failed to run a copy of the game for matinee capture."));
-	}
-	FPlatformProcess::CloseProc(ProcessHandle);
-}
-
 
 void UEditorEngine::RequestEndPlayMap()
 {
diff --git a/Engine/Source/Editor/UnrealEd/Private/UnrealEdEngine.cpp b/Engine/Source/Editor/UnrealEd/Private/UnrealEdEngine.cpp
--- a/Engine/Source/Editor/UnrealEd/Private/UnrealEdEngine.cpp
+++ b/Engine/Source/Editor/UnrealEd/Private/UnrealEdEngine.cpp
@@ -36,6 +36,7 @@
 #include "Components/BillboardComponent.h"
 #include "Components/ArrowComponent.h"
 #include "Engine/Selection.h"
+#include "IMovieSceneCapture.h"
 #include "EngineUtils.h"
 
 DEFINE_LOG_CATEGORY_STATIC(LogUnrealEdEngine, Log, All);
@@ -312,12 +313,6 @@ void UUnrealEdEngine::PreExit()
 	// Notify edit modes we're mode at exit
 	FEditorModeRegistry::Get().Shutdown();
 
-	FAVIWriter* AVIWriter = FAVIWriter::GetInstance();
-	if (AVIWriter)
-	{
-		AVIWriter->Close();
-	}
-
 	Super::PreExit();
 }
 
@@ -403,12 +398,6 @@ void UUnrealEdEngine::Tick(float DeltaSeconds, bool bIdleMode)
 
 	// Update lightmass
 	UpdateBuildLighting();
-
-	FAVIWriter* AVIWriter = FAVIWriter::GetInstance();
-	if (AVIWriter)
-	{
-		AVIWriter->Update(DeltaSeconds);
-	}
 	
 	ICrashTrackerModule* CrashTracker = FModuleManager::LoadModulePtr<ICrashTrackerModule>( FName("CrashTracker") );
 	bool bCrashTrackerEnabled = false;
diff --git a/Engine/Source/Editor/UnrealEd/UnrealEd.Build.cs b/Engine/Source/Editor/UnrealEd/UnrealEd.Build.cs
--- a/Engine/Source/Editor/UnrealEd/UnrealEd.Build.cs
+++ b/Engine/Source/Editor/UnrealEd/UnrealEd.Build.cs
@@ -42,6 +42,7 @@ public class UnrealEd : ModuleRules
                 "MergeActors",
 				"MeshUtilities",
 				"Messaging",
+				"MovieSceneCapture",
 				"NiagaraEditor",
 				"PlacementMode",
 				"Settings",
diff --git a/Engine/Source/Runtime/Engine/Classes/Engine/Engine.h b/Engine/Source/Runtime/Engine/Classes/Engine/Engine.h
--- a/Engine/Source/Runtime/Engine/Classes/Engine/Engine.h
+++ b/Engine/Source/Runtime/Engine/Classes/Engine/Engine.h
@@ -5,7 +5,6 @@
 #include "World.h"
 #include "Engine.generated.h"
 
-
 class FScreenSaverInhibitor;
 class UDeviceProfileManager;
 class FViewport;
@@ -456,19 +455,6 @@ struct FScreenMessageString
 };
 
 
-UENUM()
-namespace EMatineeCaptureType
-{
-	enum Type
-	{
-		AVI		UMETA(DisplayName="AVI Movie"),
-		BMP		UMETA(DisplayName="BMP Image Sequence"),
-		PNG		UMETA(DisplayName="PNG Image Sequence"),
-		JPEG	UMETA(DisplayName="JPEG Image Sequence")
-	};
-}
-
-
 USTRUCT()
 struct FGameNameRedirect
 {
@@ -538,42 +524,6 @@ class IAnalyticsProvider;
 DECLARE_DELEGATE_OneParam(FBeginStreamingPauseDelegate, FViewport*);
 DECLARE_DELEGATE(FEndStreamingPauseDelegate);
 
-USTRUCT()
-struct FMatineeScreenshotOptions
-{
-	GENERATED_USTRUCT_BODY()
-
-	/** determines if we should start the matinee capture as soon as the game loads */
-	UPROPERTY(transient)
-	uint32 bStartWithMatineeCapture:1;
-
-	/** should we compress the capture */
-	UPROPERTY(transient)
-	uint32 bCompressMatineeCapture:1;
-
-	/** the name of the matine that we want to record */
-	UPROPERTY(transient)
-	FString MatineeCaptureName;
-
-	/** The package name where the matinee belongs to */
-	UPROPERTY(transient)
-	FString MatineePackageCaptureName;
-
-	/** the fps of the matine that we want to record */
-	UPROPERTY(transient)
-	int32 MatineeCaptureFPS;
-
-	/** The capture type, e.g. AVI or Screen Shots */
-	UPROPERTY(transient)
-	TEnumAsByte<EMatineeCaptureType::Type> MatineeCaptureType;
-
-	/** Whether or not to disable texture streaming during matinee movie capture */
-	UPROPERTY(transient)
-	bool bNoTextureStreaming;
-
-	UPROPERTY(transient)
-	bool bHideHud;
-};
 
 /**
  * Abstract base class of all Engine classes, responsible for management of systems critical to editor or game systems.
@@ -1400,8 +1350,6 @@ private:
 	int32 ScreenSaverInhibitorSemaphore;
 
 public:
-	UPROPERTY(transient)
-	FMatineeScreenshotOptions MatineeScreenshotOptions;
 
 	/** true if the the user cannot modify levels that are read only. */
 	UPROPERTY(transient)
@@ -2734,9 +2682,6 @@ public:
 	const UGameUserSettings* GetGameUserSettings() const;
 	UGameUserSettings* GetGameUserSettings();
 
-	/** Delegate handler for screenshots */
-	void HandleScreenshotCaptured(int32 Width, int32 Height, const TArray<FColor>& Colors);
-
 private:
 	void CreateGameUserSettings();
 
diff --git a/Engine/Source/Runtime/Engine/Classes/Engine/GameEngine.h b/Engine/Source/Runtime/Engine/Classes/Engine/GameEngine.h
--- a/Engine/Source/Runtime/Engine/Classes/Engine/GameEngine.h
+++ b/Engine/Source/Runtime/Engine/Classes/Engine/GameEngine.h
@@ -3,6 +3,7 @@
 #pragma once
 
 #include "Engine/Engine.h"
+#include "Runtime/MovieSceneCapture/Public/MovieSceneCaptureHandle.h"
 #include "GameEngine.generated.h"
 
 
@@ -15,10 +16,6 @@ class ENGINE_API UGameEngine
 {
 	GENERATED_UCLASS_BODY()
 
-	/** Check to see if we need to start a movie capture (used on the first tick when we want to record a Matinee) */
-	UPROPERTY(transient)
-	uint32 bCheckForMovieCapture:1;
-
 	/** Maximium delta time the engine uses to populate FApp::DeltaTime. If 0, unbound. */
 	UPROPERTY(config)
 	float MaxDeltaTime;
@@ -41,6 +38,9 @@ public:
 	
 	/**
 	 * Creates the viewport widget where the games Slate UI is added to.
+	 *
+	 * @param GameViewportClient	The viewport client to use in the game
+	 * @param MovieCapture			Optional Movie capture implementation for this viewport
 	 */
 	void CreateGameViewportWidget( UGameViewportClient* GameViewportClient );
 
@@ -48,6 +48,7 @@ public:
 	 * Creates the game viewport
 	 *
 	 * @param GameViewportClient	The viewport client to use in the game
+	 * @param MovieCapture			Optional Movie capture implementation for this viewport
 	 */
 	void CreateGameViewport( UGameViewportClient* GameViewportClient );
 	
@@ -149,6 +150,9 @@ protected:
 
 	const FSceneViewport* GetGameSceneViewport(UGameViewportClient* ViewportClient) const;
 
+	/** Handle to a movie capture implementation to create on startup */
+	FMovieSceneCaptureHandle StartupMovieCaptureHandle;
+
 private:
 
 	virtual void HandleNetworkFailure_NotifyGameInstance(UWorld* World, UNetDriver* NetDriver, ENetworkFailure::Type FailureType) override;
diff --git a/Engine/Source/Runtime/Engine/Classes/Engine/GameInstance.h b/Engine/Source/Runtime/Engine/Classes/Engine/GameInstance.h
--- a/Engine/Source/Runtime/Engine/Classes/Engine/GameInstance.h
+++ b/Engine/Source/Runtime/Engine/Classes/Engine/GameInstance.h
@@ -95,6 +95,10 @@ public:
 	bool InitializePIE(bool bAnyBlueprintErrors, int32 PIEInstance, bool bRunAsDedicated);
 
 	virtual bool StartPIEGameInstance(ULocalPlayer* LocalPlayer, bool bInSimulateInEditor, bool bAnyBlueprintErrors, bool bStartInSpectatorMode);
+
+	/** Initial map to load when starting the game instance. Overrides anything set on the command line */
+	FString InitialMapOverride;
+
 #endif
 
 	class UEngine* GetEngine() const;
diff --git a/Engine/Source/Runtime/Engine/Classes/GameFramework/GameMode.h b/Engine/Source/Runtime/Engine/Classes/GameFramework/GameMode.h
--- a/Engine/Source/Runtime/Engine/Classes/GameFramework/GameMode.h
+++ b/Engine/Source/Runtime/Engine/Classes/GameFramework/GameMode.h
@@ -314,7 +314,7 @@ public:
 	 *	@return	bool			true if we should turn on cinematic mode, 
 	 *							false if if we should not.
 	 */
-	bool ShouldStartInCinematicMode(bool& OutHidePlayer, bool& OutHideHud, bool& OutDisableMovement, bool& OutDisableTurning);
+	bool ShouldStartInCinematicMode(APlayerController* Player, bool& OutHidePlayer, bool& OutHideHud, bool& OutDisableMovement, bool& OutDisableTurning);
 
 	/**
 	 * Initialize the game.
diff --git a/Engine/Source/Runtime/Engine/Engine.Build.cs b/Engine/Source/Runtime/Engine/Engine.Build.cs
--- a/Engine/Source/Runtime/Engine/Engine.Build.cs
+++ b/Engine/Source/Runtime/Engine/Engine.Build.cs
@@ -29,7 +29,8 @@ public class Engine : ModuleRules
 				"ImageWrapper",
 				"HeadMountedDisplay",
 				"Advertising",
-				"NetworkReplayStreaming"
+				"NetworkReplayStreaming",
+				"MovieSceneCapture"
 			}
 		);
 
@@ -114,6 +115,7 @@ public class Engine : ModuleRules
 			new string[]
 			{
 				"MovieScene",
+				"MovieSceneCapture",
 				"MovieSceneTracks",
 				"HeadMountedDisplay",
 				"StreamingPauseRendering",
diff --git a/Engine/Source/Runtime/Engine/Private/AVIWriter.cpp b/Engine/Source/Runtime/Engine/Private/AVIWriter.cpp
--- a/Engine/Source/Runtime/Engine/Private/AVIWriter.cpp
+++ b/Engine/Source/Runtime/Engine/Private/AVIWriter.cpp
@@ -5,7 +5,6 @@
 =============================================================================*/
 #include "EnginePrivate.h"
 #include "AVIWriter.h"
-#include "CapturePin.h"
 
 DEFINE_LOG_CATEGORY_STATIC(LogAVIWriter, Log, All);
 
@@ -23,15 +22,13 @@ typedef TCHAR* PTCHAR;
 #include <initguid.h>
 #include "HideWindowsPlatformTypes.h"
 
+#include "CapturePin.h"
 #include "CaptureSource.h"
 #include "SlateCore.h"
 
 
 #define g_wszCapture    L"Capture Filter"
 
-/** An event for synchronizing direct show encoding with the game thread */
-FEvent* GCaptureSyncEvent = NULL;
-
 // Filter setup data
 const AMOVIESETUP_MEDIATYPE sudOpPinTypes =
 {
@@ -65,31 +62,102 @@ const AMOVIESETUP_PIN sudOutputPinDesktop =
 
 const AMOVIESETUP_FILTER sudPushSourceDesktop =
 {
-	&CLSID_CaptureSource,	// Filter CLSID
+	&CLSID_ViewportCaptureSource,	// Filter CLSID
 	g_wszCapture,			// String name
 	MERIT_DO_NOT_USE,       // Filter merit
 	1,                      // Number pins
 	&sudOutputPinDesktop    // Pin details
 };
 
-
-// List of class IDs and creator functions for the class factory. This
-// provides the link between the OLE entry point in the DLL and an object
-// being created. The class factory will call the static CreateInstance.
-// We provide a set of filters in this one DLL.
-
 CFactoryTemplate g_Templates[1] = 
 {
 	{ 
 		g_wszCapture,					// Name
-		&CLSID_CaptureSource,			// CLSID
-		FCaptureSource::CreateInstance, // Method to create an instance of MyComponent
-		NULL,                           // Initialization function
+		nullptr,						// CLSID
+		nullptr, 						// Method to create an instance of MyComponent
+		nullptr,                        // Initialization function
 		&sudPushSourceDesktop           // Set-up information (for filters)
 	},
 };
+int32 g_cTemplates = 0;
+
+/** Find a pin on the specified filter that matches the specified direction */
+IPin* GetPin(IBaseFilter* Filter, PIN_DIRECTION PinDir)
+{
+	IEnumPins* pEnumDownFilterPins = nullptr;
+	IPin* Pin = nullptr;
+
+	if (FAILED(Filter->EnumPins(&pEnumDownFilterPins)))
+	{
+		return nullptr;
+	}
+
+	while(S_OK == pEnumDownFilterPins->Next(1, &Pin, nullptr) )
+	{
+		PIN_DIRECTION ThisPinDir;
+		if (SUCCEEDED(Pin->QueryDirection(&ThisPinDir)) && PinDir == ThisPinDir)
+		{
+			pEnumDownFilterPins->Release();
+			return Pin;
+		}
+
+		Pin->Release();
+	}
+
+	pEnumDownFilterPins->Release();	
+	return nullptr;
+}
+
+IBaseFilter* FindEncodingFilter(const FString& Name)
+{
+	// Create an encoding filter
+	ICreateDevEnum* DeviceDenumerator = nullptr;
+	if (FAILED(CoCreateInstance(CLSID_SystemDeviceEnum, NULL, CLSCTX_INPROC, IID_ICreateDevEnum, (void**)&DeviceDenumerator)))
+	{
+		return nullptr;
+	}
+
+	IEnumMoniker* EnumIterator = nullptr;
+	if (DeviceDenumerator->CreateClassEnumerator(CLSID_VideoCompressorCategory, &EnumIterator, 0) != S_OK)
+	{
+		return nullptr;
+	}
+
+	IMoniker* Moniker = nullptr;
+	while (EnumIterator->Next(1, &Moniker, nullptr) == S_OK)
+	{
+		IPropertyBag* Properties = nullptr;
+		if (FAILED(Moniker->BindToStorage(0, 0, IID_IPropertyBag, (void**)&Properties)))
+		{
+			Moniker->Release();
+			continue;
+		}
 
-int32 g_cTemplates = sizeof(g_Templates) / sizeof(g_Templates[0]);    
+		bool bUseThisEncoder = false;
+
+		VARIANT varName;
+		VariantInit(&varName);
+		if (SUCCEEDED(Properties->Read(L"FriendlyName", &varName, 0)) && FCString::Stricmp(*Name, varName.bstrVal) == 0)
+		{
+			bUseThisEncoder = true;
+		}
+		VariantClear(&varName);
+
+		Properties->Release();
+
+		IBaseFilter* Filter = nullptr;
+		if (bUseThisEncoder && Moniker->BindToObject(nullptr, nullptr, IID_IBaseFilter, (void**)&Filter) == S_OK)
+		{
+			Filter->AddRef();
+			Moniker->Release();
+			return Filter;
+		}
+
+		Moniker->Release();
+	}
+
+	return nullptr;
+}
 
 /**
  * Windows implementation relying on DirectShow.
@@ -97,265 +165,205 @@ int32 g_cTemplates = sizeof(g_Templates) / sizeof(g_Templates[0]);
 class FAVIWriterWin : public FAVIWriter
 {
 public:
-	FAVIWriterWin( ) 
-		: Graph(NULL),
-		Control(NULL),
-		Capture(NULL),
-		CapturePin(NULL),
-		CaptionSource(NULL)
+	FAVIWriterWin( const FAVIWriterOptions& InOptions ) 
+		: FAVIWriter(InOptions)
+		, Graph(nullptr)
+		, Control(nullptr)
+		, Capture(nullptr)
+		, CaptureFilter(nullptr)
+		, EncodingFilter(nullptr)
 	{
-
 	};
 
 public:
 
-	void StartCapture(FViewport* Viewport , const FString& OutputPath /*=FString()*/)
+	void StartCapture(FViewport* Viewport)
 	{
-		GCaptureSyncEvent = FPlatformProcess::GetSynchEventFromPool();
-		if (!bCapturing)
+		if (bCapturing)
 		{
-			if (!Viewport)
-			{
-				Viewport = GEngine->GameViewport != NULL ? GEngine->GameViewport->Viewport : NULL;
-				if (!Viewport)
-				{
-					UE_LOG(LogAVIWriter, Error, TEXT( "ERROR - Could not get a valid viewport for capture!" ));
-					return;
-				}
-			}
+			return;
+		}
 
-			if ( Viewport->GetSizeXY().GetMin() <= 0)
+		if (!Viewport)
+		{
+			Viewport = GEngine->GameViewport != NULL ? GEngine->GameViewport->Viewport : NULL;
+			if (!Viewport)
 			{
-				UE_LOG(LogAVIWriter, Error, TEXT( "ERROR - Attempting to record a 0 sized viewport!" ));
+				UE_LOG(LogAVIWriter, Error, TEXT( "ERROR - Could not get a valid viewport for capture!" ));
 				return;
 			}
+		}
 
-			CaptureViewport = Viewport;
+		if ( Viewport->GetSizeXY().GetMin() <= 0)
+		{
+			UE_LOG(LogAVIWriter, Error, TEXT( "ERROR - Attempting to record a 0 sized viewport!" ));
+			return;
+		}
 
-			ViewportColorBuffer.AddUninitialized( Viewport->GetSizeXY().X * Viewport->GetSizeXY().Y );
-			
+		CaptureViewport = Viewport;
 
-			// Initialize the COM library.
-			if (!FWindowsPlatformMisc::CoInitialize()) 
-			{
-				UE_LOG(LogAVIWriter, Error, TEXT( "ERROR - Could not initialize COM library!" ));
-				return;
-			}
+		// Initialize the COM library.
+		if (!FWindowsPlatformMisc::CoInitialize()) 
+		{
+			UE_LOG(LogAVIWriter, Error, TEXT( "ERROR - Could not initialize COM library!" ));
+			return;
+		}
 
-			// Create the filter graph manager and query for interfaces.
-			HRESULT hr = CoCreateInstance(CLSID_FilterGraph, NULL, CLSCTX_INPROC_SERVER, IID_IGraphBuilder, (void **)&Graph);
-			if (FAILED(hr)) 
-			{
-				UE_LOG(LogAVIWriter, Error, TEXT( "ERROR - Could not create the Filter Graph Manager!" ));
-				FWindowsPlatformMisc::CoUninitialize();
-				return;
-			}
+		// Create the filter graph manager and query for interfaces.
+		HRESULT hr = CoCreateInstance(CLSID_FilterGraph, NULL, CLSCTX_INPROC_SERVER, IID_IGraphBuilder, (void **)&Graph);
+		if (FAILED(hr)) 
+		{
+			UE_LOG(LogAVIWriter, Error, TEXT( "ERROR - Could not create the Filter Graph Manager!" ));
+			FWindowsPlatformMisc::CoUninitialize();
+			return;
+		}
 
-			// Create the capture graph builder
-			hr = CoCreateInstance (CLSID_CaptureGraphBuilder2 , NULL, CLSCTX_INPROC, IID_ICaptureGraphBuilder2, (void **) &Capture);
-			if (FAILED(hr)) 
-			{
-				UE_LOG(LogAVIWriter, Error, TEXT( "ERROR - Could not create the Capture Graph Manager!" ));
-				FWindowsPlatformMisc::CoUninitialize();
-				return;
-			}
+		// Create the capture graph builder
+		hr = CoCreateInstance(CLSID_CaptureGraphBuilder2 , NULL, CLSCTX_INPROC, IID_ICaptureGraphBuilder2, (void **) &Capture);
+		if (FAILED(hr)) 
+		{
+			UE_LOG(LogAVIWriter, Error, TEXT( "ERROR - Could not create the Capture Graph Manager!" ));
+			FWindowsPlatformMisc::CoUninitialize();
+			return;
+		}
 
-			// Specify a filter graph for the capture graph builder to use
-			hr = Capture->SetFiltergraph(Graph);
-			if (FAILED(hr))
-			{
-				UE_LOG(LogAVIWriter, Error, TEXT( "ERROR - Failed to set capture filter graph!" ));
-				FWindowsPlatformMisc::CoUninitialize();
-				return;
-			}
+		// Specify a filter graph for the capture graph builder to use
+		hr = Capture->SetFiltergraph(Graph);
+		if (FAILED(hr))
+		{
+			UE_LOG(LogAVIWriter, Error, TEXT( "ERROR - Failed to set capture filter graph!" ));
+			FWindowsPlatformMisc::CoUninitialize();
+			return;
+		}
 
-			CaptionSource = new FCaptureSource(NULL, &hr);
-			if (FAILED(hr)) 
-			{
-				UE_LOG(LogAVIWriter, Error, TEXT( "ERROR - Could not create CaptureSource filter!" ));
-				Graph->Release();
-				FWindowsPlatformMisc::CoUninitialize(); 
-				return;
-			}
-			CaptionSource->AddRef();
+		CaptureFilter = new FCaptureSource(*this);
+		if (FAILED(hr)) 
+		{
+			UE_LOG(LogAVIWriter, Error, TEXT( "ERROR - Could not create CaptureSource filter!" ));
+			Graph->Release();
+			FWindowsPlatformMisc::CoUninitialize(); 
+			return;
+		}
+		CaptureFilter->AddRef();
 
-			hr = Graph->AddFilter(CaptionSource, L"PushSource"); 
-			if (FAILED(hr)) 
-			{
-				UE_LOG(LogAVIWriter, Error, TEXT( "ERROR - Could not add CaptureSource filter!" ));
-				CaptionSource->Release();
-				Graph->Release();
-				FWindowsPlatformMisc::CoUninitialize(); 
-				return;
-			}
+		hr = Graph->AddFilter(CaptureFilter, L"Capture"); 
+		if (FAILED(hr)) 
+		{
+			UE_LOG(LogAVIWriter, Error, TEXT( "ERROR - Could not add CaptureSource filter!" ));
+			CaptureFilter->Release();
+			Graph->Release();
+			FWindowsPlatformMisc::CoUninitialize(); 
+			return;
+		}
 
-			IAMVideoCompression* VideoCompression = NULL;
-			if (GEngine->MatineeScreenshotOptions.bCompressMatineeCapture)
+		if (!Options.CodecName.IsEmpty())
+		{
+			EncodingFilter = FindEncodingFilter(Options.CodecName);
+			EncodingFilter->AddRef();
+			EncodingFilter->AddRef();
+			EncodingFilter->AddRef();
+			EncodingFilter->AddRef();
+			if (EncodingFilter)
 			{
-				// Create the filter graph manager and query for interfaces.
-				hr = CoCreateInstance (CLSID_MJPGEnc, NULL, CLSCTX_INPROC, IID_IBaseFilter, ( void ** ) &VideoCompression);
-				if (FAILED(hr)) 
-				{
-					UE_LOG(LogAVIWriter, Error, TEXT( "ERROR - Could not create the video compression!" ));
-					FWindowsPlatformMisc::CoUninitialize();
-					return;
-				}
-
-				hr = Graph->AddFilter( (IBaseFilter *)VideoCompression, TEXT("Compressor") );
-
-				if (FAILED(hr)) 
-				{
-					UE_LOG(LogAVIWriter, Error, TEXT( "ERROR - Could not add filter!" ));
-					CaptionSource->Release();
-					Graph->Release();
-					FWindowsPlatformMisc::CoUninitialize(); 
-					return;
-				}	
+				Graph->AddFilter( EncodingFilter, TEXT("Encoder") );
 			}
+		}
 
-			IPin* temp = NULL;
-			hr = CaptionSource->FindPin(L"1", &temp);
-			if (FAILED(hr)) 
+		if (Options.CompressionQuality.IsSet())
+		{
+			if (!EncodingFilter)
 			{
-				UE_LOG(LogAVIWriter, Error, TEXT( "ERROR - Could not find pin of filter CaptureSource!" ));
-				CaptionSource->Release();
-				Graph->Release();
-				FWindowsPlatformMisc::CoUninitialize(); 
-				return;
+				// Attempt to use a default encoder
+				CoCreateInstance(CLSID_MJPGEnc, NULL, CLSCTX_INPROC, IID_IBaseFilter, (void**)&EncodingFilter);
+				Graph->AddFilter( EncodingFilter, TEXT("Encoder") );
 			}
-			CapturePin = (FCapturePin*)temp;
-
-			// Attempt to make the dir if it doesn't exist.
-			TCHAR File[MAX_SPRINTF] = TEXT("");
-			if (OutputPath.IsEmpty())
-			{
-				FString OutputDir = FPaths::VideoCaptureDir();
-				IFileManager::Get().MakeDirectory(*OutputDir, true);
 
-				if (GEngine->MatineeScreenshotOptions.bCompressMatineeCapture)
-				{
-					bool bFoundUnusedName = false;
-					while (!bFoundUnusedName)
-					{
-						const FString FileNameInfo = FString::Printf( TEXT("_%dfps_%dx%d"), 
-							GEngine->MatineeScreenshotOptions.MatineeCaptureFPS, CaptureViewport->GetSizeXY().X, CaptureViewport->GetSizeXY().Y );
-						if (GEngine->MatineeScreenshotOptions.bCompressMatineeCapture)
-						{
-							FCString::Sprintf( File, TEXT("%s_%d.avi"), *(OutputDir + TEXT("/") + GEngine->MatineeScreenshotOptions.MatineePackageCaptureName + "_" + GEngine->MatineeScreenshotOptions.MatineeCaptureName + FileNameInfo + "_compressed"), MovieCaptureIndex );
-						}
-						else
-						{
-							FCString::Sprintf( File, TEXT("%s_%d.avi"), *(OutputDir + TEXT("/") + GEngine->MatineeScreenshotOptions.MatineePackageCaptureName + "_" + GEngine->MatineeScreenshotOptions.MatineeCaptureName + FileNameInfo), MovieCaptureIndex );
-						}
-						if (IFileManager::Get().FileSize(File) != -1)
-						{
-							MovieCaptureIndex++;
-						}
-						else
-						{
-							bFoundUnusedName = true;
-						}
-					}
-				}
-				else
-				{
-					bool bFoundUnusedName = false;
-					while (!bFoundUnusedName)
-					{
-						FString FileName = FString::Printf( TEXT("Movie%i.avi"), MovieCaptureIndex );
-						FCString::Sprintf( File, TEXT("%s"), *(OutputDir / FileName) );
-						if (IFileManager::Get().FileSize(File) != -1)
-						{
-							MovieCaptureIndex++;
-						}
-						else
-						{
-							bFoundUnusedName = true;
-						}
-					}
-				}
-			}
-			else
+			IAMVideoCompression* CompressionImpl = nullptr;
+			if(EncodingFilter && SUCCEEDED(EncodingFilter->QueryInterface(IID_IAMVideoCompression, (void **)&CompressionImpl)))
 			{
-				FCString::Sprintf( File, TEXT("%s"), *OutputPath );
+				CompressionImpl->put_Quality(Options.CompressionQuality.GetValue());
+				CompressionImpl->Release();
 			}
+		}
+		
+		IBaseFilter *pMux;
+		CoCreateInstance(CLSID_AviDest, NULL, CLSCTX_INPROC, IID_IBaseFilter, (void**)&pMux);
+		hr = Graph->AddFilter(pMux, TEXT("AVI Mux"));
+		if (FAILED(hr)) 
+		{
+			UE_LOG(LogAVIWriter, Error, TEXT( "ERROR - Failed to create AVI Mux!" ));
+			Graph->Release();
+			FWindowsPlatformMisc::CoUninitialize(); 
+			return;
+		}
 
-			IBaseFilter *pMux;
-			hr = Capture->SetOutputFileName(
-				&MEDIASUBTYPE_Avi,  // Specifies AVI for the target file.
-				File,				// File name.
-				&pMux,              // Receives a pointer to the mux.
-				NULL);              // (Optional) Receives a pointer to the file sink.
+		IBaseFilter *FileWriter;
+		CoCreateInstance(CLSID_FileWriter, NULL, CLSCTX_INPROC, IID_IBaseFilter, (void**)&FileWriter);
+		hr = Graph->AddFilter(FileWriter, TEXT("File Writer"));
+		if (FAILED(hr)) 
+		{
+			UE_LOG(LogAVIWriter, Error, TEXT( "ERROR - Failed to create file writer!" ));
+			Graph->Release();
+			FWindowsPlatformMisc::CoUninitialize(); 
+			return;
+		}
+		
+		IFileSinkFilter* Sink = nullptr;
+		if(SUCCEEDED(FileWriter->QueryInterface(IID_IFileSinkFilter, (void **)&Sink)))
+		{
+			Sink->SetFileName(*Options.OutputFilename, nullptr);
+		}
 
-			if (FAILED(hr)) 
-			{
-				UE_LOG(LogAVIWriter, Error, TEXT( "ERROR - Failed to set capture filter graph output name!" ));
-				CaptionSource->Release();
-				Graph->Release();
-				FWindowsPlatformMisc::CoUninitialize(); 
-				return;
-			}
+		// Now connect the graph
+		if (EncodingFilter)
+		{
+			Graph->Connect(GetPin(CaptureFilter, PINDIR_OUTPUT), GetPin(EncodingFilter, PINDIR_INPUT));
+			Graph->Connect(GetPin(EncodingFilter, PINDIR_OUTPUT), GetPin(pMux, PINDIR_INPUT));
+		}
+		else
+		{
+			Graph->Connect(GetPin(CaptureFilter, PINDIR_OUTPUT), GetPin(pMux, PINDIR_INPUT));
+		}
 
-			hr = Capture->RenderStream(
-				NULL,								// Pin category.
-				&MEDIATYPE_Video,					// Media type.
-				CaptionSource,					// Capture filter.
-				(IBaseFilter *)VideoCompression,	// Intermediate filter (optional).
-				pMux);								// Mux or file sink filter.
+		Graph->Connect(GetPin(pMux, PINDIR_OUTPUT), GetPin(FileWriter, PINDIR_INPUT));
 
+		if (SUCCEEDED(Graph->QueryInterface(IID_IMediaControl, (void **)&Control)))
+		{
+			FString Directory = Options.OutputFilename;
+			FString Ext = FPaths::GetExtension(Directory, true);
 
-			if (FAILED(hr)) 
-			{
-				UE_LOG(LogAVIWriter, Error, TEXT( "ERROR - Could not start capture!" ));
-				CaptionSource->Release();
-				Graph->Release();
-				FWindowsPlatformMisc::CoUninitialize(); 
-				return;
-			}
-			// Release the mux filter.
-			pMux->Release();
+			// Keep 3 seconds worth of frames in memory
+			CapturedFrames.Reset(new FCapturedFrames(Directory.LeftChop(Ext.Len()) + TEXT("_tmp"), Options.CaptureFPS * 3));
 
-			SAFE_RELEASE(VideoCompression);
-			
-			Graph->QueryInterface(IID_IMediaControl, (void **)&Control);
+			hr = Control->Run();
 
-			MovieCaptureIndex++;
-			bReadyForCapture = true;
+			bCapturing = true;
 		}
+		
+		pMux->Release();
 	}
 
 	void StopCapture()
 	{
-		if (bCapturing)
+		if (!bCapturing)
 		{
-			Control->Stop();
+			return;
+		}
 
-			// Trigger the event to force the capture thread to unblock and quit.
-			GCaptureSyncEvent->Trigger();
-			GCaptureSyncEvent->Wait();
+		// Stop the capture pin first to ensure we have all the frames. This blocks until all frames have been sent downstream.
+		CaptureFilter->StopCapturing();
+		Control->Stop();
 
-			MovieCaptureIndex = 0;
-			bCapturing = false;
-			bReadyForCapture = false;
-			bMatineeFinished = false;
-			CaptureViewport = NULL;
-			CaptureSlateRenderer = NULL;
-			FrameNumber = 0;
-			CurrentAccumSeconds = 0.f;
-			SAFE_RELEASE(CapturePin);
-			SAFE_RELEASE(CaptionSource);
-			SAFE_RELEASE(Capture);
-			SAFE_RELEASE(Control);
-			SAFE_RELEASE(Graph);
-			FWindowsPlatformMisc::CoUninitialize();
-			FPlatformProcess::ReturnSynchEventToPool(GCaptureSyncEvent);
-			GCaptureSyncEvent = nullptr;
+		bCapturing = false;
+		CaptureViewport = NULL;
+		FrameNumber = 0;
 
-			if(CaptureViewport)
-			{
-				ViewportColorBuffer.Empty();
-			}
-		}
+		SAFE_RELEASE(EncodingFilter);
+		SAFE_RELEASE(CaptureFilter);
+		SAFE_RELEASE(Capture);
+		SAFE_RELEASE(Control);
+		SAFE_RELEASE(Graph);
+		FWindowsPlatformMisc::CoUninitialize();
 	}
 
 	void Close()
@@ -363,95 +371,17 @@ public:
 		StopCapture();
 	}
 
-	void Update(float DeltaSeconds)
+	virtual void DropFrames(int32 NumFramesToDrop) override
 	{
-		bool bShouldUpdate = true;
-		const float CaptureFrequency = 1.0f/(float)GEngine->MatineeScreenshotOptions.MatineeCaptureFPS;
-		if (CaptureFrequency > 0.f)
-		{
-			CurrentAccumSeconds += DeltaSeconds;
-
-			if (CurrentAccumSeconds > CaptureFrequency)
-			{
-				while (CurrentAccumSeconds > CaptureFrequency)
-				{
-					CurrentAccumSeconds -= CaptureFrequency;
-				}
-			}
-			else
-			{
-				bShouldUpdate = false;
-			}
-		}
-
-		if (GIsRequestingExit || !bShouldUpdate || CaptureSlateRenderer)
-		{
-			return;
-		}
-
-		if (bMatineeFinished)
-		{
-			if (GEngine->MatineeScreenshotOptions.MatineeCaptureType == 0)
-			{
-				if (CaptureViewport)
-				{
-					CaptureViewport->GetClient()->CloseRequested(CaptureViewport);
-				}
-			}
-			else
-			{
-				FViewport* ViewportUsed = GEngine->GameViewport != NULL ? GEngine->GameViewport->Viewport : NULL;
-				if (ViewportUsed)
-				{
-					ViewportUsed->GetClient()->CloseRequested(ViewportUsed);
-				}
-			}
-			StopCapture();
-		}
-		else if (bCapturing)
-		{
-			if ( GCaptureSyncEvent )
-			{
-				// Wait for the directshow thread to finish encoding the last data
-				GCaptureSyncEvent->Wait();
-			}
-
-			CaptureViewport->ReadPixels(ViewportColorBuffer, FReadSurfaceDataFlags());
-
-			// Allow the directshow thread to encode more data if the event is still available
-			if ( GCaptureSyncEvent )
-			{
-				GCaptureSyncEvent->Trigger();
-
-				UE_LOG(LogMovieCapture, Log, TEXT("-----------------START------------------"));
-				UE_LOG(LogMovieCapture, Log, TEXT(" INCREASE FrameNumber from %d "), FrameNumber);
-				FrameNumber++;
-			}
-		}
-		else if (bReadyForCapture)
-		{
-			CaptureViewport->ReadPixels(ViewportColorBuffer, FReadSurfaceDataFlags());
-
-			if ( GCaptureSyncEvent )
-			{
-				// Allow the directshow thread to process the pixels we just read
-				GCaptureSyncEvent->Trigger();
-				Control->Run();
-				bReadyForCapture = false;
-				bCapturing = true;
-				UE_LOG(LogMovieCapture, Log, TEXT("-----------------START------------------"));
-				UE_LOG(LogMovieCapture, Log, TEXT(" INCREASE FrameNumber from %d "), FrameNumber);
-				FrameNumber++;
-			}
-		}
+		FrameNumber += NumFramesToDrop;
 	}
 
 private:
 	IGraphBuilder* Graph;
 	IMediaControl* Control;
 	ICaptureGraphBuilder2* Capture;
-	FCapturePin* CapturePin;
-	IBaseFilter* CaptionSource;
+	FCaptureSource* CaptureFilter;
+	IBaseFilter* EncodingFilter;
 };
 
 #elif PLATFORM_MAC && !UE_BUILD_MINIMAL
@@ -464,17 +394,18 @@ private:
 class FAVIWriterMac : public FAVIWriter
 {
 public:
-	FAVIWriterMac( )
-	: AVFWriterRef(nil)
+	FAVIWriterMac( const FAVIWriterOptions& InOptions )
+	: FAVIWriter(InOptions)
+	, AVFWriterRef(nil)
 	, AVFWriterInputRef(nil)
 	, AVFPixelBufferAdaptorRef(nil)
+	, bShutdownRequested(false)
 	{
-		
 	};
 	
 public:
 	
-	void StartCapture(FViewport* Viewport , const FString& OutputPath /*=FString()*/)
+	void StartCapture(FViewport* Viewport)
 	{
 		SCOPED_AUTORELEASE_POOL;
 		if (!bCapturing)
@@ -497,62 +428,10 @@ public:
 			
 			CaptureViewport = Viewport;
 			
-			ViewportColorBuffer.AddUninitialized( Viewport->GetSizeXY().X * Viewport->GetSizeXY().Y );
-			
 			// Attempt to make the dir if it doesn't exist.
 			TCHAR File[MAX_SPRINTF] = TEXT("");
-			if (OutputPath.IsEmpty())
-			{
-				FString OutputDir = FPaths::VideoCaptureDir();
-				IFileManager::Get().MakeDirectory(*OutputDir, true);
-				
-				if (GEngine->MatineeScreenshotOptions.bStartWithMatineeCapture)
-				{
-					bool bFoundUnusedName = false;
-					while (!bFoundUnusedName)
-					{
-						const FString FileNameInfo = FString::Printf( TEXT("_%dfps_%dx%d"),
-																	 GEngine->MatineeScreenshotOptions.MatineeCaptureFPS, CaptureViewport->GetSizeXY().X, CaptureViewport->GetSizeXY().Y );
-						if (GEngine->MatineeScreenshotOptions.bCompressMatineeCapture)
-						{
-							FCString::Sprintf( File, TEXT("%s_%d.mov"), *(OutputDir + TEXT("/") + GEngine->MatineeScreenshotOptions.MatineePackageCaptureName + "_" + GEngine->MatineeScreenshotOptions.MatineeCaptureName + FileNameInfo + "_compressed"), MovieCaptureIndex );
-						}
-						else
-						{
-							FCString::Sprintf( File, TEXT("%s_%d.mov"), *(OutputDir + TEXT("/") + GEngine->MatineeScreenshotOptions.MatineePackageCaptureName + "_" + GEngine->MatineeScreenshotOptions.MatineeCaptureName + FileNameInfo), MovieCaptureIndex );
-						}
-						if (IFileManager::Get().FileSize(File) != -1)
-						{
-							MovieCaptureIndex++;
-						}
-						else
-						{
-							bFoundUnusedName = true;
-						}
-					}
-				}
-				else
-				{
-					bool bFoundUnusedName = false;
-					while (!bFoundUnusedName)
-					{
-						FString FileName = FString::Printf( TEXT("Movie%i.mov"), MovieCaptureIndex );
-						FCString::Sprintf( File, TEXT("%s"), *(OutputDir / FileName) );
-						if (IFileManager::Get().FileSize(File) != -1)
-						{
-							MovieCaptureIndex++;
-						}
-						else
-						{
-							bFoundUnusedName = true;
-						}
-					}
-				}
-			}
-			else
-			{
-				FCString::Sprintf( File, TEXT("%s"), *OutputPath );
-			}
+			IFileManager::Get().MakeDirectory(*FPaths::GetPath(Options.OutputFilename), true);
+			FCString::Sprintf( File, TEXT("%s"), *Options.OutputFilename );
 			
 			NSError *Error = nil;
 			
@@ -570,10 +449,11 @@ public:
 			}
 			
 			NSDictionary* VideoSettings = nil;
-			if (GEngine->MatineeScreenshotOptions.bCompressMatineeCapture && GEngine->MatineeScreenshotOptions.bStartWithMatineeCapture)
+			if (Options.CompressionQuality.IsSet())
 			{
 				VideoSettings = [NSDictionary dictionaryWithObjectsAndKeys:
-								 AVVideoCodecH264, AVVideoCodecKey,
+								 AVVideoCodecJPEG, AVVideoCodecKey,
+								 [NSDictionary dictionaryWithObjectsAndKeys:[NSNumber numberWithFloat:Options.CompressionQuality.GetValue()], AVVideoQualityKey, nil], AVVideoCompressionPropertiesKey,
 								 [NSNumber numberWithInt:Viewport->GetSizeXY().X], AVVideoWidthKey,
 								 [NSNumber numberWithInt:Viewport->GetSizeXY().Y], AVVideoHeightKey,
 								 nil];
@@ -581,8 +461,7 @@ public:
 			else
 			{
 				VideoSettings = [NSDictionary dictionaryWithObjectsAndKeys:
-								 AVVideoCodecJPEG, AVVideoCodecKey,
-								 [NSDictionary dictionaryWithObjectsAndKeys:[NSNumber numberWithInt:1], AVVideoQualityKey, nil], AVVideoCompressionPropertiesKey,
+								 AVVideoCodecH264, AVVideoCodecKey,
 								 [NSNumber numberWithInt:Viewport->GetSizeXY().X], AVVideoWidthKey,
 								 [NSNumber numberWithInt:Viewport->GetSizeXY().Y], AVVideoHeightKey,
 								 nil];
@@ -605,9 +484,14 @@ public:
 			[AVFWriterRef startWriting];
 			[AVFWriterRef startSessionAtSourceTime:kCMTimeZero];
 			
-			MovieCaptureIndex++;
-			bReadyForCapture = true;
+			FString Directory = Options.OutputFilename;
+			FString Ext = FPaths::GetExtension(Directory, true);
+
+			// Keep 3 seconds worth of frames in memory
+			CapturedFrames.Reset(new FCapturedFrames(Directory.LeftChop(Ext.Len()) + TEXT("_tmp"), Options.CaptureFPS * 3));
+			
 			bCapturing = true;
+			ThreadTaskFuture = Async<void>(EAsyncExecution::Thread, [this]{	TaskThread(); });
 		}
 	}
 	
@@ -616,36 +500,19 @@ public:
 		if (bCapturing)
 		{
 			SCOPED_AUTORELEASE_POOL;
-			[AVFWriterInputRef markAsFinished];
-			// This will finish asynchronously and then destroy the relevant objects.
-			// We must wait for this to complete.
-			FEvent* Event = FPlatformProcess::GetSynchEventFromPool(true);
-			[AVFWriterRef finishWritingWithCompletionHandler:^{
-				check(AVFWriterRef.status == AVAssetWriterStatusCompleted);
-				Event->Trigger();
-			}];
-			Event->Wait(~0u);
-			FPlatformProcess::ReturnSynchEventToPool(Event);
-			Event = nullptr;
+
+			bShutdownRequested = true;
+			ThreadTaskFuture.Get();
 			[AVFWriterInputRef release];
 			[AVFWriterRef release];
 			[AVFPixelBufferAdaptorRef release];
 			AVFWriterInputRef = nil;
 			AVFWriterRef = nil;
 			AVFPixelBufferAdaptorRef = nil;
-			
-			if (CaptureViewport)
-			{
-				ViewportColorBuffer.Empty();
-			}
-			MovieCaptureIndex = 0;
+
 			bCapturing = false;
-			bReadyForCapture = false;
-			bMatineeFinished = false;
 			CaptureViewport = NULL;
-			CaptureSlateRenderer = NULL;
 			FrameNumber = 0;
-			CurrentAccumSeconds = 0.f;
 		}
 	}
 	
@@ -654,97 +521,312 @@ public:
 		StopCapture();
 	}
 	
-	void Update(float DeltaSeconds)
+	void TaskThread()
 	{
-		bool bShouldUpdate = true;
-		const float CaptureFrequency = 1.0f/(float)GEngine->MatineeScreenshotOptions.MatineeCaptureFPS;
-		if (CaptureFrequency > 0.f)
+		SCOPED_AUTORELEASE_POOL;
+		
+		for(;;)
 		{
-			CurrentAccumSeconds += DeltaSeconds;
-			
-			if (CurrentAccumSeconds > CaptureFrequency)
+			uint32 WaitTimeMs = 100;
+			TArray<FCapturedFrame> PendingFrames = GetFrameData(WaitTimeMs);
+
+			// Capture the frames that we have
+			for (auto& CurrentFrame : PendingFrames)
 			{
-				while (CurrentAccumSeconds > CaptureFrequency)
+				while(![AVFWriterInputRef isReadyForMoreMediaData])
 				{
-					CurrentAccumSeconds -= CaptureFrequency;
+					FPlatformProcess::Sleep(0.0001f);
 				}
-			}
-			else
-			{
-				bShouldUpdate = false;
-			}
-		}
-		
-		if (GIsRequestingExit || !bShouldUpdate || CaptureSlateRenderer)
-		{
-			return;
-		}
-		
-		if (bMatineeFinished)
-		{
-			if (GEngine->MatineeScreenshotOptions.MatineeCaptureType == 0)
-			{
-				if (CaptureViewport)
+				
+				CVPixelBufferRef PixelBuffer = NULL;
+				CVPixelBufferPoolCreatePixelBuffer (NULL, AVFPixelBufferAdaptorRef.pixelBufferPool, &PixelBuffer);
+				if(!PixelBuffer)
 				{
-					CaptureViewport->GetClient()->CloseRequested(CaptureViewport);
+					CVPixelBufferCreate(kCFAllocatorDefault, CaptureViewport->GetSizeXY().X, CaptureViewport->GetSizeXY().Y, kCVPixelFormatType_32BGRA, NULL, &PixelBuffer);
 				}
-			}
-			else
-			{
-				FViewport* ViewportUsed = GEngine->GameViewport != NULL ? GEngine->GameViewport->Viewport : NULL;
-				if (ViewportUsed)
+				check(PixelBuffer);
+				
+				CVPixelBufferLockBaseAddress(PixelBuffer, 0);
+				void* Data = CVPixelBufferGetBaseAddress(PixelBuffer);
+				FMemory::Memcpy(Data, CurrentFrame.FrameData.GetData(), CurrentFrame.FrameData.Num()*sizeof(FColor));
+				CVPixelBufferUnlockBaseAddress(PixelBuffer, 0);
+				
+				CMTime PresentTime = CurrentFrame.FrameIndex > 0 ? CMTimeMake(CurrentFrame.FrameIndex, Options.CaptureFPS) : kCMTimeZero;
+				BOOL OK = [AVFPixelBufferAdaptorRef appendPixelBuffer:PixelBuffer withPresentationTime:PresentTime];
+				check(OK);
+				
+				CVPixelBufferRelease(PixelBuffer);
+
+				if (CurrentFrame.FrameProcessedEvent)
 				{
-					ViewportUsed->GetClient()->CloseRequested(ViewportUsed);
+					CurrentFrame.FrameProcessedEvent->Trigger();
 				}
 			}
-			StopCapture();
-		}
-		else if (bCapturing && bReadyForCapture && [AVFWriterInputRef isReadyForMoreMediaData])
-		{
-			SCOPED_AUTORELEASE_POOL;
-			CaptureViewport->ReadPixels(ViewportColorBuffer, FReadSurfaceDataFlags());
-			
-			CVPixelBufferRef PixeBuffer = NULL;
-			CVPixelBufferPoolCreatePixelBuffer (NULL, AVFPixelBufferAdaptorRef.pixelBufferPool, &PixeBuffer);
-			if(!PixeBuffer)
+
+			if (bShutdownRequested && GetNumOutstandingFrames() == 0)
 			{
-				CVPixelBufferCreate(kCFAllocatorDefault, CaptureViewport->GetSizeXY().X, CaptureViewport->GetSizeXY().Y, kCVPixelFormatType_32BGRA, NULL, &PixeBuffer);
+				break;
 			}
-			check(PixeBuffer);
-			
-			CVPixelBufferLockBaseAddress(PixeBuffer, 0);
-			void* Data = CVPixelBufferGetBaseAddress(PixeBuffer);
-			FMemory::Memcpy(Data, ViewportColorBuffer.GetData(), ViewportColorBuffer.Num()*sizeof(FColor));
-			CVPixelBufferUnlockBaseAddress(PixeBuffer, 0);
-			
-            CMTime PresentTime = FrameNumber > 0 ? CMTimeMake(FrameNumber, GEngine->MatineeScreenshotOptions.MatineeCaptureFPS) : kCMTimeZero;
-			BOOL OK = [AVFPixelBufferAdaptorRef appendPixelBuffer:PixeBuffer withPresentationTime:PresentTime];
-			check(OK);
-			
-			CVPixelBufferRelease(PixeBuffer);
-			
-			UE_LOG(LogMovieCapture, Log, TEXT("-----------------START------------------"));
-			UE_LOG(LogMovieCapture, Log, TEXT(" INCREASE FrameNumber from %d "), FrameNumber);
-			FrameNumber++;
 		}
+
+		[AVFWriterInputRef markAsFinished];
+		// This will finish asynchronously and then destroy the relevant objects.
+		// We must wait for this to complete.
+		FEvent* Event = FPlatformProcess::GetSynchEventFromPool(true);
+		[AVFWriterRef finishWritingWithCompletionHandler:^{
+			check(AVFWriterRef.status == AVAssetWriterStatusCompleted);
+			Event->Trigger();
+		}];
+		Event->Wait(~0u);
+		FPlatformProcess::ReturnSynchEventToPool(Event);
+	}
+
+	virtual void DropFrames(int32 NumFramesToDrop) override
+	{
+		FrameNumber += NumFramesToDrop;
 	}
 	
 private:
 	AVAssetWriter* AVFWriterRef;
 	AVAssetWriterInput* AVFWriterInputRef;
 	AVAssetWriterInputPixelBufferAdaptor* AVFPixelBufferAdaptorRef;
+	FThreadSafeBool bShutdownRequested;
+	TFuture<void> ThreadTaskFuture;
 };
 #endif
 
-FAVIWriter* FAVIWriter::GetInstance()
+FCapturedFrame::FCapturedFrame(double InStartTimeSeconds, double InEndTimeSeconds, uint32 InFrameIndex, TArray<FColor> InFrameData)
+	: StartTimeSeconds(InStartTimeSeconds)
+	, EndTimeSeconds(InEndTimeSeconds)
+	, FrameIndex(InFrameIndex)
+	, FrameData(MoveTemp(InFrameData))
+	, FrameProcessedEvent(nullptr)
+{
+}
+
+FCapturedFrame::~FCapturedFrame()
+{
+}
+
+FCapturedFrames::FCapturedFrames(const FString& InArchiveDirectory, int32 InMaxInMemoryFrames)
+	: ArchiveDirectory(InArchiveDirectory)
+	, MaxInMemoryFrames(InMaxInMemoryFrames)
+{
+	FrameReady = FPlatformProcess::GetSynchEventFromPool();
+
+	// Ensure the file doesn't exist
+	auto& PlatformFile = FPlatformFileManager::Get().GetPlatformFile();
+
+	PlatformFile.DeleteDirectoryRecursively(*ArchiveDirectory);
+	PlatformFile.CreateDirectory(*ArchiveDirectory);
+
+	TotalArchivedFrames = 0;
+	InMemoryFrames.Reserve(MaxInMemoryFrames);
+}
+
+FCapturedFrames::~FCapturedFrames()
+{
+	FPlatformProcess::ReturnSynchEventToPool(FrameReady);
+	FPlatformFileManager::Get().GetPlatformFile().DeleteDirectoryRecursively(*ArchiveDirectory);
+}
+
+void FCapturedFrames::Add(FCapturedFrame Frame)
+{
+	bool bShouldArchive = false;
+	{
+		FScopeLock Lock(&ArchiveFrameMutex);
+		bShouldArchive = ArchivedFrames.Num() != 0;
+	}
+	
+	if (!bShouldArchive)
+	{
+		FScopeLock Lock(&InMemoryFrameMutex);
+		if (InMemoryFrames.Num() < MaxInMemoryFrames)
+		{
+			InMemoryFrames.Add(MoveTemp(Frame));
+		}
+		else
+		{
+			bShouldArchive = true;
+		}
+	}
+
+	if (bShouldArchive)
+	{
+		ArchiveFrame(MoveTemp(Frame));
+	}
+	else
+	{
+		FrameReady->Trigger();
+	}
+}
+
+FArchive &operator<<(FArchive& Ar, FCapturedFrame& Frame)
+{
+	Ar << Frame.StartTimeSeconds;
+	Ar << Frame.EndTimeSeconds;
+	Ar << Frame.FrameIndex;
+	Ar << Frame.FrameData;
+	return Ar;
+}
+
+void FCapturedFrames::ArchiveFrame(FCapturedFrame Frame)
+{
+	// Get (and increment) a unique index for this frame
+	uint32 ArchivedFrameIndex = ++TotalArchivedFrames;
+
+	FString Filename = ArchiveDirectory / FString::Printf(TEXT("%d.frame"), ArchivedFrameIndex);
+	TUniquePtr<FArchive> Archive(IFileManager::Get().CreateFileWriter(*Filename));
+	if (ensure(Archive.IsValid()))
+	{
+		*Archive << Frame;
+		Archive->Close();
+
+		// Add the archived frame to the array
+		FScopeLock Lock(&ArchiveFrameMutex);
+		ArchivedFrames.Add(ArchivedFrameIndex);
+	}
+}
+
+TOptional<FCapturedFrame> FCapturedFrames::UnArchiveFrame(uint32 FrameIndex) const
+{
+	FString Filename = ArchiveDirectory / FString::Printf(TEXT("%d.frame"), FrameIndex);
+	TUniquePtr<FArchive> Archive(IFileManager::Get().CreateFileReader(*Filename));
+	if (ensure(Archive.IsValid()))
+	{
+		FCapturedFrame Frame;
+		*Archive << Frame;
+		Archive->Close();
+
+		FPlatformFileManager::Get().GetPlatformFile().DeleteFile(*Filename);
+		return MoveTemp(Frame);
+	}
+
+	return TOptional<FCapturedFrame>();
+}
+
+void FCapturedFrames::StartUnArchiving()
+{
+	if (UnarchiveTask.IsSet())
+	{
+		return;
+	}
+
+	UnarchiveTask = Async<void>(EAsyncExecution::Thread, [this]{
+
+		// Attempt to unarchive any archived frames
+		ArchiveFrameMutex.Lock();
+		TArray<uint32> ArchivedFramesToGet = ArchivedFrames;
+		ArchiveFrameMutex.Unlock();
+
+		int32 MaxNumToProcess = FMath::Min(ArchivedFramesToGet.Num(), MaxInMemoryFrames);
+		for (int32 Index = 0; Index < MaxNumToProcess; ++Index)
+		{
+			TOptional<FCapturedFrame> Frame = UnArchiveFrame(ArchivedFramesToGet[Index]);
+
+			if (Frame.IsSet())
+			{
+				FScopeLock Lock(&InMemoryFrameMutex);
+				InMemoryFrames.Add(MoveTemp(Frame.GetValue()));
+			}
+		}
+
+		if (MaxNumToProcess)
+		{
+			// Only remove the archived frame indices once we have fully processed them (so that FCapturedFrames::Add knows when to archive frames)
+			{
+				FScopeLock Lock(&ArchiveFrameMutex);
+				ArchivedFrames.RemoveAt(0, MaxNumToProcess, false);
+			}
+
+			FrameReady->Trigger();
+		}
+	});
+}
+
+TArray<FCapturedFrame> FCapturedFrames::ReadFrames(uint32 WaitTimeMs)
+{
+	if (!FrameReady->Wait(WaitTimeMs))
+	{
+		StartUnArchiving();
+		return TArray<FCapturedFrame>();
+	}
+
+	UnarchiveTask = TOptional<TFuture<void>>();
+
+	TArray<FCapturedFrame> Frames;
+	Frames.Reserve(MaxInMemoryFrames);
+
+	// Swap the frames
+	{
+		FScopeLock Lock(&InMemoryFrameMutex);
+		Swap(Frames, InMemoryFrames);
+	}
+
+	StartUnArchiving();
+
+	return Frames;
+}
+
+int32 FCapturedFrames::GetNumOutstandingFrames() const
+{
+	int32 TotalNumFrames = 0;
+	{
+		FScopeLock Lock(&InMemoryFrameMutex);
+		TotalNumFrames += InMemoryFrames.Num();
+	}
+
+	{
+		FScopeLock Lock(&ArchiveFrameMutex);
+		TotalNumFrames += ArchivedFrames.Num();
+	}
+	
+	return TotalNumFrames;
+}
+
+FAVIWriter::~FAVIWriter()
+{
+}
+
+void FAVIWriter::Update(double FrameTimeSeconds)
+{
+	if (!bCapturing)
+	{
+		return;
+	}
+
+	TArray<FColor> Data;
+	CaptureViewport->ReadPixels(Data, FReadSurfaceDataFlags());
+
+	double FrameLength = 1.0 / Options.CaptureFPS;
+	double FrameStart = FrameNumber * FrameLength;
+	FCapturedFrame Frame(FrameStart, FrameStart + FrameLength, FrameNumber, MoveTemp(Data));
+
+	FEvent* SyncEvent = nullptr;
+	if (Options.bSynchronizeFrames)
+	{
+		SyncEvent = FPlatformProcess::GetSynchEventFromPool();
+		Frame.FrameProcessedEvent = SyncEvent;
+	}
+
+	// Add the frame
+	CapturedFrames->Add(MoveTemp(Frame));
+	FrameNumber++;
+
+	if (SyncEvent)
+	{
+		SyncEvent->Wait(MAX_uint32);
+		FPlatformProcess::ReturnSynchEventToPool(SyncEvent);
+	}
+}
+
+FAVIWriter* FAVIWriter::CreateInstance(const FAVIWriterOptions& InOptions)
 {
 #if PLATFORM_WINDOWS && !UE_BUILD_MINIMAL
-	static FAVIWriterWin Instance;
-	return &Instance;
+	return new FAVIWriterWin(InOptions);
 #elif PLATFORM_MAC && !UE_BUILD_MINIMAL
-	static FAVIWriterMac Instance;
-	return &Instance;
+	return new FAVIWriterMac(InOptions);
 #else
-	return NULL;
+	return nullptr;
 #endif
 }
diff --git a/Engine/Source/Runtime/Engine/Private/CapturePin.cpp b/Engine/Source/Runtime/Engine/Private/CapturePin.cpp
--- a/Engine/Source/Runtime/Engine/Private/CapturePin.cpp
+++ b/Engine/Source/Runtime/Engine/Private/CapturePin.cpp
@@ -8,31 +8,22 @@
 
 #include "CapturePin.h"
 #include "AVIWriter.h"
+#include "ScopeExit.h"
 
 DEFINE_LOG_CATEGORY(LogMovieCapture);
 #if PLATFORM_WINDOWS && !UE_BUILD_MINIMAL
 
-
-
-FCapturePin::FCapturePin(HRESULT *phr, CSource *pFilter)
-        : CSourceStream(NAME("Push Source"), phr, pFilter, L"Out"),
-		FramesWritten(0),
-		bZeroMemory(0),
-		FrameLength(UNITS/GEngine->MatineeScreenshotOptions.MatineeCaptureFPS),
-		CurrentBitDepth(32)
+FCapturePin::FCapturePin(HRESULT *phr, CSource *pFilter, const FAVIWriter& InWriter)
+        : CSourceStream(NAME("Push Source"), phr, pFilter, L"Capture")
+		, FrameLength(UNITS/InWriter.Options.CaptureFPS)
+		, Writer(InWriter)
 {
-	// Get the dimensions of the window
-	Screen.left   = Screen.top = 0;
-	Screen.right  = FAVIWriter::GetInstance()->GetWidth();
-	Screen.bottom = FAVIWriter::GetInstance()->GetHeight();
-
-	ImageWidth  = Screen.right  - Screen.left;
-	ImageHeight = Screen.bottom - Screen.top;
+	ImageWidth  = Writer.GetWidth();
+	ImageHeight = Writer.GetHeight();
 }
 
 FCapturePin::~FCapturePin()
-{   
-
+{
 }
 
 
@@ -75,70 +66,18 @@ HRESULT FCapturePin::GetMediaType(int32 iPosition, CMediaType *pmt)
 	// Initialize the VideoInfo structure before configuring its members
 	ZeroMemory(pvi, sizeof(VIDEOINFO));
 
-	switch(iPosition)
-	{
-		case 0:
-		{    
-			// Return our highest quality 32bit format
-
-			// Since we use RGB888 (the default for 32 bit), there is
-			// no reason to use BI_BITFIELDS to specify the RGB
-			// masks. Also, not everything supports BI_BITFIELDS
-			pvi->bmiHeader.biCompression = BI_RGB;
-			pvi->bmiHeader.biBitCount    = 32;
-			break;
-		}
-
-		case 1:
-		{   // Return our 24bit format
-			pvi->bmiHeader.biCompression = BI_RGB;
-			pvi->bmiHeader.biBitCount    = 24;
-			break;
-		}
-
-		case 2:
-		{       
-			// 16 bit per pixel RGB565
-
-			// Place the RGB masks as the first 3 doublewords in the palette area
-			for(int32 i = 0; i < 3; i++)
-			{
-				pvi->TrueColorInfo.dwBitMasks[i] = bits565[i];
-			}
-
-			pvi->bmiHeader.biCompression = BI_BITFIELDS;
-			pvi->bmiHeader.biBitCount    = 16;
-			break;
-		}
+	// We only supply 32bit RGB
 
-		case 3:
-		{   // 16 bits per pixel RGB555
-
-			// Place the RGB masks as the first 3 doublewords in the palette area
-			for(int32 i = 0; i < 3; i++)
-			{
-				pvi->TrueColorInfo.dwBitMasks[i] = bits555[i];
-			}
-
-			pvi->bmiHeader.biCompression = BI_BITFIELDS;
-			pvi->bmiHeader.biBitCount    = 16;
-			break;
-		}
-
-		case 4:
-		{   // 8 bit palettised
-
-			pvi->bmiHeader.biCompression = BI_RGB;
-			pvi->bmiHeader.biBitCount    = 8;
-			pvi->bmiHeader.biClrUsed     = iPALETTE_COLORS;
-			break;
-		}
-	}
+	// Since we use RGB888 (the default for 32 bit), there is
+	// no reason to use BI_BITFIELDS to specify the RGB
+	// masks. Also, not everything supports BI_BITFIELDS
+	pvi->bmiHeader.biCompression = BI_RGB;
+	pvi->bmiHeader.biBitCount    = 32;
 
 	// Adjust the parameters common to all formats
 	pvi->bmiHeader.biSize       = sizeof(BITMAPINFOHEADER);
 	pvi->bmiHeader.biWidth      = ImageWidth;
-	pvi->bmiHeader.biHeight     = ImageHeight * (-1); // negative 1 to flip the image vertically
+	pvi->bmiHeader.biHeight     = ImageHeight;// * (-1); // negative 1 to flip the image vertically
 	pvi->bmiHeader.biPlanes     = 1;
 	pvi->bmiHeader.biSizeImage  = GetBitmapSize(&pvi->bmiHeader);
 	pvi->bmiHeader.biClrImportant = 0;
@@ -149,11 +88,11 @@ HRESULT FCapturePin::GetMediaType(int32 iPosition, CMediaType *pmt)
 
 	pmt->SetType(&MEDIATYPE_Video);
 	pmt->SetFormatType(&FORMAT_VideoInfo);
-	pmt->SetTemporalCompression(false);
+	pmt->SetTemporalCompression(true);
 
 	// Work out the GUID for the subtype from the header info.
-	const GUID SubTypeGUID = GetBitmapSubtype(&pvi->bmiHeader);
-	pmt->SetSubtype(&SubTypeGUID);
+	//const GUID SubTypeGUID = GetBitmapSubtype(&pvi->bmiHeader);
+	pmt->SetSubtype(&MEDIASUBTYPE_RGB32);
 	pmt->SetSampleSize(pvi->bmiHeader.biSizeImage);
 
 	return NOERROR;
@@ -185,11 +124,7 @@ HRESULT FCapturePin::CheckMediaType(const CMediaType *pMediaType)
 		return E_INVALIDARG;
 	}
 
-	if(    (*SubType != MEDIASUBTYPE_RGB8)
-		&& (*SubType != MEDIASUBTYPE_RGB565)
-		&& (*SubType != MEDIASUBTYPE_RGB555)
-		&& (*SubType != MEDIASUBTYPE_RGB24)
-		&& (*SubType != MEDIASUBTYPE_RGB32))
+	if (*SubType != MEDIASUBTYPE_RGB32)
 	{
 		return E_INVALIDARG;
 	}
@@ -279,27 +214,15 @@ HRESULT FCapturePin::SetMediaType(const CMediaType *pMediaType)
 			return E_UNEXPECTED;
 		}
 
-		switch(pvi->bmiHeader.biBitCount)
+		if (pvi->bmiHeader.biBitCount == 32)
 		{
-			case 8:     // 8-bit palettized
-			case 16:    // RGB565, RGB555
-			case 24:    // RGB24
-			case 32:    // RGB32
-				// Save the current media type and bit depth
-				MediaType = *pMediaType;
-				CurrentBitDepth = pvi->bmiHeader.biBitCount;
-				hr = S_OK;
-				break;
-
-			default:
-				// We should never agree any other media types
-				check(false);
-				hr = E_INVALIDARG;
-				break;
+			return S_OK;
 		}
-	} 
+	}
 
-	return hr;
+	// We should never agree any other media types
+	check(false);
+	return E_INVALIDARG;
 
 } // SetMediaType
 
@@ -326,31 +249,28 @@ HRESULT FCapturePin::FillBuffer(IMediaSample *pSample)
 
 	if (pData)
 	{
-		uint32 sizeInBytes = FAVIWriter::GetInstance()->GetWidth() * FAVIWriter::GetInstance()->GetHeight() * sizeof(FColor);
-		const TArray<FColor>& Buffer = FAVIWriter::GetInstance()->GetColorBuffer();
-		uint32 smallest = FMath::Min((uint32)pVih->bmiHeader.biSizeImage, (uint32)cbData);
-		// Copy the DIB bits over into our filter's output buffer.
-		// Since sample size may be larger than the image size, bound the copy size
-		FMemory::Memcpy(pData, Buffer.GetData(), FMath::Min(smallest, sizeInBytes));
+		uint32 BytesPerRow = sizeof(FColor) * ImageWidth;
+
+		// Fill buffer bottom up
+		for (int32 Row = ImageHeight - 1; Row >= 0; --Row, pData += BytesPerRow)
+		{
+			const FColor* Read = &CurrentFrame->FrameData.GetData()[ImageWidth * Row];
+			FMemory::Memcpy(pData, Read, BytesPerRow);
+		}
 	}
 
 	// Set the timestamps that will govern playback frame rate.
-	// set the Average Time Per Frame for the AVI Header.
 	// The current time is the sample's start.
 
-	int32 FrameNumber = FAVIWriter::GetInstance()->GetFrameNumber();
-	UE_LOG(LogMovieCapture, Log, TEXT(" FillBuffer: FrameNumber = %d  FramesWritten = %d"), FrameNumber, FramesWritten);
-
-	REFERENCE_TIME Start = FramesWritten * FrameLength;
-	REFERENCE_TIME Stop  = Start + FrameLength;
-	FramesWritten++;
+	// Not strictly necessary since AVI is constant framerate
+	REFERENCE_TIME StartTime = UNITS*CurrentFrame->StartTimeSeconds;
+	REFERENCE_TIME StopTime = UNITS*CurrentFrame->EndTimeSeconds;
+	pSample->SetTime(&StartTime, &StopTime);
 
-	UE_LOG(LogMovieCapture, Log, TEXT(" FillBuffer: (%d, %d)"), Start, Stop);
-	UE_LOG(LogMovieCapture, Log, TEXT("-----------------END------------------"));
+	REFERENCE_TIME StartMediaTime = CurrentFrame->FrameIndex;
+	REFERENCE_TIME StopMediaTime  = CurrentFrame->FrameIndex+1;
+	pSample->SetMediaTime(&StartMediaTime, &StopMediaTime);
 
-	pSample->SetTime(&Start, &Stop);
-
-	// Set true on every sample for uncompressed frames
 	pSample->SetSyncPoint(true);
 
 	return S_OK;
@@ -359,70 +279,101 @@ HRESULT FCapturePin::FillBuffer(IMediaSample *pSample)
 // the loop executed while running
 HRESULT FCapturePin::DoBufferProcessingLoop(void) 
 {
-	Command com;
+	ON_SCOPE_EXIT
+	{
+		static_cast<FCaptureSource*>(m_pFilter)->OnFinishedCapturing();
+	};
 
 	OnThreadStartPlay();
-	int32 LastFrame = -1;
 
-	do 
+	bool bPaused = false, bShutdownRequested = false;
+	for (;;)
 	{
-		while (!CheckRequest(&com)) 
+		Command com;
+		if (CheckRequest(&com))
 		{
-			// Wait for the next frame from the game thread 
-			if ( !GCaptureSyncEvent->Wait(1000) )
+			switch(com)
 			{
-				FPlatformProcess::Sleep( 0.01f );
-				continue;	// Reevaluate request
+			case CMD_RUN:
+			case CMD_PAUSE:
+				bPaused = com == CMD_PAUSE;
+				Reply(NOERROR);
+				break;
+
+			case CMD_STOP:
+				break;
+
+			default:
+				Reply((uint32) E_UNEXPECTED);
+				break;
 			}
+		}
 
-			IMediaSample *pSample;
-			int32 FrameNumber = FAVIWriter::GetInstance()->GetFrameNumber();
-			if (FrameNumber > LastFrame)
+		bShutdownRequested = bShutdownRequested || !static_cast<FCaptureSource*>(m_pFilter)->ShouldCapture();
+
+		if (!bPaused)
+		{
+			TOptional<HRESULT> Result = ProcessFrames();
+			if (Result.IsSet())
 			{
-				UE_LOG(LogMovieCapture, Log, TEXT(" FrameNumber > LastFrame = %d > %d"), FrameNumber, LastFrame);
-				HRESULT hr = GetDeliveryBuffer(&pSample,NULL,NULL,0);
-				if (FAILED(hr)) 
-				{
-					if (pSample)
-					{
-						pSample->Release();
-					}
-				}
-				else
-				{
-					LastFrame = FrameNumber;
-					hr = FillBuffer(pSample);
-
-					if (hr == S_OK) 
-					{
-						hr = Deliver(pSample);
-						pSample->Release();
-						// downstream filter returns S_FALSE if it wants us to
-						// stop or an error if it's reporting an error.
-						if(hr != S_OK)
-						{
-							UE_LOG(LogMovieCapture, Log, TEXT("Deliver() returned %08x; stopping"), hr);
-							return S_OK;
-						}
-					}
-				}
+				return Result.GetValue();
 			}
-			// Allow the game thread read more data
-			GCaptureSyncEvent->Trigger();
 		}
 
-		// For all commands sent to us there must be a Reply call!
-		if (com == CMD_RUN || com == CMD_PAUSE) 
+		if (bShutdownRequested && 
+			(bPaused || Writer.GetNumOutstandingFrames() == 0))
 		{
-			Reply(NOERROR);
-		} 
-		else if (com != CMD_STOP) 
-		{
-			Reply((uint32) E_UNEXPECTED);
+			break;
 		}
-	} while (com != CMD_STOP);
+	}
 
+	static_cast<FCaptureSource*>(m_pFilter)->OnFinishedCapturing();
 	return S_FALSE;
 }
 
+TOptional<HRESULT> FCapturePin::ProcessFrames()
+{
+	uint32 WaitTimeMs = 100;
+	TArray<FCapturedFrame> PendingFrames = Writer.GetFrameData(WaitTimeMs);
+
+	// Capture the frames that we have
+	for (int32 FrameIndex = 0; FrameIndex < PendingFrames.Num(); ++FrameIndex)
+	{
+		CurrentFrame = &PendingFrames[FrameIndex];
+		IMediaSample *pSample = nullptr;
+
+		HRESULT hr;
+		hr = GetDeliveryBuffer(&pSample,NULL,NULL,0);
+		if (FAILED(hr))
+		{
+			UE_LOG(LogMovieCapture, Warning, TEXT("Failed to get delivery buffer: %08x; Stopping."), hr);
+			return S_OK;
+		}
+
+		hr = FillBuffer(pSample);
+		if (FAILED(hr))
+		{
+			pSample->Release();
+			UE_LOG(LogMovieCapture, Warning, TEXT("Failed to fill sample buffer: %08x; Stopping."), hr);
+			return S_OK;
+		}
+
+		hr = Deliver(pSample);
+		pSample->Release();
+		// downstream filter returns S_FALSE if it wants us to
+		// stop or an error if it's reporting an error.
+		if(hr != S_OK)
+		{
+			UE_LOG(LogMovieCapture, Warning, TEXT("Deliver() returned %08x; Stopping."), hr);
+			return S_OK;
+		}
+
+		if (CurrentFrame->FrameProcessedEvent)
+		{
+			CurrentFrame->FrameProcessedEvent->Trigger();
+		}
+	}
+
+	return TOptional<HRESULT>();
+}
 #endif //#if PLATFORM_WINDOWS
diff --git a/Engine/Source/Runtime/Engine/Private/CapturePin.h b/Engine/Source/Runtime/Engine/Private/CapturePin.h
--- a/Engine/Source/Runtime/Engine/Private/CapturePin.h
+++ b/Engine/Source/Runtime/Engine/Private/CapturePin.h
@@ -33,35 +33,26 @@ class FCapturePin : public CSourceStream
 {
 protected:
 
-	/** To track where we are in the file */
-	int32 FramesWritten;
-	/** Do we need to clear the buffer? */
-	bool bZeroMemory; 
-	/** The time stamp for each sample */
-	CRefTime SampleTime;
 	/** The length of the frame, used for playback */
 	const REFERENCE_TIME FrameLength;
-	/** Rect containing entire screen coordinates */
-	RECT Screen;                     
 	/** The current image height */
 	int32 ImageHeight;
 	/** And current image width */
 	int32 ImageWidth;
-	/** Time in msec between frames */
-	int32 RepeatTime;
-	/** Screen bit depth */
-	int32 CurrentBitDepth;
-
-	/** Stores the media type to use */
-	CMediaType MediaType;
 	/** Protects our internal state */
 	CCritSec SharedState;
-	/** Figures out our media type for us */
-	CImageDisplay Display;
+
+	/** The writer to which we belong */
+	const FAVIWriter& Writer;
+
+	/** The frame we're currently processing */
+	const FCapturedFrame* CurrentFrame;
+
+	TOptional<HRESULT> ProcessFrames();
 
 public:
 
-	FCapturePin(HRESULT *phr, CSource *pFilter);
+	FCapturePin(HRESULT *phr, CSource *pFilter, const FAVIWriter& InWriter);
 	~FCapturePin();
 
 	// Override the version that offers exactly one media type
diff --git a/Engine/Source/Runtime/Engine/Private/CaptureSource.cpp b/Engine/Source/Runtime/Engine/Private/CaptureSource.cpp
--- a/Engine/Source/Runtime/Engine/Private/CaptureSource.cpp
+++ b/Engine/Source/Runtime/Engine/Private/CaptureSource.cpp
@@ -11,40 +11,30 @@
 
 #if PLATFORM_WINDOWS && !UE_BUILD_MINIMAL
 
-FCaptureSource::FCaptureSource(IUnknown *pUnk, HRESULT *phr)
-           : CSource(NAME("PushSource"), pUnk, CLSID_CaptureSource)
+FCaptureSource::FCaptureSource(const FAVIWriter& Writer)
+           : CSource(NAME("ViewportCaptureFilter"), nullptr, CLSID_ViewportCaptureSource)
 {
-	CapturePin = new FCapturePin(phr, this);
-
-	if (phr)
-	{
-		if (CapturePin == NULL)
-			*phr = E_OUTOFMEMORY;
-		else
-			*phr = S_OK;
-	}  
-}
+	HRESULT hr;
+	new FCapturePin(&hr, this, Writer);
 
+	ShutdownEvent = FPlatformProcess::GetSynchEventFromPool();
+	bShutdownRequested = false;
+}
 
 FCaptureSource::~FCaptureSource()
 {
-	delete CapturePin;
+	FPlatformProcess::ReturnSynchEventToPool(ShutdownEvent);
 }
 
-
-CUnknown * WINAPI FCaptureSource::CreateInstance(IUnknown *pUnk, HRESULT *phr)
+void FCaptureSource::StopCapturing()
 {
-	FCaptureSource *pNewFilter = new FCaptureSource(pUnk, phr );
-
-	if (phr)
-	{
-		if (pNewFilter == NULL) 
-			*phr = E_OUTOFMEMORY;
-		else
-			*phr = S_OK;
-	}
-	return pNewFilter;
+	bShutdownRequested = true;
+	ShutdownEvent->Wait(~0);
+}
 
+void FCaptureSource::OnFinishedCapturing()
+{
+	ShutdownEvent->Trigger();
 }
 
 #endif //#if PLATFORM_WINDOWS
diff --git a/Engine/Source/Runtime/Engine/Private/CaptureSource.h b/Engine/Source/Runtime/Engine/Private/CaptureSource.h
--- a/Engine/Source/Runtime/Engine/Private/CaptureSource.h
+++ b/Engine/Source/Runtime/Engine/Private/CaptureSource.h
@@ -17,22 +17,24 @@
 #pragma warning(pop)
 class FCapturePin;
 
-// {9A80E195-3BBA-4821-B18B-21BB496F80F8}
-DEFINE_GUID(CLSID_CaptureSource, 
-			0x9a80e195, 0x3bba, 0x4821, 0xb1, 0x8b, 0x21, 0xbb, 0x49, 0x6f, 0x80, 0xf8);
+// {F817F8A7-DE00-42CF-826A-7A5654602D8E}
+DEFINE_GUID(CLSID_ViewportCaptureSource, 
+	0xf817f8a7, 0xde00, 0x42cf, 0x82, 0x6a, 0x7a, 0x56, 0x54, 0x60, 0x2d, 0x8e);
+
 
 class FCaptureSource : public CSource
 {
-
 public:
-	FCaptureSource(IUnknown *pUnk, HRESULT *phr);
+	FCaptureSource(const FAVIWriter& Writer);
 	~FCaptureSource();
-private:
-	FCapturePin *CapturePin;
 
-public:
-	static CUnknown * WINAPI CreateInstance(IUnknown *pUnk, HRESULT *phr);  
+	void StopCapturing();
+	void OnFinishedCapturing();
+	bool ShouldCapture() const { return !bShutdownRequested; }
 
+private:
+	FEvent* ShutdownEvent;
+	FThreadSafeBool bShutdownRequested;
 };
 #endif //#if PLATFORM_WINDOWS
 
diff --git a/Engine/Source/Runtime/Engine/Private/GameEngine.cpp b/Engine/Source/Runtime/Engine/Private/GameEngine.cpp
--- a/Engine/Source/Runtime/Engine/Private/GameEngine.cpp
+++ b/Engine/Source/Runtime/Engine/Private/GameEngine.cpp
@@ -18,13 +18,14 @@
 #include "Engine/RendererSettings.h"
 #include "Engine/UserInterfaceSettings.h"
 #include "GeneralProjectSettings.h"
-#include "AVIWriter.h"
 
 #include "SlateBasics.h"
 #include "Slate/SceneViewport.h"
 #include "SVirtualJoystick.h"
 
-#include "AVIWriter.h"
+#include "MovieSceneCaptureModule.h"
+#include "MovieSceneCaptureSettings.h"
+
 #include "AssetRegistryModule.h"
 #include "SynthBenchmark.h"
 
@@ -94,7 +95,7 @@ UGameEngine::UGameEngine(const FObjectInitializer& ObjectInitializer)
 
 void UGameEngine::CreateGameViewportWidget( UGameViewportClient* GameViewportClient )
 {
-	bool bRenderDirectlyToWindow = !GEngine->MatineeScreenshotOptions.bStartWithMatineeCapture && GIsDumpingMovie == 0; 
+	bool bRenderDirectlyToWindow = !StartupMovieCaptureHandle.IsValid() && GIsDumpingMovie == 0; 
 	const bool bStereoAllowed = bRenderDirectlyToWindow;
 	TSharedRef<SOverlay> ViewportOverlayWidgetRef = SNew( SOverlay );
 
@@ -155,7 +156,6 @@ void UGameEngine::CreateGameViewport( UGameViewportClient* GameViewportClient )
 
 	SceneViewport = MakeShareable( new FSceneViewport( GameViewportClient, GameViewportWidgetRef ) );
 	GameViewportClient->Viewport = SceneViewport.Get();
-
 	//GameViewportClient->CreateHighresScreenshotCaptureRegionWidget(); //  Disabled until mouse based input system can be made to work correctly.
 
 	// The viewport widget needs an interface so it knows what should render
@@ -265,6 +265,7 @@ TSharedRef<SWindow> UGameEngine::CreateGameWindow()
 	int32 ResX = GSystemResolution.ResX;
 	int32 ResY = GSystemResolution.ResY;
 	EWindowMode::Type WindowMode = GSystemResolution.WindowMode;
+
 	ConditionallyOverrideSettings(ResX, ResY, WindowMode);
 
 	// If the current settings have been overridden, apply them back into the system
@@ -344,9 +345,15 @@ void UGameEngine::SwitchGameWindowToUseGameViewport()
 		
 		GameViewportWindowPtr->SetContent(GameViewportWidgetRef);
 		GameViewportWindowPtr->SlatePrepass();
-		
+
 		SceneViewport->ResizeFrame((uint32)GSystemResolution.ResX, (uint32)GSystemResolution.ResY, GSystemResolution.WindowMode, 0, 0);
 
+		IMovieSceneCaptureInterface* MovieSceneCaptureImpl = StartupMovieCaptureHandle.IsValid() ? IMovieSceneCaptureModule::Get().RetrieveMovieSceneInterface(StartupMovieCaptureHandle) : nullptr;
+		if (MovieSceneCaptureImpl)
+		{
+			MovieSceneCaptureImpl->Initialize(SceneViewport.Get());
+		}
+
 		// Move the registration of the game viewport to that messages are correctly received.
 		if (!FPlatformProperties::SupportsWindowedMode())
 		{
@@ -463,6 +470,19 @@ void UGameEngine::Init(IEngineLoop* InEngineLoop)
 //  	// Creates the initial world context. For GameEngine, this should be the only WorldContext that ever gets created.
 //  	FWorldContext& InitialWorldContext = CreateNewWorldContext(EWorldType::Game);
 
+	IMovieSceneCaptureInterface* MovieSceneCaptureImpl = nullptr;
+#if WITH_EDITOR
+	if (!IsRunningDedicatedServer() && !IsRunningCommandlet())
+	{
+		MovieSceneCaptureImpl = IMovieSceneCaptureModule::Get().InitializeFromCommandLine();
+		if (MovieSceneCaptureImpl)
+		{
+			StartupMovieCaptureHandle = MovieSceneCaptureImpl->GetHandle();
+			GameInstance->InitialMapOverride = MovieSceneCaptureImpl->GetPackageName();
+		}
+	}
+#endif
+
 	// Initialize the viewport client.
 	UGameViewportClient* ViewportClient = NULL;
 	if(GIsClient)
@@ -473,7 +493,6 @@ void UGameEngine::Init(IEngineLoop* InEngineLoop)
 		GameInstance->GetWorldContext()->GameViewport = ViewportClient;
 	}
 
-	bCheckForMovieCapture = true;
 	LastTimeLogsFlushed = FPlatformTime::Seconds();
 
 	// Attach the viewport client to a new viewport.
@@ -504,7 +523,6 @@ void UGameEngine::Init(IEngineLoop* InEngineLoop)
 
 	GameInstance->StartGameInstance();
 
-
 	UE_LOG(LogInit, Display, TEXT("Game Engine Initialized.") );
 
 	// for IsInitialized()
@@ -514,12 +532,6 @@ void UGameEngine::Init(IEngineLoop* InEngineLoop)
 
 void UGameEngine::PreExit()
 {
-	FAVIWriter* AVIWriter = FAVIWriter::GetInstance();
-	if (AVIWriter)
-	{
-		AVIWriter->Close();
-	}
-
 	Super::PreExit();
 
 	// Stop tracking, automatically flushes.
@@ -1059,28 +1071,6 @@ void UGameEngine::Tick( float DeltaSeconds, bool bIdleMode )
 			GetRendererModule().TickRenderTargetPool();
 		});
 	}
-	
-	// Skip AVI work in environments where there is no rendering
-	if (!IsRunningDedicatedServer() && !IsRunningCommandlet())
-	{
-		FAVIWriter* AVIWriter = FAVIWriter::GetInstance();
-		if (AVIWriter)
-		{
-			QUICK_SCOPE_CYCLE_COUNTER(STAT_UGameEngine_Tick_AVIWriter);
-			AVIWriter->Update(DeltaSeconds);
-		}
-
-		// Start the movie capture if needed
-		if (bCheckForMovieCapture && GEngine->MatineeScreenshotOptions.bStartWithMatineeCapture && GEngine->MatineeScreenshotOptions.MatineeCaptureType == EMatineeCaptureType::AVI && GameViewport->Viewport->GetSizeXY() != FIntPoint::ZeroValue )
-		{
-			if (AVIWriter)
-			{
-				QUICK_SCOPE_CYCLE_COUNTER(STAT_UGameEngine_Tick_StartCapture);
-				AVIWriter->StartCapture(GameViewport->Viewport);
-			}
-			bCheckForMovieCapture = false;
-		}
-	}
 
 #if WITH_EDITOR
 	// tick animation recorder. available only in editor builds
diff --git a/Engine/Source/Runtime/Engine/Private/GameInstance.cpp b/Engine/Source/Runtime/Engine/Private/GameInstance.cpp
--- a/Engine/Source/Runtime/Engine/Private/GameInstance.cpp
+++ b/Engine/Source/Runtime/Engine/Private/GameInstance.cpp
@@ -300,7 +300,7 @@ void UGameInstance::StartGameInstance()
 	// Enter initial world.
 	EBrowseReturnVal::Type BrowseRet = EBrowseReturnVal::Failure;
 	FString Error;
-	TCHAR Parm[4096] = TEXT("");
+	
 	const TCHAR* Tmp = FCommandLine::Get();
 
 #if UE_BUILD_SHIPPING
@@ -318,12 +318,19 @@ void UGameInstance::StartGameInstance()
 
 	const UGameMapsSettings* GameMapsSettings = GetDefault<UGameMapsSettings>();
 	const FString& DefaultMap = GameMapsSettings->GetGameDefaultMap();
-	if (!FParse::Token(Tmp, Parm, ARRAY_COUNT(Parm), 0) || Parm[0] == '-')
+
+	FString PackageName;
+
+#if WITH_EDITOR
+	PackageName = InitialMapOverride;
+#endif
+
+	if (PackageName.IsEmpty() && (!FParse::Token(Tmp, PackageName, 0) || **PackageName == '-'))
 	{
-		FCString::Strcpy(Parm, *(DefaultMap + GameMapsSettings->LocalMapOptions));
+		PackageName = DefaultMap + GameMapsSettings->LocalMapOptions;
 	}
 
-	FURL URL(&DefaultURL, Parm, TRAVEL_Partial);
+	FURL URL(&DefaultURL, *PackageName, TRAVEL_Partial);
 	if (URL.Valid)
 	{
 		BrowseRet = Engine->Browse(*WorldContext, URL, Error);
@@ -335,7 +342,7 @@ void UGameInstance::StartGameInstance()
 		UE_LOG(LogLoad, Error, TEXT("%s"), *FString::Printf(TEXT("Failed to enter %s: %s. Please check the log for errors."), *URL.Map, *Error));
 
 		// the map specified on the command-line couldn't be loaded.  ask the user if we should load the default map instead
-		if (FCString::Stricmp(Parm, *DefaultMap) != 0)
+		if (FCString::Stricmp(*PackageName, *DefaultMap) != 0)
 		{
 			const FText Message = FText::Format(NSLOCTEXT("Engine", "MapNotFound", "The map specified on the commandline '{0}' could not be found. Would you like to load the default map instead?"), FText::FromString(URL.Map));
 			if (   FCString::Stricmp(*URL.Map, *DefaultMap) != 0  
@@ -368,7 +375,6 @@ void UGameInstance::StartGameInstance()
 		FPlatformMisc::RequestExit(false);
 		return;
 	}
-
 }
 
 
diff --git a/Engine/Source/Runtime/Engine/Private/GameMode.cpp b/Engine/Source/Runtime/Engine/Private/GameMode.cpp
--- a/Engine/Source/Runtime/Engine/Private/GameMode.cpp
+++ b/Engine/Source/Runtime/Engine/Private/GameMode.cpp
@@ -21,6 +21,7 @@
 #include "GameFramework/GameMode.h"
 #include "Engine/ChildConnection.h"
 #include "Engine/GameInstance.h"
+#include "MovieSceneCaptureSettings.h"
 
 DEFINE_LOG_CATEGORY_STATIC(LogGameMode, Log, All);
 
@@ -283,7 +284,7 @@ void AGameMode::PostLogin( APlayerController* NewPlayer )
 	bool HidePlayer=false, HideHUD=false, DisableMovement=false, DisableTurning=false;
 
 	//Check to see if we should start in cinematic mode (matinee movie capture)
-	if(ShouldStartInCinematicMode(HidePlayer, HideHUD, DisableMovement, DisableTurning))
+	if(ShouldStartInCinematicMode(NewPlayer, HidePlayer, HideHUD, DisableMovement, DisableTurning))
 	{
 		NewPlayer->SetCinematicMode(true, HidePlayer, HideHUD, DisableMovement, DisableTurning);
 	}
@@ -295,22 +296,29 @@ void AGameMode::PostLogin( APlayerController* NewPlayer )
 	K2_PostLogin(NewPlayer);
 }
 
-bool AGameMode::ShouldStartInCinematicMode(bool& OutHidePlayer,bool& OutHideHUD,bool& OutDisableMovement,bool& OutDisableTurning)
+bool AGameMode::ShouldStartInCinematicMode(APlayerController* Player, bool& OutHidePlayer,bool& OutHideHUD,bool& OutDisableMovement,bool& OutDisableTurning)
 {
-	bool StartInCinematicMode = false;
-	if(GEngine->MatineeScreenshotOptions.bStartWithMatineeCapture)
+	ULocalPlayer* LocPlayer = Player->GetLocalPlayer();
+	if (!LocPlayer)
 	{
-		GConfig->GetBool( TEXT("MatineeCreateMovieOptions"), TEXT("CinematicMode"), StartInCinematicMode, GEditorPerProjectIni );
-		if(StartInCinematicMode)
+		return false;
+	}
+	else if(LocPlayer->ViewportClient && LocPlayer->ViewportClient->Viewport)
+	{
+		if (auto* MovieSceneCapture = LocPlayer->ViewportClient->Viewport->GetMovieSceneCapture())
 		{
-			GConfig->GetBool( TEXT("MatineeCreateMovieOptions"), TEXT("DisableMovement"), OutDisableMovement, GEditorPerProjectIni );
-			GConfig->GetBool( TEXT("MatineeCreateMovieOptions"), TEXT("DisableTurning"), OutDisableTurning, GEditorPerProjectIni );
-			GConfig->GetBool( TEXT("MatineeCreateMovieOptions"), TEXT("HidePlayer"), OutHidePlayer, GEditorPerProjectIni );
-			GConfig->GetBool( TEXT("MatineeCreateMovieOptions"), TEXT("HideHUD"), OutHideHUD, GEditorPerProjectIni );
-			return StartInCinematicMode;
+			const FMovieSceneCaptureSettings& Settings = MovieSceneCapture->GetSettings();
+			if (Settings.bCinematicMode)
+			{
+				OutDisableMovement = !Settings.bAllowMovement;
+				OutDisableTurning = !Settings.bAllowTurning;
+				OutHidePlayer = !Settings.bShowPlayer;
+				OutHideHUD = !Settings.bShowHUD;
+				return true;
+			}
 		}
 	}
-	return StartInCinematicMode;
+	return false;
 }
 
 
diff --git a/Engine/Source/Runtime/Engine/Private/GameViewportClient.cpp b/Engine/Source/Runtime/Engine/Private/GameViewportClient.cpp
--- a/Engine/Source/Runtime/Engine/Private/GameViewportClient.cpp
+++ b/Engine/Source/Runtime/Engine/Private/GameViewportClient.cpp
@@ -32,6 +32,8 @@
 #include "GameFramework/GameUserSettings.h"
 #include "Runtime/Engine/Classes/Engine/UserInterfaceSettings.h"
 #include "SGameLayerManager.h"
+#include "IMovieSceneCapture.h"
+#include "MovieSceneCaptureSettings.h"
 #include "ActorEditorUtils.h"
 
 #define LOCTEXT_NAMESPACE "GameViewport"
@@ -1239,6 +1241,7 @@ void UGameViewportClient::ProcessScreenShots(FViewport* InViewport)
 {
 	if (GIsDumpingMovie || FScreenshotRequest::IsScreenshotRequested() || GIsHighResScreenshot)
 	{
+	
 		TArray<FColor> Bitmap;
 
 		bool bShowUI = false;
@@ -1247,10 +1250,6 @@ void UGameViewportClient::ProcessScreenShots(FViewport* InViewport)
 		{
 			bShowUI = true;
 		}
-		else if( GIsDumpingMovie && WindowPtr.IsValid() && !GEngine->MatineeScreenshotOptions.bHideHud )
-		{
-			bShowUI = true;
-		}
 
 		bool bScreenshotSuccessful = false;
 		FIntVector Size(InViewport->GetSizeXY().X, InViewport->GetSizeXY().Y, 0);
diff --git a/Engine/Source/Runtime/Engine/Private/UnrealClient.cpp b/Engine/Source/Runtime/Engine/Private/UnrealClient.cpp
--- a/Engine/Source/Runtime/Engine/Private/UnrealClient.cpp
+++ b/Engine/Source/Runtime/Engine/Private/UnrealClient.cpp
@@ -19,6 +19,7 @@
 #include "EngineModule.h"
 #include "ContentStreaming.h"
 #include "SceneUtils.h"
+#include "MovieSceneCaptureModule.h"
 #include "NotificationManager.h"
 
 DEFINE_LOG_CATEGORY_STATIC(LogClient, Log, All);
@@ -699,6 +700,16 @@ FViewport::FViewport(FViewportClient* InViewportClient):
 	bIsPlayInEditorViewport = false;
 }
 
+FViewport::~FViewport()
+{
+#if WITH_EDITOR
+	if (auto* MovieSceneCapture = GetMovieSceneCapture())
+	{
+		MovieSceneCapture->Close();
+	}
+#endif
+}
+
 bool FViewport::TakeHighResScreenShot()
 {
 	if( GScreenshotResolutionX == 0 && GScreenshotResolutionY == 0 )
@@ -1529,6 +1540,15 @@ void FViewport::SetInitialSize( FIntPoint InitialSizeXY )
 	}
 }
 
+IMovieSceneCaptureInterface* FViewport::GetMovieSceneCapture() const
+{
+	if (MovieSceneCaptureHandle.IsValid())
+	{
+		return IMovieSceneCaptureModule::Get().RetrieveMovieSceneInterface(MovieSceneCaptureHandle);
+	}
+
+	return nullptr;
+}
 
 
 ENGINE_API bool GetViewportScreenShot(FViewport* Viewport, TArray<FColor>& Bitmap, const FIntRect& ViewRect /*= FIntRect()*/)
diff --git a/Engine/Source/Runtime/Engine/Private/UnrealEngine.cpp b/Engine/Source/Runtime/Engine/Private/UnrealEngine.cpp
--- a/Engine/Source/Runtime/Engine/Private/UnrealEngine.cpp
+++ b/Engine/Source/Runtime/Engine/Private/UnrealEngine.cpp
@@ -115,6 +115,7 @@
 #include "TypeContainer.h"
 #endif
 
+#include "MovieSceneCaptureModule.h"
 #include "GameFramework/OnlineSession.h"
 #include "ABTesting.h"
 
@@ -918,14 +919,6 @@ void UEngine::Init(IEngineLoop* InEngineLoop)
 	// Connect the engine analytics provider
 	FEngineAnalytics::Initialize();
 
-#if WITH_EDITOR
-	// register screenshot capture if we are dumping a movie
-	if(GIsDumpingMovie)
-	{
-		HandleScreenshotCapturedDelegateHandle = UGameViewportClient::OnScreenshotCaptured().AddUObject(this, &UEngine::HandleScreenshotCaptured);
-	}
-#endif
-
 	//Load the streaming pause rendering module.
 	FModuleManager::LoadModulePtr<IModuleInterface>(TEXT("StreamingPauseRendering"));
 
@@ -1007,10 +1000,6 @@ void UEngine::PreExit()
 	ShutdownRenderingCVarsCaching();
 	FEngineAnalytics::Shutdown();
 
-#if WITH_EDITOR
-	UGameViewportClient::OnScreenshotCaptured().Remove(HandleScreenshotCapturedDelegateHandle);
-#endif
-
 	if (ScreenSaverInhibitor)
 	{
 		// Resume the thread to avoid a deadlock while waiting for finish.
@@ -1248,66 +1237,6 @@ void UEngine::ParseCommandline()
 	{
 		bDisableAILogging = false;
 	}
-
-	MatineeScreenshotOptions.bStartWithMatineeCapture = false;
-	MatineeScreenshotOptions.bCompressMatineeCapture = false;
-#if WITH_EDITOR
-	if (!GIsEditor && FParse::Value(FCommandLine::Get(), TEXT("-MATINEEAVICAPTURE="), MatineeScreenshotOptions.MatineeCaptureName))
-	{
-		MatineeScreenshotOptions.MatineeCaptureType = EMatineeCaptureType::AVI;
-		MatineeScreenshotOptions.bStartWithMatineeCapture = true;
-	}
-	else if (!GIsEditor && FParse::Value(FCommandLine::Get(), TEXT("-MATINEESSCAPTURE="), MatineeScreenshotOptions.MatineeCaptureName))
-	{
-		MatineeScreenshotOptions.MatineeCaptureType = EMatineeCaptureType::BMP;
-
-		FString MatineeCaptureFormat;
-		if(FParse::Value(FCommandLine::Get(), TEXT("-MATINEESSFORMAT="), MatineeCaptureFormat))
-		{
-			if(MatineeCaptureFormat == TEXT("BMP"))
-			{
-				MatineeScreenshotOptions.MatineeCaptureType = EMatineeCaptureType::BMP;
-			}
-			else if(MatineeCaptureFormat == TEXT("PNG"))
-			{
-				MatineeScreenshotOptions.MatineeCaptureType = EMatineeCaptureType::PNG;
-			}
-			else if(MatineeCaptureFormat == TEXT("JPEG"))
-			{
-				MatineeScreenshotOptions.MatineeCaptureType = EMatineeCaptureType::JPEG;
-			}
-		}
-
-		MatineeScreenshotOptions.bStartWithMatineeCapture = true;
-	}
-
-	if( !GIsEditor && MatineeScreenshotOptions.bStartWithMatineeCapture )
-	{
-		GConfig->GetBool( TEXT("MatineeCreateMovieOptions"), TEXT("HideHUD"), MatineeScreenshotOptions.bHideHud, GEditorPerProjectIni );
-	}
-
-	// If we are capturing a matinee movie and we want to dump the buffer visualization shots too, for on all required functionality
-	if (!GIsEditor && FParse::Param(FCommandLine::Get(), TEXT("MATINEEBUFFERVISUALIZATIONDUMP")))
-	{
-		static IConsoleVariable* CVarDumpFrames = IConsoleManager::Get().FindConsoleVariable(TEXT("r.BufferVisualizationDumpFrames"));
-	
-		if (CVarDumpFrames)
-		{
-			CVarDumpFrames->Set(1, ECVF_SetByCommandline);
-		}
-	}
-
-	if (MatineeScreenshotOptions.bStartWithMatineeCapture)
-	{
-		FParse::Value(FCommandLine::Get(), TEXT("-MATINEEPACKAGE="), MatineeScreenshotOptions.MatineePackageCaptureName);
-	}
-
-	if ( !GIsEditor && FParse::Param(FCommandLine::Get(), TEXT("COMPRESSCAPTURE")) )
-	{
-		MatineeScreenshotOptions.bCompressMatineeCapture = true;
-	}
-#endif
-	MatineeScreenshotOptions.MatineeCaptureFPS = 30;
 }
 
 
@@ -2373,11 +2302,11 @@ bool UEngine::Exec( UWorld* InWorld, const TCHAR* Cmd, FOutputDevice& Ar )
 	{
 		return HandleStatCommand(InWorld, GStatProcessingViewportClient, Cmd, Ar);
 	}
-	else if( FParse::Command(&Cmd,TEXT("STARTMOVIECAPTURE")) && (GEngine->MatineeScreenshotOptions.bStartWithMatineeCapture == true || GIsEditor) )
+	else if( FParse::Command(&Cmd,TEXT("STARTMOVIECAPTURE")) && GIsEditor )
 	{
 		return HandleStartMovieCaptureCommand( Cmd, Ar );
 	}
-	else if( FParse::Command(&Cmd,TEXT("STOPMOVIECAPTURE")) && (GEngine->MatineeScreenshotOptions.bStartWithMatineeCapture == true || GIsEditor) )
+	else if( FParse::Command(&Cmd,TEXT("STOPMOVIECAPTURE")) && GIsEditor )
 	{
 		return HandleStopMovieCaptureCommand( Cmd, Ar );
 	}
@@ -2724,23 +2653,27 @@ bool UEngine::Exec( UWorld* InWorld, const TCHAR* Cmd, FOutputDevice& Ar )
 	return true;
 }
 
+FMovieSceneCaptureHandle GMovieCaptureHandle;
 bool UEngine::HandleStartMovieCaptureCommand( const TCHAR* Cmd, FOutputDevice& Ar )
 {
-	FAVIWriter* AVIWriter = FAVIWriter::GetInstance();
-	if (AVIWriter && !AVIWriter->IsCapturing())
+	if (!GMovieCaptureHandle.IsValid())
 	{
-		AVIWriter->StartCapture();
-		return true;
+		IMovieSceneCaptureInterface* CaptureInterface = IMovieSceneCaptureModule::Get().CreateMovieSceneCapture(GameViewport->Viewport);
+		if (CaptureInterface)
+		{
+			GMovieCaptureHandle = CaptureInterface->GetHandle();
+			return true;
+		}
 	}
 	return false;
 }
 
 bool UEngine::HandleStopMovieCaptureCommand( const TCHAR* Cmd, FOutputDevice& Ar )
 {
-	FAVIWriter* AVIWriter = FAVIWriter::GetInstance();
-	if (AVIWriter && AVIWriter->IsCapturing() && !AVIWriter->IsCapturingSlateRenderer())
+	if (GMovieCaptureHandle.IsValid())
 	{
-		AVIWriter->StopCapture();
+		IMovieSceneCaptureModule::Get().DestroyMovieSceneCapture(GMovieCaptureHandle);
+		GMovieCaptureHandle = FMovieSceneCaptureHandle();
 		return true;
 	}
 	return false;
@@ -11051,99 +10984,6 @@ void FSystemResolution::RequestResolutionChange(int32 InResX, int32 InResY, EWin
 	CVarSystemResolution->Set(*NewValue, ECVF_SetByConsole);
 }
 
-void RemoveAlphaFromColors( TArray<FColor>& Colors )
-{
-	for( FColor& Color : Colors )
-	{
-		Color.A = 255;
-	}
-}
-
-void UEngine::HandleScreenshotCaptured(int32 Width, int32 Height, const TArray<FColor>& Colors)
-{
-#if WITH_EDITOR
-	if(GIsDumpingMovie && Colors.Num() > 0)
-	{
-		struct Local
-		{
-			
-
-			static FString GenerateScreenshotFilename(const FString& Extension)
-			{
-				static const int32 MaxTestScreenShotIndex = 65536;
-				static int32 ScreenShotIndex = 0;
-
-				FString BaseFileName;
-				FScreenshotRequest::CreateViewportScreenShotFilename(BaseFileName);
-
-				for (int32 TestScreenShotIndex = ScreenShotIndex + 1; TestScreenShotIndex < MaxTestScreenShotIndex; ++TestScreenShotIndex)
-				{
-					const FString TestFileName = FString::Printf(TEXT("%s%05i.%s"), *BaseFileName, TestScreenShotIndex, *Extension);
-					if (IFileManager::Get().FileSize(*TestFileName) < 0)
-					{
-						ScreenShotIndex = TestScreenShotIndex;
-						return TestFileName;
-					}
-				}
-
-				UE_LOG(LogEngine, Error, TEXT("Could not generate valid screenshot filename"));
-				return FString();
-			}
-		};
-
-		IImageWrapperModule& ImageWrapperModule = FModuleManager::LoadModuleChecked<IImageWrapperModule>( FName("ImageWrapper") );
-
-		switch(MatineeScreenshotOptions.MatineeCaptureType.GetValue())
-		{
-		default:
-		case EMatineeCaptureType::BMP:
-			{
-				const FString Filename = Local::GenerateScreenshotFilename(TEXT("bmp"));
-				if (Filename.Len() > 0)
-				{
-					FFileHelper::CreateBitmap(*Filename, Width, Height, Colors.GetData());
-				}
-			}
-			break;
-		case EMatineeCaptureType::PNG:
-			{
-				const FString Filename = Local::GenerateScreenshotFilename(TEXT("png"));
-				if (Filename.Len() > 0)
-				{
-					IImageWrapperPtr ImageWrapper = ImageWrapperModule.CreateImageWrapper(EImageFormat::PNG);
-
-					TArray<FColor> FinalColors = Colors;
-					// Matinee PNG export does not support writing alpha values
-					RemoveAlphaFromColors( FinalColors );
-
-					if (ImageWrapper.IsValid() && ImageWrapper->SetRaw(&FinalColors[0], FinalColors.Num() * sizeof(FColor), Width, Height, ERGBFormat::BGRA, 8))
-					{
-						FFileHelper::SaveArrayToFile(ImageWrapper->GetCompressed(), *Filename);
-					}
-				}
-			}
-			break;
-		case EMatineeCaptureType::JPEG:
-			{
-				const FString Filename = Local::GenerateScreenshotFilename(TEXT("jpeg"));
-				if (Filename.Len() > 0)
-				{
-					IImageWrapperPtr ImageWrapper = ImageWrapperModule.CreateImageWrapper(EImageFormat::JPEG);
-					if (ImageWrapper.IsValid() && ImageWrapper->SetRaw(&Colors[0], Colors.Num() * sizeof(FColor), Width, Height, ERGBFormat::BGRA, 8))
-					{
-						FFileHelper::SaveArrayToFile(ImageWrapper->GetCompressed(), *Filename);
-					}
-				}
-			}
-			break;
-		case EMatineeCaptureType::AVI:
-			// Do nothing in this case
-			break;
-		}
-	}
-#endif
-}
-
 //////////////////////////////////////////////////////////////////////////
 // STATS
 
diff --git a/Engine/Source/Runtime/Engine/Public/AVIWriter.h b/Engine/Source/Runtime/Engine/Public/AVIWriter.h
--- a/Engine/Source/Runtime/Engine/Public/AVIWriter.h
+++ b/Engine/Source/Runtime/Engine/Public/AVIWriter.h
@@ -4,79 +4,184 @@
 	AVIWriter.h: Helper class for creating AVI files.
 =============================================================================*/
 
-#ifndef _AVIWRITER_HEADER_
-#define _AVIWRITER_HEADER_
+#pragma once
+
+#include "UnrealClient.h"
 
 DECLARE_LOG_CATEGORY_EXTERN(LogMovieCapture, Warning, All);
 
-/** An event for synchronizing direct show encoding with the game thread */
-extern FEvent* GCaptureSyncEvent;
+DECLARE_DELEGATE_RetVal_TwoParams(FString, FResolveFileFormat, const FString&, const FString&);
+
+/** Creation options for the AVI writer */
+struct FAVIWriterOptions
+{
+	FAVIWriterOptions()
+		: OutputFilename(FPaths::VideoCaptureDir() / TEXT("Capture.avi"))
+		, CaptureFPS(30)
+		, bSynchronizeFrames(false)
+	{}
+
+	/** Output filename */
+	FString OutputFilename;
+
+	/** Constant framerate of the captured video */
+	int32 CaptureFPS;
+
+	/** Optional compression quality, as a value between 0 and 1 */
+	TOptional<float> CompressionQuality;
+
+	/** Optional codec to use for compression */
+	FString CodecName;
+
+	/** When true, the game thread will block until captured frames have been processed by the avi writer */
+	bool bSynchronizeFrames;
+};
+
+/** Data structure representing a captured frame */
+struct FCapturedFrame
+{
+	FCapturedFrame()
+		: StartTimeSeconds(), EndTimeSeconds(), FrameIndex(), FrameProcessedEvent(nullptr)
+	{
+	}
+
+	FCapturedFrame(double InStartTimeSeconds, double InEndTimeSeconds, uint32 InFrameIndex, TArray<FColor> InFrameData);
+	~FCapturedFrame();
+
+	FCapturedFrame(FCapturedFrame&& In) : StartTimeSeconds(In.StartTimeSeconds), EndTimeSeconds(In.EndTimeSeconds), FrameIndex(In.FrameIndex), FrameData(MoveTemp(In.FrameData)), FrameProcessedEvent(In.FrameProcessedEvent) {}
+	FCapturedFrame& operator=(FCapturedFrame&& In) { StartTimeSeconds = In.StartTimeSeconds; EndTimeSeconds = In.EndTimeSeconds; FrameIndex = In.FrameIndex; FrameData = MoveTemp(In.FrameData); FrameProcessedEvent = In.FrameProcessedEvent; return *this; }
+
+	/** The start time of this frame */
+	double StartTimeSeconds;
+	/** The End time of this frame */
+	double EndTimeSeconds;
+	/** The frame index of this frame in the stream */
+	uint32 FrameIndex;
+	/** The frame data itself (empty for a dropped frame) */
+	TArray<FColor> FrameData;
+	/** Triggered when the frame has been processed */
+	FEvent* FrameProcessedEvent;
+};
+
+/** Container for managing captured frames. Temporarily archives frames to the file system when capture rate drops. */
+struct FCapturedFrames
+{
+	/** Construct from a directory to place archives in, and a maximum number of frames we can hold in */
+	FCapturedFrames(const FString& InArchiveDirectory, int32 InMaxInMemoryFrames);
+	~FCapturedFrames();
+
+	/** Add a captured frame to this container. Only to be called from the owner tasread. */
+	void Add(FCapturedFrame Frame);
+
+	/** Read frames from this container (potentially from a thread) */
+	TArray<FCapturedFrame> ReadFrames(uint32 WaitTimeMs);
+
+	/** Retrieve the number of oustanding frames we have not processed yet */
+	int32 GetNumOutstandingFrames() const;
+
+private:
+
+	/** Archive a frame */
+	void ArchiveFrame(FCapturedFrame Frame);
+
+	/** Unarchive a frame represented by the given unique index */
+	TOptional<FCapturedFrame> UnArchiveFrame(uint32 FrameIndex) const;
+	
+	/** Start a task to unarchive some frames */
+	void StartUnArchiving();
+
+	/** The directory in which we will place temporarily archived frames */
+	FString ArchiveDirectory;
+
+	/** A mutex to protect the archived frame indices */
+	mutable FCriticalSection ArchiveFrameMutex;
+	/** Array of unique frame indices that have been archived */
+	TArray<uint32> ArchivedFrames;
+
+	/** The total number of frames that have been archived since capturing started */
+	uint32 TotalArchivedFrames;
 
+	/** An event that triggers when there are in-memory frames ready for collection */
+	FEvent* FrameReady;
+
+	/** Mutex to protect the in memory frames */
+	mutable FCriticalSection InMemoryFrameMutex;
+	/** Array of in-memory captured frames */
+	TArray<FCapturedFrame> InMemoryFrames;
+
+	/** The maximum number of frames we are to store in memory before archiving */
+	int32 MaxInMemoryFrames;
+
+	/** Unarchive task result */
+	TOptional<TFuture<void>> UnarchiveTask;
+};
+
+/** Class responsible for capturing frames from a viewport, and writing them out to an AVI file */
 class FAVIWriter
 {
 protected:
 	/** Protected constructor to avoid abuse. */
-	FAVIWriter() 
-		: ColorBuffer(NULL)
-		, MovieCaptureIndex(0)
-		, bCapturing(false)
-		, bReadyForCapture(false)
-		, bMatineeFinished(false)
-		, CaptureViewport(NULL)
-		, CaptureSlateRenderer(NULL)
+	FAVIWriter(const FAVIWriterOptions& InOptions) 
+		: bCapturing(false)
+		, CaptureViewport(nullptr)
 		, FrameNumber(0)
-		, CurrentAccumSeconds( 0.f )
+		, Options(InOptions)
 	{
-
 	}
 
-	FColor* ColorBuffer;
-	int32 MovieCaptureIndex;
-	bool bCapturing;
-	bool bReadyForCapture;
-	bool bMatineeFinished;
-	TArray< FColor > ViewportColorBuffer;
+	/** Whether we are capturing or not */
+	FThreadSafeBool bCapturing;
+
+	/** The viewport we are capturing from */
 	FViewport* CaptureViewport;
-	class FSlateRenderer* CaptureSlateRenderer;
+
+	/** The current frame number */
 	int32 FrameNumber;
-	float CurrentAccumSeconds;
+	
+	/** Container that manages frames that we have already captured */
+	mutable TUniquePtr<FCapturedFrames> CapturedFrames;
 
 public:
-	ENGINE_API static FAVIWriter* GetInstance();
-
-	uint32 GetWidth()
+	
+	/** Public destruction */
+	ENGINE_API virtual ~FAVIWriter();
+	
+	/** Creation options */
+	FAVIWriterOptions Options;
+
+	/** Create a new avi writer from the specified options */
+	ENGINE_API static FAVIWriter* CreateInstance(const FAVIWriterOptions& InOptions);
+
+	/** Access captured frame data. Safe to be called from any thread. */
+	TArray<FCapturedFrame> GetFrameData(uint32 WaitTimeMs) const
 	{
-		return CaptureViewport ? CaptureViewport->GetSizeXY().X : 0;
+		return CapturedFrames.IsValid() ? CapturedFrames->ReadFrames(WaitTimeMs) : TArray<FCapturedFrame>();
 	}
 
-	uint32 GetHeight()
+	/** Retrieve the number of oustanding frames we have not processed yet */
+	int32 GetNumOutstandingFrames() const
 	{
-		return CaptureViewport ? CaptureViewport->GetSizeXY().Y : 0;
+		return CapturedFrames.IsValid() ? CapturedFrames->GetNumOutstandingFrames() : 0;
 	}
 
-	int32 GetFrameNumber()
+	uint32 GetWidth() const
 	{
-		return FrameNumber;
-	}
-
-	bool IsCapturing()
-	{
-		return bCapturing;
+		return CaptureViewport ? CaptureViewport->GetSizeXY().X : 0;
 	}
 
-	bool IsCapturedMatineeFinished()
+	uint32 GetHeight() const
 	{
-		return bMatineeFinished;
+		return CaptureViewport ? CaptureViewport->GetSizeXY().Y : 0;
 	}
 
-	bool IsCapturingSlateRenderer()
+	int32 GetFrameNumber() const
 	{
-		return CaptureSlateRenderer != NULL;
+		return FrameNumber;
 	}
 
-	void SetCapturedMatineeFinished(bool finished = true)
+	bool IsCapturing() const
 	{
-		bMatineeFinished = finished;
+		return bCapturing;
 	}
 
 	FViewport* GetViewport()
@@ -84,15 +189,10 @@ public:
 		return CaptureViewport;
 	}
 
-	const TArray<FColor>& GetColorBuffer()
-	{
-		return ViewportColorBuffer;
-	}
-
-	virtual void StartCapture(FViewport* Viewport = NULL, const FString& OutputPath = FString()) = 0;
+	ENGINE_API void Update(double FrameTimeSeconds);
+	
+	virtual void StartCapture(FViewport* Viewport = NULL) = 0;
 	virtual void StopCapture() = 0;
 	virtual void Close() = 0;
-	virtual void Update(float DeltaSeconds) = 0;
-
-};
-#endif	//#ifndef _AVIWRITER_HEADER_
\ No newline at end of file
+	virtual void DropFrames(int32 NumFramesToDrop) = 0;
+};
\ No newline at end of file
diff --git a/Engine/Source/Runtime/Engine/Public/UnrealClient.h b/Engine/Source/Runtime/Engine/Public/UnrealClient.h
--- a/Engine/Source/Runtime/Engine/Public/UnrealClient.h
+++ b/Engine/Source/Runtime/Engine/Public/UnrealClient.h
@@ -10,10 +10,12 @@
 #include "HitProxies.h"
 #include "InputCoreTypes.h"
 #include "Engine/EngineBaseTypes.h"
+#include "Runtime/MovieSceneCapture/Public/MovieSceneCaptureHandle.h"
 #include "PopupMethodReply.h"
-
+	
 class FCanvas;
 class FViewportClient;
+class IMovieSceneCaptureInterface;
 
 class SWidget;
 class FCursorReply;
@@ -247,7 +249,7 @@ public:
 	// Constructor.
 	ENGINE_API FViewport(FViewportClient* InViewportClient);
 	// Destructor
-	virtual ~FViewport(){}
+	ENGINE_API virtual ~FViewport();
 
 	//~ Begin FViewport Interface.
 	virtual void* GetWindow() = 0;
@@ -621,6 +623,23 @@ protected:
 	ENGINE_API virtual void ReleaseDynamicRHI() override;
 	ENGINE_API virtual void InitRHI() override;
 	ENGINE_API virtual void ReleaseRHI() override;
+
+public:
+
+	/** Check if this viewport has a movie scene capture implementation or not */
+	bool HasMovieSceneCapture() const { return MovieSceneCaptureHandle.IsValid(); }
+
+	/** Access the current movie scene capture interface, if set. */
+	IMovieSceneCaptureInterface* GetMovieSceneCapture() const;
+
+	/** Set this viewport's capture implementation */
+	void SetMovieSceneCapture(FMovieSceneCaptureHandle Handle) { MovieSceneCaptureHandle = Handle; }
+
+protected:
+
+	/** Movie scene capture implementation handle. */
+	FMovieSceneCaptureHandle MovieSceneCaptureHandle;
+
 };
 
 // Shortcuts for checking the state of both left&right variations of control keys.
diff --git a/Engine/Source/Runtime/Launch/Private/LaunchEngineLoop.cpp b/Engine/Source/Runtime/Launch/Private/LaunchEngineLoop.cpp
--- a/Engine/Source/Runtime/Launch/Private/LaunchEngineLoop.cpp
+++ b/Engine/Source/Runtime/Launch/Private/LaunchEngineLoop.cpp
@@ -2077,7 +2077,6 @@ void FEngineLoop::InitTime()
 	FParse::Value(FCommandLine::Get(),TEXT("FPS="),FixedFPS);
 	if( FixedFPS > 0 )
 	{
-		GEngine->MatineeScreenshotOptions.MatineeCaptureFPS = (int32)FixedFPS;
 		FApp::SetFixedDeltaTime(1 / FixedFPS);
 	}
 
diff --git a/Engine/Source/Runtime/MovieSceneCapture/MovieSceneCapture.Build.cs b/Engine/Source/Runtime/MovieSceneCapture/MovieSceneCapture.Build.cs
new file mode 100644
--- /dev/null
+++ b/Engine/Source/Runtime/MovieSceneCapture/MovieSceneCapture.Build.cs
@@ -0,0 +1,46 @@
+// Copyright 1998-2015 Epic Games, Inc. All Rights Reserved.
+
+using UnrealBuildTool;
+
+public class MovieSceneCapture : ModuleRules
+{
+	public MovieSceneCapture(TargetInfo Target)
+	{
+		PrivateIncludePaths.AddRange(
+			new string[] {
+				"Runtime/MovieSceneCapture/Private"
+			}
+		);
+
+		var ExternalModules = new string[] {
+			"PropertyEditor",
+			"ImageWrapper",
+		};
+		PrivateIncludePathModuleNames.AddRange(ExternalModules);
+		DynamicallyLoadedModuleNames.AddRange(ExternalModules);
+
+		PublicDependencyModuleNames.AddRange(
+			new string[] {
+				"ActorAnimation",
+			}
+		);
+
+		PrivateDependencyModuleNames.AddRange(
+			new string[] {
+				"Core",
+				"CoreUObject",
+				"EditorStyle",
+				"Engine",
+				"UnrealEd",
+				"InputCore",
+				"Json",
+				"JsonUtilities",
+				"MovieScene",
+				"Slate",
+				"SlateCore",
+			}
+		);
+
+
+	}
+}
diff --git a/Engine/Source/Runtime/MovieSceneCapture/Private/AutomatedActorAnimationCapture.cpp b/Engine/Source/Runtime/MovieSceneCapture/Private/AutomatedActorAnimationCapture.cpp
new file mode 100644
--- /dev/null
+++ b/Engine/Source/Runtime/MovieSceneCapture/Private/AutomatedActorAnimationCapture.cpp
@@ -0,0 +1,68 @@
+// Copyright 1998-2015 Epic Games, Inc. All Rights Reserved.
+
+#include "MovieSceneCapturePCH.h"
+#include "ActorAnimationPlayer.h"
+#include "AutomatedActorAnimationCapture.h"
+#include "ErrorCodes.h"
+
+UAutomatedActorAnimationCapture::UAutomatedActorAnimationCapture(const FObjectInitializer& Init)
+	: Super(Init)
+{
+	PlaybackSettings.LoopCount = 0;
+}
+
+void UAutomatedActorAnimationCapture::Initialize(FViewport* InViewport)
+{
+	if (UActorAnimation* ActorAnimationAsset = Cast<UActorAnimation>(ActorAnimation.TryLoad()))
+	{
+		// Set up an instance of the asset for the current world, so we can resolve the bindings
+		AnimationInstance = NewObject<UActorAnimationInstance>(this, "AnimationInstance");
+		AnimationInstance->Initialize(ActorAnimationAsset, GWorld, false);
+
+		// Set up an animation player that can play back the instance
+		AnimationPlayback = NewObject<UActorAnimationPlayer>(this, "AnimationPlayer");
+		AnimationPlayback->Initialize(AnimationInstance, GWorld, PlaybackSettings);
+
+		// Immediately play/pause the animation. This ensures that the scene is set up, ready for playing.
+		// If a preroll has been specified, this will allow temporal PP effects to settle down before capturing
+		AnimationPlayback->Play();
+		AnimationPlayback->Pause();
+
+		PrerollTime = 0.f;
+	}
+	else
+	{
+		FPlatformMisc::RequestExit(FMovieSceneCaptureExitCodes::AssetNotFound);
+	}
+
+	CaptureStrategy = MakeShareable(new FFixedTimeStepCaptureStrategy(Settings.FrameRate));
+	Super::Initialize(InViewport);
+}
+
+void UAutomatedActorAnimationCapture::CaptureFrame(float DeltaSeconds)
+{
+	if (AnimationPlayback)
+	{
+		if (PrerollTime.IsSet())
+		{
+			float& Value = PrerollTime.GetValue();
+			Value += DeltaSeconds;
+			if (Value >= PrerollAmount)
+			{
+				AnimationPlayback->Play();
+				StartCapture();
+				PrerollTime.Reset();
+			}
+		}
+		else
+		{
+			Super::CaptureFrame(DeltaSeconds);
+			AnimationPlayback->Update(DeltaSeconds);
+			
+			if (!AnimationPlayback->IsPlaying())
+			{
+				FPlatformMisc::RequestExit(0);
+			}
+		}
+	}
+}
\ No newline at end of file
diff --git a/Engine/Source/Runtime/MovieSceneCapture/Private/LevelCapture.cpp b/Engine/Source/Runtime/MovieSceneCapture/Private/LevelCapture.cpp
new file mode 100644
--- /dev/null
+++ b/Engine/Source/Runtime/MovieSceneCapture/Private/LevelCapture.cpp
@@ -0,0 +1,12 @@
+// Copyright 1998-2015 Epic Games, Inc. All Rights Reserved.
+
+#include "MovieSceneCapturePCH.h"
+#include "LevelCapture.h"
+
+void ULevelCapture::Initialize(FViewport* InViewport)
+{
+	Super::Initialize(InViewport);
+
+	CaptureStrategy = MakeShareable(new FFixedTimeStepCaptureStrategy(Settings.FrameRate));
+	StartCapture();
+}
\ No newline at end of file
diff --git a/Engine/Source/Runtime/MovieSceneCapture/Private/MovieSceneCapture.cpp b/Engine/Source/Runtime/MovieSceneCapture/Private/MovieSceneCapture.cpp
new file mode 100644
--- /dev/null
+++ b/Engine/Source/Runtime/MovieSceneCapture/Private/MovieSceneCapture.cpp
@@ -0,0 +1,343 @@
+// Copyright 1998-2015 Epic Games, Inc. All Rights Reserved.
+
+#include "MovieSceneCapturePCH.h"
+#include "MovieSceneCapture.h"
+#include "ImageWrapper.h"
+#include "MovieSceneCaptureModule.h"
+
+struct FUniqueMovieSceneCaptureHandle : FMovieSceneCaptureHandle
+{
+	FUniqueMovieSceneCaptureHandle()
+	{
+		/// Start IDs at index 1 since 0 is deemed invalid
+		static uint32 Unique = 1;
+		ID = Unique++;
+	}
+};
+
+FMovieSceneCaptureSettings::FMovieSceneCaptureSettings()
+	: Resolution(1280, 720)
+{
+	OutputDirectory.Path = FPaths::VideoCaptureDir();
+	OutputFormat = NSLOCTEXT("MovieCapture", "DefaultFormat", "MovieCapture_{width}x{height}_{quality}").ToString();
+	FrameRate = 24;
+	CaptureType = EMovieCaptureType::AVI;
+	bUseCompression = true;
+	CompressionQuality = 1.f;
+	bEnableTextureStreaming = true;
+	bCinematicMode = true;
+	bAllowMovement = true;
+	bAllowTurning = true;
+	bShowPlayer = true;
+	bShowHUD = true;
+}
+
+UMovieSceneCapture::UMovieSceneCapture(const FObjectInitializer& Initializer)
+	: Super(Initializer)
+{
+	bBufferVisualizationDumpFrames = false;
+
+	TArray<FString> Tokens, Switches;
+	FCommandLine::Parse( FCommandLine::Get(), Tokens, Switches );
+	for (auto& Switch : Switches)
+	{
+		AdditionalCommandLineArguments.AppendChar('-');
+		AdditionalCommandLineArguments.Append(Switch);
+		AdditionalCommandLineArguments.AppendChar(' ');
+	}
+	AdditionalCommandLineArguments += TEXT("-PIEVIACONSOLE -nomovie ");
+
+	// renderer overrides - hack
+	AdditionalCommandLineArguments += FParse::Param(FCommandLine::Get(), TEXT("d3d11"))		?	TEXT("-d3d11 ")		: TEXT("");
+	AdditionalCommandLineArguments += FParse::Param(FCommandLine::Get(), TEXT("sm5"))		?	TEXT("-sm5 ")		: TEXT("");
+	AdditionalCommandLineArguments += FParse::Param(FCommandLine::Get(), TEXT("dx11"))		?	TEXT("-dx11 ")		: TEXT("");
+	AdditionalCommandLineArguments += FParse::Param(FCommandLine::Get(), TEXT("d3d10"))		?	TEXT("-d3d10 ")		: TEXT("");
+	AdditionalCommandLineArguments += FParse::Param(FCommandLine::Get(), TEXT("sm4"))		?	TEXT("-sm4 ")		: TEXT("");
+	AdditionalCommandLineArguments += FParse::Param(FCommandLine::Get(), TEXT("dx10"))		?	TEXT("-dx10 ")		: TEXT("");
+	AdditionalCommandLineArguments += FParse::Param(FCommandLine::Get(), TEXT("opengl"))	?	TEXT("-opengl ")	: TEXT("");
+	AdditionalCommandLineArguments += FParse::Param(FCommandLine::Get(), TEXT("opengl3"))	?	TEXT("-opengl3 ")	: TEXT("");
+	AdditionalCommandLineArguments += FParse::Param(FCommandLine::Get(), TEXT("opengl4"))	?	TEXT("-opengl4 ")	: TEXT("");
+
+	Handle = FUniqueMovieSceneCaptureHandle();
+
+	CurrentTimeSeconds = 0;
+}
+
+void UMovieSceneCapture::Initialize(FViewport* InViewport)
+{
+	Viewport = InViewport;
+}
+
+void UMovieSceneCapture::StartCapture()
+{
+	if (!CaptureStrategy.IsValid())
+	{
+		CaptureStrategy = MakeShareable(new FRealTimeCaptureStrategy(Settings.FrameRate));
+	}
+
+	CaptureStrategy->OnStart();
+
+	CurrentTimeSeconds = 0;
+
+	if (bBufferVisualizationDumpFrames)
+	{
+		static IConsoleVariable* CVarDumpFrames = IConsoleManager::Get().FindConsoleVariable(TEXT("r.BufferVisualizationDumpFrames"));
+		if (CVarDumpFrames)
+		{
+			CVarDumpFrames->Set(1, ECVF_SetByCommandline);
+		}
+	}
+
+	CachedMetrics.Width = Viewport->GetSizeXY().X;
+	CachedMetrics.Height = Viewport->GetSizeXY().Y;
+
+	if (Settings.CaptureType == EMovieCaptureType::AVI)
+	{
+		FAVIWriterOptions Options;
+		Options.OutputFilename = ResolveUniqueFilename();
+		Options.CaptureFPS = Settings.FrameRate;
+		Options.CodecName = Settings.Codec;
+		Options.bSynchronizeFrames = CaptureStrategy->ShouldSynchronizeFrames();
+
+		if (Settings.bUseCompression)
+		{
+			Options.CompressionQuality = Settings.CompressionQuality;
+		}
+
+		AVIWriter.Reset(FAVIWriter::CreateInstance(Options));
+		AVIWriter->StartCapture(Viewport);
+	}
+	else
+	{
+		Viewport->SetMovieSceneCapture(Handle);
+	}
+}
+
+void UMovieSceneCapture::CaptureFrame(float DeltaSeconds)
+{
+	if (!CaptureStrategy.IsValid())
+	{
+		return;
+	}
+
+	CurrentTimeSeconds += DeltaSeconds;
+
+	if (CaptureStrategy->ShouldPresent(CurrentTimeSeconds, CachedMetrics.Frame))
+	{
+		uint32 NumDroppedFrames = CaptureStrategy->GetDroppedFrames(CurrentTimeSeconds, CachedMetrics.Frame);
+		CachedMetrics.Frame += NumDroppedFrames;
+
+		CaptureStrategy->OnPresent(CurrentTimeSeconds, CachedMetrics.Frame);
+		++CachedMetrics.Frame;
+		
+		if (AVIWriter)
+		{
+			AVIWriter->DropFrames(NumDroppedFrames);
+			AVIWriter->Update(CurrentTimeSeconds);
+		}
+		else
+		{
+			TArray<FColor> Data;
+			Viewport->ReadPixels(Data, FReadSurfaceDataFlags());
+			CaptureSnapshot(Data);
+		}
+	}
+}
+
+void UMovieSceneCapture::StopCapture()
+{
+	CaptureStrategy->OnStop();
+	CaptureStrategy = nullptr;
+
+	if (AVIWriter)
+	{
+		AVIWriter->StopCapture();
+		AVIWriter.Reset();
+	}
+}
+
+void UMovieSceneCapture::Close()
+{
+	StopCapture();
+	IMovieSceneCaptureModule::Get().OnMovieSceneCaptureFinished(this);
+}
+
+void UMovieSceneCapture::CaptureSnapshot(const TArray<FColor>& Colors)
+{
+	if (Colors.Num() == 0)
+	{
+		return;
+	}
+
+	FString Filename = ResolveUniqueFilename();
+	if (Filename.IsEmpty())
+	{
+		return;
+	}
+
+	IImageWrapperModule& ImageWrapperModule = FModuleManager::LoadModuleChecked<IImageWrapperModule>( FName("ImageWrapper") );
+
+	// clamp to 1-100 range
+	const int32 ImageCompressionQuality = FMath::Clamp<int32>(Settings.CompressionQuality * 100, 1, 100);
+
+	switch (Settings.CaptureType)
+	{
+	case EMovieCaptureType::BMP:
+		FFileHelper::CreateBitmap(*Filename, CachedMetrics.Width, CachedMetrics.Height, Colors.GetData());
+		break;
+
+	case EMovieCaptureType::PNG:
+		{
+			IImageWrapperPtr ImageWrapper = ImageWrapperModule.CreateImageWrapper(EImageFormat::PNG);
+
+			TArray<FColor> FinalColors = Colors;
+			for (FColor& Color : FinalColors)
+			{
+				Color.A = 255;
+			}
+
+			if (ImageWrapper.IsValid() && ImageWrapper->SetRaw(&FinalColors[0], FinalColors.Num() * sizeof(FColor), CachedMetrics.Width, CachedMetrics.Height, ERGBFormat::BGRA, 8))
+			{
+				FFileHelper::SaveArrayToFile(ImageWrapper->GetCompressed(ImageCompressionQuality), *Filename);
+			}
+		}
+		break;
+
+	case EMovieCaptureType::JPEG:
+		{
+			IImageWrapperPtr ImageWrapper = ImageWrapperModule.CreateImageWrapper(EImageFormat::JPEG);
+			if (ImageWrapper.IsValid() && ImageWrapper->SetRaw(&Colors[0], Colors.Num() * sizeof(FColor), CachedMetrics.Width, CachedMetrics.Height, ERGBFormat::BGRA, 8))
+			{
+				FFileHelper::SaveArrayToFile(ImageWrapper->GetCompressed(ImageCompressionQuality), *Filename);
+			}
+		}
+		break;
+
+	default:
+		break;
+	}
+}
+
+const TCHAR* UMovieSceneCapture::GetDefaultFileExtension() const
+{
+	switch (Settings.CaptureType)
+	{
+		case EMovieCaptureType::BMP:	return TEXT(".bmp");
+		case EMovieCaptureType::PNG:	return TEXT(".png");
+		case EMovieCaptureType::JPEG:	return TEXT(".jpeg");
+#if PLATFORM_MAC
+		default:						return TEXT(".mov");
+#else
+		default:						return TEXT(".avi");
+#endif
+	}
+}
+
+FString UMovieSceneCapture::ResolveFileFormat(const FString& Folder, const FString& Format) const
+{
+	TMap<FString, FStringFormatArg> Mappings;
+	Mappings.Add(TEXT("fps"), FString::Printf(TEXT("%d"), Settings.FrameRate));
+	Mappings.Add(TEXT("frame"), FString::Printf(TEXT("%04d"), CachedMetrics.Frame+1));
+	Mappings.Add(TEXT("width"), FString::Printf(TEXT("%d"), CachedMetrics.Width));
+	Mappings.Add(TEXT("height"), FString::Printf(TEXT("%d"), CachedMetrics.Height));
+
+	if (Settings.bUseCompression)
+	{
+		Mappings.Add(TEXT("quality"), FString::Printf(TEXT("%.2f"), Settings.CompressionQuality));
+	}
+	else
+	{
+		Mappings.Add(TEXT("quality"), NSLOCTEXT("MovieCapture", "Uncompressed", "Uncompressed").ToString());
+	}
+
+	return Folder / FString::Format(*Format, Mappings);
+}
+
+FString UMovieSceneCapture::ResolveUniqueFilename()
+{
+	FString BaseFilename = ResolveFileFormat(Settings.OutputDirectory.Path, Settings.OutputFormat);
+	FString ThisTry = BaseFilename + GetDefaultFileExtension();
+
+	if (IFileManager::Get().FileSize(*ThisTry) != -1)
+	{
+		return ThisTry;
+	}
+
+	uint32 Index = 2;
+	for (;;)
+	{
+		ThisTry = BaseFilename + FString::Printf(TEXT(" %02d"), Index) + GetDefaultFileExtension();
+
+		// If the file doesn't exist, we can use that, else, increment the index and try again
+		if (IFileManager::Get().FileSize(*ThisTry) == -1)
+		{
+			return ThisTry;
+		}
+
+		++Index;
+	}
+
+	return ThisTry;
+}
+
+FFixedTimeStepCaptureStrategy::FFixedTimeStepCaptureStrategy(uint32 InTargetFPS)
+	: TargetFPS(InTargetFPS)
+{
+}
+
+void FFixedTimeStepCaptureStrategy::OnStart()
+{
+	FApp::SetFixedDeltaTime(1.0 / TargetFPS);
+	FApp::SetUseFixedTimeStep(true);
+}
+
+void FFixedTimeStepCaptureStrategy::OnStop()
+{
+	FApp::SetUseFixedTimeStep(false);
+}
+
+void FFixedTimeStepCaptureStrategy::OnPresent(double CurrentTimeSeconds, uint32 FrameIndex)
+{
+}
+
+bool FFixedTimeStepCaptureStrategy::ShouldPresent(double CurrentTimeSeconds, uint32 FrameIndex) const
+{
+	return true;
+}
+
+int32 FFixedTimeStepCaptureStrategy::GetDroppedFrames(double CurrentTimeSeconds, uint32 FrameIndex) const
+{
+	return 0;
+}
+
+FRealTimeCaptureStrategy::FRealTimeCaptureStrategy(uint32 InTargetFPS)
+	: NextPresentTimeS(0), FrameLength(1.0 / InTargetFPS)
+{
+}
+
+void FRealTimeCaptureStrategy::OnStart()
+{
+}
+
+void FRealTimeCaptureStrategy::OnStop()
+{
+}
+
+void FRealTimeCaptureStrategy::OnPresent(double CurrentTimeSeconds, uint32 FrameIndex)
+{
+}
+
+bool FRealTimeCaptureStrategy::ShouldPresent(double CurrentTimeSeconds, uint32 FrameIndex) const
+{
+	return CurrentTimeSeconds >= FrameIndex * FrameLength;
+}
+
+int32 FRealTimeCaptureStrategy::GetDroppedFrames(double CurrentTimeSeconds, uint32 FrameIndex) const
+{
+	uint32 ThisFrame = FMath::FloorToInt(CurrentTimeSeconds / FrameLength);
+	if (ThisFrame > FrameIndex)
+	{
+		return ThisFrame - FrameIndex;
+	}
+	return 0;
+}
\ No newline at end of file
diff --git a/Engine/Source/Runtime/MovieSceneCapture/Private/MovieSceneCaptureModule.cpp b/Engine/Source/Runtime/MovieSceneCapture/Private/MovieSceneCaptureModule.cpp
new file mode 100644
--- /dev/null
+++ b/Engine/Source/Runtime/MovieSceneCapture/Private/MovieSceneCaptureModule.cpp
@@ -0,0 +1,604 @@
+// Copyright 1998-2015 Epic Games, Inc. All Rights Reserved.
+
+#include "MovieSceneCapturePCH.h"
+#include "MovieSceneCapture.h"
+#include "MovieSceneCaptureModule.h"
+#include "EditorStyle.h"
+#include "SDockTab.h"
+#include "Editor.h"
+#include "PropertyEditing.h"
+#include "JsonObjectConverter.h"
+#include "INotificationWidget.h"
+#include "SNotificationList.h"
+#include "NotificationManager.h"
+
+#include "SlateExtras.h"
+
+#define LOCTEXT_NAMESPACE "MovieSceneCapture"
+
+const FName MovieSceneRenderTabId(TEXT("MovieSceneRenderTab"));
+
+DECLARE_DELEGATE_RetVal_OneParam(FText, FOnStartCapture, UMovieSceneCapture*);
+
+class SRenderMovieSceneSettings : public SCompoundWidget, public FGCObject
+{
+	SLATE_BEGIN_ARGS(SRenderMovieSceneSettings) : _InitialObject(nullptr) {}
+		SLATE_EVENT(FOnStartCapture, OnStartCapture)
+		SLATE_ARGUMENT(UMovieSceneCapture*, InitialObject)
+	SLATE_END_ARGS()
+
+	void Construct(const FArguments& InArgs)
+	{
+		FPropertyEditorModule& PropertyEditor = FModuleManager::LoadModuleChecked<FPropertyEditorModule>("PropertyEditor");
+
+		FDetailsViewArgs DetailsViewArgs;
+		DetailsViewArgs.bUpdatesFromSelection = false;
+		DetailsViewArgs.bLockable = false;
+		DetailsViewArgs.NameAreaSettings = FDetailsViewArgs::HideNameArea;
+		DetailsViewArgs.ViewIdentifier = "RenderMovieScene";
+
+		DetailView = PropertyEditor.CreateDetailView(DetailsViewArgs);
+
+		OnStartCapture = InArgs._OnStartCapture;
+
+		ChildSlot
+		[
+			SNew(SVerticalBox)
+
+			+ SVerticalBox::Slot()
+			[
+				DetailView.ToSharedRef()
+			]
+
+			+ SVerticalBox::Slot()
+			.AutoHeight()
+			[
+				SAssignNew(ErrorText, STextBlock)
+				.Visibility(EVisibility::Hidden)
+			]
+
+			+ SVerticalBox::Slot()
+			.AutoHeight()
+			.HAlign(HAlign_Right)
+			.Padding(5.f)
+			[
+				SNew(SButton)
+				.ContentPadding(FMargin(10, 5))
+				.Text(LOCTEXT("Export", "Capture Movie"))
+				.OnClicked(this, &SRenderMovieSceneSettings::OnStartClicked)
+			]
+			
+		];
+
+		if (InArgs._InitialObject)
+		{
+			SetObject(InArgs._InitialObject);
+		}
+	}
+
+	void SetObject(UMovieSceneCapture* InMovieSceneCapture)
+	{
+		MovieSceneCapture = InMovieSceneCapture;
+
+		DetailView->SetObject(InMovieSceneCapture);
+
+		ErrorText->SetText(FText());
+		ErrorText->SetVisibility(EVisibility::Hidden);
+	}
+
+	virtual void AddReferencedObjects( FReferenceCollector& Collector ) override
+	{
+		Collector.AddReferencedObject(MovieSceneCapture);
+	}
+
+private:
+
+	FReply OnStartClicked()
+	{
+		FText Error;
+		if (OnStartCapture.IsBound())
+		{
+			Error = OnStartCapture.Execute(MovieSceneCapture);
+		}
+
+		ErrorText->SetText(Error);
+		ErrorText->SetVisibility(Error.IsEmpty() ? EVisibility::Hidden : EVisibility::Visible);
+
+		return FReply::Handled();
+	}
+
+	TSharedPtr<IDetailsView> DetailView;
+	TSharedPtr<STextBlock> ErrorText;
+	FOnStartCapture OnStartCapture;
+	UMovieSceneCapture* MovieSceneCapture;
+};
+
+
+
+DECLARE_DELEGATE_OneParam(FOnProcessClosed, int32);
+
+class SCaptureMovieNotification : public SCompoundWidget, public INotificationWidget
+{
+public:
+	SLATE_BEGIN_ARGS(SCaptureMovieNotification){}
+		SLATE_EVENT(FOnProcessClosed, OnProcessClosed)
+		SLATE_EVENT(FSimpleDelegate, OnNotificationClosed)
+		SLATE_ARGUMENT(FString, BrowseToFolder)
+	SLATE_END_ARGS()
+
+	void Construct(const FArguments& InArgs, FProcHandle InProcHandle)
+	{
+		OnNotificationClosed = InArgs._OnNotificationClosed;
+		OnProcessClosed = InArgs._OnProcessClosed;
+		ProcHandle = InProcHandle;
+
+		FString BrowseToFolder = FPaths::ConvertRelativePathToFull(InArgs._BrowseToFolder);
+		auto OnBrowseToFolder = [=]{
+			FPlatformProcess::ExploreFolder(*BrowseToFolder);
+		};
+
+		ChildSlot
+		[
+			SNew(SBorder)
+			.Padding(FMargin(15.0f))
+			.BorderImage(FCoreStyle::Get().GetBrush("NotificationList.ItemBackground"))
+			[
+				SNew(SVerticalBox)
+
+				+ SVerticalBox::Slot()
+				.Padding(FMargin(0,0,0,5.0f))
+				.HAlign(HAlign_Right)
+				.AutoHeight()
+				[
+					SNew(SHorizontalBox)
+
+					+ SHorizontalBox::Slot()
+					.VAlign(VAlign_Center)
+					[
+						SAssignNew(TextBlock, STextBlock)
+						.Font(FCoreStyle::Get().GetFontStyle(TEXT("NotificationList.FontBold")))
+						.Text(LOCTEXT("RenderingVideo", "Capturing video"))
+					]
+
+					+ SHorizontalBox::Slot()
+					.AutoWidth()
+					.Padding(FMargin(15.f,0,0,0))
+					[
+						SNew(SThrobber)
+						.Visibility( this, &SCaptureMovieNotification::GetThrobberVisibility )
+					]
+				]
+
+				+ SVerticalBox::Slot()
+				.AutoHeight()
+				[
+					SNew(SHorizontalBox)
+
+					+ SHorizontalBox::Slot()
+					.AutoWidth()
+					.VAlign(VAlign_Center)
+					.Padding(FMargin(0,0,50.0f,0))
+					[
+						SNew(SHyperlink)
+						.Visibility( this, &SCaptureMovieNotification::GetHyperlinkVisibility )
+						.Text(LOCTEXT("OpenFolder", "Open Capture Folder..."))
+						.OnNavigate_Lambda(OnBrowseToFolder)
+					]
+
+					+ SHorizontalBox::Slot()
+					[
+						SNew(SSpacer)
+					]
+
+					+ SHorizontalBox::Slot()
+					.AutoWidth()
+					.VAlign(VAlign_Center)
+					[
+						SAssignNew(Button, SButton)
+						.OnClicked(this, &SCaptureMovieNotification::ButtonClicked)
+					]
+				]
+			]
+		];
+	}
+
+	virtual void Tick( const FGeometry& AllottedGeometry, const double InCurrentTime, const float InDeltaTime )
+	{
+		if (State != SNotificationItem::CS_Pending)
+		{
+			return;
+		}
+
+		if (!ProcHandle.IsValid() || !FPlatformProcess::IsProcRunning(ProcHandle))
+		{
+			int32 RetCode = 0;
+			FPlatformProcess::GetProcReturnCode(ProcHandle, &RetCode);
+			OnProcessClosed.ExecuteIfBound(RetCode);
+
+			if (RetCode == 0)
+			{
+				TextBlock->SetText(LOCTEXT("Finished", "Capture Finished"));
+			}
+		}
+	}
+
+	virtual void OnSetCompletionState(SNotificationItem::ECompletionState InState)
+	{
+		State = InState;
+		if (State == SNotificationItem::CS_Pending)
+		{
+			Button->SetContent(SNew(STextBlock).Text(LOCTEXT("StopButton", "Stop")));
+		}
+		else
+		{
+			Button->SetContent(SNew(STextBlock).Text(LOCTEXT("CloseButton", "Close")));
+		}
+	}
+
+	virtual TSharedRef< SWidget > AsWidget()
+	{
+		return AsShared();
+	}
+
+private:
+
+	EVisibility GetThrobberVisibility() const
+	{
+		return State == SNotificationItem::CS_Pending ? EVisibility::Visible : EVisibility::Collapsed;
+	}
+
+	EVisibility GetHyperlinkVisibility() const
+	{
+		return State != SNotificationItem::CS_Pending ? EVisibility::Visible : EVisibility::Hidden;
+	}
+
+	FReply ButtonClicked()
+	{
+		if (State == SNotificationItem::CS_Pending)
+		{
+			FPlatformProcess::TerminateProc(ProcHandle);
+		}
+		else
+		{
+			OnNotificationClosed.ExecuteIfBound();
+		}
+		return FReply::Handled();
+	}
+	
+private:
+	TSharedPtr<SButton> Button;
+	TSharedPtr<STextBlock> TextBlock;
+	SNotificationItem::ECompletionState State;
+	FOnProcessClosed OnProcessClosed;
+	FSimpleDelegate OnNotificationClosed;
+	FProcHandle ProcHandle;
+};
+
+class FMovieSceneCaptureModule : public IMovieSceneCaptureModule, public FTickableGameObject, public FGCObject
+{
+public:
+
+	void TickCapture(UMovieSceneCapture* Capture)
+	{
+		ExistingCaptures.Add(Capture);
+	}
+
+	void StopTickingCapture(UMovieSceneCapture* Capture)
+	{
+		ExistingCaptures.Remove(Capture);
+	}
+
+private:
+
+	TArray<UMovieSceneCapture*> ExistingCaptures;
+
+	/** Tickable interface */
+	virtual bool IsTickableInEditor() const override { return false; }
+	virtual bool IsTickable() const override { return true; }
+	virtual bool IsTickableWhenPaused() const override { return false; }
+	virtual TStatId GetStatId() const override { RETURN_QUICK_DECLARE_CYCLE_STAT(FCachePutAsyncWorker, STATGROUP_ThreadPoolAsyncTasks); }
+	virtual void Tick( float DeltaTime ) override
+	{
+		for (auto* Obj : ExistingCaptures)
+		{
+			Obj->CaptureFrame(DeltaTime);
+		}
+	}
+
+	virtual void AddReferencedObjects(FReferenceCollector& Collector) override
+	{
+		for (auto* Obj : ExistingCaptures)
+		{
+			Collector.AddReferencedObject(Obj);
+		}
+	}
+
+	/** End Tickable interface */
+	virtual void StartupModule() override
+	{
+		FCoreDelegates::OnPreExit.AddRaw(this, &FMovieSceneCaptureModule::PreExit);
+	}
+
+	void PreExit()
+	{
+		TArray<UMovieSceneCapture*> ExistingCapturesCopy;
+		Swap(ExistingCaptures, ExistingCapturesCopy);
+
+		for (auto* Obj : ExistingCapturesCopy)
+		{
+			Obj->Close();
+		}
+	}
+
+	virtual void PreUnloadCallback() override
+	{
+		InProgressCaptureNotification = nullptr;
+		FCoreDelegates::OnPreExit.RemoveAll(this);
+		PreExit();
+	}
+
+	virtual IMovieSceneCaptureInterface* InitializeFromCommandLine() override
+	{
+		if (GIsEditor)
+		{
+			return nullptr;
+		}
+
+		FString ManifestPath;
+		if (!FParse::Value(FCommandLine::Get(), TEXT("-MovieSceneCaptureManifest="), ManifestPath) || ManifestPath.IsEmpty())
+		{
+			return nullptr;
+		}
+
+		UMovieSceneCapture* Capture = nullptr;
+		FString Json;
+		if (FFileHelper::LoadFileToString(Json, *ManifestPath))
+		{
+			TSharedPtr<FJsonObject> RootObject;
+			TSharedRef<TJsonReader<> > JsonReader = TJsonReaderFactory<>::Create(Json);
+			if (FJsonSerializer::Deserialize(JsonReader, RootObject) && RootObject.IsValid())
+			{
+				auto TypeField = RootObject->TryGetField(TEXT("Type"));
+				if (!TypeField.IsValid())
+				{
+					return nullptr;
+				}
+
+				UClass* Class = FindObject<UClass>(nullptr, *TypeField->AsString());
+				if (!Class)
+				{
+					return nullptr;
+				}
+
+				Capture = NewObject<UMovieSceneCapture>(GetTransientPackage(), Class);
+				if (!Capture)
+				{
+					return nullptr;
+				}
+
+				auto DataField = RootObject->TryGetField(TEXT("Data"));
+				if (!DataField.IsValid())
+				{
+					return nullptr;
+				}
+
+				if (!FJsonObjectConverter::JsonAttributesToUStruct(DataField->AsObject()->Values, Class, Capture, 0, 0))
+				{
+					return nullptr;
+				}
+
+				TickCapture(Capture);
+				return Capture;
+			}
+		}
+
+		return nullptr;
+	}
+
+	virtual IMovieSceneCaptureInterface* CreateMovieSceneCapture(FViewport* Viewport) override
+	{
+		UMovieSceneCapture* Capture = NewObject<UMovieSceneCapture>(GetTransientPackage());
+		TickCapture(Capture);
+		Capture->Initialize(Viewport);
+		Capture->StartCapture();
+		return Capture;
+	}
+
+	virtual IMovieSceneCaptureInterface* RetrieveMovieSceneInterface(FMovieSceneCaptureHandle Handle)
+	{
+		auto** Existing = ExistingCaptures.FindByPredicate([&](const UMovieSceneCapture* In){ return In->GetHandle() == Handle; });
+		return Existing ? *Existing : nullptr;
+	}
+
+	virtual void OnMovieSceneCaptureFinished(IMovieSceneCaptureInterface* Capture)
+	{
+		StopTickingCapture(static_cast<UMovieSceneCapture*>(Capture));
+	}
+
+	virtual void DestroyMovieSceneCapture(FMovieSceneCaptureHandle Handle)
+	{
+		// Calling Close() can remove itself from the array
+		auto Captures = ExistingCaptures;
+		for (auto* Existing : Captures)
+		{
+			if (Existing->GetHandle() == Handle)
+			{
+				Existing->Close();
+			}
+		}
+	}
+
+	virtual void OpenCaptureSettings(const TSharedRef<FTabManager>& TabManager, UMovieSceneCapture* CaptureObject) override
+	{
+		TSharedPtr<SWindow> ExistingWindow = CaptureSettingsWindow.Pin();
+		if (ExistingWindow.IsValid())
+		{
+			ExistingWindow->BringToFront();
+		}
+		else
+		{
+			ExistingWindow = SNew(SWindow)
+				.Title( LOCTEXT("RenderMovieSettingsTitle", "Render Movie Settings") )
+				.HasCloseButton(true)
+				.SupportsMaximize(false)
+				.SupportsMinimize(false)
+				.ClientSize(FVector2D(350, 650));
+
+			TSharedPtr<SDockTab> OwnerTab = TabManager->GetOwnerTab();
+			TSharedPtr<SWindow> RootWindow = OwnerTab.IsValid() ? OwnerTab->GetParentWindow() : TSharedPtr<SWindow>();
+			if(RootWindow.IsValid())
+			{
+				FSlateApplication::Get().AddWindowAsNativeChild(ExistingWindow.ToSharedRef(), RootWindow.ToSharedRef());
+			}
+			else
+			{
+				FSlateApplication::Get().AddWindow(ExistingWindow.ToSharedRef());
+			}
+		}
+
+		ExistingWindow->SetContent(
+			SNew(SRenderMovieSceneSettings)
+			.InitialObject(CaptureObject)
+			.OnStartCapture_Raw(this, &FMovieSceneCaptureModule::OnStartCapture)
+		);
+
+		CaptureSettingsWindow = ExistingWindow;
+	}
+
+	void OnNotificationClosed()
+	{
+		InProgressCaptureNotification->Fadeout();
+		InProgressCaptureNotification = nullptr;
+	}
+
+	void OnMovieCaptureProcessClosed(int32 RetCode)
+	{
+		if (RetCode == 0)
+		{
+			InProgressCaptureNotification->SetCompletionState(SNotificationItem::CS_Success);
+		}
+		else
+		{
+			// todo: error to message log
+			InProgressCaptureNotification->SetCompletionState(SNotificationItem::CS_Fail);
+		}
+	}
+
+	FText OnStartCapture(UMovieSceneCapture* CaptureObject)
+	{
+		if (InProgressCaptureNotification.IsValid())
+		{
+			return LOCTEXT("AlreadyCapturing", "There is already a movie scene capture process open. Please close it and try again.");
+		}
+
+		// If buffer visualization dumping is enabled, we need to tell capture process to enable it too
+		static const auto CVarDumpFrames = IConsoleManager::Get().FindTConsoleVariableDataInt(TEXT("r.BufferVisualizationDumpFrames"));
+		if (CVarDumpFrames && CVarDumpFrames->GetValueOnGameThread())
+		{
+			CaptureObject->bBufferVisualizationDumpFrames = true;
+		}
+
+		// Save out the capture manifest to json
+		FString Filename = FPaths::GameSavedDir() / TEXT("MovieSceneCapture/Manifest.json");
+
+		TSharedRef<FJsonObject> Object = MakeShareable(new FJsonObject);
+		if (FJsonObjectConverter::UStructToJsonObject(CaptureObject->GetClass(), CaptureObject, Object, 0, 0))
+		{
+			TSharedRef<FJsonObject> RootObject = MakeShareable(new FJsonObject);
+			RootObject->SetField(TEXT("Type"), MakeShareable(new FJsonValueString(CaptureObject->GetClass()->GetPathName())));
+			RootObject->SetField(TEXT("Data"), MakeShareable(new FJsonValueObject(Object)));
+
+			FString Json;
+			TSharedRef<TJsonWriter<> > JsonWriter = TJsonWriterFactory<>::Create(&Json, 0);
+			if (FJsonSerializer::Serialize( RootObject, JsonWriter ))
+			{
+				FFileHelper::SaveStringToFile(Json, *Filename);
+			}
+		}
+		else
+		{
+			return LOCTEXT("UnableToSaveCaptureManifest", "Unable to save capture manifest");
+		}
+
+		// this parameter tells UE4Editor to run in game mode
+		FString EditorCommandLine = FString::Printf(TEXT("-MovieSceneCaptureManifest=\"%s\" -game"), *Filename);
+		if (!CaptureObject->AdditionalCommandLineArguments.IsEmpty())
+		{
+			EditorCommandLine.AppendChar(' ');
+			EditorCommandLine.Append(CaptureObject->AdditionalCommandLineArguments);
+		}
+
+		if (!CaptureObject->Settings.bEnableTextureStreaming)
+		{
+			EditorCommandLine.Append(TEXT(" -NoTextureStreaming"));
+		}
+		
+		EditorCommandLine += FString::Printf(TEXT(" -ResX=%d"), CaptureObject->Settings.Resolution.ResX);
+		EditorCommandLine += FString::Printf(TEXT(" -ResY=%d"), CaptureObject->Settings.Resolution.ResY);
+
+		// PlayMap - in editor. Not working yet.
+
+		// FMovieSceneCaptureModule& Module = static_cast<FMovieSceneCaptureModule&>(IMovieSceneCaptureModule::Get());
+		// Module.TickCapture(CaptureObject);
+
+		// UGameViewportClient::OnViewportCreated().AddLambda([=](){
+
+		// 	for (const FWorldContext& Context : GEngine->GetWorldContexts())
+		// 	{
+		// 		if (Context.WorldType == EWorldType::PIE)
+		// 		{
+		// 			CaptureObject->Initialize(Context.GameViewport->Viewport);
+		// 			break;
+		// 		}
+		// 	}
+
+		// });
+
+		// FEditorDelegates::EndPIE.AddLambda([&](bool bIsSimulating){
+		// 	Module.StopTickingCapture(CaptureObject);
+		// });
+
+		// GEditor->PlayMap();
+
+		CaptureObject->SaveConfig();
+
+		FString Params;
+		if (FPaths::IsProjectFilePathSet())
+		{
+			Params = FString::Printf(TEXT("\"%s\" %s %s"), *FPaths::GetProjectFilePath(), *EditorCommandLine, *FCommandLine::GetSubprocessCommandline());
+		}
+		else
+		{
+			Params = FString::Printf(TEXT("%s %s %s"), FApp::GetGameName(), *EditorCommandLine, *FCommandLine::GetSubprocessCommandline());
+		}
+
+		FString GamePath = FPlatformProcess::GenerateApplicationPath(FApp::GetName(), FApp::GetBuildConfiguration());
+		FProcHandle ProcessHandle = FPlatformProcess::CreateProc(*GamePath, *Params, true, false, false, NULL, 0, NULL, NULL);
+
+		// @todo: progress reporting, UI feedback
+		if (ProcessHandle.IsValid())
+		{
+			FNotificationInfo Info(
+				SNew(SCaptureMovieNotification, ProcessHandle)
+				.BrowseToFolder(CaptureObject->Settings.OutputDirectory.Path)
+				.OnProcessClosed_Raw(this, &FMovieSceneCaptureModule::OnMovieCaptureProcessClosed)
+				.OnNotificationClosed_Raw(this, &FMovieSceneCaptureModule::OnNotificationClosed)
+			);
+			Info.bFireAndForget = false;
+			InProgressCaptureNotification = FSlateNotificationManager::Get().AddNotification(Info);
+			InProgressCaptureNotification->SetCompletionState(SNotificationItem::CS_Pending);
+		}
+		
+
+		return FText();
+	}
+private:
+	/** */
+	TWeakPtr<SWindow> CaptureSettingsWindow;
+	TSharedPtr<SNotificationItem> InProgressCaptureNotification;
+
+};
+
+IMPLEMENT_MODULE( FMovieSceneCaptureModule, MovieSceneCapture )
+
+#undef LOCTEXT_NAMESPACE
\ No newline at end of file
diff --git a/Engine/Source/Runtime/MovieSceneCapture/Private/MovieSceneCapturePCH.h b/Engine/Source/Runtime/MovieSceneCapture/Private/MovieSceneCapturePCH.h
new file mode 100644
--- /dev/null
+++ b/Engine/Source/Runtime/MovieSceneCapture/Private/MovieSceneCapturePCH.h
@@ -0,0 +1,10 @@
+// Copyright 1998-2015 Epic Games, Inc. All Rights Reserved.
+
+#pragma once
+
+#include "Core.h"
+#include "Engine.h"
+
+#include "MovieScene.h"
+#include "ExpressionParser.h"
+#include "ActorAnimationPlayer.h"
\ No newline at end of file
diff --git a/Engine/Source/Runtime/MovieSceneCapture/Public/AutomatedActorAnimationCapture.h b/Engine/Source/Runtime/MovieSceneCapture/Public/AutomatedActorAnimationCapture.h
new file mode 100644
--- /dev/null
+++ b/Engine/Source/Runtime/MovieSceneCapture/Public/AutomatedActorAnimationCapture.h
@@ -0,0 +1,52 @@
+// Copyright 1998-2015 Epic Games, Inc. All Rights Reserved.
+
+#pragma once
+
+#include "MovieSceneCapture.h"
+#include "ActorAnimation.h"
+#include "ActorAnimationPlayer.h"
+#include "AutomatedActorAnimationCapture.generated.h"
+
+UCLASS(config=EditorSettings)
+class MOVIESCENECAPTURE_API UAutomatedActorAnimationCapture : public UMovieSceneCapture
+{
+public:
+	UAutomatedActorAnimationCapture(const FObjectInitializer&);
+
+	GENERATED_BODY()
+
+	/** The actor animation to use during capture */
+	UPROPERTY(EditAnywhere, Category=General, meta=(AllowedClasses="ActorAnimation"))
+	FStringAssetReference ActorAnimation;
+
+	/** The level to use for the capture */
+	UPROPERTY(EditAnywhere, Category=General, meta=(AllowedClasses="World"))
+	FStringAssetReference Level;
+
+	/** Specific playback settings */
+	UPROPERTY(config, EditAnywhere, Category="Playback Settings", meta=(ShowOnlyInnerProperties))
+	FActorAnimationPlaybackSettings PlaybackSettings;
+
+	/** The amount of time to wait before playback and capture start. Useful for allowing Post Processing effects to settle down before capturing the animation. */
+	UPROPERTY(config, EditAnywhere, Category=CaptureSettings, AdvancedDisplay, meta=(Units=Seconds, ClampMin=0))
+	float PrerollAmount;
+
+public:
+	// UMovieSceneCapture interface
+	virtual FString GetPackageName() const override { return Level.ToString(); }
+	virtual void Initialize(FViewport* InViewport) override;
+	virtual void CaptureFrame(float DeltaSeconds) override;
+
+private:
+
+	/** Animation instance created at runtime before we start capturing */
+	UPROPERTY(transient)
+	UActorAnimationInstance* AnimationInstance;
+
+	/** Animation player used to playback the animation at runtime */
+	UPROPERTY(transient)
+	UActorAnimationPlayer* AnimationPlayback;
+
+	/** Transient amount of time that has passed since the level was started. When >= Settings.Preroll, playback will start  */
+	TOptional<float> PrerollTime;
+};
\ No newline at end of file
diff --git a/Engine/Source/Runtime/MovieSceneCapture/Public/ErrorCodes.h b/Engine/Source/Runtime/MovieSceneCapture/Public/ErrorCodes.h
new file mode 100644
--- /dev/null
+++ b/Engine/Source/Runtime/MovieSceneCapture/Public/ErrorCodes.h
@@ -0,0 +1,10 @@
+// Copyright 1998-2015 Epic Games, Inc. All Rights Reserved.
+
+#pragma once
+
+/** Process exit codes */
+struct FMovieSceneCaptureExitCodes
+{
+	static const int32 Base = 0xAC701000;
+	static const int32 AssetNotFound = Base + 1;
+};
\ No newline at end of file
diff --git a/Engine/Source/Runtime/MovieSceneCapture/Public/IMovieSceneCapture.h b/Engine/Source/Runtime/MovieSceneCapture/Public/IMovieSceneCapture.h
new file mode 100644
--- /dev/null
+++ b/Engine/Source/Runtime/MovieSceneCapture/Public/IMovieSceneCapture.h
@@ -0,0 +1,38 @@
+// Copyright 1998-2015 Epic Games, Inc. All Rights Reserved.
+
+#pragma once
+
+#include "IMovieSceneCapture.generated.h"
+
+class FViewport;
+struct FMovieSceneCaptureSettings;
+struct FMovieSceneCaptureHandle;
+
+UINTERFACE(MinimalAPI, meta=(CannotImplementInterfaceInBlueprint))
+class UMovieSceneCaptureInterface : public UInterface
+{
+public:
+	GENERATED_BODY()
+};
+
+/** Interface for a movie capture object */
+class IMovieSceneCaptureInterface
+{
+public:
+	GENERATED_BODY()
+
+	/** Initialize this capture object by binding it to the specified viewport */
+	virtual void Initialize(FViewport* Viewport) = 0;
+
+	/** Shut down this movie capture */
+	virtual void Close() = 0;
+
+	/** Get a unique handle to this object */
+	virtual FMovieSceneCaptureHandle GetHandle() const = 0;
+
+	/** Access specific movie scene capture settings */
+	virtual const FMovieSceneCaptureSettings& GetSettings() const = 0;
+	
+	/** Get the name of a map package name that this capture wants to load */
+	virtual FString GetPackageName() const { return FString(); }
+};
\ No newline at end of file
diff --git a/Engine/Source/Runtime/MovieSceneCapture/Public/LevelCapture.h b/Engine/Source/Runtime/MovieSceneCapture/Public/LevelCapture.h
new file mode 100644
--- /dev/null
+++ b/Engine/Source/Runtime/MovieSceneCapture/Public/LevelCapture.h
@@ -0,0 +1,20 @@
+// Copyright 1998-2015 Epic Games, Inc. All Rights Reserved.
+
+#pragma once
+
+#include "MovieSceneCapture.h"
+#include "LevelCapture.generated.h"
+
+UCLASS()
+class MOVIESCENECAPTURE_API ULevelCapture : public UMovieSceneCapture
+{
+public:
+	GENERATED_BODY()
+
+	/** The level we want to load and capture */
+	UPROPERTY(EditAnywhere, Category="General", meta=(AllowedClasses="World"))
+	FStringAssetReference Level;
+
+	virtual FString GetPackageName() const override { return Level.ToString(); }
+	virtual void Initialize(FViewport* InViewport) override;
+};
\ No newline at end of file
diff --git a/Engine/Source/Runtime/MovieSceneCapture/Public/MovieSceneCapture.h b/Engine/Source/Runtime/MovieSceneCapture/Public/MovieSceneCapture.h
new file mode 100644
--- /dev/null
+++ b/Engine/Source/Runtime/MovieSceneCapture/Public/MovieSceneCapture.h
@@ -0,0 +1,142 @@
+// Copyright 1998-2015 Epic Games, Inc. All Rights Reserved.
+
+#pragma once
+
+#include "MovieSceneCaptureSettings.h"
+#include "IMovieSceneCapture.h"
+#include "MovieSceneCaptureHandle.h"
+#include "MovieScene.h"
+#include "AVIWriter.h"
+#include "MovieSceneCapture.generated.h"
+
+/** Interface that defines when to capture or drop frames */
+struct ICaptureStrategy
+{
+	virtual ~ICaptureStrategy(){}
+	
+	virtual void OnStart() = 0;
+	virtual void OnStop() = 0;
+	virtual void OnPresent(double CurrentTimeSeconds, uint32 FrameIndex) = 0;
+	virtual bool ShouldSynchronizeFrames() const { return true; }
+
+	virtual bool ShouldPresent(double CurrentTimeSeconds, uint32 FrameIndex) const = 0;
+	virtual int32 GetDroppedFrames(double CurrentTimeSeconds, uint32 FrameIndex) const = 0;
+};
+
+/** Class responsible for capturing scene data */
+UCLASS(config=EditorSettings)
+class MOVIESCENECAPTURE_API UMovieSceneCapture : public UObject, public IMovieSceneCaptureInterface
+{
+public:
+
+	UMovieSceneCapture(const FObjectInitializer& Initializer);
+
+	GENERATED_BODY()
+
+public:
+
+	// Begin IMovieSceneCaptureInterface
+	virtual void Initialize(FViewport* InViewport) override;
+	virtual void Close() override;
+	virtual FMovieSceneCaptureHandle GetHandle() const override { return Handle; }
+	const FMovieSceneCaptureSettings& GetSettings() const override { return Settings; }
+	// End IMovieSceneCaptureInterface
+
+public:
+
+	/** Settings that define how to capture */
+	UPROPERTY(EditAnywhere, config, Category=CaptureSettings, meta=(ShowOnlyInnerProperties))
+	FMovieSceneCaptureSettings Settings;
+
+	/** Additional command line arguments to pass to the external process when capturing */
+	UPROPERTY(EditAnywhere, transient, Category=General, AdvancedDisplay)
+	FString AdditionalCommandLineArguments;
+
+	/** Value used to control the BufferVisualizationDumpFrames cvar in the child process */
+	UPROPERTY()
+	bool bBufferVisualizationDumpFrames;
+
+public:
+
+	/** Initialize the capture */
+	void StartCapture();
+
+	/** Finalize the capture, waiting for any outstanding processing */
+	void StopCapture();
+
+	/** Capture a frame from our bound viewport */
+	virtual void CaptureFrame(float DeltaSeconds);
+
+protected:
+
+	/** Implementation function that saves out a snapshot file from the specified color data */
+	void CaptureSnapshot(const TArray<FColor>& Colors);
+
+protected:
+
+	/** Resolve the specified format using the user supplied formatting rules. */
+	FString ResolveFileFormat(const FString& Folder, const FString& Format) const;
+
+	/** Get the default extension required for this capture (avi, jpg etc) */
+	const TCHAR* GetDefaultFileExtension() const;
+
+	/** Calculate a unique index for the {unique} formatting rule */
+	FString ResolveUniqueFilename();
+
+protected:
+
+	/** Structure used to cache various metrics for our capture */
+	struct FCachedMetrics
+	{
+		FCachedMetrics() : Width(0), Height(0), Frame(0) {}
+
+		/** The width/Height of the frame */
+		int32 Width, Height;
+		/** The current frame number */
+		int32 Frame;
+	};
+	
+	/** The viewport we are bound to */
+	FViewport* Viewport;
+	/** Our unique handle, used for external representation without having to link to the MovieSceneCapture module */
+	FMovieSceneCaptureHandle Handle;
+	/** Cached metrics for this capture operation */
+	FCachedMetrics CachedMetrics;
+	/** The current capture time, in seconds */
+	double CurrentTimeSeconds;
+	/** Optional AVI file writer used when capturing to video */
+	TUniquePtr<FAVIWriter> AVIWriter;
+	/** Strategy used for capture (real-time/fixed-time-step) */
+	TSharedPtr<ICaptureStrategy> CaptureStrategy;
+};
+
+/** A strategy that employs a fixed frame time-step, and as such never drops a frame. Potentially accelerated.  */
+struct MOVIESCENECAPTURE_API FFixedTimeStepCaptureStrategy : ICaptureStrategy
+{
+	FFixedTimeStepCaptureStrategy(uint32 InTargetFPS);
+
+	virtual void OnStart() override;
+	virtual void OnStop() override;
+	virtual void OnPresent(double CurrentTimeSeconds, uint32 FrameIndex) override;	
+	virtual bool ShouldPresent(double CurrentTimeSeconds, uint32 FrameIndex) const override;
+	virtual int32 GetDroppedFrames(double CurrentTimeSeconds, uint32 FrameIndex) const override;
+
+private:
+	uint32 TargetFPS;
+};
+
+/** A capture strategy that captures in real-time, potentially dropping frames to maintain a stable constant framerate video. */
+struct MOVIESCENECAPTURE_API FRealTimeCaptureStrategy : ICaptureStrategy
+{
+	FRealTimeCaptureStrategy(uint32 InTargetFPS);
+
+	virtual void OnStart() override;
+	virtual void OnStop() override;
+	virtual void OnPresent(double CurrentTimeSeconds, uint32 FrameIndex) override;
+	virtual bool ShouldSynchronizeFrames() const override { return false; }
+	virtual bool ShouldPresent(double CurrentTimeSeconds, uint32 FrameIndex) const override;
+	virtual int32 GetDroppedFrames(double CurrentTimeSeconds, uint32 FrameIndex) const override;
+
+private:
+	double NextPresentTimeS, FrameLength;
+};
\ No newline at end of file
diff --git a/Engine/Source/Runtime/MovieSceneCapture/Public/MovieSceneCaptureHandle.h b/Engine/Source/Runtime/MovieSceneCapture/Public/MovieSceneCaptureHandle.h
new file mode 100644
--- /dev/null
+++ b/Engine/Source/Runtime/MovieSceneCapture/Public/MovieSceneCaptureHandle.h
@@ -0,0 +1,22 @@
+// Copyright 1998-2015 Epic Games, Inc. All Rights Reserved.
+
+#pragma once
+
+/** A handle to a movie scene capture implementation */
+struct FMovieSceneCaptureHandle
+{
+	FMovieSceneCaptureHandle() : ID(0) {}
+
+	/** Equality operator */
+	friend bool operator==(const FMovieSceneCaptureHandle& A, const FMovieSceneCaptureHandle& B) { return A.ID == B.ID; }
+	
+	/** Inequality operator */
+	friend bool operator!=(const FMovieSceneCaptureHandle& A, const FMovieSceneCaptureHandle& B) { return A.ID != B.ID; }
+
+	/** Check if this handle is valid */
+	bool IsValid() const { return ID > 0; }
+
+protected:
+	/** Unique identifier for this handle */
+	uint32 ID;
+};
\ No newline at end of file
diff --git a/Engine/Source/Runtime/MovieSceneCapture/Public/MovieSceneCaptureModule.h b/Engine/Source/Runtime/MovieSceneCapture/Public/MovieSceneCaptureModule.h
new file mode 100644
--- /dev/null
+++ b/Engine/Source/Runtime/MovieSceneCapture/Public/MovieSceneCaptureModule.h
@@ -0,0 +1,30 @@
+// Copyright 1998-2015 Epic Games, Inc. All Rights Reserved.
+
+#pragma once
+
+#include "ModuleManager.h"
+#include "IMovieSceneCapture.h"
+#include "MovieSceneCaptureHandle.h"
+
+class UMovieSceneCapture;
+
+class IMovieSceneCaptureModule : public IModuleInterface
+{
+public:
+	static IMovieSceneCaptureModule& Get()
+	{
+		static const FName ModuleName(TEXT("MovieSceneCapture"));
+		return FModuleManager::LoadModuleChecked< IMovieSceneCaptureModule >(ModuleName);
+	}
+
+	virtual IMovieSceneCaptureInterface* InitializeFromCommandLine() = 0;
+	virtual IMovieSceneCaptureInterface* CreateMovieSceneCapture(FViewport* Viewport) = 0;
+
+	virtual IMovieSceneCaptureInterface* RetrieveMovieSceneInterface(FMovieSceneCaptureHandle Handle) = 0;
+	virtual void DestroyMovieSceneCapture(FMovieSceneCaptureHandle Handle) = 0;
+
+	virtual void OnMovieSceneCaptureFinished(IMovieSceneCaptureInterface*) = 0;
+
+	virtual void OpenCaptureSettings(const TSharedRef<class FTabManager>& TabManager, UMovieSceneCapture* CaptureObject) = 0;
+};
+
diff --git a/Engine/Source/Runtime/MovieSceneCapture/Public/MovieSceneCaptureSettings.h b/Engine/Source/Runtime/MovieSceneCapture/Public/MovieSceneCaptureSettings.h
new file mode 100644
--- /dev/null
+++ b/Engine/Source/Runtime/MovieSceneCapture/Public/MovieSceneCaptureSettings.h
@@ -0,0 +1,99 @@
+// Copyright 1998-2015 Epic Games, Inc. All Rights Reserved.
+
+#pragma once
+
+#include "MovieSceneCaptureSettings.generated.h"
+
+UENUM()
+enum class EMovieCaptureType : uint8
+{
+	AVI		UMETA(DisplayName="AVI Movie"),
+	BMP		UMETA(DisplayName="BMP Image Sequence"),
+	PNG		UMETA(DisplayName="PNG Image Sequence"),
+	JPEG	UMETA(DisplayName="JPEG Image Sequence")
+};
+
+/** Structure representing a capture resolution */
+USTRUCT()
+struct MOVIESCENECAPTURE_API FCaptureResolution
+{
+	FCaptureResolution(uint32 InX = 0, uint32 InY = 0) : ResX(InX), ResY(InY) {}
+	
+	GENERATED_BODY()
+
+	UPROPERTY(config, EditAnywhere, Category=Resolution, meta=(ClampMin = 16, ClampMax=7680))
+	uint32 ResX;
+
+	UPROPERTY(config, EditAnywhere, Category=Resolution, meta=(ClampMin = 16, ClampMax=7680))
+	uint32 ResY;
+};
+
+/** Common movie-scene capture settings */
+USTRUCT()
+struct MOVIESCENECAPTURE_API FMovieSceneCaptureSettings
+{
+	FMovieSceneCaptureSettings();
+
+	GENERATED_BODY()
+
+	/** The directory to output the captured file(s) in */
+	UPROPERTY(config, EditAnywhere, Category=General, AdvancedDisplay, meta=(RelativePath))
+	FDirectoryPath OutputDirectory;
+
+	/** The format to use for the resulting filename. Extension will be added automatically. Any tokens of the form {token} will be replaced with the corresponding value:
+	 * {fps}		- The captured framerate
+	 * {frame}		- The current frame number (only relevant for image sequences)
+	 * {width}		- The width of the captured frames
+	 * {height}		- The height of the captured frames
+	 */
+	UPROPERTY(config, EditAnywhere, Category=General, AdvancedDisplay)
+	FString OutputFormat;
+
+	/** The frame rate at which to capture */
+	UPROPERTY(config, EditAnywhere, Category=CaptureSettings)
+	int32 FrameRate;
+
+	/** The resolution at which to capture */
+	UPROPERTY(config, EditAnywhere, Category=CaptureSettings, meta=(ShowOnlyInnerProperties))
+	FCaptureResolution Resolution;
+
+	/** The type of capture to perform */
+	UPROPERTY(config, EditAnywhere, Category=CaptureSettings, AdvancedDisplay)
+	EMovieCaptureType CaptureType;
+
+	/** Whether compression is enabled on the resulting file(s) */
+	UPROPERTY(config, EditAnywhere, Category=CaptureSettings)
+	bool bUseCompression;
+
+	/** The level of compression to apply to the captured file(s) (0-1) */
+	UPROPERTY(config, EditAnywhere, Category=CaptureSettings, AdvancedDisplay, meta=(ClampMin=0, ClampMax=1))
+	float CompressionQuality;
+
+	/** (Experimental) - An optional codec to use for video encoding */
+	UPROPERTY(config, EditAnywhere, Category=CaptureSettings, AdvancedDisplay)
+	FString Codec;
+	
+	/** Whether to enable texture streaming whilst capturing */
+	UPROPERTY(config, EditAnywhere, Category=CaptureSettings)
+	bool bEnableTextureStreaming;
+
+	/** Whether to enable cinematic mode whilst capturing */
+	UPROPERTY(config, EditAnywhere, Category=Cinematic)
+	bool bCinematicMode;
+
+	/** Whether to allow player movement whilst capturing */
+	UPROPERTY(config, EditAnywhere, Category=Cinematic, AdvancedDisplay)
+	bool bAllowMovement;
+
+	/** Whether to allow player rotation whilst capturing */
+	UPROPERTY(config, EditAnywhere, Category=Cinematic, AdvancedDisplay)
+	bool bAllowTurning;
+
+	/** Whether to show the local player whilst capturing */
+	UPROPERTY(config, EditAnywhere, Category=Cinematic, AdvancedDisplay)
+	bool bShowPlayer;
+
+	/** Whether to show the in-game HUD whilst capturing */
+	UPROPERTY(config, EditAnywhere, Category=Cinematic, AdvancedDisplay)
+	bool bShowHUD;
+};
\ No newline at end of file
