Added missing include

diff --git a/Engine/Source/Runtime/Engine/Private/GameMode.cpp b/Engine/Source/Runtime/Engine/Private/GameMode.cpp
--- a/Engine/Source/Runtime/Engine/Private/GameMode.cpp
+++ b/Engine/Source/Runtime/Engine/Private/GameMode.cpp
@@ -21,6 +21,7 @@
 #include "GameFramework/GameMode.h"
 #include "Engine/ChildConnection.h"
 #include "Engine/GameInstance.h"
+#include "IMovieSceneCapture.h"
 #include "MovieSceneCaptureSettings.h"
 
 DEFINE_LOG_CATEGORY_STATIC(LogGameMode, Log, All);
