- Replace UpdateComponentToWorld with less expensive MarkRenderTransformDirty
- Move InvalidateCachedBounds into PostBlendPhysics, to avoid doing it more than necessary
- Remove MarkRenderDynamicDataDirty from PostAnimEvaluation (already done in PostBlendPhysics)

diff --git a/Engine/Source/Runtime/Engine/Private/Components/SkeletalMeshComponent.cpp b/Engine/Source/Runtime/Engine/Private/Components/SkeletalMeshComponent.cpp
--- a/Engine/Source/Runtime/Engine/Private/Components/SkeletalMeshComponent.cpp
+++ b/Engine/Source/Runtime/Engine/Private/Components/SkeletalMeshComponent.cpp
@@ -1271,21 +1271,17 @@ void USkeletalMeshComponent::PostAnimEvaluation(FAnimationEvaluationContext& Eva
 
 	bNeedToFlipSpaceBaseBuffers = true;
 
-	// Transforms updated, cached local bounds are now out of date.
-	InvalidateCachedBounds();
-
 	// update physics data from animated data
 	UpdateKinematicBonesToAnim(GetEditableSpaceBases(), ETeleportType::None, true);
 	UpdateRBJointMotors();
 
+	// If we have no physics to blend, we are done
 	if (!ShouldBlendPhysicsBones())
 	{
+		// Flip buffers, update bounds, attachments etc.
 		PostBlendPhysics();
 	}
 
-
-	MarkRenderDynamicDataDirty();
-
 	PrepareCloth();
 }
 
diff --git a/Engine/Source/Runtime/Engine/Private/PhysicsEngine/PhysAnim.cpp b/Engine/Source/Runtime/Engine/Private/PhysicsEngine/PhysAnim.cpp
--- a/Engine/Source/Runtime/Engine/Private/PhysicsEngine/PhysAnim.cpp
+++ b/Engine/Source/Runtime/Engine/Private/PhysicsEngine/PhysAnim.cpp
@@ -342,9 +342,6 @@ void USkeletalMeshComponent::PerformBlendPhysicsBones(const TArray<FBoneIndexTyp
     }
 #endif
 	
-
-	// Transforms updated, cached local bounds are now out of date.
-	InvalidateCachedBounds();
 }
 
 
@@ -423,19 +420,24 @@ void USkeletalMeshComponent::PostBlendPhysics()
 {
 	SCOPE_CYCLE_COUNTER(STAT_UpdateLocalToWorldAndOverlaps);
 	
+	// Flip bone buffer and send 'post anim' notification
 	FinalizeBoneTransform();
 
-	UpdateComponentToWorld();
-
 	// Update Child Transform - The above function changes bone transform, so will need to update child transform
 	UpdateChildTransforms();
 
 	// animation often change overlap. 
 	UpdateOverlaps();
 
+	// Cached local bounds are now out of date
+	InvalidateCachedBounds();
+
 	// update bounds
 	UpdateBounds();
 
+	// Need to send new bounds to 
+	MarkRenderTransformDirty();
+
 	// New bone positions need to be sent to render thread
 	MarkRenderDynamicDataDirty();
 }
