fixed MovingAgent cleanup in NavLinkCustomComponent
#ue4

diff --git a/Engine/Source/Runtime/Engine/Private/AI/Navigation/NavLinkCustomComponent.cpp b/Engine/Source/Runtime/Engine/Private/AI/Navigation/NavLinkCustomComponent.cpp
--- a/Engine/Source/Runtime/Engine/Private/AI/Navigation/NavLinkCustomComponent.cpp
+++ b/Engine/Source/Runtime/Engine/Private/AI/Navigation/NavLinkCustomComponent.cpp
@@ -83,7 +83,7 @@ bool UNavLinkCustomComponent::OnLinkMoveStarted(UPathFollowingComponent* PathCom
 
 void UNavLinkCustomComponent::OnLinkMoveFinished(UPathFollowingComponent* PathComp)
 {
-	TWeakObjectPtr<UPathFollowingComponent> WeakPathComp;
+	TWeakObjectPtr<UPathFollowingComponent> WeakPathComp(PathComp);
 	MovingAgents.Remove(WeakPathComp);
 }
 
