Ensure KeysToConsume and FoundChords are reset when breaking out due to a block

diff --git a/Engine/Source/Runtime/Engine/Private/UserInterface/PlayerInput.cpp b/Engine/Source/Runtime/Engine/Private/UserInterface/PlayerInput.cpp
--- a/Engine/Source/Runtime/Engine/Private/UserInterface/PlayerInput.cpp
+++ b/Engine/Source/Runtime/Engine/Private/UserInterface/PlayerInput.cpp
@@ -1081,16 +1081,16 @@ void UPlayerInput::ProcessInputStack(const TArray<UInputComponent*>& InputCompon
 			{
 				// stop traversing the stack, all input has been consumed by this InputComponent
 				--StackIndex;
+				KeysToConsume.Reset();
+				FoundChords.Reset();
 				break;
 			}
-			else
+
+			// we do this after finishing the whole component, so we don't consume a key while there might be more bindings to it
+			for (int32 KeyIndex=0; KeyIndex<KeysToConsume.Num(); ++KeyIndex)
 			{
-				// we do this after finishing the whole component, so we don't consume a key while there might be more bindings to it
-				for (int32 KeyIndex=0; KeyIndex<KeysToConsume.Num(); ++KeyIndex)
-				{
-					ConsumeKey(KeysToConsume[KeyIndex]);
-				}
-			}		
+				ConsumeKey(KeysToConsume[KeyIndex]);
+			}
 			KeysToConsume.Reset();
 			FoundChords.Reset();
 		}
