Forward declare class AHUD for delegate declaration

diff --git a/Engine/Source/Runtime/Engine/Classes/GameFramework/HUD.h b/Engine/Source/Runtime/Engine/Classes/GameFramework/HUD.h
--- a/Engine/Source/Runtime/Engine/Classes/GameFramework/HUD.h
+++ b/Engine/Source/Runtime/Engine/Classes/GameFramework/HUD.h
@@ -12,6 +12,7 @@ class AActor;
 class UCanvas;
 class APlayerController;
 class UFont;
+class AHUD;
 
 DECLARE_MULTICAST_DELEGATE_FiveParams(FOnShowDebugInfo, AHUD* /* HUD */, UCanvas* /* Canvas */, const FDebugDisplayInfo& /* DisplayInfo */, float& /* YL */, float& /* YPos */);
 
