Fix delegate that wasn't getting cleaned up in TermGamePhys()

Merging 2700823

//depot/UE4-UT/...

to //depot/UE4/...

diff --git a/Engine/Source/Runtime/Engine/Private/PhysicsEngine/PhysLevel.cpp b/Engine/Source/Runtime/Engine/Private/PhysicsEngine/PhysLevel.cpp
--- a/Engine/Source/Runtime/Engine/Private/PhysicsEngine/PhysLevel.cpp
+++ b/Engine/Source/Runtime/Engine/Private/PhysicsEngine/PhysLevel.cpp
@@ -50,6 +50,7 @@ DEFINE_STAT(STAT_NumShapesAsync);
 
 
 FPhysCommandHandler * GPhysCommandHandler = NULL;
+FDelegateHandle GPreGarbageCollectDelegateHandle;
 
 // CVars
 static TAutoConsoleVariable<float> CVarToleranceScaleLength(
@@ -349,7 +350,7 @@ void InitGamePhys()
 
 	GPhysCommandHandler = new FPhysCommandHandler();
 
-	FCoreUObjectDelegates::PreGarbageCollect.AddRaw(GPhysCommandHandler, &FPhysCommandHandler::Flush);
+	GPreGarbageCollectDelegateHandle = FCoreUObjectDelegates::PreGarbageCollect.AddRaw(GPhysCommandHandler, &FPhysCommandHandler::Flush);
 
 	// Init Extensions
 	PxInitExtensions(*GPhysXSDK);
@@ -488,6 +489,7 @@ void TermGamePhys()
 	if (GPhysCommandHandler != NULL)
 	{
 		GPhysCommandHandler->Flush();	//finish off any remaining commands
+		FCoreUObjectDelegates::PreGarbageCollect.Remove(GPreGarbageCollectDelegateHandle);
 		delete GPhysCommandHandler;
 		GPhysCommandHandler = NULL;
 	}
