Fix for objects being culled by the near plane of the shadow frustum with RTDF shadows (shadow casters behind the camera pop out)
New TraceDistance property on directional lights

diff --git a/Engine/Shaders/DistanceFieldShadowing.usf b/Engine/Shaders/DistanceFieldShadowing.usf
--- a/Engine/Shaders/DistanceFieldShadowing.usf
+++ b/Engine/Shaders/DistanceFieldShadowing.usf
@@ -200,6 +200,14 @@ void ShadowObjectCullVS(
 	float EffectiveRadius = (ObjectPositionAndRadius.w + MinRadius) * ConservativeRadiusScale;
 	float3 WorldPosition = InPosition.xyz * EffectiveRadius + ObjectPositionAndRadius.xyz;
 	OutPosition = mul(float4(WorldPosition, 1), WorldToShadow);
+
+	// Clamp the vertex to the near plane if it is in front of the near plane
+	if (OutPosition.z < 0)
+	{
+		OutPosition.z = 0.000001f;
+		OutPosition.w = 1.0f;
+	}
+
 	Output.PositionAndRadius = ObjectPositionAndRadius;
 	Output.ObjectIndex = ObjectIndex;
 } 
@@ -226,7 +234,8 @@ void ShadowObjectCullPS(
 	float TileExpand = 1;
 	ShadowTileMin.xy = (TilePositionForCulling - TileExpand) / (float2)NumGroups * 2 - 1;
 	ShadowTileMax.xy = (TilePositionForCulling + 1 + TileExpand) / (float2)NumGroups * 2 - 1;
-	ShadowTileMin.z = 0;
+	// Extrude toward light to avoid culling objects between the light and the shadow frustum
+	ShadowTileMin.z = -1000;
 	ShadowTileMax.z = 1;
 
 	float3 ObjectViewSpaceMin;
@@ -235,7 +244,8 @@ void ShadowObjectCullPS(
 
 	BRANCH
 	// Separating axis test on the AABB
-	if (all(ObjectViewSpaceMax > ShadowTileMin) && all(ObjectViewSpaceMin < ShadowTileMax))
+	// Note: don't clip by near plane, objects closer to the light can still cast into the frustum
+	if (all(ObjectViewSpaceMax.xy > ShadowTileMin.xy) && all(ObjectViewSpaceMin < ShadowTileMax))
 	{
 		float3 ObjectCenter = .5f * (ObjectViewSpaceMin + ObjectViewSpaceMax);
 		float3 MinProjections = 500000;
@@ -303,7 +313,7 @@ float3 LightDirection;
 float4 LightPositionAndInvRadius;
 float LightSourceRadius;
 float RayStartOffsetDepthScale;
-float2 TanLightAngleAndNormalThreshold;
+float3 TanLightAngleAndNormalThreshold;
 int4 ScissorRectMinAndSize;
 
 /** Min and Max depth for this tile. */
@@ -489,7 +499,7 @@ void DistanceFieldShadowingCS(
 	float3 OpaqueWorldPosition = HomogeneousWorldPosition.xyz / HomogeneousWorldPosition.w;
 
 	// Distance for directional lights to trace
-	float TraceDistance = 10000;
+	float TraceDistance = TanLightAngleAndNormalThreshold.z;
 
 	//@todo - take advantage of CSM covering near pixels
 	//@todo - handle invalid tile depth ranges due to MinDepth
diff --git a/Engine/Source/Runtime/Engine/Classes/Components/DirectionalLightComponent.h b/Engine/Source/Runtime/Engine/Classes/Components/DirectionalLightComponent.h
--- a/Engine/Source/Runtime/Engine/Classes/Components/DirectionalLightComponent.h
+++ b/Engine/Source/Runtime/Engine/Classes/Components/DirectionalLightComponent.h
@@ -115,6 +115,10 @@ class UDirectionalLightComponent : public ULightComponent
 	UPROPERTY(EditAnywhere, BlueprintReadOnly, Category=DistanceFieldShadows, meta=(UIMin = "0", UIMax = "5"))
 	float LightSourceAngle;
 
+	/** Determines how far shadows can be cast, in world units.  Larger values increase the shadowing cost. */
+	UPROPERTY(EditAnywhere, AdvancedDisplay, BlueprintReadOnly, Category=DistanceFieldShadows, meta=(UIMin = "1000", UIMax = "100000"), DisplayName = "DistanceField Trace Distance")
+	float TraceDistance;
+
 	/** The Lightmass settings for this object. */
 	UPROPERTY(EditAnywhere, BlueprintReadOnly, Category=Light, meta=(ShowOnlyInnerProperties))
 	struct FLightmassDirectionalLightSettings LightmassSettings;
diff --git a/Engine/Source/Runtime/Engine/Private/Components/DirectionalLightComponent.cpp b/Engine/Source/Runtime/Engine/Private/Components/DirectionalLightComponent.cpp
--- a/Engine/Source/Runtime/Engine/Private/Components/DirectionalLightComponent.cpp
+++ b/Engine/Source/Runtime/Engine/Private/Components/DirectionalLightComponent.cpp
@@ -93,6 +93,9 @@ public:
 	/** Light source angle in degrees. */
 	float LightSourceAngle;
 
+	/** Determines how far shadows can be cast, in world units.  Larger values increase the shadowing cost. */
+	float TraceDistance;
+
 	/** Initialization constructor. */
 	FDirectionalLightSceneProxy(const UDirectionalLightComponent* Component):
 		FLightSceneProxy(Component),
@@ -106,7 +109,8 @@ public:
 		ShadowDistanceFadeoutFraction(Component->ShadowDistanceFadeoutFraction),
 		bUseInsetShadowsForMovableObjects(Component->bUseInsetShadowsForMovableObjects),
 		DistanceFieldShadowDistance(Component->bUseRayTracedDistanceFieldShadows ? Component->DistanceFieldShadowDistance : 0),
-		LightSourceAngle(Component->LightSourceAngle)
+		LightSourceAngle(Component->LightSourceAngle),
+		TraceDistance(FMath::Clamp(Component->TraceDistance, 1000.0f, 1000000.0f))
 	{
 		LightShaftOverrideDirection.Normalize();
 
@@ -169,6 +173,11 @@ public:
 		return LightSourceAngle;
 	}
 
+	virtual float GetTraceDistance() const override
+	{
+		return TraceDistance;
+	}
+
 	virtual bool GetLightShaftOcclusionParameters(float& OutOcclusionMaskDarkness, float& OutOcclusionDepthRange) const override
 	{
 		OutOcclusionMaskDarkness = OcclusionMaskDarkness;
@@ -673,6 +682,7 @@ UDirectionalLightComponent::UDirectionalLightComponent(const FObjectInitializer&
 	DynamicShadowDistanceStationaryLight = 0.f;
 
 	DistanceFieldShadowDistance = 30000.0f;
+	TraceDistance = 10000.0f;
 	FarShadowDistance = 300000.0f;
 	LightSourceAngle = 1;
 
@@ -749,6 +759,7 @@ bool UDirectionalLightComponent::CanEditChange(const UProperty* InProperty) cons
 		}
 
 		if (PropertyName == GET_MEMBER_NAME_STRING_CHECKED(UDirectionalLightComponent, DistanceFieldShadowDistance)
+			|| PropertyName == GET_MEMBER_NAME_STRING_CHECKED(UDirectionalLightComponent, TraceDistance)
 			|| (PropertyName == GET_MEMBER_NAME_STRING_CHECKED(UDirectionalLightComponent, LightSourceAngle)
 				// Make sure we don't accidentally affect the LightSourceAngle property in LightmassSettings
 				&& InProperty && InProperty->GetOuter() && InProperty->GetOuter()->GetName() == TEXT("DirectionalLightComponent")))
diff --git a/Engine/Source/Runtime/Engine/Public/SceneManagement.h b/Engine/Source/Runtime/Engine/Public/SceneManagement.h
--- a/Engine/Source/Runtime/Engine/Public/SceneManagement.h
+++ b/Engine/Source/Runtime/Engine/Public/SceneManagement.h
@@ -694,6 +694,7 @@ public:
 	virtual float GetSourceRadius() const { return 0.0f; }
 	virtual bool IsInverseSquared() const { return false; }
 	virtual float GetLightSourceAngle() const { return 0.0f; }
+	virtual float GetTraceDistance() const { return 0.0f; }
 
 	virtual FVector2D GetLightShaftConeParams() const
 	{
diff --git a/Engine/Source/Runtime/Renderer/Private/DistanceFieldShadowing.cpp b/Engine/Source/Runtime/Renderer/Private/DistanceFieldShadowing.cpp
--- a/Engine/Source/Runtime/Renderer/Private/DistanceFieldShadowing.cpp
+++ b/Engine/Source/Runtime/Renderer/Private/DistanceFieldShadowing.cpp
@@ -465,7 +465,7 @@ public:
 		SetShaderValue(RHICmdList, ShaderRHI, RayStartOffsetDepthScale, LightProxy.GetRayStartOffsetDepthScale());
 
 		const float LightSourceAngle = FMath::Clamp(LightProxy.GetLightSourceAngle(), 0.001f, 5.0f) * PI / 180.0f;
-		const FVector2D TanLightAngleAndNormalThresholdValue(FMath::Tan(LightSourceAngle), FMath::Cos(PI / 2 + LightSourceAngle));
+		const FVector TanLightAngleAndNormalThresholdValue(FMath::Tan(LightSourceAngle), FMath::Cos(PI / 2 + LightSourceAngle), LightProxy.GetTraceDistance());
 		SetShaderValue(RHICmdList, ShaderRHI, TanLightAngleAndNormalThreshold, TanLightAngleAndNormalThresholdValue);
 
 		SetShaderValue(RHICmdList, ShaderRHI, ScissorRectMinAndSize, FIntRect(ScissorRect.Min, ScissorRect.Size()));
