#ue4 - Optimize FCollisionQueryParams::AddIgnoredActor() and AddIgnoredComponents().

- Removed AddUnique() calls when adding components, sort is deferred until the first requests for the ignored components lists.
- FCollisionQueryParams::IgnoreComponents is now private, use GetIgnoredComponents() instead.

Based on: https://github.com/EpicGames/UnrealEngine/pull/1442/files
GitHub #1442

diff --git a/Engine/Source/Runtime/Engine/Private/Collision/Collision.cpp b/Engine/Source/Runtime/Engine/Private/Collision/Collision.cpp
--- a/Engine/Source/Runtime/Engine/Private/Collision/Collision.cpp
+++ b/Engine/Source/Runtime/Engine/Private/Collision/Collision.cpp
@@ -154,6 +154,7 @@ FCollisionQueryParams::FCollisionQueryParams(FName InTraceTag, bool bInTraceComp
 	bFindInitialOverlaps = true;
 	bReturnFaceIndex = false;
 	bReturnPhysicalMaterial = false;
+	bComponentListUnique = true;
 
 	AddIgnoredActor(InIgnoreActor);
 	if (InIgnoreActor != NULL)
@@ -171,42 +172,53 @@ static FORCEINLINE_DEBUGGABLE bool IsQueryCollisionEnabled(const UPrimitiveCompo
 		|| CollisionEnabled == ECollisionEnabled::QueryOnly;
 }
 
+static FORCEINLINE_DEBUGGABLE bool CheckForCollision(const AActor* Actor)
+{
+	// Note: This code operates differently in the editor because the actors in the editor can have their collision setting become out of sync with physx.  This happens because even when collision is disabled on an actor, in the editor we 
+	// tell phsx that we still require queries( See FBodyInstance::UpdatePhysicsFilterData).  Doing so allows us to perform editor only traces against objects with collision disabled.  Due to this, we cannot assume that the collision enabled flag here
+	// setting is correct compared to physx and so we still ignore specified components regardless of their collision setting
+#if WITH_EDITOR
+	const bool bCheckForCollision = Actor && Actor->GetWorld() && Actor->GetWorld()->IsGameWorld();
+#else
+	const bool bCheckForCollision = true;
+#endif
+	return bCheckForCollision;
+}
 
-void FCollisionQueryParams::AddIgnoredActor(const AActor* InIgnoreActor)
+
+static FORCEINLINE_DEBUGGABLE bool CheckForCollision(const UPrimitiveComponent* PrimComponent)
 {
 	// Note: This code operates differently in the editor because the actors in the editor can have their collision setting become out of sync with physx.  This happens because even when collision is disabled on an actor, in the editor we 
 	// tell phsx that we still require queries( See FBodyInstance::UpdatePhysicsFilterData).  Doing so allows us to perform editor only traces against objects with collision disabled.  Due to this, we cannot assume that the collision enabled flag here
 	// setting is correct compared to physx and so we still ignore specified components regardless of their collision setting
 #if WITH_EDITOR
-	const bool bCheckForCollision = InIgnoreActor && InIgnoreActor->GetWorld() && InIgnoreActor->GetWorld()->IsGameWorld();
+	bool bCheckForCollision = PrimComponent && PrimComponent->GetWorld() && PrimComponent->GetWorld()->IsGameWorld();
 #else
 	const bool bCheckForCollision = true;
 #endif
+	return bCheckForCollision;
+}
 
+
+void FCollisionQueryParams::AddIgnoredActor(const AActor* InIgnoreActor)
+{
 	if (InIgnoreActor)
-	{
-		if (IgnoreComponents.Num() == 0)
+	{	
+		const bool bCheckForCollision = CheckForCollision(InIgnoreActor);
+		const int32 InitialCount = IgnoreComponents.Num();
+		for (const UActorComponent* ActorComponent : InIgnoreActor->GetComponents())
 		{
-			// Don't need to AddUnique if we start empty, as GetComponents should never contain duplicate components.
-			for (const UActorComponent* ActorComponent : InIgnoreActor->GetComponents())
+			const UPrimitiveComponent* PrimComponent = Cast<const UPrimitiveComponent>(ActorComponent);
+			if (PrimComponent && (!bCheckForCollision || IsQueryCollisionEnabled(PrimComponent)) )
 			{
-				const UPrimitiveComponent* PrimComponent = Cast<const UPrimitiveComponent>(ActorComponent);
-				if (PrimComponent && (!bCheckForCollision || IsQueryCollisionEnabled(PrimComponent)) )
-				{
-					IgnoreComponents.Add(PrimComponent->GetUniqueID());
-				}
+				IgnoreComponents.Add(PrimComponent->GetUniqueID());
 			}
 		}
-		else
+
+		// If we added entries to a non-empty array, we aren't sure if the list is unique. Assumes GetComponents() above contains no duplicates.
+		if ((InitialCount != 0) && (IgnoreComponents.Num() - InitialCount > 0))
 		{
-			for (const UActorComponent* ActorComponent : InIgnoreActor->GetComponents())
-			{
-				const UPrimitiveComponent* PrimComponent = Cast<const UPrimitiveComponent>(ActorComponent);
-				if (PrimComponent && (!bCheckForCollision || IsQueryCollisionEnabled(PrimComponent)) )
-				{
-					IgnoreComponents.AddUnique(PrimComponent->GetUniqueID());
-				}
-			}
+			bComponentListUnique = false;
 		}
 	}
 }
@@ -215,10 +227,7 @@ void FCollisionQueryParams::AddIgnoredActors(const TArray<AActor*>& InIgnoreActo
 {
 	for (int32 Idx = 0; Idx < InIgnoreActors.Num(); ++Idx)
 	{
-		if (AActor const* const A = InIgnoreActors[Idx])
-		{
-			AddIgnoredActor(A);
-		}
+		AddIgnoredActor(InIgnoreActors[Idx]);
 	}
 }
 
@@ -226,36 +235,29 @@ void FCollisionQueryParams::AddIgnoredActors(const TArray<TWeakObjectPtr<AActor>
 {
 	for (int32 Idx = 0; Idx < InIgnoreActors.Num(); ++Idx)
 	{
-		if (AActor const* const A = InIgnoreActors[Idx].Get())
-		{
-			AddIgnoredActor(A);
-		}
+		AddIgnoredActor(InIgnoreActors[Idx].Get());
 	}
 }
 
-
-void FCollisionQueryParams::AddIgnoredComponent(const UPrimitiveComponent* InIgnoreComponent)
+FORCEINLINE_DEBUGGABLE void FCollisionQueryParams::Internal_AddIgnoredComponent(const UPrimitiveComponent* InIgnoreComponent)
 {
-	// Note: This code operates differently in the editor because the actors in the editor can have their collision setting become out of sync with physx.  This happens because even when collision is disabled on an actor, in the editor we 
-	// tell phsx that we still require queries( See FBodyInstance::UpdatePhysicsFilterData).  Doing so allows us to perform editor only traces against objects with collision disabled.  Due to this, we cannot assume that the collision enabled flag here
-	// setting is correct compared to physx and so we still ignore specified components regardless of their collision setting
-#if WITH_EDITOR
-	const bool bCheckForCollision = InIgnoreComponent && InIgnoreComponent->GetWorld() && InIgnoreComponent->GetWorld()->IsGameWorld();
-#else
-	const bool bCheckForCollision = true;
-#endif
-
-	if (InIgnoreComponent && (!bCheckForCollision || IsQueryCollisionEnabled(InIgnoreComponent)))
+	if (InIgnoreComponent && (!CheckForCollision(InIgnoreComponent) || IsQueryCollisionEnabled(InIgnoreComponent)))
 	{
-		IgnoreComponents.AddUnique(InIgnoreComponent->GetUniqueID());
+		IgnoreComponents.Add(InIgnoreComponent->GetUniqueID());
+		bComponentListUnique = false;
 	}
 }
 
+void FCollisionQueryParams::AddIgnoredComponent(const UPrimitiveComponent* InIgnoreComponent)
+{
+	Internal_AddIgnoredComponent(InIgnoreComponent);
+}
+
 void FCollisionQueryParams::AddIgnoredComponents(const TArray<UPrimitiveComponent*>& InIgnoreComponents)
 {
 	for (const UPrimitiveComponent* IgnoreComponent : InIgnoreComponents)
 	{
-		AddIgnoredComponent(IgnoreComponent);
+		Internal_AddIgnoredComponent(IgnoreComponent);
 	}
 }
 
@@ -263,11 +265,52 @@ void FCollisionQueryParams::AddIgnoredComponents(const TArray<TWeakObjectPtr<UPr
 {
 	for (TWeakObjectPtr<UPrimitiveComponent> IgnoreComponent : InIgnoreComponents)
 	{
-		if (IgnoreComponent.IsValid())
+		Internal_AddIgnoredComponent(IgnoreComponent.Get());
+	}
+}
+
+void FCollisionQueryParams::AddIgnoredComponent_LikelyDuplicatedRoot(const UPrimitiveComponent* InIgnoreComponent)
+{
+	if (InIgnoreComponent && (!CheckForCollision(InIgnoreComponent) || IsQueryCollisionEnabled(InIgnoreComponent)))
+	{
+		// Code calling this is usually just making sure they don't add the root component to queries right before the actual query.
+		// We try to avoid invalidating the uniqueness of the array if this is the case.
+		const uint32 ComponentID = InIgnoreComponent->GetUniqueID();
+		if (IgnoreComponents.Num() == 0 || IgnoreComponents[0] != ComponentID)
+		{
+			IgnoreComponents.Add(ComponentID);
+			bComponentListUnique = false;
+		}
+	}
+}
+
+const FCollisionQueryParams::IgnoreComponentsArrayType& FCollisionQueryParams::GetIgnoredComponents() const
+{
+	if (!bComponentListUnique)
+	{
+		// Make unique
+		bComponentListUnique = true;
+		if (IgnoreComponents.Num() > 1)
 		{
-			AddIgnoredComponent(IgnoreComponent.Get());
+			// For adding a collection to ignore it's faster to sort and remove duplicates
+			// than to check for duplicates at each addition.
+			IgnoreComponents.Sort();
+			uint32* U = IgnoreComponents.GetData();
+			uint32* D = U+1;
+			uint32* E = U + IgnoreComponents.Num();
+			while (D != E)
+			{
+				if (*U != *D)
+				{
+					*(++U) = *D;
+				}
+				D += 1;
+			}
+			IgnoreComponents.SetNum(U - IgnoreComponents.GetData() + 1, /*bAllowShrinking=*/ false);
 		}
 	}
+
+	return IgnoreComponents;
 }
 
 //////////////////////////////////////////////////////////////////////////
diff --git a/Engine/Source/Runtime/Engine/Private/Collision/PhysXCollision.cpp b/Engine/Source/Runtime/Engine/Private/Collision/PhysXCollision.cpp
--- a/Engine/Source/Runtime/Engine/Private/Collision/PhysXCollision.cpp
+++ b/Engine/Source/Runtime/Engine/Private/Collision/PhysXCollision.cpp
@@ -181,7 +181,7 @@ FRaycastHelperBox2D_Base::EResponse FRaycastHelperBox2D_Base::CalcQueryHitType(b
 		FBodyInstance* OwnerBodyInstance = static_cast<FBodyInstance*>(Fixture->GetBody()->GetUserData());
 		UPrimitiveComponent* OwnerComponent = OwnerBodyInstance->OwnerComponent.Get();
 		
-		if ((OwnerComponent != nullptr) && Params.IgnoreComponents.Contains(OwnerComponent->GetUniqueID()))
+		if ((OwnerComponent != nullptr) && Params.GetIgnoredComponents().Contains(OwnerComponent->GetUniqueID()))
 		{
 			Result = ResponseNone;
 		}
@@ -831,7 +831,7 @@ bool RaycastTest(const UWorld* World, const FVector Start, const FVector End, EC
 			PxFilterData PFilter = CreateQueryFilterData(TraceChannel, Params.bTraceComplex, ResponseParams.CollisionResponse, ObjectParams, false);
 			PxSceneQueryFilterData PQueryFilterData(PFilter, PxSceneQueryFilterFlag::eSTATIC | PxSceneQueryFilterFlag::eDYNAMIC | PxSceneQueryFilterFlag::ePREFILTER);
 			PxSceneQueryFlags POutputFlags = PxHitFlags();
-			FPxQueryFilterCallback PQueryCallback(Params.IgnoreComponents);
+			FPxQueryFilterCallback PQueryCallback(Params.GetIgnoredComponents());
 			PQueryCallback.bIgnoreTouches = true; // pre-filter to ignore touches and only get blocking hits.
 
 			FPhysScene* PhysScene = World->GetPhysicsScene();
@@ -904,7 +904,7 @@ bool RaycastSingle(const UWorld* World, struct FHitResult& OutHit, const FVector
 			PxFilterData PFilter = CreateQueryFilterData(TraceChannel, Params.bTraceComplex, ResponseParams.CollisionResponse, ObjectParams, false);
 			PxSceneQueryFilterData PQueryFilterData(PFilter, PxSceneQueryFilterFlag::eSTATIC | PxSceneQueryFilterFlag::eDYNAMIC | PxSceneQueryFilterFlag::ePREFILTER);
 			PxSceneQueryFlags POutputFlags = PxSceneQueryFlag::ePOSITION | PxSceneQueryFlag::eNORMAL | PxSceneQueryFlag::eDISTANCE | PxSceneQueryFlag::eMTD;
-			FPxQueryFilterCallback PQueryCallback(Params.IgnoreComponents);
+			FPxQueryFilterCallback PQueryCallback(Params.GetIgnoredComponents());
 			PQueryCallback.bIgnoreTouches = true; // pre-filter to ignore touches and only get blocking hits.
 
 			PxVec3 PStart = U2PVector(Start);
@@ -1019,7 +1019,7 @@ bool RaycastMulti(const UWorld* World, TArray<struct FHitResult>& OutHits, const
 		PxFilterData PFilter = CreateQueryFilterData(TraceChannel, Params.bTraceComplex, ResponseParams.CollisionResponse, ObjectParams, true);
 		PxSceneQueryFilterData PQueryFilterData(PFilter, PxSceneQueryFilterFlag::eSTATIC | PxSceneQueryFilterFlag::eDYNAMIC | PxSceneQueryFilterFlag::ePREFILTER);
 		PxSceneQueryFlags POutputFlags = PxSceneQueryFlag::ePOSITION | PxSceneQueryFlag::eNORMAL | PxSceneQueryFlag::eDISTANCE | PxSceneQueryFlag::eMTD;
-		FPxQueryFilterCallback PQueryCallback(Params.IgnoreComponents);
+		FPxQueryFilterCallback PQueryCallback(Params.GetIgnoredComponents());
 
 		// Create buffer for hits. Note: memory is not initialized (for perf reasons), since API does not require it.
 		TTypeCompatibleBytes<PxRaycastHit> RawPHits[HIT_BUFFER_SIZE];
@@ -1190,7 +1190,7 @@ bool GeomSweepTest(const UWorld* World, const struct FCollisionShape& CollisionS
 		PxSceneQueryFilterData PQueryFilterData(PFilter, PxSceneQueryFilterFlag::eSTATIC | PxSceneQueryFilterFlag::eDYNAMIC | PxSceneQueryFilterFlag::ePREFILTER | PxSceneQueryFilterFlag::ePOSTFILTER);
 		PxSceneQueryFlags POutputFlags; 
 
-		FPxQueryFilterCallbackSweep PQueryCallbackSweep(Params.IgnoreComponents);
+		FPxQueryFilterCallbackSweep PQueryCallbackSweep(Params.GetIgnoredComponents());
 		PQueryCallbackSweep.DiscardInitialOverlaps = !Params.bFindInitialOverlaps;
 		PQueryCallbackSweep.bIgnoreTouches = true; // pre-filter to ignore touches and only get blocking hits.
 
@@ -1270,7 +1270,7 @@ bool GeomSweepSingle(const UWorld* World, const struct FCollisionShape& Collisio
 		//UE_LOG(LogCollision, Log, TEXT("PFilter: %x %x %x %x"), PFilter.word0, PFilter.word1, PFilter.word2, PFilter.word3);
 		PxSceneQueryFilterData PQueryFilterData(PFilter, PxSceneQueryFilterFlag::eSTATIC | PxSceneQueryFilterFlag::eDYNAMIC | PxSceneQueryFilterFlag::ePREFILTER);
 		PxSceneQueryFlags POutputFlags = PxSceneQueryFlag::ePOSITION | PxSceneQueryFlag::eNORMAL | PxSceneQueryFlag::eDISTANCE | PxSceneQueryFlag::eMTD;
-		FPxQueryFilterCallbackSweep PQueryCallbackSweep(Params.IgnoreComponents);
+		FPxQueryFilterCallbackSweep PQueryCallbackSweep(Params.GetIgnoredComponents());
 		PQueryCallbackSweep.bIgnoreTouches = true; // pre-filter to ignore touches and only get blocking hits.
 		PQueryCallbackSweep.DiscardInitialOverlaps = !Params.bFindInitialOverlaps;	
 
@@ -1366,7 +1366,7 @@ bool GeomSweepMulti_PhysX(const UWorld* World, const PxGeometry& PGeom, const Px
 	PxFilterData PFilter = CreateQueryFilterData(TraceChannel, Params.bTraceComplex, ResponseParams.CollisionResponse, ObjectParams, true);
 	PxSceneQueryFilterData PQueryFilterData(PFilter, PxSceneQueryFilterFlag::eSTATIC | PxSceneQueryFilterFlag::eDYNAMIC | PxSceneQueryFilterFlag::ePREFILTER | PxSceneQueryFilterFlag::ePOSTFILTER);
 	PxSceneQueryFlags POutputFlags = PxSceneQueryFlag::ePOSITION | PxSceneQueryFlag::eNORMAL | PxSceneQueryFlag::eDISTANCE | PxSceneQueryFlag::eMTD;
-	FPxQueryFilterCallbackSweep PQueryCallbackSweep(Params.IgnoreComponents);
+	FPxQueryFilterCallbackSweep PQueryCallbackSweep(Params.GetIgnoredComponents());
 	PQueryCallbackSweep.DiscardInitialOverlaps = !Params.bFindInitialOverlaps;
 
 	const FVector Delta = End - Start;
@@ -1532,7 +1532,7 @@ bool GeomOverlapMultiImp_PhysX(const UWorld* World, const PxGeometry& PGeom, con
 		// Create filter data used to filter collisions
 		PxFilterData PFilter = CreateQueryFilterData(TraceChannel, Params.bTraceComplex, ResponseParams.CollisionResponse, ObjectParams, InfoType != EQueryInfo::IsAnything);
 		PxSceneQueryFilterData PQueryFilterData(PFilter, PxSceneQueryFilterFlag::eSTATIC | PxSceneQueryFilterFlag::eDYNAMIC | PxSceneQueryFilterFlag::ePREFILTER);
-		FPxQueryFilterCallback PQueryCallback(Params.IgnoreComponents);
+		FPxQueryFilterCallback PQueryCallback(Params.GetIgnoredComponents());
 		PQueryCallback.bIgnoreTouches = (InfoType == EQueryInfo::IsBlocking); // pre-filter to ignore touches and only get blocking hits, if that's what we're after.
 
 		// Enable scene locks, in case they are required
diff --git a/Engine/Source/Runtime/Engine/Private/Components/PrimitiveComponent.cpp b/Engine/Source/Runtime/Engine/Private/Components/PrimitiveComponent.cpp
--- a/Engine/Source/Runtime/Engine/Private/Components/PrimitiveComponent.cpp
+++ b/Engine/Source/Runtime/Engine/Private/Components/PrimitiveComponent.cpp
@@ -2505,7 +2505,7 @@ void UPrimitiveComponent::UpdateOverlaps(const TArray<FOverlapInfo>* NewPendingO
 bool UPrimitiveComponent::ComponentOverlapMultiImpl(TArray<struct FOverlapResult>& OutOverlaps, const UWorld* World, const FVector& Pos, const FQuat& Quat, ECollisionChannel TestChannel, const struct FComponentQueryParams& Params, const struct FCollisionObjectQueryParams& ObjectQueryParams) const
 {
 	FComponentQueryParams ParamsWithSelf = Params;
-	ParamsWithSelf.AddIgnoredComponent(this);
+	ParamsWithSelf.AddIgnoredComponent_LikelyDuplicatedRoot(this);
 	OutOverlaps.Reset();
 	return BodyInstance.OverlapMulti(OutOverlaps, World, /*pWorldToComponent=*/ nullptr, Pos, Quat, TestChannel, ParamsWithSelf, FCollisionResponseParams(GetCollisionResponseToChannels()), ObjectQueryParams);
 }
diff --git a/Engine/Source/Runtime/Engine/Public/CollisionQueryParams.h b/Engine/Source/Runtime/Engine/Public/CollisionQueryParams.h
--- a/Engine/Source/Runtime/Engine/Public/CollisionQueryParams.h
+++ b/Engine/Source/Runtime/Engine/Public/CollisionQueryParams.h
@@ -38,9 +38,20 @@ struct ENGINE_API FCollisionQueryParams
 	/** TArray typedef of components to ignore. */
 	typedef TArray<uint32, TInlineAllocator<NumInlinedActorComponents>> IgnoreComponentsArrayType;
 
+private:
+
+	/** Tracks whether the IgnoreComponents list is verified unique. */
+	mutable bool bComponentListUnique;
+
 	/** Set of components to ignore during the trace */
-	IgnoreComponentsArrayType IgnoreComponents;
+	mutable IgnoreComponentsArrayType IgnoreComponents;
+
+	void Internal_AddIgnoredComponent(const UPrimitiveComponent* InIgnoreComponent);
 
+public:
+
+	/** Returns set of components to ignore during the trace. */
+	const IgnoreComponentsArrayType& GetIgnoredComponents() const;
 
 	// Constructors
 	FCollisionQueryParams(bool bInTraceComplex=false)
@@ -51,6 +62,7 @@ struct ENGINE_API FCollisionQueryParams
 		bFindInitialOverlaps = true;
 		bReturnFaceIndex = false;
 		bReturnPhysicalMaterial = false;
+		bComponentListUnique = true;
 	}
 
 	FCollisionQueryParams(FName InTraceTag, bool bInTraceComplex=false, const AActor* InIgnoreActor=NULL);
@@ -75,6 +87,12 @@ struct ENGINE_API FCollisionQueryParams
 	/** Variant that uses an array of TWeakObjectPtrs */
 	void AddIgnoredComponents(const TArray<TWeakObjectPtr<UPrimitiveComponent>>& InIgnoreComponents);
 
+	/**
+	 * Special variant that hints that we are likely adding a duplicate of the root component or first ignored component.
+	 * Helps avoid invalidating the potential uniquess of the IgnoreComponents array.
+	 */
+	void AddIgnoredComponent_LikelyDuplicatedRoot(const UPrimitiveComponent* InIgnoreComponent);
+
 	FString ToString() const
 	{
 		return FString::Printf(TEXT("[%s:%s] TraceAsync(%d), TraceComplex(%d)"), *OwnerTag.ToString(), *TraceTag.ToString(), bTraceAsyncScene, bTraceComplex );
