Non-unity compile fixes

diff --git a/Engine/Source/Runtime/Engine/Private/CapturePin.h b/Engine/Source/Runtime/Engine/Private/CapturePin.h
--- a/Engine/Source/Runtime/Engine/Private/CapturePin.h
+++ b/Engine/Source/Runtime/Engine/Private/CapturePin.h
@@ -28,6 +28,7 @@ typedef TCHAR* PTCHAR;
 	#pragma warning(pop)
 #endif
 
+class FAVIWriter;
 
 class FCapturePin : public CSourceStream
 {
diff --git a/Engine/Source/Runtime/Engine/Private/CaptureSource.cpp b/Engine/Source/Runtime/Engine/Private/CaptureSource.cpp
--- a/Engine/Source/Runtime/Engine/Private/CaptureSource.cpp
+++ b/Engine/Source/Runtime/Engine/Private/CaptureSource.cpp
@@ -7,6 +7,7 @@
 
 #include "CapturePin.h"
 #include "CaptureSource.h"
+#include "AVIWriter.h"
 
 
 #if PLATFORM_WINDOWS && !UE_BUILD_MINIMAL
diff --git a/Engine/Source/Runtime/Engine/Private/CaptureSource.h b/Engine/Source/Runtime/Engine/Private/CaptureSource.h
--- a/Engine/Source/Runtime/Engine/Private/CaptureSource.h
+++ b/Engine/Source/Runtime/Engine/Private/CaptureSource.h
@@ -16,6 +16,7 @@
 #include "HideWindowsPlatformTypes.h"
 #pragma warning(pop)
 class FCapturePin;
+class FAVIWriter;
 
 // {F817F8A7-DE00-42CF-826A-7A5654602D8E}
 DEFINE_GUID(CLSID_ViewportCaptureSource, 
