Fix SoundNodeQualityLevel not setting cached quality level when in PIE

diff --git a/Engine/Source/Runtime/Engine/Private/SoundNodeQualityLevel.cpp b/Engine/Source/Runtime/Engine/Private/SoundNodeQualityLevel.cpp
--- a/Engine/Source/Runtime/Engine/Private/SoundNodeQualityLevel.cpp
+++ b/Engine/Source/Runtime/Engine/Private/SoundNodeQualityLevel.cpp
@@ -71,9 +71,11 @@ void USoundNodeQualityLevel::ParseNodes( FAudioDevice* AudioDevice, const UPTRIN
 			const bool bIsPIESound = ((GEditor->bIsSimulatingInEditor || GEditor->PlayWorld != NULL) && ActiveSound.World != NULL);
 			if (bIsPIESound)
 			{
-				QualityLevel = GetDefault<ULevelEditorPlaySettings>()->PlayInEditorSoundQualityLevel;
+				CachedQualityLevel = GetDefault<ULevelEditorPlaySettings>()->PlayInEditorSoundQualityLevel;
 			}
 		}
+
+		QualityLevel = CachedQualityLevel;
 	}
 	else
 	{
