PR #1532: Marked bEnableRootMotion with AssetRegistrySearchable (Contributed by [REDACTED_PERSON])

diff --git a/Engine/Source/Runtime/Engine/Classes/Animation/AnimSequence.h b/Engine/Source/Runtime/Engine/Classes/Animation/AnimSequence.h
--- a/Engine/Source/Runtime/Engine/Classes/Animation/AnimSequence.h
+++ b/Engine/Source/Runtime/Engine/Classes/Animation/AnimSequence.h
@@ -474,7 +474,7 @@ class ENGINE_API UAnimSequence : public UAnimSequenceBase
 	FName RetargetSource;
 
 	/** If this is on, it will allow extracting of root motion **/
-	UPROPERTY(EditAnywhere, Category = RootMotion)
+	UPROPERTY(EditAnywhere, AssetRegistrySearchable, Category = RootMotion, meta = (DisplayName = "EnableRootMotion"))
 	bool bEnableRootMotion;
 
 	/** Root Bone will be locked to that position when extracting root motion.**/
