#ue4 - Add FCollisionQueryParams::ClearIgnoredComponents() and SetNumIgnoredComponents().

diff --git a/Engine/Source/Runtime/Engine/Private/Collision/Collision.cpp b/Engine/Source/Runtime/Engine/Private/Collision/Collision.cpp
--- a/Engine/Source/Runtime/Engine/Private/Collision/Collision.cpp
+++ b/Engine/Source/Runtime/Engine/Private/Collision/Collision.cpp
@@ -313,6 +313,22 @@ const FCollisionQueryParams::IgnoreComponentsArrayType& FCollisionQueryParams::G
 	return IgnoreComponents;
 }
 
+void FCollisionQueryParams::SetNumIgnoredComponents(int32 NewNum)
+{
+	if (NewNum > 0)
+	{
+		// We can only make it smaller (and uniqueness does not change).
+		if (NewNum < IgnoreComponents.Num())
+		{
+			IgnoreComponents.SetNum(NewNum, /*bAllowShrinking=*/ false);
+		}
+	}
+	else
+	{
+		ClearIgnoredComponents();
+	}
+}
+
 //////////////////////////////////////////////////////////////////////////
 // FSeparatingAxisPointCheck
 
diff --git a/Engine/Source/Runtime/Engine/Public/CollisionQueryParams.h b/Engine/Source/Runtime/Engine/Public/CollisionQueryParams.h
--- a/Engine/Source/Runtime/Engine/Public/CollisionQueryParams.h
+++ b/Engine/Source/Runtime/Engine/Public/CollisionQueryParams.h
@@ -50,9 +50,22 @@ private:
 
 public:
 
-	/** Returns set of components to ignore during the trace. */
+	/** Returns set of unique components to ignore during the trace. Elements are guaranteed to be unique (they are made so internally if they are not already). */
 	const IgnoreComponentsArrayType& GetIgnoredComponents() const;
 
+	/** Clears the set of components to ignore during the trace. */
+	void ClearIgnoredComponents()
+	{
+		IgnoreComponents.Reset();
+		bComponentListUnique = true;
+	}
+
+	/**
+	 * Set the number of ignored components in the list. Uniqueness is not changed, it operates on the current state (unique or not).
+	 * Useful for temporarily adding some, then restoring to a previous size. NewNum must be <= number of current components for there to be any effect.
+	 */
+	void SetNumIgnoredComponents(int32 NewNum);
+
 	// Constructors
 	FCollisionQueryParams(bool bInTraceComplex=false)
 	{
