Particles: Add a column for Looping in the content browser (True if any emitter in any LOD has Emitter Loops = 0)
[INTEGRATE] Change 2694232 by [REDACTED_PERSON][REDACTED_HANDLE] on 2015/09/16 17:27:42

diff --git a/Engine/Source/Runtime/Engine/Private/Particles/ParticleComponents.cpp b/Engine/Source/Runtime/Engine/Private/Particles/ParticleComponents.cpp
--- a/Engine/Source/Runtime/Engine/Private/Particles/ParticleComponents.cpp
+++ b/Engine/Source/Runtime/Engine/Private/Particles/ParticleComponents.cpp
@@ -2087,6 +2087,31 @@ void UParticleSystem::GetAssetRegistryTags(TArray<FAssetRegistryTag>& OutTags) c
 	}
 	OutTags.Add(FAssetRegistryTag("LODMethod", LODMethodString, FAssetRegistryTag::TT_Alphabetical));
 
+	// Run thru the emitters and see if any will loop forever
+	bool bAnyEmitterLoopsForever = false;
+	for (int32 EmitterIndex = 0; (EmitterIndex < Emitters.Num()) && !bAnyEmitterLoopsForever; ++EmitterIndex)
+	{
+		if (const UParticleEmitter* Emitter = Emitters[EmitterIndex])
+		{
+			for (const UParticleLODLevel* LODLevel : Emitter->LODLevels)
+			{
+				if (LODLevel != nullptr)
+				{
+					const UParticleModuleRequired* RequiredModule = LODLevel->RequiredModule;
+					check(RequiredModule);
+
+					if (RequiredModule->EmitterLoops == 0)
+					{
+						UE_LOG(LogTemp, Warning, TEXT("ParticleSystem %s has emitter %s that is loop forever"), *this->GetName(), *Emitter->EmitterName.ToString());
+						bAnyEmitterLoopsForever = true;
+						break;
+					}
+				}
+			}
+		}
+	}
+	OutTags.Add(FAssetRegistryTag("Looping", bAnyEmitterLoopsForever ? TEXT("True") : TEXT("False"), FAssetRegistryTag::TT_Alphabetical));
+
 	Super::GetAssetRegistryTags(OutTags);
 }
 
