Particles: Removed temporary log spam in UParticleSystem::GetAssetRegistryTags
[INTEGRATE] Change 2694240 by [REDACTED_PERSON][REDACTED_HANDLE] on 2015/09/16 17:32:22

diff --git a/Engine/Source/Runtime/Engine/Private/Particles/ParticleComponents.cpp b/Engine/Source/Runtime/Engine/Private/Particles/ParticleComponents.cpp
--- a/Engine/Source/Runtime/Engine/Private/Particles/ParticleComponents.cpp
+++ b/Engine/Source/Runtime/Engine/Private/Particles/ParticleComponents.cpp
@@ -2102,7 +2102,6 @@ void UParticleSystem::GetAssetRegistryTags(TArray<FAssetRegistryTag>& OutTags) c
 
 					if (RequiredModule->EmitterLoops == 0)
 					{
-						UE_LOG(LogTemp, Warning, TEXT("ParticleSystem %s has emitter %s that is loop forever"), *this->GetName(), *Emitter->EmitterName.ToString());
 						bAnyEmitterLoopsForever = true;
 						break;
 					}
