Engine: Increased size of curve texture used for GPU particle curves (now 512x256)
[INTEGRATE] Change 2695899 by [REDACTED_PERSON][REDACTED_HANDLE] on 2015/09/17 16:25:05

diff --git a/Engine/Source/Runtime/Engine/Private/Particles/ParticleCurveTexture.cpp b/Engine/Source/Runtime/Engine/Private/Particles/ParticleCurveTexture.cpp
--- a/Engine/Source/Runtime/Engine/Private/Particles/ParticleCurveTexture.cpp
+++ b/Engine/Source/Runtime/Engine/Private/Particles/ParticleCurveTexture.cpp
@@ -17,7 +17,7 @@
 #include "SceneUtils.h"
 
 /** The texture size allocated for particle curves. */
-extern const int32 GParticleCurveTextureSizeX = 256;
+extern const int32 GParticleCurveTextureSizeX = 512;
 extern const int32 GParticleCurveTextureSizeY = 256;
 
 /** The texel allocator uses 16-bit integers internally. */
@@ -549,7 +549,7 @@ FTexelAllocation FParticleCurveTexture::AddCurve( const TArray<FColor>& CurveSam
 				FMemory::Memcpy( PendingCurve->Samples, CurveSamples.GetData(), TexelAllocation.Size * sizeof(FColor) );
 				return TexelAllocation;
 			}
-			UE_LOG(LogParticles, Warning, TEXT("FParticleCurveTexture: Failed to allocate %d texels for a curve."), CurveSamples.Num() );
+			UE_LOG(LogParticles, Warning, TEXT("FParticleCurveTexture: Failed to allocate %d texels for a curve (may need to increase the size of GParticleCurveTextureSizeX or GParticleCurveTextureSizeY)."), CurveSamples.Num() );
 
 			return TexelAllocation;
 		}
