Engine: Optimized FURL construction when passed a short map name by trying the asset registry before resorting to a glacial disk search
[INTEGRATE] Change 2700123 by [REDACTED_PERSON][REDACTED_HANDLE] on 2015/09/21 20:13:12

diff --git a/Engine/Source/Runtime/Engine/Private/URL.cpp b/Engine/Source/Runtime/Engine/Private/URL.cpp
--- a/Engine/Source/Runtime/Engine/Private/URL.cpp
+++ b/Engine/Source/Runtime/Engine/Private/URL.cpp
@@ -6,6 +6,7 @@
 
 #include "EnginePrivate.h"
 #include "Net/UnrealNetwork.h"
+#include "AssetRegistryModule.h"
 
 /*-----------------------------------------------------------------------------
 	FURL Statics.
@@ -331,15 +332,50 @@ FURL::FURL( FURL* Base, const TCHAR* TextURL, ETravelType Type )
 			// find full pathname from short map name
 			FString MapFullName;
 			FText MapNameError;
+			bool bFoundMap = false;
 			if (FPaths::FileExists(URL))
 			{
 				Map = FPackageName::FilenameToLongPackageName(URL);
+				bFoundMap = true;
 			}
-			else if (!FPackageName::DoesPackageNameContainInvalidCharacters(URL, &MapNameError) && FPackageName::SearchForPackageOnDisk(FString(URL) + FPackageName::GetMapPackageExtension(), &MapFullName))
+			else if (!FPackageName::DoesPackageNameContainInvalidCharacters(URL, &MapNameError))
 			{
-				Map = MapFullName;
+				// First try to use the asset registry if it is available and finished scanning
+				if (FModuleManager::Get().IsModuleLoaded("AssetRegistry"))
+				{
+					IAssetRegistry& AssetRegistry = FModuleManager::Get().LoadModuleChecked<FAssetRegistryModule>("AssetRegistry").Get();
+
+					if (!AssetRegistry.IsLoadingAssets())
+					{
+						TArray<FAssetData> MapList;
+						if (AssetRegistry.GetAssetsByClass(UWorld::StaticClass()->GetFName(), /*out*/ MapList))
+						{
+							FName TargetTestName(URL);
+							for (const FAssetData& MapAsset : MapList)
+							{
+								if (MapAsset.AssetName == TargetTestName)
+								{
+									Map = MapAsset.PackageName.ToString();
+									bFoundMap = true;
+									break;
+								}
+							}
+						}
+					}
+				}
+
+				if (!bFoundMap)
+				{
+					// Fall back to incredibly slow disk scan for the package
+					if (FPackageName::SearchForPackageOnDisk(FString(URL) + FPackageName::GetMapPackageExtension(), &MapFullName))
+					{
+						Map = MapFullName;
+						bFoundMap = true;
+					}
+				}
 			}
-			else
+
+			if (!bFoundMap)
 			{
 				// can't find file, invalidate and bail
 				UE_CLOG(MapNameError.ToString().Len() > 0, LogLongPackageNames, Warning, TEXT("URL: %s: %s"), URL, *MapNameError.ToString());
