#ue4 - Fix crash in CheckStillInWorld() after Actor has been destroyed or is not in the world.

diff --git a/Engine/Source/Runtime/Engine/Private/Actor.cpp b/Engine/Source/Runtime/Engine/Private/Actor.cpp
--- a/Engine/Source/Runtime/Engine/Private/Actor.cpp
+++ b/Engine/Source/Runtime/Engine/Private/Actor.cpp
@@ -997,8 +997,18 @@ FBox AActor::GetComponentsBoundingBox(bool bNonColliding) const
 
 bool AActor::CheckStillInWorld()
 {
+	if (IsPendingKill())
+	{
+		return false;
+	}
+	UWorld* MyWorld = GetWorld();
+	if (!MyWorld)
+	{
+		return false;
+	}
+
 	// check the variations of KillZ
-	AWorldSettings* WorldSettings = GetWorld()->GetWorldSettings( true );
+	AWorldSettings* WorldSettings = MyWorld->GetWorldSettings( true );
 
 	if (!WorldSettings->bEnableWorldBoundsChecks)
 	{
diff --git a/Engine/Source/Runtime/Engine/Private/Components/InterpToMovementComponent.cpp b/Engine/Source/Runtime/Engine/Private/Components/InterpToMovementComponent.cpp
--- a/Engine/Source/Runtime/Engine/Private/Components/InterpToMovementComponent.cpp
+++ b/Engine/Source/Runtime/Engine/Private/Components/InterpToMovementComponent.cpp
@@ -337,14 +337,19 @@ bool UInterpToMovementComponent::CheckStillInWorld()
 	{
 		return false;
 	}
+	const UWorld* MyWorld = GetWorld();
+	if (!MyWorld)
+	{
+		return false;
+	}
 	// check the variations of KillZ
-	AWorldSettings* WorldSettings = GetWorld()->GetWorldSettings(true);
+	AWorldSettings* WorldSettings = MyWorld->GetWorldSettings(true);
 	if (!WorldSettings->bEnableWorldBoundsChecks)
 	{
 		return true;
 	}
 	AActor* ActorOwner = UpdatedComponent->GetOwner();
-	if (!ActorOwner)
+	if (!IsValid(ActorOwner))
 	{
 		return false;
 	}
diff --git a/Engine/Source/Runtime/Engine/Private/Components/ProjectileMovementComponent.cpp b/Engine/Source/Runtime/Engine/Private/Components/ProjectileMovementComponent.cpp
--- a/Engine/Source/Runtime/Engine/Private/Components/ProjectileMovementComponent.cpp
+++ b/Engine/Source/Runtime/Engine/Private/Components/ProjectileMovementComponent.cpp
@@ -506,14 +506,21 @@ bool UProjectileMovementComponent::CheckStillInWorld()
 	{
 		return false;
 	}
+
+	const UWorld* MyWorld = GetWorld();
+	if (!MyWorld)
+	{
+		return false;
+	}
+
 	// check the variations of KillZ
-	AWorldSettings* WorldSettings = GetWorld()->GetWorldSettings( true );
+	AWorldSettings* WorldSettings = MyWorld->GetWorldSettings( true );
 	if (!WorldSettings->bEnableWorldBoundsChecks)
 	{
 		return true;
 	}
 	AActor* ActorOwner = UpdatedComponent->GetOwner();
-	if ( !ActorOwner )
+	if (!IsValid(ActorOwner))
 	{
 		return false;
 	}
