Engine: Allow command line options -execcmds, -exec, -vsync, and -novsync to work in Test builds (still disabled in Shipping)
[INTEGRATE] Change 2701266 by [REDACTED_PERSON][REDACTED_HANDLE] on 2015/09/22 15:47:14

diff --git a/Engine/Source/Runtime/Engine/Private/UnrealEngine.cpp b/Engine/Source/Runtime/Engine/Private/UnrealEngine.cpp
--- a/Engine/Source/Runtime/Engine/Private/UnrealEngine.cpp
+++ b/Engine/Source/Runtime/Engine/Private/UnrealEngine.cpp
@@ -830,7 +830,7 @@ void UEngine::Init(IEngineLoop* InEngineLoop)
 		bUseConsoleInput = true;
 	}
 
-#if !(UE_BUILD_SHIPPING || UE_BUILD_TEST)
+#if !(UE_BUILD_SHIPPING)
 	// Optionally Exec an exec file
 	FString Temp;
 	if( FParse::Value(FCommandLine::Get(), TEXT("EXEC="), Temp) )
@@ -873,7 +873,7 @@ void UEngine::Init(IEngineLoop* InEngineLoop)
 	{
 		new(GEngine->DeferredCommands) FString(TEXT("r.vsync 0"));
 	}
-#endif // !(UE_BUILD_SHIPPING || UE_BUILD_TEST)
+#endif // !(UE_BUILD_SHIPPING)
 
 	if (GetDerivedDataCache())
 	{
