Merging

//depot/UE4-Orion/Engine/Source/Runtime/Engine/Private/Animation/AnimMontage.cpp

to //depot/UE4/Engine/Source/Runtime/Engine/Private/Animation/AnimMontage.cpp

diff --git a/Engine/Source/Runtime/Engine/Private/Animation/AnimMontage.cpp b/Engine/Source/Runtime/Engine/Private/Animation/AnimMontage.cpp
--- a/Engine/Source/Runtime/Engine/Private/Animation/AnimMontage.cpp
+++ b/Engine/Source/Runtime/Engine/Private/Animation/AnimMontage.cpp
@@ -464,13 +464,13 @@ void UAnimMontage::PostLoad()
 	}
 
 	// fix up blending time deprecated variable
-	if (BlendInTime_DEPRECATED != -1.f)
+	if (BlendInTime_DEPRECATED == -1.f)
 	{
 		BlendIn.SetBlendTime(BlendInTime_DEPRECATED);
 		BlendInTime_DEPRECATED = -1.f;
 	}
 
-	if(BlendOutTime_DEPRECATED != -1.f)
+	if(BlendOutTime_DEPRECATED == -1.f)
 	{
 		BlendOut.SetBlendTime(BlendOutTime_DEPRECATED);
 		BlendOutTime_DEPRECATED = -1.f;
