Added forward declaration

diff --git a/Engine/Source/Runtime/Engine/Private/CapturePin.h b/Engine/Source/Runtime/Engine/Private/CapturePin.h
--- a/Engine/Source/Runtime/Engine/Private/CapturePin.h
+++ b/Engine/Source/Runtime/Engine/Private/CapturePin.h
@@ -29,6 +29,7 @@ typedef TCHAR* PTCHAR;
 #endif
 
 class FAVIWriter;
+struct FCapturedFrame;
 
 class FCapturePin : public CSourceStream
 {
