Use 64bit int to hold memory sizes in FStreamMemoryTracker, will still break down when memory sizes go above the max signed 64 bit int, but fixes the issue with large 32 bit texture pools

Merging 2701660

//depot/UE4-UT/...

to //depot/UE4/...

diff --git a/Engine/Source/Runtime/Engine/Public/ContentStreaming.h b/Engine/Source/Runtime/Engine/Public/ContentStreaming.h
--- a/Engine/Source/Runtime/Engine/Public/ContentStreaming.h
+++ b/Engine/Source/Runtime/Engine/Public/ContentStreaming.h
@@ -726,11 +726,8 @@ struct FSpawnedTextureInstance
  */
 struct FStreamMemoryTracker
 {
-#if PLATFORM_64BITS
 	typedef int64 TSize;
-#else
-	typedef int32 TSize;
-#endif
+
 	/** Stream-in memory that hasn't been allocated yet. */
 	volatile TSize PendingStreamIn;
 	/** Temp memory that hasn't been allocated yet. */
