Fixed non-unity build

diff --git a/Engine/Source/Runtime/Engine/Private/CapturePin.cpp b/Engine/Source/Runtime/Engine/Private/CapturePin.cpp
--- a/Engine/Source/Runtime/Engine/Private/CapturePin.cpp
+++ b/Engine/Source/Runtime/Engine/Private/CapturePin.cpp
@@ -1,18 +1,17 @@
 // Copyright 1998-2015 Epic Games, Inc. All Rights Reserved.
 
-/*=============================================================================
-	CapturePin.cpp: CapturePin implementation
-=============================================================================*/
-
 #include "EnginePrivate.h"
-
 #include "CapturePin.h"
+#include "CaptureSource.h"
 #include "AVIWriter.h"
 #include "ScopeExit.h"
 
+
 DEFINE_LOG_CATEGORY(LogMovieCapture);
+
 #if PLATFORM_WINDOWS && !UE_BUILD_MINIMAL
 
+
 FCapturePin::FCapturePin(HRESULT *phr, CSource *pFilter, const FAVIWriter& InWriter)
         : CSourceStream(NAME("Push Source"), phr, pFilter, L"Capture")
 		, FrameLength(UNITS/InWriter.Options.CaptureFPS)
