Fix WindowsXP build

Merging 2702593

//depot/UE4-UT/...

to //depot/UE4/...

diff --git a/Engine/Source/Runtime/Engine/Public/ContentStreaming.h b/Engine/Source/Runtime/Engine/Public/ContentStreaming.h
--- a/Engine/Source/Runtime/Engine/Public/ContentStreaming.h
+++ b/Engine/Source/Runtime/Engine/Public/ContentStreaming.h
@@ -726,7 +726,12 @@ struct FSpawnedTextureInstance
  */
 struct FStreamMemoryTracker
 {
+#if WINVER < 0x0600
+	// Windows XP does not have InterlockedIncrement64
+	typedef int32 TSize;
+#else
 	typedef int64 TSize;
+#endif
 
 	/** Stream-in memory that hasn't been allocated yet. */
 	volatile TSize PendingStreamIn;
