Fix reimport apex asset crash
- fix old code path for non parallel not using morptarget path

diff --git a/Engine/Source/Editor/UnrealEd/Private/ApexClothingUtils.cpp b/Engine/Source/Editor/UnrealEd/Private/ApexClothingUtils.cpp
--- a/Engine/Source/Editor/UnrealEd/Private/ApexClothingUtils.cpp
+++ b/Engine/Source/Editor/UnrealEd/Private/ApexClothingUtils.cpp
@@ -1201,16 +1201,12 @@ NxClothingAsset* ApplyTransform(NxClothingAsset* ApexClothingAsset)
 
 void RestoreAllClothingSections(USkeletalMesh* SkelMesh, uint32 LODIndex, uint32 AssetIndex)
 {
-	int32 NumAssets = SkelMesh->ClothingAssets.Num();
-
-	check((int32)AssetIndex < NumAssets);
-
 	TArray<uint32> SectionIndices;
 	SkelMesh->GetOriginSectionIndicesWithCloth(LODIndex, AssetIndex, SectionIndices);
 
 	for(int32 i=0; i < SectionIndices.Num(); i++)
 	{
-		RestoreOriginalClothingSection(SkelMesh, LODIndex, SectionIndices[i]);
+		RestoreOriginalClothingSection(SkelMesh, LODIndex, SectionIndices[i], false);
 	}
 }
 
@@ -1235,12 +1231,15 @@ void RemoveAssetFromSkeletalMesh(USkeletalMesh* SkelMesh, uint32 AssetIndex, boo
 		}
 	}
 
-	// release and make it invalid
-	GPhysCommandHandler->DeferredRelease(SkelMesh->ClothingAssets[AssetIndex].ApexClothingAsset);
-	SkelMesh->ClothingAssets[AssetIndex].ApexClothingAsset = NULL;
 
-	//this requires to refresh UI layout
-	SkelMesh->ClothingAssets.RemoveAt(AssetIndex);
+	NxClothingAsset* ApexClothingAsset = SkelMesh->ClothingAssets[AssetIndex].ApexClothingAsset;	//Can't delete apex asset until after apex actors so we save this for now and reregister component (which will trigger the actor delete)
+	SkelMesh->ClothingAssets.RemoveAt(AssetIndex);	//have to remove the asset from the array so that new actors are not created for asset pending deleting
+	ReregisterSkelMeshComponents(SkelMesh);
+
+	//Now we can actually delete the asset
+	GPhysCommandHandler->DeferredRelease(ApexClothingAsset);
+	
+
 
 	if(bRecreateSkelMeshComponent)
 	{
@@ -1249,7 +1248,7 @@ void RemoveAssetFromSkeletalMesh(USkeletalMesh* SkelMesh, uint32 AssetIndex, boo
 	}
 }
 
-void RestoreOriginalClothingSection(USkeletalMesh* SkelMesh, uint32 LODIndex, uint32 SectionIndex)
+void RestoreOriginalClothingSection(USkeletalMesh* SkelMesh, uint32 LODIndex, uint32 SectionIndex, bool bReregisterSkelMeshComponent)
 {
 	FSkeletalMeshResource* ImportedResource= SkelMesh->GetImportedResource();
 	FStaticLODModel& LODModel = ImportedResource->LODModels[LODIndex];
@@ -1313,8 +1312,6 @@ void RestoreOriginalClothingSection(USkeletalMesh* SkelMesh, uint32 LODIndex, ui
 
 		LODModel.MultiSizeIndexContainer.CopyIndexBuffer(OutIndexBuffer);
 
-		FClothingAssetData& CorrespondAssetData = SkelMesh->ClothingAssets[ClothChunk.CorrespondClothAssetIndex];
-
 		LODModel.Chunks.RemoveAt(ClothSection.ChunkIndex);
 		LODModel.Sections.RemoveAt(OriginSection.CorrespondClothSectionIndex);
 
@@ -1343,9 +1340,11 @@ void RestoreOriginalClothingSection(USkeletalMesh* SkelMesh, uint32 LODIndex, ui
 		OriginSection.bDisabled = false;
 		OriginSection.CorrespondClothSectionIndex = -1;
 
-		// Cloth chunk is removed so updates Vertex Factories
-		ReregisterSkelMeshComponents(SkelMesh);
-		// Reinitialize the render data
+		if(bReregisterSkelMeshComponent)
+		{
+			ReregisterSkelMeshComponents(SkelMesh);
+		}
+		
 		SkelMesh->PostEditChange();
 	}
 	else
@@ -1548,6 +1547,8 @@ EClothUtilRetType ImportApexAssetFromApexFile(FString& ApexFile, USkeletalMesh*
 	FName AssetName = FName(*FPaths::GetCleanFilename(ApexFile));
 	int32 NumAssets = SkelMesh->ClothingAssets.Num();
 
+	NxClothingAsset* ExistingAsset = nullptr;
+
 	if(!bReimport)
 	{
 		// new asset
@@ -1562,11 +1563,7 @@ EClothUtilRetType ImportApexAssetFromApexFile(FString& ApexFile, USkeletalMesh*
 	{
 		// re-import
 		FClothingAssetData& AssetData = SkelMesh->ClothingAssets[AssetIndex];
-		if (AssetData.ApexClothingAsset)
-		{
-			GPhysCommandHandler->DeferredRelease(AssetData.ApexClothingAsset);
-			AssetData.ApexClothingAsset = NULL;
-		}
+		ExistingAsset = AssetData.ApexClothingAsset;
 		AssetData.ApexClothingAsset = ApexClothingAsset;
 		AssetData.ApexFileName = ApexFile;
 
@@ -1596,14 +1593,22 @@ EClothUtilRetType ImportApexAssetFromApexFile(FString& ApexFile, USkeletalMesh*
 					if(!ImportClothingSectionFromClothingAsset(SkelMesh, LODIndex, SecIdx, AssetIndex, SubmeshIdx))
 					{
 						// if import is failed, restores original section because sub-mesh might be removed
-						RestoreOriginalClothingSection(SkelMesh, LODIndex, SecIdx);
+						RestoreOriginalClothingSection(SkelMesh, LODIndex, SecIdx, false);
 					}
 				}
 			}
 		}
 	}
 
+	SkelMesh->BuildApexToUnrealBoneMapping();
 	SkelMesh->LoadClothCollisionVolumes(AssetIndex, ApexClothingAsset);
+	ReregisterSkelMeshComponents(SkelMesh);
+
+	//We've recreated the component and the cloth actors so it's now safe to remove the asset.
+	if(ExistingAsset)
+	{
+		GPhysCommandHandler->DeferredRelease(ExistingAsset);
+	}
 
 	return CURT_Ok;
 }
@@ -1864,6 +1869,8 @@ void ReapplyClothingDataToSkeletalMesh( USkeletalMesh* SkelMesh, FClothingBackup
 		);
 
 	}
+
+	SkelMesh->BuildApexToUnrealBoneMapping();
 }
 
 int32 GetNumLODs(NxClothingAsset *InAsset)
diff --git a/Engine/Source/Editor/UnrealEd/Public/ApexClothingUtils.h b/Engine/Source/Editor/UnrealEd/Public/ApexClothingUtils.h
--- a/Engine/Source/Editor/UnrealEd/Public/ApexClothingUtils.h
+++ b/Engine/Source/Editor/UnrealEd/Public/ApexClothingUtils.h
@@ -67,7 +67,7 @@ namespace ApexClothingUtils
 	// Function to restore all clothing section to original mesh section related to specified asset index
 	UNREALED_API void RemoveAssetFromSkeletalMesh(USkeletalMesh* SkelMesh, uint32 AssetIndex, bool bRecreateSkelMeshComponent = false);
 	// Function to restore clothing section to original mesh section
-	UNREALED_API void RestoreOriginalClothingSection(USkeletalMesh* SkelMesh, uint32 LODIndex, uint32 SectionIndex);
+	UNREALED_API void RestoreOriginalClothingSection(USkeletalMesh* SkelMesh, uint32 LODIndex, uint32 SectionIndex, bool bReregisterSkelMeshComponent = true);
 	UNREALED_API void RestoreOriginalClothingSectionAllLODs(USkeletalMesh* SkelMesh, uint32 LOD0_SectionIndex);
 
 	// using global variable for backing up clothing data
diff --git a/Engine/Source/Runtime/Engine/Classes/Components/SkeletalMeshComponent.h b/Engine/Source/Runtime/Engine/Classes/Components/SkeletalMeshComponent.h
--- a/Engine/Source/Runtime/Engine/Classes/Components/SkeletalMeshComponent.h
+++ b/Engine/Source/Runtime/Engine/Classes/Components/SkeletalMeshComponent.h
@@ -66,11 +66,16 @@ public:
 	physx::apex::NxClothingAsset*	ParentClothingAsset;
 	/** APEX clothing actor is created from APEX clothing asset for cloth simulation */
 	physx::apex::NxClothingActor*		ApexClothingActor;
+
+	/** The corresponding clothing asset index */
+	int32 ParentClothingAssetIndex;
 };
 
 //The data that cloth needs for simulation prep in parallel. These properties are accessible via double buffer
 struct FClothSimulationContext
 {
+	FClothSimulationContext();
+
 	/** used for pre-computation using TeleportRotationThreshold property */
 	float ClothTeleportCosineThresholdInRad;
 	/** used for pre-computation using tTeleportDistanceThreshold property */
@@ -81,6 +86,10 @@ private:
 	/** previous root bone matrix to compare the difference and decide to do clothing teleport  */
 	FMatrix	PrevRootBoneMatrix;
 
+	USkinnedMeshComponent* InMasterPoseComponent;
+	TArray<int32> InMasterBoneMap;
+	int32 InMasterBoneMapCacheCount;
+
 	friend class USkeletalMeshComponent;
 };
 
@@ -752,7 +761,7 @@ public:
 	void ReleaseAllClothingResources();
 	bool IsExecutingParallelCloth() const;
 
-	bool IsValidClothingActor(int32 ActorIndex) const;
+	bool IsValidClothingActor(const FClothingActor& ClothingActor) const;
 	/** Draws APEX Clothing simulated normals on cloth meshes **/
 	void DrawClothingNormals(FPrimitiveDrawInterface* PDI);
 	/** Draws APEX Clothing Graphical Tangents on cloth meshes **/
@@ -774,7 +783,7 @@ public:
 	/** Loads clothing extra infos dynamically just for Previewing in Editor 
 	 *  such as MaxDistances, Physical mesh wire
 	 **/
-	void LoadClothingVisualizationInfo(int32 AssetIndex);
+	void LoadClothingVisualizationInfo(FClothingAssetData& ClothAssetData);
 	void LoadAllClothingVisualizationInfos();
 
 	/** freezing clothing actor now */
@@ -1135,7 +1144,6 @@ public:
 	void SetClothingLOD(int32 LODIndex);
 	/** check whether clothing teleport is needed or not to avoid a weird simulation result */
 	virtual void CheckClothTeleport(float DeltaTime, FClothSimulationContext& ClothSimulationContext) const;
-
 	/** 
 	* methods for cloth morph targets 
 	*/
@@ -1157,7 +1165,7 @@ public:
 	void UpdateClothTransform();
 
 	/** only check whether there are valid clothing actors or not */
-	bool HasValidClothingActors();
+	bool HasValidClothingActors() const;
 
 	/** get root bone matrix by the root bone index which Apex cloth asset is holding */
 	void GetClothRootBoneMatrix(int32 AssetIndex, FMatrix& OutRootBoneMatrix) const;
diff --git a/Engine/Source/Runtime/Engine/Classes/Components/SkinnedMeshComponent.h b/Engine/Source/Runtime/Engine/Classes/Components/SkinnedMeshComponent.h
--- a/Engine/Source/Runtime/Engine/Classes/Components/SkinnedMeshComponent.h
+++ b/Engine/Source/Runtime/Engine/Classes/Components/SkinnedMeshComponent.h
@@ -193,13 +193,19 @@ protected:
 	 * Note this is weak object ptr, so it will go away unless you have other strong reference
 	 */
 	TArray< TWeakObjectPtr<USkinnedMeshComponent> > SlavePoseComponents;
-public:
+
 	/**
 	 *	Mapping between bone indices in this component and the parent one. Each element is the index of the bone in the MasterPoseComponent.
 	 *	Size should be the same as USkeletalMesh.RefSkeleton size (ie number of bones in this skeleton).
 	 */
 	TArray<int32> MasterBoneMap;
 
+	/** Incremented every time the master bone map changes. Used to keep in sync with any duplicate data needed by other threads */
+	int32 MasterBoneMapCacheCount;
+public:
+
+	const TArray<int32>& GetMasterBoneMap() const { return MasterBoneMap; }
+
 	/** 
 	 * When true, we will just using the bounds from our MasterPoseComponent.  This is useful for when we have a Mesh Parented
 	 * to the main SkelMesh (e.g. outline mesh or a full body overdraw effect that is toggled) that is always going to be the same
diff --git a/Engine/Source/Runtime/Engine/Classes/Engine/SkeletalMesh.h b/Engine/Source/Runtime/Engine/Classes/Engine/SkeletalMesh.h
--- a/Engine/Source/Runtime/Engine/Classes/Engine/SkeletalMesh.h
+++ b/Engine/Source/Runtime/Engine/Classes/Engine/SkeletalMesh.h
@@ -856,6 +856,7 @@ public:
 	ENGINE_API void  LoadClothCollisionVolumes(int32 AssetIndex, physx::apex::NxClothingAsset* ClothingAsset);
 	ENGINE_API bool IsMappedClothingLOD(int32 LODIndex, int32 AssetIndex);
 	ENGINE_API int32 GetClothAssetIndex(int32 LODIndex, int32 SectionIndex);
+	ENGINE_API void BuildApexToUnrealBoneMapping();
 #endif
 
 	/** 
diff --git a/Engine/Source/Runtime/Engine/Private/Components/SkeletalMeshComponent.cpp b/Engine/Source/Runtime/Engine/Private/Components/SkeletalMeshComponent.cpp
--- a/Engine/Source/Runtime/Engine/Private/Components/SkeletalMeshComponent.cpp
+++ b/Engine/Source/Runtime/Engine/Private/Components/SkeletalMeshComponent.cpp
@@ -1019,6 +1019,8 @@ void USkeletalMeshComponent::UpdateSlaveComponent()
 		}
 	}
 
+	PrepareCloth();
+
 	Super::UpdateSlaveComponent();
 }
 
@@ -1066,6 +1068,13 @@ void USkeletalMeshComponent::SubmitClothSimulationContext()
 	InternalClothSimulationContext.ClothTeleportMode = EditableClothSimulationContext.ClothTeleportMode == FClothingActor::Default ? InternalClothSimulationContext.ClothTeleportMode : EditableClothSimulationContext.ClothTeleportMode;
 	EditableClothSimulationContext.ClothTeleportMode = FClothingActor::Default;	//reset editable teleport mode to default because user has to set it directly
 
+	if(InternalClothSimulationContext.InMasterBoneMapCacheCount != MasterBoneMapCacheCount)
+	{
+		InternalClothSimulationContext.InMasterBoneMapCacheCount = MasterBoneMapCacheCount;
+		InternalClothSimulationContext.InMasterPoseComponent = MasterPoseComponent.Get();
+		InternalClothSimulationContext.InMasterBoneMap = MasterBoneMap;
+	}
+
 	//we intentionally ignore PrevRootBone which is only needed internally, but is still in the Context struct for the sake of const
 }
 #endif
@@ -1218,12 +1227,25 @@ void USkeletalMeshComponent::PrepareCloth()
 			FClothManager* ClothManager = PhysScene->GetClothManager();
 			bool bClothNeedsPhysics = BodyInstance.bSimulatePhysics || IsAnySimulatingPhysics();	//TODO: this errs on the side of simulating later. We may want to optimize this so that it's only needed when cloth bodies are simulating
 			PrepareClothSchedule PrepareSchedule = bClothNeedsPhysics ? PrepareClothSchedule::WaitOnPhysics : PrepareClothSchedule::IgnorePhysics;
+
+			SubmitClothSimulationContext();	//duplicate needed data for off-thread work
+
 			ClothManager->RegisterForPrepareCloth(this, PrepareSchedule);
 		}
 	}
 #endif
 }
 
+FClothSimulationContext::FClothSimulationContext()
+{
+	ClothTeleportCosineThresholdInRad = 0.f;
+	ClothTeleportDistThresholdSquared = 0.f;
+	ClothTeleportMode = FClothingActor::TeleportMode::Default;
+	
+	USkeletalMeshComponent* InMasterComponent = nullptr;
+	InMasterBoneMapCacheCount = -1;
+}
+
 void USkeletalMeshComponent::PostAnimEvaluation(FAnimationEvaluationContext& EvaluationContext)
 {
 	AnimEvaluationContext.Clear();
@@ -1282,6 +1304,7 @@ void USkeletalMeshComponent::PostAnimEvaluation(FAnimationEvaluationContext& Eva
 		PostBlendPhysics();
 	}
 
+
 	PrepareCloth();
 }
 
diff --git a/Engine/Source/Runtime/Engine/Private/Components/SkinnedMeshComponent.cpp b/Engine/Source/Runtime/Engine/Private/Components/SkinnedMeshComponent.cpp
--- a/Engine/Source/Runtime/Engine/Private/Components/SkinnedMeshComponent.cpp
+++ b/Engine/Source/Runtime/Engine/Private/Components/SkinnedMeshComponent.cpp
@@ -1141,6 +1141,8 @@ void USkinnedMeshComponent::UpdateMasterBoneMap()
 			}
 		}
 	}
+
+	MasterBoneMapCacheCount += 1;
 }
 
 FTransform USkinnedMeshComponent::GetSocketTransform(FName InSocketName, ERelativeTransformSpace TransformSpace) const
diff --git a/Engine/Source/Runtime/Engine/Private/LevelTick.cpp b/Engine/Source/Runtime/Engine/Private/LevelTick.cpp
--- a/Engine/Source/Runtime/Engine/Private/LevelTick.cpp
+++ b/Engine/Source/Runtime/Engine/Private/LevelTick.cpp
@@ -829,20 +829,20 @@ void UWorld::SendAllEndOfFrameUpdates()
 		{
 			QUICK_SCOPE_CYCLE_COUNTER(STAT_PostTickComponentUpdate_ForcedGameThread);
 			for (TArray<TWeakObjectPtr<UActorComponent> >::TIterator It(LocalComponentsThatNeedEndOfFrameUpdate_OnGameThread); It; ++It)
-			{
-				UActorComponent* Component = It->Get();
-				if (Component)
-				{
+	{
+		UActorComponent* Component = It->Get();
+		if (Component)
+		{
 					if (Component->IsRegistered() && !Component->IsTemplate())
-					{
-						FScopeCycleCounterUObject ComponentScope(Component);
-						FScopeCycleCounterUObject AdditionalScope(STATS ? Component->AdditionalStatObject() : NULL);
-						Component->DoDeferredRenderUpdates_Concurrent();
-					}
-					check(Component->GetMarkedForEndOfFrameUpdateState() == EComponentMarkedForEndOfFrameUpdateState::MarkedForGameThread);
-					FMarkComponentEndOfFrameUpdateState::Set(Component, EComponentMarkedForEndOfFrameUpdateState::Unmarked);
-				}
+			{
+				FScopeCycleCounterUObject ComponentScope(Component);
+				FScopeCycleCounterUObject AdditionalScope(STATS ? Component->AdditionalStatObject() : NULL);
+				Component->DoDeferredRenderUpdates_Concurrent();
 			}
+			check(Component->GetMarkedForEndOfFrameUpdateState() == EComponentMarkedForEndOfFrameUpdateState::MarkedForGameThread);
+			FMarkComponentEndOfFrameUpdateState::Set(Component, EComponentMarkedForEndOfFrameUpdateState::Unmarked);
+		}
+	}
 			LocalComponentsThatNeedEndOfFrameUpdate_OnGameThread.Reset();
 		};
 
@@ -851,11 +851,11 @@ void UWorld::SendAllEndOfFrameUpdates()
 		ParallelForWithPreWork(ComponentsThatNeedEndOfFrameUpdate.Num(), ParallelWork, GTWork);
 	}
 	else
-	{
+					{
 		GTWork();
 		ParallelFor(ComponentsThatNeedEndOfFrameUpdate.Num(), ParallelWork);
-	}
-	ComponentsThatNeedEndOfFrameUpdate.Reset();
+			}
+		ComponentsThatNeedEndOfFrameUpdate.Reset();
 }
 
 
diff --git a/Engine/Source/Runtime/Engine/Private/Particles/ParticleModules_Location.cpp b/Engine/Source/Runtime/Engine/Private/Particles/ParticleModules_Location.cpp
--- a/Engine/Source/Runtime/Engine/Private/Particles/ParticleModules_Location.cpp
+++ b/Engine/Source/Runtime/Engine/Private/Particles/ParticleModules_Location.cpp
@@ -2673,6 +2673,7 @@ bool UParticleModuleLocationSkelVertSurface::VertInfluencedByActiveBone(FParticl
 template<bool bExtraBoneInfluencesT>
 bool UParticleModuleLocationSkelVertSurface::VertInfluencedByActiveBoneTyped(bool bSoftVertex, FStaticLODModel& Model, const FSkelMeshChunk& Chunk, int32 VertIndex, USkeletalMeshComponent* InSkelMeshComponent, FModuleLocationVertSurfaceInstancePayload* InstancePayload, int32* OutBoneIndex)
 {
+	const TArray<int32>& MasterBoneMap = InSkelMeshComponent->GetMasterBoneMap();
 	// Do soft skinning for this vertex.
 	if(bSoftVertex)
 	{
@@ -2688,8 +2689,8 @@ bool UParticleModuleLocationSkelVertSurface::VertInfluencedByActiveBoneTyped(boo
 			int32 BoneIndex = Chunk.BoneMap[SrcSoftVertex->InfluenceBones[InfluenceIndex]];
 			if(InSkelMeshComponent->MasterPoseComponent.IsValid())
 			{		
-				check(InSkelMeshComponent->MasterBoneMap.Num() == InSkelMeshComponent->SkeletalMesh->RefSkeleton.GetNum());
-				BoneIndex = InSkelMeshComponent->MasterBoneMap[BoneIndex];
+				check(MasterBoneMap.Num() == InSkelMeshComponent->SkeletalMesh->RefSkeleton.GetNum());
+				BoneIndex = MasterBoneMap[BoneIndex];
 			}
 
 			if(!InstancePayload->NumValidAssociatedBoneIndices || InstancePayload->ValidAssociatedBoneIndices.Contains(BoneIndex))
@@ -2712,8 +2713,8 @@ bool UParticleModuleLocationSkelVertSurface::VertInfluencedByActiveBoneTyped(boo
 		int32 BoneIndex = Chunk.BoneMap[SrcRigidVertex->InfluenceBones[RigidInfluenceIndex]];
 		if(InSkelMeshComponent->MasterPoseComponent.IsValid())
 		{
-			check(InSkelMeshComponent->MasterBoneMap.Num() == InSkelMeshComponent->SkeletalMesh->RefSkeleton.GetNum());
-			BoneIndex = InSkelMeshComponent->MasterBoneMap[BoneIndex];
+			check(MasterBoneMap.Num() == InSkelMeshComponent->SkeletalMesh->RefSkeleton.GetNum());
+			BoneIndex = MasterBoneMap[BoneIndex];
 		}
 
 		if(!InstancePayload->NumValidAssociatedBoneIndices || InstancePayload->ValidAssociatedBoneIndices.Contains(BoneIndex))
diff --git a/Engine/Source/Runtime/Engine/Private/PhysicsEngine/ClothManager.cpp b/Engine/Source/Runtime/Engine/Private/PhysicsEngine/ClothManager.cpp
--- a/Engine/Source/Runtime/Engine/Private/PhysicsEngine/ClothManager.cpp
+++ b/Engine/Source/Runtime/Engine/Private/PhysicsEngine/ClothManager.cpp
@@ -10,7 +10,9 @@ FClothManager::FClothManager(UWorld* InAssociatedWorld)
 {
 }
 
-void FClothManagerData::PrepareCloth(float DeltaTime)
+static TAutoConsoleVariable<int32> CVarParallelCloth(TEXT("p.ParallelCloth"), 0, TEXT("Whether to prepare and start cloth sim off the game thread"));
+
+void FClothManagerData::PrepareCloth(float DeltaTime, FTickFunction& TickFunction)
 {
 #if WITH_APEX_CLOTHING
 	QUICK_SCOPE_CYCLE_COUNTER(STAT_FClothManager_PrepareCloth);
@@ -19,10 +21,16 @@ void FClothManagerData::PrepareCloth(float DeltaTime)
 		IsPreparingCloth.AtomicSet(true);
 		for (USkeletalMeshComponent* SkeletalMeshComponent : SkeletalMeshComponents)
 		{
-			SkeletalMeshComponent->SubmitClothSimulationContext();	//make sure user params are passed internally
-
 			FClothSimulationContext& ClothSimulationContext = SkeletalMeshComponent->InternalClothSimulationContext;
-			SkeletalMeshComponent->ParallelTickClothing(DeltaTime, ClothSimulationContext);
+			if(CVarParallelCloth.GetValueOnAnyThread())
+			{
+				SkeletalMeshComponent->ParallelTickClothing(DeltaTime, ClothSimulationContext);
+			}else
+			{
+				check(IsInGameThread());
+				SkeletalMeshComponent->TickClothing(DeltaTime, TickFunction);
+			}
+			
 		}
 
 		IsPreparingCloth.AtomicSet(false);
@@ -39,7 +47,7 @@ void FClothManager::RegisterForPrepareCloth(USkeletalMeshComponent* SkeletalMesh
 void FClothManager::StartCloth(float DeltaTime)
 {
 	//First we must prepare any cloth that waits on physics sim
-	PrepareClothDataArray[(int32)PrepareClothSchedule::WaitOnPhysics].PrepareCloth(DeltaTime);
+	PrepareClothDataArray[(int32)PrepareClothSchedule::WaitOnPhysics].PrepareCloth(DeltaTime, StartClothTickFunction);
 
 	//Reset skeletal mesh components
 	bool bNeedSimulateCloth = false;
@@ -62,8 +70,6 @@ void FClothManager::StartCloth(float DeltaTime)
 	}
 }
 
-static TAutoConsoleVariable<int32> CVarParallelCloth(TEXT("p.ParallelCloth"), 0, TEXT("Whether to prepare and start cloth sim off the game thread"));
-
 void FClothManager::SetupClothTickFunction(bool bTickClothSim)
 {
 	check(IsInGameThread());
@@ -117,7 +123,7 @@ void FClothManager::EndCloth()
 
 void FStartIgnorePhysicsClothTickFunction::ExecuteTick(float DeltaTime, enum ELevelTick TickType, ENamedThreads::Type CurrentThread, const FGraphEventRef& MyCompletionGraphEvent)
 {
-	Target->PrepareClothDataArray[(int32)PrepareClothSchedule::IgnorePhysics].PrepareCloth(DeltaTime);
+	Target->PrepareClothDataArray[(int32)PrepareClothSchedule::IgnorePhysics].PrepareCloth(DeltaTime, *this);
 }
 
 FString FStartIgnorePhysicsClothTickFunction::DiagnosticMessage()
diff --git a/Engine/Source/Runtime/Engine/Private/PhysicsEngine/ClothManager.h b/Engine/Source/Runtime/Engine/Private/PhysicsEngine/ClothManager.h
--- a/Engine/Source/Runtime/Engine/Private/PhysicsEngine/ClothManager.h
+++ b/Engine/Source/Runtime/Engine/Private/PhysicsEngine/ClothManager.h
@@ -27,7 +27,7 @@ public:
 	FGraphEventRef PrepareCompletion;
 
 	/** Go through all SkeletalMeshComponents registered for work this frame and call the solver as needed*/
-	void PrepareCloth(float DeltaTime);
+	void PrepareCloth(float DeltaTime, FTickFunction& TickFunction);
 };
 
 struct FStartIgnorePhysicsClothTickFunction : public FTickFunction
diff --git a/Engine/Source/Runtime/Engine/Private/SkeletalMesh.cpp b/Engine/Source/Runtime/Engine/Private/SkeletalMesh.cpp
--- a/Engine/Source/Runtime/Engine/Private/SkeletalMesh.cpp
+++ b/Engine/Source/Runtime/Engine/Private/SkeletalMesh.cpp
@@ -2430,19 +2430,22 @@ bool USkeletalMesh::IsReadyForFinishDestroy()
 FName GetConvertedBoneName(NxClothingAsset* ApexClothingAsset, int32 BoneIndex);
 
 
-void BuildApexToUnrealBoneMapping(const USkeletalMesh* SkeletalMesh, FClothingAssetData& ClothingAsset)
+void USkeletalMesh::BuildApexToUnrealBoneMapping()
 {
-	NxClothingAsset* ApexClothingAsset = ClothingAsset.ApexClothingAsset;
-	uint32 NumUsedBones = ClothingAsset.ApexClothingAsset->getNumUsedBones();
-	
-	check(ClothingAsset.ApexToUnrealBoneMapping.Num() == 0);
-	ClothingAsset.ApexToUnrealBoneMapping.AddUninitialized(NumUsedBones);
-
-	for (uint32 Index = 0; Index < NumUsedBones; Index++)
+	for(FClothingAssetData& ClothingAsset : ClothingAssets)
 	{
-		FName BoneName = GetConvertedBoneName(ApexClothingAsset, Index);
-		int32 BoneIndex = SkeletalMesh->RefSkeleton.FindBoneIndex(BoneName);
-		ClothingAsset.ApexToUnrealBoneMapping[Index] = BoneIndex;
+		NxClothingAsset* ApexClothingAsset = ClothingAsset.ApexClothingAsset;
+		uint32 NumUsedBones = ClothingAsset.ApexClothingAsset->getNumUsedBones();
+
+		ClothingAsset.ApexToUnrealBoneMapping.Empty();
+		ClothingAsset.ApexToUnrealBoneMapping.AddUninitialized(NumUsedBones);
+
+		for (uint32 Index = 0; Index < NumUsedBones; Index++)
+		{
+			FName BoneName = GetConvertedBoneName(ApexClothingAsset, Index);
+			int32 BoneIndex = RefSkeleton.FindBoneIndex(BoneName);
+			ClothingAsset.ApexToUnrealBoneMapping[Index] = BoneIndex;
+		}
 	}
 }
 #endif
@@ -2539,13 +2542,14 @@ void USkeletalMesh::Serialize( FArchive& Ar )
 		for( int32 Idx=0;Idx<ClothingAssets.Num();Idx++ )
 		{
 			Ar << ClothingAssets[Idx];
+		}
+
 #if WITH_APEX_CLOTHING
-			if(Ar.IsLoading())
-			{
-				BuildApexToUnrealBoneMapping(this, ClothingAssets[Idx]);
-			}
-#endif
+		if (Ar.IsLoading())
+		{
+			BuildApexToUnrealBoneMapping();
 		}
+#endif
 	}
 
 	if ( Ar.IsLoading() && Ar.UE4Ver() < VER_UE4_MOVE_SKELETALMESH_SHADOWCASTING )
@@ -4801,6 +4805,7 @@ USkinnedMeshComponent::USkinnedMeshComponent(const FObjectInitializer& ObjectIni
 	bNeedToFlipSpaceBaseBuffers = false;
 
 	bCanEverAffectNavigation = false;
+	MasterBoneMapCacheCount = 0;
 }
 
 
diff --git a/Engine/Source/Runtime/Engine/Private/SkeletalMeshComponentPhysics.cpp b/Engine/Source/Runtime/Engine/Private/SkeletalMeshComponentPhysics.cpp
--- a/Engine/Source/Runtime/Engine/Private/SkeletalMeshComponentPhysics.cpp
+++ b/Engine/Source/Runtime/Engine/Private/SkeletalMeshComponentPhysics.cpp
@@ -70,6 +70,7 @@ void FClothingActor::Clear(bool bReleaseResource)
 
 	ParentClothingAsset = NULL;
 	ApexClothingActor = NULL;
+	ParentClothingAssetIndex = INDEX_NONE;
 }
 
 //
@@ -1997,15 +1998,9 @@ bool USkeletalMeshComponent::ComponentOverlapMultiImpl(TArray<struct FOverlapRes
 
 void USkeletalMeshComponent::AddClothingBounds(FBoxSphereBounds& InOutBounds) const
 {
-	int32 NumAssets = ClothingActors.Num();
-
-	for(int32 i=0; i < NumAssets; i++)
-	{
-		if(IsValidClothingActor(i))
+	for(const FClothingActor& ClothingActor : ClothingActors)
 		{
-			NxClothingActor* Actor = ClothingActors[i].ApexClothingActor;
-
-			if(Actor)
+		if(NxClothingActor* Actor = ClothingActor.ApexClothingActor)
 			{
 				physx::PxBounds3 ApexClothingBounds = Actor->getBounds();
 
@@ -2016,14 +2011,13 @@ void USkeletalMeshComponent::AddClothingBounds(FBoxSphereBounds& InOutBounds) co
 				}
 			}
 		}
-	}
 }
 
-bool USkeletalMeshComponent::HasValidClothingActors()
+bool USkeletalMeshComponent::HasValidClothingActors() const
 {
-	for(int32 i=0; i < ClothingActors.Num(); i++)
+	for(const FClothingActor& ClothingActor : ClothingActors)
 	{
-		if(IsValidClothingActor(i))
+		if(IsValidClothingActor(ClothingActor))
 		{
 			return true;
 		}
@@ -2061,8 +2055,7 @@ void USkeletalMeshComponent::RecreateClothingActors()
 */
 bool USkeletalMeshComponent::CreateClothingActor(int32 AssetIndex, physx::apex::NxClothingAsset* ClothingAsset, TArray<FVector>* BlendedDelta)
 {	
-	check(ClothingActors.Num() == AssetIndex);
-	ClothingActors.AddZeroed();
+	int32 ActorIndex = ClothingActors.AddZeroed();
 
 	UWorld* World = GetWorld();
 	FPhysScene* PhysScene = World ? World->GetPhysicsScene() : nullptr;
@@ -2074,15 +2067,15 @@ bool USkeletalMeshComponent::CreateClothingActor(int32 AssetIndex, physx::apex::
 	}
 
 	//TODO: this is a workaround while waiting for fix: apex cloth factory lifetime should depend on cloth asset. For now we make sure to only create asset if the scene ticks which is the only way to get refcounted
-	if(World->bShouldSimulatePhysics == false)
-	{
-		return false;
-	}
+	if (World->bShouldSimulatePhysics == false)
+		{
+			return false;
+		}
 
 	physx::apex::NxApexScene* ScenePtr = PhysScene->GetApexScene(PST_Cloth);
 
 	if (!ScenePtr)
-	{
+		{
 		// can't create clothing actor
 		UE_LOG(LogSkeletalMesh, Log, TEXT("CreateClothingActor: Failed to create an actor becauase PhysX Scene doesn't exist"));
 		return false;
@@ -2147,27 +2140,30 @@ bool USkeletalMeshComponent::CreateClothingActor(int32 AssetIndex, physx::apex::
 	}
 
 	physx::apex::NxApexActor* apexActor = ClothingAsset->createApexActor(*ActorDesc, *ScenePtr);
-	physx::apex::NxClothingActor* ClothingActor = static_cast<physx::apex::NxClothingActor*>(apexActor);
+	physx::apex::NxClothingActor* ApexClothingActor = static_cast<physx::apex::NxClothingActor*>(apexActor);
 
-	ClothingActors[AssetIndex].ApexClothingActor = ClothingActor;
+	FClothingActor& ClothingActor = ClothingActors[ActorIndex];
+	
+	ClothingActor.ApexClothingActor = ApexClothingActor;
 
-	if(!ClothingActor)
+	if(!ApexClothingActor)
 	{
 		UE_LOG(LogSkeletalMesh, Log, TEXT("CreateClothingActor: Failed to create an clothing actor (%s)"), ANSI_TO_TCHAR(ClothingAsset->getName()));
 		return false;
 	}
 
 	//set parent pointer to verify later whether became invalid or not
-	ClothingActors[AssetIndex].ParentClothingAsset = ClothingAsset;
+	ClothingActor.ParentClothingAsset = ClothingAsset;
+	ClothingActor.ParentClothingAssetIndex = AssetIndex;
 
 	// budget is millisecond units
 	ScenePtr->setLODResourceBudget(100); // for temporary, 100ms
 
-	ClothingActor->setGraphicalLOD(PredictedLODLevel);
+	ApexClothingActor->setGraphicalLOD(PredictedLODLevel);
 
 	// 0 means no simulation
-	ClothingActor->forcePhysicalLod(1); // 1 will be changed to "GetActivePhysicalLod()" later
-	ClothingActor->setFrozen(false);
+	ApexClothingActor->forcePhysicalLod(1); // 1 will be changed to "GetActivePhysicalLod()" later
+	ApexClothingActor->setFrozen(false);
 
 #if WITH_CLOTH_COLLISION_DETECTION
 	// process clothing collisions once for the case that this component doesn't move
@@ -2183,49 +2179,47 @@ bool USkeletalMeshComponent::CreateClothingActor(int32 AssetIndex, physx::apex::
 void USkeletalMeshComponent::SetClothingLOD(int32 LODIndex)
 {
 	//TODO: are we sure this can't be called during parallel animation?
-	int32 NumActors = ClothingActors.Num();
 
 	bool bFrozen = false;
 
-	for (int32 AssetIndex = 0; AssetIndex<NumActors; AssetIndex++)
+	for (FClothingActor& ClothingActor : ClothingActors)
 	{
-		if (IsValidClothingActor(AssetIndex))
+		if (IsValidClothingActor(ClothingActor))
 		{			
-			FClothingActor& Actor = ClothingActors[AssetIndex];
-
-			int32 CurLODIndex = (int32)Actor.ApexClothingActor->getGraphicalLod();
+			int32 CurLODIndex = (int32)ClothingActor.ApexClothingActor->getGraphicalLod();
 
 			// check whether clothing LOD is mapped for this LOD index
-			bool IsMappedClothLOD = SkeletalMesh->IsMappedClothingLOD(LODIndex, AssetIndex);
+			bool IsMappedClothLOD = SkeletalMesh->IsMappedClothingLOD(LODIndex, ClothingActor.ParentClothingAssetIndex);
 
 			// Change Clothing LOD if a new LOD index is different from the current index
 			if (IsMappedClothLOD && CurLODIndex != LODIndex)
 			{
 				//physical LOD is changed by graphical LOD
-				Actor.ApexClothingActor->setGraphicalLOD(LODIndex);
+				ClothingActor.ApexClothingActor->setGraphicalLOD(LODIndex);
 
-				if (Actor.ApexClothingActor->isFrozen())
+				if (ClothingActor.ApexClothingActor->isFrozen())
 				{
 					bFrozen = true;
 				}
 			}
 
-			int32 NumClothLODs = Actor.ParentClothingAsset->getNumGraphicalLodLevels();
+			int32 NumClothLODs = ClothingActor.ParentClothingAsset->getNumGraphicalLodLevels();
 
 			// NOTE: forcePhysicalLod() is an deferred operation.   getActivePhysicalLod() may return previous state.
 			// decide whether should enable or disable
 			if (!IsMappedClothLOD || (LODIndex >= NumClothLODs))
 			{
 				//disable clothing simulation
-				Actor.ApexClothingActor->forcePhysicalLod(0);
+				ClothingActor.ApexClothingActor->forcePhysicalLod(0);
 			}
 			else
 			{	
-				// enable clothing simulation
-				Actor.ApexClothingActor->forcePhysicalLod(1);
+				int32 CurPhysLOD = ClothingActor.ApexClothingActor->getActivePhysicalLod();
+				// if disabled, enables clothing simulation
+				ClothingActor.ApexClothingActor->forcePhysicalLod(1);
+				}
 			}
 		}
-	}
 
 	if(bFrozen)
 	{
@@ -2248,11 +2242,10 @@ bool USkeletalMeshComponent::IsExecutingParallelCloth() const
 void USkeletalMeshComponent::RemoveAllClothingActors()
 {
 	check(!IsExecutingParallelCloth());
-	int32 NumActors = ClothingActors.Num();
 
-	for(int32 i=0; i<NumActors; i++)
+	for(FClothingActor& ClothingActor : ClothingActors)
 	{
-		ClothingActors[i].Clear(IsValidClothingActor(i));
+		ClothingActor.Clear(IsValidClothingActor(ClothingActor));
 	}
 
 	ClothingActors.Empty();
@@ -2713,15 +2706,15 @@ void USkeletalMeshComponent::CreateInternalClothCollisions(TArray<FApexClothColl
 	NumCollisions = FMath::Min(NumCollisions, MaxNumCapsules);
 	int32 NumActors = ClothingActors.Num();
 
-		for(int32 ActorIdx=0; ActorIdx < NumActors; ActorIdx++)
+	for(FClothingActor& ClothingActor : ClothingActors)
 		{
-			if (!IsValidClothingActor(ActorIdx))
+		if (!IsValidClothingActor(ClothingActor))
 			{
 				continue;
 			}
 
-		NxClothingActor* Actor = ClothingActors[ActorIdx].ApexClothingActor;
-		int32 NumCurrentCapsules = SkeletalMesh->ClothingAssets[ActorIdx].ClothCollisionVolumes.Num(); // # of capsules 
+		NxClothingActor* Actor = ClothingActor.ApexClothingActor;
+		int32 NumCurrentCapsules = SkeletalMesh->ClothingAssets[ClothingActor.ParentClothingAssetIndex].ClothCollisionVolumes.Num(); // # of capsules 
 
 			for(int32 ColIdx=0; ColIdx < NumCollisions; ColIdx++)
 			{
@@ -2773,7 +2766,7 @@ void USkeletalMeshComponent::CreateInternalClothCollisions(TArray<FApexClothColl
 				}
 			}
 		}
-	}
+}
 
 void USkeletalMeshComponent::CopyClothCollisionsToChildren()
 {
@@ -2908,9 +2901,9 @@ FApexClothCollisionInfo* USkeletalMeshComponent::CreateNewClothingCollsions(UPri
 
 	int32 NumActors = ClothingActors.Num();
 
-	for(int32 ActorIdx=0; ActorIdx < NumActors; ActorIdx++)
+	for(FClothingActor& ClothingActor : ClothingActors)
 	{
-		NxClothingActor* Actor = ClothingActors[ActorIdx].ApexClothingActor;
+		NxClothingActor* Actor = ClothingActor.ApexClothingActor;
 
 		if(Actor)
 		{	
@@ -3192,21 +3185,13 @@ void USkeletalMeshComponent::UpdateClothTransform()
 
 	physx::PxMat44 PxGlobalPose = U2PMatrix(ComponentToWorld.ToMatrixWithScale());
 
-	for(int32 ActorIdx=0; ActorIdx<NumActors; ActorIdx++)
+	for(FClothingActor& ClothingActor : ClothingActors)
 	{
-		// skip if ClothingActor is NULL or invalid
-		if(!IsValidClothingActor(ActorIdx))
+		if(NxClothingActor* ApexClothingActor = ClothingActor.ApexClothingActor)
 		{
-			continue;
+			NxParameterized::Interface* ActorDesc = ApexClothingActor->getActorDesc();
+			verify(NxParameterized::setParamMat44(*ActorDesc, "globalPose", PxGlobalPose));
 		}
-
-		NxClothingActor* ClothingActor = ClothingActors[ActorIdx].ApexClothingActor;
-
-		check(ClothingActor);
-
-		NxParameterized::Interface* ActorDesc = ClothingActor->getActorDesc();
-
-		verify(NxParameterized::setParamMat44(*ActorDesc, "globalPose", PxGlobalPose));
 	}
 }
 
@@ -3247,6 +3232,7 @@ void USkeletalMeshComponent::CheckClothTeleport(float DeltaTime, FClothSimulatio
 
 void USkeletalMeshComponent::ChangeClothMorphTargetMapping(FClothMorphTargetData& MorphData, FName CurrentActivatedMorphName)
 {
+	check(!IsExecutingParallelCloth());
 	int32 AssetIndex = MorphData.ClothAssetIndex;
 
 	if(SkeletalMesh && SkeletalMesh->ClothingAssets.IsValidIndex(AssetIndex))
@@ -3296,6 +3282,8 @@ void USkeletalMeshComponent::PrepareClothMorphTargets()
 		return;
 	}
 
+	check(!IsExecutingParallelCloth());
+
 	ClothMorphTargets.Empty();
 
 	for (UMorphTarget* MorphTarget : SkeletalMesh->MorphTargets)
@@ -3425,6 +3413,8 @@ void USkeletalMeshComponent::UpdateClothMorphTarget()
 		return;
 	}
 
+	check(!IsExecutingParallelCloth());
+
 	for (int32 MorphIdx = 0; MorphIdx < ActiveVertexAnims.Num(); MorphIdx++)
 	{
 		if (ActiveVertexAnims[MorphIdx].Weight > 0.0f)
@@ -3464,9 +3454,13 @@ void USkeletalMeshComponent::UpdateClothMorphTarget()
 
 					if (OriginDelta.Num() > 0)
 					{
-						int32 ActorIndex = MorphData.ClothAssetIndex;
-
-						if (!IsValidClothingActor(ActorIndex))
+						int32 AssetIndex = MorphData.ClothAssetIndex;
+						for(int32 ActorIndex = 0; ActorIndex < ClothingActors.Num(); ++ActorIndex)
+						{
+							const FClothingActor& ClothingActor = ClothingActors[ActorIndex];
+							if(ClothingActor.ParentClothingAssetIndex == AssetIndex)
+							{
+								if (!IsValidClothingActor(ClothingActor))
 						{
 							continue;
 						}
@@ -3478,16 +3472,18 @@ void USkeletalMeshComponent::UpdateClothMorphTarget()
 							BlendedDelta[DeltaIdx] = (OriginDelta[DeltaIdx] * CurWeight);
 						}
 
-						// make current actor invalid
-						ClothingActors[ActorIndex].Clear(true);
-						CreateClothingActor(ActorIndex, SkeletalMesh->ClothingAssets[ActorIndex].ApexClothingAsset, &BlendedDelta);
+								CreateClothingActor(AssetIndex, SkeletalMesh->ClothingAssets[AssetIndex].ApexClothingAsset, &BlendedDelta);
+								//Remove the old actor
+								ClothingActors.RemoveAt(ActorIndex);
+								break;
+							}
+						}
 					}
 				}
 			}
 		}
 	}
 }
-
 void USkeletalMeshComponent::UpdateClothState(float DeltaTime)
 {
 	if (CVarEnableClothPhysics.GetValueOnGameThread() == 0)
@@ -3518,7 +3514,7 @@ void USkeletalMeshComponent::UpdateClothState(float DeltaTime)
 	}
 #endif // WITH_CLOTH_COLLISION_DETECTION
 
-	ParallelUpdateClothState(DeltaTime, EditableClothSimulationContext);
+	ParallelUpdateClothState(DeltaTime, InternalClothSimulationContext);
 }
 
 void USkeletalMeshComponent::ParallelUpdateClothState(float DeltaTime, FClothSimulationContext& ClothSimulationContext) const
@@ -3528,13 +3524,7 @@ void USkeletalMeshComponent::ParallelUpdateClothState(float DeltaTime, FClothSim
 		return;
 	}
 
-	int32 NumActors = ClothingActors.Num();
-	if (NumActors == 0)
-	{
-		return;
-	}
-
-	const TArray<FTransform>& BoneTransforms = MasterPoseComponent.IsValid() ? MasterPoseComponent.Get()->GetSpaceBases() : GetSpaceBases();
+	const TArray<FTransform>& BoneTransforms = ClothSimulationContext.InMasterPoseComponent ? ClothSimulationContext.InMasterPoseComponent->GetSpaceBases() : GetSpaceBases();
 
 	if (BoneTransforms.Num() == 0)
 	{
@@ -3550,10 +3540,9 @@ void USkeletalMeshComponent::ParallelUpdateClothState(float DeltaTime, FClothSim
 
 	TArray<physx::PxMat44> BoneMatrices;
 
-	for (int32 ClothIndex = 0; ClothIndex < ClothingActors.Num(); ++ClothIndex)
-	{
-		const FClothingActor& ClothingActor = ClothingActors[ClothIndex];
-		const FClothingAssetData& ClothingAsset = SkeletalMesh->ClothingAssets[ClothIndex];
+	for (const FClothingActor& ClothingActor : ClothingActors)
+		{
+		const FClothingAssetData& ClothingAsset = SkeletalMesh->ClothingAssets[ClothingActor.ParentClothingAssetIndex];
 
 		ApplyWindForCloth(ClothingActor);
 
@@ -3568,17 +3557,17 @@ void USkeletalMeshComponent::ParallelUpdateClothState(float DeltaTime, FClothSim
 		{
 			int32 BoneIndex = ClothingAsset.ApexToUnrealBoneMapping[Index];
 			
-			if (MasterPoseComponent.IsValid())
+			if (ClothSimulationContext.InMasterPoseComponent)
 			{
 				int32 TempBoneIndex = BoneIndex;
 				BoneIndex = INDEX_NONE;
 				if (TempBoneIndex < MasterBoneMap.Num())
 				{
-					int32 MasterBoneIndex = MasterBoneMap[TempBoneIndex];
+					int32 MasterBoneIndex = ClothSimulationContext.InMasterBoneMap[TempBoneIndex];
 
 					// If ParentBoneIndex is valid, grab matrix from MasterPoseComponent.
 					if (MasterBoneIndex != INDEX_NONE &&
-						MasterBoneIndex < MasterPoseComponent->GetNumSpaceBases())
+						MasterBoneIndex < ClothSimulationContext.InMasterPoseComponent->GetNumSpaceBases())
 					{
 						BoneIndex = MasterBoneIndex;
 					}
@@ -3586,13 +3575,13 @@ void USkeletalMeshComponent::ParallelUpdateClothState(float DeltaTime, FClothSim
 			}
 
 			if (BoneIndex != INDEX_NONE)
-			{
-				BoneMatrices[Index] = U2PMatrix(BoneTransforms[BoneIndex].ToMatrixWithScale());
-			}
-			else
-			{
-				BoneMatrices[Index] = PxMat44::createIdentity();
-			}
+		   {
+			   BoneMatrices[Index] = U2PMatrix(BoneTransforms[BoneIndex].ToMatrixWithScale());
+		   }
+		   else
+		   {
+			   BoneMatrices[Index] = PxMat44::createIdentity();
+		   }
 		}
 
 		NxClothingActor* ApexClothingActor = ClothingActor.ApexClothingActor;
@@ -3617,9 +3606,13 @@ void USkeletalMeshComponent::ParallelUpdateClothState(float DeltaTime, FClothSim
 
 void USkeletalMeshComponent::GetClothRootBoneMatrix(int32 AssetIndex, FMatrix& OutRootBoneMatrix) const
 {
-	if (IsValidClothingActor(AssetIndex))
+	for (const FClothingActor ClothingActor : ClothingActors)
+	{
+		if (ClothingActor.ParentClothingAssetIndex == AssetIndex)
 	{
-		NxClothingAsset* Asset = ClothingActors[AssetIndex].ParentClothingAsset;
+			if (NxClothingActor* ApexClothingActor = ClothingActor.ApexClothingActor)
+			{
+				NxClothingAsset* Asset = ClothingActor.ParentClothingAsset;
 
 		check(Asset);
 
@@ -3631,29 +3624,30 @@ void USkeletalMeshComponent::GetClothRootBoneMatrix(int32 AssetIndex, FMatrix& O
 		int32 BoneIndex = GetBoneIndex(BoneName);
 		check(BoneIndex >= 0);
 		OutRootBoneMatrix = GetBoneMatrix(BoneIndex);
+				return;	//TODO: This assumes only one actor per asset, do we want to change this?
+			}
 	}
-	else
-	{
-		OutRootBoneMatrix = GetBoneMatrix(0);
 	}
+	
+		OutRootBoneMatrix = GetBoneMatrix(0);
 }
 
 bool USkeletalMeshComponent::GetClothSimulatedPosition(int32 AssetIndex, int32 VertexIndex, FVector& OutSimulPos) const
 {
 	bool bSucceed = false;
 
-	if (IsValidClothingActor(AssetIndex))
+	for(const FClothingActor ClothingActor : ClothingActors)
 	{
-		NxClothingActor* ApexClothingActor = ClothingActors[AssetIndex].ApexClothingActor;
-
-		if (ApexClothingActor)
+		if(ClothingActor.ParentClothingAssetIndex == AssetIndex)
+		{
+			if (NxClothingActor* ApexClothingActor = ClothingActor.ApexClothingActor)
 		{
 			uint32 NumSimulVertices = ApexClothingActor->getNumSimulationVertices();
 
 			// handles only simulated vertices, indices of fixed vertices are bigger than # of simulated vertices
 			if ((uint32)VertexIndex < NumSimulVertices)
 			{
-				//TODO: this is not thread safe! apex requires a lock. Not going to add it here because it's called for many vertices, but this should be addressed!
+					//TODO: this is not thread safe! apex requires a lock. Not going to add it here because it's called for many vertices, but this should be addressed!
 				bSucceed = true;
 				const physx::PxVec3* Vertices = ApexClothingActor->getSimulationPositions();
 
@@ -3669,6 +3663,9 @@ bool USkeletalMeshComponent::GetClothSimulatedPosition(int32 AssetIndex, int32 V
 				}
 			}
 		}
+
+			break; //TODO: This finds the first actor. Is there a reason to not support many instances?
+		}
 	}
 
 	return bSucceed;
@@ -3791,19 +3788,18 @@ void USkeletalMeshComponent::GetUpdateClothSimulationData(TArray<FClothSimulData
 
 	for(int32 ActorIndex=0; ActorIndex<NumClothingActors; ActorIndex++)
 	{
-		if(!IsValidClothingActor(ActorIndex))
+		const FClothingActor& ClothingActor = ClothingActors[ActorIndex];
+		if(!IsValidClothingActor(ClothingActor))
 		{
 			OutClothSimData[ActorIndex].ClothSimulPositions.Reset();
 			OutClothSimData[ActorIndex].ClothSimulNormals.Reset();
 			continue;
 		}
 
-		NxClothingActor* ClothingActor = ClothingActors[ActorIndex].ApexClothingActor;
-
 		// update simulation positions & normals
-		if(ClothingActor)
+		if(NxClothingActor* ApexClothingActor = ClothingActor.ApexClothingActor)
 		{
-			uint32 NumSimulVertices = ClothingActor->getNumSimulationVertices();
+			uint32 NumSimulVertices = ApexClothingActor->getNumSimulationVertices();
 
 			if(NumSimulVertices > 0)
 			{
@@ -3819,14 +3815,14 @@ void USkeletalMeshComponent::GetUpdateClothSimulationData(TArray<FClothSimulData
 					ClothData.ClothSimulNormals.AddUninitialized(NumSimulVertices);
 				}
 
-				const physx::PxVec3* Vertices = ClothingActor->getSimulationPositions();
-				const physx::PxVec3* Normals = ClothingActor->getSimulationNormals();
+				const physx::PxVec3* Vertices = ApexClothingActor->getSimulationPositions();
+				const physx::PxVec3* Normals = ApexClothingActor->getSimulationNormals();
 
 				// In case of Local space simulation, need to transform simulated positions with the internal bone matrix
 				if (bLocalSpaceSimulation)
 				{
 					FMatrix ClothRootBoneMatrix;
-					(OverrideLocalRootComponent) ? OverrideLocalRootComponent->GetClothRootBoneMatrix(ActorIndex, ClothRootBoneMatrix) : GetClothRootBoneMatrix(ActorIndex, ClothRootBoneMatrix);
+					(OverrideLocalRootComponent) ? OverrideLocalRootComponent->GetClothRootBoneMatrix(ClothingActor.ParentClothingAssetIndex, ClothRootBoneMatrix) : GetClothRootBoneMatrix(ClothingActor.ParentClothingAssetIndex, ClothRootBoneMatrix);
 					for (uint32 VertexIndex = 0; VertexIndex < NumSimulVertices; VertexIndex++)
 					{						
 						ClothData.ClothSimulPositions[VertexIndex] = ClothRootBoneMatrix.TransformPosition(P2UVector(Vertices[VertexIndex]));
@@ -3859,42 +3855,20 @@ void USkeletalMeshComponent::FreezeClothSection(bool bFreeze)
 	check(!IsExecutingParallelCloth());
 	int32 NumActors = ClothingActors.Num();
 
-	for(int32 ActorIdx=0; ActorIdx<NumActors; ActorIdx++)
+	for(FClothingActor& ClothingActor : ClothingActors)
 	{
-		NxClothingActor* ClothingActor = ClothingActors[ActorIdx].ApexClothingActor;
-
-		if (ClothingActor)
+		if (NxClothingActor* ApexClothingActor = ClothingActor.ApexClothingActor)
 		{
-			ClothingActor->setFrozen(bFreeze);
+			ApexClothingActor->setFrozen(bFreeze);
 		}
 	}
 #endif// #if WITH_APEX_CLOTHING
 }
 
-bool USkeletalMeshComponent::IsValidClothingActor(int32 ActorIndex) const
+bool USkeletalMeshComponent::IsValidClothingActor(const FClothingActor& ClothingActor) const
 {
 #if WITH_APEX_CLOTHING
-
-	if (!SkeletalMesh)
-	{
-		return false;
-	}
-
-	//false if ActorIndex is out-range
-	if(ActorIndex >= SkeletalMesh->ClothingAssets.Num()
-	|| ActorIndex >= ClothingActors.Num())
-	{
-		return false;
-	}
-
-	if(ClothingActors[ActorIndex].ApexClothingActor
-	&& ClothingActors[ActorIndex].ParentClothingAsset == SkeletalMesh->ClothingAssets[ActorIndex].ApexClothingAsset)
-	{
-		return true;
-	}
-
-	return false;
-	
+	return SkeletalMesh && ClothingActor.ApexClothingActor;	
 #else
 	return false;
 #endif// #if WITH_APEX_CLOTHING
@@ -3911,24 +3885,17 @@ void USkeletalMeshComponent::DrawClothingNormals(FPrimitiveDrawInterface* PDI)
 
 	int32 NumActors = ClothingActors.Num();
 
-	for(int32 ActorIdx=0; ActorIdx<NumActors; ActorIdx++)
+	for(FClothingActor& ClothingActor : ClothingActors)
 	{
-		NxClothingActor* ClothingActor = ClothingActors[ActorIdx].ApexClothingActor;
-
-		if(!IsValidClothingActor(ActorIdx))
+		if (NxClothingActor* ApexClothingActor = ClothingActor.ApexClothingActor)
 		{
-			continue;
-		}
-
-		if (ClothingActor)
-		{
-			uint32 NumSimulVertices = ClothingActor->getNumSimulationVertices();
+			uint32 NumSimulVertices = ApexClothingActor->getNumSimulationVertices();
 
 			// simulation is enabled and there exist simulated vertices
 			if(!bDisableClothSimulation && NumSimulVertices > 0)
 			{
-				const physx::PxVec3* Vertices = ClothingActor->getSimulationPositions();
-				const physx::PxVec3* Normals = ClothingActor->getSimulationNormals();
+				const physx::PxVec3* Vertices = ApexClothingActor->getSimulationPositions();
+				const physx::PxVec3* Normals = ApexClothingActor->getSimulationNormals();
 
 				FLinearColor LineColor = FColor::Red;
 				FVector Start, End;
@@ -3946,20 +3913,22 @@ void USkeletalMeshComponent::DrawClothingNormals(FPrimitiveDrawInterface* PDI)
 			}
 			else // otherwise, draws initial values loaded from the asset file
 			{
-				if(SkeletalMesh->ClothingAssets[ActorIdx].ClothVisualizationInfos.Num() == 0)
+				const int32 ClothingAssetIndex = ClothingActor.ParentClothingAssetIndex;
+				FClothingAssetData& ClothingAsset = SkeletalMesh->ClothingAssets[ClothingAssetIndex];
+				if(ClothingAsset.ClothVisualizationInfos.Num() == 0)
 				{
-					LoadClothingVisualizationInfo(ActorIdx);
+					LoadClothingVisualizationInfo(ClothingAsset);
 				}
 
 				int32 PhysicalMeshLOD = PredictedLODLevel;
 
 				// if this asset doesn't have current physical mesh LOD, then skip to visualize
-				if(!SkeletalMesh->ClothingAssets[ActorIdx].ClothVisualizationInfos.IsValidIndex(PhysicalMeshLOD))
+				if(!ClothingAsset.ClothVisualizationInfos.IsValidIndex(PhysicalMeshLOD))
 				{
 					continue;
 				}
 
-				FClothVisualizationInfo& VisualInfo = SkeletalMesh->ClothingAssets[ActorIdx].ClothVisualizationInfos[PhysicalMeshLOD];				
+				const FClothVisualizationInfo& VisualInfo = ClothingAsset.ClothVisualizationInfos[PhysicalMeshLOD];				
 
 				uint32 NumVertices = VisualInfo.ClothPhysicalMeshVertices.Num();
 
@@ -3994,18 +3963,16 @@ void USkeletalMeshComponent::DrawClothingTangents(FPrimitiveDrawInterface* PDI)
 
 	int32 NumActors = ClothingActors.Num();
 
-	for(int32 ActorIdx=0; ActorIdx<NumActors; ActorIdx++)
+	for(const FClothingActor& ClothingActor : ClothingActors)
 	{
-		if(!IsValidClothingActor(ActorIdx))
+		if(!IsValidClothingActor(ClothingActor))
 		{
 			continue;
 		}
 
-		NxClothingActor* ClothingActor = ClothingActors[ActorIdx].ApexClothingActor;
-
-		if (ClothingActor)
+		if (NxClothingActor* ApexClothingActor = ClothingActor.ApexClothingActor)
 		{
-			uint32 NumSimulVertices = ClothingActor->getNumSimulationVertices();
+			uint32 NumSimulVertices = ApexClothingActor->getNumSimulationVertices();
 
 			// simulation is enabled and there exist simulated vertices
 			if(!bDisableClothSimulation && NumSimulVertices > 0)
@@ -4015,7 +3982,7 @@ void USkeletalMeshComponent::DrawClothingTangents(FPrimitiveDrawInterface* PDI)
 				FStaticLODModel& LODModel = MeshObject->GetSkeletalMeshResource().LODModels[PredictedLODLevel];
 
 				TArray<uint32> SectionIndices;
-				SkeletalMesh->GetClothSectionIndices(PredictedLODLevel, ActorIdx, SectionIndices);
+				SkeletalMesh->GetClothSectionIndices(PredictedLODLevel, ClothingActor.ParentClothingAssetIndex, SectionIndices);
 
 				int32 NumSections = SectionIndices.Num();
 
@@ -4027,8 +3994,8 @@ void USkeletalMeshComponent::DrawClothingTangents(FPrimitiveDrawInterface* PDI)
 
 					FSkelMeshChunk& Chunk = LODModel.Chunks[ChunkIndex];
 
-					const physx::PxVec3* SimulVertices = ClothingActor->getSimulationPositions();
-					const physx::PxVec3* SimulNormals = ClothingActor->getSimulationNormals();
+					const physx::PxVec3* SimulVertices = ApexClothingActor->getSimulationPositions();
+					const physx::PxVec3* SimulNormals = ApexClothingActor->getSimulationNormals();
 
 					uint32 NumMppingData = Chunk.ApexClothMappingData.Num();
 
@@ -4106,20 +4073,21 @@ void USkeletalMeshComponent::DrawClothingTangents(FPrimitiveDrawInterface* PDI)
 			}
 			else
 			{
-				if(SkeletalMesh->ClothingAssets[ActorIdx].ClothVisualizationInfos.Num() == 0)
+				FClothingAssetData& ClothingAsset = SkeletalMesh->ClothingAssets[ClothingActor.ParentClothingAssetIndex];
+				if(ClothingAsset.ClothVisualizationInfos.Num() == 0)
 				{
-					LoadClothingVisualizationInfo(ActorIdx);
+					LoadClothingVisualizationInfo(ClothingAsset);
 				}
 
 				int32 PhysicalMeshLOD = PredictedLODLevel;
 
 				// if this asset doesn't have current physical mesh LOD, then skip to visualize
-				if(!SkeletalMesh->ClothingAssets[ActorIdx].ClothVisualizationInfos.IsValidIndex(PhysicalMeshLOD))
+				if(!ClothingAsset.ClothVisualizationInfos.IsValidIndex(PhysicalMeshLOD))
 				{
 					continue;
 				}
 
-				FClothVisualizationInfo& VisualInfo = SkeletalMesh->ClothingAssets[ActorIdx].ClothVisualizationInfos[PhysicalMeshLOD];				
+				FClothVisualizationInfo& VisualInfo = ClothingAsset.ClothVisualizationInfos[PhysicalMeshLOD];				
 
 				NumSimulVertices = VisualInfo.ClothPhysicalMeshVertices.Num();
 
@@ -4128,7 +4096,7 @@ void USkeletalMeshComponent::DrawClothingTangents(FPrimitiveDrawInterface* PDI)
 				FStaticLODModel& LODModel = MeshObject->GetSkeletalMeshResource().LODModels[PredictedLODLevel];
 
 				TArray<uint32> SectionIndices;
-				SkeletalMesh->GetClothSectionIndices(PredictedLODLevel, ActorIdx, SectionIndices);
+				SkeletalMesh->GetClothSectionIndices(PredictedLODLevel, ClothingActor.ParentClothingAssetIndex, SectionIndices);
 
 				int32 NumSections = SectionIndices.Num();
 
@@ -4400,33 +4368,26 @@ void USkeletalMeshComponent::DrawClothingFixedVertices(FPrimitiveDrawInterface*
 	UpdateRefToLocalMatrices(RefToLocals, this, GetSkeletalMeshResource(), 0, NULL);
 	const float Inv_255 = 1.f/255.f;
 
-	for(int32 ActorIdx=0; ActorIdx<NumActors; ActorIdx++)
+	for(const FClothingActor& ClothingActor : ClothingActors)
 	{
-		if(!IsValidClothingActor(ActorIdx))
+		if (NxClothingActor* ApexClothingActor = ClothingActor.ApexClothingActor)
 		{
-			continue;
-		}
-
-		NxClothingActor* ClothingActor = ClothingActors[ActorIdx].ApexClothingActor;
-
-		if (ClothingActor)
+			uint32 NumSimulVertices = ApexClothingActor->getNumSimulationVertices();
+			FClothingAssetData& ClothingAssetData = SkeletalMesh->ClothingAssets[ClothingActor.ParentClothingAssetIndex];
+			if(ClothingAssetData.ClothVisualizationInfos.Num() == 0)
 		{
-			uint32 NumSimulVertices = ClothingActor->getNumSimulationVertices();
-
-			if(SkeletalMesh->ClothingAssets[ActorIdx].ClothVisualizationInfos.Num() == 0)
-			{
-				LoadClothingVisualizationInfo(ActorIdx);
+				LoadClothingVisualizationInfo(ClothingAssetData);
 			}
 
 			int32 PhysicalMeshLOD = PredictedLODLevel;
 
 			// if this asset doesn't have current physical mesh LOD, then skip to visualize
-			if(!SkeletalMesh->ClothingAssets[ActorIdx].ClothVisualizationInfos.IsValidIndex(PhysicalMeshLOD))
+			if(!ClothingAssetData.ClothVisualizationInfos.IsValidIndex(PhysicalMeshLOD))
 			{
 				continue;
 			}
 
-			FClothVisualizationInfo& VisualInfo = SkeletalMesh->ClothingAssets[ActorIdx].ClothVisualizationInfos[PhysicalMeshLOD];				
+			FClothVisualizationInfo& VisualInfo = ClothingAssetData.ClothVisualizationInfos[PhysicalMeshLOD];				
 
 			NumSimulVertices = VisualInfo.ClothPhysicalMeshVertices.Num();
 
@@ -4435,7 +4396,7 @@ void USkeletalMeshComponent::DrawClothingFixedVertices(FPrimitiveDrawInterface*
 			FStaticLODModel& LODModel = MeshObject->GetSkeletalMeshResource().LODModels[PredictedLODLevel];
 
 			TArray<uint32> SectionIndices;
-			SkeletalMesh->GetClothSectionIndices(PredictedLODLevel, ActorIdx, SectionIndices);
+			SkeletalMesh->GetClothSectionIndices(PredictedLODLevel, ClothingActor.ParentClothingAssetIndex, SectionIndices);
 
 			int32 NumSections = SectionIndices.Num();
 
@@ -4508,17 +4469,16 @@ void USkeletalMeshComponent::DrawClothingFixedVertices(FPrimitiveDrawInterface*
 #endif // #if WITH_APEX_CLOTHING
 }
 
-void USkeletalMeshComponent::LoadClothingVisualizationInfo(int32 AssetIndex)
+void USkeletalMeshComponent::LoadClothingVisualizationInfo(FClothingAssetData& AssetData)
 {
 #if WITH_APEX_CLOTHING
 
-	if (!SkeletalMesh || !SkeletalMesh->ClothingAssets.IsValidIndex(AssetIndex)) 
+	if (!SkeletalMesh) 
 	{
 		return;
 	}
 
 
-	FClothingAssetData& AssetData = SkeletalMesh->ClothingAssets[AssetIndex];
 	NxClothingAsset* ApexClothingAsset = AssetData.ApexClothingAsset;
 	const NxParameterized::Interface* AssetParams = ApexClothingAsset->getAssetNxParameterized();
 
@@ -4701,11 +4661,9 @@ void USkeletalMeshComponent::LoadAllClothingVisualizationInfos()
 		return;
 	}
 
-	int32 NumAssets = SkeletalMesh->ClothingAssets.Num();
-
-	for(int32 AssetIdx=0; AssetIdx<NumAssets; AssetIdx++)
+	for(FClothingAssetData& ClothAssetData : SkeletalMesh->ClothingAssets)
 	{
-		LoadClothingVisualizationInfo(AssetIdx);
+		LoadClothingVisualizationInfo(ClothAssetData);
 	}
 #endif // #if WITH_APEX_CLOTHING
 }
@@ -4721,22 +4679,20 @@ void USkeletalMeshComponent::DrawClothingMaxDistances(FPrimitiveDrawInterface* P
 
 	int32 PhysicalMeshLOD = PredictedLODLevel;
 
-	int32 NumAssets = SkeletalMesh->ClothingAssets.Num();
-
-	for(int32 AssetIdx=0; AssetIdx < NumAssets; AssetIdx++)
+	for(FClothingAssetData& ClothingAssetData : SkeletalMesh->ClothingAssets)
 	{
-		if(SkeletalMesh->ClothingAssets[AssetIdx].ClothVisualizationInfos.Num() == 0)
+		if(ClothingAssetData.ClothVisualizationInfos.Num() == 0)
 		{
-			LoadClothingVisualizationInfo(AssetIdx);
+			LoadClothingVisualizationInfo(ClothingAssetData);
 		}
 
 		// if this asset doesn't have current physical mesh LOD, then skip to visualize
-		if(!SkeletalMesh->ClothingAssets[AssetIdx].ClothVisualizationInfos.IsValidIndex(PhysicalMeshLOD))
+		if(!ClothingAssetData.ClothVisualizationInfos.IsValidIndex(PhysicalMeshLOD))
 		{
 			continue;
 		}
 
-		FClothVisualizationInfo& VisualInfo = SkeletalMesh->ClothingAssets[AssetIdx].ClothVisualizationInfos[PhysicalMeshLOD];
+		FClothVisualizationInfo& VisualInfo = ClothingAssetData.ClothVisualizationInfos[PhysicalMeshLOD];
 
 		int32 NumMaxDistances = VisualInfo.ClothConstrainCoeffs.Num();
 		for(int32 MaxDistIdx=0; MaxDistIdx < NumMaxDistances; MaxDistIdx++)
@@ -4775,20 +4731,20 @@ void USkeletalMeshComponent::DrawClothingBackstops(FPrimitiveDrawInterface* PDI)
 
 	int32 NumAssets = SkeletalMesh->ClothingAssets.Num();
 
-	for(int32 AssetIdx=0; AssetIdx < NumAssets; AssetIdx++)
+	for(FClothingAssetData& ClothingAssetData : SkeletalMesh->ClothingAssets)
 	{
-		if(SkeletalMesh->ClothingAssets[AssetIdx].ClothVisualizationInfos.Num() == 0)
+		if(ClothingAssetData.ClothVisualizationInfos.Num() == 0)
 		{
-			LoadClothingVisualizationInfo(AssetIdx);
+			LoadClothingVisualizationInfo(ClothingAssetData);
 		}
 
 		// if this asset doesn't have current physical mesh LOD, then skip to visualize
-		if(!SkeletalMesh->ClothingAssets[AssetIdx].ClothVisualizationInfos.IsValidIndex(PhysicalMeshLOD))
+		if(!ClothingAssetData.ClothVisualizationInfos.IsValidIndex(PhysicalMeshLOD))
 		{
 			continue;
 		}
 
-		FClothVisualizationInfo& VisualInfo = SkeletalMesh->ClothingAssets[AssetIdx].ClothVisualizationInfos[PhysicalMeshLOD];
+		FClothVisualizationInfo& VisualInfo = ClothingAssetData.ClothVisualizationInfos[PhysicalMeshLOD];
 
 		int32 NumBackstopSpheres = VisualInfo.ClothConstrainCoeffs.Num();
 		for(int32 BackstopIdx=0; BackstopIdx < NumBackstopSpheres; BackstopIdx++)
@@ -4838,26 +4794,25 @@ void USkeletalMeshComponent::DrawClothingPhysicalMeshWire(FPrimitiveDrawInterfac
 
 	int32 PhysicalMeshLOD = PredictedLODLevel;
 
-	int32 NumAssets = SkeletalMesh->ClothingAssets.Num();
-
 	TArray<FMatrix> RefToLocals;
 	// get local matrices for skinning
 	UpdateRefToLocalMatrices(RefToLocals, this, GetSkeletalMeshResource(), 0, NULL);
 
-	for(int32 AssetIdx=0; AssetIdx < NumAssets; AssetIdx++)
+	for(const FClothingActor& ClothingActor : ClothingActors)
 	{
-		if(SkeletalMesh->ClothingAssets[AssetIdx].ClothVisualizationInfos.Num() == 0)
+		FClothingAssetData& ClothingAsset = SkeletalMesh->ClothingAssets[ClothingActor.ParentClothingAssetIndex];
+		if(ClothingAsset.ClothVisualizationInfos.Num() == 0)
 		{
-			LoadClothingVisualizationInfo(AssetIdx);
+			LoadClothingVisualizationInfo(ClothingAsset);
 		}
 
 		// if this asset doesn't have current physical mesh LOD, then skip to visualize
-		if(!SkeletalMesh->ClothingAssets[AssetIdx].ClothVisualizationInfos.IsValidIndex(PhysicalMeshLOD))
+		if(!ClothingAsset.ClothVisualizationInfos.IsValidIndex(PhysicalMeshLOD))
 		{
 			continue;
 		}
 
-		FClothVisualizationInfo& VisualInfo = SkeletalMesh->ClothingAssets[AssetIdx].ClothVisualizationInfos[PhysicalMeshLOD];
+		FClothVisualizationInfo& VisualInfo = ClothingAsset.ClothVisualizationInfos[PhysicalMeshLOD];
 
 		uint32 NumPhysicalMeshVerts = VisualInfo.ClothPhysicalMeshVertices.Num();
 
@@ -4866,13 +4821,11 @@ void USkeletalMeshComponent::DrawClothingPhysicalMeshWire(FPrimitiveDrawInterfac
 		TArray<FVector> SimulatedPhysicalMeshVertices;
 		// if cloth simulation is enabled and there exist a clothing actor and simulation vertices, 
 		// then draws wire frame using simulated vertices
-		if(!bDisableClothSimulation && ClothingActors.IsValidIndex(AssetIdx))
+		if(!bDisableClothSimulation)
 		{
-			NxClothingActor* ClothingActor = ClothingActors[AssetIdx].ApexClothingActor;
-
-			if(ClothingActor)
+			if(NxClothingActor* ApexClothingActor = ClothingActor.ApexClothingActor)
 			{
-				uint32 NumSimulVertices = ClothingActor->getNumSimulationVertices();
+				uint32 NumSimulVertices = ApexClothingActor->getNumSimulationVertices();
 
 				if(NumSimulVertices > 0)
 				{
@@ -4886,7 +4839,7 @@ void USkeletalMeshComponent::DrawClothingPhysicalMeshWire(FPrimitiveDrawInterfac
 					bUseSimulatedResult = true;
 					SimulatedPhysicalMeshVertices.AddUninitialized(NumPhysicalMeshVerts);
 
-					const physx::PxVec3* SimulVertices = ClothingActor->getSimulationPositions();
+					const physx::PxVec3* SimulVertices = ApexClothingActor->getSimulationPositions();
 					for(uint32 SimulVertIdx=0; SimulVertIdx < NumSimulVertices; SimulVertIdx++)
 					{
 						SimulatedPhysicalMeshVertices[SimulVertIdx] = P2UVector(SimulVertices[SimulVertIdx]);
@@ -4901,7 +4854,7 @@ void USkeletalMeshComponent::DrawClothingPhysicalMeshWire(FPrimitiveDrawInterfac
 						{
 							uint16 ApexBoneIndex = VisualInfo.ClothPhysicalMeshBoneWeightsInfo[FixedVertIdx].Indices[BoneWeightIdx];
 
-							FName BoneName = GetConvertedBoneName(SkeletalMesh->ClothingAssets[AssetIdx].ApexClothingAsset, ApexBoneIndex);
+							FName BoneName = GetConvertedBoneName(SkeletalMesh->ClothingAssets[ClothingActor.ParentClothingAssetIndex].ApexClothingAsset, ApexBoneIndex);
 
 							int32 BoneIndex = GetBoneIndex(BoneName);
 
@@ -5034,23 +4987,13 @@ void USkeletalMeshComponent::SetClothMaxDistanceScale(float Scale)
 	//this scale parameter is also used when new clothing actor is created
 	ClothMaxDistanceScale = Scale;
 
-	int32 NumActors = ClothingActors.Num();
-
-	for(int32 ActorIdx=0; ActorIdx<NumActors; ActorIdx++)
+	for(const FClothingActor& ClothingActor : ClothingActors)
 	{
-		// skip if ClothingActor is NULL or invalid
-		if(!IsValidClothingActor(ActorIdx))
+		if(NxClothingActor* ApexClothingActor = ClothingActor.ApexClothingActor)
 		{
-			continue;
+			NxParameterized::Interface* ActorDesc = ApexClothingActor->getActorDesc();
+			verify(NxParameterized::setParamF32(*ActorDesc, "maxDistanceScale.Scale", Scale));
 		}
-
-		NxClothingActor* ClothingActor = ClothingActors[ActorIdx].ApexClothingActor;
-
-		check(ClothingActor);
-
-		NxParameterized::Interface* ActorDesc = ClothingActor->getActorDesc();
-
-		verify(NxParameterized::setParamF32(*ActorDesc, "maxDistanceScale.Scale", Scale));
 	}
 #endif// #if WITH_APEX_CLOTHING
 }
diff --git a/Engine/Source/Runtime/Engine/Private/SkeletalRender.cpp b/Engine/Source/Runtime/Engine/Private/SkeletalRender.cpp
--- a/Engine/Source/Runtime/Engine/Private/SkeletalRender.cpp
+++ b/Engine/Source/Runtime/Engine/Private/SkeletalRender.cpp
@@ -199,6 +199,8 @@ void UpdateRefToLocalMatrices( TArray<FMatrix>& ReferenceToLocal, const USkinned
 	const USkeletalMesh* const MasterCompMesh = MasterComp? MasterComp->SkeletalMesh : nullptr;
 	const FStaticLODModel& LOD = InSkeletalMeshResource->LODModels[LODIndex];
 
+	const TArray<int32>& MasterBoneMap = InMeshComponent->GetMasterBoneMap();
+
 	check( ThisMesh->RefBasesInvMatrix.Num() != 0 );
 	if(ReferenceToLocal.Num() != ThisMesh->RefBasesInvMatrix.Num())
 	{
@@ -206,7 +208,7 @@ void UpdateRefToLocalMatrices( TArray<FMatrix>& ReferenceToLocal, const USkinned
 		ReferenceToLocal.AddUninitialized(ThisMesh->RefBasesInvMatrix.Num());
 	}
 
-	const bool bIsMasterCompValid = MasterComp && InMeshComponent->MasterBoneMap.Num() == ThisMesh->RefSkeleton.GetNum();
+	const bool bIsMasterCompValid = MasterComp && MasterBoneMap.Num() == ThisMesh->RefSkeleton.GetNum();
 
 	const TArray<FBoneIndexType>* RequiredBoneSets[3] = { &LOD.ActiveBoneIndices, ExtraRequiredBoneIndices, NULL };
 
@@ -230,7 +232,7 @@ void UpdateRefToLocalMatrices( TArray<FMatrix>& ReferenceToLocal, const USkinned
 				if( bIsMasterCompValid )
 				{
 					// If valid, use matrix from parent component.
-					const int32 MasterBoneIndex = InMeshComponent->MasterBoneMap[ThisBoneIndex];
+					const int32 MasterBoneIndex = MasterBoneMap[ThisBoneIndex];
 					if ( MasterComp->GetSpaceBases().IsValidIndex(MasterBoneIndex) )
 					{
 						const int32 ParentIndex = MasterCompMesh->RefSkeleton.GetParentIndex(MasterBoneIndex);
@@ -299,6 +301,7 @@ void UpdateCustomLeftRightVectors( TArray<FTwoVectors>& OutVectors, const USkinn
 	const USkinnedMeshComponent* const MasterComp = InMeshComponent->MasterPoseComponent.Get();
 	const FStaticLODModel& LOD = InSkeletalMeshResource->LODModels[LODIndex];
 	const FSkeletalMeshLODInfo& LODInfo = ThisMesh->LODInfo[LODIndex];
+	const TArray<int32>& MasterBoneMap = InMeshComponent->GetMasterBoneMap();
 
 	if(OutVectors.Num() != LODInfo.TriangleSortSettings.Num())
 	{
@@ -324,10 +327,10 @@ void UpdateCustomLeftRightVectors( TArray<FTwoVectors>& OutVectors, const USkinn
 				const TArray<FTransform>* SpaceBases = &InMeshComponent->GetSpaceBases();
 				
 				// Handle case of using MasterPoseComponent for SpaceBases.
-				if( MasterComp && InMeshComponent->MasterBoneMap.Num() == ThisMesh->RefSkeleton.GetNum() && SpaceBasesBoneIndex != INDEX_NONE )
+				if( MasterComp && MasterBoneMap.Num() == ThisMesh->RefSkeleton.GetNum() && SpaceBasesBoneIndex != INDEX_NONE )
 				{
 					// If valid, use matrix from parent component.
-					SpaceBasesBoneIndex = InMeshComponent->MasterBoneMap[SpaceBasesBoneIndex];
+					SpaceBasesBoneIndex = MasterBoneMap[SpaceBasesBoneIndex];
 					SpaceBases = &MasterComp->GetSpaceBases();
 				}
 
