UE4 - Fix missing hud materials

diff --git a/Engine/Source/Runtime/Engine/Private/Materials/MaterialShared.cpp b/Engine/Source/Runtime/Engine/Private/Materials/MaterialShared.cpp
--- a/Engine/Source/Runtime/Engine/Private/Materials/MaterialShared.cpp
+++ b/Engine/Source/Runtime/Engine/Private/Materials/MaterialShared.cpp
@@ -789,7 +789,7 @@ bool FMaterialResource::IsCrackFreeDisplacementEnabled() const
 
 bool FMaterialResource::IsSeparateTranslucencyEnabled() const 
 { 
-	return Material->bEnableSeparateTranslucency;
+	return Material->bEnableSeparateTranslucency && !IsUIMaterial();
 }
 
 bool FMaterialResource::IsAdaptiveTessellationEnabled() const
@@ -804,7 +804,7 @@ bool FMaterialResource::IsFullyRough() const
 
 bool FMaterialResource::OutputsVelocityOnBasePass() const
 {
-	return Material->bOutputVelocityOnBasePass;
+	return Material->bOutputVelocityOnBasePass && !IsUIMaterial();
 }
 
 bool FMaterialResource::IsNonmetal() const
