Block cheat cvars from being set by system setting ini

Merging 2703124

//depot/UE4-UT/...

to //depot/UE4/...

diff --git a/Engine/Source/Runtime/Engine/Private/SystemSettings.cpp b/Engine/Source/Runtime/Engine/Private/SystemSettings.cpp
--- a/Engine/Source/Runtime/Engine/Private/SystemSettings.cpp
+++ b/Engine/Source/Runtime/Engine/Private/SystemSettings.cpp
@@ -71,26 +71,29 @@ public:
 
 		if(CVar)
 		{
-			// Set if the variable exists.
-
-			// to be more compatible with non console variable boolean system settings we allow verbal values
-
-			if(FCString::Stricmp(Value, TEXT("True")) == 0
-			|| FCString::Stricmp(Value, TEXT("Yes")) == 0
-			|| FCString::Stricmp(Value, TEXT("On")) == 0)
-			{
-				CVar->Set(1, ECVF_SetBySystemSettingsIni);
-			}
-			else
-			if(FCString::Stricmp(Value, TEXT("False")) == 0
-			|| FCString::Stricmp(Value, TEXT("No")) == 0
-			|| FCString::Stricmp(Value, TEXT("Off")) == 0)
-			{
-				CVar->Set(0, ECVF_SetBySystemSettingsIni);
-			}
-			else
+			if (!CVar->TestFlags(EConsoleVariableFlags::ECVF_Cheat))
 			{
-				CVar->Set(Value, ECVF_SetBySystemSettingsIni);
+				// Set if the variable exists.
+
+				// to be more compatible with non console variable boolean system settings we allow verbal values
+
+				if (FCString::Stricmp(Value, TEXT("True")) == 0
+					|| FCString::Stricmp(Value, TEXT("Yes")) == 0
+					|| FCString::Stricmp(Value, TEXT("On")) == 0)
+				{
+					CVar->Set(1, ECVF_SetBySystemSettingsIni);
+				}
+				else
+				if (FCString::Stricmp(Value, TEXT("False")) == 0
+					|| FCString::Stricmp(Value, TEXT("No")) == 0
+					|| FCString::Stricmp(Value, TEXT("Off")) == 0)
+				{
+					CVar->Set(0, ECVF_SetBySystemSettingsIni);
+				}
+				else
+				{
+					CVar->Set(Value, ECVF_SetBySystemSettingsIni);
+				}
 			}
 		}
 	}
