Removal of AI-navigation headers from Engine.h merged over from Orion branch #UE4

diff --git a/Engine/Source/Runtime/Engine/Public/Engine.h b/Engine/Source/Runtime/Engine/Public/Engine.h
--- a/Engine/Source/Runtime/Engine/Public/Engine.h
+++ b/Engine/Source/Runtime/Engine/Public/Engine.h
@@ -32,9 +32,6 @@
 #include "GameFramework/CameraBlockingVolume.h"
 #include "Engine/CullDistanceVolume.h"
 #include "Engine/LevelStreamingVolume.h"
-#include "AI/Navigation/NavMeshBoundsVolume.h"
-#include "AI/Navigation/NavRelevantInterface.h"
-#include "AI/Navigation/NavModifierVolume.h"
 #include "GameFramework/PhysicsVolume.h"
 #include "GameFramework/DefaultPhysicsVolume.h"
 #include "GameFramework/KillZVolume.h"
@@ -44,12 +41,7 @@
 #include "Sound/AudioVolume.h"
 #include "Engine/TriggerVolume.h"
 #include "Camera/CameraActor.h"
-#include "AI/Navigation/NavAgentInterface.h"
-#include "AI/Navigation/NavLinkDefinition.h"
-#include "AI/Navigation/NavigationTypes.h"
 #include "Kismet/BlueprintFunctionLibrary.h"
-#include "AI/Navigation/NavigationSystem.h"
-#include "AI/Navigation/NavCollision.h"
 #include "GameFramework/PlayerMuteList.h"
 #include "Camera/PlayerCameraManager.h"
 #include "GameFramework/Info.h"
@@ -77,18 +69,9 @@
 #include "Engine/PointLight.h"
 #include "Engine/SpotLight.h"
 #include "Engine/GeneratedMeshAreaLight.h"
-#include "AI/Navigation/NavigationData.h"
-#include "AI/Navigation/NavigationGraph.h"
-#include "AI/Navigation/NavAreas/NavArea.h"
-#include "AI/Navigation/RecastNavMesh.h"
-#include "AI/Navigation/NavigationGraphNode.h"
 #include "Engine/NavigationObjectBase.h"
 #include "GameFramework/PlayerStart.h"
 #include "Engine/PlayerStartPIE.h"
-#include "AI/Navigation/NavPathObserverInterface.h"
-#include "AI/Navigation/NavigationTestingActor.h"
-#include "AI/Navigation/NavLinkHostInterface.h"
-#include "AI/Navigation/NavLinkProxy.h"
 #include "Engine/Note.h"
 #include "GameFramework/MovementComponent.h"
 #include "GameFramework/NavMovementComponent.h"
@@ -129,7 +112,6 @@
 #include "Components/ExponentialHeightFogComponent.h"
 #include "Components/DirectionalLightComponent.h"
 #include "Components/SkyLightComponent.h"
-#include "AI/Navigation/NavigationGraphNodeComponent.h"
 #include "PhysicsEngine/PhysicsThrusterComponent.h"
 #include "Components/PostProcessComponent.h"
 #include "Components/ArrowComponent.h"
@@ -146,9 +128,6 @@
 #include "Components/InstancedStaticMeshComponent.h"
 #include "Interfaces/Interface_CollisionDataProvider.h"
 #include "Components/ModelComponent.h"
-#include "AI/Navigation/NavLinkRenderingComponent.h"
-#include "AI/Navigation/NavMeshRenderingComponent.h"
-#include "AI/Navigation/NavTestRenderingComponent.h"
 #include "Components/ShapeComponent.h"
 #include "Components/DrawSphereComponent.h"
 #include "Components/TextRenderComponent.h"
@@ -183,7 +162,6 @@
 #include "Animation/AnimNotifies/AnimNotifyState_Trail.h"
 #include "EditorFramework/AssetImportData.h"
 #include "Engine/AssetUserData.h"
-#include "AI/Navigation/AvoidanceManager.h"
 #include "Engine/BlendableInterface.h"
 #include "Engine/BlueprintCore.h"
 #include "Engine/Blueprint.h"
