Integrate CL 2702348 and 2698748 from Orion to Main
- Move CompleteParallelAnimationEvaluation & ParallelAnimationEvaluation implementations into cpp file, add stat for swapping buffers when completing parallel anim
- Add 'Calc SkelMesh Bounds' stat, to differentiate between using cached bounds and recalculating them
- Add stat object scopes to UpdateComponentToWorld and AnimCompletionTask
- Add stats for component registration and transform update

diff --git a/Engine/Source/Runtime/Engine/Classes/Components/SkeletalMeshComponent.h b/Engine/Source/Runtime/Engine/Classes/Components/SkeletalMeshComponent.h
--- a/Engine/Source/Runtime/Engine/Classes/Components/SkeletalMeshComponent.h
+++ b/Engine/Source/Runtime/Engine/Classes/Components/SkeletalMeshComponent.h
@@ -1248,22 +1248,9 @@ private:
 
 public:
 	// Parallel evaluation wrappers
-	void ParallelAnimationEvaluation() { PerformAnimationEvaluation(AnimEvaluationContext.SkeletalMesh, AnimEvaluationContext.AnimInstance, AnimEvaluationContext.SpaceBases, AnimEvaluationContext.LocalAtoms, AnimEvaluationContext.VertexAnims, AnimEvaluationContext.RootBoneTranslation, AnimEvaluationContext.Curve); }
-	void CompleteParallelAnimationEvaluation(bool bDoPostAnimEvaluation)
-	{
-		ParallelAnimationEvaluationTask.SafeRelease(); //We are done with this task now, clean up!
-
-		if (bDoPostAnimEvaluation && (AnimEvaluationContext.AnimInstance == AnimScriptInstance) && (AnimEvaluationContext.SkeletalMesh == SkeletalMesh) && (AnimEvaluationContext.SpaceBases.Num() == GetNumSpaceBases()))
-		{
-			Exchange(AnimEvaluationContext.SpaceBases, AnimEvaluationContext.bDoInterpolation ? CachedSpaceBases : GetEditableSpaceBases() );
-			Exchange(AnimEvaluationContext.LocalAtoms, AnimEvaluationContext.bDoInterpolation ? CachedLocalAtoms : LocalAtoms);
-			Exchange(AnimEvaluationContext.VertexAnims, ActiveVertexAnims);
-			Exchange(AnimEvaluationContext.RootBoneTranslation, RootBoneTranslation);
-
-			PostAnimEvaluation(AnimEvaluationContext);
-		}
-		AnimEvaluationContext.Clear();
-	}
+	void ParallelAnimationEvaluation();
+	void CompleteParallelAnimationEvaluation(bool bDoPostAnimEvaluation);
+
 
 	// Returns whether we are currently trying to run a parallel animation evaluation task
 	bool IsRunningParallelEvaluation() const { return IsValidRef(ParallelAnimationEvaluationTask); }
diff --git a/Engine/Source/Runtime/Engine/Private/Animation/AnimInstance.cpp b/Engine/Source/Runtime/Engine/Private/Animation/AnimInstance.cpp
--- a/Engine/Source/Runtime/Engine/Private/Animation/AnimInstance.cpp
+++ b/Engine/Source/Runtime/Engine/Private/Animation/AnimInstance.cpp
@@ -25,7 +25,7 @@
 #include "Animation/AnimInstance.h"
 /** Anim stats */
 
-DEFINE_STAT(STAT_UpdateSkelMeshBounds);
+DEFINE_STAT(STAT_CalcSkelMeshBounds);
 DEFINE_STAT(STAT_MeshObjectUpdate);
 DEFINE_STAT(STAT_BlendInPhysics);
 DEFINE_STAT(STAT_SkelCompUpdateTransform);
diff --git a/Engine/Source/Runtime/Engine/Private/Components/ActorComponent.cpp b/Engine/Source/Runtime/Engine/Private/Components/ActorComponent.cpp
--- a/Engine/Source/Runtime/Engine/Private/Components/ActorComponent.cpp
+++ b/Engine/Source/Runtime/Engine/Private/Components/ActorComponent.cpp
@@ -11,11 +11,22 @@
 #include "ComponentReregisterContext.h"
 #include "ComponentRecreateRenderStateContext.h"
 #include "Engine/SimpleConstructionScript.h"
+#include "ComponentUtils.h"
 
 #define LOCTEXT_NAMESPACE "ActorComponent"
 
 DEFINE_LOG_CATEGORY(LogActorComponent);
 
+
+DECLARE_CYCLE_STAT(TEXT("Component OnRegister"), STAT_ComponentOnRegister, STATGROUP_Component);
+DECLARE_CYCLE_STAT(TEXT("Component OnUnregister"), STAT_ComponentOnUnregister, STATGROUP_Component);
+
+DECLARE_CYCLE_STAT(TEXT("Component CreateRenderState"), STAT_ComponentCreateRenderState, STATGROUP_Component);
+DECLARE_CYCLE_STAT(TEXT("Component DestroyRenderState"), STAT_ComponentDestroyRenderState, STATGROUP_Component);
+
+DECLARE_CYCLE_STAT(TEXT("Component CreatePhysicsState"), STAT_ComponentCreatePhysicsState, STATGROUP_Component);
+DECLARE_CYCLE_STAT(TEXT("Component DestroyPhysicsState"), STAT_ComponentDestroyPhysicsState, STATGROUP_Component);
+
 /** Enable to log out all render state create, destroy and updatetransform events */
 #define LOG_RENDER_STATE 0
 
@@ -1053,18 +1064,21 @@ void UActorComponent::ExecuteRegisterEvents()
 {
 	if(!bRegistered)
 	{
+		SCOPE_CYCLE_COUNTER(STAT_ComponentOnRegister);
 		OnRegister();
 		checkf(bRegistered, TEXT("Failed to route OnRegister (%s)"), *GetFullName());
 	}
 
 	if(FApp::CanEverRender() && !bRenderStateCreated && World->Scene && ShouldCreateRenderState())
 	{
+		SCOPE_CYCLE_COUNTER(STAT_ComponentCreateRenderState);
 		CreateRenderState_Concurrent();
 		checkf(bRenderStateCreated, TEXT("Failed to route CreateRenderState_Concurrent (%s)"), *GetFullName());
 	}
 
 	if(!bPhysicsStateCreated && World->GetPhysicsScene() && ShouldCreatePhysicsState())
 	{
+		SCOPE_CYCLE_COUNTER(STAT_ComponentCreatePhysicsState);
 		CreatePhysicsState();
 		checkf(bPhysicsStateCreated, TEXT("Failed to route CreatePhysicsState (%s)"), *GetFullName());
 	}
@@ -1075,6 +1089,7 @@ void UActorComponent::ExecuteUnregisterEvents()
 {
 	if(bPhysicsStateCreated)
 	{
+		SCOPE_CYCLE_COUNTER(STAT_ComponentDestroyPhysicsState);
 		check(bRegistered); // should not have physics state unless we are registered
 		DestroyPhysicsState();
 		checkf(!bPhysicsStateCreated, TEXT("Failed to route DestroyPhysicsState (%s)"), *GetFullName());
@@ -1083,6 +1098,7 @@ void UActorComponent::ExecuteUnregisterEvents()
 
 	if(bRenderStateCreated)
 	{
+		SCOPE_CYCLE_COUNTER(STAT_ComponentDestroyRenderState);
 		check(bRegistered);
 		DestroyRenderState_Concurrent();
 		checkf(!bRenderStateCreated, TEXT("Failed to route DestroyRenderState_Concurrent (%s)"), *GetFullName());
@@ -1090,6 +1106,7 @@ void UActorComponent::ExecuteUnregisterEvents()
 
 	if(bRegistered)
 	{
+		SCOPE_CYCLE_COUNTER(STAT_ComponentOnUnregister);
 		OnUnregister();
 		checkf(!bRegistered, TEXT("Failed to route OnUnregister (%s)"), *GetFullName());
 	}
diff --git a/Engine/Source/Runtime/Engine/Private/Components/SceneComponent.cpp b/Engine/Source/Runtime/Engine/Private/Components/SceneComponent.cpp
--- a/Engine/Source/Runtime/Engine/Private/Components/SceneComponent.cpp
+++ b/Engine/Source/Runtime/Engine/Private/Components/SceneComponent.cpp
@@ -29,6 +29,10 @@ namespace SceneComponentStatics
 
 DEFINE_LOG_CATEGORY_STATIC(LogSceneComponent, Log, All);
 
+DECLARE_CYCLE_STAT(TEXT("UpdateComponentToWorld"), STAT_UpdateComponentToWorld, STATGROUP_Component);
+DECLARE_CYCLE_STAT(TEXT("Update Child Transforms"), STAT_UpdateChildTransforms, STATGROUP_Component);
+
+
 FOverlapInfo::FOverlapInfo(UPrimitiveComponent* InComponent, int32 InBodyIndex)
 	: bFromSweep(false)
 {
@@ -333,7 +337,9 @@ void USceneComponent::OnUpdateTransform(bool bSkipPhysicsMove, ETeleportType Tel
 
 void USceneComponent::UpdateComponentToWorldWithParent(USceneComponent* Parent,FName SocketName, bool bSkipPhysicsMove, const FQuat& RelativeRotationQuat, ETeleportType Teleport)
 {
-	//QUICK_SCOPE_CYCLE_COUNTER(STAT_USceneComponent_UpdateComponentToWorldWithParent);
+	SCOPE_CYCLE_COUNTER(STAT_UpdateComponentToWorld);
+	FScopeCycleCounterUObject ComponentScope(this);
+
 	// If our parent hasn't been updated before, we'll need walk up our parent attach hierarchy
 	if (Parent && !Parent->bWorldToComponentUpdated)
 	{
@@ -1553,6 +1559,8 @@ FActorComponentInstanceData* USceneComponent::GetComponentInstanceData() const
 
 void USceneComponent::UpdateChildTransforms(bool bSkipPhysicsMove, ETeleportType Teleport)
 {
+	SCOPE_CYCLE_COUNTER(STAT_UpdateChildTransforms);
+
 	for(int32 i=0; i<AttachChildren.Num(); i++)
 	{
 		USceneComponent* ChildComp = AttachChildren[i];
diff --git a/Engine/Source/Runtime/Engine/Private/Components/SkeletalMeshComponent.cpp b/Engine/Source/Runtime/Engine/Private/Components/SkeletalMeshComponent.cpp
--- a/Engine/Source/Runtime/Engine/Private/Components/SkeletalMeshComponent.cpp
+++ b/Engine/Source/Runtime/Engine/Private/Components/SkeletalMeshComponent.cpp
@@ -39,6 +39,8 @@
 
 TAutoConsoleVariable<int32> CVarUseParallelAnimationEvaluation(TEXT("a.ParallelAnimEvaluation"), 1, TEXT("If 1, animation evaluation will be run across the task graph system. If 0, evaluation will run purely on the game thread"));
 
+DECLARE_CYCLE_STAT(TEXT("Swap Anim Buffers"), STAT_CompleteAnimSwapBuffers, STATGROUP_Anim);
+DECLARE_CYCLE_STAT(TEXT("Update SkelMesh Bounds"), STAT_UpdateSkelMeshBounds, STATGROUP_Anim);
 DECLARE_CYCLE_STAT_EXTERN(TEXT("Anim Instance Spawn Time"), STAT_AnimSpawnTime, STATGROUP_Anim, );
 DEFINE_STAT(STAT_AnimSpawnTime);
 DEFINE_STAT(STAT_PostAnimEvaluation);
@@ -101,10 +103,13 @@ public:
 	void DoTask(ENamedThreads::Type CurrentThread, const FGraphEventRef& MyCompletionGraphEvent)
 	{
 		SCOPE_CYCLE_COUNTER(STAT_AnimGameThreadTime);
+
 		if (USkeletalMeshComponent* Comp = SkeletalMeshComponent.Get())
 		{
+			FScopeCycleCounterUObject ComponentScope(Comp);
+			FScopeCycleCounterUObject MeshScope(Comp->SkeletalMesh);
+
 			const bool bPerformPostAnimEvaluation = true;
-			QUICK_SCOPE_CYCLE_COUNTER(STAT_USkeletalMeshComponent_CompleteParallelAnimationEvaluation);
 			Comp->CompleteParallelAnimationEvaluation(bPerformPostAnimEvaluation);
 		}
 	}
@@ -1314,6 +1319,8 @@ void USkeletalMeshComponent::PostAnimEvaluation(FAnimationEvaluationContext& Eva
 
 void USkeletalMeshComponent::UpdateBounds()
 {
+	SCOPE_CYCLE_COUNTER(STAT_UpdateSkelMeshBounds);
+
 #if WITH_EDITOR
 	FBoxSphereBounds OriginalBounds = Bounds; // Save old bounds
 #endif
@@ -1355,6 +1362,8 @@ void USkeletalMeshComponent::UpdateBounds()
 
 FBoxSphereBounds USkeletalMeshComponent::CalcBounds(const FTransform& LocalToWorld) const
 {
+	SCOPE_CYCLE_COUNTER(STAT_CalcSkelMeshBounds);
+
 	FVector RootBoneOffset = RootBoneTranslation;
 
 	// if to use MasterPoseComponent's fixed skel bounds, 
@@ -2201,6 +2210,31 @@ void USkeletalMeshComponent::RefreshActiveVertexAnims()
 	}
 }
 
+void USkeletalMeshComponent::ParallelAnimationEvaluation() 
+{ 
+	PerformAnimationEvaluation(AnimEvaluationContext.SkeletalMesh, AnimEvaluationContext.AnimInstance, AnimEvaluationContext.SpaceBases, AnimEvaluationContext.LocalAtoms, AnimEvaluationContext.VertexAnims, AnimEvaluationContext.RootBoneTranslation, AnimEvaluationContext.Curve); 
+}
+
+void USkeletalMeshComponent::CompleteParallelAnimationEvaluation(bool bDoPostAnimEvaluation)
+{
+	ParallelAnimationEvaluationTask.SafeRelease(); //We are done with this task now, clean up!
+
+	if (bDoPostAnimEvaluation && (AnimEvaluationContext.AnimInstance == AnimScriptInstance) && (AnimEvaluationContext.SkeletalMesh == SkeletalMesh) && (AnimEvaluationContext.SpaceBases.Num() == GetNumSpaceBases()))
+	{
+		{
+			SCOPE_CYCLE_COUNTER(STAT_CompleteAnimSwapBuffers);
+
+			Exchange(AnimEvaluationContext.SpaceBases, AnimEvaluationContext.bDoInterpolation ? CachedSpaceBases : GetEditableSpaceBases());
+			Exchange(AnimEvaluationContext.LocalAtoms, AnimEvaluationContext.bDoInterpolation ? CachedLocalAtoms : LocalAtoms);
+			Exchange(AnimEvaluationContext.VertexAnims, ActiveVertexAnims);
+			Exchange(AnimEvaluationContext.RootBoneTranslation, RootBoneTranslation);
+		}
+
+		PostAnimEvaluation(AnimEvaluationContext);
+	}
+	AnimEvaluationContext.Clear();
+}
+
 bool USkeletalMeshComponent::HandleExistingParallelEvaluationTask(bool bBlockOnTask, bool bPerformPostAnimEvaluation)
 {
 	if (IsValidRef(ParallelAnimationEvaluationTask)) // We are already processing eval on another thread
diff --git a/Engine/Source/Runtime/Engine/Private/Components/SkinnedMeshComponent.cpp b/Engine/Source/Runtime/Engine/Private/Components/SkinnedMeshComponent.cpp
--- a/Engine/Source/Runtime/Engine/Private/Components/SkinnedMeshComponent.cpp
+++ b/Engine/Source/Runtime/Engine/Private/Components/SkinnedMeshComponent.cpp
@@ -615,6 +615,8 @@ void USkinnedMeshComponent::RebuildVisibilityArray()
 
 FBoxSphereBounds USkinnedMeshComponent::CalcBounds(const FTransform& LocalToWorld) const
 {
+	SCOPE_CYCLE_COUNTER(STAT_CalcSkelMeshBounds);
+
 	return CalcMeshBound( FVector::ZeroVector, false, LocalToWorld );
 }
 
@@ -636,8 +638,6 @@ class UPhysicsAsset* USkinnedMeshComponent::GetPhysicsAsset() const
 
 FBoxSphereBounds USkinnedMeshComponent::CalcMeshBound(const FVector& RootOffset, bool UsePhysicsAsset, const FTransform& LocalToWorld) const
 {
-	SCOPE_CYCLE_COUNTER(STAT_UpdateSkelMeshBounds);
-
 	FBoxSphereBounds NewBounds;
 
 	// If physics are asleep, and actor is using physics to move, skip updating the bounds.
diff --git a/Engine/Source/Runtime/Engine/Public/Animation/AnimStats.h b/Engine/Source/Runtime/Engine/Public/Animation/AnimStats.h
--- a/Engine/Source/Runtime/Engine/Public/Animation/AnimStats.h
+++ b/Engine/Source/Runtime/Engine/Public/Animation/AnimStats.h
@@ -22,7 +22,7 @@ DECLARE_CYCLE_STAT_EXTERN(TEXT("UpdateRBJointsMotors"), STAT_UpdateRBJoints, STA
 DECLARE_CYCLE_STAT_EXTERN(TEXT("UpdateLocalToWorldAndOverlaps"), STAT_UpdateLocalToWorldAndOverlaps, STATGROUP_Anim, );
 DECLARE_CYCLE_STAT_EXTERN(TEXT("SkelComp UpdateTransform"), STAT_SkelCompUpdateTransform, STATGROUP_Anim, );
 DECLARE_CYCLE_STAT_EXTERN(TEXT("MeshObject Update"), STAT_MeshObjectUpdate, STATGROUP_Anim, );
-DECLARE_CYCLE_STAT_EXTERN(TEXT("Update SkelMesh Bounds"), STAT_UpdateSkelMeshBounds, STATGROUP_Anim, );
+DECLARE_CYCLE_STAT_EXTERN(TEXT("Calc SkelMesh Bounds"), STAT_CalcSkelMeshBounds, STATGROUP_Anim, );
 DECLARE_CYCLE_STAT_EXTERN(TEXT("BlendInPhysics"), STAT_BlendInPhysics, STATGROUP_Anim, );
 DECLARE_CYCLE_STAT_EXTERN(TEXT("SkinPerPolyVertices);"), STAT_SkinPerPolyVertices, STATGROUP_Anim, );
 DECLARE_CYCLE_STAT_EXTERN(TEXT("UpdateTriMeshVertices"), STAT_UpdateTriMeshVertices, STATGROUP_Anim, );
diff --git a/Engine/Source/Runtime/Engine/Public/ComponentUtils.h b/Engine/Source/Runtime/Engine/Public/ComponentUtils.h
--- a/Engine/Source/Runtime/Engine/Public/ComponentUtils.h
+++ b/Engine/Source/Runtime/Engine/Public/ComponentUtils.h
@@ -6,6 +6,8 @@ class USimpleConstructionScript;
 class USCS_Node;
 class USceneComponent;
 
+DECLARE_STATS_GROUP(TEXT("Component"), STATGROUP_Component, STATCAT_Advanced);
+
 namespace ComponentUtils
 {
 	/**
