- Collision Analyzer improvements
  - Add support for recording info about overlap queries
  - Separate 'type' (raycast, sweep, overlap) from 'shape', add 'mode' as well (test,single,multi)
  - Fix crash when sorting by ID and then grouping
  - Add support for saving/loading data to file

diff --git a/Engine/Source/Developer/CollisionAnalyzer/CollisionAnalyzer.Build.cs b/Engine/Source/Developer/CollisionAnalyzer/CollisionAnalyzer.Build.cs
--- a/Engine/Source/Developer/CollisionAnalyzer/CollisionAnalyzer.Build.cs
+++ b/Engine/Source/Developer/CollisionAnalyzer/CollisionAnalyzer.Build.cs
@@ -6,6 +6,13 @@ public class CollisionAnalyzer : ModuleRules
 {
 	public CollisionAnalyzer(TargetInfo Target)
 	{
+        PrivateIncludePathModuleNames.AddRange(
+            new string[] {
+				"DesktopPlatform",
+				"MainFrame",
+			}
+        );
+
 		PrivateDependencyModuleNames.AddRange(
 			new string[] {
 				"Core",
@@ -18,5 +25,12 @@ public class CollisionAnalyzer : ModuleRules
 				"WorkspaceMenuStructure",
 			}
 		);
+
+        DynamicallyLoadedModuleNames.AddRange(
+            new string[] {
+				"DesktopPlatform",
+				"MainFrame",
+			}
+        );
 	}
 }
diff --git a/Engine/Source/Developer/CollisionAnalyzer/Private/CollisionAnalyzer.cpp b/Engine/Source/Developer/CollisionAnalyzer/Private/CollisionAnalyzer.cpp
--- a/Engine/Source/Developer/CollisionAnalyzer/Private/CollisionAnalyzer.cpp
+++ b/Engine/Source/Developer/CollisionAnalyzer/Private/CollisionAnalyzer.cpp
@@ -4,6 +4,85 @@
 #include "CollisionDebugDrawingPublic.h"
 
 
+/** Magic value, determining that file is a collision analyzer file.				*/
+#define COLLISION_ANALYZER_MAGIC					0x2DFF34FC
+/** Version of collision analyzer. Incremented on serialization changes.			*/
+#define COLLISION_ANALYZER_VERSION					0
+
+/** Util for serializing an FHitResult struct */
+void SerializeHitResult(FArchive& Ar, FHitResult& Result)
+{
+	bool bTempBlocking = Result.bBlockingHit;
+	Ar << bTempBlocking;
+	Result.bBlockingHit = bTempBlocking;
+
+	bool bTempPenetrating = Result.bStartPenetrating;
+	Ar << bTempPenetrating;
+	Result.bStartPenetrating = bTempPenetrating;
+
+	Ar << Result.Time;
+	Ar << Result.Distance;
+	Ar << Result.Location;
+	Ar << Result.ImpactPoint;
+	Ar << Result.Normal;
+	Ar << Result.ImpactNormal;
+	Ar << Result.TraceStart;
+	Ar << Result.TraceEnd;
+	Ar << Result.PenetrationDepth;
+	Ar << Result.BoneName;
+	Ar << Result.PhysMaterial;
+	Ar << Result.Actor;
+	Ar << Result.Component;
+	Ar << Result.FaceIndex;
+}
+
+FArchive& operator << (FArchive& Ar, FCAQuery& Query)
+{
+	Ar << Query.Start;
+	Ar << Query.End;
+	Ar << Query.Rot;
+	Ar << (int32&)Query.Type;
+	Ar << (int32&)Query.Shape;
+	Ar << (int32&)Query.Mode;
+	Ar << Query.Dims;
+	Ar << (int32&)Query.Channel;
+
+	Ar << Query.Params.TraceTag;
+	Ar << Query.Params.OwnerTag;
+	Ar << Query.Params.bTraceAsyncScene;
+	Ar << Query.Params.bTraceComplex;
+	Ar << Query.Params.bFindInitialOverlaps;
+	Ar << Query.Params.bReturnFaceIndex;
+	Ar << Query.Params.bReturnPhysicalMaterial;
+
+	int32 NumResults = 0;
+
+	if (Ar.IsLoading())
+	{
+		// Load array.
+		Ar << NumResults;
+		Query.Results.SetNum(NumResults);
+	}
+	else if(Ar.IsSaving())
+	{
+		// Save array.
+		NumResults = Query.Results.Num();
+		Ar << NumResults;
+	}
+
+	for (int32 i = 0; i < NumResults; i++)
+	{
+		SerializeHitResult(Ar, Query.Results[i]);
+	}
+
+	Ar << Query.FrameNum;
+	Ar << Query.CPUTime;
+	Ar << Query.ID;
+
+	return Ar;
+}
+
+
 //////////////////////////////////////////////////////////////////////////
 
 
@@ -11,7 +90,8 @@ void FCollisionAnalyzer::CaptureQuery(	const FVector& Start,
 										const FVector& End, 
 										const FQuat& Rot, 
 										ECAQueryType::Type QueryType, 
-										ECAQueryShape::Type QueryShape, 
+										ECAQueryShape::Type QueryShape,
+										ECAQueryMode::Type QueryMode,
 										const FVector& Dims, 
 										ECollisionChannel TraceChannel, 
 										const struct FCollisionQueryParams& Params, 
@@ -30,6 +110,7 @@ void FCollisionAnalyzer::CaptureQuery(	const FVector& Start,
 		NewQuery.Rot = Rot;
 		NewQuery.Type = QueryType;
 		NewQuery.Shape = QueryShape;
+		NewQuery.Mode = QueryMode;
 		NewQuery.Dims = Dims;
 		NewQuery.Channel = TraceChannel;
 		NewQuery.Params = Params;
@@ -85,21 +166,24 @@ void FCollisionAnalyzer::TickAnalyzer(UWorld* World)
 		if(QueryIdx < Queries.Num())
 		{
 			FCAQuery& DrawQuery = Queries[QueryIdx];
-			if(DrawQuery.Shape == ECAQueryShape::Raycast)
+			if(DrawQuery.Type == ECAQueryType::Raycast)
 			{
 				DrawLineTraces(World, DrawQuery.Start, DrawQuery.End, DrawQuery.Results, 0.f);
 			}
-			else if(DrawQuery.Shape == ECAQueryShape::SphereSweep)
+			else if(DrawQuery.Type == ECAQueryType::GeomSweep)
 			{
-				DrawSphereSweeps(World, DrawQuery.Start, DrawQuery.End, DrawQuery.Dims.X, DrawQuery.Results, 0.f);
-			}
-			else if(DrawQuery.Shape == ECAQueryShape::BoxSweep)
-			{
-				DrawBoxSweeps(World, DrawQuery.Start, DrawQuery.End, DrawQuery.Dims, DrawQuery.Rot, DrawQuery.Results, 0.f);
-			}
-			else if(DrawQuery.Shape == ECAQueryShape::CapsuleSweep)
-			{
-				DrawCapsuleSweeps(World, DrawQuery.Start, DrawQuery.End, DrawQuery.Dims.Z, DrawQuery.Dims.X, DrawQuery.Rot, DrawQuery.Results, 0.f);
+				if (DrawQuery.Shape == ECAQueryShape::Sphere)
+				{
+					DrawSphereSweeps(World, DrawQuery.Start, DrawQuery.End, DrawQuery.Dims.X, DrawQuery.Results, 0.f);
+				}
+				else if (DrawQuery.Shape == ECAQueryShape::Box)
+				{
+					DrawBoxSweeps(World, DrawQuery.Start, DrawQuery.End, DrawQuery.Dims, DrawQuery.Rot, DrawQuery.Results, 0.f);
+				}
+				else if (DrawQuery.Shape == ECAQueryShape::Capsule)
+				{
+					DrawCapsuleSweeps(World, DrawQuery.Start, DrawQuery.End, DrawQuery.Dims.Z, DrawQuery.Dims.X, DrawQuery.Rot, DrawQuery.Results, 0.f);
+				}
 			}
 		}
 	}
@@ -134,3 +218,78 @@ int32 FCollisionAnalyzer::GetNumFramesOfRecording()
 {
 	return CurrentFrameNum+1;
 }
+
+void FCollisionAnalyzer::SaveCollisionProfileData(FString ProfileFileName)
+{
+	FArchive* FileWriter = IFileManager::Get().CreateFileWriter(*ProfileFileName);
+
+	if(FileWriter != nullptr)
+	{
+		FNameAsStringProxyArchive Ar(*FileWriter);
+
+		int32 Magic = COLLISION_ANALYZER_MAGIC;
+		int32 Version = COLLISION_ANALYZER_VERSION;
+
+		Ar << Magic;
+		Ar << Version;
+		Ar << Queries;
+
+		FileWriter->Close();
+
+		delete FileWriter;
+		FileWriter = NULL;
+
+		UE_LOG(LogCollisionAnalyzer, Log, TEXT("Saved collision analyzer data to file '%s'."), *ProfileFileName);
+	}
+	else
+	{
+		UE_LOG(LogCollisionAnalyzer, Warning, TEXT("Unable to save collision analyzer data to file '%s'."), *ProfileFileName);
+	}
+}
+
+void FCollisionAnalyzer::LoadCollisionProfileData(FString ProfileFileName)
+{
+	FArchive* FileReader = IFileManager::Get().CreateFileReader(*ProfileFileName);
+
+	bool bSuccess = false;
+
+	if (FileReader != nullptr)
+	{
+		FNameAsStringProxyArchive Ar(*FileReader);
+
+		int32 Magic;
+		Ar << Magic;
+		if(Magic == COLLISION_ANALYZER_MAGIC)
+		{
+			int32 Version;
+			Ar << Version;
+			if(Version == COLLISION_ANALYZER_VERSION)
+			{
+				// First clear old data
+				Queries.Reset();
+				DrawQueryIndices.Empty();
+				CurrentFrameNum = 0;
+
+				// Load data from file
+				Ar << Queries;
+
+				// Invoke event that data has changed
+				QueriesChangedEvent.Broadcast();
+
+				UE_LOG(LogCollisionAnalyzer, Log, TEXT("Loaded collision analyzer data from file '%s'."), *ProfileFileName);
+
+				bSuccess = true;
+			}
+		}
+
+		// Close file
+		FileReader->Close();
+		delete FileReader;
+		FileReader = NULL;
+	}
+	
+	if(!bSuccess)
+	{
+		UE_LOG(LogCollisionAnalyzer, Warning, TEXT("Unable to load collision analyzer data from file '%s'."), *ProfileFileName);
+	}
+}
\ No newline at end of file
diff --git a/Engine/Source/Developer/CollisionAnalyzer/Private/CollisionAnalyzer.h b/Engine/Source/Developer/CollisionAnalyzer/Private/CollisionAnalyzer.h
--- a/Engine/Source/Developer/CollisionAnalyzer/Private/CollisionAnalyzer.h
+++ b/Engine/Source/Developer/CollisionAnalyzer/Private/CollisionAnalyzer.h
@@ -2,6 +2,7 @@
 
 #pragma once
 
+
 /** Stores information about one collision query */
 struct FCAQuery
 {
@@ -10,6 +11,7 @@ struct FCAQuery
 	FQuat						Rot;
 	ECAQueryType::Type			Type;
 	ECAQueryShape::Type			Shape;
+	ECAQueryMode::Type			Mode;
 	FVector						Dims;
 	ECollisionChannel			Channel;
 	FCollisionQueryParams		Params;
@@ -20,6 +22,8 @@ struct FCAQuery
 	int32						FrameNum;
 	float						CPUTime; /** In ms */
 	int32						ID;
+
+	friend FArchive& operator << ( FArchive& Ar, FCAQuery& Query );
 };
 
 /** Actual implementation of CollisionAnalyzer, private inside module */
@@ -32,7 +36,8 @@ public:
 		const FVector& End, 
 		const FQuat& Rot, 
 		ECAQueryType::Type QueryType, 
-		ECAQueryShape::Type QueryShape, 
+		ECAQueryShape::Type QueryShape,
+		ECAQueryMode::Type QueryMode,
 		const FVector& Dims, 
 		ECollisionChannel TraceChannel, 
 		const struct FCollisionQueryParams& Params, 
@@ -76,6 +81,11 @@ public:
 	FQueriesChangedEvent& OnQueriesChanged() { return QueriesChangedEvent; }
 	FQueriesChangedEvent& OnQueryAdded() { return QueryAddedEvent; }
 
+	/** Save current data to a file */
+	void SaveCollisionProfileData(FString ProfileFileName);
+	/** Load data from file */
+	void LoadCollisionProfileData(FString ProfileFileName);
+
 private:
 	/** The current frame number we are on while recording */
 	int32				CurrentFrameNum;
diff --git a/Engine/Source/Developer/CollisionAnalyzer/Private/CollisionAnalyzerStyle.cpp b/Engine/Source/Developer/CollisionAnalyzer/Private/CollisionAnalyzerStyle.cpp
--- a/Engine/Source/Developer/CollisionAnalyzer/Private/CollisionAnalyzerStyle.cpp
+++ b/Engine/Source/Developer/CollisionAnalyzer/Private/CollisionAnalyzerStyle.cpp
@@ -47,6 +47,8 @@ void FCollisionAnalyzerStyle::FStyle::Initialize()
 	StyleInstance->Set("CollisionAnalyzer.ShowRecent", new IMAGE_BRUSH("Icons/CA_ShowRecent", Icon24x24));
 	StyleInstance->Set("CollisionAnalyzer.Group", new IMAGE_BRUSH("Icons/CA_Group", FVector2D(10, 18)));
 	StyleInstance->Set("CollisionAnalyzer.GroupBackground", new BOX_BRUSH("Icons/CA_GroupBackground", FMargin(4.f / 16.f)));
+	StyleInstance->Set("CollisionAnalyzer.Save", new IMAGE_BRUSH("Icons/icon_file_save_40x", Icon24x24));
+	StyleInstance->Set("CollisionAnalyzer.Load", new IMAGE_BRUSH("Icons/icon_file_open_40x", Icon24x24));
 	StyleInstance->Set("Menu.Background", new BOX_BRUSH("Old/Menu_Background", FMargin(8.0f / 64.0f)));
 	StyleInstance->Set("BoldFont", TTF_CORE_FONT("Fonts/Roboto-Bold", 9));
 
diff --git a/Engine/Source/Developer/CollisionAnalyzer/Private/SCAQueryTableRow.cpp b/Engine/Source/Developer/CollisionAnalyzer/Private/SCAQueryTableRow.cpp
--- a/Engine/Source/Developer/CollisionAnalyzer/Private/SCAQueryTableRow.cpp
+++ b/Engine/Source/Developer/CollisionAnalyzer/Private/SCAQueryTableRow.cpp
@@ -23,29 +23,29 @@ TSharedRef<SWidget> SCAQueryTableRow::GenerateWidgetForColumn(const FName& Colum
 	// GROUP
 	if(Item->bIsGroup)
 	{
-		if (ColumnName == TEXT("ID"))
+		if (ColumnName == SCollisionAnalyzer::IDColumnName)
 		{
 			return	SNew(SExpanderArrow, SharedThis(this));
 		}
-		else if (ColumnName == TEXT("Frame") && OwnerAnalyzerWidget->GroupBy == EQueryGroupMode::ByFrameNum)
+		else if (ColumnName == SCollisionAnalyzer::FrameColumnName && OwnerAnalyzerWidget->GroupBy == EQueryGroupMode::ByFrameNum)
 		{
 			return	SNew(STextBlock)
 					.Font(FCollisionAnalyzerStyle::Get()->GetFontStyle("BoldFont"))
 					.Text(FText::AsNumber(Item->FrameNum));
 		}
-		else if (ColumnName == TEXT("Tag") && OwnerAnalyzerWidget->GroupBy == EQueryGroupMode::ByTag)
+		else if (ColumnName == SCollisionAnalyzer::TagColumnName && OwnerAnalyzerWidget->GroupBy == EQueryGroupMode::ByTag)
 		{
 			return	SNew(STextBlock)
 					.Font(FCollisionAnalyzerStyle::Get()->GetFontStyle("BoldFont"))
 					.Text(FText::FromName(Item->GroupName));
 		}
-		else if (ColumnName == TEXT("Owner") && OwnerAnalyzerWidget->GroupBy == EQueryGroupMode::ByOwnerTag)
+		else if (ColumnName == SCollisionAnalyzer::OwnerColumnName && OwnerAnalyzerWidget->GroupBy == EQueryGroupMode::ByOwnerTag)
 		{
 			return	SNew(STextBlock)
 					.Font(FCollisionAnalyzerStyle::Get()->GetFontStyle("BoldFont"))
 					.Text(FText::FromName(Item->GroupName));
 		}
-		else if (ColumnName == TEXT("Time"))
+		else if (ColumnName == SCollisionAnalyzer::TimeColumnName)
 		{
 			return	SNew(STextBlock)
 					.Font(FCollisionAnalyzerStyle::Get()->GetFontStyle("BoldFont"))
@@ -64,37 +64,49 @@ TSharedRef<SWidget> SCAQueryTableRow::GenerateWidgetForColumn(const FName& Colum
 
 		FCAQuery& Query = OwnerAnalyzerWidget->Analyzer->Queries[QueryId];
 
-		if (ColumnName == TEXT("ID"))
+		if (ColumnName == SCollisionAnalyzer::IDColumnName)
 		{	
 			return	SNew(STextBlock)
 					.Text( FText::AsNumber(QueryId) );
 		}
-		else if (ColumnName == TEXT("Frame"))
+		else if (ColumnName == SCollisionAnalyzer::FrameColumnName)
 		{
 			return	SNew(STextBlock)
 					.Text( FText::AsNumber(Query.FrameNum) );
 		}
-		else if (ColumnName == TEXT("Type"))
+		else if (ColumnName == SCollisionAnalyzer::TypeColumnName)
 		{
 			return	SNew(STextBlock)
 					.Text(FText::FromString(SCollisionAnalyzer::QueryTypeToString(Query.Type)));
 		}
-		else if (ColumnName == TEXT("Shape"))
+		else if (ColumnName == SCollisionAnalyzer::ShapeColumnName)
+		{
+			// Leave shape string blank if this is a raycast, it doesn't matter
+			FString ShapeString;
+			if(Query.Type != ECAQueryType::Raycast)
+			{
+				ShapeString = SCollisionAnalyzer::QueryShapeToString(Query.Shape);
+			}
+
+			return	SNew(STextBlock)
+					.Text(FText::FromString(ShapeString));
+		}
+		else if (ColumnName == SCollisionAnalyzer::ModeColumnName)
 		{
 			return	SNew(STextBlock)
-					.Text(FText::FromString(SCollisionAnalyzer::QueryShapeToString(Query.Shape)));
+					.Text(FText::FromString(SCollisionAnalyzer::QueryModeToString(Query.Mode)));
 		}
-		else if (ColumnName == TEXT("Tag"))
+		else if (ColumnName == SCollisionAnalyzer::TagColumnName)
 		{
 			return	SNew(STextBlock)
 					.Text(FText::FromName(Query.Params.TraceTag));
 		}
-		else if (ColumnName == TEXT("Owner"))
+		else if (ColumnName == SCollisionAnalyzer::OwnerColumnName)
 		{
 			return	SNew(STextBlock)
 					.Text(FText::FromName(Query.Params.OwnerTag));
 		}
-		else if (ColumnName == TEXT("NumBlock"))
+		else if (ColumnName == SCollisionAnalyzer::NumBlockColumnName)
 		{
 			FHitResult* FirstHit = FHitResult::GetFirstBlockingHit(Query.Results);
 			bool bStartPenetrating = (FirstHit != NULL) && FirstHit->bStartPenetrating;
@@ -103,19 +115,15 @@ TSharedRef<SWidget> SCAQueryTableRow::GenerateWidgetForColumn(const FName& Colum
 					.Text(FText::AsNumber(FHitResult::GetNumBlockingHits(Query.Results)))
 					.ColorAndOpacity(bStartPenetrating ? FLinearColor(1.f,0.25f,0.25f) : FSlateColor::UseForeground() );
 		}
-		else if (ColumnName == TEXT("NumTouch"))
+		else if (ColumnName == SCollisionAnalyzer::NumTouchColumnName)
 		{
 			return	SNew(STextBlock)
 					.Text(FText::AsNumber(FHitResult::GetNumOverlapHits(Query.Results)));
 		}
-		else if (ColumnName == TEXT("Time"))
+		else if (ColumnName == SCollisionAnalyzer::TimeColumnName)
 		{
-			static const FNumberFormattingOptions TimeFormatOptions = FNumberFormattingOptions()
-				.SetMinimumFractionalDigits(3)
-				.SetMaximumFractionalDigits(3);
-
 			return	SNew(STextBlock)
-					.Text(FText::AsNumber(Query.CPUTime, &TimeFormatOptions));
+					.Text(FText::FromString(FString::Printf(TEXT("%.3f"), Query.CPUTime)));
 		}
 	}
 
@@ -124,12 +132,8 @@ TSharedRef<SWidget> SCAQueryTableRow::GenerateWidgetForColumn(const FName& Colum
 
 FText SCAQueryTableRow::GetTotalTimeText() const
 {
-	static const FNumberFormattingOptions TimeFormatOptions = FNumberFormattingOptions()
-		.SetMinimumFractionalDigits(3)
-		.SetMaximumFractionalDigits(3);
-
 	check(Item->bIsGroup)
-	return FText::AsNumber(Item->TotalCPUTime, &TimeFormatOptions);
+	return FText::FromString(FString::Printf(TEXT("%.3f"), Item->TotalCPUTime));
 }
 
 #undef LOCTEXT_NAMESPACE
\ No newline at end of file
diff --git a/Engine/Source/Developer/CollisionAnalyzer/Private/SCollisionAnalyzer.cpp b/Engine/Source/Developer/CollisionAnalyzer/Private/SCollisionAnalyzer.cpp
--- a/Engine/Source/Developer/CollisionAnalyzer/Private/SCollisionAnalyzer.cpp
+++ b/Engine/Source/Developer/CollisionAnalyzer/Private/SCollisionAnalyzer.cpp
@@ -1,13 +1,24 @@
 // Copyright 1998-2015 Epic Games, Inc. All Rights Reserved.
 
 #include "CollisionAnalyzerPCH.h"
+#include "DesktopPlatformModule.h"
+#include "MainFrame.h"
 
 #define LOCTEXT_NAMESPACE "SCollisionAnalyzer"
 
 static const int32 NumDrawRecentQueries = 10;
 
-static FName TimeColumnName(TEXT("Name"));
-static FName IDColumnName(TEXT("ID"));
+// Column names
+const FName SCollisionAnalyzer::IDColumnName(TEXT("ID"));
+const FName SCollisionAnalyzer::FrameColumnName(TEXT("Frame"));
+const FName SCollisionAnalyzer::TypeColumnName(TEXT("Type"));
+const FName SCollisionAnalyzer::ShapeColumnName(TEXT("Shape"));
+const FName SCollisionAnalyzer::ModeColumnName(TEXT("Mode"));
+const FName SCollisionAnalyzer::TagColumnName(TEXT("Tag"));
+const FName SCollisionAnalyzer::OwnerColumnName(TEXT("Owner"));
+const FName SCollisionAnalyzer::NumBlockColumnName(TEXT("NumBlock"));
+const FName SCollisionAnalyzer::NumTouchColumnName(TEXT("NumTouch"));
+const FName SCollisionAnalyzer::TimeColumnName(TEXT("Time"));
 
 BEGIN_SLATE_FUNCTION_BUILD_OPTIMIZATION
 void SCollisionAnalyzer::Construct(const FArguments& InArgs, FCollisionAnalyzer* InAnalyzer)
@@ -60,6 +71,32 @@ void SCollisionAnalyzer::Construct(const FArguments& InArgs, FCollisionAnalyzer*
 						.Image(FCollisionAnalyzerStyle::Get()->GetBrush("CollisionAnalyzer.ShowRecent"))
 					]
 				]
+				// Load profile
+				+ SHorizontalBox::Slot()
+				.AutoWidth()
+				.Padding(1)
+				[
+					SNew(SButton)
+					.OnClicked(this, &SCollisionAnalyzer::OnLoadButtonClicked)
+					.Content()
+					[
+						SNew(SImage)
+						.Image(FCollisionAnalyzerStyle::Get()->GetBrush("CollisionAnalyzer.Load"))
+					]
+				]
+				// Save profile
+				+ SHorizontalBox::Slot()
+				.AutoWidth()
+				.Padding(1)
+				[
+					SNew(SButton)
+					.OnClicked(this, &SCollisionAnalyzer::OnSaveButtonClicked)
+					.Content()
+					[
+						SNew(SImage)
+						.Image(FCollisionAnalyzerStyle::Get()->GetBrush("CollisionAnalyzer.Save"))
+					]
+				]
 			]
 		]
 		// List area
@@ -95,7 +132,7 @@ void SCollisionAnalyzer::Construct(const FArguments& InArgs, FCollisionAnalyzer*
 							.Text(LOCTEXT("QueryListIdHeader", "ID"))
 						]
 						// Frame number
-						+SHeaderRow::Column("Frame")
+						+SHeaderRow::Column(FrameColumnName)
 						.FixedWidth(48)
 						[
 							SNew(SVerticalBox)
@@ -135,21 +172,28 @@ void SCollisionAnalyzer::Construct(const FArguments& InArgs, FCollisionAnalyzer*
 							]
 						]
 						// Type
-						+SHeaderRow::Column("Type")
-						.FillWidth(0.75)
+						+SHeaderRow::Column(TypeColumnName)
+						.FillWidth(0.5)
 						[
 							SNew(STextBlock)
 							.Text(LOCTEXT("QueryListTypeHeader", "Type"))
 						]
 						// Shape
-						+SHeaderRow::Column("Shape")
-						.FillWidth(0.75)
+						+SHeaderRow::Column(ShapeColumnName)
+						.FillWidth(0.5)
 						[
 							SNew(STextBlock)
 							.Text(LOCTEXT("QueryListShapeHeader", "Shape"))
 						]
+						// Shape
+						+SHeaderRow::Column(ModeColumnName)
+						.FillWidth(0.5)
+						[
+							SNew(STextBlock)
+							.Text(LOCTEXT("QueryListModeHeader", "Mode"))
+						]
 						// Tag
-						+SHeaderRow::Column("Tag")
+						+SHeaderRow::Column(TagColumnName)
 						.FillWidth(1.5)
 						[
 							SNew(SVerticalBox)
@@ -189,7 +233,7 @@ void SCollisionAnalyzer::Construct(const FArguments& InArgs, FCollisionAnalyzer*
 							]
 						]
 						// Owner
-						+SHeaderRow::Column("Owner")
+						+SHeaderRow::Column(OwnerColumnName)
 						.FillWidth(1.5)
 						[
 							SNew(SVerticalBox)
@@ -229,7 +273,7 @@ void SCollisionAnalyzer::Construct(const FArguments& InArgs, FCollisionAnalyzer*
 							]
 						]
 						// Num blocking hits
-						+SHeaderRow::Column("NumBlock")
+						+SHeaderRow::Column(NumBlockColumnName)
 						.FixedWidth(24)
 						[
 							SNew(STextBlock)
@@ -237,7 +281,7 @@ void SCollisionAnalyzer::Construct(const FArguments& InArgs, FCollisionAnalyzer*
 							.ToolTipText( LOCTEXT("NumberBlocksTooltip", "Number of blocking results, red means 'started penetrating'") )
 						]
 						// Num touching hits
-						+SHeaderRow::Column("NumTouch")
+						+SHeaderRow::Column(NumTouchColumnName)
 						.FixedWidth(24)
 						[
 							SNew(STextBlock)
@@ -538,13 +582,16 @@ void SCollisionAnalyzer::AddQueryToGroupedQueries(int32 NewQueryIndex, bool bPer
 	{
 		GroupedQueries.Add(NewItem);
 
-		if(SortBy == EQuerySortMode::ByTime && bPerformSort)
-		{
-			GroupedQueries.Sort( FCompareQueryByCPUTime(Analyzer) );
-		}
-		else if (SortBy == EQuerySortMode::ByID)
+		if(bPerformSort)
 		{
-			GroupedQueries.Sort( FCompareQueryByID(Analyzer, SortDirection) );
+			if(SortBy == EQuerySortMode::ByTime)
+			{
+				GroupedQueries.Sort( FCompareQueryByCPUTime(Analyzer) );
+			}
+			else if (SortBy == EQuerySortMode::ByID)
+			{
+				GroupedQueries.Sort( FCompareQueryByID(Analyzer, SortDirection) );
+			}
 		}
 	}
 	// If we are grouping..
@@ -581,6 +628,11 @@ void SCollisionAnalyzer::AddQueryToGroupedQueries(int32 NewQueryIndex, bool bPer
 				GroupedQueries.Sort( FCompareGroupByCPUTime() );
 				AddToGroup->QueriesInGroup.Sort( FCompareQueryByCPUTime(Analyzer) );
 			}
+			else if (SortBy == EQuerySortMode::ByID)
+			{
+				GroupedQueries.Sort( FCompareGroupByCPUTime() );
+				AddToGroup->QueriesInGroup.Sort( FCompareQueryByID(Analyzer, SortDirection) );
+			}
 		}
 	}
 
@@ -637,7 +689,22 @@ void SCollisionAnalyzer::RebuildFilteredList()
 	}
 	else if (SortBy == EQuerySortMode::ByID)
 	{
-		GroupedQueries.Sort(FCompareQueryByID(Analyzer, SortDirection));
+		// Ungrouped
+		if (GroupBy == EQueryGroupMode::Ungrouped)
+		{
+			GroupedQueries.Sort(FCompareQueryByID(Analyzer, SortDirection));
+		}
+		// Grouped
+		else
+		{
+			GroupedQueries.Sort(FCompareGroupByCPUTime()); // Sort groups by time again, doesn't make sense to sort by 
+			for (int i = 0; i < GroupedQueries.Num(); i++)
+			{
+				TSharedPtr<FQueryTreeItem> Group = GroupedQueries[i];
+				check(Group->bIsGroup);
+				Group->QueriesInGroup.Sort(FCompareQueryByID(Analyzer, SortDirection));
+			}
+		}
 	}
 	
 	// When underlying array changes, refresh list
@@ -682,6 +749,83 @@ FReply SCollisionAnalyzer::OnRecordButtonClicked()
 	return FReply::Handled();
 }
 
+FReply SCollisionAnalyzer::OnLoadButtonClicked()
+{
+	IDesktopPlatform* DesktopPlatform = FDesktopPlatformModule::Get();
+	if (DesktopPlatform)
+	{
+		// Get the window handles
+		void* ParentWindowWindowHandle = NULL;
+		IMainFrameModule& MainFrameModule = FModuleManager::LoadModuleChecked<IMainFrameModule>(TEXT("MainFrame"));
+		const TSharedPtr<SWindow>& MainFrameParentWindow = MainFrameModule.GetParentWindow();
+		if (MainFrameParentWindow.IsValid() && MainFrameParentWindow->GetNativeWindow().IsValid())
+		{
+			ParentWindowWindowHandle = MainFrameParentWindow->GetNativeWindow()->GetOSWindowHandle();
+		}
+
+		// Default path to find stats
+		const FString DefaultPath = FPaths::ProfilingDir() + TEXT("CollisionAnalyzer");
+
+		// File open dialog
+		TArray<FString> Filenames;
+		bool bOpened = DesktopPlatform->OpenFileDialog(
+			ParentWindowWindowHandle,
+			LOCTEXT("OpenProjectBrowseTitle", "Open Project").ToString(),
+			DefaultPath,
+			TEXT(""),
+			TEXT( "UCA file|*.uca" ),
+			EFileDialogFlags::None,
+			Filenames
+			);
+
+		// If we chose a file
+		if(bOpened && Filenames.Num() > 0)
+		{
+			Analyzer->LoadCollisionProfileData(Filenames[0]);
+		}
+	}
+	return FReply::Handled();
+}
+
+FReply SCollisionAnalyzer::OnSaveButtonClicked()
+{
+	IDesktopPlatform* DesktopPlatform = FDesktopPlatformModule::Get();
+	if (DesktopPlatform)
+	{
+		// Get the window handles
+		void* ParentWindowWindowHandle = NULL;
+		IMainFrameModule& MainFrameModule = FModuleManager::LoadModuleChecked<IMainFrameModule>(TEXT("MainFrame"));
+		const TSharedPtr<SWindow>& MainFrameParentWindow = MainFrameModule.GetParentWindow();
+		if (MainFrameParentWindow.IsValid() && MainFrameParentWindow->GetNativeWindow().IsValid())
+		{
+			ParentWindowWindowHandle = MainFrameParentWindow->GetNativeWindow()->GetOSWindowHandle();
+		}
+
+		// Default path to find stats
+		const FString DefaultPath = FPaths::ProfilingDir() + TEXT("CollisionAnalyzer");
+
+		// File save dialog
+		TArray<FString> Filenames;
+		bool bSaved = DesktopPlatform->SaveFileDialog(
+			ParentWindowWindowHandle,
+			LOCTEXT("CollisionFileLocation", "Choose file location").ToString(),
+			DefaultPath,
+			TEXT(""),
+			TEXT( "UCA file|*.uca" ),
+			EFileDialogFlags::None,
+			Filenames
+			);
+
+		// If we chose a file
+		if(bSaved && Filenames.Num() > 0)
+		{
+			Analyzer->SaveCollisionProfileData(Filenames[0]);
+		}
+	}
+
+	return FReply::Handled();
+}
+
 void SCollisionAnalyzer::OnDrawRecentChanged(ECheckBoxState NewState)
 {
 	bDrawRecentQueries = (NewState == ECheckBoxState::Checked);
@@ -769,19 +913,32 @@ EColumnSortMode::Type SCollisionAnalyzer::GetTimeSortMode() const
 	return (SortBy == EQuerySortMode::ByTime) ? SortDirection : EColumnSortMode::None;
 }
 
+FString SCollisionAnalyzer::QueryTypeToString(ECAQueryType::Type QueryType)
+{
+	switch (QueryType)
+	{
+	case ECAQueryType::Raycast:
+		return FString(TEXT("Raycast"));
+	case ECAQueryType::GeomSweep:
+		return FString(TEXT("Sweep"));
+	case ECAQueryType::GeomOverlap:
+		return FString(TEXT("Overlap"));
+	}
+
+	return FString(TEXT("UNKNOWNN"));
+}
+
 FString SCollisionAnalyzer::QueryShapeToString(ECAQueryShape::Type QueryShape)
 {
 	switch(QueryShape)
 	{
-	case ECAQueryShape::Raycast:
-		return FString(TEXT("Raycast"));
-	case ECAQueryShape::SphereSweep:
+	case ECAQueryShape::Sphere:
 		return FString(TEXT("Sphere"));
-	case ECAQueryShape::BoxSweep:
+	case ECAQueryShape::Box:
 		return FString(TEXT("Box"));
-	case ECAQueryShape::CapsuleSweep:
+	case ECAQueryShape::Capsule:
 		return FString(TEXT("Capsule"));	
-	case ECAQueryShape::ConvexSweep:
+	case ECAQueryShape::Convex:
 		return FString(TEXT("Convex"));
 	}
 
@@ -789,15 +946,15 @@ FString SCollisionAnalyzer::QueryShapeToString(ECAQueryShape::Type QueryShape)
 }
 
 
-FString SCollisionAnalyzer::QueryTypeToString(ECAQueryType::Type QueryType)
+FString SCollisionAnalyzer::QueryModeToString(ECAQueryMode::Type QueryMode)
 {
-	switch(QueryType)
+	switch(QueryMode)
 	{
-	case ECAQueryType::Test:
+	case ECAQueryMode::Test:
 		return FString(TEXT("Test"));
-	case ECAQueryType::Single:
+	case ECAQueryMode::Single:
 		return FString(TEXT("Single"));
-	case ECAQueryType::Multi:
+	case ECAQueryMode::Multi:
 		return FString(TEXT("Multi"));
 	}
 
diff --git a/Engine/Source/Developer/CollisionAnalyzer/Private/SCollisionAnalyzer.h b/Engine/Source/Developer/CollisionAnalyzer/Private/SCollisionAnalyzer.h
--- a/Engine/Source/Developer/CollisionAnalyzer/Private/SCollisionAnalyzer.h
+++ b/Engine/Source/Developer/CollisionAnalyzer/Private/SCollisionAnalyzer.h
@@ -103,6 +103,8 @@ public:
 	EColumnSortMode::Type GetTimeSortMode() const;
 	// Handler delegates
 	FReply OnRecordButtonClicked();
+	FReply OnLoadButtonClicked();
+	FReply OnSaveButtonClicked();
 	void OnDrawRecentChanged(ECheckBoxState NewState);
 	void OnGroupByFrameChanged(ECheckBoxState NewState);
 	void OnGroupByTagChanged(ECheckBoxState NewState);
@@ -120,6 +122,9 @@ public:
 	/** Util to convert a query type to a string */
 	static FString QueryTypeToString(ECAQueryType::Type QueryType);
 
+	/** Util to convert a query mode to a string */
+	static FString QueryModeToString(ECAQueryMode::Type QueryMode);
+
 	/** Pointer to the analyzer object we want to show ui for */
 	FCollisionAnalyzer*				Analyzer;
 
@@ -181,4 +186,17 @@ private:
 	/** Box to filter by time */
 	TSharedPtr<SEditableTextBox>	TimeFilterBox;
 
+	// Column names
+public:
+	static const FName IDColumnName;
+	static const FName FrameColumnName;
+	static const FName TypeColumnName;
+	static const FName ShapeColumnName;
+	static const FName ModeColumnName;
+	static const FName TagColumnName;
+	static const FName OwnerColumnName;
+	static const FName NumBlockColumnName;
+	static const FName NumTouchColumnName;
+	static const FName TimeColumnName;
+
 };
diff --git a/Engine/Source/Developer/CollisionAnalyzer/Public/ICollisionAnalyzer.h b/Engine/Source/Developer/CollisionAnalyzer/Public/ICollisionAnalyzer.h
--- a/Engine/Source/Developer/CollisionAnalyzer/Public/ICollisionAnalyzer.h
+++ b/Engine/Source/Developer/CollisionAnalyzer/Public/ICollisionAnalyzer.h
@@ -2,21 +2,31 @@
 
 #pragma once
 
+/** Enum to describe type of the query */
+namespace ECAQueryType
+{
+	enum Type
+	{
+		Raycast,
+		GeomSweep,
+		GeomOverlap
+	};
+}
+
 /** Enum to describe shape of the query */
 namespace ECAQueryShape
 {
 	enum Type
 	{
-		Raycast,
-		SphereSweep,
-		CapsuleSweep,
-		BoxSweep,
-		ConvexSweep
+		Sphere,
+		Box,
+		Capsule,
+		Convex
 	};
 }
 
-/** Enum to describe the type of query performed */
-namespace ECAQueryType
+/** Enum to describe the mode of query performed */
+namespace ECAQueryMode
 {
 	enum Type
 	{
@@ -38,7 +48,8 @@ public:
 		const FVector& End, 
 		const FQuat& Rot, 
 		ECAQueryType::Type QueryType, 
-		ECAQueryShape::Type QueryShape, 
+		ECAQueryShape::Type QueryShape,
+		ECAQueryMode::Type QueryMode,
 		const FVector& Dims, 
 		ECollisionChannel TraceChannel, 
 		const struct FCollisionQueryParams& Params, 
diff --git a/Engine/Source/Runtime/Engine/Private/Collision/PhysXCollision.cpp b/Engine/Source/Runtime/Engine/Private/Collision/PhysXCollision.cpp
--- a/Engine/Source/Runtime/Engine/Private/Collision/PhysXCollision.cpp
+++ b/Engine/Source/Runtime/Engine/Private/Collision/PhysXCollision.cpp
@@ -712,55 +712,65 @@ bool GCollisionAnalyzerIsRecording = false;
 
 bool bSkipCapture = false;
 
-/** Util to extract type and dimensions from physx geom being swept, and pass info to CollisionAnalyzer, if its recording */
-void CaptureGeomSweep(const UWorld* World, const FVector& Start, const FVector& End, const PxQuat& PRot, ECAQueryType::Type QueryType, const PxGeometry& PGeom, ECollisionChannel TraceChannel, const struct FCollisionQueryParams& Params, const struct FCollisionResponseParams& ResponseParams, const struct FCollisionObjectQueryParams& ObjectParams, const TArray<FHitResult>& Results, double CPUTime)
+/** Util to convert from PhysX shape and rotation to unreal shape enum, dimension vector and rotation */
+static void P2UGeomAndRot(const PxGeometry& PGeom, const PxQuat& PRot, ECAQueryShape::Type& OutQueryShape, FVector& OutDims, FQuat& OutQuat)
 {
-	if(bSkipCapture || !GCollisionAnalyzerIsRecording || !IsInGameThread())
-	{
-		return;
-	}
-
-	ECAQueryShape::Type QueryShape = ECAQueryShape::Raycast;
-	FVector Dims(0,0,0);
-	FQuat Rot = FQuat::Identity;
+	OutQueryShape = ECAQueryShape::Capsule;
+	OutDims = FVector(0, 0, 0);
+	OutQuat = FQuat::Identity;
 
-	switch(PGeom.getType())
+	switch (PGeom.getType())
 	{
 	case PxGeometryType::eCAPSULE:
-		{
-			QueryShape = ECAQueryShape::CapsuleSweep;
-			PxCapsuleGeometry* PCapsuleGeom = (PxCapsuleGeometry*)&PGeom;
-			Dims = FVector(PCapsuleGeom->radius, PCapsuleGeom->radius, PCapsuleGeom->halfHeight + PCapsuleGeom->radius);
-			Rot = ConvertToUECapsuleRot(PRot);
-			break;
-		}
+	{
+		OutQueryShape = ECAQueryShape::Capsule;
+		PxCapsuleGeometry* PCapsuleGeom = (PxCapsuleGeometry*)&PGeom;
+		OutDims = FVector(PCapsuleGeom->radius, PCapsuleGeom->radius, PCapsuleGeom->halfHeight + PCapsuleGeom->radius);
+		OutQuat = ConvertToUECapsuleRot(PRot);
+		break;
+	}
 
 	case PxGeometryType::eSPHERE:
-		{
-			QueryShape = ECAQueryShape::SphereSweep;
-			PxSphereGeometry* PSphereGeom = (PxSphereGeometry*)&PGeom;
-			Dims = FVector(PSphereGeom->radius);
-			break;
-		}
+	{
+		OutQueryShape = ECAQueryShape::Sphere;
+		PxSphereGeometry* PSphereGeom = (PxSphereGeometry*)&PGeom;
+		OutDims = FVector(PSphereGeom->radius);
+		break;
+	}
 
 	case PxGeometryType::eBOX:
-		{
-			QueryShape = ECAQueryShape::BoxSweep;
-			PxBoxGeometry* PBoxGeom = (PxBoxGeometry*)&PGeom;
-			Dims = P2UVector(PBoxGeom->halfExtents);
-			Rot = P2UQuat(PRot);
-			break;
-		}
+	{
+		OutQueryShape = ECAQueryShape::Box;
+		PxBoxGeometry* PBoxGeom = (PxBoxGeometry*)&PGeom;
+		OutDims = P2UVector(PBoxGeom->halfExtents);
+		OutQuat = P2UQuat(PRot);
+		break;
+	}
 
 	case PxGeometryType::eCONVEXMESH:
-		{
-			QueryShape = ECAQueryShape::ConvexSweep;
-			break;
-		}
+	{
+		OutQueryShape = ECAQueryShape::Convex;
+		break;
+	}
 
 	default:
 		UE_LOG(LogCollision, Warning, TEXT("CaptureGeomSweep: Unknown geom type."));
 	}
+}
+
+/** Util to extract type and dimensions from physx geom being swept, and pass info to CollisionAnalyzer, if its recording */
+void CaptureGeomSweep(const UWorld* World, const FVector& Start, const FVector& End, const PxQuat& PRot, ECAQueryMode::Type QueryMode, const PxGeometry& PGeom, ECollisionChannel TraceChannel, const struct FCollisionQueryParams& Params, const struct FCollisionResponseParams& ResponseParams, const struct FCollisionObjectQueryParams& ObjectParams, const TArray<FHitResult>& Results, double CPUTime)
+{
+	if(bSkipCapture || !GCollisionAnalyzerIsRecording || !IsInGameThread())
+	{
+		return;
+	}
+
+	// Convert from PhysX to Unreal types
+	ECAQueryShape::Type QueryShape = ECAQueryShape::Sphere;
+	FVector Dims(0,0,0);
+	FQuat Rot = FQuat::Identity;
+	P2UGeomAndRot(PGeom, PRot, QueryShape, Dims, Rot);
 
 	// Do a touch all query to find things we _didn't_ hit
 	bSkipCapture = true;
@@ -769,11 +779,11 @@ void CaptureGeomSweep(const UWorld* World, const FVector& Start, const FVector&
 	bSkipCapture = false;
 
 	// Now tell analyzer
-	FCollisionAnalyzerModule::Get()->CaptureQuery(Start, End, Rot, QueryType, QueryShape, Dims, TraceChannel, Params, ResponseParams, ObjectParams, Results, TouchAllResults, CPUTime);
+	FCollisionAnalyzerModule::Get()->CaptureQuery(Start, End, Rot, ECAQueryType::GeomSweep, QueryShape, QueryMode, Dims, TraceChannel, Params, ResponseParams, ObjectParams, Results, TouchAllResults, CPUTime);
 }
 
 /** Util to capture a raycast with the CollisionAnalyzer if recording */
-void CaptureRaycast(const UWorld* World, const FVector& Start, const FVector& End, ECAQueryType::Type QueryType, ECollisionChannel TraceChannel, const struct FCollisionQueryParams& Params, const struct FCollisionResponseParams& ResponseParams, const struct FCollisionObjectQueryParams& ObjectParams, const TArray<FHitResult>& Results, double CPUTime)
+void CaptureRaycast(const UWorld* World, const FVector& Start, const FVector& End, ECAQueryMode::Type QueryMode, ECollisionChannel TraceChannel, const struct FCollisionQueryParams& Params, const struct FCollisionResponseParams& ResponseParams, const struct FCollisionObjectQueryParams& ObjectParams, const TArray<FHitResult>& Results, double CPUTime)
 {
 	if(bSkipCapture || !GCollisionAnalyzerIsRecording || !IsInGameThread())
 	{
@@ -786,18 +796,50 @@ void CaptureRaycast(const UWorld* World, const FVector& Start, const FVector& En
 	RaycastMulti(World, TouchAllResults, Start, End, DefaultCollisionChannel, Params, ResponseParams, FCollisionObjectQueryParams(FCollisionObjectQueryParams::InitType::AllObjects));
 	bSkipCapture = false;
 
-	FCollisionAnalyzerModule::Get()->CaptureQuery(Start, End, FQuat::Identity, QueryType, ECAQueryShape::Raycast, FVector(0,0,0), TraceChannel, Params, ResponseParams, ObjectParams, Results, TouchAllResults, CPUTime);
+	FCollisionAnalyzerModule::Get()->CaptureQuery(Start, End, FQuat::Identity, ECAQueryType::Raycast, ECAQueryShape::Sphere, QueryMode, FVector(0,0,0), TraceChannel, Params, ResponseParams, ObjectParams, Results, TouchAllResults, CPUTime);
+}
+
+void CaptureOverlap(const UWorld* World, const PxGeometry& PGeom, const PxTransform& PGeomPose, ECAQueryMode::Type QueryMode, ECollisionChannel TraceChannel, const struct FCollisionQueryParams& Params, const struct FCollisionResponseParams& ResponseParams, const struct FCollisionObjectQueryParams& ObjectParams, TArray<FOverlapResult>& Results, double CPUTime)
+{
+	if (bSkipCapture || !GCollisionAnalyzerIsRecording || !IsInGameThread())
+	{
+		return;
+	}
+
+	ECAQueryShape::Type QueryShape = ECAQueryShape::Sphere;
+	FVector Dims(0, 0, 0);
+	FQuat Rot = FQuat::Identity;
+	P2UGeomAndRot(PGeom, PGeomPose.q, QueryShape, Dims, Rot);
+
+	TArray<FHitResult> HitResults;
+	for(const FOverlapResult& OverlapResult : Results)
+	{
+		FHitResult NewResult = FHitResult(0.f);
+		NewResult.bBlockingHit = OverlapResult.bBlockingHit;
+		NewResult.Actor = OverlapResult.Actor;
+		NewResult.Component = OverlapResult.Component;
+		NewResult.Item = OverlapResult.ItemIndex;
+		HitResults.Add(NewResult);
+	}
+
+	TArray<FHitResult> TouchAllResults;
+	// TODO: Fill in 'all results' for overlaps
+
+	FCollisionAnalyzerModule::Get()->CaptureQuery(P2UVector(PGeomPose.p), FVector(0,0,0), Rot, ECAQueryType::GeomOverlap, QueryShape, QueryMode, Dims, TraceChannel, Params, ResponseParams, ObjectParams, HitResults, TouchAllResults, CPUTime);
 }
 
 #define STARTQUERYTIMER() double StartTime = FPlatformTime::Seconds()
-#define CAPTUREGEOMSWEEP(World, Start, End, Rot, QueryType, PGeom, TraceChannel, Params, ResponseParam, ObjectParam, Results) if (GCollisionAnalyzerIsRecording) { CaptureGeomSweep(World, Start, End, Rot, QueryType, PGeom, TraceChannel, Params, ResponseParam, ObjectParam, Results, FPlatformTime::Seconds() - StartTime); }
-#define CAPTURERAYCAST(World, Start, End, QueryType, TraceChannel, Params, ResponseParam, ObjectParam, Results)	if (GCollisionAnalyzerIsRecording) { CaptureRaycast(World, Start, End, QueryType, TraceChannel, Params, ResponseParam, ObjectParam, Results, FPlatformTime::Seconds() - StartTime); }
+#define CAPTUREGEOMSWEEP(World, Start, End, Rot, QueryMode, PGeom, TraceChannel, Params, ResponseParam, ObjectParam, Results) if (GCollisionAnalyzerIsRecording) { CaptureGeomSweep(World, Start, End, Rot, QueryMode, PGeom, TraceChannel, Params, ResponseParam, ObjectParam, Results, FPlatformTime::Seconds() - StartTime); }
+#define CAPTURERAYCAST(World, Start, End, QueryMode, TraceChannel, Params, ResponseParam, ObjectParam, Results)	if (GCollisionAnalyzerIsRecording) { CaptureRaycast(World, Start, End, QueryMode, TraceChannel, Params, ResponseParam, ObjectParam, Results, FPlatformTime::Seconds() - StartTime); }
+#define CAPTUREGEOMOVERLAP(World, PGeom, PGeomPose, QueryMode, TraceChannel, Params, ResponseParams, ObjectParams, Results)	if (GCollisionAnalyzerIsRecording) { CaptureOverlap(World, PGeom, PGeomPose, QueryMode, TraceChannel, Params, ResponseParams, ObjectParams, Results, FPlatformTime::Seconds() - StartTime); }
+
 
 #else
 
 #define STARTQUERYTIMER() 
 #define CAPTUREGEOMSWEEP(...)
 #define CAPTURERAYCAST(...)
+#define CAPTUREGEOMOVERLAP(...)
 
 #endif // ENABLE_COLLISION_ANALYZER
 
@@ -871,7 +913,7 @@ bool RaycastTest(const UWorld* World, const FVector Start, const FVector End, EC
 		DrawLineTraces(World, Start, End, Hits, DebugLineLifetime);
 	}
 #endif //!(UE_BUILD_SHIPPING || UE_BUILD_TEST)
-	CAPTURERAYCAST(World, Start, End, ECAQueryType::Test, TraceChannel, Params, ResponseParams, ObjectParams, Hits);
+	CAPTURERAYCAST(World, Start, End, ECAQueryMode::Test, TraceChannel, Params, ResponseParams, ObjectParams, Hits);
 
 	return bHaveBlockingHit;
 }
@@ -988,7 +1030,7 @@ bool RaycastSingle(const UWorld* World, struct FHitResult& OutHit, const FVector
 		{
 			Hits.Add(OutHit);
 		}
-		CAPTURERAYCAST(World, Start, End, ECAQueryType::Single, TraceChannel, Params, ResponseParams, ObjectParams, Hits);
+		CAPTURERAYCAST(World, Start, End, ECAQueryMode::Single, TraceChannel, Params, ResponseParams, ObjectParams, Hits);
 	}
 #endif
 
@@ -1157,7 +1199,7 @@ bool RaycastMulti(const UWorld* World, TArray<struct FHitResult>& OutHits, const
 	}
 #endif //!(UE_BUILD_SHIPPING || UE_BUILD_TEST)
 
-	CAPTURERAYCAST(World, Start, End, ECAQueryType::Multi, TraceChannel, Params, ResponseParams, ObjectParams, OutHits);
+	CAPTURERAYCAST(World, Start, End, ECAQueryMode::Multi, TraceChannel, Params, ResponseParams, ObjectParams, OutHits);
 
 	return bHaveBlockingHit;
 }
@@ -1228,7 +1270,7 @@ bool GeomSweepTest(const UWorld* World, const struct FCollisionShape& CollisionS
 	if (GCollisionAnalyzerIsRecording)
 	{
 		TArray<FHitResult> Hits;
-		CAPTUREGEOMSWEEP(World, Start, End, PGeomRot, ECAQueryType::Test, PGeom, TraceChannel, Params, ResponseParams, ObjectParams, Hits);
+		CAPTUREGEOMSWEEP(World, Start, End, PGeomRot, ECAQueryMode::Test, PGeom, TraceChannel, Params, ResponseParams, ObjectParams, Hits);
 	}
 #endif // ENABLE_COLLISION_ANALYZER
 
@@ -1343,7 +1385,7 @@ bool GeomSweepSingle(const UWorld* World, const struct FCollisionShape& Collisio
 		{
 			Hits.Add(OutHit);
 		}
-		CAPTUREGEOMSWEEP(World, Start, End, PGeomRot, ECAQueryType::Single, PGeom, TraceChannel, Params, ResponseParams, ObjectParams, Hits);
+		CAPTUREGEOMSWEEP(World, Start, End, PGeomRot, ECAQueryMode::Single, PGeom, TraceChannel, Params, ResponseParams, ObjectParams, Hits);
 	}
 #endif
 
@@ -1469,7 +1511,7 @@ bool GeomSweepMulti_PhysX(const UWorld* World, const PxGeometry& PGeom, const Px
 	{
 		TArray<FHitResult> OnlyMyHits(OutHits);
 		OnlyMyHits.RemoveAt(0, InitialHitCount, false); // Remove whatever was there initially.
-		CAPTUREGEOMSWEEP(World, Start, End, PGeomRot, ECAQueryType::Multi, PGeom, TraceChannel, Params, ResponseParams, ObjectParams, OnlyMyHits);
+		CAPTUREGEOMSWEEP(World, Start, End, PGeomRot, ECAQueryMode::Multi, PGeom, TraceChannel, Params, ResponseParams, ObjectParams, OnlyMyHits);
 	}
 #endif // ENABLE_COLLISION_ANALYZER
 
@@ -1524,6 +1566,8 @@ template <EQueryInfo::Type InfoType>
 bool GeomOverlapMultiImp_PhysX(const UWorld* World, const PxGeometry& PGeom, const PxTransform& PGeomPose, TArray<FOverlapResult>& OutOverlaps, ECollisionChannel TraceChannel, const struct FCollisionQueryParams& Params, const struct FCollisionResponseParams& ResponseParams, const struct FCollisionObjectQueryParams& ObjectParams)
 {
 	SCOPE_CYCLE_COUNTER(STAT_Collision_GeomOverlapMultiple);
+	STARTQUERYTIMER();
+
 	bool bHaveBlockingHit = false;
 
 	// overlapMultiple only supports sphere/capsule/box 
@@ -1636,6 +1680,16 @@ bool GeomOverlapMultiImp_PhysX(const UWorld* World, const PxGeometry& PGeom, con
 			}
 #endif //!(UE_BUILD_SHIPPING || UE_BUILD_TEST)
 		}
+
+#if ENABLE_COLLISION_ANALYZER
+		if (GCollisionAnalyzerIsRecording)
+		{
+			// Determine query mode ('single' doesn't really exist for overlaps)
+			ECAQueryMode::Type QueryMode = (InfoType == EQueryInfo::GatherAll)  ? ECAQueryMode::Multi : ECAQueryMode::Test;
+
+			CAPTUREGEOMOVERLAP(World, PGeom, PGeomPose, QueryMode, TraceChannel, Params, ResponseParams, ObjectParams, OutOverlaps);
+		}
+#endif // ENABLE_COLLISION_ANALYZER
 	}
 	else
 	{
