Integrate CL 2704096 from Orion to Main
- Use !IsValid instead of ContainsNan in UpdateKinematicBonesToAnim to validate transform before passing to PhysX
- Change check to ensure in SetKinematicTarget_AssumesLocked

diff --git a/Engine/Source/Runtime/Engine/Private/PhysicsEngine/PhysAnim.cpp b/Engine/Source/Runtime/Engine/Private/PhysicsEngine/PhysAnim.cpp
--- a/Engine/Source/Runtime/Engine/Private/PhysicsEngine/PhysAnim.cpp
+++ b/Engine/Source/Runtime/Engine/Private/PhysicsEngine/PhysAnim.cpp
@@ -565,7 +565,7 @@ void USkeletalMeshComponent::UpdateKinematicBonesToAnim(const TArray<FTransform>
 #if WITH_PHYSX
 						// update bone transform to world
 						const FTransform BoneTransform = InSpaceBases[BoneIndex] * CurrentLocalToWorld;
-						if(BoneTransform.ContainsNaN())
+						if(!BoneTransform.IsValid())
 						{
 							const FName BodyName = PhysicsAsset->BodySetup[i]->BoneName;
 							UE_LOG(LogPhysics, Warning, TEXT("UpdateKinematicBonesToAnim: Trying to set transform with bad data %s on PhysicsAsset '%s' in SkeletalMesh '%s' for bone '%s'"), *BoneTransform.ToHumanReadableString(), *PhysicsAsset->GetName(), *SkeletalMesh->GetName(), *BodyName.ToString());
diff --git a/Engine/Source/Runtime/Engine/Private/PhysicsEngine/PhysScene.cpp b/Engine/Source/Runtime/Engine/Private/PhysicsEngine/PhysScene.cpp
--- a/Engine/Source/Runtime/Engine/Private/PhysicsEngine/PhysScene.cpp
+++ b/Engine/Source/Runtime/Engine/Private/PhysicsEngine/PhysScene.cpp
@@ -275,7 +275,7 @@ void FPhysScene::SetKinematicTarget_AssumesLocked(FBodyInstance* BodyInstance, c
 #endif
 		{
 			const PxTransform PNewPose = U2PTransform(TargetTransform);
-			check(PNewPose.isValid());
+			ensure(PNewPose.isValid());
 			PRigidDynamic->setKinematicTarget(PNewPose);
 		}
 	}
