Optimization to UpdateKinematicBonesToAnim, to avoid querying body flags

diff --git a/Engine/Source/Runtime/Engine/Private/PhysicsEngine/PhysAnim.cpp b/Engine/Source/Runtime/Engine/Private/PhysicsEngine/PhysAnim.cpp
--- a/Engine/Source/Runtime/Engine/Private/PhysicsEngine/PhysAnim.cpp
+++ b/Engine/Source/Runtime/Engine/Private/PhysicsEngine/PhysAnim.cpp
@@ -573,8 +573,7 @@ void USkeletalMeshComponent::UpdateKinematicBonesToAnim(const TArray<FTransform>
 						}					
 
 						// If kinematic and not teleporting, set kinematic target
-						PxRigidDynamic* PRigidDynamic = BodyInst->GetPxRigidDynamic_AssumesLocked();
-						if (!IsRigidBodyNonKinematic_AssumesLocked(PRigidDynamic) && !bTeleport)
+						if (!BodyInst->IsInstanceSimulatingPhysics() && !bTeleport)
 						{
 							PhysScene->SetKinematicTarget_AssumesLocked(BodyInst, BoneTransform, true);
 						}
@@ -583,7 +582,8 @@ void USkeletalMeshComponent::UpdateKinematicBonesToAnim(const TArray<FTransform>
 						{
 							const PxTransform PNewPose = U2PTransform(BoneTransform);
 							ensure(PNewPose.isValid());
-							PRigidDynamic->setGlobalPose(PNewPose);
+							PxRigidActor* RigidActor = BodyInst->RigidActorSync ? BodyInst->RigidActorSync : BodyInst->RigidActorAsync;
+							RigidActor->setGlobalPose(PNewPose);
 						}
 #endif
 
