UE4 - Fixed bug with SetTemplate on paritcle system not recomputing parallel elligibility.

diff --git a/Engine/Source/Runtime/Engine/Private/Particles/ParticleComponents.cpp b/Engine/Source/Runtime/Engine/Private/Particles/ParticleComponents.cpp
--- a/Engine/Source/Runtime/Engine/Private/Particles/ParticleComponents.cpp
+++ b/Engine/Source/Runtime/Engine/Private/Particles/ParticleComponents.cpp
@@ -4370,6 +4370,8 @@ void UParticleSystemComponent::SetTemplate(class UParticleSystem* NewTemplate)
 		bool bResetInstances = false;
 		if (NewTemplate != Template)
 		{
+			bIsElligibleForAsyncTick = false;
+			bIsElligibleForAsyncTickComputed = false;
 			bResetInstances = true;
 		}
 		if (bIsTemplate == false)
