Use GetPxRigidActor_AssumesLocked instead of direct accessing BodyInstance

diff --git a/Engine/Source/Runtime/Engine/Private/PhysicsEngine/PhysAnim.cpp b/Engine/Source/Runtime/Engine/Private/PhysicsEngine/PhysAnim.cpp
--- a/Engine/Source/Runtime/Engine/Private/PhysicsEngine/PhysAnim.cpp
+++ b/Engine/Source/Runtime/Engine/Private/PhysicsEngine/PhysAnim.cpp
@@ -582,7 +582,7 @@ void USkeletalMeshComponent::UpdateKinematicBonesToAnim(const TArray<FTransform>
 						{
 							const PxTransform PNewPose = U2PTransform(BoneTransform);
 							ensure(PNewPose.isValid());
-							PxRigidActor* RigidActor = BodyInst->RigidActorSync ? BodyInst->RigidActorSync : BodyInst->RigidActorAsync;
+							PxRigidActor* RigidActor = BodyInst->GetPxRigidActor_AssumesLocked(); // This should never fail because IsValidBodyInstance() passed above
 							RigidActor->setGlobalPose(PNewPose);
 						}
 #endif
