add back message log that got removed during merge

diff --git a/Engine/Source/Runtime/Engine/Private/Animation/Skeleton.cpp b/Engine/Source/Runtime/Engine/Private/Animation/Skeleton.cpp
--- a/Engine/Source/Runtime/Engine/Private/Animation/Skeleton.cpp
+++ b/Engine/Source/Runtime/Engine/Private/Animation/Skeleton.cpp
@@ -393,6 +393,17 @@ int32 USkeleton::BuildLinkup(const USkeletalMesh* InSkelMesh)
 				bDismissedMessage = true;
 			}
 
+			static FName NAME_LoadErrors("LoadErrors");
+			FMessageLog LoadErrors(NAME_LoadErrors);
+
+			TSharedRef<FTokenizedMessage> Message = LoadErrors.Error();
+			Message->AddToken(FTextToken::Create(LOCTEXT("SkeletonBuildLinkupMissingBones1", "The Skeleton ")));
+			Message->AddToken(FAssetNameToken::Create(GetPathName(), FText::FromString( GetNameSafe(this) ) ));
+			Message->AddToken(FTextToken::Create(LOCTEXT("SkeletonBuildLinkupMissingBones2", " is missing bones that SkeletalMesh ")));
+			Message->AddToken(FAssetNameToken::Create(InSkelMesh->GetPathName(), FText::FromString( GetNameSafe(InSkelMesh) )));
+			Message->AddToken(FTextToken::Create(LOCTEXT("SkeletonBuildLinkupMissingBones3", "  needs. They will be added now. Please save the Skeleton!")));
+			LoadErrors.Open();
+
 			// Re-add all SkelMesh bones to the Skeleton.
 			MergeAllBonesToBoneTree(InSkelMesh);
 
