Converted to message log for resave message

Merging

//depot/UE4-Orion/Engine/Source/Runtime/Engine/Private/Animation/AnimSequence.cpp

to //depot/UE4-Orion/Engine/Source/Runtime/Engine/Private/Animation/AnimSequence.cpp

diff --git a/Engine/Source/Runtime/Engine/Private/Animation/AnimSequence.cpp b/Engine/Source/Runtime/Engine/Private/Animation/AnimSequence.cpp
--- a/Engine/Source/Runtime/Engine/Private/Animation/AnimSequence.cpp
+++ b/Engine/Source/Runtime/Engine/Private/Animation/AnimSequence.cpp
@@ -16,8 +16,10 @@
 #include "Animation/Rig.h"
 #include "EditorFramework/AssetImportData.h"
 #include "Animation/AnimStats.h"
+#include "MessageLog.h"
 
 #define USE_SLERP 0
+#define LOCTEXT_NAMESPACE "AnimSequence"
 
 /////////////////////////////////////////////////////
 // FRawAnimSequenceTrackNativeDeprecated
@@ -535,6 +537,15 @@ void UAnimSequence::VerifyTrackMap(USkeleton* MySkeleton)
 	{
 		UE_LOG(LogAnimation, Warning, TEXT("RESAVE ANIMATION NEEDED(%s): Fixing track names."), *GetName());
 
+		static FName NAME_LoadErrors("LoadErrors");
+		FMessageLog LoadErrors(NAME_LoadErrors);
+
+		TSharedRef<FTokenizedMessage> Message = LoadErrors.Warning();
+		Message->AddToken(FTextToken::Create(LOCTEXT("AnimationNeedsResave1", "The Animation ")));
+		Message->AddToken(FAssetNameToken::Create(GetPathName(), FText::FromString(GetNameSafe(this))));
+		Message->AddToken(FTextToken::Create(LOCTEXT("AnimationNeedsResave2", " needs resave.")));
+		LoadErrors.Open();
+
 		const TArray<FBoneNode>& BoneTree = UseSkeleton->GetBoneTree();
 		AnimationTrackNames.Empty();
 		AnimationTrackNames.AddUninitialized(TrackToSkeletonMapTable.Num());
@@ -571,6 +582,16 @@ void UAnimSequence::VerifyTrackMap(USkeleton* MySkeleton)
 			{
 				UE_LOG(LogAnimation, Warning, TEXT("RESAVE ANIMATION NEEDED(%s): Fixing track index."), *GetName());
 
+				static FName NAME_LoadErrors("LoadErrors");
+				FMessageLog LoadErrors(NAME_LoadErrors);
+
+				TSharedRef<FTokenizedMessage> Message = LoadErrors.Warning();
+				Message->AddToken(FTextToken::Create(LOCTEXT("AnimationNeedsResave1", "The Animation ")));
+				Message->AddToken(FAssetNameToken::Create(GetPathName(), FText::FromString(GetNameSafe(this))));
+				Message->AddToken(FTextToken::Create(LOCTEXT("AnimationNeedsResave2", " needs resave.")));
+
+				LoadErrors.Open();
+
 				const TArray<FBoneNode>& BoneTree = UseSkeleton->GetBoneTree();
 				for(int32 I=NumTracks-1; I>=0; --I)
 				{
