Merging

//depot/UE4-Orion/Engine/Source/Runtime/Engine/Private/Animation/AnimSequence.cpp

to //depot/UE4/Engine/Source/Runtime/Engine/Private/Animation/AnimSequence.cpp

diff --git a/Engine/Source/Runtime/Engine/Private/Animation/AnimSequence.cpp b/Engine/Source/Runtime/Engine/Private/Animation/AnimSequence.cpp
--- a/Engine/Source/Runtime/Engine/Private/Animation/AnimSequence.cpp
+++ b/Engine/Source/Runtime/Engine/Private/Animation/AnimSequence.cpp
@@ -586,7 +586,7 @@ void UAnimSequence::VerifyTrackMap(USkeleton* MySkeleton)
 				FMessageLog LoadErrors(NAME_LoadErrors);
 
 				TSharedRef<FTokenizedMessage> Message = LoadErrors.Warning();
-				Message->AddToken(FTextToken::Create(LOCTEXT("AnimationNeedsResave1", "The Animation ")));
+				Message->AddToken(FTextToken::Create(LOCTEXT("AnimationNeedsResave1", "Skeleton has been changed. The Animation ")));
 				Message->AddToken(FAssetNameToken::Create(GetPathName(), FText::FromString(GetNameSafe(this))));
 				Message->AddToken(FTextToken::Create(LOCTEXT("AnimationNeedsResave2", " needs resave.")));
 
