Fix incorrect filter for valid ReverbEffect in UGameplayStatics::ActivateReverbEffect

diff --git a/Engine/Source/Runtime/Engine/Private/GameplayStatics.cpp b/Engine/Source/Runtime/Engine/Private/GameplayStatics.cpp
--- a/Engine/Source/Runtime/Engine/Private/GameplayStatics.cpp
+++ b/Engine/Source/Runtime/Engine/Private/GameplayStatics.cpp
@@ -1069,7 +1069,7 @@ void UGameplayStatics::ClearSoundMixModifiers(UObject* WorldContextObject)
 
 void UGameplayStatics::ActivateReverbEffect(UObject* WorldContextObject, class UReverbEffect* ReverbEffect, FName TagName, float Priority, float Volume, float FadeTime)
 {
-	if (ReverbEffect || !GEngine || !GEngine->UseSound())
+	if (ReverbEffect == nullptr || !GEngine || !GEngine->UseSound())
 	{
 		return;
 	}
