Merging

//depot/UE4-Orion/Engine/Source/Runtime/Engine/Private/Animation/AnimSequence.cpp

to //depot/UE4-Orion/Engine/Source/Runtime/Engine/Private/Animation/AnimSequence.cpp

diff --git a/Engine/Source/Runtime/Engine/Private/Animation/AnimSequence.cpp b/Engine/Source/Runtime/Engine/Private/Animation/AnimSequence.cpp
--- a/Engine/Source/Runtime/Engine/Private/Animation/AnimSequence.cpp
+++ b/Engine/Source/Runtime/Engine/Private/Animation/AnimSequence.cpp
@@ -4263,4 +4263,4 @@ void GatherAnimSequenceStats(FOutputDevice& Ar)
 	Ar.Logf( TEXT("Approx Key Data Memory: %i bytes"), GlobalApproxKeyDataSize);
 }
 #endif
-
+#undef LOCTEXT_NAMESPACE 
\ No newline at end of file
