Disabled reversed indices while working on a fix for bad material indices.

diff --git a/Engine/Source/Runtime/Engine/Private/StaticMeshRender.cpp b/Engine/Source/Runtime/Engine/Private/StaticMeshRender.cpp
--- a/Engine/Source/Runtime/Engine/Private/StaticMeshRender.cpp
+++ b/Engine/Source/Runtime/Engine/Private/StaticMeshRender.cpp
@@ -21,7 +21,7 @@
 static bool GUseShadowIndexBuffer = true;
 
 /** If true, reversed index buffer are used for mesh with negative transform determinants. */
-static bool GUseReversedIndexBuffer = true;
+static bool GUseReversedIndexBuffer = false;
 
 static void ToggleShadowIndexBuffers()
 {
