Merging using Orion->UE4

CL# 2704181
FURL comment fixes.

diff --git a/Engine/Source/Runtime/Engine/Classes/Engine/EngineBaseTypes.h b/Engine/Source/Runtime/Engine/Classes/Engine/EngineBaseTypes.h
--- a/Engine/Source/Runtime/Engine/Classes/Engine/EngineBaseTypes.h
+++ b/Engine/Source/Runtime/Engine/Classes/Engine/EngineBaseTypes.h
@@ -617,26 +617,33 @@ struct ENGINE_API FURL
 {
 	GENERATED_USTRUCT_BODY()
 
+	// Protocol, i.e. "unreal" or "http".
 	UPROPERTY()
-	FString Protocol;    // Protocol, i.e. "unreal" or "http".
+	FString Protocol;
 
+	// Optional hostname, i.e. "204.157.115.40" or "unreal.epicgames.com", blank if local.
 	UPROPERTY()
-	FString Host;    // Optional hostname, i.e. "204.157.115.40" or "unreal.epicgames.com", blank if local.
+	FString Host;
 
+	// Optional host port.
 	UPROPERTY()
-	int32 Port;    // Optional host port.
+	int32 Port;
 
+	// Map name, i.e. "SkyCity", default is "Entry".
 	UPROPERTY()
-	FString Map;    // Map name, i.e. "SkyCity", default is "Entry".
+	FString Map;
 
+	// Optional place to download Map if client does not possess it
 	UPROPERTY()
-	FString RedirectURL;    // Optional place to download Map if client does not possess it
+	FString RedirectURL;
 
+	// Options.
 	UPROPERTY()
-	TArray<FString> Op;    // Options.
+	TArray<FString> Op;
 
+	// Portal to enter through, default is "".
 	UPROPERTY()
-	FString Portal;    // Portal to enter through, default is "".
+	FString Portal;
 
 	UPROPERTY()
 	int32 Valid;
@@ -645,11 +652,21 @@ struct ENGINE_API FURL
 	static FUrlConfig UrlConfig;
 	static bool bDefaultsInitialized;
 
-	// Constructors.
-	/* FURL() prevent default from being generated */
+	/**
+	 * Prevent default from being generated.
+	 */
 	explicit FURL( ENoInit ) { }
+
+	/**
+	 * Construct a purely default, local URL from an optional filename.
+	 */
 	FURL( const TCHAR* Filename=nullptr );
+
+	/**
+	 * Construct a URL from text and an optional relative base.
+	 */
 	FURL( FURL* Base, const TCHAR* TextURL, ETravelType Type );
+
 	static void StaticInit();
 	static void StaticExit();
 
@@ -660,19 +677,68 @@ struct ENGINE_API FURL
 	 */
 	static void FilterURLString( FString& Str );
 
-	// Functions.
+	/**
+	 * Returns whether this URL corresponds to an internal object, i.e. an Unreal
+	 * level which this app can try to connect to locally or on the net. If this
+	 * is false, the URL refers to an object that a remote application like Internet
+	 * Explorer can execute.
+	 */
 	bool IsInternal() const;
+
+	/**
+	 * Returns whether this URL corresponds to an internal object on this local 
+	 * process. In this case, no Internet use is necessary.
+	 */
 	bool IsLocalInternal() const;
+
+	/**
+	 * Tests if the URL contains an option string.
+	 */
 	bool HasOption( const TCHAR* Test ) const;
+
+	/**
+	 * Returns the value associated with an option.
+	 *
+	 * @param Match The name of the option to get.
+	 * @param Default The value to return if the option wasn't found.
+	 *
+	 * @return The value of the named option, or Default if the option wasn't found.
+	 */
 	const TCHAR* GetOption( const TCHAR* Match, const TCHAR* Default ) const;
+
+	/**
+	 * Load URL from config.
+	 */
 	void LoadURLConfig( const TCHAR* Section, const FString& Filename=GGameIni );
+
+	/**
+	 * Save URL to config.
+	 */
 	void SaveURLConfig( const TCHAR* Section, const TCHAR* Item, const FString& Filename=GGameIni ) const;
+
+	/**
+	 * Add a unique option to the URL, replacing any existing one.
+	 */
 	void AddOption( const TCHAR* Str );
+
+	/**
+	 * Remove an option from the URL
+	 */
 	void RemoveOption( const TCHAR* Key, const TCHAR* Section = nullptr, const FString& Filename = GGameIni);
+
+	/**
+	 * Convert this URL to text.
+	 */
 	FString ToString( bool FullyQualified=0 ) const;
+
+	/**
+	 * Serializes a FURL to or from an archive.
+	 */
 	ENGINE_API friend FArchive& operator<<( FArchive& Ar, FURL& U );
 
-	// Operators.
+	/**
+	 * Compare two URLs to see if they refer to the same exact thing.
+	 */
 	bool operator==( const FURL& Other ) const;
 };
 
diff --git a/Engine/Source/Runtime/Engine/Private/URL.cpp b/Engine/Source/Runtime/Engine/Private/URL.cpp
--- a/Engine/Source/Runtime/Engine/Private/URL.cpp
+++ b/Engine/Source/Runtime/Engine/Private/URL.cpp
@@ -63,9 +63,6 @@ static bool ValidNetChar( const TCHAR* c )
 	Constructors.
 -----------------------------------------------------------------------------*/
 
-//
-// Constuct a purely default, local URL from an optional filename.
-//
 FURL::FURL( const TCHAR* LocalFilename )
 :	Protocol	( UrlConfig.DefaultProtocol )
 ,	Host		( UrlConfig.DefaultHost )
@@ -120,9 +117,6 @@ void FURL::FilterURLString( FString& Str )
 
 
 
-//
-// Construct a URL from text and an optional relative base.
-//
 FURL::FURL( FURL* Base, const TCHAR* TextURL, ETravelType Type )
 :	Protocol	( UrlConfig.DefaultProtocol )
 ,	Host		( UrlConfig.DefaultHost )
@@ -469,29 +463,16 @@ FString FURL::ToString (bool FullyQualified) const
 	Informational.
 -----------------------------------------------------------------------------*/
 
-//
-// Return whether this URL corresponds to an internal object, i.e. an Unreal
-// level which this app can try to connect to locally or on the net. If this
-// is false, the URL refers to an object that a remote application like Internet
-// Explorer can execute.
-//
 bool FURL::IsInternal() const
 {
 	return Protocol == UrlConfig.DefaultProtocol;
 }
 
-//
-// Return whether this URL corresponds to an internal object on this local 
-// process. In this case, no Internet use is necessary.
-//
 bool FURL::IsLocalInternal() const
 {
 	return IsInternal() && Host.Len()==0;
 }
 
-//
-// Add a unique option to the URL, replacing any existing one.
-//
 void FURL::AddOption( const TCHAR* Str )
 {
 	int32 Match = FCString::Strchr(Str, '=') ? (FCString::Strchr(Str, '=') - Str) : FCString::Strlen(Str);
@@ -516,9 +497,6 @@ void FURL::AddOption( const TCHAR* Str )
 }
 
 
-//
-// Remove an option from the URL
-//
 void FURL::RemoveOption( const TCHAR* Key, const TCHAR* Section, const FString& Filename )
 {
 	if ( !Key )
@@ -540,9 +518,6 @@ void FURL::RemoveOption( const TCHAR* Key, const TCHAR* Section, const FString&
 	}
 }
 
-//
-// Load URL from config.
-//
 void FURL::LoadURLConfig( const TCHAR* Section, const FString& Filename )
 {
 	TArray<FString> Options;
@@ -553,9 +528,6 @@ void FURL::LoadURLConfig( const TCHAR* Section, const FString& Filename )
 	}
 }
 
-//
-// Save URL to config.
-//
 void FURL::SaveURLConfig( const TCHAR* Section, const TCHAR* Item, const FString& Filename ) const
 {
 	for( int32 i=0; i<Op.Num(); i++ )
@@ -572,9 +544,6 @@ void FURL::SaveURLConfig( const TCHAR* Section, const TCHAR* Item, const FString
 	}
 }
 
-//
-// See if the URL contains an option string.
-//
 bool FURL::HasOption( const TCHAR* Test ) const
 {
 	return GetOption( Test, NULL ) != NULL;
@@ -606,9 +575,6 @@ const TCHAR* FURL::GetOption( const TCHAR* Match, const TCHAR* Default ) const
 	Comparing.
 -----------------------------------------------------------------------------*/
 
-//
-// Compare two URL's to see if they refer to the same exact thing.
-//
 bool FURL::operator==( const FURL& Other ) const
 {
 	if
