Increment FFastArraySerializer replication keys during reading as well as writing, to support replicating again on clients for replay recording.

--------
Integrated using branch UE4-Orion-To-UE4 of change#2704893 by [REDACTED_PERSON] on 2015/09/24 14:48:25.

diff --git a/Engine/Source/Runtime/Engine/Classes/Engine/NetSerialization.h b/Engine/Source/Runtime/Engine/Classes/Engine/NetSerialization.h
--- a/Engine/Source/Runtime/Engine/Classes/Engine/NetSerialization.h
+++ b/Engine/Source/Runtime/Engine/Classes/Engine/NetSerialization.h
@@ -677,6 +677,12 @@ bool FFastArraySerializer::FastArrayDeltaSerialize( TArray<Type> &Items, FNetDel
 
 		UE_LOG( LogNetSerialization, Verbose, TEXT( "Read NumChanged: %d NumDeletes: %d." ), NumChanged, NumDeletes );
 
+		// Increment keys so that a client can re-serialize the array if needed, such as for client replay recording
+		if ( NumDeletes > 0 || NumChanged > 0 )
+		{
+			++ArraySerializer.ArrayReplicationKey;
+		}
+
 		TArray<int32> DeleteIndices;
 
 		//---------------
@@ -741,6 +747,7 @@ bool FFastArraySerializer::FastArrayDeltaSerialize( TArray<Type> &Items, FNetDel
 				UE_LOG(LogNetSerialization, Log, TEXT("   Changed. ID: %d -> Idx: %d"), ElementID, *ElementIndexPtr);
 				ElementIndex = *ElementIndexPtr;
 				ThisElement = &Items[ElementIndex];
+				++Items[ElementIndex].ReplicationKey;
 			}
 
 			// Let package map know we want to track and know about any guids that are unmapped during the serialize call
