Fix matinee audio tracks not starting at the correct point when scrubbing

diff --git a/Engine/Source/Runtime/Engine/Private/Interpolation.cpp b/Engine/Source/Runtime/Engine/Private/Interpolation.cpp
--- a/Engine/Source/Runtime/Engine/Private/Interpolation.cpp
+++ b/Engine/Source/Runtime/Engine/Private/Interpolation.cpp
@@ -7908,6 +7908,14 @@ void UInterpTrackSound::UpdateTrack(float NewPosition, UInterpTrackInst* TrInst,
 			}
 			else if (!bTreatAsDialogue || !Actor) // Don't play at all if we had a dialogue actor but they are not available/dead now
 			{
+				float StartTime = NewPosition - SoundTrackKey.Time;
+				if (StartTime <= FApp::GetDeltaTime())
+				{
+					// If the start time is within the past frames delta time, start from the beginning
+					// TODO: Consider dealing with play rate, time dilation, clamping
+					StartTime = 0.f;
+				}
+
 				// If we have a sound playing already (ie. an AudioComponent exists) stop it now.
 				if(SoundInst->PlayAudioComp)
 				{
@@ -7930,7 +7938,7 @@ void UInterpTrackSound::UpdateTrack(float NewPosition, UInterpTrackInst* TrInst,
 					SoundInst->PlayAudioComp->SetVolumeMultiplier(VolumePitchValue.X);
 					SoundInst->PlayAudioComp->SetPitchMultiplier(VolumePitchValue.Y);
 					SoundInst->PlayAudioComp->SubtitlePriority = bSuppressSubtitles ? 0.f : SUBTITLE_PRIORITY_MATINEE;
-					SoundInst->PlayAudioComp->Play((bJump ? NewPosition - SoundTrackKey.Time : 0.f));
+					SoundInst->PlayAudioComp->Play(StartTime);
 				}
 				else
 				{
@@ -7960,7 +7968,7 @@ void UInterpTrackSound::UpdateTrack(float NewPosition, UInterpTrackInst* TrInst,
 						SoundInst->PlayAudioComp->SetVolumeMultiplier(VolumePitchValue.X);
 						SoundInst->PlayAudioComp->SetPitchMultiplier(VolumePitchValue.Y);
 						SoundInst->PlayAudioComp->SubtitlePriority = bSuppressSubtitles ? 0.f : SUBTITLE_PRIORITY_MATINEE;
-						SoundInst->PlayAudioComp->Play((bJump ? NewPosition - SoundTrackKey.Time : 0.f));
+						SoundInst->PlayAudioComp->Play(StartTime);
 					}
 				}
 			}
