Sequencer: Exposed curve keys in details panel

diff --git a/Engine/Source/Runtime/Engine/Classes/Curves/CurveBase.h b/Engine/Source/Runtime/Engine/Classes/Curves/CurveBase.h
--- a/Engine/Source/Runtime/Engine/Classes/Curves/CurveBase.h
+++ b/Engine/Source/Runtime/Engine/Classes/Curves/CurveBase.h
@@ -86,7 +86,6 @@ struct TStructOpsTypeTraits< FKeyHandleMap > : public TStructOpsTypeTraitsBase
 	};
 };
 
-
 // @todo Some heavy refactoring can be done here. Much more stuff can go in this base class.
 /** A curve base class which enables key handles to index lookups */
 USTRUCT()
@@ -189,11 +188,11 @@ struct ENGINE_API FRichCurveKey
 	TEnumAsByte<ERichCurveTangentWeightMode>	TangentWeightMode;
 
 	/** Time at this key */
-	UPROPERTY()
+	UPROPERTY(EditAnywhere, Category="Key")
 	float Time;
 
 	/** Value at this key */
-	UPROPERTY()
+	UPROPERTY(EditAnywhere, Category="Key")
 	float Value;
 
 	/** If RCIM_Cubic, the arriving tangent at this key */
@@ -412,11 +411,11 @@ public:
 	TEnumAsByte<ERichCurveExtrapolation> PostInfinityExtrap;
 
 	/** Sorted array of keys */
-	UPROPERTY()
+	UPROPERTY(EditAnywhere, EditFixedSize, Category="Curve")
 	TArray<FRichCurveKey> Keys;
 
 	/** Default value */
-	UPROPERTY()
+	UPROPERTY(EditAnywhere, Category="Curve")
 	float DefaultValue;
 };
 
@@ -663,10 +662,10 @@ public:
 
 private:
 	/** The keys, ordered by time */
-	UPROPERTY()
+	UPROPERTY(EditAnywhere, Category="Curve")
 	TArray<FIntegralKey> Keys;
 
 	/** Default value */
-	UPROPERTY()
+	UPROPERTY(EditAnywhere, Category="Curve")
 	int32 DefaultValue;
 };
diff --git a/Engine/Source/Runtime/MovieSceneTracks/Public/Sections/MovieSceneAudioSection.h b/Engine/Source/Runtime/MovieSceneTracks/Public/Sections/MovieSceneAudioSection.h
--- a/Engine/Source/Runtime/MovieSceneTracks/Public/Sections/MovieSceneAudioSection.h
+++ b/Engine/Source/Runtime/MovieSceneTracks/Public/Sections/MovieSceneAudioSection.h
@@ -56,13 +56,17 @@ public:
 		return Position >= AudioRange.GetLowerBoundValue() && Position <= AudioRange.GetUpperBoundValue();
 	}
 
-	/** MovieSceneSection interface */
+public:
+
+	// MovieSceneSection interface
+
 	virtual void MoveSection( float DeltaPosition, TSet<FKeyHandle>& KeyHandles ) override;
 	virtual void DilateSection( float DilationFactor, float Origin, TSet<FKeyHandle>& KeyHandles ) override;
 	virtual void GetKeyHandles(TSet<FKeyHandle>& KeyHandles) const override;
 	virtual void GetSnapTimes(TArray<float>& OutSnapTimes, bool bGetSectionBorders) const override;
 
 private:
+
 	/** The sound cue or wave that this section plays */
 	UPROPERTY(EditAnywhere, Category="Audio")
 	USoundBase* Sound;
diff --git a/Engine/Source/Runtime/MovieSceneTracks/Public/Sections/MovieSceneByteSection.h b/Engine/Source/Runtime/MovieSceneTracks/Public/Sections/MovieSceneByteSection.h
--- a/Engine/Source/Runtime/MovieSceneTracks/Public/Sections/MovieSceneByteSection.h
+++ b/Engine/Source/Runtime/MovieSceneTracks/Public/Sections/MovieSceneByteSection.h
@@ -52,6 +52,6 @@ private:
 	/** Ordered curve data */
 	// @todo Sequencer This could be optimized by packing the bytes separately
 	// but that may not be worth the effort
-	UPROPERTY()
+	UPROPERTY(EditAnywhere, Category="Curve")
 	FIntegralCurve ByteCurve;
 };
diff --git a/Engine/Source/Runtime/MovieSceneTracks/Public/Sections/MovieSceneFloatSection.h b/Engine/Source/Runtime/MovieSceneTracks/Public/Sections/MovieSceneFloatSection.h
--- a/Engine/Source/Runtime/MovieSceneTracks/Public/Sections/MovieSceneFloatSection.h
+++ b/Engine/Source/Runtime/MovieSceneTracks/Public/Sections/MovieSceneFloatSection.h
@@ -51,6 +51,6 @@ public:
 	FRichCurve& GetFloatCurve() { return FloatCurve; }
 private:
 	/** Curve data */
-	UPROPERTY()
+	UPROPERTY(EditAnywhere, Category="Curve")
 	FRichCurve FloatCurve;
 };
\ No newline at end of file
diff --git a/Engine/Source/Runtime/MovieSceneTracks/Public/Sections/MovieSceneParticleSection.h b/Engine/Source/Runtime/MovieSceneTracks/Public/Sections/MovieSceneParticleSection.h
--- a/Engine/Source/Runtime/MovieSceneTracks/Public/Sections/MovieSceneParticleSection.h
+++ b/Engine/Source/Runtime/MovieSceneTracks/Public/Sections/MovieSceneParticleSection.h
@@ -5,8 +5,9 @@
 #include "MovieSceneSection.h"
 #include "MovieSceneParticleSection.generated.h"
 
+
 /**
- * The kind of key this particle section is
+ * The kind of key this particle section is.
  */
 UENUM()
 namespace EParticleKey
@@ -18,14 +19,18 @@ namespace EParticleKey
 	};
 }
 
+
 /**
- * Particle section, for particle toggling and triggering
+ * Particle section, for particle toggling and triggering.
  */
-UCLASS( MinimalAPI )
-class UMovieSceneParticleSection : public UMovieSceneSection
+UCLASS(MinimalAPI)
+class UMovieSceneParticleSection
+	: public UMovieSceneSection
 {
 	GENERATED_UCLASS_BODY()
+
 public:
+
 	/** Sets this section's key type*/
 	void SetKeyType(EParticleKey::Type InKeyType) {KeyType = InKeyType;}
 
@@ -33,7 +38,8 @@ public:
 	EParticleKey::Type GetKeyType() {return KeyType;}
 	
 private:
+
 	/** The way this particle key will operate */
-	UPROPERTY()
+	UPROPERTY(EditAnywhere, Category="Particles")
 	TEnumAsByte<EParticleKey::Type> KeyType;
 };
