[UE-20764] Fix for a potential infinite loop case in post-transform navmesh updates.

[from Mieszko]
Changed the way Navigation System gathers actors attached to actor being currently processed #UE4

- Taking advantage of GetAttachedActors() API
- Added a dedicated function that created a flat list of attached actrors and simply processes one by one.

diff --git a/Engine/Source/Runtime/Engine/Classes/AI/Navigation/NavigationSystem.h b/Engine/Source/Runtime/Engine/Classes/AI/Navigation/NavigationSystem.h
--- a/Engine/Source/Runtime/Engine/Classes/AI/Navigation/NavigationSystem.h
+++ b/Engine/Source/Runtime/Engine/Classes/AI/Navigation/NavigationSystem.h
@@ -591,10 +591,15 @@ public:
 	static void UpdateNavOctree(AActor* Actor);
 	static void UpdateNavOctree(UActorComponent* Comp);
 	/** update all navoctree entries for actor and its components */
-	static void UpdateNavOctreeAll(AActor* Actor);
+	static void UpdateNavOctreeAll(AActor* Actor, bool bUpdateAttachedActors = true);
 	/** update all navoctree entries for actor and its non scene components after root movement */
 	static void UpdateNavOctreeAfterMove(USceneComponent* Comp);
 
+protected:
+	/** updates navoctree information on actors attached to RootActor */
+	static void UpdateAttachedActorsInNavOctree(AActor& RootActor);
+
+public:
 	/** removes all navoctree entries for actor and its components */
 	static void ClearNavOctreeAll(AActor* Actor);
 	/** updates bounds of all components implementing INavRelevantInterface */
diff --git a/Engine/Source/Runtime/Engine/Private/AI/Navigation/NavigationSystem.cpp b/Engine/Source/Runtime/Engine/Private/AI/Navigation/NavigationSystem.cpp
--- a/Engine/Source/Runtime/Engine/Private/AI/Navigation/NavigationSystem.cpp
+++ b/Engine/Source/Runtime/Engine/Private/AI/Navigation/NavigationSystem.cpp
@@ -2511,7 +2511,7 @@ void UNavigationSystem::UpdateNavOctree(UActorComponent* Comp)
 	}
 }
 
-void UNavigationSystem::UpdateNavOctreeAll(AActor* Actor)
+void UNavigationSystem::UpdateNavOctreeAll(AActor* Actor, bool bUpdateAttachedActors)
 {
 	if (Actor)
 	{
@@ -2525,15 +2525,9 @@ void UNavigationSystem::UpdateNavOctreeAll(AActor* Actor)
 			UpdateNavOctree(Components[ComponentIndex]);
 		}
 
-		if (Actor->GetRootComponent())
+		if (bUpdateAttachedActors)
 		{
-			for (int32 RootChildIndex = 0; RootChildIndex < Actor->GetRootComponent()->AttachChildren.Num(); RootChildIndex++)
-			{
-				if (Actor->GetRootComponent()->AttachChildren[RootChildIndex] && Actor->GetRootComponent()->AttachChildren[RootChildIndex]->GetOuter() != Actor)
-				{
-					UpdateNavOctreeAll(Cast<AActor>(Actor->GetRootComponent()->AttachChildren[RootChildIndex]->GetOuter()));
-				}
-			}
+			UpdateAttachedActorsInNavOctree(*Actor);
 		}
 	}
 }
@@ -2557,14 +2551,34 @@ void UNavigationSystem::UpdateNavOctreeAfterMove(USceneComponent* Comp)
 			}
 		}
 
-		for (int32 RootChildIndex = 0; RootChildIndex < Comp->AttachChildren.Num(); RootChildIndex++)
+		UpdateAttachedActorsInNavOctree(*OwnerActor);
+	}
+}
+
+void UNavigationSystem::UpdateAttachedActorsInNavOctree(AActor& RootActor)
+{
+	TArray<AActor*> UniqueAttachedActors;
+	UniqueAttachedActors.Add(&RootActor);
+
+	TArray<AActor*> TempAttachedActors;
+	for (int32 ActorIndex = 0; ActorIndex < UniqueAttachedActors.Num(); ++ActorIndex)
+	{
+		check(UniqueAttachedActors[ActorIndex]);
+		// find all attached actors
+		UniqueAttachedActors[ActorIndex]->GetAttachedActors(TempAttachedActors);
+		
+		for (int32 AttachmentIndex = 0; AttachmentIndex < TempAttachedActors.Num(); ++AttachmentIndex)
 		{
-			if (Comp->AttachChildren[RootChildIndex] && Comp->AttachChildren[RootChildIndex]->GetOuter() != OwnerActor)
-			{
-				UpdateNavOctreeAll(Cast<AActor>(Comp->AttachChildren[RootChildIndex]->GetOuter()));
-			}
+			// and store the ones we don't know about yet
+			UniqueAttachedActors.AddUnique(TempAttachedActors[AttachmentIndex]);
 		}
 	}
+	
+	// skipping the first item since that's the root, and we just care about the attached actors
+	for (int32 ActorIndex = 1; ActorIndex < UniqueAttachedActors.Num(); ++ActorIndex)
+	{
+		UpdateNavOctreeAll(UniqueAttachedActors[ActorIndex], /*bUpdateAttachedActors = */false);
+	}
 }
 
 void UNavigationSystem::UpdateNavOctreeBounds(AActor* Actor)
diff --git a/Engine/Source/Runtime/Engine/Private/ActorEditor.cpp b/Engine/Source/Runtime/Engine/Private/ActorEditor.cpp
--- a/Engine/Source/Runtime/Engine/Private/ActorEditor.cpp
+++ b/Engine/Source/Runtime/Engine/Private/ActorEditor.cpp
@@ -166,16 +166,16 @@ void AActor::PostEditMove(bool bFinished)
 		// update actor and all its components in navigation system after finishing move
 		// USceneComponent::UpdateNavigationData works only in game world
 		UNavigationSystem::UpdateNavOctreeBounds(this);
-		UNavigationSystem::UpdateNavOctreeAll(this);
 
 		TArray<AActor*> ParentedActors;
 		GetAttachedActors(ParentedActors);
-
 		for (int32 Idx = 0; Idx < ParentedActors.Num(); Idx++)
 		{
 			UNavigationSystem::UpdateNavOctreeBounds(ParentedActors[Idx]);
-			UNavigationSystem::UpdateNavOctreeAll(ParentedActors[Idx]);
 		}
+
+		// not doing manual update of all attached actors since UpdateNavOctreeAll should take care of it
+		UNavigationSystem::UpdateNavOctreeAll(this);
 	}
 }
 
