Engine: CurveBase cleanup pass

diff --git a/Engine/Source/Runtime/Engine/Classes/Curves/CurveBase.h b/Engine/Source/Runtime/Engine/Classes/Curves/CurveBase.h
--- a/Engine/Source/Runtime/Engine/Classes/Curves/CurveBase.h
+++ b/Engine/Source/Runtime/Engine/Classes/Curves/CurveBase.h
@@ -4,25 +4,28 @@
 
 #include "CurveBase.generated.h"
 
-//////////////////////////////////////////////////////////////////////////
-// FIndexedCurve
 
-/** Key handles are used to keep a handle to a key. They are completely transient. */
+/* Basic types
+ *****************************************************************************/
+
+/**
+ * Key handles are used to keep a handle to a key. They are completely transient.
+ */
 struct FKeyHandle
 {
 	ENGINE_API FKeyHandle();
 
-	bool operator == (const FKeyHandle& Other) const
+	bool operator ==(const FKeyHandle& Other) const
 	{
 		return Index == Other.Index;
 	}
 	
-	bool operator != (const FKeyHandle& Other) const
+	bool operator !=(const FKeyHandle& Other) const
 	{
 		return Index != Other.Index;
 	}
 	
-	friend uint32 GetTypeHash( const FKeyHandle& Handle )
+	friend uint32 GetTypeHash(const FKeyHandle& Handle)
 	{
 		return GetTypeHash(Handle.Index);
 	}
@@ -32,10 +35,13 @@ struct FKeyHandle
 		Ar << Handle.Index;
 		return Ar;
 	}
+
 private:
+
 	uint32 Index;
 };
 
+
 /**
  * Represents a mapping of key handles to key index that may be serialized 
  */
@@ -72,11 +78,14 @@ public:
 	}
 
 private:
+
 	TMap<FKeyHandle, int32> KeyHandlesToIndices;
 };
 
+
 template<>
-struct TStructOpsTypeTraits< FKeyHandleMap > : public TStructOpsTypeTraitsBase
+struct TStructOpsTypeTraits<FKeyHandleMap>
+	: public TStructOpsTypeTraitsBase
 {
 	enum 
 	{
@@ -86,17 +95,26 @@ struct TStructOpsTypeTraits< FKeyHandleMap > : public TStructOpsTypeTraitsBase
 	};
 };
 
-// @todo Some heavy refactoring can be done here. Much more stuff can go in this base class.
-/** A curve base class which enables key handles to index lookups */
+
+/* Index curve
+ *****************************************************************************/
+
+/**
+ * A curve base class which enables key handles to index lookups.
+ *
+ * @todo sequencer: Some heavy refactoring can be done here. Much more stuff can go in this base class.
+ */
 USTRUCT()
 struct ENGINE_API FIndexedCurve
 {
 	GENERATED_USTRUCT_BODY()
+
 public:
-	FIndexedCurve() {}
+
+	FIndexedCurve() { }
 
 	/** Get number of keys in curve. */
-	virtual int32 GetNumKeys() const PURE_VIRTUAL(FIndexedCurve::GetNumKeys,return 0;);
+	virtual int32 GetNumKeys() const PURE_VIRTUAL(FIndexedCurve::GetNumKeys, return 0;);
 
 	/** Const iterator for the handles. */
 	TMap<FKeyHandle, int32>::TConstIterator GetKeyHandleIterator() const;
@@ -108,6 +126,7 @@ public:
 	virtual bool IsKeyHandleValid(FKeyHandle KeyHandle) const;
 
 protected:
+
 	/** Internal tool to get a handle from an index. */
 	FKeyHandle GetKeyHandle(int32 KeyIndex) const;
 	
@@ -119,13 +138,15 @@ protected:
 	void EnsureAllIndicesHaveHandles() const;
 
 protected:
+
 	/** Map of which key handles go to which indices. */
 	UPROPERTY(transient)
 	mutable FKeyHandleMap KeyHandlesToIndices;
 };
 
-//////////////////////////////////////////////////////////////////////////
-// Rich curve data
+
+/* Rich curves
+ *****************************************************************************/
 
 /** Method of interpolation between this key and the next. */
 UENUM()
@@ -137,6 +158,7 @@ enum ERichCurveInterpMode
 	RCIM_None
 };
 
+
 /** If using RCIM_Cubic, this enum describes how the tangents should be controlled in editor. */
 UENUM()
 enum ERichCurveTangentMode
@@ -147,7 +169,8 @@ enum ERichCurveTangentMode
 	RCTM_None
 };
 
-/** Enum to indicate whether if a tangent is 'weighted' (ie can be stretched). */
+
+/** Enumerates tangent weight modes. */
 UENUM()
 enum ERichCurveTangentWeightMode
 {
@@ -157,7 +180,8 @@ enum ERichCurveTangentWeightMode
 	RCTWM_WeightedBoth
 };
 
-/** Enum to indicate extrapolation option. */
+
+/** Enumerates extrapolation options. */
 UENUM()
 enum ERichCurveExtrapolation
 {
@@ -169,6 +193,7 @@ enum ERichCurveExtrapolation
 	RCCE_None
 };
 
+
 /** One key in a rich, editable float curve */
 USTRUCT()
 struct ENGINE_API FRichCurveKey
@@ -212,40 +237,40 @@ struct ENGINE_API FRichCurveKey
 	float LeaveTangentWeight;
 
 	FRichCurveKey()
-	: InterpMode(RCIM_Linear)
-	, TangentMode(RCTM_Auto)
-	, TangentWeightMode(RCTWM_WeightedNone)
-	, Time(0.f)
-	, Value(0.f)
-	, ArriveTangent(0.f)
-	, ArriveTangentWeight(0.f)
-	, LeaveTangent(0.f)
-	, LeaveTangentWeight(0.f)
-	{}
+		: InterpMode(RCIM_Linear)
+		, TangentMode(RCTM_Auto)
+		, TangentWeightMode(RCTWM_WeightedNone)
+		, Time(0.f)
+		, Value(0.f)
+		, ArriveTangent(0.f)
+		, ArriveTangentWeight(0.f)
+		, LeaveTangent(0.f)
+		, LeaveTangentWeight(0.f)
+	{ }
 
 	FRichCurveKey(float InTime, float InValue)
-	: InterpMode(RCIM_Linear)
-	, TangentMode(RCTM_Auto)
-	, TangentWeightMode(RCTWM_WeightedNone)
-	, Time(InTime)
-	, Value(InValue)
-	, ArriveTangent(0.f)
-	, ArriveTangentWeight(0.f)
-	, LeaveTangent(0.f)
-	, LeaveTangentWeight(0.f)
-	{}
+		: InterpMode(RCIM_Linear)
+		, TangentMode(RCTM_Auto)
+		, TangentWeightMode(RCTWM_WeightedNone)
+		, Time(InTime)
+		, Value(InValue)
+		, ArriveTangent(0.f)
+		, ArriveTangentWeight(0.f)
+		, LeaveTangent(0.f)
+		, LeaveTangentWeight(0.f)
+	{ }
 
 	FRichCurveKey(float InTime, float InValue, float InArriveTangent, const float InLeaveTangent, ERichCurveInterpMode InInterpMode)
-	: InterpMode(InInterpMode)
-	, TangentMode(RCTM_Auto)
-	, TangentWeightMode(RCTWM_WeightedNone)
-	, Time(InTime)
-	, Value(InValue)
-	, ArriveTangent(InArriveTangent)
-	, ArriveTangentWeight(0.f)
-	, LeaveTangent(InLeaveTangent)
-	, LeaveTangentWeight(0.f)
-	{}
+		: InterpMode(InInterpMode)
+		, TangentMode(RCTM_Auto)
+		, TangentWeightMode(RCTWM_WeightedNone)
+		, Time(InTime)
+		, Value(InValue)
+		, ArriveTangent(InArriveTangent)
+		, ArriveTangentWeight(0.f)
+		, LeaveTangent(InLeaveTangent)
+		, LeaveTangentWeight(0.f)
+	{ }
 
 	/** Conversion constructor */
 	FRichCurveKey(const FInterpCurvePoint<float>& InPoint);
@@ -261,13 +286,19 @@ struct ENGINE_API FRichCurveKey
 		P.Serialize(Ar);
 		return Ar;
 	}
+};
+
 
+template<>
+struct TIsPODType<FRichCurveKey>
+{
+	enum { Value = true };
 };
 
-template<> struct TIsPODType<FRichCurveKey> { enum { Value = true }; };
 
 template<>
-struct TStructOpsTypeTraits< FRichCurveKey > : public TStructOpsTypeTraitsBase
+struct TStructOpsTypeTraits<FRichCurveKey>
+	: public TStructOpsTypeTraitsBase
 {
 	enum
 	{
@@ -277,24 +308,25 @@ struct TStructOpsTypeTraits< FRichCurveKey > : public TStructOpsTypeTraitsBase
 	};
 };
 
+
 /** A rich, editable float curve */
 USTRUCT()
-struct ENGINE_API FRichCurve : public FIndexedCurve
+struct ENGINE_API FRichCurve
+	: public FIndexedCurve
 {
 	GENERATED_USTRUCT_BODY()
 
 public:
+
 	FRichCurve() 
-	: FIndexedCurve()
-	, PreInfinityExtrap(RCCE_Constant)
-	, PostInfinityExtrap(RCCE_Constant)
-	, DefaultValue(MAX_flt)
-	{
-	}
+		: FIndexedCurve()
+		, PreInfinityExtrap(RCCE_Constant)
+		, PostInfinityExtrap(RCCE_Constant)
+		, DefaultValue(MAX_flt)
+	{ }
 
-	virtual ~FRichCurve()
-	{
-	}
+	/** Virtual destructor. */
+	virtual ~FRichCurve() { }
 
 	/** Gets a copy of the keys, so indices and handles can't be meddled with */
 	TArray<FRichCurveKey> GetCopyOfKeys() const;
@@ -314,12 +346,6 @@ public:
 	FKeyHandle GetNextKey(FKeyHandle KeyHandle) const;
 	FKeyHandle GetPreviousKey(FKeyHandle KeyHandle) const;
 
-	/** Get number of key in curve. */
-	virtual int32 GetNumKeys() const override;
-	
-	/** Checks to see if the key handle is valid for this curve */
-	virtual bool IsKeyHandleValid(FKeyHandle KeyHandle) const override;
-
 	/**
 	  * Add a new key to the curve with the supplied Time and Value. Returns the handle of the new key.
 	  * 
@@ -402,6 +428,15 @@ public:
 	/** Determine if two RichCurves are the same */
 	bool operator == (const FRichCurve& Curve) const;
 
+public:
+
+	// FIndexedCurve interface
+
+	virtual int32 GetNumKeys() const override;
+	virtual bool IsKeyHandleValid(FKeyHandle KeyHandle) const override;
+
+public:
+
 	/** Pre-infinity extrapolation state */
 	UPROPERTY()
 	TEnumAsByte<ERichCurveExtrapolation> PreInfinityExtrap;
@@ -419,32 +454,36 @@ public:
 	float DefaultValue;
 };
 
-//////////////////////////////////////////////////////////////////////////
-//~ Begin Curve editor Interface
 
-/** Info about a curve to be edited */
-template<class T> struct FRichCurveEditInfoTemplate
+/* Curve editors
+ *****************************************************************************/
+
+/**
+ * Info about a curve to be edited.
+ */
+template<class T>
+struct FRichCurveEditInfoTemplate
 {
 	/** Name of curve, used when displaying in editor. Can include comma's to allow tree expansion in editor */
-	FName			CurveName;
+	FName CurveName;
 
 	/** Pointer to curves to be edited */
-	T				CurveToEdit;
+	T CurveToEdit;
 
 	FRichCurveEditInfoTemplate()
 		: CurveName(NAME_None)
 		, CurveToEdit(nullptr)
-	{}
+	{ }
 
 	FRichCurveEditInfoTemplate(T InCurveToEdit)
-	:	CurveName(NAME_None)
-	,	CurveToEdit(InCurveToEdit)
-	{}
+		: CurveName(NAME_None)
+		, CurveToEdit(InCurveToEdit)
+	{ }
 
 	FRichCurveEditInfoTemplate(T InCurveToEdit, FName InCurveName)
-	:	CurveName(InCurveName)
-	,	CurveToEdit(InCurveToEdit)
-	{}
+		: CurveName(InCurveName)
+		, CurveToEdit(InCurveToEdit)
+	{ }
 
 	inline bool operator==(const FRichCurveEditInfoTemplate<T>& Other) const
 	{
@@ -457,21 +496,26 @@ template<class T> struct FRichCurveEditInfoTemplate
 	}
 };
 
+
 template<class T>
 inline uint32 GetTypeHash( const FRichCurveEditInfoTemplate<T>& RichCurveEditInfo )
 {
 	return RichCurveEditInfo.GetTypeHash();
 }
 
+
 typedef FRichCurveEditInfoTemplate<FRichCurve*>			FRichCurveEditInfo;
 typedef FRichCurveEditInfoTemplate<const FRichCurve*>	FRichCurveEditInfoConst;
 
 
-/** Interface you implement if you want the CurveEditor to be able to edit curves on you */
+/**
+ * Interface you implement if you want the CurveEditor to be able to edit curves on you.
+ */
 class FCurveOwnerInterface
 {
 public:
-	virtual ~FCurveOwnerInterface() {}
+
+	virtual ~FCurveOwnerInterface() { }
 
 	/** Returns set of curves to edit. Must not release the curves while being edited. */
 	virtual TArray<FRichCurveEditInfoConst> GetCurves() const = 0;
@@ -489,13 +533,22 @@ public:
 	virtual void OnCurveChanged(const TArray<FRichCurveEditInfo>& ChangedCurveEditInfos) = 0;
 
 	/** Whether the curve represents a linear color */
-	virtual bool IsLinearColorCurve() const { return false; }
+	virtual bool IsLinearColorCurve() const
+	{
+		return false;
+	}
 
 	/** Evaluate this color curve at the specified time */
-	virtual FLinearColor GetLinearColorValue(float InTime) const { return FLinearColor::Black; }
+	virtual FLinearColor GetLinearColorValue(float InTime) const
+	{
+		return FLinearColor::Black;
+	}
 
 	/** @return True if the curve has any alpha keys */
-	virtual bool HasAnyAlphaKeys() const { return false; }
+	virtual bool HasAnyAlphaKeys() const
+	{
+		return false;
+	}
 
 	/** Validates that a previously retrieved curve is still valid for editing. */
 	virtual bool IsValidCurve(FRichCurveEditInfo CurveInfo) = 0;
@@ -508,7 +561,9 @@ public:
  * Defines a curve of interpolated points to evaluate over a given range
  */
 UCLASS(abstract)
-class ENGINE_API UCurveBase : public UObject, public FCurveOwnerInterface
+class ENGINE_API UCurveBase
+	: public UObject
+	, public FCurveOwnerInterface
 {
 	GENERATED_UCLASS_BODY()
 
@@ -520,8 +575,21 @@ class ENGINE_API UCurveBase : public UObject, public FCurveOwnerInterface
 	UFUNCTION(BlueprintCallable, Category="Math|Curves")
 	void GetValueRange(float& MinValue, float& MaxValue) const;
 
+	/** 
+	 *	Create curve from CSV style comma-separated string.
+	 *
+	 * @param InString The input string to parse.
+	 *	@return	Set of problems encountered while processing input
+	 */
+	TArray<FString> CreateCurveFromCSVString(const FString& InString);
+
+	/** Reset all curve data */
+	void ResetCurve();
+
 public:
-	// Begin FCurveOwnerInterface
+
+	// FCurveOwnerInterface interface
+
 	virtual TArray<FRichCurveEditInfoConst> GetCurves() const override
 	{
 		TArray<FRichCurveEditInfoConst> Curves;
@@ -538,24 +606,17 @@ public:
 	virtual void MakeTransactional() override;
 	virtual void OnCurveChanged(const TArray<FRichCurveEditInfo>& ChangedCurveEditInfos) override;
 
-	virtual bool IsValidCurve( FRichCurveEditInfo CurveInfo ) override { return false; };
-	// End FCurveOwnerInterface
-
-	//~ Begin UCurveBase Interface
-
-	/** 
-	 *	Create curve from CSV style comma-separated string. 
-	 *	@return	Set of problems encountered while processing input
-	 */
-	TArray<FString> CreateCurveFromCSVString(const FString& InString);
+	virtual bool IsValidCurve(FRichCurveEditInfo CurveInfo) override
+	{
+		return false;
+	};
 
-	/** Reset all curve data */
-	void ResetCurve();
+public:
 
-	//~ End UCurveBase Interface
+	// UObject overrides
 
-	//~ Begin UObject Interface
 #if WITH_EDITORONLY_DATA
+
 	/** Override to ensure we write out the asset import data */
 	virtual void GetAssetRegistryTags(TArray<FAssetRegistryTag>& OutTags) const override;
 	virtual void PostInitProperties() override;
@@ -569,9 +630,9 @@ public:
 	FString ImportPath_DEPRECATED;
 
 #endif
-	//~ End UObject Interface
 };
 
+
 //////////////////////////////////////////////////////////////////////////
 
 /** An integral key, which holds the key time and the key value */
@@ -580,9 +641,11 @@ struct FIntegralKey
 {
 	GENERATED_USTRUCT_BODY()
 public:
+
 	FIntegralKey(float InTime = 0.f, int32 InValue = 0)
 		: Time(InTime)
-		, Value(InValue) {}
+		, Value(InValue)
+	{ }
 	
 	/** The keyed time */
 	UPROPERTY()
@@ -593,21 +656,24 @@ public:
 	int32 Value;
 };
 
+
 /** An integral curve, which holds the key time and the key value */
 USTRUCT()
-struct ENGINE_API FIntegralCurve : public FIndexedCurve
+struct ENGINE_API FIntegralCurve
+	: public FIndexedCurve
 {
 	GENERATED_USTRUCT_BODY()
+
 public:
+
+	/** Default constructor. */
 	FIntegralCurve() 
-	: FIndexedCurve()
-	, DefaultValue(MAX_int32)
-	{
-	}
+		: FIndexedCurve()
+		, DefaultValue(MAX_int32)
+	{ }
 
-	virtual ~FIntegralCurve()
-	{
-	}
+	/** Virtual destructor. */
+	virtual ~FIntegralCurve() { }
 	
 	/** Get number of keys in curve. */
 	virtual int32 GetNumKeys() const override;
@@ -622,9 +688,10 @@ public:
 	TArray<FIntegralKey>::TConstIterator GetKeyIterator() const;
 	
 	/**
-	  * Add a new key to the curve with the supplied Time and Value. Returns the handle of the new key.
+	  * Add a new key to the curve with the supplied Time and Value.
 	  * 
-	  * @param  KeyHandle			Optionally can specify what handle this new key should have, otherwise, it'll make a new one
+	  * @param KeyHandle Optionally can specify what handle this new key should have, otherwise, it'll make a new one
+	  * @return The handle of the new key
 	  */
 	FKeyHandle AddKey( float InTime, int32 InValue, FKeyHandle KeyHandle = FKeyHandle() );
 	
@@ -661,6 +728,7 @@ public:
 	FKeyHandle FindKey(float KeyTime) const;
 
 private:
+
 	/** The keys, ordered by time */
 	UPROPERTY(EditAnywhere, Category="Curve")
 	TArray<FIntegralKey> Keys;
diff --git a/Engine/Source/Runtime/Engine/Private/CurveBase.cpp b/Engine/Source/Runtime/Engine/Private/CurveBase.cpp
--- a/Engine/Source/Runtime/Engine/Private/CurveBase.cpp
+++ b/Engine/Source/Runtime/Engine/Private/CurveBase.cpp
@@ -4,62 +4,77 @@
 #include "CsvParser.h"
 #include "EditorFramework/AssetImportData.h"
 
+
+/* FKeyHandleMap
+ *****************************************************************************/
+
 void FKeyHandleMap::Add( const FKeyHandle& InHandle, int32 InIndex )
 {
 	KeyHandlesToIndices.Add( InHandle, InIndex );
 }
 
+
 void FKeyHandleMap::Empty()
 {
 	KeyHandlesToIndices.Empty();
 }
 
+
 void FKeyHandleMap::Remove( const FKeyHandle& InHandle )
 {
 	KeyHandlesToIndices.Remove( InHandle );
 }
 
+
 const int32* FKeyHandleMap::Find( const FKeyHandle& InHandle ) const
 {
 	return KeyHandlesToIndices.Find( InHandle );
 }
 
+
 const FKeyHandle* FKeyHandleMap::FindKey( int32 KeyIndex ) const
 {
 	return KeyHandlesToIndices.FindKey( KeyIndex );
 }
 
+
 int32 FKeyHandleMap::Num() const
 {
 	return KeyHandlesToIndices.Num();
 }
 
+
 TMap<FKeyHandle, int32>::TConstIterator FKeyHandleMap::CreateConstIterator() const
 {
 	return KeyHandlesToIndices.CreateConstIterator();
 }
 
+
 TMap<FKeyHandle, int32>::TIterator FKeyHandleMap::CreateIterator() 
 {
 	return KeyHandlesToIndices.CreateIterator();
 }
 
+
 bool FKeyHandleMap::Serialize(FArchive& Ar)
 {
+	// only allow this map to be saved to the transaction buffer
 	if( Ar.IsTransacting() )
 	{
-		// Only allow this map to be saved to the transaction buffer
 		Ar << KeyHandlesToIndices;
 	}
+
 	return true;
 }
 
+
 bool FKeyHandleMap::operator==(const FKeyHandleMap& Other) const
 {
 	if (KeyHandlesToIndices.Num() != Other.KeyHandlesToIndices.Num())
 	{
 		return false;
 	}
+
 	for (TMap<FKeyHandle, int32>::TConstIterator It(KeyHandlesToIndices); It; ++It)
 	{
 		int32 const* OtherVal = Other.KeyHandlesToIndices.Find(It.Key());
@@ -68,48 +83,55 @@ bool FKeyHandleMap::operator==(const FKeyHandleMap& Other) const
 			return false;
 		}
 	}
+
 	return true;
 }
 
+
 bool FKeyHandleMap::operator!=(const FKeyHandleMap& Other) const
 {
 	return !(*this==Other);
 }
 
-//////////////////////////////////////////////////////////////////////////
-// FIndexedCurve
+
+/* FIndexedCurve
+ *****************************************************************************/
 
 int32 FIndexedCurve::GetIndexSafe(FKeyHandle KeyHandle) const
 {
 	return IsKeyHandleValid(KeyHandle) ? *KeyHandlesToIndices.Find(KeyHandle) : INDEX_NONE;
 }
 
+
 int32 FIndexedCurve::GetIndex(FKeyHandle KeyHandle) const
 {
 	return *KeyHandlesToIndices.Find(KeyHandle);
 }
 
+
 FKeyHandle FIndexedCurve::GetKeyHandle(int32 KeyIndex) const
 {
 	check(KeyIndex >= 0 && KeyIndex < GetNumKeys());
-
 	EnsureIndexHasAHandle(KeyIndex);
 
 	return *KeyHandlesToIndices.FindKey(KeyIndex);
 }
 
+
 TMap<FKeyHandle, int32>::TConstIterator FIndexedCurve::GetKeyHandleIterator() const
 {
 	EnsureAllIndicesHaveHandles();
 	return KeyHandlesToIndices.CreateConstIterator();
 }
 
+
 bool FIndexedCurve::IsKeyHandleValid(FKeyHandle KeyHandle) const
 {
 	EnsureAllIndicesHaveHandles();
 	return KeyHandlesToIndices.Find(KeyHandle) != NULL;
 }
 
+
 void FIndexedCurve::EnsureIndexHasAHandle(int32 KeyIndex) const
 {
 	const FKeyHandle* KeyHandle = KeyHandlesToIndices.FindKey(KeyIndex);
@@ -120,6 +142,7 @@ void FIndexedCurve::EnsureIndexHasAHandle(int32 KeyIndex) const
 	}
 }
 
+
 void FIndexedCurve::EnsureAllIndicesHaveHandles() const
 {
 	if (KeyHandlesToIndices.Num() != GetNumKeys())
@@ -131,8 +154,9 @@ void FIndexedCurve::EnsureAllIndicesHaveHandles() const
 	}
 }
 
-//////////////////////////////////////////////////////////////////////////
-// FRichCurveKey
+
+/* FRichCurveKey
+ *****************************************************************************/
 
 static void SetModesFromLegacy(FRichCurveKey& InKey, EInterpCurveMode InterpMode)
 {
@@ -140,11 +164,11 @@ static void SetModesFromLegacy(FRichCurveKey& InKey, EInterpCurveMode InterpMode
 	InKey.TangentWeightMode = RCTWM_WeightedNone;
 	InKey.TangentMode = RCTM_Auto;
 
-	if(InterpMode == CIM_Constant)
+	if (InterpMode == CIM_Constant)
 	{
 		InKey.InterpMode = RCIM_Constant;
 	}
-	else if(InterpMode == CIM_Linear)
+	else if (InterpMode == CIM_Linear)
 	{
 		InKey.InterpMode = RCIM_Linear;
 	}
@@ -152,21 +176,22 @@ static void SetModesFromLegacy(FRichCurveKey& InKey, EInterpCurveMode InterpMode
 	{
 		InKey.InterpMode = RCIM_Cubic;
 
-		if(InterpMode == CIM_CurveAuto || InterpMode == CIM_CurveAutoClamped)
+		if (InterpMode == CIM_CurveAuto || InterpMode == CIM_CurveAutoClamped)
 		{
 			InKey.TangentMode = RCTM_Auto;
 		}
-		else if(InterpMode == CIM_CurveBreak)
+		else if (InterpMode == CIM_CurveBreak)
 		{
 			InKey.TangentMode = RCTM_Break;
 		}
-		else if(InterpMode == CIM_CurveUser)
+		else if (InterpMode == CIM_CurveUser)
 		{
 			InKey.TangentMode = RCTM_User;
 		}
 	}
 }
 
+
 FRichCurveKey::FRichCurveKey(const FInterpCurvePoint<float>& InPoint)
 {
 	SetModesFromLegacy(*this, InPoint.InterpMode);
@@ -181,19 +206,20 @@ FRichCurveKey::FRichCurveKey(const FInterpCurvePoint<float>& InPoint)
 	LeaveTangentWeight = 0.f;
 }
 
+
 FRichCurveKey::FRichCurveKey(const FInterpCurvePoint<FVector>& InPoint, int32 ComponentIndex)
 {
 	SetModesFromLegacy(*this, InPoint.InterpMode);
 
 	Time = InPoint.InVal;
 
-	if(ComponentIndex == 0)
+	if (ComponentIndex == 0)
 	{
 		Value = InPoint.OutVal.X;
 		ArriveTangent = InPoint.ArriveTangent.X;
 		LeaveTangent = InPoint.LeaveTangent.X;
 	}
-	else if(ComponentIndex == 1)
+	else if (ComponentIndex == 1)
 	{
 		Value = InPoint.OutVal.Y;
 		ArriveTangent = InPoint.ArriveTangent.Y;
@@ -210,30 +236,30 @@ FRichCurveKey::FRichCurveKey(const FInterpCurvePoint<FVector>& InPoint, int32 Co
 	LeaveTangentWeight = 0.f;
 }
 
+
 bool FRichCurveKey::Serialize(FArchive& Ar)
 {
-	if( Ar.UE4Ver() >= VER_UE4_SERIALIZE_RICH_CURVE_KEY )
-	{
-		// Serialization is handled manually to avoid the extra size overhead of UProperty tagging.
-		// Otherwise with many keys in a rich curve the size can become quite large.
-		Ar << InterpMode;
-		Ar << TangentMode;
-		Ar << TangentWeightMode;
-		Ar << Time;
-		Ar << Value;
-		Ar << ArriveTangent;
-		Ar << ArriveTangentWeight;
-		Ar << LeaveTangent;
-		Ar << LeaveTangentWeight;
-
-		return true;
-	}
-	else
+	if (Ar.UE4Ver() < VER_UE4_SERIALIZE_RICH_CURVE_KEY)
 	{
 		return false;
 	}
+
+	// Serialization is handled manually to avoid the extra size overhead of UProperty tagging.
+	// Otherwise with many keys in a rich curve the size can become quite large.
+	Ar << InterpMode;
+	Ar << TangentMode;
+	Ar << TangentWeightMode;
+	Ar << Time;
+	Ar << Value;
+	Ar << ArriveTangent;
+	Ar << ArriveTangentWeight;
+	Ar << LeaveTangent;
+	Ar << LeaveTangentWeight;
+
+	return true;
 }
 
+
 bool FRichCurveKey::operator==( const FRichCurveKey& Curve ) const
 {
 	return (Time == Curve.Time) && (Value == Curve.Value) && (InterpMode == Curve.InterpMode) &&
@@ -243,98 +269,112 @@ bool FRichCurveKey::operator==( const FRichCurveKey& Curve ) const
 
 }
 
+
 bool FRichCurveKey::operator!=(const FRichCurveKey& Other) const
 {
 	return !(*this == Other);
 }
 
-//////////////////////////////////////////////////////////////////////////
-// FRichCurve
 
 FKeyHandle::FKeyHandle()
 {
 	static uint32 LastKeyHandleIndex = 1;
 	check(LastKeyHandleIndex != 0); // check in the unlikely event that this overflows
+
 	Index = ++LastKeyHandleIndex;
 }
 
+
+/* FRichCurve
+ *****************************************************************************/
+
 TArray<FRichCurveKey> FRichCurve::GetCopyOfKeys() const
 {
 	return Keys;
 }
 
+
 TArray<FRichCurveKey>::TConstIterator FRichCurve::GetKeyIterator() const
 {
 	return Keys.CreateConstIterator();
 }
 
+
 FRichCurveKey& FRichCurve::GetKey(FKeyHandle KeyHandle)
 {
 	EnsureAllIndicesHaveHandles();
 	return Keys[GetIndex(KeyHandle)];
 }
 
+
 FRichCurveKey FRichCurve::GetKey(FKeyHandle KeyHandle) const
 {
 	EnsureAllIndicesHaveHandles();
 	return Keys[GetIndex(KeyHandle)];
 }
 
+
 FRichCurveKey FRichCurve::GetFirstKey() const
 {
 	check(Keys.Num() > 0);
 	return Keys[0];
 }
 
+
 FRichCurveKey FRichCurve::GetLastKey() const
 {
 	check(Keys.Num() > 0);
 	return Keys[Keys.Num()-1];
 }
 
+
 FKeyHandle FRichCurve::GetNextKey(FKeyHandle KeyHandle) const
 {
 	int32 KeyIndex = GetIndex(KeyHandle);
 	KeyIndex++;
+
 	if (Keys.IsValidIndex(KeyIndex))
 	{
 		return GetKeyHandle(KeyIndex);
 	}
-	else
-	{
-		return FKeyHandle();
-	}
+
+	return FKeyHandle();
 }
 
+
 FKeyHandle FRichCurve::GetPreviousKey(FKeyHandle KeyHandle) const
 {
 	int32 KeyIndex = GetIndex(KeyHandle);
 	KeyIndex--;
+
 	if (Keys.IsValidIndex(KeyIndex))
 	{
 		return GetKeyHandle(KeyIndex);
 	}
-	else
-	{
-		return FKeyHandle();
-	}
+
+	return FKeyHandle();
 }
 
+
 int32 FRichCurve::GetNumKeys() const
 {
 	return Keys.Num();
 }
 
+
 bool FRichCurve::IsKeyHandleValid(FKeyHandle KeyHandle) const
 {
 	bool bValid = false;
+
 	if (FIndexedCurve::IsKeyHandleValid(KeyHandle))
 	{
 		bValid = Keys.IsValidIndex( GetIndex(KeyHandle) );
 	}
+
 	return bValid;
 }
 
+
 FKeyHandle FRichCurve::AddKey( const float InTime, const float InValue, const bool bUnwindRotation, FKeyHandle NewHandle )
 {
 	int32 Index = 0;
@@ -374,6 +414,7 @@ FKeyHandle FRichCurve::AddKey( const float InTime, const float InValue, const bo
 	return GetKeyHandle(Index);
 }
 
+
 void FRichCurve::DeleteKey(FKeyHandle InKeyHandle)
 {
 	int32 Index = GetIndex(InKeyHandle);
@@ -388,38 +429,47 @@ void FRichCurve::DeleteKey(FKeyHandle InKeyHandle)
 		const FKeyHandle& KeyHandle = It.Key();
 		int32& KeyIndex = It.Value();
 
-		if (KeyIndex >= Index) {--KeyIndex;}
+		if (KeyIndex >= Index)
+		{
+			--KeyIndex;
+		}
 	}
 }
 
+
 FKeyHandle FRichCurve::UpdateOrAddKey(float InTime, float InValue, const bool bUnwindRotation)
 {
 	// Search for a key that already exists at the time and if found, update its value
-	for( int32 KeyIndex = 0; KeyIndex < Keys.Num(); ++KeyIndex )
+	for (int32 KeyIndex = 0; KeyIndex < Keys.Num(); ++KeyIndex)
 	{
 		float KeyTime = Keys[KeyIndex].Time;
 
-		if( KeyTime > InTime )
+		if (KeyTime > InTime)
 		{
 			// All the rest of the keys exist after the key we want to add
 			// so there is no point in searching
 			break;
 		}
 
-		if( FMath::IsNearlyEqual( KeyTime, InTime ) )
+		if (FMath::IsNearlyEqual(KeyTime, InTime))
 		{
 			Keys[KeyIndex].Value = InValue;
+
 			return GetKeyHandle(KeyIndex);
 		}
 	}
 
 	// A key wasnt found, add it now
-	return AddKey( InTime, InValue, bUnwindRotation );
+	return AddKey(InTime, InValue, bUnwindRotation);
 }
 
+
 FKeyHandle FRichCurve::SetKeyTime( FKeyHandle KeyHandle, float NewTime )
 {
-	if (!IsKeyHandleValid(KeyHandle)) {return KeyHandle;}
+	if (!IsKeyHandleValid(KeyHandle))
+	{
+		return KeyHandle;
+	}
 
 	FRichCurveKey OldKey = GetKey(KeyHandle);
 	
@@ -433,29 +483,34 @@ FKeyHandle FRichCurve::SetKeyTime( FKeyHandle KeyHandle, float NewTime )
 	return KeyHandle;
 }
 
+
 float FRichCurve::GetKeyTime(FKeyHandle KeyHandle) const
 {
-	if (!IsKeyHandleValid(KeyHandle)) {return 0.f;}
+	if (!IsKeyHandleValid(KeyHandle))
+	{
+		return 0.f;
+	}
 
 	return GetKey(KeyHandle).Time;
 }
 
-FKeyHandle FRichCurve::FindKey( float KeyTime ) const
+
+FKeyHandle FRichCurve::FindKey(float KeyTime) const
 {
 	int32 Start = 0;
 	int32 End = Keys.Num()-1;
 
 	// Binary search since the keys are in sorted order
-	while( Start <= End )
+	while (Start <= End)
 	{
 		int32 TestPos = Start + (End-Start) / 2;
-
 		float TestKeyTime = Keys[TestPos].Time;
-		if( TestKeyTime == KeyTime )
+
+		if (TestKeyTime == KeyTime)
 		{
-			return GetKeyHandle( TestPos );
+			return GetKeyHandle(TestPos);
 		}
-		else if( TestKeyTime < KeyTime )
+		else if(TestKeyTime < KeyTime)
 		{
 			Start = TestPos+1;
 		}
@@ -468,24 +523,34 @@ FKeyHandle FRichCurve::FindKey( float KeyTime ) const
 	return FKeyHandle();
 }
 
+
 void FRichCurve::SetKeyValue(FKeyHandle KeyHandle, float NewValue, bool bAutoSetTangents)
 {
-	if (!IsKeyHandleValid(KeyHandle)) {return;}
+	if (!IsKeyHandleValid(KeyHandle))
+	{
+		return;
+	}
 
 	GetKey(KeyHandle).Value = NewValue;
-	if ( bAutoSetTangents )
+
+	if (bAutoSetTangents)
 	{
 		AutoSetTangents();
 	}
 }
 
+
 float FRichCurve::GetKeyValue(FKeyHandle KeyHandle) const
 {
-	if (!IsKeyHandleValid(KeyHandle)) {return 0.f;}
+	if (!IsKeyHandleValid(KeyHandle))
+	{
+		return 0.f;
+	}
 
 	return GetKey(KeyHandle).Value;
 }
 
+
 void FRichCurve::ShiftCurve(float DeltaTime)
 {
 	TSet<FKeyHandle> KeyHandles;
@@ -498,6 +563,7 @@ void FRichCurve::ShiftCurve(float DeltaTime)
 	ShiftCurve(DeltaTime, KeyHandles);
 }
 
+
 void FRichCurve::ShiftCurve(float DeltaTime, TSet<FKeyHandle>& KeyHandles)
 {
 	for (auto It = KeyHandlesToIndices.CreateIterator(); It; ++It)
@@ -510,6 +576,7 @@ void FRichCurve::ShiftCurve(float DeltaTime, TSet<FKeyHandle>& KeyHandles)
 	}
 }
 
+
 void FRichCurve::ScaleCurve(float ScaleOrigin, float ScaleFactor)
 {
 	TSet<FKeyHandle> KeyHandles;
@@ -522,6 +589,7 @@ void FRichCurve::ScaleCurve(float ScaleOrigin, float ScaleFactor)
 	ScaleCurve(ScaleOrigin, ScaleFactor, KeyHandles);
 }
 
+
 void FRichCurve::ScaleCurve(float ScaleOrigin, float ScaleFactor, TSet<FKeyHandle>& KeyHandles)
 {
 	for (auto It = KeyHandlesToIndices.CreateIterator(); It; ++It)
@@ -534,47 +602,68 @@ void FRichCurve::ScaleCurve(float ScaleOrigin, float ScaleFactor, TSet<FKeyHandl
 	}
 }
 
+
 void FRichCurve::SetKeyInterpMode(FKeyHandle KeyHandle, ERichCurveInterpMode NewInterpMode)
 {
-	if (!IsKeyHandleValid(KeyHandle)) {return;}
+	if (!IsKeyHandleValid(KeyHandle))
+	{
+		return;
+	}
 
 	GetKey(KeyHandle).InterpMode = NewInterpMode;
 	AutoSetTangents();
 }
 
+
 void FRichCurve::SetKeyTangentMode(FKeyHandle KeyHandle, ERichCurveTangentMode NewTangentMode)
 {
-	if (!IsKeyHandleValid(KeyHandle)) {return;}
+	if (!IsKeyHandleValid(KeyHandle))
+	{
+		return;
+	}
 
 	GetKey(KeyHandle).TangentMode = NewTangentMode;
 	AutoSetTangents();
 }
 
+
 void FRichCurve::SetKeyTangentWeightMode(FKeyHandle KeyHandle, ERichCurveTangentWeightMode NewTangentWeightMode)
 {
-	if (!IsKeyHandleValid(KeyHandle)) {return;}
+	if (!IsKeyHandleValid(KeyHandle))
+	{
+		return;
+	}
 
 	GetKey(KeyHandle).TangentWeightMode = NewTangentWeightMode;
 	AutoSetTangents();
 }
 
+
 ERichCurveInterpMode FRichCurve::GetKeyInterpMode(FKeyHandle KeyHandle) const
 {
-	if (!IsKeyHandleValid(KeyHandle)) {return RCIM_Linear;}
+	if (!IsKeyHandleValid(KeyHandle))
+	{
+		return RCIM_Linear;
+	}
 
 	return GetKey(KeyHandle).InterpMode;
 }
 
+
 ERichCurveTangentMode FRichCurve::GetKeyTangentMode(FKeyHandle KeyHandle) const
 {
-	if (!IsKeyHandleValid(KeyHandle)) { return RCTM_Auto; }
+	if (!IsKeyHandleValid(KeyHandle))
+	{
+		return RCTM_Auto;
+	}
 
 	return GetKey(KeyHandle).TangentMode;
 }
 
+
 void FRichCurve::GetTimeRange(float& MinTime, float& MaxTime) const
 {
-	if(Keys.Num() == 0)
+	if (Keys.Num() == 0)
 	{
 		MinTime = 0.f;
 		MaxTime = 0.f;
@@ -586,41 +675,47 @@ void FRichCurve::GetTimeRange(float& MinTime, float& MaxTime) const
 	}
 }
 
+
 /*	 Finds min/max for cubic curves:
 	Looks for feature points in the signal(determined by change in direction of local tangent), these locations are then re-examined in closer detail recursively */
 template<class T>
 void FeaturePointMethod(T& Function , float StartTime,  float EndTime, float StartValue,float Mu, int Depth, int MaxDepth, float& MaxV, float& MinVal)
 {
-	if(Depth < MaxDepth)
+	if (Depth >= MaxDepth)
 	{
-		float PrevValue   = StartValue;
-		float PrevTangent = StartValue - Function.Eval(StartTime-Mu);
-		EndTime += Mu;
-		for(float f = StartTime + Mu;f < EndTime; f += Mu)
-		{
-			float Value		 = Function.Eval(f);
+		return;
+	}
+
+	float PrevValue   = StartValue;
+	float PrevTangent = StartValue - Function.Eval(StartTime-Mu);
+	EndTime += Mu;
 
-			MaxV   = FMath::Max(Value, MaxV);
-			MinVal = FMath::Min(Value, MinVal);
+	for (float f = StartTime + Mu;f < EndTime; f += Mu)
+	{
+		float Value		 = Function.Eval(f);
+
+		MaxV   = FMath::Max(Value, MaxV);
+		MinVal = FMath::Min(Value, MinVal);
 
-			float CurTangent = Value - PrevValue;
+		float CurTangent = Value - PrevValue;
 		
-			//Change direction? Examine this area closer
-			if(FMath::Sign(CurTangent) != FMath::Sign(PrevTangent))
-			{
-				//feature point centered around the previous tangent
-				float FeaturePointTime = f-Mu*2.0f;
-				FeaturePointMethod(Function, FeaturePointTime, f, Function.Eval(FeaturePointTime), Mu*0.4f,Depth+1, MaxDepth, MaxV, MinVal);
-			}
-			PrevTangent = CurTangent;
-			PrevValue = Value;
+		//Change direction? Examine this area closer
+		if (FMath::Sign(CurTangent) != FMath::Sign(PrevTangent))
+		{
+			//feature point centered around the previous tangent
+			float FeaturePointTime = f-Mu*2.0f;
+			FeaturePointMethod(Function, FeaturePointTime, f, Function.Eval(FeaturePointTime), Mu*0.4f,Depth+1, MaxDepth, MaxV, MinVal);
 		}
+
+		PrevTangent = CurTangent;
+		PrevValue = Value;
 	}
 }
 
+
 void FRichCurve::GetValueRange(float& MinValue, float& MaxValue) const
 {
-	if(Keys.Num() == 0)
+	if (Keys.Num() == 0)
 	{
 		MinValue = MaxValue = 0.f;
 	}
@@ -628,47 +723,48 @@ void FRichCurve::GetValueRange(float& MinValue, float& MaxValue) const
 	{
 		int32 LastKeyIndex = Keys.Num()-1;
 		MinValue = MaxValue = Keys[0].Value;
-		for(int32 i=0; i<Keys.Num(); i++)
+
+		for (int32 i = 0; i < Keys.Num(); i++)
 		{
 			const FRichCurveKey& Key = Keys[i];
-		
+
 			MinValue = FMath::Min(MinValue, Key.Value);
 			MaxValue = FMath::Max(MaxValue, Key.Value);
 
-			if(Key.InterpMode == RCIM_Cubic && i != LastKeyIndex)
+			if (Key.InterpMode == RCIM_Cubic && i != LastKeyIndex)
 			{
 				const FRichCurveKey& NextKey = Keys[i+1];
+				float TimeStep = (NextKey.Time - Key.Time) * 0.2f;
 
-				float TimeStep = (NextKey.Time - Key.Time)*0.2f;
-
-				FeaturePointMethod(*this, Key.Time, NextKey.Time, Key.Value, TimeStep, 0, 3, MaxValue, MinValue );
+				FeaturePointMethod(*this, Key.Time, NextKey.Time, Key.Value, TimeStep, 0, 3, MaxValue, MinValue);
 			}
 		}
 	}
 }
 
+
 void FRichCurve::Reset()
 {
 	Keys.Empty();
-	
 	KeyHandlesToIndices.Empty();
 }
 
+
 void FRichCurve::AutoSetTangents(float Tension)
 {
 	// Iterate over all points in this InterpCurve
-	for(int32 KeyIndex=0; KeyIndex<Keys.Num(); KeyIndex++)
+	for (int32 KeyIndex = 0; KeyIndex<Keys.Num(); KeyIndex++)
 	{
-		FRichCurveKey& Key =  Keys[KeyIndex];
+		FRichCurveKey& Key = Keys[KeyIndex];
 		float ArriveTangent = Key.ArriveTangent;
 		float LeaveTangent  = Key.LeaveTangent;
 
-		if(KeyIndex == 0)
+		if (KeyIndex == 0)
 		{
-			if(KeyIndex < Keys.Num()-1) // Start point
+			if (KeyIndex < Keys.Num()-1) // Start point
 			{
 				// If first section is not a curve, or is a curve and first point has manual tangent setting.
-				if( Key.TangentMode == RCTM_Auto )
+				if (Key.TangentMode == RCTM_Auto)
 				{
 					LeaveTangent = 0.0f;
 				}
@@ -677,11 +773,11 @@ void FRichCurve::AutoSetTangents(float Tension)
 		else
 		{
 			
-			if(KeyIndex < Keys.Num()-1) // Inner point
+			if (KeyIndex < Keys.Num() - 1) // Inner point
 			{
 				FRichCurveKey& PrevKey =  Keys[KeyIndex-1];
 
-				if( Key.InterpMode == RCIM_Cubic && (Key.TangentMode == RCTM_Auto ))
+				if (Key.InterpMode == RCIM_Cubic && (Key.TangentMode == RCTM_Auto))
 				{
 						FRichCurveKey& NextKey =  Keys[KeyIndex+1];
 						ComputeCurveTangent(
@@ -698,12 +794,13 @@ void FRichCurve::AutoSetTangents(float Tension)
 						// In 'auto' mode, arrive and leave tangents are always the same
 						LeaveTangent = ArriveTangent;
 				}
-				else if( PrevKey.InterpMode == RCIM_Constant || Key.InterpMode == RCIM_Constant )
+				else if ((PrevKey.InterpMode == RCIM_Constant) || (Key.InterpMode == RCIM_Constant))
 				{
-					if(Keys[ KeyIndex - 1 ].InterpMode != RCIM_Cubic)
+					if (Keys[ KeyIndex - 1 ].InterpMode != RCIM_Cubic)
 					{
 						ArriveTangent = 0.0f;
 					}
+
 					LeaveTangent  = 0.0f;
 				}
 				
@@ -711,7 +808,7 @@ void FRichCurve::AutoSetTangents(float Tension)
 			else // End point
 			{
 				// If last section is not a curve, or is a curve and final point has manual tangent setting.
-				if( Key.InterpMode == RCIM_Cubic && Key.TangentMode == RCTM_Auto)
+				if (Key.InterpMode == RCIM_Cubic && Key.TangentMode == RCTM_Auto)
 				{
 					ArriveTangent = 0.0f;
 				}
@@ -723,10 +820,12 @@ void FRichCurve::AutoSetTangents(float Tension)
 	}
 }
 
+
 void FRichCurve::ReadjustTimeRange(float NewMinTimeRange, float NewMaxTimeRange, bool bInsert/* whether insert or remove*/, float OldStartTime, float OldEndTime)
 {
 	// first readjust modified time keys
 	float ModifiedDuration = OldEndTime - OldStartTime;
+
 	if (bInsert)
 	{
 		for(int32 KeyIndex=0; KeyIndex<Keys.Num(); ++KeyIndex)
@@ -846,6 +945,7 @@ void FRichCurve::ReadjustTimeRange(float NewMinTimeRange, float NewMaxTimeRange,
 	}
 }
 
+
 /** Util to find float value on bezier defined by 4 control points */ 
 static float BezierInterp(float P0, float P1, float P2, float P3, float Alpha)
 {
@@ -855,21 +955,25 @@ static float BezierInterp(float P0, float P1, float P2, float P3, float Alpha)
 	const float P012 = FMath::Lerp(P01, P12, Alpha);
 	const float P123 = FMath::Lerp(P12, P23, Alpha);
 	const float P0123 = FMath::Lerp(P012, P123, Alpha);
+
 	return P0123;
 }
 
+
 static float BezierInterp2(float P0, float Y1, float Y2, float P3, float mu)
 {
-	float P1 = ( -5.f* P0 + 18.f* Y1 -  9.f *Y2 + 2.f* P3 ) / 6.f;
-	float P2 = (  2.f* P0 -  9.f* Y1 + 18.f *Y2 - 5.f* P3 ) / 6.f;
-	float A = P3 - 3.0f*P2 + 3.0f*P1 - P0;
-	float B = 3.0f*P2 - 6.0f*P1 + 3.0f*P0;
-	float C = 3.0f*P1 - 3.0f*P0;
+	float P1 = (-5.f * P0 + 18.f * Y1 - 9.f * Y2 + 2.f * P3) / 6.f;
+	float P2 = (2.f * P0 -  9.f * Y1 + 18.f * Y2 - 5.f * P3) / 6.f;
+	float A = P3 - 3.0f * P2 + 3.0f * P1 - P0;
+	float B = 3.0f * P2 - 6.0f * P1 + 3.0f * P0;
+	float C = 3.0f * P1 - 3.0f * P0;
 	float D = P0;
-	float Result = A * (mu*mu*mu) + B*(mu*mu) + C * mu + D;
+	float Result = A * (mu * mu * mu) + B * (mu * mu) + C * mu + D;
+
 	return Result;
 }
 
+
 static void CycleTime(float MinTime, float MaxTime, float& InTime, int& CycleCount)
 {
 	float InitTime = InTime;
@@ -899,14 +1003,17 @@ static void CycleTime(float MinTime, float MaxTime, float& InTime, int& CycleCou
 	CycleCount = FMath::Abs(CycleCount);
 }
 
+
 void FRichCurve::RemapTimeValue(float& InTime, float& CycleValueOffset) const
 {
 	const int32 NumKeys = Keys.Num();
+	
 	if (NumKeys < 2)
 	{
 		return;
 	} 
-	else if (InTime <= Keys[0].Time)
+
+	if (InTime <= Keys[0].Time)
 	{
 		if (PreInfinityExtrap != RCCE_Linear && PreInfinityExtrap != RCCE_Constant)
 		{
@@ -956,6 +1063,7 @@ void FRichCurve::RemapTimeValue(float& InTime, float& CycleValueOffset) const
 	}
 }
 
+
 float FRichCurve::Eval(float InTime, float InDefaultValue) const
 {
 	// Remap time if extrapolation is present and compute offset value to use if cycling 
@@ -985,6 +1093,7 @@ float FRichCurve::Eval(float InTime, float InDefaultValue) const
 			{
 				float DV = Keys[1].Value - Keys[0].Value;
 				float Slope = DV / DT;
+
 				InterpVal = Slope * (InTime - Keys[0].Time) + Keys[0].Value;
 			}
 		}
@@ -1058,6 +1167,7 @@ float FRichCurve::Eval(float InTime, float InDefaultValue) const
 			{
 				float DV = Keys[NumKeys - 2].Value - Keys[NumKeys - 1].Value;
 				float Slope = DV / DT;
+
 				InterpVal = Slope * (InTime - Keys[NumKeys - 1].Time) + Keys[NumKeys - 1].Value;
 			}
 		}
@@ -1078,6 +1188,7 @@ bool FRichCurve::operator==(const FRichCurve& Curve) const
 	{
 		return false;
 	}
+
 	for(int32 i = 0;i<Keys.Num();++i)
 	{
 		if(!(Keys[i] == Curve.Keys[i]))
@@ -1094,13 +1205,14 @@ bool FRichCurve::operator==(const FRichCurve& Curve) const
 	return true;
 }
 
-//////////////////////////////////////////////////////////////////////////
-// UCurveBase
+
+/* UCurveBase
+ *****************************************************************************/
 
 UCurveBase::UCurveBase(const FObjectInitializer& ObjectInitializer)
 	: Super(ObjectInitializer)
-{
-}
+{ }
+
 
 #if WITH_EDITORONLY_DATA
 void UCurveBase::GetAssetRegistryTags(TArray<FAssetRegistryTag>& OutTags) const
@@ -1113,6 +1225,7 @@ void UCurveBase::GetAssetRegistryTags(TArray<FAssetRegistryTag>& OutTags) const
 	Super::GetAssetRegistryTags(OutTags);
 }
 
+
 void UCurveBase::PostInitProperties()
 {
 	if (!HasAnyFlags(RF_ClassDefaultObject))
@@ -1123,6 +1236,7 @@ void UCurveBase::PostInitProperties()
 	Super::PostInitProperties();
 }
 
+
 void UCurveBase::PostLoad()
 {
 	Super::PostLoad();
@@ -1135,6 +1249,7 @@ void UCurveBase::PostLoad()
 }
 #endif
 
+
 void UCurveBase::GetTimeRange(float& MinTime, float& MaxTime) const
 {
 	TArray<FRichCurveEditInfoConst> Curves = GetCurves();
@@ -1155,6 +1270,7 @@ void UCurveBase::GetTimeRange(float& MinTime, float& MaxTime) const
 	}
 }
 
+
 void UCurveBase::GetValueRange(float& MinValue, float& MaxValue) const
 {
 	TArray<FRichCurveEditInfoConst> Curves = GetCurves();
@@ -1175,21 +1291,24 @@ void UCurveBase::GetValueRange(float& MinValue, float& MaxValue) const
 	}
 }
 
+
 void UCurveBase::ModifyOwner() 
 {
 	Modify(true);
 }
 
+
 void UCurveBase::MakeTransactional() 
 {
 	SetFlags(GetFlags() | RF_Transactional);
 }
 
+
 void UCurveBase::OnCurveChanged(const TArray<FRichCurveEditInfo>& ChangedCurveEditInfos)
 {
-
 }
 
+
 void UCurveBase::ResetCurve()
 {
 	TArray<FRichCurveEditInfo> Curves = GetCurves();
@@ -1203,6 +1322,7 @@ void UCurveBase::ResetCurve()
 	}
 }
 
+
 TArray<FString> UCurveBase::CreateCurveFromCSVString(const FString& InString)
 {
 	// Array used to store problems about curve import
@@ -1260,9 +1380,11 @@ TArray<FString> UCurveBase::CreateCurveFromCSVString(const FString& InString)
 	}
 
 	Modify(true);
+
 	return OutProblems;
 }
 
+
 //////////////////////////////////////////////////////////////////////////
 
 int32 FIntegralCurve::GetNumKeys() const
@@ -1270,6 +1392,7 @@ int32 FIntegralCurve::GetNumKeys() const
 	return Keys.Num();
 }
 
+
 bool FIntegralCurve::IsKeyHandleValid(FKeyHandle KeyHandle) const
 {
 	bool bValid = false;
@@ -1280,6 +1403,7 @@ bool FIntegralCurve::IsKeyHandleValid(FKeyHandle KeyHandle) const
 	return bValid;
 }
 
+
 int32 FIntegralCurve::Evaluate(float Time, int32 InDefaultValue) const
 {
 	// If the default value hasn't been initialized, use the incoming default value
@@ -1315,11 +1439,13 @@ int32 FIntegralCurve::Evaluate(float Time, int32 InDefaultValue) const
 	return ReturnVal;
 }
 
+
 TArray<FIntegralKey>::TConstIterator FIntegralCurve::GetKeyIterator() const
 {
 	return Keys.CreateConstIterator();
 }
 
+
 FKeyHandle FIntegralCurve::AddKey( float InTime, int32 InValue, FKeyHandle InKeyHandle )
 {
 	int32 Index = 0;
@@ -1331,7 +1457,10 @@ FKeyHandle FIntegralCurve::AddKey( float InTime, int32 InValue, FKeyHandle InKey
 		const FKeyHandle& KeyHandle = It.Key();
 		int32& KeyIndex = It.Value();
 
-		if (KeyIndex >= Index) {++KeyIndex;}
+		if (KeyIndex >= Index)
+		{
+			++KeyIndex;
+		}
 	}
 
 	KeyHandlesToIndices.Add(InKeyHandle, Index);
@@ -1339,6 +1468,7 @@ FKeyHandle FIntegralCurve::AddKey( float InTime, int32 InValue, FKeyHandle InKey
 	return GetKeyHandle(Index);
 }
 
+
 void FIntegralCurve::DeleteKey(FKeyHandle InKeyHandle)
 {
 	int32 Index = GetIndex(InKeyHandle);
@@ -1352,10 +1482,14 @@ void FIntegralCurve::DeleteKey(FKeyHandle InKeyHandle)
 		const FKeyHandle& KeyHandle = It.Key();
 		int32& KeyIndex = It.Value();
 
-		if (KeyIndex >= Index) {--KeyIndex;}
+		if (KeyIndex >= Index)
+		{
+			--KeyIndex;
+		}
 	}
 }
 
+
 FKeyHandle FIntegralCurve::UpdateOrAddKey( float Time, int32 Value )
 {
 	for( int32 KeyIndex = 0; KeyIndex < Keys.Num(); ++KeyIndex )
@@ -1380,9 +1514,13 @@ FKeyHandle FIntegralCurve::UpdateOrAddKey( float Time, int32 Value )
 	return AddKey( Time, Value );
 }
 
+
 FKeyHandle FIntegralCurve::SetKeyTime( FKeyHandle KeyHandle, float NewTime )
 {
-	if (!IsKeyHandleValid(KeyHandle)) {return KeyHandle;}
+	if (!IsKeyHandleValid(KeyHandle))
+	{
+		return KeyHandle;
+	}
 
 	FIntegralKey OldKey = GetKey(KeyHandle);
 	
@@ -1396,13 +1534,18 @@ FKeyHandle FIntegralCurve::SetKeyTime( FKeyHandle KeyHandle, float NewTime )
 	return KeyHandle;
 }
 
+
 float FIntegralCurve::GetKeyTime(FKeyHandle KeyHandle) const
 {
-	if (!IsKeyHandleValid(KeyHandle)) {return 0.f;}
+	if (!IsKeyHandleValid(KeyHandle))
+	{
+		return 0.f;
+	}
 
 	return GetKey(KeyHandle).Time;
 }
 
+
 void FIntegralCurve::ShiftCurve(float DeltaTime)
 {
 	TSet<FKeyHandle> KeyHandles;
@@ -1415,6 +1558,7 @@ void FIntegralCurve::ShiftCurve(float DeltaTime)
 	ShiftCurve(DeltaTime, KeyHandles);
 }
 
+
 void FIntegralCurve::ShiftCurve(float DeltaTime, TSet<FKeyHandle>& KeyHandles)
 {
 	for (auto It = KeyHandlesToIndices.CreateIterator(); It; ++It)
@@ -1427,6 +1571,7 @@ void FIntegralCurve::ShiftCurve(float DeltaTime, TSet<FKeyHandle>& KeyHandles)
 	}
 }
 
+
 void FIntegralCurve::ScaleCurve(float ScaleOrigin, float ScaleFactor)
 {
 	TSet<FKeyHandle> KeyHandles;
@@ -1439,6 +1584,7 @@ void FIntegralCurve::ScaleCurve(float ScaleOrigin, float ScaleFactor)
 	ScaleCurve(ScaleOrigin, ScaleFactor, KeyHandles);
 }
 
+
 void FIntegralCurve::ScaleCurve(float ScaleOrigin, float ScaleFactor, TSet<FKeyHandle>& KeyHandles)
 {
 	for (auto It = KeyHandlesToIndices.CreateIterator(); It; ++It)
@@ -1451,18 +1597,21 @@ void FIntegralCurve::ScaleCurve(float ScaleOrigin, float ScaleFactor, TSet<FKeyH
 	}
 }
 
+
 FIntegralKey& FIntegralCurve::GetKey(FKeyHandle KeyHandle)
 {
 	EnsureAllIndicesHaveHandles();
 	return Keys[GetIndex(KeyHandle)];
 }
 
+
 FIntegralKey FIntegralCurve::GetKey(FKeyHandle KeyHandle) const
 {
 	EnsureAllIndicesHaveHandles();
 	return Keys[GetIndex(KeyHandle)];
 }
 
+
 FKeyHandle FIntegralCurve::FindKey(float KeyTime) const
 {
 	int32 Start = 0;
@@ -1472,8 +1621,8 @@ FKeyHandle FIntegralCurve::FindKey(float KeyTime) const
 	while (Start <= End)
 	{
 		int32 TestPos = Start + (End - Start) / 2;
-
 		float TestKeyTime = Keys[TestPos].Time;
+
 		if (TestKeyTime == KeyTime)
 		{
 			return GetKeyHandle(TestPos);
