Basic engine support for multiple worlds per game instance.
-UGameEngine::Tick now ticks systems that are independent of worlds outside of the loop over FWorldContexts
-Make sure FTimerManager only ticks once per frame, since it ticks in UWorld::Tick
-Added ability to get the PlayerController of a UPlayer from a specific world, since with multiple worlds this may need to be disambiguated
-Added ability to associate an FLocalPlayerContext with a specific world in addition to a local player, to disambiguate cases when multiple worlds exist
-Allow overriding the world and URL options for UGameInstance::PlayReplay

--------
Integrated using branch UE4-Orion-To-UE4 of change#2705078 by [REDACTED_PERSON] on 2015/09/24 16:01:28.

diff --git a/Engine/Source/Runtime/Engine/Classes/Engine/GameEngine.h b/Engine/Source/Runtime/Engine/Classes/Engine/GameEngine.h
--- a/Engine/Source/Runtime/Engine/Classes/Engine/GameEngine.h
+++ b/Engine/Source/Runtime/Engine/Classes/Engine/GameEngine.h
@@ -104,6 +104,7 @@ public:
 	virtual float GetMaxTickRate( float DeltaTime, bool bAllowFrameRateSmoothing = true ) const override;
 	virtual void ProcessToggleFreezeCommand( UWorld* InWorld ) override;
 	virtual void ProcessToggleFreezeStreamingCommand( UWorld* InWorld ) override;
+	virtual bool NetworkRemapPath(UWorld* InWorld, FString& Str, bool bReading = true) override;
 
 public:
 
diff --git a/Engine/Source/Runtime/Engine/Classes/Engine/GameInstance.h b/Engine/Source/Runtime/Engine/Classes/Engine/GameInstance.h
--- a/Engine/Source/Runtime/Engine/Classes/Engine/GameInstance.h
+++ b/Engine/Source/Runtime/Engine/Classes/Engine/GameInstance.h
@@ -156,7 +156,7 @@ public:
 
 	int32					GetNumLocalPlayers() const;
 	ULocalPlayer*			GetLocalPlayerByIndex(const int32 Index) const;
-	APlayerController*		GetFirstLocalPlayerController() const;
+	APlayerController*		GetFirstLocalPlayerController(UWorld* World = nullptr) const;
 	ULocalPlayer*			FindLocalPlayerFromControllerId(const int32 ControllerId) const;
 	ULocalPlayer*			FindLocalPlayerFromUniqueNetId(TSharedPtr<const FUniqueNetId> UniqueNetId) const;
 	ULocalPlayer*			FindLocalPlayerFromUniqueNetId(const FUniqueNetId& UniqueNetId) const;
@@ -207,7 +207,7 @@ public:
 	virtual void StopRecordingReplay();
 
 	/** Start playing back a previously recorded replay. */
-	virtual void PlayReplay(const FString& InName);
+	virtual void PlayReplay(const FString& InName, UWorld* WorldOverride = nullptr, const TArray<FString>& AdditionalOptions = TArray<FString>());
 
 	/**
 	 * Adds a join-in-progress user to the set of users associated with the currently recording replay (if any)
diff --git a/Engine/Source/Runtime/Engine/Classes/Engine/LocalPlayer.h b/Engine/Source/Runtime/Engine/Classes/Engine/LocalPlayer.h
--- a/Engine/Source/Runtime/Engine/Classes/Engine/LocalPlayer.h
+++ b/Engine/Source/Runtime/Engine/Classes/Engine/LocalPlayer.h
@@ -13,12 +13,13 @@
 #define INVALID_CONTROLLERID 255
 
 class FUniqueNetId;
+class UWorld;
 
 /** A context object that binds to a LocalPlayer. Useful for UI or other things that need to pass around player references */
 struct ENGINE_API FLocalPlayerContext
 {
 	FLocalPlayerContext();
-	FLocalPlayerContext(const class ULocalPlayer* InLocalPlayer);
+	FLocalPlayerContext(const class ULocalPlayer* InLocalPlayer, UWorld* InWorld = nullptr);
 	FLocalPlayerContext(const class APlayerController* InPlayerController);
 	FLocalPlayerContext(const FLocalPlayerContext& InPlayerContext);
 
@@ -132,6 +133,8 @@ private:
 	void SetPlayerController( const class APlayerController* InPlayerController );
 
 	TWeakObjectPtr<class ULocalPlayer>		LocalPlayer;
+
+	TWeakObjectPtr<UWorld>					World;
 };
 
 /**
diff --git a/Engine/Source/Runtime/Engine/Classes/Engine/Player.h b/Engine/Source/Runtime/Engine/Classes/Engine/Player.h
--- a/Engine/Source/Runtime/Engine/Classes/Engine/Player.h
+++ b/Engine/Source/Runtime/Engine/Classes/Engine/Player.h
@@ -48,4 +48,13 @@ public:
 	 * @param bWriteToLog write out to the log
 	 */
 	ENGINE_API FString ConsoleCommand(const FString& Cmd, bool bWriteToLog = true);
+
+	/**
+	 * Gets the player controller in the given world for this player.
+	 *
+	 * @param InWorld The world in which to search for player controllers.
+	 *
+	 * @return The controller associated with this player in InWorld, if one exists.
+	 */
+	ENGINE_API APlayerController* GetPlayerController(UWorld* InWorld) const;
 };
diff --git a/Engine/Source/Runtime/Engine/Private/GameEngine.cpp b/Engine/Source/Runtime/Engine/Private/GameEngine.cpp
--- a/Engine/Source/Runtime/Engine/Private/GameEngine.cpp
+++ b/Engine/Source/Runtime/Engine/Private/GameEngine.cpp
@@ -580,6 +580,41 @@ void UGameEngine::FinishDestroy()
 	Super::FinishDestroy();
 }
 
+bool UGameEngine::NetworkRemapPath(UWorld* InWorld, FString& Str, bool bReading /*= true*/)
+{
+	// If the game has created multiple worlds, some of them may have prefixed package names,
+	// so we need to remap the world package and streaming levels for replay playback to work correctly.
+	FWorldContext& Context = GetWorldContextFromWorldChecked(InWorld);
+	if (Context.PIEInstance == INDEX_NONE || !bReading)
+	{
+		return false;
+	}
+
+	// If the prefixed path matches the world package name or the name of a streaming level,
+	// return the prefixed name.
+	const FString PrefixedName = UWorld::ConvertToPIEPackageName(Str, Context.PIEInstance);
+	const FString WorldPackageName = InWorld->GetOutermost()->GetName();
+	if (WorldPackageName == PrefixedName)
+	{
+		Str = PrefixedName;
+		return true;
+	}
+
+	for( ULevelStreaming* StreamingLevel : InWorld->StreamingLevels)
+	{
+		if (StreamingLevel != nullptr)
+		{
+			const FString StreamingLevelName = StreamingLevel->WorldAsset.GetLongPackageName();
+			if (StreamingLevelName == PrefixedName)
+			{
+				Str = PrefixedName;
+				return true;
+			}
+		}
+	}
+
+	return false;
+}
 
 void UGameEngine::LoadRuntimeEngineStartupModules()
 {
@@ -896,6 +931,8 @@ void UGameEngine::Tick( float DeltaSeconds, bool bIdleMode )
 	// Begin ticking worlds
 	// -----------------------------------------------------
 
+	bool bIsAnyNonPreviewWorldUnpaused = false;
+
 	FName OriginalGWorldContext = NAME_None;
 	for (int32 i=0; i < WorldList.Num(); ++i)
 	{
@@ -959,38 +996,8 @@ void UGameEngine::Tick( float DeltaSeconds, bool bIdleMode )
 
 			Context.World()->bTriggerPostLoadMap = true;
 		}
-
-		// Tick the viewports.
-		if ( GameViewport != NULL && !bIdleMode )
-		{
-			SCOPE_CYCLE_COUNTER(STAT_GameViewportTick);
-			GameViewport->Tick(DeltaSeconds);
-		}
 	
 		UpdateTransitionType(Context.World());
-	
-		// fixme: this will only happen once due to the static bool, but still need to figure out how to handle this for multiple worlds
-		if (FPlatformProperties::SupportsWindowedMode())
-		{
-			// Hide the splashscreen and show the game window
-			static bool bFirstTime = true;
-			if ( bFirstTime )
-			{
-				bFirstTime = false;
-				FPlatformSplash::Hide();
-				if ( GameViewportWindow.IsValid() )
-				{
-					GameViewportWindow.Pin()->ShowWindow();
-					FSlateApplication::Get().RegisterGameViewport( GameViewportWidget.ToSharedRef() );
-				}
-			}
-		}
-
-		if (!bIdleMode && !IsRunningDedicatedServer() && !IsRunningCommandlet())
-		{
-			// Render everything.
-			RedrawViewports();
-		}
 
 		// Block on async loading if requested.
 		if (Context.World()->bRequestedBlockOnAsyncLoading)
@@ -1006,42 +1013,17 @@ void UGameEngine::Tick( float DeltaSeconds, bool bIdleMode )
 			Context.World()->UpdateLevelStreaming();
 		}
 
-		if (Context.WorldType != EWorldType::Preview)
-		{
-		// Update Audio. This needs to occur after rendering as the rendering code updates the listener position.
-		if (FAudioDevice* AudioDevice = Context.World()->GetAudioDevice())
-		{
-			AudioDevice->Update(!Context.World()->IsPaused());
-		}
-		}
-
-		if( GIsClient )
-		{
-			// IStreamingManager is updated outside of a world context. For now, assuming it needs to tick here, before possibly calling PostLoadMap. 
-			// Will need to take another look when trying to support multiple worlds.
-
-			// Update resource streaming after viewports have had a chance to update view information. Normal update.
-			{
-				QUICK_SCOPE_CYCLE_COUNTER(STAT_UGameEngine_Tick_IStreamingManager);
-				IStreamingManager::Get().Tick( DeltaSeconds );
-			}
-
-			if ( Context.World()->bTriggerPostLoadMap )
-			{
-				Context.World()->bTriggerPostLoadMap = false;
-
-				// Turns off the loading movie (if it was turned on by LoadMap) and other post-load cleanup.
-				QUICK_SCOPE_CYCLE_COUNTER(STAT_UGameEngine_Tick_IStreamingManager);
-				PostLoadMap();
-			}
-		}
-
 		UNCLOCK_CYCLES(LocalTickCycles);
 		TickCycles=LocalTickCycles;
 
 		// See whether any map changes are pending and we requested them to be committed.
 		QUICK_SCOPE_CYCLE_COUNTER(STAT_UGameEngine_Tick_ConditionalCommitMapChange);
 		ConditionalCommitMapChange(Context);
+
+		if (Context.WorldType != EWorldType::Preview && !Context.World()->IsPaused())
+		{
+			bIsAnyNonPreviewWorldUnpaused = true;
+		}
 	}
 
 	// ----------------------------
@@ -1055,6 +1037,49 @@ void UGameEngine::Tick( float DeltaSeconds, bool bIdleMode )
 		GWorld = GetWorldContextFromHandleChecked(OriginalGWorldContext).World();
 	}
 
+	// Tick the viewport
+	if ( GameViewport != NULL && !bIdleMode )
+	{
+		SCOPE_CYCLE_COUNTER(STAT_GameViewportTick);
+		GameViewport->Tick(DeltaSeconds);
+	}
+
+	if (FPlatformProperties::SupportsWindowedMode())
+	{
+		// Hide the splashscreen and show the game window
+		static bool bFirstTime = true;
+		if ( bFirstTime )
+		{
+			bFirstTime = false;
+			FPlatformSplash::Hide();
+			if ( GameViewportWindow.IsValid() )
+			{
+				GameViewportWindow.Pin()->ShowWindow();
+				FSlateApplication::Get().RegisterGameViewport( GameViewportWidget.ToSharedRef() );
+			}
+		}
+	}
+
+	if (!bIdleMode && !IsRunningDedicatedServer() && !IsRunningCommandlet())
+	{
+		// Render everything.
+		RedrawViewports();
+	}
+
+	if( GIsClient )
+	{
+		// Update resource streaming after viewports have had a chance to update view information. Normal update.
+		QUICK_SCOPE_CYCLE_COUNTER(STAT_UGameEngine_Tick_IStreamingManager);
+		IStreamingManager::Get().Tick( DeltaSeconds );
+	}
+
+	// Update Audio. This needs to occur after rendering as the rendering code updates the listener position.
+	FAudioDeviceManager* AudioDeviceManager = GEngine->GetAudioDeviceManager();
+	if (AudioDeviceManager)
+	{
+		AudioDeviceManager->UpdateActiveAudioDevices(bIsAnyNonPreviewWorldUnpaused);
+	}
+
 	// rendering thread commands
 	{
 		ENQUEUE_UNIQUE_RENDER_COMMAND_TWOPARAMETER(
diff --git a/Engine/Source/Runtime/Engine/Private/GameInstance.cpp b/Engine/Source/Runtime/Engine/Private/GameInstance.cpp
--- a/Engine/Source/Runtime/Engine/Private/GameInstance.cpp
+++ b/Engine/Source/Runtime/Engine/Private/GameInstance.cpp
@@ -577,14 +577,29 @@ ULocalPlayer* UGameInstance::GetLocalPlayerByIndex(const int32 Index) const
 	return LocalPlayers[Index];
 }
 
-APlayerController* UGameInstance::GetFirstLocalPlayerController() const
+APlayerController* UGameInstance::GetFirstLocalPlayerController(UWorld* World) const
 {
-	for (ULocalPlayer* Player : LocalPlayers)
+	if (World == nullptr)
+	{
+		for (ULocalPlayer* Player : LocalPlayers)
+		{
+			// Returns the first non-null UPlayer::PlayerController without filtering by UWorld.
+			if (Player && Player->PlayerController)
+			{
+				// return first non-null entry
+				return Player->PlayerController;
+			}
+		}
+	}
+	else
 	{
-		if (Player && Player->PlayerController)
+		// Only return a local PlayerController from the given World.
+		for (FConstPlayerControllerIterator Iterator = World->GetPlayerControllerIterator(); Iterator; ++Iterator)
 		{
-			// return first non-null entry
-			return Player->PlayerController;
+			if (*Iterator != nullptr && (*Iterator)->IsLocalController())
+			{
+				return *Iterator;
+			}
 		}
 	}
 
@@ -820,9 +835,9 @@ void UGameInstance::StopRecordingReplay()
 	CurrentWorld->DestroyDemoNetDriver();
 }
 
-void UGameInstance::PlayReplay(const FString& Name)
+void UGameInstance::PlayReplay(const FString& Name, UWorld* WorldOverride, const TArray<FString>& AdditionalOptions)
 {
-	UWorld* CurrentWorld = GetWorld();
+	UWorld* CurrentWorld = WorldOverride != nullptr ? WorldOverride : GetWorld();
 
 	if ( CurrentWorld == nullptr )
 	{
@@ -836,6 +851,11 @@ void UGameInstance::PlayReplay(const FString& Name)
 	UE_LOG( LogDemo, Log, TEXT( "PlayReplay: Attempting to play demo %s" ), *Name );
 
 	DemoURL.Map = Name;
+	
+	for (const FString& Option : AdditionalOptions)
+	{
+		DemoURL.AddOption(*Option);
+	}
 
 	const FName NAME_DemoNetDriver( TEXT( "DemoNetDriver" ) );
 
diff --git a/Engine/Source/Runtime/Engine/Private/LocalPlayer.cpp b/Engine/Source/Runtime/Engine/Private/LocalPlayer.cpp
--- a/Engine/Source/Runtime/Engine/Private/LocalPlayer.cpp
+++ b/Engine/Source/Runtime/Engine/Private/LocalPlayer.cpp
@@ -55,7 +55,8 @@ FLocalPlayerContext::FLocalPlayerContext()
 
 }
 
-FLocalPlayerContext::FLocalPlayerContext( const class ULocalPlayer* InLocalPlayer )
+FLocalPlayerContext::FLocalPlayerContext( const class ULocalPlayer* InLocalPlayer, UWorld* InWorld )
+	: World(InWorld)
 {
 	SetLocalPlayer( InLocalPlayer );
 }
@@ -66,6 +67,7 @@ FLocalPlayerContext::FLocalPlayerContext( const class APlayerController* InPlaye
 }
 
 FLocalPlayerContext::FLocalPlayerContext( const FLocalPlayerContext& InPlayerContext )
+	: World(InPlayerContext.World)
 {
 	check(InPlayerContext.GetLocalPlayer());
 	SetLocalPlayer(InPlayerContext.GetLocalPlayer());
@@ -83,6 +85,12 @@ bool FLocalPlayerContext::IsInitialized() const
 
 UWorld* FLocalPlayerContext::GetWorld() const
 {
+	UWorld* WorldPtr = World.Get();
+	if (WorldPtr != nullptr)
+	{
+		return WorldPtr;	
+	}
+
 	check( LocalPlayer.IsValid() );
 	return LocalPlayer->GetWorld();	
 }
@@ -96,34 +104,38 @@ ULocalPlayer* FLocalPlayerContext::GetLocalPlayer() const
 APlayerController* FLocalPlayerContext::GetPlayerController() const
 {
 	check( LocalPlayer.IsValid() );
-	return LocalPlayer->PlayerController;
+	UWorld* WorldPtr = World.Get();
+	return WorldPtr != nullptr ? LocalPlayer->GetPlayerController(WorldPtr) : LocalPlayer->PlayerController;
 }
 
 AGameState* FLocalPlayerContext::GetGameState() const
 {
+	UWorld* WorldPtr = World.Get();
+	if (WorldPtr != nullptr)
+	{
+		return WorldPtr->GameState;
+	}
+	
 	check(LocalPlayer.IsValid());
-	UWorld* World = LocalPlayer->GetWorld();
-	return World ? World->GameState : NULL;
+	UWorld* LocalPlayerWorld = LocalPlayer->GetWorld();
+	return LocalPlayerWorld ? LocalPlayerWorld->GameState : NULL;
 }
 
 APlayerState* FLocalPlayerContext::GetPlayerState() const
 {
-	check(LocalPlayer.IsValid());
-	APlayerController* PC = LocalPlayer->PlayerController;
+	APlayerController* PC = GetPlayerController();
 	return PC ? PC->PlayerState : NULL;
 }
 
 AHUD* FLocalPlayerContext::GetHUD() const
 {
-	check(LocalPlayer.IsValid())
-	APlayerController* PC = LocalPlayer->PlayerController;
+	APlayerController* PC = GetPlayerController();
 	return PC ? PC->MyHUD : NULL;
 }
 
 class APawn* FLocalPlayerContext::GetPawn() const
 {
-	check(LocalPlayer.IsValid())
-	APlayerController* PC = LocalPlayer->PlayerController;
+	APlayerController* PC = GetPlayerController();
 	return PC ? PC->GetPawn() : NULL;
 }
 
@@ -136,6 +148,7 @@ void FLocalPlayerContext::SetPlayerController( const APlayerController* InPlayer
 {
 	check( InPlayerController->IsLocalPlayerController() );
 	LocalPlayer = CastChecked<ULocalPlayer>(InPlayerController->Player);
+	World = InPlayerController->GetWorld();
 }
 
 bool FLocalPlayerContext::IsFromLocalPlayer(const AActor* ActorToTest) const
diff --git a/Engine/Source/Runtime/Engine/Private/NetConnection.cpp b/Engine/Source/Runtime/Engine/Private/NetConnection.cpp
--- a/Engine/Source/Runtime/Engine/Private/NetConnection.cpp
+++ b/Engine/Source/Runtime/Engine/Private/NetConnection.cpp
@@ -1522,9 +1522,9 @@ void UNetConnection::HandleClientPlayer( APlayerController *PC, UNetConnection*
 		break;
 	}
 
-	// Detach old player.
+	// Detach old player if same world.
 	check(LocalPlayer);
-	if( LocalPlayer->PlayerController )
+	if( LocalPlayer->PlayerController && LocalPlayer->PlayerController->GetWorld() == PC->GetWorld() )
 	{
 		if (LocalPlayer->PlayerController->Role == ROLE_Authority)
 		{
diff --git a/Engine/Source/Runtime/Engine/Private/Player.cpp b/Engine/Source/Runtime/Engine/Private/Player.cpp
--- a/Engine/Source/Runtime/Engine/Private/Player.cpp
+++ b/Engine/Source/Runtime/Engine/Private/Player.cpp
@@ -72,6 +72,24 @@ FString UPlayer::ConsoleCommand(const FString& Cmd, bool bWriteToLog)
 	return TEXT("");
 }
 
+APlayerController* UPlayer::GetPlayerController(UWorld* InWorld) const
+{
+	if (InWorld == nullptr)
+	{
+		return PlayerController;
+	}
+
+	for ( FConstPlayerControllerIterator Iterator = InWorld->GetPlayerControllerIterator(); Iterator; ++Iterator)
+	{
+		if ( (*Iterator)->GetLocalPlayer() == this )
+		{
+			return *Iterator;
+		}
+	}
+
+	return nullptr;
+}
+
 bool UPlayer::Exec( UWorld* InWorld, const TCHAR* Cmd,FOutputDevice& Ar)
 {
 	if(PlayerController)
diff --git a/Engine/Source/Runtime/Engine/Private/PlayerController.cpp b/Engine/Source/Runtime/Engine/Private/PlayerController.cpp
--- a/Engine/Source/Runtime/Engine/Private/PlayerController.cpp
+++ b/Engine/Source/Runtime/Engine/Private/PlayerController.cpp
@@ -3640,8 +3640,9 @@ void APlayerController::SetPlayer( UPlayer* InPlayer )
 {
 	check(InPlayer!=NULL);
 
-	// Detach old player.
-	if (InPlayer->PlayerController)
+	const bool bIsSameWorld = InPlayer->PlayerController && (InPlayer->PlayerController->GetWorld() == GetWorld());
+	// Detach old player if same world.
+	if (bIsSameWorld)
 	{
 		InPlayer->PlayerController->Player = NULL;
 	}
diff --git a/Engine/Source/Runtime/Engine/Private/TimerManager.cpp b/Engine/Source/Runtime/Engine/Private/TimerManager.cpp
--- a/Engine/Source/Runtime/Engine/Private/TimerManager.cpp
+++ b/Engine/Source/Runtime/Engine/Private/TimerManager.cpp
@@ -416,6 +416,11 @@ void FTimerManager::Tick(float DeltaTime)
 
 	INC_DWORD_STAT_BY(STAT_NumHeapEntries, ActiveTimerHeap.Num());
 
+	if (HasBeenTickedThisFrame())
+	{
+		return;
+	}
+
 	InternalTime += DeltaTime;
 
 	while (ActiveTimerHeap.Num() > 0)
diff --git a/Engine/Source/Runtime/Engine/Private/UnrealEngine.cpp b/Engine/Source/Runtime/Engine/Private/UnrealEngine.cpp
--- a/Engine/Source/Runtime/Engine/Private/UnrealEngine.cpp
+++ b/Engine/Source/Runtime/Engine/Private/UnrealEngine.cpp
@@ -485,6 +485,59 @@ namespace
 			}
 		}
 	}
+
+	UWorld* CreatePIEWorldByLoadingFromPackage(const FWorldContext& WorldContext, const FString& SourceWorldPackage, UPackage*& OutPackage)
+	{
+		// Load map from the disk in case editor does not have it
+		const FString PIEPackageName = UWorld::ConvertToPIEPackageName(SourceWorldPackage, WorldContext.PIEInstance);
+
+		// Set the world type in the static map, so that UWorld::PostLoad can set the world type
+		const FName PIEPackageFName = FName(*PIEPackageName);
+		UWorld::WorldTypePreLoadMap.FindOrAdd( PIEPackageFName ) = WorldContext.WorldType;
+
+		uint32 LoadFlags = LOAD_None;
+		UPackage* NewPackage = CreatePackage(NULL, *PIEPackageName);
+		if (NewPackage != nullptr && WorldContext.WorldType == EWorldType::PIE)
+		{
+			NewPackage->PackageFlags |= PKG_PlayInEditor;
+			LoadFlags |= LOAD_PackageForPIE;
+		}
+		OutPackage = LoadPackage(NewPackage, *SourceWorldPackage, LoadFlags);
+
+		// Clean up the world type list now that PostLoad has occurred
+		UWorld::WorldTypePreLoadMap.Remove( PIEPackageFName );
+
+		if (OutPackage == nullptr)
+		{
+			return nullptr;
+		}
+
+		UWorld* NewWorld = UWorld::FindWorldInPackage(OutPackage);
+
+		// If the world was not found, follow a redirector if there is one.
+		if ( !NewWorld )
+		{
+			NewWorld = UWorld::FollowWorldRedirectorInPackage(OutPackage);
+			if ( NewWorld )
+			{
+				OutPackage = NewWorld->GetOutermost();
+			}
+		}
+
+		check(NewWorld);
+
+#if WITH_EDITOR
+		OutPackage->PIEInstanceID = WorldContext.PIEInstance;
+#endif
+
+		// Rename streaming levels to PIE
+		for (ULevelStreaming* StreamingLevel : NewWorld->StreamingLevels)
+		{
+			StreamingLevel->RenameForPIE(WorldContext.PIEInstance);
+		}
+
+		return NewWorld;
+	}
 }
 /*-----------------------------------------------------------------------------
 	Object class implementation.
@@ -1964,7 +2017,7 @@ APlayerController* UEngine::GetFirstLocalPlayerController(UWorld *InWorld)
 {
 	const FWorldContext &Context = GetWorldContextFromWorldChecked(InWorld);
 	
-	return ( Context.OwningGameInstance != NULL ) ? Context.OwningGameInstance->GetFirstLocalPlayerController() : NULL;
+	return ( Context.OwningGameInstance != NULL ) ? Context.OwningGameInstance->GetFirstLocalPlayerController(InWorld) : NULL;
 }
 
 void UEngine::GetAllLocalPlayerControllers(TArray<APlayerController*> & PlayerList)
@@ -9193,13 +9246,13 @@ bool UEngine::LoadMap( FWorldContext& WorldContext, FURL URL, class UPendingNetG
 		// without a call to RemoveFromWorld for each)
 		FWorldDelegates::LevelRemovedFromWorld.Broadcast(NULL, WorldContext.World());
 
-		// Disassociate the players from their PlayerControllers.
+		// Disassociate the players from their PlayerControllers in this world.
 		if (WorldContext.OwningGameInstance != NULL)
 		{
 			for(auto It = WorldContext.OwningGameInstance->GetLocalPlayerIterator(); It; ++It)
 			{
 				ULocalPlayer *Player = *It;
-				if(Player->PlayerController)
+				if(Player->PlayerController && Player->PlayerController->GetWorld() == WorldContext.World())
 				{
 					if(Player->PlayerController->GetPawn())
 					{
@@ -9355,52 +9408,12 @@ bool UEngine::LoadMap( FWorldContext& WorldContext, FURL URL, class UPendingNetG
 			NewWorld = UWorld::DuplicateWorldForPIE(SourceWorldPackage, NULL);
 			if (NewWorld == nullptr) 
 			{
-				// Load map from the disk in case editor does not have it
-				const FString PIEPackageName = *UWorld::ConvertToPIEPackageName(SourceWorldPackage, WorldContext.PIEInstance);
-
-				// Set the world type in the static map, so that UWorld::PostLoad can set the world type
-				const FName PIEPackageFName = FName(*PIEPackageName);
-				UWorld::WorldTypePreLoadMap.FindOrAdd( PIEPackageFName ) = WorldContext.WorldType;
-
-				uint32 LoadFlags = LOAD_None;
-				auto NewPackage = CreatePackage(NULL, *PIEPackageName);
-				if (NewPackage != nullptr && WorldContext.WorldType == EWorldType::PIE)
-				{
-					NewPackage->PackageFlags |= PKG_PlayInEditor;
-					LoadFlags |= LOAD_PackageForPIE;
-				}
-				WorldPackage = LoadPackage(NewPackage, *SourceWorldPackage, LoadFlags);
-
-				// Clean up the world type list now that PostLoad has occurred
-				UWorld::WorldTypePreLoadMap.Remove( PIEPackageFName );
-
-				if (WorldPackage == nullptr)
+				NewWorld = CreatePIEWorldByLoadingFromPackage(WorldContext, SourceWorldPackage, WorldPackage);
+				if (NewWorld == nullptr)
 				{
 					Error = FString::Printf(TEXT("Failed to load package '%s' while in PIE"), *SourceWorldPackage);
 					return false;
 				}
-
-				NewWorld = UWorld::FindWorldInPackage(WorldPackage);
-
-				// If the world was not found, follow a redirector if there is one.
-				if ( !NewWorld )
-				{
-					NewWorld = UWorld::FollowWorldRedirectorInPackage(WorldPackage);
-					if ( NewWorld )
-					{
-						WorldPackage = NewWorld->GetOutermost();
-					}
-				}
-
-				check(NewWorld);
-#if WITH_EDITOR
-				WorldPackage->PIEInstanceID = WorldContext.PIEInstance;
-#endif			
-				// Rename streaming levels to PIE
-				for (auto StreamingLevel : NewWorld->StreamingLevels)
-				{
-					StreamingLevel->RenameForPIE(WorldContext.PIEInstance);
-				}
 			}
 			else
 			{
@@ -9477,6 +9490,23 @@ bool UEngine::LoadMap( FWorldContext& WorldContext, FURL URL, class UPendingNetG
 				NewWorld->RenameToPIEWorld(WorldContext.PIEInstance);
 			}
 		}
+		else if (WorldContext.WorldType == EWorldType::Game)
+		{
+			// If we are a game world and the world we just found is already initialized, then we're probably trying to load
+			// an independent fresh copy of the world into a different context. Create a package with a prefixed name
+			// and load the world from disk to keep the instances independent. If this is the case, assume the creator
+			// of the FWorldContext was aware and set WorldContext.PIEInstance to a valid value.
+			if (NewWorld->bIsWorldInitialized && WorldContext.PIEInstance != -1)
+			{
+				NewWorld = CreatePIEWorldByLoadingFromPackage(WorldContext, URL.Map, WorldPackage);
+
+				if (NewWorld == nullptr)
+				{
+					Error = FString::Printf(TEXT("Failed to load package '%s' into a new game world."), *URL.Map);
+					return false;
+				}
+			}
+		}
 	}
 	NewWorld->SetGameInstance(WorldContext.OwningGameInstance);
 
@@ -10175,21 +10205,25 @@ bool UEngine::WorldHasValidContext(UWorld *InWorld)
 
 void UEngine::VerifyLoadMapWorldCleanup()
 {
-	// All worlds at this point should be the CurrentWorld of some context or preview worlds.
+	// All worlds at this point should be the CurrentWorld of some context, preview worlds, or streaming level
+	// worlds that are owned by the CurrentWorld of some context.
 	for( TObjectIterator<UWorld> It; It; ++It )
 	{
 		UWorld* World = *It;
 		const bool bIsPersistantWorldType = (World->WorldType == EWorldType::Inactive) || (World->WorldType == EWorldType::Preview);
 		if (!bIsPersistantWorldType && !WorldHasValidContext(World))
 		{
-			// Print some debug information...
-			UE_LOG(LogLoad, Log, TEXT("%s not cleaned up by garbage collection! "), *World->GetFullName());
-			StaticExec(World, *FString::Printf(TEXT("OBJ REFS CLASS=WORLD NAME=%s"), *World->GetPathName()));
-			TMap<UObject*,UProperty*>	Route		= FArchiveTraceRoute::FindShortestRootPath( World, true, GARBAGE_COLLECTION_KEEPFLAGS );
-			FString						ErrorString	= FArchiveTraceRoute::PrintRootPath( Route, World );
-			UE_LOG(LogLoad, Log, TEXT("%s"),*ErrorString);
-			// before asserting.
-			UE_LOG(LogLoad, Fatal, TEXT("%s not cleaned up by garbage collection!") LINE_TERMINATOR TEXT("%s") , *World->GetFullName(), *ErrorString );
+			if (World->PersistentLevel != nullptr && !WorldHasValidContext(World->PersistentLevel->OwningWorld))
+			{
+				// Print some debug information...
+				UE_LOG(LogLoad, Log, TEXT("%s not cleaned up by garbage collection! "), *World->GetFullName());
+				StaticExec(World, *FString::Printf(TEXT("OBJ REFS CLASS=WORLD NAME=%s"), *World->GetPathName()));
+				TMap<UObject*,UProperty*>	Route		= FArchiveTraceRoute::FindShortestRootPath( World, true, GARBAGE_COLLECTION_KEEPFLAGS );
+				FString						ErrorString	= FArchiveTraceRoute::PrintRootPath( Route, World );
+				UE_LOG(LogLoad, Log, TEXT("%s"),*ErrorString);
+				// before asserting.
+				UE_LOG(LogLoad, Fatal, TEXT("%s not cleaned up by garbage collection!") LINE_TERMINATOR TEXT("%s") , *World->GetFullName(), *ErrorString );
+			}
 		}
 	}
 }
diff --git a/Engine/Source/Runtime/Engine/Private/World.cpp b/Engine/Source/Runtime/Engine/Private/World.cpp
--- a/Engine/Source/Runtime/Engine/Private/World.cpp
+++ b/Engine/Source/Runtime/Engine/Private/World.cpp
@@ -1913,7 +1913,8 @@ void UWorld::AddToWorld( ULevel* Level, const FTransform& LevelTransform )
 		{
 			for (FLocalPlayerIterator It(GEngine, this); It; ++It)
 			{
-				if (It->PlayerController != NULL)
+				APlayerController* LocalPlayerController = It->GetPlayerController(this);
+				if (LocalPlayerController != NULL)
 				{
 					// Remap packagename for PIE networking before sending out to server
 					FName PackageName = Level->GetOutermost()->GetFName();
@@ -1923,7 +1924,7 @@ void UWorld::AddToWorld( ULevel* Level, const FTransform& LevelTransform )
 						PackageName = FName(*PackageNameStr);
 					}
 
-					It->PlayerController->ServerUpdateLevelVisibility(PackageName, true);
+					LocalPlayerController->ServerUpdateLevelVisibility(PackageName, true);
 				}
 			}
 		}
@@ -2029,7 +2030,8 @@ void UWorld::RemoveFromWorld( ULevel* Level )
 		{
 			for (FLocalPlayerIterator It(GEngine, this); It; ++It)
 			{
-				if (It->PlayerController != NULL)
+				APlayerController* LocalPlayerController = It->GetPlayerController(this);
+				if (LocalPlayerController != NULL)
 				{
 					// Remap packagename for PIE networking before sending out to server
 					FName PackageName = Level->GetOutermost()->GetFName();
@@ -2039,7 +2041,7 @@ void UWorld::RemoveFromWorld( ULevel* Level )
 						PackageName = FName(*PackageNameStr);
 					}
 
-					It->PlayerController->ServerUpdateLevelVisibility(PackageName, false);
+					LocalPlayerController->ServerUpdateLevelVisibility(PackageName, false);
 				}
 			}
 		}
diff --git a/Engine/Source/Runtime/Online/OnlineSubsystemUtils/Private/OnlineSessionClient.cpp b/Engine/Source/Runtime/Online/OnlineSubsystemUtils/Private/OnlineSessionClient.cpp
--- a/Engine/Source/Runtime/Online/OnlineSubsystemUtils/Private/OnlineSessionClient.cpp
+++ b/Engine/Source/Runtime/Online/OnlineSubsystemUtils/Private/OnlineSessionClient.cpp
@@ -194,7 +194,7 @@ void UOnlineSessionClient::OnJoinSessionComplete(FName SessionName, EOnJoinSessi
 			FString URL;
 			if (SessionInt->GetResolvedConnectString(SessionName, URL))
 			{
-				APlayerController* PC = GetGameInstance()->GetFirstLocalPlayerController();;
+				APlayerController* PC = GetGameInstance()->GetFirstLocalPlayerController(GetWorld());
 				if (PC)
 				{
 					if (bIsFromInvite)
diff --git a/Engine/Source/Runtime/UMG/Private/Components/WidgetComponent.cpp b/Engine/Source/Runtime/UMG/Private/Components/WidgetComponent.cpp
--- a/Engine/Source/Runtime/UMG/Private/Components/WidgetComponent.cpp
+++ b/Engine/Source/Runtime/UMG/Private/Components/WidgetComponent.cpp
@@ -851,7 +851,7 @@ void UWidgetComponent::TickComponent(float DeltaTime, enum ELevelTick TickType,
 								
 								bAddedToScreen = true;
 								
-								Widget->SetPlayerContext(TargetPlayer);
+								Widget->SetPlayerContext(FLocalPlayerContext(TargetPlayer, GetWorld()));
 								ScreenLayer->AddComponent(this);
 							}
 						}
diff --git a/Engine/Source/Runtime/UMG/Private/UserWidget.cpp b/Engine/Source/Runtime/UMG/Private/UserWidget.cpp
--- a/Engine/Source/Runtime/UMG/Private/UserWidget.cpp
+++ b/Engine/Source/Runtime/UMG/Private/UserWidget.cpp
@@ -609,7 +609,7 @@ void UUserWidget::SetOwningLocalPlayer(ULocalPlayer* LocalPlayer)
 {
 	if ( LocalPlayer )
 	{
-		PlayerContext = FLocalPlayerContext(LocalPlayer);
+		PlayerContext = FLocalPlayerContext(LocalPlayer, GetWorld());
 	}
 }
 
