Attempt to spawn a PlayerState for the replay spectator controller during client replay recording, by falling back to the default game mode object if the authority game mode can't be found (only if GetNetMode() != NM_Client).
Expose the ability to get an actor's FNetworkGUID, and the actor that corresponds to an FNetworkGUID, if available. Allows games to correlate actors across multiple local worlds if there is a client-server relationship between them.

--------
Integrated using branch UE4-Orion-To-UE4 of change#2705235 by [REDACTED_PERSON] on 2015/09/24 17:12:18.

diff --git a/Engine/Source/Runtime/CoreUObject/Public/UObject/CoreNet.h b/Engine/Source/Runtime/CoreUObject/Public/UObject/CoreNet.h
--- a/Engine/Source/Runtime/CoreUObject/Public/UObject/CoreNet.h
+++ b/Engine/Source/Runtime/CoreUObject/Public/UObject/CoreNet.h
@@ -150,9 +150,10 @@ class COREUOBJECT_API UPackageMap : public UObject
 	void							ResetTrackedUnmappedGuids( bool bShouldTrack ) { TrackedUnmappedNetGuids.Empty(); bShouldTrackUnmappedGuids = bShouldTrack; }
 	const TArray< FNetworkGUID > &	GetTrackedUnmappedGuids() const { return TrackedUnmappedNetGuids; }
 
-	virtual void		LogDebugInfo( FOutputDevice & Ar) { }
-	virtual UObject*	GetObjectFromNetGUID( const FNetworkGUID& NetGUID, const bool bIgnoreMustBeMapped ) { return NULL; }
-	virtual bool		IsGUIDBroken( const FNetworkGUID& NetGUID, const bool bMustBeRegistered ) const { return false; }
+	virtual void			LogDebugInfo( FOutputDevice & Ar) { }
+	virtual UObject*		GetObjectFromNetGUID( const FNetworkGUID& NetGUID, const bool bIgnoreMustBeMapped ) { return NULL; }
+	virtual FNetworkGUID	GetNetGUIDFromObject( const UObject* InObject) const { return FNetworkGUID(); }
+	virtual bool			IsGUIDBroken( const FNetworkGUID& NetGUID, const bool bMustBeRegistered ) const { return false; }
 
 protected:
 
diff --git a/Engine/Source/Runtime/Engine/Classes/Engine/DemoNetDriver.h b/Engine/Source/Runtime/Engine/Classes/Engine/DemoNetDriver.h
--- a/Engine/Source/Runtime/Engine/Classes/Engine/DemoNetDriver.h
+++ b/Engine/Source/Runtime/Engine/Classes/Engine/DemoNetDriver.h
@@ -2,6 +2,7 @@
 
 #pragma once
 
+#include "NetDriver.h"
 #include "Runtime/Online/HTTP/Public/Interfaces/IHttpRequest.h"
 #include "Runtime/NetworkReplayStreaming/NetworkReplayStreaming/Public/NetworkReplayStreaming.h"
 #include "DemoNetDriver.generated.h"
@@ -124,6 +125,8 @@ public:
 	virtual bool ShouldClientDestroyTearOffActors() const override;
 	virtual bool ShouldSkipRepNotifies() const override;
 	virtual bool ShouldQueueBunchesForActorGUID(FNetworkGUID InGUID) const override;
+	virtual FNetworkGUID GetGUIDForActor(const AActor* InActor) const override;
+	virtual AActor* GetActorForGUID(FNetworkGUID InGUID) const override;
 
 	void GotoTimeInSeconds(const float TimeInSeconds, const FOnGotoTimeDelegate& InOnGotoTimeDelegate = FOnGotoTimeDelegate());
 
@@ -177,4 +180,6 @@ public:
 	void ClearReplayTasks();
 	bool ProcessReplayTasks();
 	bool IsNamedTaskInQueue( const FString& Name );
+
+	void AddNonQueuedGUIDForScrubbing(FNetworkGUID InGUID);
 };
diff --git a/Engine/Source/Runtime/Engine/Classes/Engine/NetDriver.h b/Engine/Source/Runtime/Engine/Classes/Engine/NetDriver.h
--- a/Engine/Source/Runtime/Engine/Classes/Engine/NetDriver.h
+++ b/Engine/Source/Runtime/Engine/Classes/Engine/NetDriver.h
@@ -553,6 +553,12 @@ public:
 	/** Returns true if actor channels with InGUID should queue up bunches, even if they wouldn't otherwise be queued. */
 	virtual bool ShouldQueueBunchesForActorGUID(FNetworkGUID InGUID) const { return false; }
 
+	/** Returns the existing FNetworkGUID of InActor, if it has one. */
+	virtual FNetworkGUID GetGUIDForActor(const AActor* InActor) const { return FNetworkGUID(); }
+
+	/** Returns the actor that corresponds to InGUID, if one can be found. */
+	virtual AActor* GetActorForGUID(FNetworkGUID InGUID) const { return nullptr; }
+
 protected:
 
 	/** Adds (fully initialized, ready to go) client connection to the ClientConnections list + any other game related setup */
diff --git a/Engine/Source/Runtime/Engine/Classes/Engine/PackageMapClient.h b/Engine/Source/Runtime/Engine/Classes/Engine/PackageMapClient.h
--- a/Engine/Source/Runtime/Engine/Classes/Engine/PackageMapClient.h
+++ b/Engine/Source/Runtime/Engine/Classes/Engine/PackageMapClient.h
@@ -57,6 +57,7 @@ public:
 	bool			IsDynamicObject( const UObject* Object );
 	bool			IsNetGUIDAuthority() const;
 	FNetworkGUID	GetOrAssignNetGUID( const UObject* Object );
+	FNetworkGUID	GetNetGUID( const UObject* Object ) const;
 	FNetworkGUID	AssignNewNetGUID_Server( const UObject* Object );
 	void			RegisterNetGUID_Internal( const FNetworkGUID& NetGUID, const FNetGuidCacheObject& CacheObject );
 	void			RegisterNetGUID_Server( const FNetworkGUID& NetGUID, const UObject* Object );
@@ -150,9 +151,10 @@ public:
 
 	virtual bool PrintExportBatch() override;
 
-	virtual void		LogDebugInfo( FOutputDevice & Ar) override;
-	virtual UObject *	GetObjectFromNetGUID( const FNetworkGUID& NetGUID, const bool bIgnoreMustBeMapped ) override;
-	virtual bool		IsGUIDBroken( const FNetworkGUID& NetGUID, const bool bMustBeRegistered ) const override { return GuidCache->IsGUIDBroken( NetGUID, bMustBeRegistered ); }
+	virtual void			LogDebugInfo( FOutputDevice & Ar) override;
+	virtual UObject *		GetObjectFromNetGUID( const FNetworkGUID& NetGUID, const bool bIgnoreMustBeMapped ) override;
+	virtual FNetworkGUID	GetNetGUIDFromObject( const UObject* InObject) const override;
+	virtual bool			IsGUIDBroken( const FNetworkGUID& NetGUID, const bool bMustBeRegistered ) const override { return GuidCache->IsGUIDBroken( NetGUID, bMustBeRegistered ); }
 
 	TArray< FNetworkGUID > & GetMustBeMappedGuidsInLastBunch() { return MustBeMappedGuidsInLastBunch; }
 
diff --git a/Engine/Source/Runtime/Engine/Private/Controller.cpp b/Engine/Source/Runtime/Engine/Private/Controller.cpp
--- a/Engine/Source/Runtime/Engine/Private/Controller.cpp
+++ b/Engine/Source/Runtime/Engine/Private/Controller.cpp
@@ -11,6 +11,7 @@
 #include "VisualLogger/VisualLogger.h"
 #include "GameFramework/Character.h"
 #include "GameFramework/GameMode.h"
+#include "GameFramework/GameState.h"
 #include "GameFramework/PawnMovementComponent.h"
 #include "MessageLog.h"
 
@@ -438,7 +439,17 @@ void AController::InitPlayerState()
 	if ( GetNetMode() != NM_Client )
 	{
 		UWorld* const World = GetWorld();
-		AGameMode* const GameMode = World ? World->GetAuthGameMode() : NULL;
+		const AGameMode* GameMode = World ? World->GetAuthGameMode() : NULL;
+
+		// If the GameMode is null, this might be a network client that's trying to
+		// record a replay. Try to use the default game mode in this case so that
+		// we can still spawn a PlayerState.
+		if (GameMode == NULL)
+		{
+			const AGameState* const GameState = World ? World->GetGameState() : NULL;
+			GameMode = GameState ? GameState->GetDefaultGameMode() : NULL;
+		}
+
 		if (GameMode != NULL)
 		{
 			FActorSpawnParameters SpawnInfo;
diff --git a/Engine/Source/Runtime/Engine/Private/DemoNetDriver.cpp b/Engine/Source/Runtime/Engine/Private/DemoNetDriver.cpp
--- a/Engine/Source/Runtime/Engine/Private/DemoNetDriver.cpp
+++ b/Engine/Source/Runtime/Engine/Private/DemoNetDriver.cpp
@@ -1031,12 +1031,14 @@ static void DemoReplicateActor(AActor* Actor, UNetConnection* Connection, bool I
 		//  1. We're recording a replay on a client that's connected to a live server, and
 		//  2. the actor isn't bTearOff, and
 		//  3. the actor isn't the replay spectator controller.
+		//  3. the actor isn't the replay spectator controller or its PlayerState.
 		// This is to ensure the roles appear correct when playing back this demo.
 		UWorld* ActorWorld = Actor->GetWorld();
 		bool bShouldSwapRoles =
 			ActorWorld != nullptr && ActorWorld->IsRecordingClientReplay() &&
 			!Actor->bTearOff &&
-			Actor != SpectatorController;
+			Actor != SpectatorController &&
+			Actor != SpectatorController->PlayerState;
 
 		FScopedActorRoleSwap RoleSwap(bShouldSwapRoles ? Actor : nullptr);
 
@@ -1131,6 +1133,8 @@ static void SerializeGuidCache( TSharedPtr< class FNetGUIDCache > GuidCache, FAr
 
 void UDemoNetDriver::SaveCheckpoint()
 {
+	DECLARE_SCOPE_CYCLE_COUNTER(TEXT("SaveCheckpoint time"), STAT_ReplayCheckpointSaveTime, STATGROUP_Net);
+
 	FArchive* CheckpointArchive = ReplayStreamer->GetCheckpointArchive();
 
 	if ( CheckpointArchive == nullptr )
@@ -1698,6 +1702,9 @@ void UDemoNetDriver::FinalizeFastForward( const float StartTime )
 		}
 	}
 
+	// Reset the never-queue GUID list, we'll rebuild it
+	NonQueuedGUIDsForScrubbing.Reset();
+
 	const auto FastForwardTotalSeconds = FPlatformTime::Seconds() - StartTime;
 
 	OnGotoTimeDelegate.ExecuteIfBound( true );
@@ -1746,6 +1753,17 @@ void UDemoNetDriver::SpawnDemoRecSpectator( UNetConnection* Connection, const FU
 		return;
 	}
 
+	// Make sure SpectatorController->GetNetDriver returns this driver. Ensures functions that depend on it,
+	// such as IsLocalController, work as expected.
+	SpectatorController->NetDriverName = NetDriverName;
+
+	// If the controller doesn't have a player state, we are probably recording on a client.
+	// Spawn one manually.
+	if ( SpectatorController->PlayerState == nullptr && GetWorld() != nullptr && GetWorld()->IsRecordingClientReplay())
+	{
+		SpectatorController->InitPlayerState();
+	}
+
 	for ( FActorIterator It( World ); It; ++It)
 	{
 		if ( It->IsA( APlayerStart::StaticClass() ) )
@@ -1754,10 +1772,7 @@ void UDemoNetDriver::SpawnDemoRecSpectator( UNetConnection* Connection, const FU
 			break;
 		}
 	}
-
-	// Make sure SpectatorController->GetNetDriver returns this driver. Ensures functions that depend on it,
-	// such as IsLocalController, work as expected.
-	SpectatorController->NetDriverName = NetDriverName;
+	
 	SpectatorController->SetReplicates( true );
 	SpectatorController->SetAutonomousProxy( true );
 
@@ -1824,9 +1839,6 @@ void UDemoNetDriver::LoadCheckpoint( FArchive* GotoCheckpointArchive, int64 Goto
 	check( !bIsFastForwardingForCheckpoint );
 	check( !bIsFastForwarding );
 
-	// Reset the never-queue GUID list, we'll rebuild it
-	NonQueuedGUIDsForScrubbing.Empty();
-
 	// Save off the current spectator position
 	// Check for NULL, which can be the case if we haven't played any of the demo yet but want to fast forward (joining live game for example)
 	if ( SpectatorController != NULL )
@@ -2034,6 +2046,43 @@ bool UDemoNetDriver::ShouldQueueBunchesForActorGUID(FNetworkGUID InGUID) const
 	return false;
 }
 
+FNetworkGUID UDemoNetDriver::GetGUIDForActor(const AActor* InActor) const
+{
+	UNetConnection* Connection = ServerConnection;
+	
+	if ( ClientConnections.Num() > 0)
+	{
+		Connection = ClientConnections[0];
+	}
+
+	if ( !Connection )
+	{
+		return FNetworkGUID();
+	}
+
+	FNetworkGUID Guid = Connection->PackageMap->GetNetGUIDFromObject(InActor);
+	return Guid;
+}
+
+AActor* UDemoNetDriver::GetActorForGUID(FNetworkGUID InGUID) const
+{
+	UNetConnection* Connection = ServerConnection;
+	
+	if ( ClientConnections.Num() > 0)
+	{
+		Connection = ClientConnections[0];
+	}
+
+	if ( !Connection )
+	{
+		return nullptr;
+	}
+
+	UObject* FoundObject = Connection->PackageMap->GetObjectFromNetGUID(InGUID, true);
+	return Cast<AActor>(FoundObject);
+
+}
+
 void UDemoNetDriver::AddNonQueuedActorForScrubbing(AActor* Actor)
 {
 	auto FoundChannel = ServerConnection->ActorChannels.Find(Actor);
@@ -2044,6 +2093,14 @@ void UDemoNetDriver::AddNonQueuedActorForScrubbing(AActor* Actor)
 	}
 }
 
+void UDemoNetDriver::AddNonQueuedGUIDForScrubbing(FNetworkGUID InGUID)
+{
+	if (InGUID.IsValid())
+	{
+		NonQueuedGUIDsForScrubbing.Add(InGUID);
+	}
+}
+
 /*-----------------------------------------------------------------------------
 	UDemoNetConnection.
 -----------------------------------------------------------------------------*/
diff --git a/Engine/Source/Runtime/Engine/Private/PackageMapClient.cpp b/Engine/Source/Runtime/Engine/Private/PackageMapClient.cpp
--- a/Engine/Source/Runtime/Engine/Private/PackageMapClient.cpp
+++ b/Engine/Source/Runtime/Engine/Private/PackageMapClient.cpp
@@ -192,6 +192,8 @@ bool UPackageMapClient::SerializeNewActor(FArchive& Ar, class UActorChannel *Cha
 		return false;
 	}
 
+	Channel->ActorNetGUID = NetGUID;
+
 	Actor = Cast<AActor>(NewObj);
 
 	// When we return an actor, we don't necessarily always spawn it (we might have found it already in memory)
@@ -1347,6 +1349,11 @@ UObject* UPackageMapClient::GetObjectFromNetGUID( const FNetworkGUID& NetGUID, c
 	return GuidCache->GetObjectFromNetGUID( NetGUID, bIgnoreMustBeMapped );
 }
 
+FNetworkGUID UPackageMapClient::GetNetGUIDFromObject(const UObject* InObject) const
+{
+	return GuidCache->GetNetGUID(InObject);
+}
+
 //----------------------------------------------------------------------------------------
 //	FNetGUIDCache
 //----------------------------------------------------------------------------------------
@@ -1520,6 +1527,18 @@ FNetworkGUID FNetGUIDCache::GetOrAssignNetGUID( const UObject* Object )
 	return AssignNewNetGUID_Server( Object );
 }
 
+FNetworkGUID FNetGUIDCache::GetNetGUID(const UObject* Object) const
+{
+	if ( !Object || !SupportsObject( Object ) )
+	{
+		// Null of unsupported object, leave as default NetGUID and just return mapped=true
+		return FNetworkGUID();
+	}
+
+	FNetworkGUID NetGUID = NetGUIDLookup.FindRef( Object );
+	return NetGUID;
+}
+
 /**
  *	Generate a new NetGUID for this object and assign it.
  */
