Sequencer: Implemented name curves

diff --git a/Engine/Source/Runtime/Engine/Classes/Curves/CurveBase.h b/Engine/Source/Runtime/Engine/Classes/Curves/CurveBase.h
--- a/Engine/Source/Runtime/Engine/Classes/Curves/CurveBase.h
+++ b/Engine/Source/Runtime/Engine/Classes/Curves/CurveBase.h
@@ -145,6 +145,164 @@ protected:
 };
 
 
+/* Name curves
+ *****************************************************************************/
+
+/**
+ * One key in a curve of FNames.
+ */
+USTRUCT()
+struct ENGINE_API FNameCurveKey
+{
+	GENERATED_USTRUCT_BODY()
+
+	/** Time at this key */
+	UPROPERTY(EditAnywhere, Category="Key")
+	float Time;
+
+	/** Value at this key */
+	UPROPERTY(EditAnywhere, Category="Key")
+	FName Value;
+
+	/** Default constructor. */
+	FNameCurveKey()
+		: Time(0.0f)
+		, Value(NAME_None)
+	{ }
+
+	/** Creates and initializes a new instance. */
+	FNameCurveKey(float InTime, const FName& InValue)
+		: Time(InTime)
+		, Value(InValue)
+	{ }
+
+public:
+
+	// TStructOpsTypeTraits interface
+
+	bool operator==(const FNameCurveKey& Other) const;
+	bool operator!=(const FNameCurveKey& Other) const;
+	bool Serialize(FArchive& Ar);
+
+	/** Serializes a name curve key from or into an archive. */
+	friend FArchive& operator<<(FArchive& Ar, FNameCurveKey& Key)
+	{
+		Key.Serialize(Ar);
+		return Ar;
+	}
+};
+
+
+template<>
+struct TIsPODType<FNameCurveKey>
+{
+	enum { Value = true };
+};
+
+
+template<>
+struct TStructOpsTypeTraits<FNameCurveKey>
+	: public TStructOpsTypeTraitsBase
+{
+	enum
+	{
+		WithSerializer = true,
+		WithCopy = false,
+		WithIdenticalViaEquality = true,
+	};
+};
+
+
+/**
+ * Implements a curve of FNames.
+ */
+USTRUCT()
+struct ENGINE_API FNameCurve
+	: public FIndexedCurve
+{
+	GENERATED_USTRUCT_BODY()
+
+	/** Virtual destructor. */
+	virtual ~FNameCurve() { }
+
+public:
+
+	/**
+	  * Add a new key to the curve with the supplied Time and Value. Returns the handle of the new key.
+	  * 
+	  * @param InTime The time at which to add the key.
+	  * @param InValue The value of the key.
+	  * @param KeyHandle Optional handle for the new key.
+	  */
+	FKeyHandle AddKey(float InTime, const FName& InValue, FKeyHandle KeyHandle = FKeyHandle());
+
+	/**
+	 * Remove the specified key from the curve.
+	 *
+	 * @param KeyHandle Handle to the key to remove.
+	 */
+	void DeleteKey(FKeyHandle KeyHandle);
+
+	/**
+	 * Get a key.
+	 *
+	 * @param KeyHandle The handle of the key to get.
+	 * @return The key.
+	 */
+	FNameCurveKey& GetKey(FKeyHandle KeyHandle);
+	FNameCurveKey GetKey(FKeyHandle KeyHandle) const;
+
+	/**
+	 * Get the time for the Key with the specified index.
+	 *
+	 * @param KeyHandle Handle to the key whose time to get.
+	 * @return The key's time.
+	 */
+	float GetKeyTime(FKeyHandle KeyHandle) const;
+
+	/**
+	 * Move a key to a new time.
+	 *
+	 * This may change the index of the key, so the new key index is returned.
+	 *
+	 * @param KeyHandle The handle of the key to change.
+	 * @param NewTime The new time to set on the key.
+	 */
+	FKeyHandle SetKeyTime(FKeyHandle KeyHandle, float NewTime);
+
+	/**
+	 * Finds the key at InTime, and updates its value. If it can't find the key, it adds one at that time.
+	 *
+	 * @param InTime The time at which the key should be added or updated.
+	 * @param InValue The value of the key.
+	 */
+	FKeyHandle UpdateOrAddKey(float InTime, const FName& InValue);
+
+public:
+
+	/** Shifts all keys forwards or backwards in time by an even amount, preserving order. */
+	void ShiftCurve(float DeltaTime);
+	void ShiftCurve(float DeltaTime, TSet<FKeyHandle>& KeyHandles);
+	
+	/** Scales all keys about an origin, preserving order. */
+	void ScaleCurve(float ScaleOrigin, float ScaleFactor);
+	void ScaleCurve(float ScaleOrigin, float ScaleFactor, TSet<FKeyHandle>& KeyHandles);
+
+public:
+
+	// FIndexedCurve interface
+
+	virtual int32 GetNumKeys() const override;
+	virtual bool IsKeyHandleValid(FKeyHandle KeyHandle) const override;
+
+public:
+
+	/** Sorted array of keys */
+	UPROPERTY(EditAnywhere, EditFixedSize, Category="Curve")
+	TArray<FNameCurveKey> Keys;
+};
+
+
 /* Rich curves
  *****************************************************************************/
 
@@ -316,8 +474,6 @@ struct ENGINE_API FRichCurve
 {
 	GENERATED_USTRUCT_BODY()
 
-public:
-
 	FRichCurve() 
 		: FIndexedCurve()
 		, PreInfinityExtrap(RCCE_Constant)
@@ -328,6 +484,8 @@ public:
 	/** Virtual destructor. */
 	virtual ~FRichCurve() { }
 
+public:
+
 	/** Gets a copy of the keys, so indices and handles can't be meddled with */
 	TArray<FRichCurveKey> GetCopyOfKeys() const;
 
diff --git a/Engine/Source/Runtime/Engine/Private/CurveBase.cpp b/Engine/Source/Runtime/Engine/Private/CurveBase.cpp
--- a/Engine/Source/Runtime/Engine/Private/CurveBase.cpp
+++ b/Engine/Source/Runtime/Engine/Private/CurveBase.cpp
@@ -155,6 +155,217 @@ void FIndexedCurve::EnsureAllIndicesHaveHandles() const
 }
 
 
+/* FNameCurveKey
+ *****************************************************************************/
+
+bool FNameCurveKey::operator==(const FNameCurveKey& Curve) const
+{
+	return ((Time == Curve.Time) && (Value == Curve.Value));
+}
+
+
+bool FNameCurveKey::operator!=(const FNameCurveKey& Other) const
+{
+	return !(*this == Other);
+}
+
+
+bool FNameCurveKey::Serialize(FArchive& Ar)
+{
+	Ar << Time << Value;
+	return true;
+}
+
+
+/* FNameCurve
+ *****************************************************************************/
+
+FKeyHandle FNameCurve::AddKey(float InTime, const FName& InValue, FKeyHandle KeyHandle)
+{
+	int32 Index = 0;
+
+	// insert key
+	for(; Index < Keys.Num() && Keys[Index].Time < InTime; ++Index);
+	Keys.Insert(FNameCurveKey(InTime, InValue), Index);
+
+	// update key indices
+	for ( auto It = KeyHandlesToIndices.CreateIterator(); It; ++It)
+	{
+		const FKeyHandle& KeyHandle = It.Key();
+		int32& KeyIndex = It.Value();
+
+		if (KeyIndex >= Index) {++KeyIndex;}
+	}
+
+	KeyHandlesToIndices.Add(KeyHandle, Index);
+
+	return GetKeyHandle(Index);
+}
+
+
+void FNameCurve::DeleteKey(FKeyHandle KeyHandle)
+{
+	// remove key
+	int32 Index = GetIndex(KeyHandle);
+	Keys.RemoveAt(Index);
+
+	// update key indices
+	for (auto It = KeyHandlesToIndices.CreateIterator(); It; ++It)
+	{
+		const FKeyHandle& KeyHandle = It.Key();
+		int32& KeyIndex = It.Value();
+
+		if (KeyIndex >= Index)
+		{
+			--KeyIndex;
+		}
+	}
+}
+
+
+FNameCurveKey& FNameCurve::GetKey(FKeyHandle KeyHandle)
+{
+	EnsureAllIndicesHaveHandles();
+	return Keys[GetIndex(KeyHandle)];
+}
+
+
+FNameCurveKey FNameCurve::GetKey(FKeyHandle KeyHandle) const
+{
+	EnsureAllIndicesHaveHandles();
+	return Keys[GetIndex(KeyHandle)];
+}
+
+
+float FNameCurve::GetKeyTime(FKeyHandle KeyHandle) const
+{
+	if (!IsKeyHandleValid(KeyHandle))
+	{
+		return 0.f;
+	}
+
+	return GetKey(KeyHandle).Time;
+}
+
+
+FKeyHandle FNameCurve::SetKeyTime(FKeyHandle KeyHandle, float NewTime)
+{
+	if (!IsKeyHandleValid(KeyHandle))
+	{
+		return KeyHandle;
+	}
+
+	FNameCurveKey OldKey = GetKey(KeyHandle);
+	
+	DeleteKey(KeyHandle);
+	AddKey(NewTime, OldKey.Value, KeyHandle);
+
+	// Copy all properties from old key, but then fix time to be the new time
+	GetKey(KeyHandle) = OldKey;
+	GetKey(KeyHandle).Time = NewTime;
+
+	return KeyHandle;
+}
+
+
+void FNameCurve::ShiftCurve(float DeltaTime)
+{
+	TSet<FKeyHandle> KeyHandles;
+	for (auto It = KeyHandlesToIndices.CreateIterator(); It; ++It)
+	{
+		FKeyHandle& KeyHandle = It.Key();
+		KeyHandles.Add(KeyHandle);
+	}
+
+	ShiftCurve(DeltaTime, KeyHandles);
+}
+
+
+void FNameCurve::ShiftCurve(float DeltaTime, TSet<FKeyHandle>& KeyHandles)
+{
+	for (auto It = KeyHandlesToIndices.CreateIterator(); It; ++It)
+	{
+		const FKeyHandle& KeyHandle = It.Key();
+		if (KeyHandles.Num() != 0 && KeyHandles.Contains(KeyHandle))
+		{
+			SetKeyTime(KeyHandle, GetKeyTime(KeyHandle) + DeltaTime);
+		}
+	}
+}
+
+
+void FNameCurve::ScaleCurve(float ScaleOrigin, float ScaleFactor)
+{
+	TSet<FKeyHandle> KeyHandles;
+	for (auto It = KeyHandlesToIndices.CreateIterator(); It; ++It)
+	{
+		FKeyHandle& KeyHandle = It.Key();
+		KeyHandles.Add(KeyHandle);
+	}
+
+	ScaleCurve(ScaleOrigin, ScaleFactor, KeyHandles);
+}
+
+
+void FNameCurve::ScaleCurve(float ScaleOrigin, float ScaleFactor, TSet<FKeyHandle>& KeyHandles)
+{
+	for (auto It = KeyHandlesToIndices.CreateIterator(); It; ++It)
+	{
+		const FKeyHandle& KeyHandle = It.Key();
+		if (KeyHandles.Num() != 0 && KeyHandles.Contains(KeyHandle))
+		{
+			SetKeyTime(KeyHandle, (GetKeyTime(KeyHandle) - ScaleOrigin) * ScaleFactor + ScaleOrigin);
+		}
+	}
+}
+
+
+FKeyHandle FNameCurve::UpdateOrAddKey(float InTime, const FName& InValue)
+{
+	// Search for a key that already exists at the time and if found, update its value
+	for (int32 KeyIndex = 0; KeyIndex < Keys.Num(); ++KeyIndex)
+	{
+		float KeyTime = Keys[KeyIndex].Time;
+
+		if (KeyTime > InTime)
+		{
+			// All the rest of the keys exist after the key we want to add
+			// so there is no point in searching
+			break;
+		}
+
+		if (FMath::IsNearlyEqual(KeyTime, InTime))
+		{
+			Keys[KeyIndex].Value = InValue;
+
+			return GetKeyHandle(KeyIndex);
+		}
+	}
+
+	// A key wasnt found, add it now
+	return AddKey(InTime, InValue);
+}
+
+
+int32 FNameCurve::GetNumKeys() const
+{
+	return Keys.Num();
+}
+
+
+bool FNameCurve::IsKeyHandleValid(FKeyHandle KeyHandle) const
+{
+	bool bValid = false;
+
+	if (FIndexedCurve::IsKeyHandleValid(KeyHandle))
+	{
+		bValid = Keys.IsValidIndex(GetIndex(KeyHandle));
+	}
+
+	return bValid;
+}
+
+
 /* FRichCurveKey
  *****************************************************************************/
 
