Sequencer: Added FindKey() to name curve

diff --git a/Engine/Source/Runtime/Engine/Classes/Curves/CurveBase.h b/Engine/Source/Runtime/Engine/Classes/Curves/CurveBase.h
--- a/Engine/Source/Runtime/Engine/Classes/Curves/CurveBase.h
+++ b/Engine/Source/Runtime/Engine/Classes/Curves/CurveBase.h
@@ -243,6 +243,14 @@ public:
 	 */
 	void DeleteKey(FKeyHandle KeyHandle);
 
+	/**
+	 * Finds a key a the specified time.
+	 *
+	 * @param KeyTime The time at which to find the key.
+	 * @return A handle to the key, or invalid key handle if not found.
+	 */
+	FKeyHandle FindKey(float KeyTime) const;
+
 	/**
 	 * Get a key.
 	 *
diff --git a/Engine/Source/Runtime/Engine/Private/CurveBase.cpp b/Engine/Source/Runtime/Engine/Private/CurveBase.cpp
--- a/Engine/Source/Runtime/Engine/Private/CurveBase.cpp
+++ b/Engine/Source/Runtime/Engine/Private/CurveBase.cpp
@@ -223,6 +223,36 @@ void FNameCurve::DeleteKey(FKeyHandle KeyHandle)
 }
 
 
+FKeyHandle FNameCurve::FindKey(float KeyTime) const
+{
+	int32 Start = 0;
+	int32 End = Keys.Num()-1;
+
+	// Binary search since the keys are in sorted order
+	while (Start <= End)
+	{
+		int32 TestPos = Start + (End-Start) / 2;
+		float TestKeyTime = Keys[TestPos].Time;
+
+		if (TestKeyTime == KeyTime)
+		{
+			return GetKeyHandle(TestPos);
+		}
+
+		if (TestKeyTime < KeyTime)
+		{
+			Start = TestPos+1;
+		}
+		else
+		{
+			End = TestPos-1;
+		}
+	}
+
+	return FKeyHandle();
+}
+
+
 FNameCurveKey& FNameCurve::GetKey(FKeyHandle KeyHandle)
 {
 	EnsureAllIndicesHaveHandles();
