Fixed problem with landscape grass bounds causing visibility flickering

diff --git a/Engine/Source/Runtime/Engine/Private/HierarchicalInstancedStaticMesh.cpp b/Engine/Source/Runtime/Engine/Private/HierarchicalInstancedStaticMesh.cpp
--- a/Engine/Source/Runtime/Engine/Private/HierarchicalInstancedStaticMesh.cpp
+++ b/Engine/Source/Runtime/Engine/Private/HierarchicalInstancedStaticMesh.cpp
@@ -1602,10 +1602,11 @@ const TArray<FBoxSphereBounds>* FHierarchicalStaticMeshSceneProxy::GetOcclusionQ
 
 FBoxSphereBounds UHierarchicalInstancedStaticMeshComponent::CalcBounds(const FTransform& BoundTransform) const
 {
+	ensure(BuiltInstanceBounds.IsValid || ClusterTreePtr->Num() == 0);
+
 	if (BuiltInstanceBounds.IsValid || UnbuiltInstanceBounds.IsValid)
 	{
 		FBoxSphereBounds Result = BuiltInstanceBounds + UnbuiltInstanceBounds;
-
 		return Result.TransformBy(BoundTransform);
 	}
 	else
@@ -2049,12 +2050,13 @@ void UHierarchicalInstancedStaticMeshComponent::BuildTree()
 		UnbuiltInstanceBounds.Init();
 		RemovedInstances.Empty();
 		UnbuiltInstanceBoundsList.Empty();
+		BuiltInstanceBounds = (Builder->Result->Nodes.Num() > 0 ? FBox(Builder->Result->Nodes[0].BoundMin, Builder->Result->Nodes[0].BoundMax) : FBox(0));
 
 		ClusterTreePtr = MakeShareable(new TArray<FClusterNode>);
 		Exchange(*ClusterTreePtr, Builder->Result->Nodes);
 		Exchange(InstanceReorderTable, Builder->Result->InstanceReorderTable);
 		Exchange(SortedInstances, Builder->Result->SortedInstances);
-		
+
 		FlushAccumulatedNavigationUpdates();
 				
 		PostBuildStats();
@@ -2097,6 +2099,7 @@ void UHierarchicalInstancedStaticMeshComponent::AcceptPrebuiltTree(TArray<FClust
 	InstanceReorderTable.Empty();
 	SortedInstances.Empty();
 	OcclusionLayerNumNodes = InOcclusionLayerNumNodes;
+	BuiltInstanceBounds = (InClusterTree.Num() > 0 ? FBox(InClusterTree[0].BoundMin, InClusterTree[0].BoundMax) : FBox(0));
 
 	// Verify that the mesh is valid before using it.
 	const bool bMeshIsValid = 
@@ -2260,7 +2263,7 @@ void UHierarchicalInstancedStaticMeshComponent::ApplyBuildTreeAsync(ENamedThread
 		OcclusionLayerNumNodes = Builder->Result->OutOcclusionLayerNum;
 		BuiltInstanceBounds = (ClusterTree.Num() > 0 ? FBox(ClusterTree[0].BoundMin, ClusterTree[0].BoundMax) : FBox(0));
 
-		//UE_LOG(LogStaticMesh, Display, TEXT("Built a foliage hierarchy with %d of %d elements in %.1fs."), NumBuiltInstances, PerInstanceSMData.Num(), float(FPlatformTime::Seconds() - StartTime));
+		UE_LOG(LogStaticMesh, Verbose, TEXT("Built a foliage hierarchy with %d of %d elements in %.1fs."), NumBuiltInstances, PerInstanceSMData.Num(), float(FPlatformTime::Seconds() - StartTime));
 
 		if (NumBuiltInstances < PerInstanceSMData.Num())
 		{
@@ -2316,7 +2319,7 @@ void UHierarchicalInstancedStaticMeshComponent::BuildTreeAsync()
 			InstanceTransforms[Index] = PerInstanceSMData[Index].Transform;
 		}
 
-		//UE_LOG(LogStaticMesh, Display, TEXT("Copied %d transforms in %.3fs."), Num, float(FPlatformTime::Seconds() - StartTime));
+		UE_LOG(LogStaticMesh, Verbose, TEXT("Copied %d transforms in %.3fs."), Num, float(FPlatformTime::Seconds() - StartTime));
 
 		TSharedRef<FClusterBuilder, ESPMode::ThreadSafe> Builder(new FClusterBuilder(InstanceTransforms, StaticMesh->GetBounds().GetBox(), DesiredInstancesPerLeaf()));
 
