Integrate CL 2702298 Orion -> Main
Compile out 'show pre phys bones' in Shipping/Test

diff --git a/Engine/Source/Runtime/Engine/Private/PhysicsEngine/PhysAnim.cpp b/Engine/Source/Runtime/Engine/Private/PhysicsEngine/PhysAnim.cpp
--- a/Engine/Source/Runtime/Engine/Private/PhysicsEngine/PhysAnim.cpp
+++ b/Engine/Source/Runtime/Engine/Private/PhysicsEngine/PhysAnim.cpp
@@ -491,6 +491,7 @@ void USkeletalMeshComponent::UpdateKinematicBonesToAnim(const TArray<FTransform>
 		return;
 	}
 
+#if !(UE_BUILD_SHIPPING || UE_BUILD_TEST)
 	// If desired, draw the skeleton at the point where we pass it to the physics.
 	if (bShowPrePhysBones && SkeletalMesh && InSpaceBases.Num() == SkeletalMesh->RefSkeleton.GetNum())
 	{
@@ -504,6 +505,7 @@ void USkeletalMeshComponent::UpdateKinematicBonesToAnim(const TArray<FTransform>
 			GetWorld()->LineBatcher->DrawLine(ThisPos, ParentPos, AnimSkelDrawColor, SDPG_Foreground);
 		}
 	}
+#endif
 
 	// warn if it has non-uniform scale
 	const FVector& MeshScale3D = CurrentLocalToWorld.GetScale3D();
