UEMOB-36

Add material shader quality settings.
Enables quality overrides for android and ios (ES2 flavours)

diff --git a/Engine/Shaders/BasePassForForwardShadingPixelShader.usf b/Engine/Shaders/BasePassForForwardShadingPixelShader.usf
--- a/Engine/Shaders/BasePassForForwardShadingPixelShader.usf
+++ b/Engine/Shaders/BasePassForForwardShadingPixelShader.usf
@@ -23,8 +23,11 @@ float4 LightPositionAndInvRadius[MAX_DYNAMIC_POINT_LIGHTS];
 float4 LightColorAndFalloffExponent[MAX_DYNAMIC_POINT_LIGHTS];
 #endif
 
+#define FULLY_ROUGH (MATERIAL_FULLY_ROUGH || FORWARD_QL_FORCE_FULLY_ROUGH)
+#define NONMETAL (MATERIAL_NONMETAL || FORWARD_QL_FORCE_NONMETAL)
+
 /** Prenormalized capture of the scene that's closest to the object being rendered. */
-#if !MATERIAL_FULLY_ROUGH
+#if !FULLY_ROUGH
 TextureCube ReflectionCubemap;
 SamplerState ReflectionCubemapSampler;
 #endif
@@ -44,7 +47,7 @@ half PhongApprox( half Roughness, half RoL )
 	// Total 7 instr
 }
 
-#if !MATERIAL_FULLY_ROUGH
+#if !FULLY_ROUGH
 half3 GetImageBasedReflectionLighting(FMaterialPixelParameters MaterialParameters, half Roughness)
 {
 	half3 ProjectedCaptureVector = MaterialParameters.ReflectionVector;	
@@ -160,7 +163,7 @@ void Main(
 	half  Metallic = GetMaterialMetallic( MaterialParameters );
 	half  Specular = GetMaterialSpecular( MaterialParameters );
 
-#if MATERIAL_NONMETAL
+#if NONMETAL
 	half3 DiffuseColor = BaseColor;
 	half SpecularColor = 0.04;
 #else
@@ -171,14 +174,14 @@ void Main(
 
 	half Roughness = GetMaterialRoughness(MaterialParameters);
 
-#if MATERIAL_FULLY_ROUGH
+#if FULLY_ROUGH
 	// Factors derived from EnvBRDFApprox( SpecularColor, 1, 1 ) == SpecularColor * 0.4524 - 0.0024
 	DiffuseColor += SpecularColor * 0.45;
 	SpecularColor = 0;
 #else
 	half NoV = max( dot( MaterialParameters.WorldNormal, MaterialParameters.CameraVector ), 0 );
 
-#if MATERIAL_NONMETAL
+#if NONMETAL
 	// If nothing is hooked up to Metalic and Specular,
 	// then defaults are the same as a non-metal,
 	// so this define is safe.
@@ -276,7 +279,7 @@ void Main(
 		float NoL = max(0, dot(MaterialParameters.WorldNormal, View.DirectionalLightDirection));
 		float RoL = max(0, dot(MaterialParameters.ReflectionVector, View.DirectionalLightDirection));
 
-		#if MATERIAL_FULLY_ROUGH
+		#if FULLY_ROUGH
 			Color += (Shadow * NoL) * View.DirectionalLightColor.rgb * DiffuseColor;
 		#else
 			Color += (Shadow * NoL) * View.DirectionalLightColor.rgb * (DiffuseColor + SpecularColor * PhongApprox(Roughness, RoL));
@@ -332,7 +335,7 @@ void Main(
 					Attenuation = RadialAttenuation(ToLight * LightPositionAndInvRadius[i].w, LightColorAndFalloffExponent[i].w);		
 				}
 
-				#if !MATERIAL_FULLY_ROUGH
+				#if !FULLY_ROUGH
 					Color += (Attenuation * PointNoL) * LightColorAndFalloffExponent[i].rgb * (DiffuseColor + SpecularColor * PhongApprox(Roughness, PointRoL));
 				#else
 					Color += (Attenuation * PointNoL) * LightColorAndFalloffExponent[i].rgb * DiffuseColor;
diff --git a/Engine/Shaders/LightmapCommon.usf b/Engine/Shaders/LightmapCommon.usf
--- a/Engine/Shaders/LightmapCommon.usf
+++ b/Engine/Shaders/LightmapCommon.usf
@@ -6,6 +6,14 @@
 #define MATERIAL_USE_LM_DIRECTIONALITY	1
 #endif
 
+
+// Material quality overrides
+#ifndef QL_FORCEDISABLE_LM_DIRECTIONALITY
+#define QL_FORCEDISABLE_LM_DIRECTIONALITY 0
+#endif
+
+#define USE_LM_DIRECTIONALITY (MATERIAL_USE_LM_DIRECTIONALITY && !QL_FORCEDISABLE_LM_DIRECTIONALITY)
+
 Texture2D LightMapTexture;
 SamplerState LightMapSampler;
 	
@@ -32,7 +40,7 @@ half4 GetLightMapColorLQ( float2 LightmapUV0, float2 LightmapUV1, half3 WorldNor
 	const half LogBlackPoint = 0.00390625;	// exp2(-8);
 	half L = exp2( LogL * 16 - 8 ) - LogBlackPoint;		// 1 exp2, 1 smad, 1 ssub
 
-#if MATERIAL_USE_LM_DIRECTIONALITY
+#if USE_LM_DIRECTIONALITY
 	// Alpha doesn't matter, will scaled by zero
 	half4 SH = Lightmap1 * LightMapScale[1] + LightMapAdd[1];	// 1 vmad
 
@@ -70,7 +78,7 @@ void GetLightMapColorHQ( float2 LightmapUV0, float2 LightmapUV1, half3 WorldNorm
 	const half LogBlackPoint = 0.01858136;
 	half L = exp2( LogL ) - LogBlackPoint;
 
-#if MATERIAL_USE_LM_DIRECTIONALITY
+#if USE_LM_DIRECTIONALITY
 	// Range scale SH. Alpha doesn't matter, will scale with zero
 	half4 SH = Lightmap1 * LightMapScale[1] + LightMapAdd[1];
 
diff --git a/Engine/Source/Developer/Android/AndroidPlatformEditor/AndroidPlatformEditor.Build.cs b/Engine/Source/Developer/Android/AndroidPlatformEditor/AndroidPlatformEditor.Build.cs
--- a/Engine/Source/Developer/Android/AndroidPlatformEditor/AndroidPlatformEditor.Build.cs
+++ b/Engine/Source/Developer/Android/AndroidPlatformEditor/AndroidPlatformEditor.Build.cs
@@ -24,7 +24,9 @@ public class AndroidPlatformEditor : ModuleRules
 				"SourceControl",
 				"AndroidRuntimeSettings",
                 "AndroidDeviceDetection",
-                "TargetPlatform"
+                "TargetPlatform",
+                "RenderCore",
+                "MaterialShaderQualitySettings",
 			}
 		);
 
diff --git a/Engine/Source/Developer/Android/AndroidPlatformEditor/Private/AndroidPlatformEditorModule.cpp b/Engine/Source/Developer/Android/AndroidPlatformEditor/Private/AndroidPlatformEditorModule.cpp
--- a/Engine/Source/Developer/Android/AndroidPlatformEditor/Private/AndroidPlatformEditorModule.cpp
+++ b/Engine/Source/Developer/Android/AndroidPlatformEditor/Private/AndroidPlatformEditorModule.cpp
@@ -6,6 +6,11 @@
 #include "ModuleInterface.h"
 #include "ISettingsModule.h"
 #include "ModuleManager.h"
+#include "MaterialShaderQualitySettingsCustomization.h"
+#include "MaterialShaderQualitySettings.h"
+#include "Materials/Material.h"
+#include "RenderingThread.h"
+#include "ComponentRecreateRenderStateContext.h"
 
 
 #define LOCTEXT_NAMESPACE "FAndroidPlatformEditorModule"
@@ -33,6 +38,19 @@ class FAndroidPlatformEditorModule
 			FOnGetDetailCustomizationInstance::CreateStatic(&FAndroidSDKSettingsCustomization::MakeInstance)
 			);
 
+		FOnUpdateMaterialShaderQuality UpdateMaterials = FOnUpdateMaterialShaderQuality::CreateLambda([]()
+		{
+			FGlobalComponentRecreateRenderStateContext Recreate;
+			FlushRenderingCommands();
+			UMaterial::AllMaterialsCacheResourceShadersForRendering();
+			UMaterialInstance::AllMaterialsCacheResourceShadersForRendering();
+		});
+
+		PropertyModule.RegisterCustomClassLayout(
+			UShaderPlatformQualitySettings::StaticClass()->GetFName(),
+			FOnGetDetailCustomizationInstance::CreateStatic(&FMaterialShaderQualitySettingsCustomization::MakeInstance, UpdateMaterials)
+			);
+
 		PropertyModule.NotifyCustomizationModuleChanged();
 
 		// register settings
@@ -51,6 +69,14 @@ class FAndroidPlatformEditorModule
  				LOCTEXT("SDKSettingsDescription", "Settings for Android SDK (for all projects)"),
  				GetMutableDefault<UAndroidSDKSettings>()
 			);
+
+			static FName NAME_OPENGL_ES2(TEXT("GLSL_ES2"));
+			const UShaderPlatformQualitySettings* AndroidMaterialQualitySettings = UMaterialShaderQualitySettings::Get()->GetShaderPlatformQualitySettings(NAME_OPENGL_ES2);
+			SettingsModule->RegisterSettings("Project", "Platforms", "AndroidES2Quality",
+				LOCTEXT("AndroidES2QualitySettingsName", "Android ES2 Quality"),
+				LOCTEXT("AndroidES2QualitySettingsDescription", "Settings for Android ES2 material quality"),
+				AndroidMaterialQualitySettings
+			);
 		}
 
 		// Force the SDK settings into a sane state initially so we can make use of them
@@ -70,6 +96,7 @@ class FAndroidPlatformEditorModule
 		{
 			SettingsModule->UnregisterSettings("Project", "Platforms", "Android");
 			SettingsModule->UnregisterSettings("Project", "Platforms", "AndroidSDK");
+			SettingsModule->UnregisterSettings("Project", "Platforms", "AndroidES2Quality");
 		}
 	}
 };
diff --git a/Engine/Source/Developer/iOS/IOSPlatformEditor/IOSPlatformEditor.Build.cs b/Engine/Source/Developer/iOS/IOSPlatformEditor/IOSPlatformEditor.Build.cs
--- a/Engine/Source/Developer/iOS/IOSPlatformEditor/IOSPlatformEditor.Build.cs
+++ b/Engine/Source/Developer/iOS/IOSPlatformEditor/IOSPlatformEditor.Build.cs
@@ -24,6 +24,8 @@ public class IOSPlatformEditor : ModuleRules
 				"SourceControl",
 				"IOSRuntimeSettings",
 				"TargetPlatform",
+				"MaterialShaderQualitySettings",
+				"RenderCore",
 			}
 		);
 
diff --git a/Engine/Source/Developer/iOS/IOSPlatformEditor/Private/IOSPlatformEditorModule.cpp b/Engine/Source/Developer/iOS/IOSPlatformEditor/Private/IOSPlatformEditorModule.cpp
--- a/Engine/Source/Developer/iOS/IOSPlatformEditor/Private/IOSPlatformEditorModule.cpp
+++ b/Engine/Source/Developer/iOS/IOSPlatformEditor/Private/IOSPlatformEditorModule.cpp
@@ -5,7 +5,10 @@
 #include "ISettingsModule.h"
 #include "ModuleInterface.h"
 #include "ModuleManager.h"
-
+#include "MaterialShaderQualitySettings.h"
+#include "MaterialShaderQualitySettingsCustomization.h"
+#include "RenderingThread.h"
+#include "ComponentRecreateRenderStateContext.h"
 
 #define LOCTEXT_NAMESPACE "FIOSPlatformEditorModule"
 
@@ -27,6 +30,19 @@ class FIOSPlatformEditorModule
 			FOnGetDetailCustomizationInstance::CreateStatic(&FIOSTargetSettingsCustomization::MakeInstance)
 		);
 
+		FOnUpdateMaterialShaderQuality UpdateMaterials = FOnUpdateMaterialShaderQuality::CreateLambda([]()
+		{
+			FGlobalComponentRecreateRenderStateContext Recreate;
+			FlushRenderingCommands();
+			UMaterial::AllMaterialsCacheResourceShadersForRendering();
+			UMaterialInstance::AllMaterialsCacheResourceShadersForRendering();
+		});
+
+		PropertyModule.RegisterCustomClassLayout(
+			UShaderPlatformQualitySettings::StaticClass()->GetFName(),
+			FOnGetDetailCustomizationInstance::CreateStatic(&FMaterialShaderQualitySettingsCustomization::MakeInstance, UpdateMaterials)
+			);
+
 		PropertyModule.NotifyCustomizationModuleChanged();
 
 		// register settings
@@ -39,6 +55,14 @@ class FIOSPlatformEditorModule
 				LOCTEXT("RuntimeSettingsDescription", "Settings and resources for the iOS platform"),
 				GetMutableDefault<UIOSRuntimeSettings>()
 			);
+
+			static FName NAME_OPENGL_ES2_IOS(TEXT("GLSL_ES2_IOS"));
+			const UShaderPlatformQualitySettings* IOSES2MaterialQualitySettings = UMaterialShaderQualitySettings::Get()->GetShaderPlatformQualitySettings(NAME_OPENGL_ES2_IOS);
+			SettingsModule->RegisterSettings("Project", "Platforms", "iOSES2Quality",
+				LOCTEXT("IOSES2QualitySettingsName", "iOS ES2 Quality"),
+				LOCTEXT("IOSES2QualitySettingsDescription", "Settings for iOS ES2 material quality"),
+				IOSES2MaterialQualitySettings
+				);
 		}
 	}
 
@@ -49,6 +73,7 @@ class FIOSPlatformEditorModule
 		if (SettingsModule != nullptr)
 		{
 			SettingsModule->UnregisterSettings("Project", "Platforms", "iOS");
+			SettingsModule->UnregisterSettings("Project", "Platforms", "iOSES2Quality");
 		}
 	}
 };
diff --git a/Engine/Source/Editor/LevelEditor/LevelEditor.Build.cs b/Engine/Source/Editor/LevelEditor/LevelEditor.Build.cs
--- a/Engine/Source/Editor/LevelEditor/LevelEditor.Build.cs
+++ b/Engine/Source/Editor/LevelEditor/LevelEditor.Build.cs
@@ -62,7 +62,8 @@ public class LevelEditor : ModuleRules
 				"KismetWidgets",
 				"Sequencer",
                 "Foliage",
-                "HierarchicalLODOutliner"
+                "HierarchicalLODOutliner",
+				"MaterialShaderQualitySettings",
 			}
 		);
 
diff --git a/Engine/Source/Editor/LevelEditor/Private/LevelEditor.cpp b/Engine/Source/Editor/LevelEditor/Private/LevelEditor.cpp
--- a/Engine/Source/Editor/LevelEditor/Private/LevelEditor.cpp
+++ b/Engine/Source/Editor/LevelEditor/Private/LevelEditor.cpp
@@ -1635,12 +1635,32 @@ void FLevelEditorModule::BindGlobalLevelEditorCommands()
 		FExecuteAction::CreateStatic( &FLevelEditorActionCallbacks::SetMaterialQualityLevel, (EMaterialQualityLevel::Type)EMaterialQualityLevel::Low ),
 		FCanExecuteAction(),
 		FIsActionChecked::CreateStatic( &FLevelEditorActionCallbacks::IsMaterialQualityLevelChecked, (EMaterialQualityLevel::Type)EMaterialQualityLevel::Low ) );
+	ActionList.MapAction(
+		Commands.MaterialQualityLevel_Medium,
+		FExecuteAction::CreateStatic(&FLevelEditorActionCallbacks::SetMaterialQualityLevel, (EMaterialQualityLevel::Type)EMaterialQualityLevel::Medium),
+		FCanExecuteAction(),
+		FIsActionChecked::CreateStatic(&FLevelEditorActionCallbacks::IsMaterialQualityLevelChecked, (EMaterialQualityLevel::Type)EMaterialQualityLevel::Medium));
 	ActionList.MapAction( 
 		Commands.MaterialQualityLevel_High, 
 		FExecuteAction::CreateStatic( &FLevelEditorActionCallbacks::SetMaterialQualityLevel, (EMaterialQualityLevel::Type)EMaterialQualityLevel::High ),
 		FCanExecuteAction(),
 		FIsActionChecked::CreateStatic( &FLevelEditorActionCallbacks::IsMaterialQualityLevelChecked, (EMaterialQualityLevel::Type)EMaterialQualityLevel::High ) );
 
+	ActionList.MapAction(
+		Commands.PreviewPlatformOverride_DefaultES2,
+		FExecuteAction::CreateStatic(&FLevelEditorActionCallbacks::SetPreviewPlatform, FName()),
+		FCanExecuteAction(),
+		FIsActionChecked::CreateStatic(&FLevelEditorActionCallbacks::IsPreviewPlatformChecked, FName()));
+	ActionList.MapAction(
+		Commands.PreviewPlatformOverride_AndroidES2,
+		FExecuteAction::CreateStatic(&FLevelEditorActionCallbacks::SetPreviewPlatform, LegacyShaderPlatformToShaderFormat(SP_OPENGL_ES2_ANDROID)),
+		FCanExecuteAction(),
+		FIsActionChecked::CreateStatic(&FLevelEditorActionCallbacks::IsPreviewPlatformChecked, LegacyShaderPlatformToShaderFormat(SP_OPENGL_ES2_ANDROID)));
+	ActionList.MapAction(
+		Commands.PreviewPlatformOverride_IOSES2,
+		FExecuteAction::CreateStatic(&FLevelEditorActionCallbacks::SetPreviewPlatform, LegacyShaderPlatformToShaderFormat(SP_OPENGL_ES2_IOS)),
+		FCanExecuteAction(),
+		FIsActionChecked::CreateStatic(&FLevelEditorActionCallbacks::IsPreviewPlatformChecked,LegacyShaderPlatformToShaderFormat(SP_OPENGL_ES2_IOS)));
 
 	for (int32 i = 0; i < ERHIFeatureLevel::Num; ++i)
 	{
diff --git a/Engine/Source/Editor/LevelEditor/Private/LevelEditorActions.cpp b/Engine/Source/Editor/LevelEditor/Private/LevelEditorActions.cpp
--- a/Engine/Source/Editor/LevelEditor/Private/LevelEditorActions.cpp
+++ b/Engine/Source/Editor/LevelEditor/Private/LevelEditorActions.cpp
@@ -68,6 +68,7 @@
 #include "AsyncResult.h"
 #include "IPortalApplicationWindow.h"
 #include "IPortalServiceLocator.h"
+#include "MaterialShaderQualitySettings.h"
 
 DEFINE_LOG_CATEGORY_STATIC(LevelEditorActions, Log, All);
 
@@ -482,10 +483,23 @@ void FLevelEditorActionCallbacks::SetMaterialQualityLevel( EMaterialQualityLevel
 bool FLevelEditorActionCallbacks::IsMaterialQualityLevelChecked( EMaterialQualityLevel::Type TestQualityLevel )
 {
 	static const auto MaterialQualityLevelVar = IConsoleManager::Get().FindTConsoleVariableDataInt(TEXT("r.MaterialQualityLevel"));
-	EMaterialQualityLevel::Type MaterialQualityLevel = (EMaterialQualityLevel::Type)FMath::Clamp(MaterialQualityLevelVar->GetValueOnGameThread(), 0, 1);
+	EMaterialQualityLevel::Type MaterialQualityLevel = (EMaterialQualityLevel::Type)FMath::Clamp(MaterialQualityLevelVar->GetValueOnGameThread(), 0, (int32)EMaterialQualityLevel::Num-1);
 	return TestQualityLevel == MaterialQualityLevel;
 }
 
+void FLevelEditorActionCallbacks::SetPreviewPlatform(FName MaterialQualityPlatform)
+{
+	UMaterialShaderQualitySettings::Get()->SetPreviewPlatform(MaterialQualityPlatform);
+
+	SetFeatureLevelPreview(ERHIFeatureLevel::ES2);
+}
+
+bool FLevelEditorActionCallbacks::IsPreviewPlatformChecked(FName MaterialQualityPlatform)
+{
+	const FName& PreviewPlatform = UMaterialShaderQualitySettings::Get()->GetPreviewPlatform();
+	return PreviewPlatform == MaterialQualityPlatform && ERHIFeatureLevel::ES2 == GetWorld()->FeatureLevel;
+}
+
 void FLevelEditorActionCallbacks::SetFeatureLevelPreview(ERHIFeatureLevel::Type InPreviewFeatureLevel)
 {
 	// Record this feature level as we want to use it for all subsequent level creation and loading
@@ -3051,8 +3065,13 @@ void FLevelEditorCommands::RegisterCommands()
 
 	UI_COMMAND( FocusAllViewportsToSelection, "Focus Selected Actors in All Viewports", "Moves the camera in front of the selected actors in all open viewports", EUserInterfaceActionType::Button, FInputChord( EKeys::F, EModifierKey::Shift ) );
 
-	UI_COMMAND( MaterialQualityLevel_Low, "Low", "Sets material quality in the scene to low.", EUserInterfaceActionType::RadioButton, FInputChord() );
-	UI_COMMAND( MaterialQualityLevel_High, "High", "Sets material quality in the scene to high.", EUserInterfaceActionType::RadioButton, FInputChord() );
+	UI_COMMAND(MaterialQualityLevel_Low, "Low", "Sets material quality in the scene to low.", EUserInterfaceActionType::RadioButton, FInputChord());
+	UI_COMMAND(MaterialQualityLevel_Medium, "Medium", "Sets material quality in the scene to medium.", EUserInterfaceActionType::RadioButton, FInputChord());
+	UI_COMMAND(MaterialQualityLevel_High, "High", "Sets material quality in the scene to high.", EUserInterfaceActionType::RadioButton, FInputChord());
+
+	UI_COMMAND(PreviewPlatformOverride_DefaultES2, "Default Mobile / HTML5 Preview", "Use default mobile settings (no quality overrides).", EUserInterfaceActionType::RadioButton, FInputChord());
+	UI_COMMAND(PreviewPlatformOverride_AndroidES2, "Android Preview", "Mobile preview using Android's quality settings.", EUserInterfaceActionType::RadioButton, FInputChord());
+	UI_COMMAND(PreviewPlatformOverride_IOSES2, "iOS ES2 Preview", "Mobile preview using iOS's OpenGL ES2 quality settings.", EUserInterfaceActionType::RadioButton, FInputChord());
 
 	UI_COMMAND( ConnectToSourceControl, "Connect to Source Control...", "Opens a dialog to connect to source control.", EUserInterfaceActionType::Button, FInputChord());
 	UI_COMMAND( ChangeSourceControlSettings, "Change Source Control Settings...", "Opens a dialog to change source control settings.", EUserInterfaceActionType::Button, FInputChord());
diff --git a/Engine/Source/Editor/LevelEditor/Private/LevelEditorToolBar.cpp b/Engine/Source/Editor/LevelEditor/Private/LevelEditorToolBar.cpp
--- a/Engine/Source/Editor/LevelEditor/Private/LevelEditorToolBar.cpp
+++ b/Engine/Source/Editor/LevelEditor/Private/LevelEditorToolBar.cpp
@@ -1592,12 +1592,24 @@ TSharedRef< SWidget > FLevelEditorToolBar::GenerateBuildMenuContent( TSharedRef<
 	return MenuBuilder.MakeWidget();
 }
 
+static void MakeES2PreviewPlatformOverrideMenu(FMenuBuilder& MenuBuilder)
+{
+	MenuBuilder.BeginSection("LevelEditorShaderModelPreview", NSLOCTEXT("LevelToolBarViewMenu", "ES2PreviewPlatformOverrideHeading", "Preview Platform"));
+	{
+		MenuBuilder.AddMenuEntry(FLevelEditorCommands::Get().PreviewPlatformOverride_DefaultES2);
+		MenuBuilder.AddMenuEntry(FLevelEditorCommands::Get().PreviewPlatformOverride_AndroidES2);
+		MenuBuilder.AddMenuEntry(FLevelEditorCommands::Get().PreviewPlatformOverride_IOSES2);
+	}
+	MenuBuilder.EndSection();
+}
+
 static void MakeMaterialQualityLevelMenu( FMenuBuilder& MenuBuilder )
 {
 	MenuBuilder.BeginSection("LevelEditorMaterialQualityLevel", NSLOCTEXT( "LevelToolBarViewMenu", "MaterialQualityLevelHeading", "Material Quality Level" ) );
 	{
-		MenuBuilder.AddMenuEntry( FLevelEditorCommands::Get().MaterialQualityLevel_Low );
-		MenuBuilder.AddMenuEntry( FLevelEditorCommands::Get().MaterialQualityLevel_High );
+		MenuBuilder.AddMenuEntry(FLevelEditorCommands::Get().MaterialQualityLevel_Low);
+		MenuBuilder.AddMenuEntry(FLevelEditorCommands::Get().MaterialQualityLevel_Medium);
+		MenuBuilder.AddMenuEntry(FLevelEditorCommands::Get().MaterialQualityLevel_High);
 	}
 	MenuBuilder.EndSection();
 }
@@ -1608,7 +1620,14 @@ static void MakeShaderModelPreviewMenu(FMenuBuilder& MenuBuilder)
 	{
 		for (int32 i = GMaxRHIFeatureLevel; i >= 0; --i)
 		{
-			if (i != ERHIFeatureLevel::ES3_1 || GetDefault<UEditorExperimentalSettings>()->bFeatureLevelES31Preview)
+			if (i == ERHIFeatureLevel::ES2)
+			{
+				MenuBuilder.AddSubMenu(
+					FLevelEditorCommands::Get().FeatureLevelPreview[i]->GetLabel(),
+					FLevelEditorCommands::Get().FeatureLevelPreview[i]->GetDescription(),
+					FNewMenuDelegate::CreateStatic(&MakeES2PreviewPlatformOverrideMenu));
+			}
+			else if (i != ERHIFeatureLevel::ES3_1 || GetDefault<UEditorExperimentalSettings>()->bFeatureLevelES31Preview)
 			{
 				MenuBuilder.AddMenuEntry(FLevelEditorCommands::Get().FeatureLevelPreview[i]);
 			}
@@ -1707,7 +1726,7 @@ TSharedRef< SWidget > FLevelEditorToolBar::GenerateQuickSettingsMenu( TSharedRef
 
 		MenuBuilder.AddSubMenu(
 			LOCTEXT( "MaterialQualityLevelSubMenu", "Material Quality Level" ),
-			LOCTEXT( "MaterialQualityLevelSubMenu_ToolTip", "Sets the value of the CVar \"r.MaterialQualityLevel\" (low=0, high=1). This affects materials via the QualitySwitch material expression." ),
+			LOCTEXT( "MaterialQualityLevelSubMenu_ToolTip", "Sets the value of the CVar \"r.MaterialQualityLevel\" (low=0, high=1, medium=2). This affects materials via the QualitySwitch material expression." ),
 			FNewMenuDelegate::CreateStatic( &MakeMaterialQualityLevelMenu ) );
 
 		MenuBuilder.AddSubMenu(
diff --git a/Engine/Source/Editor/LevelEditor/Public/LevelEditorActions.h b/Engine/Source/Editor/LevelEditor/Public/LevelEditorActions.h
--- a/Engine/Source/Editor/LevelEditor/Public/LevelEditorActions.h
+++ b/Engine/Source/Editor/LevelEditor/Public/LevelEditorActions.h
@@ -552,10 +552,15 @@ public:
 	TSharedPtr< FUICommandInfo > AddMatinee;
 
 	TSharedPtr< FUICommandInfo > MaterialQualityLevel_Low;
+	TSharedPtr< FUICommandInfo > MaterialQualityLevel_Medium;
 	TSharedPtr< FUICommandInfo > MaterialQualityLevel_High;
 
 	TSharedPtr< FUICommandInfo > FeatureLevelPreview[ERHIFeatureLevel::Num];
 	
+	TSharedPtr< FUICommandInfo > PreviewPlatformOverride_DefaultES2;
+	TSharedPtr< FUICommandInfo > PreviewPlatformOverride_AndroidES2;
+	TSharedPtr< FUICommandInfo > PreviewPlatformOverride_IOSES2;
+
 	///**
 	// * Mode Commands                   
 	// */
@@ -737,6 +742,8 @@ public:
 	static void AttachToSocketSelection(FName SocketName, AActor* ParentActorPtr);
 	static void SetMaterialQualityLevel( EMaterialQualityLevel::Type NewQualityLevel );
 	static bool IsMaterialQualityLevelChecked( EMaterialQualityLevel::Type TestQualityLevel );
+	static void SetPreviewPlatform(FName MaterialQualityPlatform);
+	static bool IsPreviewPlatformChecked(FName MaterialQualityPlatform);
 	static void SetFeatureLevelPreview(ERHIFeatureLevel::Type InFeatureLevel);
 	static bool IsFeatureLevelPreviewChecked(ERHIFeatureLevel::Type InFeatureLevel);
 	
diff --git a/Engine/Source/Runtime/Core/Private/HAL/ConsoleManager.cpp b/Engine/Source/Runtime/Core/Private/HAL/ConsoleManager.cpp
--- a/Engine/Source/Runtime/Core/Private/HAL/ConsoleManager.cpp
+++ b/Engine/Source/Runtime/Core/Private/HAL/ConsoleManager.cpp
@@ -1737,7 +1737,7 @@ static TAutoConsoleVariable<float> CVarScreenPercentage(
 static TAutoConsoleVariable<int32> CVarMaterialQualityLevel(
 	TEXT("r.MaterialQualityLevel"),
 	1,
-	TEXT("0 corresponds to low quality materials, as defined by quality switches in materials, 1 corresponds to high."),
+	TEXT("0 corresponds to low quality materials, as defined by quality switches in materials, 1 corresponds to high and 2 for medium."),
 	ECVF_Scalability | ECVF_RenderThreadSafe);
 
 static TAutoConsoleVariable<int32> CVarUseDXT5NormalMaps(
diff --git a/Engine/Source/Runtime/Engine/Classes/Kismet/KismetSystemLibrary.h b/Engine/Source/Runtime/Engine/Classes/Kismet/KismetSystemLibrary.h
--- a/Engine/Source/Runtime/Engine/Classes/Kismet/KismetSystemLibrary.h
+++ b/Engine/Source/Runtime/Engine/Classes/Kismet/KismetSystemLibrary.h
@@ -1342,6 +1342,7 @@ class ENGINE_API UKismetSystemLibrary : public UBlueprintFunctionLibrary
 	 * Get the clamped state of r.MaterialQualityLevel, see console variable help (allows for scalability, cannot be used in construction scripts)
 	 * 0: low
 	 * 1: high
+	 * 2: medium
 	 */
 	UFUNCTION(BlueprintPure, Category="Rendering|Material", meta=(UnsafeDuringActorConstruction = "true"))
 	static int32 GetRenderingMaterialQualityLevel();
diff --git a/Engine/Source/Runtime/Engine/Classes/Materials/Material.h b/Engine/Source/Runtime/Engine/Classes/Materials/Material.h
--- a/Engine/Source/Runtime/Engine/Classes/Materials/Material.h
+++ b/Engine/Source/Runtime/Engine/Classes/Materials/Material.h
@@ -1036,9 +1036,17 @@ public:
 		}
 	}
 
-	/** Determines whether each quality level has different nodes by inspecting the material's expressions. */
-	void GetQualityLevelNodeUsage(TArray<bool, TInlineAllocator<EMaterialQualityLevel::Num> >& QualityLevelsUsed);
+	/** Determines whether each quality level has different nodes by inspecting the material's expressions. 
+	* Or is required by the material quality setting overrides.
+	* @param	QualityLevelsUsed	output array of used quality levels.
+	* @param	ShaderPlatform	The shader platform to use for the quality settings.
+	*/
+	void GetQualityLevelUsage(TArray<bool, TInlineAllocator<EMaterialQualityLevel::Num> >& QualityLevelsUsed, EShaderPlatform ShaderPlatform);
 
+	/** Determines whether each quality level has different nodes by inspecting the material's expressions.
+	* @param	QualityLevelsUsed	output array of used quality levels.
+	*/
+	void GetQualityLevelNodeUsage(TArray<bool, TInlineAllocator<EMaterialQualityLevel::Num> >& OutQualityLevelsUsed);
 
 	/**
 	 * Cache the expression texture references for this UMaterial 
@@ -1274,7 +1282,7 @@ public:
 	ENGINE_API static FMaterialCompilationFinished& OnMaterialCompilationFinished();
 
 	// For all materials, UMaterial::CacheResourceShadersForRendering
-	static void AllMaterialsCacheResourceShadersForRendering();
+	ENGINE_API static void AllMaterialsCacheResourceShadersForRendering();
 
 #if WITH_EDITORONLY_DATA
 	/**
diff --git a/Engine/Source/Runtime/Engine/Classes/Materials/MaterialInterface.h b/Engine/Source/Runtime/Engine/Classes/Materials/MaterialInterface.h
--- a/Engine/Source/Runtime/Engine/Classes/Materials/MaterialInterface.h
+++ b/Engine/Source/Runtime/Engine/Classes/Materials/MaterialInterface.h
@@ -607,4 +607,4 @@ private:
 };
 
 /** Helper function to serialize inline shader maps for the given material resources. */
-extern void SerializeInlineShaderMaps(const TMap<const class ITargetPlatform*,TArray<FMaterialResource*>>* PlatformMaterialResourcesToSave, FArchive& Ar, FMaterialResource* OutMaterialResourcesLoaded[][ERHIFeatureLevel::Num]);
\ No newline at end of file
+extern void SerializeInlineShaderMaps(const TMap<const class ITargetPlatform*, TArray<FMaterialResource*>>* PlatformMaterialResourcesToSave, FArchive& Ar, FMaterialResource* (&OutMaterialResourcesLoaded)[EMaterialQualityLevel::Num][ERHIFeatureLevel::Num]);
\ No newline at end of file
diff --git a/Engine/Source/Runtime/Engine/Engine.Build.cs b/Engine/Source/Runtime/Engine/Engine.Build.cs
--- a/Engine/Source/Runtime/Engine/Engine.Build.cs
+++ b/Engine/Source/Runtime/Engine/Engine.Build.cs
@@ -86,7 +86,8 @@ public class Engine : ModuleRules
 				"Projects",
 				"Niagara",
                 "Internationalization",
-                "PacketHandler"
+                "PacketHandler",
+                "MaterialShaderQualitySettings",
 			}
         );
 
@@ -105,7 +106,8 @@ public class Engine : ModuleRules
         CircularlyReferencedDependentModules.Add("AIModule");
 		CircularlyReferencedDependentModules.Add("Landscape");
         CircularlyReferencedDependentModules.Add("UMG");
-		CircularlyReferencedDependentModules.Add("Niagara");
+        CircularlyReferencedDependentModules.Add("Niagara");
+        CircularlyReferencedDependentModules.Add("MaterialShaderQualitySettings");
 
 		// The AnimGraphRuntime module is not needed by Engine proper, but it is loaded in LaunchEngineLoop.cpp,
 		// and needs to be listed in an always-included module in order to be compiled into standalone games
diff --git a/Engine/Source/Runtime/Engine/Private/Materials/Material.cpp b/Engine/Source/Runtime/Engine/Private/Materials/Material.cpp
--- a/Engine/Source/Runtime/Engine/Private/Materials/Material.cpp
+++ b/Engine/Source/Runtime/Engine/Private/Materials/Material.cpp
@@ -31,6 +31,7 @@
 #include "ComponentRecreateRenderStateContext.h"
 #include "ShaderCompiler.h"
 #include "Materials/MaterialParameterCollection.h"
+#include "MaterialShaderQualitySettings.h"
 #if WITH_EDITOR
 #include "MessageLog.h"
 #include "UObjectToken.h"
@@ -398,7 +399,15 @@ void UMaterialInterface::AssertDefaultMaterialsPostLoaded()
 	}
 }
 
-void SerializeInlineShaderMaps(const TMap<const ITargetPlatform*,TArray<FMaterialResource*>>* PlatformMaterialResourcesToSavePtr, FArchive& Ar, FMaterialResource* OutMaterialResourcesLoaded[][ERHIFeatureLevel::Num])
+static TAutoConsoleVariable<int32> CVarDiscardUnusedQualityLevels(
+	TEXT("r.DiscardUnusedQuality"),
+	0,
+	TEXT("Wether to keep or discard unused quality level shadermaps in memory.\n")
+	TEXT("0: keep all quality levels in memory. (default)\n")
+	TEXT("1: Discard unused quality levels on load."),
+	ECVF_ReadOnly);
+
+void SerializeInlineShaderMaps(const TMap<const ITargetPlatform*, TArray<FMaterialResource*>>* PlatformMaterialResourcesToSavePtr, FArchive& Ar, FMaterialResource* (&OutMaterialResourcesLoaded)[EMaterialQualityLevel::Num][ERHIFeatureLevel::Num])
 {
 	if (Ar.IsSaving())
 	{
@@ -444,10 +453,16 @@ void SerializeInlineShaderMaps(const TMap<const ITargetPlatform*,TArray<FMateria
 			LoadedResources.Add(LoadedResource);
 		}
 
-		// Apply in 2 passes - first pass is for shader maps without a specified quality level
-		// Second pass is where shader maps with a specified quality level override
-		for (int32 PassIndex = 0; PassIndex < 2; PassIndex++)
+		if (CVarDiscardUnusedQualityLevels.GetValueOnAnyThread())
 		{
+			// Map EMaterialQualityLevel to a score.
+			// Higher quality levels are of increasing weight such that lower quality is preferred when neighbouring hi/lo QLs exist.
+			static const int32 QualityScores[EMaterialQualityLevel::Num + 1] = { 0, 3, 1, 10 };
+
+			int32 DesiredQL = (int32)GetCachedScalabilityCVars().MaterialQualityLevel;
+			check(DesiredQL < EMaterialQualityLevel::Num);
+			const int32 DesiredScore = QualityScores[DesiredQL];
+
 			for (int32 ResourceIndex = 0; ResourceIndex < LoadedResources.Num(); ResourceIndex++)
 			{
 				FMaterialResource& LoadedResource = LoadedResources[ResourceIndex];
@@ -456,26 +471,74 @@ void SerializeInlineShaderMaps(const TMap<const ITargetPlatform*,TArray<FMateria
 				if (LoadedShaderMap && LoadedShaderMap->GetShaderPlatform() == GMaxRHIShaderPlatform)
 				{
 					EMaterialQualityLevel::Type LoadedQualityLevel = LoadedShaderMap->GetShaderMapId().QualityLevel;
+					LoadedQualityLevel = LoadedQualityLevel == EMaterialQualityLevel::Num ? EMaterialQualityLevel::High : LoadedQualityLevel;
 					ERHIFeatureLevel::Type LoadedFeatureLevel = LoadedShaderMap->GetShaderMapId().FeatureLevel;
 
-					for (int32 QualityLevelIndex = 0; QualityLevelIndex < EMaterialQualityLevel::Num; QualityLevelIndex++)
+					// find current QL:
+					int32 CurrentQL = (int32)EMaterialQualityLevel::Num;
+					for (int32 i = 0; i < EMaterialQualityLevel::Num; i++)
+					{
+						FMaterialResource* MaterialResource = OutMaterialResourcesLoaded[i][LoadedFeatureLevel];
+						if (MaterialResource != nullptr && MaterialResource->GetGameThreadShaderMap() != nullptr)
+						{
+							int32 FoundQL = MaterialResource->GetGameThreadShaderMap()->GetShaderMapId().QualityLevel;
+							CurrentQL = FoundQL == EMaterialQualityLevel::Num ? EMaterialQualityLevel::High : FoundQL;
+						}
+					}
+
+					// Determine if this is a better match than our current shader map.
+					const int32 CurrentScore = FMath::Abs(QualityScores[CurrentQL] - DesiredScore);
+					const int32 PotentialScore = FMath::Abs(QualityScores[LoadedQualityLevel] - DesiredScore);
+					if (PotentialScore < CurrentScore)
 					{
-						// Apply to all resources in the first pass if the shader map does not have a quality level specified
-						if ((PassIndex == 0 && LoadedQualityLevel == EMaterialQualityLevel::Num)
-							// Apply to just the corresponding resource in the second pass if the shader map has a quality level specified
-							|| (PassIndex == 1 && QualityLevelIndex == LoadedQualityLevel))
+						// replace existing shadermap with loadedshadermap.
+						for (int32 QualityLevelIndex = 0; QualityLevelIndex < EMaterialQualityLevel::Num; QualityLevelIndex++)
 						{
 							if (!OutMaterialResourcesLoaded[QualityLevelIndex][LoadedFeatureLevel])
 							{
 								OutMaterialResourcesLoaded[QualityLevelIndex][LoadedFeatureLevel] = new FMaterialResource();
 							}
-
+							OutMaterialResourcesLoaded[QualityLevelIndex][LoadedFeatureLevel]->ReleaseShaderMap();
 							OutMaterialResourcesLoaded[QualityLevelIndex][LoadedFeatureLevel]->SetInlineShaderMap(LoadedShaderMap);
 						}
 					}
 				}
 			}
 		}
+		else
+		{
+			// Apply in 2 passes - first pass is for shader maps without a specified quality level
+			// Second pass is where shader maps with a specified quality level override
+			for (int32 PassIndex = 0; PassIndex < 2; PassIndex++)
+			{
+				for (int32 ResourceIndex = 0; ResourceIndex < LoadedResources.Num(); ResourceIndex++)
+				{
+					FMaterialResource& LoadedResource = LoadedResources[ResourceIndex];
+					FMaterialShaderMap* LoadedShaderMap = LoadedResource.GetGameThreadShaderMap();
+
+					if (LoadedShaderMap && LoadedShaderMap->GetShaderPlatform() == GMaxRHIShaderPlatform)
+					{
+						EMaterialQualityLevel::Type LoadedQualityLevel = LoadedShaderMap->GetShaderMapId().QualityLevel;
+						ERHIFeatureLevel::Type LoadedFeatureLevel = LoadedShaderMap->GetShaderMapId().FeatureLevel;
+						for (int32 QualityLevelIndex = 0; QualityLevelIndex < EMaterialQualityLevel::Num; QualityLevelIndex++)
+						{
+							// Apply to all resources in the first pass if the shader map does not have a quality level specified
+							if ((PassIndex == 0 && LoadedQualityLevel == EMaterialQualityLevel::Num)
+								// Apply to just the corresponding resource in the second pass if the shader map has a quality level specified
+								|| (PassIndex == 1 && QualityLevelIndex == LoadedQualityLevel))
+							{
+								if (!OutMaterialResourcesLoaded[QualityLevelIndex][LoadedFeatureLevel])
+								{
+									OutMaterialResourcesLoaded[QualityLevelIndex][LoadedFeatureLevel] = new FMaterialResource();
+								}
+
+								OutMaterialResourcesLoaded[QualityLevelIndex][LoadedFeatureLevel]->SetInlineShaderMap(LoadedShaderMap);
+							}
+						}
+					}
+				}
+			}
+		}
 	}
 }
 
@@ -1788,7 +1851,7 @@ void UMaterial::CacheResourceShadersForCooking(EShaderPlatform ShaderPlatform, T
 	ERHIFeatureLevel::Type TargetFeatureLevel = GetMaxSupportedFeatureLevel(ShaderPlatform);
 
 	TArray<bool, TInlineAllocator<EMaterialQualityLevel::Num> > QualityLevelsUsed;
-	GetQualityLevelNodeUsage(QualityLevelsUsed);
+	GetQualityLevelUsage(QualityLevelsUsed, ShaderPlatform);
 
 	bool bAnyQualityLevelUsed = false;
 
@@ -2127,6 +2190,20 @@ void UMaterial::BackwardsCompatibilityInputConversion()
 #endif // WITH_EDITOR
 }
 
+void UMaterial::GetQualityLevelUsage(TArray<bool, TInlineAllocator<EMaterialQualityLevel::Num> >& OutQualityLevelsUsed, EShaderPlatform ShaderPlatform)
+{
+	GetQualityLevelNodeUsage(OutQualityLevelsUsed);
+
+	// OR in the quality overrides if we're a valid shader platform
+	if (ShaderPlatform != SP_NumPlatforms)
+	{
+		const UShaderPlatformQualitySettings* MaterialQualitySettings = UMaterialShaderQualitySettings::Get()->GetShaderPlatformQualitySettings(ShaderPlatform);
+		OutQualityLevelsUsed[EMaterialQualityLevel::Low] |= MaterialQualitySettings->GetQualityOverrides(EMaterialQualityLevel::Low).bEnableOverride;
+		OutQualityLevelsUsed[EMaterialQualityLevel::Medium] |= MaterialQualitySettings->GetQualityOverrides(EMaterialQualityLevel::Medium).bEnableOverride;
+		OutQualityLevelsUsed[EMaterialQualityLevel::High] |= MaterialQualitySettings->GetQualityOverrides(EMaterialQualityLevel::High).bEnableOverride;
+	}
+}
+
 void UMaterial::GetQualityLevelNodeUsage(TArray<bool, TInlineAllocator<EMaterialQualityLevel::Num> >& OutQualityLevelsUsed)
 {
 	OutQualityLevelsUsed.AddZeroed(EMaterialQualityLevel::Num);
@@ -2181,12 +2258,12 @@ void UMaterial::GetQualityLevelNodeUsage(TArray<bool, TInlineAllocator<EMaterial
 
 void UMaterial::UpdateResourceAllocations()
 {
-	TArray<bool, TInlineAllocator<EMaterialQualityLevel::Num> > QualityLevelsUsed;
-	GetQualityLevelNodeUsage(QualityLevelsUsed);
-
-	for (int32 QualityLevelIndex = 0; QualityLevelIndex < EMaterialQualityLevel::Num; QualityLevelIndex++)
+	for (int32 FeatureLevelIndex = 0; FeatureLevelIndex < ERHIFeatureLevel::Num; FeatureLevelIndex++)
 	{
-		for (int32 FeatureLevelIndex = 0; FeatureLevelIndex < ERHIFeatureLevel::Num; FeatureLevelIndex++)
+		TArray<bool, TInlineAllocator<EMaterialQualityLevel::Num> > QualityLevelsUsed;
+		EShaderPlatform ShaderPlatform = GShaderPlatformForFeatureLevel[FeatureLevelIndex];
+		GetQualityLevelUsage(QualityLevelsUsed, ShaderPlatform);
+		for (int32 QualityLevelIndex = 0; QualityLevelIndex < EMaterialQualityLevel::Num; QualityLevelIndex++)
 		{
 			FMaterialResource*& CurrentResource = MaterialResources[QualityLevelIndex][FeatureLevelIndex];
 
@@ -2195,9 +2272,9 @@ void UMaterial::UpdateResourceAllocations()
 				CurrentResource = AllocateResource();
 			}
 
-			const bool bQualityLevelHasDifferentNodes = QualityLevelsUsed[QualityLevelIndex];
+			const bool bHasQualityLevelUsage = QualityLevelsUsed[QualityLevelIndex];
 			// Setup transient FMaterialResource properties that are needed to use this resource for rendering or compilation
-			CurrentResource->SetMaterial(this, (EMaterialQualityLevel::Type)QualityLevelIndex, bQualityLevelHasDifferentNodes, (ERHIFeatureLevel::Type)FeatureLevelIndex);
+			CurrentResource->SetMaterial(this, (EMaterialQualityLevel::Type)QualityLevelIndex, bHasQualityLevelUsage, (ERHIFeatureLevel::Type)FeatureLevelIndex);
 		}
 	}
 }
diff --git a/Engine/Source/Runtime/Engine/Private/Materials/MaterialInstance.cpp b/Engine/Source/Runtime/Engine/Private/Materials/MaterialInstance.cpp
--- a/Engine/Source/Runtime/Engine/Private/Materials/MaterialInstance.cpp
+++ b/Engine/Source/Runtime/Engine/Private/Materials/MaterialInstance.cpp
@@ -1297,12 +1297,13 @@ void UMaterialInstance::UpdatePermutationAllocations()
 	{
 		UMaterial* BaseMaterial = GetMaterial();
 
-		TArray<bool, TInlineAllocator<EMaterialQualityLevel::Num> > QualityLevelsUsed;
-		BaseMaterial->GetQualityLevelNodeUsage(QualityLevelsUsed);
-
-		for (int32 QualityLevelIndex = 0; QualityLevelIndex < EMaterialQualityLevel::Num; QualityLevelIndex++)
+		for (int32 FeatureLevelIndex = 0; FeatureLevelIndex < ERHIFeatureLevel::Num; FeatureLevelIndex++)
 		{
-			for (int32 FeatureLevelIndex = 0; FeatureLevelIndex < ERHIFeatureLevel::Num; FeatureLevelIndex++)
+			EShaderPlatform ShaderPlatform = GShaderPlatformForFeatureLevel[FeatureLevelIndex];
+			TArray<bool, TInlineAllocator<EMaterialQualityLevel::Num> > QualityLevelsUsed;
+			BaseMaterial->GetQualityLevelUsage(QualityLevelsUsed, ShaderPlatform);
+
+			for (int32 QualityLevelIndex = 0; QualityLevelIndex < EMaterialQualityLevel::Num; QualityLevelIndex++)
 			{
 				FMaterialResource*& CurrentResource = StaticPermutationMaterialResources[QualityLevelIndex][FeatureLevelIndex];
 
@@ -1362,7 +1363,7 @@ void UMaterialInstance::CacheResourceShadersForCooking(EShaderPlatform ShaderPla
 		UMaterial* BaseMaterial = GetMaterial();
 
 		TArray<bool, TInlineAllocator<EMaterialQualityLevel::Num> > QualityLevelsUsed;
-		BaseMaterial->GetQualityLevelNodeUsage(QualityLevelsUsed);
+		BaseMaterial->GetQualityLevelUsage(QualityLevelsUsed, ShaderPlatform);
 
 		TArray<FMaterialResource*> ResourcesToCache;
 		ERHIFeatureLevel::Type TargetFeatureLevel = GetMaxSupportedFeatureLevel(ShaderPlatform);
diff --git a/Engine/Source/Runtime/Engine/Private/Materials/MaterialShared.cpp b/Engine/Source/Runtime/Engine/Private/Materials/MaterialShared.cpp
--- a/Engine/Source/Runtime/Engine/Private/Materials/MaterialShared.cpp
+++ b/Engine/Source/Runtime/Engine/Private/Materials/MaterialShared.cpp
@@ -24,12 +24,14 @@
 
 #include "VertexFactory.h"
 #include "RendererInterface.h"
+#include "MaterialShaderQualitySettings.h"
 DEFINE_LOG_CATEGORY(LogMaterial);
 
 FName MaterialQualityLevelNames[] = 
 {
 	FName(TEXT("Low")),
 	FName(TEXT("High")),
+	FName(TEXT("Medium")),
 	FName(TEXT("Num"))
 };
 
@@ -268,7 +270,7 @@ void FMaterial::GetShaderMapId(EShaderPlatform Platform, FMaterialShaderMapId& O
 	OutId.QualityLevel = GetQualityLevelForShaderMapId();
 	OutId.FeatureLevel = GetFeatureLevel();
 	OutId.SetShaderDependencies(ShaderTypes, VFTypes);
-	GetReferencedTexturesHash(OutId.TextureReferencesHash);
+	GetReferencedTexturesHash(Platform, OutId.TextureReferencesHash);
 }
 
 EMaterialTessellationMode FMaterial::GetTessellationMode() const 
@@ -1881,7 +1883,7 @@ void FMaterial::GetDependentShaderAndVFTypes(EShaderPlatform Platform, TArray<FS
 	OutVFTypes.Sort(FCompareVertexFactoryTypes());
 }
 
-void FMaterial::GetReferencedTexturesHash(FSHAHash& OutHash) const
+void FMaterial::GetReferencedTexturesHash(EShaderPlatform Platform, FSHAHash& OutHash) const
 {
 	FSHA1 HashState;
 
@@ -1899,6 +1901,15 @@ void FMaterial::GetReferencedTexturesHash(FSHAHash& OutHash) const
 		HashState.UpdateWithString(*TextureName, TextureName.Len());
 	}
 
+	// TODO:
+	// Appending the quality settings for this platform,
+	// this is being done here to avoid bumping EUnrealEngineObjectUE4Version for 4.10
+	// must be fixed for 4.11.
+	if (bHasQualityLevelUsage)
+	{
+		UMaterialShaderQualitySettings::Get()->GetShaderPlatformQualitySettings(Platform)->AppendToHashState(GetQualityLevelForShaderMapId(), HashState);
+	}
+
 	HashState.Final();
 	HashState.GetHash(&OutHash.Hash[0]);
 }
diff --git a/Engine/Source/Runtime/Engine/Private/UnrealEngine.cpp b/Engine/Source/Runtime/Engine/Private/UnrealEngine.cpp
--- a/Engine/Source/Runtime/Engine/Private/UnrealEngine.cpp
+++ b/Engine/Source/Runtime/Engine/Private/UnrealEngine.cpp
@@ -297,7 +297,7 @@ void ScalabilityCVarsSinkCallback()
 
 	{
 		static const auto MaterialQualityLevelVar = IConsoleManager::Get().FindTConsoleVariableDataInt(TEXT("r.MaterialQualityLevel"));
-		LocalScalabilityCVars.MaterialQualityLevel = (EMaterialQualityLevel::Type)FMath::Clamp(MaterialQualityLevelVar->GetValueOnGameThread(), 0, 1);
+		LocalScalabilityCVars.MaterialQualityLevel = (EMaterialQualityLevel::Type)FMath::Clamp(MaterialQualityLevelVar->GetValueOnGameThread(), 0, (int32)EMaterialQualityLevel::Num - 1);
 	}
 
 	{
diff --git a/Engine/Source/Runtime/Engine/Public/MaterialShared.h b/Engine/Source/Runtime/Engine/Public/MaterialShared.h
--- a/Engine/Source/Runtime/Engine/Public/MaterialShared.h
+++ b/Engine/Source/Runtime/Engine/Public/MaterialShared.h
@@ -872,7 +872,7 @@ public:
 		RenderingThreadShaderMap(NULL),
 		Id_DEPRECATED(0,0,0,0),
 		QualityLevel(EMaterialQualityLevel::High),
-		bQualityLevelHasDifferentNodes(false),
+		bHasQualityLevelUsage(false),
 		FeatureLevel(ERHIFeatureLevel::SM4),
 		bContainsInlineShaders(false)
 	{}
@@ -1154,10 +1154,10 @@ protected:
 	/** Useful for debugging. */
 	virtual FString GetBaseMaterialPathName() const { return TEXT(""); }
 
-	void SetQualityLevelProperties(EMaterialQualityLevel::Type InQualityLevel, bool bInQualityLevelHasDifferentNodes, ERHIFeatureLevel::Type InFeatureLevel)
+	void SetQualityLevelProperties(EMaterialQualityLevel::Type InQualityLevel, bool bInHasQualityLevelUsage, ERHIFeatureLevel::Type InFeatureLevel)
 	{
 		QualityLevel = InQualityLevel;
-		bQualityLevelHasDifferentNodes = bInQualityLevelHasDifferentNodes;
+		bHasQualityLevelUsage = bInHasQualityLevelUsage;
 		FeatureLevel = InFeatureLevel;
 	}
 
@@ -1218,7 +1218,7 @@ private:
 	EMaterialQualityLevel::Type QualityLevel;
 
 	/** Whether this material has quality level specific nodes. */
-	bool bQualityLevelHasDifferentNodes;
+	bool bHasQualityLevelUsage;
 
 	/** Feature level that this material is representing. */
 	ERHIFeatureLevel::Type FeatureLevel;
@@ -1250,11 +1250,11 @@ private:
 	 */
 	ENGINE_API FShader* GetShader(class FMeshMaterialShaderType* ShaderType, FVertexFactoryType* VertexFactoryType) const;
 
-	void GetReferencedTexturesHash(FSHAHash& OutHash) const;
+	void GetReferencedTexturesHash(EShaderPlatform Platform, FSHAHash& OutHash) const;
 
 	EMaterialQualityLevel::Type GetQualityLevelForShaderMapId() const 
 	{
-		return bQualityLevelHasDifferentNodes ? QualityLevel : EMaterialQualityLevel::Num;
+		return bHasQualityLevelUsage ? QualityLevel : EMaterialQualityLevel::Num;
 	}
 
 	friend class FMaterialShaderMap;
diff --git a/Engine/Source/Runtime/Engine/Public/SceneTypes.h b/Engine/Source/Runtime/Engine/Public/SceneTypes.h
--- a/Engine/Source/Runtime/Engine/Public/SceneTypes.h
+++ b/Engine/Source/Runtime/Engine/Public/SceneTypes.h
@@ -116,6 +116,7 @@ namespace EMaterialQualityLevel
 	{
 		Low,
 		High,
+		Medium,
 		Num
 	};
 }
diff --git a/Engine/Source/Runtime/MaterialShaderQualitySettings/Classes/MaterialShaderQualitySettings.h b/Engine/Source/Runtime/MaterialShaderQualitySettings/Classes/MaterialShaderQualitySettings.h
new file mode 100644
--- /dev/null
+++ b/Engine/Source/Runtime/MaterialShaderQualitySettings/Classes/MaterialShaderQualitySettings.h
@@ -0,0 +1,49 @@
+// Copyright 1998-2015 Epic Games, Inc. All Rights Reserved.
+
+#pragma once
+
+#include "RHIDefinitions.h"
+#include "ShaderPlatformQualitySettings.h"
+#include "MaterialShaderQualitySettings.generated.h"
+
+//UCLASS(config = Engine, defaultconfig)
+UCLASS()
+class MATERIALSHADERQUALITYSETTINGS_API UMaterialShaderQualitySettings : public UObject
+{
+public:
+	GENERATED_UCLASS_BODY()
+
+	UShaderPlatformQualitySettings* GetShaderPlatformQualitySettings(FName PlatformName);
+
+	const UShaderPlatformQualitySettings* GetShaderPlatformQualitySettings(EShaderPlatform ShaderPlatform);
+
+#if WITH_EDITOR
+	// Override GetForwardShadingQuality() return value with the specified platform's settings.
+	// An empty PlatformName or otherwise non existant platform will cause GetForwardShadingQuality() 
+	// to revert to its default behaviour.
+	void SetPreviewPlatform(FName PlatformName);
+	const FName& GetPreviewPlatform();
+
+	// UObject interface
+	virtual void PostEditChangeProperty(struct FPropertyChangedEvent& PropertyChangedEvent) override;
+	virtual void PostInitProperties() override;
+	// End of UObject interface
+#endif
+
+	static UMaterialShaderQualitySettings* Get();
+
+private:
+	UShaderPlatformQualitySettings* GetOrCreatePlatformSettings(FName PlatformName);
+
+	UPROPERTY()
+	TMap<FName, UShaderPlatformQualitySettings*> ForwardSettingMap;
+
+	UShaderPlatformQualitySettings* CurrentPlatformSettings;
+
+#if WITH_EDITOR
+	UShaderPlatformQualitySettings* PreviewPlatformSettings;
+	FName PreviewPlatformName;
+#endif
+
+	static class UMaterialShaderQualitySettings* RenderQualitySingleton;
+};
diff --git a/Engine/Source/Runtime/MaterialShaderQualitySettings/Classes/MaterialShaderQualitySettingsCustomization.h b/Engine/Source/Runtime/MaterialShaderQualitySettings/Classes/MaterialShaderQualitySettingsCustomization.h
new file mode 100644
--- /dev/null
+++ b/Engine/Source/Runtime/MaterialShaderQualitySettings/Classes/MaterialShaderQualitySettingsCustomization.h
@@ -0,0 +1,38 @@
+// Copyright 1998-2015 Epic Games, Inc. All Rights Reserved.
+
+#pragma once
+
+#if WITH_EDITOR
+#include "EditorStyle.h"
+#include "PropertyEditorModule.h"
+
+//////////////////////////////////////////////////////////////////////////
+// FMaterialShaderQualitySettingsCustomization
+
+typedef SListView< TSharedPtr< struct FShaderQualityOverridesListItem > > SMaterialQualityOverridesListView;
+
+DECLARE_DELEGATE(FOnUpdateMaterialShaderQuality);
+
+class FMaterialShaderQualitySettingsCustomization : public IDetailCustomization
+{
+public:
+	// Makes a new instance of this detail layout class for a specific detail view requesting it
+	MATERIALSHADERQUALITYSETTINGS_API static TSharedRef<IDetailCustomization> MakeInstance(const FOnUpdateMaterialShaderQuality InUpdateMaterials);
+
+	FMaterialShaderQualitySettingsCustomization(const FOnUpdateMaterialShaderQuality InUpdateMaterials);
+
+	// IDetailCustomization interface
+	virtual void CustomizeDetails(IDetailLayoutBuilder& DetailLayout) override;
+	// End of IDetailCustomization interface
+
+private:
+
+	TSharedRef<ITableRow> HandleGenerateQualityWidget(TSharedPtr< struct FShaderQualityOverridesListItem > InItem, const TSharedRef<STableViewBase>& OwnerTable);
+	FReply UpdatePreviewShaders();
+	FOnUpdateMaterialShaderQuality UpdateMaterials;
+	FMaterialShaderQualitySettingsCustomization();
+	TSharedPtr<SMaterialQualityOverridesListView> MaterialQualityOverridesListView;
+	TArray< TSharedPtr< struct FShaderQualityOverridesListItem > >	QualityOverrideListSource;
+};
+
+#endif
\ No newline at end of file
diff --git a/Engine/Source/Runtime/MaterialShaderQualitySettings/Classes/ShaderPlatformQualitySettings.h b/Engine/Source/Runtime/MaterialShaderQualitySettings/Classes/ShaderPlatformQualitySettings.h
new file mode 100644
--- /dev/null
+++ b/Engine/Source/Runtime/MaterialShaderQualitySettings/Classes/ShaderPlatformQualitySettings.h
@@ -0,0 +1,63 @@
+// Copyright 1998-2015 Epic Games, Inc. All Rights Reserved.
+
+#pragma once
+
+#include "SceneTypes.h"
+#include "ShaderPlatformQualitySettings.generated.h"
+
+/**
+* 
+*/
+
+// FMaterialQualityOverrides represents the full set of possible material overrides per quality level.
+USTRUCT()
+struct FMaterialQualityOverrides
+{
+public:
+	GENERATED_USTRUCT_BODY()
+
+	UPROPERTY(Config, Meta = (DisplayName = "Enable Quality Override"))
+	bool bEnableOverride;
+
+	UPROPERTY(Config, Meta = (DisplayName = "Force Fully Rough"))
+	bool bForceFullyRough;
+
+	UPROPERTY(Config, Meta = (DisplayName = "Force Non-metal"))
+	bool bForceNonMetal;
+
+	UPROPERTY(Config, Meta = (DisplayName = "Disable Lightmap directionality"))
+	bool bForceDisableLMDirectionality;
+};
+
+
+UCLASS(config = Engine, defaultconfig, perObjectConfig)
+class MATERIALSHADERQUALITYSETTINGS_API UShaderPlatformQualitySettings : public UObject
+{
+public:
+	GENERATED_UCLASS_BODY()
+
+	UPROPERTY(Config)
+	FMaterialQualityOverrides QualityOverrides[EMaterialQualityLevel::Num];
+
+	FMaterialQualityOverrides& GetQualityOverrides(EMaterialQualityLevel::Type QualityLevel)
+	{
+		check(QualityLevel<EMaterialQualityLevel::Num);
+		return QualityOverrides[(int32)QualityLevel];
+	}
+
+	const FMaterialQualityOverrides& GetQualityOverrides(EMaterialQualityLevel::Type QualityLevel) const
+	{
+		check(QualityLevel < EMaterialQualityLevel::Num);
+		return QualityOverrides[(int32)QualityLevel];
+	}
+
+#if WITH_EDITOR
+	// UObject interface
+	virtual void PostEditChangeProperty(struct FPropertyChangedEvent& PropertyChangedEvent) override;
+	virtual void PostInitProperties() override;
+	// End of UObject interface
+#endif
+
+	void BuildHash(EMaterialQualityLevel::Type QualityLevel, class FSHAHash& OutHash) const;
+	void AppendToHashState(EMaterialQualityLevel::Type QualityLevel, class FSHA1& HashState) const;
+};
diff --git a/Engine/Source/Runtime/MaterialShaderQualitySettings/MaterialShaderQualitySettings.Build.cs b/Engine/Source/Runtime/MaterialShaderQualitySettings/MaterialShaderQualitySettings.Build.cs
new file mode 100644
--- /dev/null
+++ b/Engine/Source/Runtime/MaterialShaderQualitySettings/MaterialShaderQualitySettings.Build.cs
@@ -0,0 +1,35 @@
+// Copyright 1998-2015 Epic Games, Inc. All Rights Reserved.
+
+using UnrealBuildTool;
+
+public class MaterialShaderQualitySettings : ModuleRules
+{
+    public MaterialShaderQualitySettings(TargetInfo Target)
+	{
+		PrivateDependencyModuleNames.AddRange(
+			new string[]
+			{
+				"Core",
+				"CoreUObject",
+                "RHI",
+			}
+		);
+        PrivateIncludePathModuleNames.Add("Engine");
+
+        if (UEBuildConfiguration.bBuildEditor == true)
+		{
+            PrivateDependencyModuleNames.AddRange(
+                new string[] {            
+       				"Slate",
+				    "SlateCore",
+				    "EditorStyle",
+                    "PropertyEditor",
+                    "TargetPlatform",
+                    "InputCore",
+                }
+            );
+
+		}
+
+	}
+}
diff --git a/Engine/Source/Runtime/MaterialShaderQualitySettings/Private/MaterialShaderQualitySettings.cpp b/Engine/Source/Runtime/MaterialShaderQualitySettings/Private/MaterialShaderQualitySettings.cpp
new file mode 100644
--- /dev/null
+++ b/Engine/Source/Runtime/MaterialShaderQualitySettings/Private/MaterialShaderQualitySettings.cpp
@@ -0,0 +1,162 @@
+// Copyright 1998-2015 Epic Games, Inc. All Rights Reserved.
+
+#include "MaterialShaderQualitySettingsPrivatePCH.h"
+#include "MaterialShaderQualitySettings.h"
+#include "ShaderPlatformQualitySettings.h"
+#include "RHI.h"
+#include "SecureHash.h"
+
+#if WITH_EDITOR
+#include "TargetPlatform.h"
+#include "PlatformInfo.h"
+#endif
+
+UMaterialShaderQualitySettings* UMaterialShaderQualitySettings::RenderQualitySingleton = nullptr;
+
+UMaterialShaderQualitySettings::UMaterialShaderQualitySettings(const FObjectInitializer& ObjectInitializer)
+	: Super(ObjectInitializer)
+{
+}
+
+UMaterialShaderQualitySettings* UMaterialShaderQualitySettings::Get()
+{
+	if( RenderQualitySingleton == nullptr )
+	{
+		static const TCHAR* SettingsContainerName = TEXT("MaterialShaderQualitySettingsContainer");
+
+		RenderQualitySingleton = FindObject<UMaterialShaderQualitySettings>(GetTransientPackage(), SettingsContainerName);
+
+		if (RenderQualitySingleton == nullptr)
+		{
+			RenderQualitySingleton = NewObject<UMaterialShaderQualitySettings>(GetTransientPackage(), UMaterialShaderQualitySettings::StaticClass(), SettingsContainerName);
+			RenderQualitySingleton->AddToRoot();
+		}
+		RenderQualitySingleton->CurrentPlatformSettings = RenderQualitySingleton->GetShaderPlatformQualitySettings(FPlatformProperties::PlatformName());
+
+		// LegacyShaderPlatformToShaderFormat(EShaderPlatform)
+		// GShaderPlatformForFeatureLevel
+		// GMaxRHIFeatureLevel
+
+		// populate shader platforms
+		RenderQualitySingleton->CurrentPlatformSettings = RenderQualitySingleton->GetShaderPlatformQualitySettings(FPlatformProperties::PlatformName());
+
+	}
+	return RenderQualitySingleton;
+}
+
+#if WITH_EDITOR
+const FName& UMaterialShaderQualitySettings::GetPreviewPlatform()
+{
+	return PreviewPlatformName;
+}
+
+void UMaterialShaderQualitySettings::SetPreviewPlatform(FName PlatformName)
+{
+	 UShaderPlatformQualitySettings** FoundPlatform = ForwardSettingMap.Find(PlatformName);
+	 PreviewPlatformSettings = FoundPlatform == nullptr ? nullptr : *FoundPlatform;
+	 PreviewPlatformName = PlatformName;
+}
+
+void UMaterialShaderQualitySettings::PostEditChangeProperty(struct FPropertyChangedEvent& PropertyChangedEvent)
+{
+	Super::PostEditChangeProperty(PropertyChangedEvent);
+}
+
+void UMaterialShaderQualitySettings::PostInitProperties()
+{
+	Super::PostInitProperties();
+}
+
+#endif
+
+UShaderPlatformQualitySettings* UMaterialShaderQualitySettings::GetOrCreatePlatformSettings(FName PlatformName)
+{
+	auto* PlatformSettings = ForwardSettingMap.Find(PlatformName);
+	if (PlatformSettings == nullptr)
+	{
+		FString ObjectName("ForwardShadingQuality_");
+		PlatformName.AppendString(ObjectName);
+
+		auto* ForwardQualitySettings = FindObject<UShaderPlatformQualitySettings>(this, *ObjectName);
+		if (ForwardQualitySettings == nullptr)
+		{
+			FName ForwardSettingsName(*ObjectName);
+			ForwardQualitySettings = NewObject<UShaderPlatformQualitySettings>(this, UShaderPlatformQualitySettings::StaticClass(), FName(*ObjectName));
+			ForwardQualitySettings->LoadConfig();
+		}
+
+		return ForwardSettingMap.Add(PlatformName, ForwardQualitySettings);
+	}
+	return *PlatformSettings;
+}
+
+static const FName GetPlatformNameFromShaderPlatform(EShaderPlatform Platform)
+{
+	return LegacyShaderPlatformToShaderFormat(Platform);
+}
+
+const UShaderPlatformQualitySettings* UMaterialShaderQualitySettings::GetShaderPlatformQualitySettings(EShaderPlatform ShaderPlatform)
+{
+ #if WITH_EDITORONLY_DATA
+	// TODO: discuss this, in order to preview render quality settings we override the 
+	// requested platform's settings.
+	// However we do not know if we are asking for the editor preview window (override able) or for thumbnails, cooking purposes etc.. (Must not override)
+	// The code below 'works' because desktop platforms do not cook for ES2 preview.
+	if (IsPCPlatform(ShaderPlatform) && IsES2Platform(ShaderPlatform))
+	{
+		// Can check this cant be cooked by iterating through target platforms and shader formats to ensure it's not covered.
+		if (PreviewPlatformSettings != nullptr)
+		{
+			return PreviewPlatformSettings;
+		}
+	}
+#endif
+	return GetShaderPlatformQualitySettings(GetPlatformNameFromShaderPlatform(ShaderPlatform));
+
+}
+
+UShaderPlatformQualitySettings* UMaterialShaderQualitySettings::GetShaderPlatformQualitySettings(FName PlatformName)
+{
+	return GetOrCreatePlatformSettings(PlatformName);
+}
+
+
+//////////////////////////////////////////////////////////////////////////
+
+UShaderPlatformQualitySettings::UShaderPlatformQualitySettings(const FObjectInitializer& ObjectInitializer)
+	: Super(ObjectInitializer)
+{
+	// high quality overrides are always enabled by default
+	GetQualityOverrides(EMaterialQualityLevel::High).bEnableOverride = true;
+}
+
+#if WITH_EDITOR
+
+void UShaderPlatformQualitySettings::PostEditChangeProperty(struct FPropertyChangedEvent& PropertyChangedEvent)
+{
+	Super::PostEditChangeProperty(PropertyChangedEvent);
+	SaveConfig(); // needed?
+}
+
+void UShaderPlatformQualitySettings::PostInitProperties()
+{
+	Super::PostInitProperties();
+}
+#endif
+
+void UShaderPlatformQualitySettings::BuildHash(EMaterialQualityLevel::Type QualityLevel, FSHAHash& OutHash) const
+{
+	FSHA1 Hash;
+
+	AppendToHashState(QualityLevel, Hash);
+
+	Hash.Final();
+	Hash.GetHash(&OutHash.Hash[0]);
+}
+
+void UShaderPlatformQualitySettings::AppendToHashState(EMaterialQualityLevel::Type QualityLevel, FSHA1& HashState) const
+{
+	const FMaterialQualityOverrides& QualityLevelOverrides = GetQualityOverrides(QualityLevel);
+	HashState.Update((const uint8*)&QualityLevelOverrides, sizeof(QualityLevelOverrides));
+}
+
diff --git a/Engine/Source/Runtime/MaterialShaderQualitySettings/Private/MaterialShaderQualitySettingsCustomization.cpp b/Engine/Source/Runtime/MaterialShaderQualitySettings/Private/MaterialShaderQualitySettingsCustomization.cpp
new file mode 100644
--- /dev/null
+++ b/Engine/Source/Runtime/MaterialShaderQualitySettings/Private/MaterialShaderQualitySettingsCustomization.cpp
@@ -0,0 +1,249 @@
+// Copyright 1998-2015 Epic Games, Inc. All Rights Reserved.
+
+#if WITH_EDITOR
+
+#include "MaterialShaderQualitySettingsPrivatePCH.h"
+#include "MaterialShaderQualitySettingsCustomization.h"
+#include "DetailLayoutBuilder.h"
+#include "DetailCategoryBuilder.h"
+#include "PropertyEditing.h"
+#include "ShaderPlatformQualitySettings.h"
+#include "SListView.h"
+#include "SCompoundWidget.h"
+#include "SceneTypes.h"
+#include "ShaderQualityOverridesListItem.h"
+
+#define LOCTEXT_NAMESPACE "MaterialShaderQualitySettings"
+
+//////////////////////////////////////////////////////////////////////////
+// FMaterialShaderQualitySettingsCustomization
+namespace FMaterialShaderQualitySettingsCustomizationConstants
+{
+	const FText DisabledTip = LOCTEXT("GitHubSourceRequiredToolTip", "This requires GitHub source.");
+}
+
+TSharedRef<IDetailCustomization> FMaterialShaderQualitySettingsCustomization::MakeInstance(const FOnUpdateMaterialShaderQuality InUpdateMaterials)
+{
+	return MakeShareable(new FMaterialShaderQualitySettingsCustomization(InUpdateMaterials));
+}
+
+FMaterialShaderQualitySettingsCustomization::FMaterialShaderQualitySettingsCustomization(const FOnUpdateMaterialShaderQuality InUpdateMaterials)
+{
+	UpdateMaterials = InUpdateMaterials;
+}
+
+class SQualityListItem
+	: public SMultiColumnTableRow< TSharedPtr<struct FShaderQualityOverridesListItem> >
+{
+public:
+
+	SLATE_BEGIN_ARGS(SQualityListItem) { }
+	SLATE_ARGUMENT(TSharedPtr<FShaderQualityOverridesListItem>, Item)
+	SLATE_END_ARGS()
+
+public:
+
+	void Construct(const FArguments& InArgs, const TSharedRef<STableViewBase>& InOwnerTableView)
+	{
+		Item = InArgs._Item;
+		check(Item.IsValid());
+		const auto args = FSuperRowType::FArguments();
+		SMultiColumnTableRow< TSharedPtr<FShaderQualityOverridesListItem> >::Construct(args, InOwnerTableView);
+	}
+
+	virtual TSharedRef<SWidget> GenerateWidgetForColumn(const FName& ColumnName) override
+	{
+		if (ColumnName == TEXT("Quality Option"))
+		{
+			return	SNew(SBox)
+				.HeightOverride(20)
+				.Padding(FMargin(3, 0))
+				.VAlign(VAlign_Center)
+				[
+					SNew(STextBlock)
+					.Text(FText::FromString(Item->RangeName))
+					.Font(IDetailLayoutBuilder::GetDetailFont())
+				];
+		}
+		else
+		{
+			EMaterialQualityLevel::Type MaterialQualityLevel = EMaterialQualityLevel::Num;
+
+			if (ColumnName == TEXT("Low"))
+			{
+				MaterialQualityLevel = EMaterialQualityLevel::Low;
+			}
+			else if (ColumnName == TEXT("Medium"))
+			{
+				MaterialQualityLevel = EMaterialQualityLevel::Medium;
+			}
+			else if (ColumnName == TEXT("High"))
+			{
+				MaterialQualityLevel = EMaterialQualityLevel::High;
+			}
+
+			if (MaterialQualityLevel < EMaterialQualityLevel::Num)
+			{
+				return	SNew(SBox)
+					.HeightOverride(20)
+					.Padding(FMargin(3, 0))
+					.VAlign(VAlign_Center)
+					[
+						SNew(SCheckBox)
+						.OnCheckStateChanged(this, &SQualityListItem::OnQualityRangeChanged, MaterialQualityLevel)
+						.IsChecked(this, &SQualityListItem::IsQualityRangeSet, MaterialQualityLevel)
+						.IsEnabled(this, &SQualityListItem::IsEnabled, MaterialQualityLevel)
+					];
+			}
+		}
+	
+		return SNullWidget::NullWidget;
+	}
+
+private:
+
+	ECheckBoxState IsQualityRangeSet(EMaterialQualityLevel::Type QualityLevel) const
+	{
+		bool* State = Item->QualityProperty->ContainerPtrToValuePtr<bool>(&Item->SettingContainer->GetQualityOverrides(QualityLevel));
+		return *State ? ECheckBoxState::Checked : ECheckBoxState::Unchecked;
+	}
+
+	void OnQualityRangeChanged(ECheckBoxState NewCheckedState, EMaterialQualityLevel::Type QualityLevel)
+	{
+		bool* State = Item->QualityProperty->ContainerPtrToValuePtr<bool>(&Item->SettingContainer->GetQualityOverrides(QualityLevel));
+		*State = NewCheckedState == ECheckBoxState::Checked;
+
+		// save changes.
+		// TODO: use property handles so FSettingSection::Save() does the config saving.
+		Item->SettingContainer->UpdateDefaultConfigFile();
+	}
+
+	bool IsEnabled(EMaterialQualityLevel::Type QualityLevel) const
+	{
+		bool* bEnableOverride = &Item->SettingContainer->GetQualityOverrides(QualityLevel).bEnableOverride;
+		bool* State = Item->QualityProperty->ContainerPtrToValuePtr<bool>(&Item->SettingContainer->GetQualityOverrides(QualityLevel));
+
+		// Enable all if EnableOverride is set,
+		// but also ensure all but the high quality override flag remains enabled.
+		const bool bIsOverrideProperty = bEnableOverride == State;
+		const bool bIsHighQualityOverrideProperty = (bIsOverrideProperty && QualityLevel == EMaterialQualityLevel::High);
+		return (!bIsOverrideProperty && *bEnableOverride == true) || (bIsOverrideProperty && !bIsHighQualityOverrideProperty);
+	}
+
+	TSharedPtr< FShaderQualityOverridesListItem > Item;
+};
+
+TSharedRef<ITableRow> FMaterialShaderQualitySettingsCustomization::HandleGenerateQualityWidget(TSharedPtr<FShaderQualityOverridesListItem> InItem, const TSharedRef<STableViewBase>& OwnerTable)
+{
+	return SNew(SQualityListItem, OwnerTable)
+		.Item(InItem)
+		;
+}
+
+FReply FMaterialShaderQualitySettingsCustomization::UpdatePreviewShaders()
+{
+	UpdateMaterials.ExecuteIfBound();
+	return FReply::Handled();
+}
+
+void FMaterialShaderQualitySettingsCustomization::CustomizeDetails(IDetailLayoutBuilder& DetailLayout)
+{
+	IDetailCategoryBuilder& ForwardRenderingCategory = DetailLayout.EditCategory(TEXT("Forward Rendering overrides"));
+
+	TArray<TWeakObjectPtr<UObject>> ObjectsBeingCustomized;
+	DetailLayout.GetObjectsBeingCustomized(ObjectsBeingCustomized);
+	UShaderPlatformQualitySettings* TargetObject = nullptr;
+
+	if (ObjectsBeingCustomized.Num() >= 1)
+	{
+		TargetObject = Cast<UShaderPlatformQualitySettings>(ObjectsBeingCustomized[0].Get());
+	}
+	check(TargetObject != nullptr);
+
+	// Build a list of FShaderQualityRange properties.
+	QualityOverrideListSource.Empty();
+	// iterate through possible override properties
+	for (UBoolProperty* Prop : TFieldRange<UBoolProperty>(FMaterialQualityOverrides::StaticStruct()))
+	{
+		QualityOverrideListSource.Add(MakeShareable(
+			new FShaderQualityOverridesListItem(Prop->GetMetaData("DisplayName"), Prop, TargetObject)));
+	}
+
+	ForwardRenderingCategory.AddCustomRow(LOCTEXT("ForwardRenderingMaterialOverrides", "Forward Rendering Material Overrides"))
+		[
+			SAssignNew(MaterialQualityOverridesListView, SMaterialQualityOverridesListView)
+			.ItemHeight(20.0f)
+			.ListItemsSource(&QualityOverrideListSource)
+			.OnGenerateRow(this, &FMaterialShaderQualitySettingsCustomization::HandleGenerateQualityWidget)
+			.SelectionMode(ESelectionMode::None)
+			.HeaderRow(
+			SNew(SHeaderRow)
+			// 
+			+ SHeaderRow::Column("Quality Option")
+			.HAlignCell(HAlign_Left)
+			.FillWidth(1)
+			.HeaderContentPadding(FMargin(0, 3))
+			[
+				SNew(SHorizontalBox)
+				+ SHorizontalBox::Slot()
+				.AutoWidth()
+				[
+					SNew(STextBlock)
+					.Text(LOCTEXT("MaterialQualityList_Category", "Quality Option"))
+					.Font(IDetailLayoutBuilder::GetDetailFont())
+				]
+			]
+			// 
+			+ SHeaderRow::Column("Low")
+				.HAlignCell(HAlign_Left)
+				.FillWidth(1)
+				.HeaderContentPadding(FMargin(0, 3))
+				[
+					SNew(STextBlock)
+					.Text(LOCTEXT("MaterialQualityList_Low", "Low"))
+					.Font(IDetailLayoutBuilder::GetDetailFont())
+				]
+			// 
+			+ SHeaderRow::Column("Medium")
+				.HAlignCell(HAlign_Left)
+				.FillWidth(1)
+				.HeaderContentPadding(FMargin(0, 3))
+				[
+					SNew(STextBlock)
+					.Text(LOCTEXT("MaterialQualityList_Medium", "Medium"))
+					.Font(IDetailLayoutBuilder::GetDetailFont())
+				]
+			// 
+			+ SHeaderRow::Column("High")
+				.HAlignCell(HAlign_Left)
+				.FillWidth(1)
+				.HeaderContentPadding(FMargin(0, 3))
+				[
+					SNew(STextBlock)
+					.Text(LOCTEXT("MaterialQualityList_High", "High"))
+					.Font(IDetailLayoutBuilder::GetDetailFont())
+				]
+			)
+		];
+
+	ForwardRenderingCategory.AddCustomRow(LOCTEXT("ForwardRenderingMaterialOverrides", "Forward Rendering Material Overrides"))
+	[
+		SNew(SHorizontalBox)
+		+ SHorizontalBox::Slot()
+		.AutoWidth()
+		[
+			SNew(SButton)
+			.Text(LOCTEXT("UpdatePreviewShaders", "Update preview shaders"))
+			.ToolTipText(LOCTEXT("UpdatePreviewShadersButton_Tooltip", "Updates the editor to reflect changes to quality settings."))
+			.OnClicked(this, &FMaterialShaderQualitySettingsCustomization::UpdatePreviewShaders)
+		]
+	];
+
+}
+
+
+//////////////////////////////////////////////////////////////////////////
+
+#undef LOCTEXT_NAMESPACE
+
+#endif // WITH_EDITOR
\ No newline at end of file
diff --git a/Engine/Source/Runtime/MaterialShaderQualitySettings/Private/MaterialShaderQualitySettingsModule.cpp b/Engine/Source/Runtime/MaterialShaderQualitySettings/Private/MaterialShaderQualitySettingsModule.cpp
new file mode 100644
--- /dev/null
+++ b/Engine/Source/Runtime/MaterialShaderQualitySettings/Private/MaterialShaderQualitySettingsModule.cpp
@@ -0,0 +1,5 @@
+// Copyright 1998-2015 Epic Games, Inc. All Rights Reserved.
+
+#include "MaterialShaderQualitySettingsPrivatePCH.h"
+
+IMPLEMENT_MODULE(FDefaultModuleImpl, MaterialShaderQualitySettings);
diff --git a/Engine/Source/Runtime/MaterialShaderQualitySettings/Private/MaterialShaderQualitySettingsPrivatePCH.h b/Engine/Source/Runtime/MaterialShaderQualitySettings/Private/MaterialShaderQualitySettingsPrivatePCH.h
new file mode 100644
--- /dev/null
+++ b/Engine/Source/Runtime/MaterialShaderQualitySettings/Private/MaterialShaderQualitySettingsPrivatePCH.h
@@ -0,0 +1,8 @@
+// Copyright 1998-2015 Epic Games, Inc. All Rights Reserved.
+
+#pragma once
+
+#include "Core.h"
+#include "CoreUObject.h"
+#include "ModuleInterface.h"
+#include "ModuleManager.h"
diff --git a/Engine/Source/Runtime/MaterialShaderQualitySettings/Private/ShaderQualityOverridesListItem.h b/Engine/Source/Runtime/MaterialShaderQualitySettings/Private/ShaderQualityOverridesListItem.h
new file mode 100644
--- /dev/null
+++ b/Engine/Source/Runtime/MaterialShaderQualitySettings/Private/ShaderQualityOverridesListItem.h
@@ -0,0 +1,33 @@
+// Copyright 1998-2015 Epic Games, Inc. All Rights Reserved.
+
+#pragma once
+
+#include "ShaderQualityOverridesListItem.generated.h"
+
+// FShaderQualityOverridesListItem
+// Helper struct for FMaterialShaderQualitySettingsCustomization, contains info required to populate a material quality row.
+USTRUCT()
+struct FShaderQualityOverridesListItem
+{
+public:
+	GENERATED_USTRUCT_BODY()
+
+	UPROPERTY()
+	FString RangeName;
+
+	UPROPERTY()
+	UBoolProperty* QualityProperty;
+
+	UPROPERTY()
+	class UShaderPlatformQualitySettings* SettingContainer;
+
+	FShaderQualityOverridesListItem() {}
+
+	FShaderQualityOverridesListItem(FString InRangeName, UBoolProperty* InQualityProperty, UShaderPlatformQualitySettings* InSettingContainer)
+		: RangeName(InRangeName)
+		  , QualityProperty(InQualityProperty)
+		  , SettingContainer(InSettingContainer)
+	{}
+
+};
+
diff --git a/Engine/Source/Runtime/Renderer/Private/ForwardBasePassRendering.cpp b/Engine/Source/Runtime/Renderer/Private/ForwardBasePassRendering.cpp
--- a/Engine/Source/Runtime/Renderer/Private/ForwardBasePassRendering.cpp
+++ b/Engine/Source/Runtime/Renderer/Private/ForwardBasePassRendering.cpp
@@ -8,6 +8,7 @@
 #include "ScenePrivate.h"
 #include "SceneUtils.h"
 #include "Algo/Partition.h"
+#include "MaterialShaderQualitySettings.h"
 
 #define IMPLEMENT_FORWARD_SHADING_BASEPASS_LIGHTMAPPED_VERTEX_SHADER_TYPE(LightMapPolicyType,LightMapPolicyName) \
 	typedef TBasePassForForwardShadingVS< LightMapPolicyType, LDR_GAMMA_32 > TBasePassForForwardShadingVS##LightMapPolicyName##LDRGamma32; \
@@ -53,6 +54,20 @@ IMPLEMENT_FORWARD_SHADING_BASEPASS_LIGHTMAPPED_SHADER_TYPE(FMovableDirectionalLi
 IMPLEMENT_FORWARD_SHADING_BASEPASS_LIGHTMAPPED_SHADER_TYPE(FMovableDirectionalLightWithLightmapLightingPolicy, FMovableDirectionalLightWithLightmapLightingPolicy);
 IMPLEMENT_FORWARD_SHADING_BASEPASS_LIGHTMAPPED_SHADER_TYPE(FMovableDirectionalLightCSMWithLightmapLightingPolicy, FMovableDirectionalLightCSMWithLightmapLightingPolicy);
 
+template<typename LightMapPolicyType, int32 NumDynamicPointLights>
+bool TBasePassForForwardShadingPSBaseType<LightMapPolicyType, NumDynamicPointLights>::ModifyCompilationEnvironmentForQualityLevel(EShaderPlatform Platform, EMaterialQualityLevel::Type QualityLevel, FShaderCompilerEnvironment& OutEnvironment)
+{
+	// Get quality settings for shader platform
+	const UShaderPlatformQualitySettings* MaterialShadingQuality = UMaterialShaderQualitySettings::Get()->GetShaderPlatformQualitySettings(Platform);
+	const FMaterialQualityOverrides& QualityOverrides = MaterialShadingQuality->GetQualityOverrides(QualityLevel);
+
+	OutEnvironment.SetDefine(TEXT("FORWARD_QL_FORCE_FULLY_ROUGH"), QualityOverrides.bEnableOverride && QualityOverrides.bForceFullyRough != 0 ? 1u : 0u);
+	OutEnvironment.SetDefine(TEXT("FORWARD_QL_FORCE_NONMETAL"), QualityOverrides.bEnableOverride && QualityOverrides.bForceNonMetal != 0 ? 1u : 0u);
+	OutEnvironment.SetDefine(TEXT("QL_FORCEDISABLE_LM_DIRECTIONALITY"), QualityOverrides.bEnableOverride && QualityOverrides.bForceDisableLMDirectionality != 0 ? 1u : 0u);
+	
+	return true;
+}
+
 FBasePassFowardDynamicPointLightInfo::FBasePassFowardDynamicPointLightInfo(const FPrimitiveSceneProxy* InSceneProxy)
 : NumDynamicPointLights(0)
 {
diff --git a/Engine/Source/Runtime/Renderer/Private/ForwardBasePassRendering.h b/Engine/Source/Runtime/Renderer/Private/ForwardBasePassRendering.h
--- a/Engine/Source/Runtime/Renderer/Private/ForwardBasePassRendering.h
+++ b/Engine/Source/Runtime/Renderer/Private/ForwardBasePassRendering.h
@@ -173,6 +173,8 @@ public:
 			OutEnvironment.SetDefine(TEXT("VARIABLE_NUM_DYNAMIC_POINT_LIGHTS"), (uint32)0);
 			OutEnvironment.SetDefine(TEXT("NUM_DYNAMIC_POINT_LIGHTS"), (uint32)NumDynamicPointLights);
 		}
+		// Modify compilation environment depending upon material shader quality level settings.
+		ModifyCompilationEnvironmentForQualityLevel(Platform, Material->GetQualityLevel(), OutEnvironment);
 	}
 
 	/** Initialization constructor. */
@@ -252,6 +254,9 @@ public:
 	}
 
 private:
+
+	static bool ModifyCompilationEnvironmentForQualityLevel(EShaderPlatform Platform, EMaterialQualityLevel::Type QualityLevel, FShaderCompilerEnvironment& OutEnvironment);
+
 	FShaderResourceParameter ReflectionCubemap;
 	FShaderResourceParameter ReflectionSampler;
 	FEditorCompositingParameters EditorCompositeParams;
diff --git a/Engine/Source/Runtime/Renderer/Renderer.Build.cs b/Engine/Source/Runtime/Renderer/Renderer.Build.cs
--- a/Engine/Source/Runtime/Renderer/Renderer.Build.cs
+++ b/Engine/Source/Runtime/Renderer/Renderer.Build.cs
@@ -16,7 +16,7 @@ public class Renderer : ModuleRules
 
 		PublicDependencyModuleNames.Add("Core");
         PublicDependencyModuleNames.Add("Engine");
-
+        PublicDependencyModuleNames.Add("MaterialShaderQualitySettings");
         // Renderer module builds faster without unity
         // Non-unity also provides faster iteration
 		// Not enabled by default as it might harm full rebuild times without XGE
