adding missed file from CL2705914
UEMOB-36

diff --git a/Engine/Source/Runtime/Engine/Classes/Materials/MaterialInstance.h b/Engine/Source/Runtime/Engine/Classes/Materials/MaterialInstance.h
--- a/Engine/Source/Runtime/Engine/Classes/Materials/MaterialInstance.h
+++ b/Engine/Source/Runtime/Engine/Classes/Materials/MaterialInstance.h
@@ -361,7 +361,7 @@ public:
 	bool HasOverridenBaseProperties()const;
 
 	// For all materials instances, UMaterialInstance::CacheResourceShadersForRendering
-	static void AllMaterialsCacheResourceShadersForRendering();
+	ENGINE_API static void AllMaterialsCacheResourceShadersForRendering();
 
 	/** 
 	 * Determine whether this Material Instance is a child of another Material
