adding missed file from CL2705914
UEMOB-36

diff --git a/Engine/Source/Runtime/Engine/Private/KismetSystemLibrary.cpp b/Engine/Source/Runtime/Engine/Private/KismetSystemLibrary.cpp
--- a/Engine/Source/Runtime/Engine/Private/KismetSystemLibrary.cpp
+++ b/Engine/Source/Runtime/Engine/Private/KismetSystemLibrary.cpp
@@ -2997,7 +2997,7 @@ int32 UKismetSystemLibrary::GetRenderingMaterialQualityLevel()
 	static const IConsoleVariable* CVar = IConsoleManager::Get().FindConsoleVariable(TEXT("r.MaterialQualityLevel"));
 
 	// clamp range
-	int32 Ret = FMath::Clamp(CVar->GetInt(), 0, 1);
+	int32 Ret = FMath::Clamp(CVar->GetInt(), 0, (int32)EMaterialQualityLevel::Num - 1);
 
 	return Ret;
 }
