Removing redundant variable, fixing CIS shadow warning.

diff --git a/Engine/Source/Runtime/Engine/Private/CurveBase.cpp b/Engine/Source/Runtime/Engine/Private/CurveBase.cpp
--- a/Engine/Source/Runtime/Engine/Private/CurveBase.cpp
+++ b/Engine/Source/Runtime/Engine/Private/CurveBase.cpp
@@ -1695,9 +1695,7 @@ FKeyHandle FIntegralCurve::AddKey( float InTime, int32 InValue, FKeyHandle InKey
 	
 	for (auto It = KeyHandlesToIndices.CreateIterator(); It; ++It)
 	{
-		const FKeyHandle& KeyHandle = It.Key();
 		int32& KeyIndex = It.Value();
-
 		if (KeyIndex >= Index)
 		{
 			++KeyIndex;
@@ -1720,9 +1718,7 @@ void FIntegralCurve::DeleteKey(FKeyHandle InKeyHandle)
 
 	for (auto It = KeyHandlesToIndices.CreateIterator(); It; ++It)
 	{
-		const FKeyHandle& KeyHandle = It.Key();
 		int32& KeyIndex = It.Value();
-
 		if (KeyIndex >= Index)
 		{
 			--KeyIndex;
