Merging

//depot/UE4-Orion/Engine/Source/Runtime/Engine/Private/Animation/BlendSpaceBase.cpp

to //depot/UE4-Orion/Engine/Source/Runtime/Engine/Private/Animation/BlendSpaceBase.cpp

diff --git a/Engine/Source/Runtime/Engine/Private/Animation/BlendSpaceBase.cpp b/Engine/Source/Runtime/Engine/Private/Animation/BlendSpaceBase.cpp
--- a/Engine/Source/Runtime/Engine/Private/Animation/BlendSpaceBase.cpp
+++ b/Engine/Source/Runtime/Engine/Private/Animation/BlendSpaceBase.cpp
@@ -290,6 +290,20 @@ void UBlendSpaceBase::TickAssetPlayerInstance(FAnimTickRecord& Instance, class U
 			float& NormalizedCurrentTime = *(Instance.TimeAccumulator);
 			const float NormalizedPreviousTime = NormalizedCurrentTime;
 
+			// @note for sync group vs non sync group
+			// in blendspace, it will still sync even if only one node in sync group
+			// so you're never non-sync group unless you have situation where some markers are relevant to one sync group but not all the time
+			// here we save NormalizedCurrentTime as Highest weighted samples' position in sync group
+			// if you're not in sync group, NormalizedCurrentTime is based on normalized length by sample weights
+			// if you move between sync to non sync within blendspace, you're going to see pop because we'll have to jump
+			// for now, our rule is to keep normalized time as highest weighted sample position within its own length
+			// also MoveDelta doesn't work if you're in sync group. It will move according to sync group position
+			// @todo consider using MoveDelta when  this is leader, but that can be scary because it's not matching with DeltaTime any more. 
+			// if you have interpolation delay, that value can be applied, but the output might be unpredictable. 
+			// 
+			// to fix this better in the future, we should use marker sync position from last tick
+			// but that still doesn't fix if you just join sync group, you're going to see pop since your animation doesn't fix
+
 			if (Context.IsLeader())
 			{
 				// advance current time - blend spaces hold normalized time as when dealing with changing anim length it would be possible to go backwards
@@ -342,9 +356,14 @@ void UBlendSpaceBase::TickAssetPlayerInstance(FAnimTickRecord& Instance, class U
 					}
 					const int32 HighestWeightIndex = GetHighestWeightSample(SampleDataList);
 					FBlendSampleData& SampleDataItem = SampleDataList[HighestWeightIndex];
+					const FBlendSample& Sample = SampleData[SampleDataItem.SampleDataIndex];
 					Instance.MarkerTickRecord = SampleDataItem.MarkerTickRecord;
+					NormalizedCurrentTime =  SampleDataItem.Time / Sample.Animation->SequenceLength; 
+				}
+				else
+				{
+					NormalizedCurrentTime =  Context.GetAnimationPositionRatio();
 				}
-				NormalizedCurrentTime = Context.GetAnimationPositionRatio();
 			}
 
 			// generate notifies and sets time
