Optimize USphereComponent::SetSphereRadius,  UCapsuleComponent::SetCapsuleSize and UBoxComponent::SetBoxExtent  by using UpdateBodyScale rather than destroying/recreating physics body.

diff --git a/Engine/Source/Runtime/Engine/Classes/PhysicsEngine/BodyInstance.h b/Engine/Source/Runtime/Engine/Classes/PhysicsEngine/BodyInstance.h
--- a/Engine/Source/Runtime/Engine/Classes/PhysicsEngine/BodyInstance.h
+++ b/Engine/Source/Runtime/Engine/Classes/PhysicsEngine/BodyInstance.h
@@ -611,10 +611,11 @@ public:
 
 	/**
 	 * Update Body Scale
-	 * @param inScale3D - new Scale3D. If that's different from previous Scale3D, it will update Body scale
+	 * @param	InScale3D		New Scale3D. If that's different from previous Scale3D, it will update Body scale.
+	 * @param	bForceUpdate	Will refresh shape dimensions from BodySetup, even if scale has not changed.
 	 * @return true if succeed
 	 */
-	bool UpdateBodyScale(const FVector& InScale3D);
+	bool UpdateBodyScale(const FVector& InScale3D, bool bForceUpdate = false);
 
 	/** Dynamically update the vertices of per-poly collision for this body. */
 	void UpdateTriMeshVertices(const TArray<FVector> & NewPositions);
diff --git a/Engine/Source/Runtime/Engine/Private/Components/BoxComponent.cpp b/Engine/Source/Runtime/Engine/Private/Components/BoxComponent.cpp
--- a/Engine/Source/Runtime/Engine/Private/Components/BoxComponent.cpp
+++ b/Engine/Source/Runtime/Engine/Private/Components/BoxComponent.cpp
@@ -17,14 +17,14 @@ void UBoxComponent::SetBoxExtent(FVector NewBoxExtent, bool bUpdateOverlaps)
 {
 	BoxExtent = NewBoxExtent;
 	MarkRenderStateDirty();
+	UpdateBodySetup();
 
 	// do this if already created
 	// otherwise, it hasn't been really created yet
 	if (bPhysicsStateCreated)
 	{
-		DestroyPhysicsState();
-		UpdateBodySetup();
-		CreatePhysicsState();
+		// Update physics engine collision shapes
+		BodyInstance.UpdateBodyScale(ComponentToWorld.GetScale3D(), true);
 
 		if ( bUpdateOverlaps && IsCollisionEnabled() && GetOwner() )
 		{
diff --git a/Engine/Source/Runtime/Engine/Private/Components/CapsuleComponent.cpp b/Engine/Source/Runtime/Engine/Private/Components/CapsuleComponent.cpp
--- a/Engine/Source/Runtime/Engine/Private/Components/CapsuleComponent.cpp
+++ b/Engine/Source/Runtime/Engine/Private/Components/CapsuleComponent.cpp
@@ -146,15 +146,15 @@ void UCapsuleComponent::SetCapsuleSize(float NewRadius, float NewHalfHeight, boo
 {
 	CapsuleHalfHeight = FMath::Max3(0.f, NewHalfHeight, NewRadius);
 	CapsuleRadius = FMath::Max(0.f, NewRadius);
+	UpdateBodySetup();
 	MarkRenderStateDirty();
 
 	// do this if already created
 	// otherwise, it hasn't been really created yet
 	if (bPhysicsStateCreated)
 	{
-		DestroyPhysicsState();
-		UpdateBodySetup();
-		CreatePhysicsState();
+		// Update physics engine collision shapes
+		BodyInstance.UpdateBodyScale(ComponentToWorld.GetScale3D(), true);
 
 		if ( bUpdateOverlaps && IsCollisionEnabled() && GetOwner() )
 		{
diff --git a/Engine/Source/Runtime/Engine/Private/Components/SphereComponent.cpp b/Engine/Source/Runtime/Engine/Private/Components/SphereComponent.cpp
--- a/Engine/Source/Runtime/Engine/Private/Components/SphereComponent.cpp
+++ b/Engine/Source/Runtime/Engine/Private/Components/SphereComponent.cpp
@@ -51,13 +51,13 @@ void USphereComponent::UpdateBodySetup()
 void USphereComponent::SetSphereRadius( float InSphereRadius, bool bUpdateOverlaps )
 {
 	SphereRadius = InSphereRadius;
+	UpdateBodySetup();
 	MarkRenderStateDirty();
 
 	if (bPhysicsStateCreated)
 	{
-		DestroyPhysicsState();
-		UpdateBodySetup();
-		CreatePhysicsState();
+		// Update physics engine collision shapes
+		BodyInstance.UpdateBodyScale(ComponentToWorld.GetScale3D(), true);
 
 		if ( bUpdateOverlaps && IsCollisionEnabled() && GetOwner() )
 		{
diff --git a/Engine/Source/Runtime/Engine/Private/PhysicsEngine/BodyInstance.cpp b/Engine/Source/Runtime/Engine/Private/PhysicsEngine/BodyInstance.cpp
--- a/Engine/Source/Runtime/Engine/Private/PhysicsEngine/BodyInstance.cpp
+++ b/Engine/Source/Runtime/Engine/Private/PhysicsEngine/BodyInstance.cpp
@@ -2070,7 +2070,7 @@ EScaleMode::Type ComputeScaleMode(const TArray<PxShape*>& PShapes)
 	return ScaleMode;
 }
 
-bool FBodyInstance::UpdateBodyScale(const FVector& InScale3D)
+bool FBodyInstance::UpdateBodyScale(const FVector& InScale3D, bool bForceUpdate)
 {
 	if (!IsValidBodyInstance())
 	{
@@ -2078,8 +2078,8 @@ bool FBodyInstance::UpdateBodyScale(const FVector& InScale3D)
 		return false;
 	}
 
-	// if same, return
-	if (Scale3D.Equals(InScale3D))
+	// if scale is already correct, and not forcing an update, do nothing
+	if (Scale3D.Equals(InScale3D) && !bForceUpdate)
 	{
 		return false;
 	}
