Integrate CL 2706042 from Orion to Main
- Use stat for BodyInstanceOverlapMulti/Test
- Remove stacking GeomOverlapMultiple stat

diff --git a/Engine/Source/Runtime/Engine/Private/Collision/PhysXCollision.cpp b/Engine/Source/Runtime/Engine/Private/Collision/PhysXCollision.cpp
--- a/Engine/Source/Runtime/Engine/Private/Collision/PhysXCollision.cpp
+++ b/Engine/Source/Runtime/Engine/Private/Collision/PhysXCollision.cpp
@@ -1709,7 +1709,6 @@ bool GeomOverlapMulti_PhysX(const UWorld* World, const PxGeometry& PGeom, const
 template <EQueryInfo::Type InfoType>
 bool GeomOverlapMultiImp(const UWorld* World, const struct FCollisionShape& CollisionShape, const FVector& Pos, const FQuat& Rot, TArray<FOverlapResult>& OutOverlaps, ECollisionChannel TraceChannel, const struct FCollisionQueryParams& Params, const struct FCollisionResponseParams& ResponseParams, const struct FCollisionObjectQueryParams& ObjectParams)
 {
-	SCOPE_CYCLE_COUNTER(STAT_Collision_GeomOverlapMultiple);
 	if ((World == NULL) || (World->GetPhysicsScene() == NULL))
 	{
 		return false;
diff --git a/Engine/Source/Runtime/Engine/Private/Collision/WorldCollision.cpp b/Engine/Source/Runtime/Engine/Private/Collision/WorldCollision.cpp
--- a/Engine/Source/Runtime/Engine/Private/Collision/WorldCollision.cpp
+++ b/Engine/Source/Runtime/Engine/Private/Collision/WorldCollision.cpp
@@ -36,7 +36,8 @@ DEFINE_STAT(STAT_Collision_GeomSweepMultiple);
 DEFINE_STAT(STAT_Collision_GeomOverlapAny);
 DEFINE_STAT(STAT_Collision_GeomOverlapSingle);
 DEFINE_STAT(STAT_Collision_GeomOverlapMultiple);
-DEFINE_STAT(STAT_Collision_GeomComputePenetration);
+DEFINE_STAT(STAT_Collision_FBodyInstance_OverlapMulti);
+DEFINE_STAT(STAT_Collision_FBodyInstance_OverlapTest);
 DEFINE_STAT(STAT_Collision_PreFilter);
 DEFINE_STAT(STAT_Collision_PostFilter);
 
diff --git a/Engine/Source/Runtime/Engine/Private/PhysicsEngine/BodyInstance.cpp b/Engine/Source/Runtime/Engine/Private/PhysicsEngine/BodyInstance.cpp
--- a/Engine/Source/Runtime/Engine/Private/PhysicsEngine/BodyInstance.cpp
+++ b/Engine/Source/Runtime/Engine/Private/PhysicsEngine/BodyInstance.cpp
@@ -4094,6 +4094,8 @@ template bool FBodyInstance::OverlapTestForBodiesImpl(const FVector& Pos, const
 
 bool FBodyInstance::OverlapTest(const FVector& Position, const FQuat& Rotation, const struct FCollisionShape& CollisionShape, FMTDResult* OutMTD) const
 {
+	SCOPE_CYCLE_COUNTER(STAT_Collision_FBodyInstance_OverlapTest);
+
 	bool bHasOverlap = false;
 
 #if WITH_PHYSX
@@ -4133,13 +4135,14 @@ FTransform RootSpaceToWeldedSpace(const FBodyInstance* BI, const FTransform& Roo
 
 bool FBodyInstance::OverlapMulti(TArray<struct FOverlapResult>& InOutOverlaps, const class UWorld* World, const FTransform* pWorldToComponent, const FVector& Pos, const FQuat& Quat, ECollisionChannel TestChannel, const struct FComponentQueryParams& Params, const struct FCollisionResponseParams& ResponseParams, const struct FCollisionObjectQueryParams& ObjectQueryParams) const
 {
+	SCOPE_CYCLE_COUNTER(STAT_Collision_FBodyInstance_OverlapMulti);
+
 	if ( !IsValidBodyInstance()  && (!WeldParent || !WeldParent->IsValidBodyInstance()))
 	{
 		UE_LOG(LogCollision, Log, TEXT("FBodyInstance::OverlapMulti : (%s) No physics data"), *GetBodyDebugName());
 		return false;
 	}
 
-	SCOPE_CYCLE_COUNTER(STAT_Collision_GeomOverlapMultiple);
 	bool bHaveBlockingHit = false;
 
 	// Determine how to convert the local space of this body instance to the test space
diff --git a/Engine/Source/Runtime/Engine/Public/Collision.h b/Engine/Source/Runtime/Engine/Public/Collision.h
--- a/Engine/Source/Runtime/Engine/Public/Collision.h
+++ b/Engine/Source/Runtime/Engine/Public/Collision.h
@@ -18,7 +18,8 @@ DECLARE_CYCLE_STAT_EXTERN(TEXT("GeomSweepMultiple"),STAT_Collision_GeomSweepMult
 DECLARE_CYCLE_STAT_EXTERN(TEXT("GeomOverlapAny"),STAT_Collision_GeomOverlapAny,STATGROUP_Collision, );
 DECLARE_CYCLE_STAT_EXTERN(TEXT("GeomOverlapSingle"),STAT_Collision_GeomOverlapSingle,STATGROUP_Collision, );
 DECLARE_CYCLE_STAT_EXTERN(TEXT("GeomOverlapMultiple"),STAT_Collision_GeomOverlapMultiple,STATGROUP_Collision, );
-DECLARE_CYCLE_STAT_EXTERN(TEXT("GeomComputePenetration"), STAT_Collision_GeomComputePenetration, STATGROUP_Collision, );
+DECLARE_CYCLE_STAT_EXTERN(TEXT("BodyInstanceOverlapMulti"), STAT_Collision_FBodyInstance_OverlapMulti, STATGROUP_Collision, );
+DECLARE_CYCLE_STAT_EXTERN(TEXT("BodyInstanceOverlapTest"), STAT_Collision_FBodyInstance_OverlapTest, STATGROUP_Collision, );
 DECLARE_CYCLE_STAT_EXTERN(TEXT("PreFilter"), STAT_Collision_PreFilter, STATGROUP_CollisionVerbose, );
 DECLARE_CYCLE_STAT_EXTERN(TEXT("PostFilter"), STAT_Collision_PostFilter, STATGROUP_CollisionVerbose, );
 
