Integrate CL 2706043 from Orion to Main
- Stop UpdateBounds from being virtual. SkeletalMeshComponent now caches/updates local bounds in CalcBounds instead.
- Add stats for UpdateNavigationData and PostUpdateNavigationData

diff --git a/Engine/Source/Runtime/Engine/Classes/Components/SceneComponent.h b/Engine/Source/Runtime/Engine/Classes/Components/SceneComponent.h
--- a/Engine/Source/Runtime/Engine/Classes/Components/SceneComponent.h
+++ b/Engine/Source/Runtime/Engine/Classes/Components/SceneComponent.h
@@ -858,7 +858,7 @@ public:
 	virtual void CalcBoundingCylinder(float& CylinderRadius, float& CylinderHalfHeight) const;
 
 	/** Update the Bounds of the component.*/
-	virtual void UpdateBounds();
+	void UpdateBounds();
 
 	/** If true, bounds should be used when placing component/actor in level. Does not affect spawning. */
 	virtual bool ShouldCollideWhenPlacing() const
diff --git a/Engine/Source/Runtime/Engine/Classes/Components/SkeletalMeshComponent.h b/Engine/Source/Runtime/Engine/Classes/Components/SkeletalMeshComponent.h
--- a/Engine/Source/Runtime/Engine/Classes/Components/SkeletalMeshComponent.h
+++ b/Engine/Source/Runtime/Engine/Classes/Components/SkeletalMeshComponent.h
@@ -836,7 +836,6 @@ public:
 	//~ End UActorComponent Interface.
 
 	//~ Begin USceneComponent Interface.
-	virtual void UpdateBounds() override;
 	virtual FBoxSphereBounds CalcBounds(const FTransform& LocalToWorld) const override;
 	virtual bool IsAnySimulatingPhysics() const override;
 	virtual void OnUpdateTransform(bool bSkipPhysicsMove, ETeleportType Teleport = ETeleportType::None) override;
diff --git a/Engine/Source/Runtime/Engine/Classes/Components/SkinnedMeshComponent.h b/Engine/Source/Runtime/Engine/Classes/Components/SkinnedMeshComponent.h
--- a/Engine/Source/Runtime/Engine/Classes/Components/SkinnedMeshComponent.h
+++ b/Engine/Source/Runtime/Engine/Classes/Components/SkinnedMeshComponent.h
@@ -650,11 +650,11 @@ protected:
 
 	/** LocalBounds cached, so they're computed just once. */
 	UPROPERTY(Transient)
-	FBoxSphereBounds CachedLocalBounds;
+	mutable FBoxSphereBounds CachedLocalBounds;
 
 	/** true when CachedLocalBounds is up to date. */
 	UPROPERTY(Transient)
-	bool bCachedLocalBoundsUpToDate;
+	mutable bool bCachedLocalBoundsUpToDate;
 
 	/** Invalidate Cached Bounds, when Mesh Component has been updated. */
 	void InvalidateCachedBounds();
diff --git a/Engine/Source/Runtime/Engine/Private/Components/SceneComponent.cpp b/Engine/Source/Runtime/Engine/Private/Components/SceneComponent.cpp
--- a/Engine/Source/Runtime/Engine/Private/Components/SceneComponent.cpp
+++ b/Engine/Source/Runtime/Engine/Private/Components/SceneComponent.cpp
@@ -30,7 +30,10 @@ namespace SceneComponentStatics
 DEFINE_LOG_CATEGORY_STATIC(LogSceneComponent, Log, All);
 
 DECLARE_CYCLE_STAT(TEXT("UpdateComponentToWorld"), STAT_UpdateComponentToWorld, STATGROUP_Component);
-DECLARE_CYCLE_STAT(TEXT("Update Child Transforms"), STAT_UpdateChildTransforms, STATGROUP_Component);
+DECLARE_CYCLE_STAT(TEXT("UpdateChildTransforms"), STAT_UpdateChildTransforms, STATGROUP_Component);
+DECLARE_CYCLE_STAT(TEXT("Component UpdateBounds"), STAT_ComponentUpdateBounds, STATGROUP_Component);
+DECLARE_CYCLE_STAT(TEXT("Component UpdateNavData"), STAT_ComponentUpdateNavData, STATGROUP_Component);
+DECLARE_CYCLE_STAT(TEXT("Component PostUpdateNavData"), STAT_ComponentPostUpdateNavData, STATGROUP_Component);
 
 
 FOverlapInfo::FOverlapInfo(UPrimitiveComponent* InComponent, int32 InBodyIndex)
@@ -725,6 +728,8 @@ void USceneComponent::CalcBoundingCylinder(float& CylinderRadius, float& Cylinde
 
 void USceneComponent::UpdateBounds()
 {
+	SCOPE_CYCLE_COUNTER(STAT_ComponentUpdateBounds);
+
 #if WITH_EDITOR
 	FBoxSphereBounds OriginalBounds = Bounds; // Save old bounds
 #endif
@@ -2587,6 +2592,8 @@ const int32 USceneComponent::GetNumUncachedStaticLightingInteractions() const
 
 void USceneComponent::UpdateNavigationData()
 {
+	SCOPE_CYCLE_COUNTER(STAT_ComponentUpdateNavData);
+
 	if (UNavigationSystem::ShouldUpdateNavOctreeOnComponentChange() &&
 		IsRegistered() && World && World->IsGameWorld() &&
 		World->GetNetMode() < ENetMode::NM_Client)
@@ -2597,6 +2604,8 @@ void USceneComponent::UpdateNavigationData()
 
 void USceneComponent::PostUpdateNavigationData()
 {
+	SCOPE_CYCLE_COUNTER(STAT_ComponentPostUpdateNavData);
+
 	if (UNavigationSystem::ShouldUpdateNavOctreeOnComponentChange() &&
 		IsRegistered() && World && World->IsGameWorld() &&
 		World->GetNetMode() < ENetMode::NM_Client)
diff --git a/Engine/Source/Runtime/Engine/Private/Components/SkeletalMeshComponent.cpp b/Engine/Source/Runtime/Engine/Private/Components/SkeletalMeshComponent.cpp
--- a/Engine/Source/Runtime/Engine/Private/Components/SkeletalMeshComponent.cpp
+++ b/Engine/Source/Runtime/Engine/Private/Components/SkeletalMeshComponent.cpp
@@ -40,7 +40,6 @@
 TAutoConsoleVariable<int32> CVarUseParallelAnimationEvaluation(TEXT("a.ParallelAnimEvaluation"), 1, TEXT("If 1, animation evaluation will be run across the task graph system. If 0, evaluation will run purely on the game thread"));
 
 DECLARE_CYCLE_STAT(TEXT("Swap Anim Buffers"), STAT_CompleteAnimSwapBuffers, STATGROUP_Anim);
-DECLARE_CYCLE_STAT(TEXT("Update SkelMesh Bounds"), STAT_UpdateSkelMeshBounds, STATGROUP_Anim);
 DECLARE_CYCLE_STAT_EXTERN(TEXT("Anim Instance Spawn Time"), STAT_AnimSpawnTime, STATGROUP_Anim, );
 DEFINE_STAT(STAT_AnimSpawnTime);
 DEFINE_STAT(STAT_PostAnimEvaluation);
@@ -1313,79 +1312,48 @@ void USkeletalMeshComponent::PostAnimEvaluation(FAnimationEvaluationContext& Eva
 	PrepareCloth();
 }
 
-//
-//	USkeletalMeshComponent::UpdateBounds
-//
-
-void USkeletalMeshComponent::UpdateBounds()
+FBoxSphereBounds USkeletalMeshComponent::CalcBounds(const FTransform& LocalToWorld) const
 {
-	SCOPE_CYCLE_COUNTER(STAT_UpdateSkelMeshBounds);
-
-#if WITH_EDITOR
-	FBoxSphereBounds OriginalBounds = Bounds; // Save old bounds
-#endif
+	SCOPE_CYCLE_COUNTER(STAT_CalcSkelMeshBounds);
 
-	// if use parent bound if attach parent exists, and the flag is set
-	// since parents tick first before child, this should work correctly
-	if ( bUseAttachParentBound && AttachParent != NULL )
+	// fixme laurent - extend concept of LocalBounds to all SceneComponent
+	// as rendered calls CalcBounds*() directly in FScene::UpdatePrimitiveTransform, which is pretty expensive for SkelMeshes.
+	// No need to calculated that again, just use cached local bounds.
+	if (bCachedLocalBoundsUpToDate)
 	{
-		Bounds = AttachParent->Bounds;
+		return CachedLocalBounds.TransformBy(LocalToWorld);
 	}
+	// Calculate new bounds
 	else
 	{
-		// fixme laurent - extend concept of LocalBounds to all SceneComponent
-		// as rendered calls CalcBounds*() directly in FScene::UpdatePrimitiveTransform, which is pretty expensive for SkelMeshes.
-		// No need to calculated that again, just use cached local bounds.
-		if( bCachedLocalBoundsUpToDate )
-		{
-			Bounds = CachedLocalBounds.TransformBy(ComponentToWorld);
-		}
-		else
-		{
-			// Calculate new bounds
-			Bounds = CalcBounds(ComponentToWorld);
+		FVector RootBoneOffset = RootBoneTranslation;
 
-			bCachedLocalBoundsUpToDate = true;
-			CachedLocalBounds = Bounds.TransformBy(ComponentToWorld.Inverse());
+		// if to use MasterPoseComponent's fixed skel bounds, 
+		// send MasterPoseComponent's Root Bone Translation
+		if (MasterPoseComponent.IsValid())
+		{
+			const USkinnedMeshComponent* const MasterPoseComponentInst = MasterPoseComponent.Get();
+			check(MasterPoseComponentInst);
+			if (MasterPoseComponentInst->SkeletalMesh &&
+				MasterPoseComponentInst->bComponentUseFixedSkelBounds &&
+				MasterPoseComponentInst->IsA((USkeletalMeshComponent::StaticClass())))
+			{
+				const USkeletalMeshComponent* BaseComponent = CastChecked<USkeletalMeshComponent>(MasterPoseComponentInst);
+				RootBoneOffset = BaseComponent->RootBoneTranslation; // Adjust bounds by root bone translation
+			}
 		}
-	}
 
-#if WITH_EDITOR
-	// If bounds have changed (in editor), trigger data rebuild
-	if ( IsRegistered() && (World != NULL) && !GetWorld()->IsGameWorld() && 
-		(OriginalBounds.Origin.Equals(Bounds.Origin) == false || OriginalBounds.BoxExtent.Equals(Bounds.BoxExtent) == false) )
-	{
-		GEngine->TriggerStreamingDataRebuild();
-	}
-#endif
-}
+		FBoxSphereBounds NewBounds = CalcMeshBound( RootBoneOffset, bHasValidBodies, LocalToWorld );
 
-FBoxSphereBounds USkeletalMeshComponent::CalcBounds(const FTransform& LocalToWorld) const
-{
-	SCOPE_CYCLE_COUNTER(STAT_CalcSkelMeshBounds);
+#if WITH_APEX_CLOTHING
+		AddClothingBounds(NewBounds);
+#endif// #if WITH_APEX_CLOTHING
 
-	FVector RootBoneOffset = RootBoneTranslation;
+		bCachedLocalBoundsUpToDate = true;
+		CachedLocalBounds = NewBounds.TransformBy(LocalToWorld.Inverse());
 
-	// if to use MasterPoseComponent's fixed skel bounds, 
-	// send MasterPoseComponent's Root Bone Translation
-	if (MasterPoseComponent.IsValid())
-	{
-		const USkinnedMeshComponent* const MasterPoseComponentInst = MasterPoseComponent.Get();
-		check(MasterPoseComponentInst);
-		if (MasterPoseComponentInst->SkeletalMesh &&
-			MasterPoseComponentInst->bComponentUseFixedSkelBounds &&
-			MasterPoseComponentInst->IsA((USkeletalMeshComponent::StaticClass())))
-		{
-			const USkeletalMeshComponent* BaseComponent = CastChecked<USkeletalMeshComponent>(MasterPoseComponentInst);
-			RootBoneOffset = BaseComponent->RootBoneTranslation; // Adjust bounds by root bone translation
-		}
+		return NewBounds;
 	}
-
-	FBoxSphereBounds NewBounds = CalcMeshBound( RootBoneOffset, bHasValidBodies, LocalToWorld );
-#if WITH_APEX_CLOTHING
-	AddClothingBounds(NewBounds);
-#endif// #if WITH_APEX_CLOTHING
-	return NewBounds;
 }
 
 void USkeletalMeshComponent::SetSkeletalMesh(USkeletalMesh* InSkelMesh)
