Fixed issue when blend input data is invalid when switching assets.

Originally was going to remove the check, but this is better fix and this has to be fixed anyway. Otherwise, it doesn't initialize filters correctly.

Related to this jira - https://jira.ol.epicgames.net/browse/UE-20748

diff --git a/Engine/Source/Runtime/Engine/Private/Animation/AnimSingleNodeInstance.cpp b/Engine/Source/Runtime/Engine/Private/Animation/AnimSingleNodeInstance.cpp
--- a/Engine/Source/Runtime/Engine/Private/Animation/AnimSingleNodeInstance.cpp
+++ b/Engine/Source/Runtime/Engine/Private/Animation/AnimSingleNodeInstance.cpp
@@ -58,6 +58,7 @@ void UAnimSingleNodeInstance::SetAnimationAsset(class UAnimationAsset* NewAsset,
 	bLooping = bIsLooping;
 	PlayRate = InPlayRate;
 	CurrentTime = 0.f;
+	BlendSpaceInput = FVector::ZeroVector;
 #if WITH_EDITORONLY_DATA
 	PreviewPoseCurrentTime = 0.0f;
 #endif
@@ -77,6 +78,12 @@ void UAnimSingleNodeInstance::SetAnimationAsset(class UAnimationAsset* NewAsset,
 	{
 		// otherwise stop all montages
 		StopAllMontages(0.25f);
+
+		UBlendSpaceBase * BlendSpace = Cast<UBlendSpaceBase>(CurrentAsset);
+		if(BlendSpace)
+		{
+			BlendSpace->InitializeFilter(&BlendFilter);
+		}
 	}
 }
 
@@ -206,16 +213,8 @@ void UAnimSingleNodeInstance::NativeInitializeAnimation()
 	CurrentTime = 0.f;
 	USkeletalMeshComponent* SkelComp = GetSkelMeshComponent();
 	SkelComp->AnimationData.Initialize(this);
-
-	if ( CurrentAsset!=NULL )
-	{
-		UBlendSpace * BlendSpace = Cast<UBlendSpace>(CurrentAsset);
-		if (BlendSpace)
-		{
-			BlendSpace->InitializeFilter(&BlendFilter);
-		}
-	}
 }
+
 void UAnimSingleNodeInstance::UpdateAnimationNode(float DeltaTimeX)
 {
 	float NewPlayRate = PlayRate;
