Fix display of global mip bias

diff --git a/Engine/Source/Runtime/Engine/Private/ContentStreaming.cpp b/Engine/Source/Runtime/Engine/Private/ContentStreaming.cpp
--- a/Engine/Source/Runtime/Engine/Private/ContentStreaming.cpp
+++ b/Engine/Source/Runtime/Engine/Private/ContentStreaming.cpp
@@ -5008,7 +5008,7 @@ void FStreamingManagerTexture::InvestigateTexture( const FString& InInvestigateT
 			UE_LOG(LogContentStreaming, Log,  TEXT("  Priority:        %.3f"), StreamingTexture.CalcPriority() );
 			UE_LOG(LogContentStreaming, Log,  TEXT("  Boost factor:    %.1f"), StreamingTexture.BoostFactor );
 			UE_LOG(LogContentStreaming, Log,  TEXT("  Allowed mips:    %d-%d"), StreamingTexture.MinAllowedMips, StreamingTexture.MaxAllowedMips );
-			UE_LOG(LogContentStreaming, Log,  TEXT("  Global mip bias: %d"), ThreadSettings.MipBias );
+			UE_LOG(LogContentStreaming, Log,  TEXT("  Global mip bias: %.1f"), ThreadSettings.MipBias );
 
 			float ScreenSizeFactor;
 			switch ( StreamingTexture.LODGroup )
