InvalidateCachedBounds() should be public

diff --git a/Engine/Source/Runtime/Engine/Classes/Components/SkinnedMeshComponent.h b/Engine/Source/Runtime/Engine/Classes/Components/SkinnedMeshComponent.h
--- a/Engine/Source/Runtime/Engine/Classes/Components/SkinnedMeshComponent.h
+++ b/Engine/Source/Runtime/Engine/Classes/Components/SkinnedMeshComponent.h
@@ -656,9 +656,13 @@ protected:
 	UPROPERTY(Transient)
 	mutable bool bCachedLocalBoundsUpToDate;
 
+public:
+
 	/** Invalidate Cached Bounds, when Mesh Component has been updated. */
 	void InvalidateCachedBounds();
 
+protected:
+
 	/** Update Mesh Bound information based on input
 	 * 
 	 * @param RootOffset	: Root Bone offset from mesh location
