#ue4 - Change some collision profiles to use existing static FNames rather than names defined at the call sites.

--------

diff --git a/Engine/Source/Runtime/Engine/Private/AI/Navigation/NavigationTestingActor.cpp b/Engine/Source/Runtime/Engine/Private/AI/Navigation/NavigationTestingActor.cpp
--- a/Engine/Source/Runtime/Engine/Private/AI/Navigation/NavigationTestingActor.cpp
+++ b/Engine/Source/Runtime/Engine/Private/AI/Navigation/NavigationTestingActor.cpp
@@ -48,12 +48,9 @@ ANavigationTestingActor::ANavigationTestingActor(const FObjectInitializer& Objec
 
 	QueryingExtent = FVector(DEFAULT_NAV_QUERY_EXTENT_HORIZONTAL, DEFAULT_NAV_QUERY_EXTENT_HORIZONTAL, DEFAULT_NAV_QUERY_EXTENT_VERTICAL);
 
-	// collision profile name set up - found in baseengine.ini
-	static FName CollisionProfileName(TEXT("Pawn"));
-
 	CapsuleComponent = CreateDefaultSubobject<UCapsuleComponent>(TEXT("CollisionCylinder"));
 	CapsuleComponent->InitCapsuleSize(NavAgentProps.AgentRadius, NavAgentProps.AgentHeight / 2);
-	CapsuleComponent->SetCollisionProfileName(CollisionProfileName);
+	CapsuleComponent->SetCollisionProfileName(UCollisionProfile::Pawn_ProfileName);
 	CapsuleComponent->CanCharacterStepUpOn = ECB_No;
 	CapsuleComponent->bShouldUpdatePhysicsVolume = true;
 	CapsuleComponent->bCanEverAffectNavigation = false;
diff --git a/Engine/Source/Runtime/Engine/Private/Character.cpp b/Engine/Source/Runtime/Engine/Private/Character.cpp
--- a/Engine/Source/Runtime/Engine/Private/Character.cpp
+++ b/Engine/Source/Runtime/Engine/Private/Character.cpp
@@ -46,9 +46,7 @@ ACharacter::ACharacter(const FObjectInitializer& ObjectInitializer)
 
 	CapsuleComponent = CreateDefaultSubobject<UCapsuleComponent>(ACharacter::CapsuleComponentName);
 	CapsuleComponent->InitCapsuleSize(34.0f, 88.0f);
-
-	static FName CollisionProfileName(TEXT("Pawn"));
-	CapsuleComponent->SetCollisionProfileName(CollisionProfileName);
+	CapsuleComponent->SetCollisionProfileName(UCollisionProfile::Pawn_ProfileName);
 
 	CapsuleComponent->CanCharacterStepUpOn = ECB_No;
 	CapsuleComponent->bShouldUpdatePhysicsVolume = true;
diff --git a/Engine/Source/Runtime/Engine/Private/DefaultPawn.cpp b/Engine/Source/Runtime/Engine/Private/DefaultPawn.cpp
--- a/Engine/Source/Runtime/Engine/Private/DefaultPawn.cpp
+++ b/Engine/Source/Runtime/Engine/Private/DefaultPawn.cpp
@@ -27,9 +27,7 @@ ADefaultPawn::ADefaultPawn(const FObjectInitializer& ObjectInitializer)
 
 	CollisionComponent = CreateDefaultSubobject<USphereComponent>(ADefaultPawn::CollisionComponentName);
 	CollisionComponent->InitSphereRadius(35.0f);
-
-	static FName CollisionProfileName(TEXT("Pawn"));
-	CollisionComponent->SetCollisionProfileName(CollisionProfileName);
+	CollisionComponent->SetCollisionProfileName(UCollisionProfile::Pawn_ProfileName);
 
 	CollisionComponent->CanCharacterStepUpOn = ECB_No;
 	CollisionComponent->bShouldUpdatePhysicsVolume = true;
@@ -62,7 +60,7 @@ ADefaultPawn::ADefaultPawn(const FObjectInitializer& ObjectInitializer)
 		MeshComponent->bAffectDynamicIndirectLighting = false;
 		MeshComponent->PrimaryComponentTick.TickGroup = TG_PrePhysics;
 		MeshComponent->AttachParent = RootComponent;
-		MeshComponent->SetCollisionProfileName(CollisionProfileName);
+		MeshComponent->SetCollisionProfileName(UCollisionProfile::Pawn_ProfileName);
 		const float Scale = CollisionComponent->GetUnscaledSphereRadius() / 160.f; // @TODO: hardcoding known size of EngineMeshes.Sphere. Should use a unit sphere instead.
 		MeshComponent->SetRelativeScale3D(FVector(Scale));
 		MeshComponent->bGenerateOverlapEvents = false;
