#ue4 - Remove collision from DefaultPhysicsVolume. It doesn't have collision geometry but it shows up in global filters I'm adding.

--------

diff --git a/Engine/Source/Runtime/Engine/Private/DefaultPhysicsVolume.cpp b/Engine/Source/Runtime/Engine/Private/DefaultPhysicsVolume.cpp
--- a/Engine/Source/Runtime/Engine/Private/DefaultPhysicsVolume.cpp
+++ b/Engine/Source/Runtime/Engine/Private/DefaultPhysicsVolume.cpp
@@ -3,6 +3,7 @@
 #include "EnginePrivate.h"
 #include "GameFramework/DefaultPhysicsVolume.h"
 #include "PhysicsEngine/PhysicsSettings.h"
+#include "Components/BrushComponent.h"
 
 ADefaultPhysicsVolume::ADefaultPhysicsVolume(const FObjectInitializer& ObjectInitializer)
 	: Super(ObjectInitializer)
@@ -16,4 +17,8 @@ ADefaultPhysicsVolume::ADefaultPhysicsVolume(const FObjectInitializer& ObjectIni
 	// update default values when world is restarted
 	TerminalVelocity = UPhysicsSettings::Get()->DefaultTerminalVelocity;
 	FluidFriction = UPhysicsSettings::Get()->DefaultFluidFriction;
-}
\ No newline at end of file
+
+	// DefaultPhysicsVolumes are spawned only as a fallback object when determining the current physics volume.
+	// They are not intended to actually have any collision response, as they don't have actual collision geometry.
+	GetBrushComponent()->SetCollisionEnabled(ECollisionEnabled::NoCollision);
+}
