Extra variable checking

diff --git a/Engine/Source/Runtime/Engine/Private/DemoNetDriver.cpp b/Engine/Source/Runtime/Engine/Private/DemoNetDriver.cpp
--- a/Engine/Source/Runtime/Engine/Private/DemoNetDriver.cpp
+++ b/Engine/Source/Runtime/Engine/Private/DemoNetDriver.cpp
@@ -1232,13 +1232,19 @@ void UDemoNetDriver::SaveCheckpoint()
 void UDemoNetDriver::AddEvent(const FString& Group, const FString& Meta, const TArray<uint8>& Data)
 {
 	uint32 SavedTimeMS = GetDemoCurrentTimeInMS();
-	ReplayStreamer->AddEvent(SavedTimeMS, Group, Meta, Data);
+	if (ReplayStreamer.IsValid())
+	{
+		ReplayStreamer->AddEvent(SavedTimeMS, Group, Meta, Data);
+	}
 	UE_LOG(LogDemo, Verbose, TEXT("Custom Event %s. Total: %i, Time: %2.2f"), *Group, Data.Num(), SavedTimeMS);
 }
 
 void UDemoNetDriver::EnumerateEvents(const FString& Group, FEnumerateEventsCompleteDelegate& EnumerationCompleteDelegate)
 {
-	ReplayStreamer->EnumerateEvents(Group, EnumerationCompleteDelegate);
+	if (ReplayStreamer.IsValid())
+	{
+		ReplayStreamer->EnumerateEvents(Group, EnumerationCompleteDelegate);
+	}
 }
 
 void UDemoNetDriver::RequestEventData(const FString& EventID, FOnRequestEventDataComplete& RequestEventDataCompleteDelegate)
