Merging

//depot/UE4-Orion/Engine/Source/Runtime/Engine/...

to //depot/UE4/Engine/Source/Runtime/Engine/...

-adding 'immortal asset tag for potentially immortal particle emitters
-adding auto-deactivate, auto-trimming emitter/effect lifetime for cases of burst-only effects

diff --git a/Engine/Source/Runtime/Engine/Classes/Particles/ParticleSystem.h b/Engine/Source/Runtime/Engine/Classes/Particles/ParticleSystem.h
--- a/Engine/Source/Runtime/Engine/Classes/Particles/ParticleSystem.h
+++ b/Engine/Source/Runtime/Engine/Classes/Particles/ParticleSystem.h
@@ -290,6 +290,9 @@ public:
 	UPROPERTY(EditAnywhere, Category=Delay)
 	uint32 bUseDelayRange:1;
 
+	UPROPERTY(EditAnywhere, Category = Performance, meta = (ToolTip = "Auto-deactivate system if all emitters are determined to not spawn particles again, regardless of lifetime."))
+	bool bAutoDeactivate;
+
 	/** Local space position that UVs generated with the ParticleMacroUV material node will be centered on. */
 	UPROPERTY(EditAnywhere, Category=MacroUV)
 	FVector MacroUVPosition;
diff --git a/Engine/Source/Runtime/Engine/Private/Particles/ParticleComponents.cpp b/Engine/Source/Runtime/Engine/Private/Particles/ParticleComponents.cpp
--- a/Engine/Source/Runtime/Engine/Private/Particles/ParticleComponents.cpp
+++ b/Engine/Source/Runtime/Engine/Private/Particles/ParticleComponents.cpp
@@ -2089,6 +2089,7 @@ void UParticleSystem::GetAssetRegistryTags(TArray<FAssetRegistryTag>& OutTags) c
 
 	// Run thru the emitters and see if any will loop forever
 	bool bAnyEmitterLoopsForever = false;
+	bool bEmitterIsDangerous = false;
 	for (int32 EmitterIndex = 0; (EmitterIndex < Emitters.Num()) && !bAnyEmitterLoopsForever; ++EmitterIndex)
 	{
 		if (const UParticleEmitter* Emitter = Emitters[EmitterIndex])
@@ -2105,11 +2106,25 @@ void UParticleSystem::GetAssetRegistryTags(TArray<FAssetRegistryTag>& OutTags) c
 						bAnyEmitterLoopsForever = true;
 						break;
 					}
+
+					UParticleModuleSpawn *SpawnModule = LODLevel->SpawnModule;
+					check(SpawnModule);
+
+					if (SpawnModule->GetMaximumSpawnRate() == 0 
+						&& RequiredModule->EmitterDuration == 0
+						&& RequiredModule->EmitterLoops == 0 
+						)
+					{
+						bEmitterIsDangerous = true;
+					}
+
 				}
 			}
 		}
 	}
+
 	OutTags.Add(FAssetRegistryTag("Looping", bAnyEmitterLoopsForever ? TEXT("True") : TEXT("False"), FAssetRegistryTag::TT_Alphabetical));
+	OutTags.Add(FAssetRegistryTag("Immortal", bEmitterIsDangerous ? TEXT("True") : TEXT("False"), FAssetRegistryTag::TT_Alphabetical));
 
 	Super::GetAssetRegistryTags(OutTags);
 }
@@ -3915,6 +3930,7 @@ void UParticleSystemComponent::ComputeTickComponent_Concurrent()
 	FScopeCycleCounterUObject AdditionalScope(AdditionalStatObject(), GET_STATID(STAT_ParticleComputeTickTime));
 	// Tick Subemitters.
 	int32 EmitterIndex;
+	bool bAllEmittersFinished = true;
 	for (EmitterIndex = 0; EmitterIndex < EmitterInstances.Num(); EmitterIndex++)
 	{
 		FParticleEmitterInstance* Instance = EmitterInstances[EmitterIndex];
@@ -3933,6 +3949,10 @@ void UParticleSystemComponent::ComputeTickComponent_Concurrent()
 			if (SpriteLODLevel && SpriteLODLevel->bEnabled)
 			{
 				Instance->Tick(DeltaTimeTick, bSuppressSpawning);
+				if (!Instance->bEmitterIsDone == true)
+				{
+					bAllEmittersFinished = false;
+				}
 
 				if (!Instance->Tick_MaterialOverrides())
 				{
@@ -3948,6 +3968,11 @@ void UParticleSystemComponent::ComputeTickComponent_Concurrent()
 			}
 		}
 	}
+	
+	if (bAllEmittersFinished == true && Template->bAutoDeactivate==true)
+	{
+		this->DeactivateSystem();
+	}
 }
 
 void UParticleSystemComponent::FinalizeTickComponent()
diff --git a/Engine/Source/Runtime/Engine/Private/Particles/ParticleEmitterInstances.cpp b/Engine/Source/Runtime/Engine/Private/Particles/ParticleEmitterInstances.cpp
--- a/Engine/Source/Runtime/Engine/Private/Particles/ParticleEmitterInstances.cpp
+++ b/Engine/Source/Runtime/Engine/Private/Particles/ParticleEmitterInstances.cpp
@@ -600,6 +600,8 @@ void FParticleEmitterInstance::Init()
 
 	// Tag it as dirty w.r.t. the renderer
 	IsRenderDataDirty	= 1;
+
+	bEmitterIsDone = false;
 }
 
 UWorld* FParticleEmitterInstance::GetWorld() const
@@ -768,6 +770,9 @@ void FParticleEmitterInstance::Tick(float DeltaTime, bool bSuppressSpawning)
 
 	Tick_ModuleFinalUpdate(DeltaTime, LODLevel);
 
+
+	CheckEmitterFinished();
+
 	// Invalidate the contents of the vertex/index buffer.
 	IsRenderDataDirty = 1;
 
@@ -783,6 +788,49 @@ void FParticleEmitterInstance::Tick(float DeltaTime, bool bSuppressSpawning)
 	INC_DWORD_STAT_BY(STAT_SpriteParticles, ActiveParticles);
 }
 
+
+/**
+*	Called from Tick to determine whether the emitter will no longer spawn particles
+*   checks for emitters with 0 loops, infinite lifetime, and no continuous spawning (only bursts)
+*	and sets bEmitterIsDone if the last burst lies in the past and there are no active particles
+*	bEmitterIsDone is checked for all emitters by ParticleSystemComponent tick, and the particle
+*	system is deactivated if it's true for all emitters, and if bAutoDeactivate is set on the ParticleSystem
+*/
+void FParticleEmitterInstance::CheckEmitterFinished()
+{
+	// Grab the current LOD level
+	UParticleLODLevel* LODLevel = GetCurrentLODLevelChecked();
+
+	// figure out if this emitter will no longer spawn particles
+	//
+	if (this->ActiveParticles == 0)
+	{
+		UParticleModuleSpawn *SpawnModule = LODLevel->SpawnModule;
+		check(SpawnModule);
+
+		FParticleBurst *LastBurst = nullptr;
+		if (SpawnModule->BurstList.Num())
+		{
+			LastBurst = &SpawnModule->BurstList.Last();
+		}
+
+		if (!LastBurst || LastBurst->Time < this->EmitterTime)
+		{
+			const UParticleModuleRequired* RequiredModule = LODLevel->RequiredModule;
+			check(RequiredModule);
+
+			if (SpawnModule->GetMaximumSpawnRate() == 0
+				&& RequiredModule->EmitterDuration == 0
+				&& RequiredModule->EmitterLoops == 0
+				)
+			{
+				bEmitterIsDone = true;
+			}
+		}
+	}
+}
+
+
 /**
  *	Tick sub-function that handle EmitterTime setup, looping, etc.
  *
@@ -979,6 +1027,7 @@ void FParticleEmitterInstance::Tick_ModuleFinalUpdate(float DeltaTime, UParticle
 		}
 	}
 
+
 	if (InCurrentLODLevel->TypeDataModule && InCurrentLODLevel->TypeDataModule->bEnabled && InCurrentLODLevel->TypeDataModule->bFinalUpdateModule)
 	{
 		uint32* Offset = ModuleOffsetMap.Find(HighestLODLevel->TypeDataModule);
diff --git a/Engine/Source/Runtime/Engine/Public/ParticleEmitterInstances.h b/Engine/Source/Runtime/Engine/Public/ParticleEmitterInstances.h
--- a/Engine/Source/Runtime/Engine/Public/ParticleEmitterInstances.h
+++ b/Engine/Source/Runtime/Engine/Public/ParticleEmitterInstances.h
@@ -284,7 +284,8 @@ public:
 	TArray<float> EmitterDurations;
 	/** The emitter's delay for the current loop		*/
 	float CurrentDelay;
-
+	/** true if the emitter has no active particles and will no longer spawn any in the future */
+	bool bEmitterIsDone;
 
 	/** The number of triangles to render								*/
 	int32	TrianglesToRender;
@@ -331,6 +332,8 @@ public:
 	virtual bool Resize(int32 NewMaxActiveParticles, bool bSetMaxActiveCount = true);
 
 	virtual void Tick(float DeltaTime, bool bSuppressSpawning);
+	void CheckEmitterFinished();
+
 	/**
 	 *	Tick sub-function that handles EmitterTime setup, looping, etc.
 	 *
