Merging

//depot/UE4-Orion/Engine/Source/Runtime/Engine/...

to //depot/UE4/Engine/Source/Runtime/Engine/...

diff --git a/Engine/Source/Runtime/Engine/Classes/Particles/ParticleSystemComponent.h b/Engine/Source/Runtime/Engine/Classes/Particles/ParticleSystemComponent.h
--- a/Engine/Source/Runtime/Engine/Classes/Particles/ParticleSystemComponent.h
+++ b/Engine/Source/Runtime/Engine/Classes/Particles/ParticleSystemComponent.h
@@ -794,6 +794,10 @@ private:
 	int32 TotalActiveParticles;
 	/** If true, it means the ASync work is done and the finalize is not */
 	bool bNeedsFinalize;
+	/** If true, it means the Async work is in process and not yet completed */
+	bool bAsyncWorkOutstanding;
+	/** This flag is only valid during finalize. It is sent back from the potentially async task to indicate that all emitters are finished */
+	bool bAllEmittersFinished;
 public:
 
 	//~ Begin UActorComponent Interface.
diff --git a/Engine/Source/Runtime/Engine/Private/Particles/ParticleComponents.cpp b/Engine/Source/Runtime/Engine/Private/Particles/ParticleComponents.cpp
--- a/Engine/Source/Runtime/Engine/Private/Particles/ParticleComponents.cpp
+++ b/Engine/Source/Runtime/Engine/Private/Particles/ParticleComponents.cpp
@@ -2743,6 +2743,7 @@ UParticleSystemComponent::UParticleSystemComponent(const FObjectInitializer& Obj
 	bIsViewRelevanceDirty = true;
 	CustomTimeDilation = 1.0f;
 	bAllowConcurrentTick = true;
+	bAsyncWorkOutstanding = false;
 	bWasActive = false;
 #if WITH_EDITORONLY_DATA
 	EditorDetailMode = -1;
@@ -3891,6 +3892,7 @@ void UParticleSystemComponent::TickComponent(float DeltaTime, enum ELevelTick Ti
 		AsyncInstanceParameters.Reset();
 		AsyncInstanceParameters.Append(InstanceParameters);
 
+		bAsyncWorkOutstanding = true;
 		{
 			QUICK_SCOPE_CYCLE_COUNTER(STAT_UParticleSystemComponent_QueueAsync);
 			AsyncWork = TGraphTask<FParticleAsyncTask>::CreateTask(nullptr, ENamedThreads::GameThread).ConstructAndDispatchWhenReady(this);
@@ -3930,7 +3932,7 @@ void UParticleSystemComponent::ComputeTickComponent_Concurrent()
 	FScopeCycleCounterUObject AdditionalScope(AdditionalStatObject(), GET_STATID(STAT_ParticleComputeTickTime));
 	// Tick Subemitters.
 	int32 EmitterIndex;
-	bool bAllEmittersFinished = true;
+	bAllEmittersFinished = true;
 	for (EmitterIndex = 0; EmitterIndex < EmitterInstances.Num(); EmitterIndex++)
 	{
 		FParticleEmitterInstance* Instance = EmitterInstances[EmitterIndex];
@@ -3968,10 +3970,10 @@ void UParticleSystemComponent::ComputeTickComponent_Concurrent()
 			}
 		}
 	}
-	
-	if (bAllEmittersFinished == true && Template->bAutoDeactivate==true)
+	if (bAsyncWorkOutstanding)
 	{
-		this->DeactivateSystem();
+		FPlatformMisc::MemoryBarrier();
+		bAsyncWorkOutstanding = false;
 	}
 }
 
@@ -3985,6 +3987,12 @@ void UParticleSystemComponent::FinalizeTickComponent()
 	}
 	AsyncWork = NULL; // this task is done
 	bNeedsFinalize = false;
+
+	// should this be moved later?
+	if (bAllEmittersFinished == true && Template->bAutoDeactivate==true)
+	{
+		this->DeactivateSystem();
+	}
 	if (FXConsoleVariables::bFreezeParticleSimulation == false)
 	{
 		int32 EmitterIndex;
@@ -4115,25 +4123,26 @@ void UParticleSystemComponent::WaitForAsyncAndFinalize(EForceAsyncWorkCompletion
 {
 	if (AsyncWork.GetReference() && !AsyncWork->IsComplete())
 	{
-#if WITH_EDITOR
 		check(IsInGameThread());
 		SCOPE_CYCLE_COUNTER(STAT_GTSTallTime);
 		double StartTime = FPlatformTime::Seconds();
 		QUICK_SCOPE_CYCLE_COUNTER(STAT_UParticleSystemComponent_WaitForAsyncAndFinalize);
+#if WITH_EDITOR
 		FTaskGraphInterface::Get().WaitUntilTaskCompletes(AsyncWork, ENamedThreads::GameThread_Local);
+#else
+		// since in the non-editor case the completion is chained to a game thread task (not a gamethread_local one), and we don't want to execute arbitrary tasks 
+		// in what is probably a very, very deep callstack, we will spin here and wait for the async task to finish. The we will do the finalize. The finalize will be attempted again later but do nothing
+		while (bAsyncWorkOutstanding)
+		{
+			FPlatformProcess::SleepNoStats(0.0f);
+		}
+#endif
 		float ThisTime = float(FPlatformTime::Seconds() - StartTime) * 1000.0f;
 		if (Behavior != SILENT)
 		{
 			UE_LOG(LogParticles, Warning, TEXT("Stalled gamethread waiting for particles %5.2fms '%s' '%s'"), ThisTime, *GetFullNameSafe(this), *GetFullNameSafe(Template));
-			if (Behavior == ENSURE_AND_STALL)
-			{
-				ensureMsgf(0,TEXT("Stalled gamethread waiting for particles %5.2fms '%s' '%s'"), ThisTime, *GetFullNameSafe(this), *GetFullNameSafe(Template));
-			}
 		}
 		const_cast<UParticleSystemComponent*>(this)->FinalizeTickComponent();
-#else
-		check(0);
-#endif
 	}
 }
 
@@ -5279,10 +5288,10 @@ int32 UParticleSystemComponent::DetermineLODLevelForLocation(const FVector& Effe
 		// This will now put everything in LODLevel 0 (high detail) by default
 		float LODDistanceSqr = (PlayerViewLocations.Num() ? FMath::Square(WORLD_MAX) : 0.0f);
 		for (const FVector& ViewLocation : PlayerViewLocations)
-		{
+			{
 			const float DistanceToEffectSqr = FVector(ViewLocation - EffectLocation).SizeSquared();
 			if (DistanceToEffectSqr < LODDistanceSqr)
-			{
+				{
 				LODDistanceSqr = DistanceToEffectSqr;
 			}
 		}
