Fixed reversed indices logic not taking into account material section.

diff --git a/Engine/Source/Developer/MeshUtilities/Private/MeshUtilities.cpp b/Engine/Source/Developer/MeshUtilities/Private/MeshUtilities.cpp
--- a/Engine/Source/Developer/MeshUtilities/Private/MeshUtilities.cpp
+++ b/Engine/Source/Developer/MeshUtilities/Private/MeshUtilities.cpp
@@ -3040,9 +3040,16 @@ public:
 				TArray<uint32> InversedIndices;
 				const int32 IndexCount = CombinedIndices.Num();
 				InversedIndices.AddUninitialized(IndexCount);
-				for (int32 i = 0; i < IndexCount; ++i)
+
+				for (int32 SectionIndex = 0; SectionIndex < LODModel.Sections.Num(); ++SectionIndex)
 				{
-					InversedIndices[i] = CombinedIndices[IndexCount - 1 - i];
+					const FStaticMeshSection& SectionInfo = LODModel.Sections[SectionIndex];
+					const int32 SectionIndexCount = SectionInfo.NumTriangles * 3;
+
+					for (int32 i = 0; i < SectionIndexCount; ++i)
+					{
+						InversedIndices[SectionInfo.FirstIndex + i] = CombinedIndices[SectionInfo.FirstIndex + SectionIndexCount - 1 - i];
+					}
 				}
 				LODModel.ReversedIndexBuffer.SetIndices(InversedIndices, bNeeds32BitIndices ? EIndexBufferStride::Force32Bit : EIndexBufferStride::Force16Bit);
 			}
diff --git a/Engine/Source/Runtime/Engine/Private/StaticMesh.cpp b/Engine/Source/Runtime/Engine/Private/StaticMesh.cpp
--- a/Engine/Source/Runtime/Engine/Private/StaticMesh.cpp
+++ b/Engine/Source/Runtime/Engine/Private/StaticMesh.cpp
@@ -1089,7 +1089,8 @@ FArchive& operator<<(FArchive& Ar, FMeshBuildSettings& BuildSettings)
 // differences, etc.) replace the version GUID below with a new one.
 // In case of merge conflicts with DDC versions, you MUST generate a new GUID
 // and set this new GUID as the version.
-#define STATICMESH_DERIVEDDATA_VER TEXT("0651E844A9544BD83BF12B1807773C2")
+#define STATICMESH_DERIVEDDATA_VER TEXT("37403C6FB9441A6992E4CB614156E6A")
+
 static const FString& GetStaticMeshDerivedDataVersion()
 {
 	static FString CachedVersionString;
diff --git a/Engine/Source/Runtime/Engine/Private/StaticMeshRender.cpp b/Engine/Source/Runtime/Engine/Private/StaticMeshRender.cpp
--- a/Engine/Source/Runtime/Engine/Private/StaticMeshRender.cpp
+++ b/Engine/Source/Runtime/Engine/Private/StaticMeshRender.cpp
@@ -21,7 +21,7 @@
 static bool GUseShadowIndexBuffer = true;
 
 /** If true, reversed index buffer are used for mesh with negative transform determinants. */
-static bool GUseReversedIndexBuffer = false;
+static bool GUseReversedIndexBuffer = true;
 
 static void ToggleShadowIndexBuffers()
 {
@@ -273,7 +273,9 @@ bool FStaticMeshSceneProxy::GetMeshElement(int32 LODIndex, int32 BatchIndex, int
 
 	const bool bWireframe = false;
 	const bool bRequiresAdjacencyInformation = RequiresAdjacencyInformation( Material, OutMeshBatch.VertexFactory->GetType(), GetScene().GetFeatureLevel() );
-	const bool bUseReversedIndices = GUseReversedIndexBuffer && !bWireframe && !bRequiresAdjacencyInformation && IsLocalToWorldDeterminantNegative() && LOD.bHasReversedIndices;
+	
+	// Two sided material use bIsFrontFace which is wrong with Reversed Indices. AdjacencyInformation use another index buffer.
+	const bool bUseReversedIndices = GUseReversedIndexBuffer && IsLocalToWorldDeterminantNegative() && LOD.bHasReversedIndices && !bWireframe && !bRequiresAdjacencyInformation && !Material->IsTwoSided();
 
 	SetIndexSource(LODIndex, SectionIndex, OutMeshBatch, bWireframe, bRequiresAdjacencyInformation, bUseReversedIndices);
 
