New feature for defining grouped max concurrency rules

Overview:
- Created a new UObject: USoundConcurrency
- Defines concurrency behaviors and allows for different USoundBase assets (e.g. SoundWave's, SoundCues) to be a part of the same concurrency group.
- This feature will allow a number of powerful techniques to limit and control the number of sound assets which can play concurrently

Details and Additional Related Features:
- USoundBase's not specifying an override or a USoundConcurrency ref will result in no concurrency limits.
- Added support to existing static gameplay functions to provide an optional argument to override the USoundConcurrency of a sound at the point of playing
- Added a few new concurrency resolution rules: Resolve based on priority and Resolve based on volume
- Volume based concurrency resolution is implemented slightly different than the others due to the fact that it's complex/expensive to compute an active sound's volume before it plays due to SoundCue graph traversal. Thus, this step is made after a sound is already made it through one update of the SoundCue graph.
- It's now possible to limit max concurrency (both in groups and per-sound-instance) to the "owner" of a sound. This way different characters/actors can use the same USoundConcurrency ref but have different max-concurrency counts. With this technique, you can, for example, limit dialogue lines per actor without needing to make a USoundConcurrency object for each actor.
- A new optional concurrency behavior was also added at a suggestion from the Epic audio team: the ability to attenuate a sound based on the instance count of a sound in its concurrency group. I.e. as more sounds play in a concurrency group, the sound instances get quieter.

Backwards Compat:
- The old way limited sounds per USoundBase instance (and not different ones). You can recover this behavior by specifying an "Override" settings structure on the USoundBase. When you do this, the feature will only limit sound concurrency per instance and not per group. All existing assets with MaxConcurrency settings will fallback to this behavior.
- The default USoundBase behavior for new USoundBase objects is the same as the old behavior - i.e. max limit of 16 and non-grouped (i.e. will be an override)

diff --git a/Engine/Source/Developer/AssetTools/Private/AssetTools.cpp b/Engine/Source/Developer/AssetTools/Private/AssetTools.cpp
--- a/Engine/Source/Developer/AssetTools/Private/AssetTools.cpp
+++ b/Engine/Source/Developer/AssetTools/Private/AssetTools.cpp
@@ -87,6 +87,7 @@ FAssetTools::FAssetTools()
 	RegisterAssetTypeActions( MakeShareable(new FAssetTypeActions_SlateBrush) );
 	RegisterAssetTypeActions( MakeShareable(new FAssetTypeActions_SlateWidgetStyle) );
 	RegisterAssetTypeActions( MakeShareable(new FAssetTypeActions_SoundAttenuation) );
+	RegisterAssetTypeActions( MakeShareable(new FAssetTypeActions_SoundConcurrency));
 	RegisterAssetTypeActions( MakeShareable(new FAssetTypeActions_SoundBase) );
 	RegisterAssetTypeActions( MakeShareable(new FAssetTypeActions_SoundClass) );
 	RegisterAssetTypeActions( MakeShareable(new FAssetTypeActions_SoundCue) );
diff --git a/Engine/Source/Developer/AssetTools/Private/AssetToolsPrivatePCH.h b/Engine/Source/Developer/AssetTools/Private/AssetToolsPrivatePCH.h
--- a/Engine/Source/Developer/AssetTools/Private/AssetToolsPrivatePCH.h
+++ b/Engine/Source/Developer/AssetTools/Private/AssetToolsPrivatePCH.h
@@ -75,6 +75,7 @@ DECLARE_LOG_CATEGORY_EXTERN(LogAssetTools, Log, All);
 #include "AssetTypeActions/AssetTypeActions_SlateBrush.h"
 #include "AssetTypeActions/AssetTypeActions_SlateWidgetStyle.h"
 #include "AssetTypeActions/AssetTypeActions_SoundAttenuation.h"
+#include "AssetTypeActions/AssetTypeActions_SoundConcurrency.h"
 #include "AssetTypeActions/AssetTypeActions_SoundClass.h"
 #include "AssetTypeActions/AssetTypeActions_SoundCue.h"
 #include "AssetTypeActions/AssetTypeActions_SoundMix.h"
diff --git a/Engine/Source/Developer/AssetTools/Private/AssetTypeActions/AssetTypeActions_SoundConcurrency.cpp b/Engine/Source/Developer/AssetTools/Private/AssetTypeActions/AssetTypeActions_SoundConcurrency.cpp
new file mode 100644
--- /dev/null
+++ b/Engine/Source/Developer/AssetTools/Private/AssetTypeActions/AssetTypeActions_SoundConcurrency.cpp
@@ -0,0 +1,9 @@
+// Copyright 1998-2015 Epic Games, Inc. All Rights Reserved.
+
+#include "AssetToolsPrivatePCH.h"
+#include "SoundDefinitions.h"
+
+UClass* FAssetTypeActions_SoundConcurrency::GetSupportedClass() const
+{
+	return USoundConcurrency::StaticClass();
+}
\ No newline at end of file
diff --git a/Engine/Source/Developer/AssetTools/Private/AssetTypeActions/AssetTypeActions_SoundConcurrency.h b/Engine/Source/Developer/AssetTools/Private/AssetTypeActions/AssetTypeActions_SoundConcurrency.h
new file mode 100644
--- /dev/null
+++ b/Engine/Source/Developer/AssetTools/Private/AssetTypeActions/AssetTypeActions_SoundConcurrency.h
@@ -0,0 +1,13 @@
+// Copyright 1998-2015 Epic Games, Inc. All Rights Reserved.
+
+#pragma once
+
+class FAssetTypeActions_SoundConcurrency : public FAssetTypeActions_Base
+{
+public:
+	// IAssetTypeActions Implementation
+	virtual FText GetName() const override { return NSLOCTEXT("AssetTypeActions", "AssetTypeActions_SoundConcurrency", "Sound Concurrency"); }
+	virtual FColor GetTypeColor() const override { return FColor(77, 100, 139); }
+	virtual UClass* GetSupportedClass() const override;
+	virtual uint32 GetCategories() override { return EAssetTypeCategories::Sounds; }
+};
\ No newline at end of file
diff --git a/Engine/Source/Editor/UnrealEd/Classes/Factories/SoundConcurrencyFactory.h b/Engine/Source/Editor/UnrealEd/Classes/Factories/SoundConcurrencyFactory.h
new file mode 100644
--- /dev/null
+++ b/Engine/Source/Editor/UnrealEd/Classes/Factories/SoundConcurrencyFactory.h
@@ -0,0 +1,21 @@
+// Copyright 1998-2015 Epic Games, Inc. All Rights Reserved.
+
+//=============================================================================
+// USoundConcurrencyFactory
+//=============================================================================
+
+#pragma once
+#include "SoundConcurrencyFactory.generated.h"
+
+UCLASS(hidecategories=Object, MinimalAPI)
+class USoundConcurrencyFactory : public UFactory
+{
+	GENERATED_UCLASS_BODY()
+
+	// Begin UFactory Interface
+	virtual UObject* FactoryCreateNew(UClass* Class,UObject* InParent,FName Name,EObjectFlags Flags,UObject* Context,FFeedbackContext* Warn) override;
+	// Begin UFactory Interface	
+};
+
+
+
diff --git a/Engine/Source/Editor/UnrealEd/Private/Factories/EditorFactories.cpp b/Engine/Source/Editor/UnrealEd/Private/Factories/EditorFactories.cpp
--- a/Engine/Source/Editor/UnrealEd/Private/Factories/EditorFactories.cpp
+++ b/Engine/Source/Editor/UnrealEd/Private/Factories/EditorFactories.cpp
@@ -2405,6 +2405,25 @@ UObject* USoundAttenuationFactory::FactoryCreateNew( UClass* Class, UObject* InP
 	return NewObject<USoundAttenuation>(InParent, Name, Flags);
 }
 
+/*------------------------------------------------------------------------------
+USoundConcurrencyFactory.
+------------------------------------------------------------------------------*/
+USoundConcurrencyFactory::USoundConcurrencyFactory(const FObjectInitializer& ObjectInitializer)
+	: Super(ObjectInitializer)
+{
+
+	SupportedClass = USoundConcurrency::StaticClass();
+	bCreateNew = true;
+	bEditorImport = false;
+	bEditAfterNew = true;
+}
+
+UObject* USoundConcurrencyFactory::FactoryCreateNew(UClass* Class, UObject* InParent, FName Name, EObjectFlags Flags, UObject* Context, FFeedbackContext* Warn)
+{
+	return NewObject<USoundConcurrency>(InParent, Name, Flags);
+}
+
+
 /*------------------------------------------------------------------------------
 	UParticleSystemFactoryNew.
 ------------------------------------------------------------------------------*/
diff --git a/Engine/Source/Editor/UnrealEd/Public/UnrealEd.h b/Engine/Source/Editor/UnrealEd/Public/UnrealEd.h
--- a/Engine/Source/Editor/UnrealEd/Public/UnrealEd.h
+++ b/Engine/Source/Editor/UnrealEd/Public/UnrealEd.h
@@ -165,6 +165,7 @@
 #include "Classes/Factories/SlateBrushAssetFactory.h"
 #include "Classes/Factories/SlateWidgetStyleAssetFactory.h"
 #include "Classes/Factories/SoundAttenuationFactory.h"
+#include "Classes/Factories/SoundConcurrencyFactory.h"
 #include "Classes/Factories/SoundClassFactory.h"
 #include "Classes/Factories/SoundCueFactoryNew.h"
 #include "Classes/Factories/SoundFactory.h"
diff --git a/Engine/Source/Runtime/Core/Public/UObject/ObjectVersion.h b/Engine/Source/Runtime/Core/Public/UObject/ObjectVersion.h
--- a/Engine/Source/Runtime/Core/Public/UObject/ObjectVersion.h
+++ b/Engine/Source/Runtime/Core/Public/UObject/ObjectVersion.h
@@ -572,6 +572,8 @@ enum EUnrealEngineObjectUE4Version
 	VER_UE4_MAKE_ROT_RENAME_AND_REORDER,
 	// K2Node_Variable will properly have the VariableReference Guid set if available
 	VER_UE4_K2NODE_VAR_REFERENCEGUIDS,
+	// Added support for sound concurrency settings structure and overrides
+	VER_UE4_SOUND_CONCURRENCY_PACKAGE,
 	// -----<new versions can be added before this line>-------------------------------------------------
 	// - this needs to be the last line (see note below)
 	VER_UE4_AUTOMATIC_VERSION_PLUS_ONE,
diff --git a/Engine/Source/Runtime/Engine/Classes/Components/AudioComponent.h b/Engine/Source/Runtime/Engine/Classes/Components/AudioComponent.h
--- a/Engine/Source/Runtime/Engine/Classes/Components/AudioComponent.h
+++ b/Engine/Source/Runtime/Engine/Classes/Components/AudioComponent.h
@@ -141,6 +141,10 @@ class ENGINE_API UAudioComponent : public USceneComponent
 	UPROPERTY(EditAnywhere, BlueprintReadWrite, Category=Attenuation, meta=(EditCondition="bOverrideAttenuation"))
 	struct FAttenuationSettings AttenuationOverrides;
 
+	/** What sound concurrency to use for sounds generated by this audio component */
+	UPROPERTY(EditAnywhere, BlueprintReadWrite, Category=Concurrency)
+	class USoundConcurrency* ConcurrencySettings;
+
 	/** while playing, this component will check for occlusion from its closest listener every this many seconds */
 	float OcclusionCheckInterval;
 
diff --git a/Engine/Source/Runtime/Engine/Classes/Kismet/GameplayStatics.h b/Engine/Source/Runtime/Engine/Classes/Kismet/GameplayStatics.h
--- a/Engine/Source/Runtime/Engine/Classes/Kismet/GameplayStatics.h
+++ b/Engine/Source/Runtime/Engine/Classes/Kismet/GameplayStatics.h
@@ -323,10 +323,11 @@ class ENGINE_API UGameplayStatics : public UBlueprintFunctionLibrary
 	 * @param Sound - Sound to play.
 	 * @param VolumeMultiplier - Multiplied with the volume to make the sound louder or softer.
 	 * @param PitchMultiplier - Multiplies the pitch.
+	 * @param ConcurrencySettings - Override concurrency settings package to play sound with
 	 * @param StartTime - How far in to the sound to begin playback at
 	 */
 	UFUNCTION(BlueprintCallable, BlueprintCosmetic, Category="Audio", meta=( WorldContext="WorldContextObject", AdvancedDisplay = "2", UnsafeDuringActorConstruction = "true" ))
-	static void PlaySound2D(UObject* WorldContextObject, class USoundBase* Sound, float VolumeMultiplier = 1.f, float PitchMultiplier = 1.f, float StartTime = 0.f);
+	static void PlaySound2D(UObject* WorldContextObject, class USoundBase* Sound, float VolumeMultiplier = 1.f, float PitchMultiplier = 1.f, float StartTime = 0.f, class USoundConcurrency* ConcurrencySettings = nullptr);
 
 	/**
 	 * Spawns a sound with no attenuation, perfect for UI sounds.
@@ -336,10 +337,11 @@ class ENGINE_API UGameplayStatics : public UBlueprintFunctionLibrary
 	 * @param VolumeMultiplier - Multiplied with the volume to make the sound louder or softer.
 	 * @param PitchMultiplier - Multiplies the pitch.
 	 * @param StartTime - How far in to the sound to begin playback at
+	 * @param ConcurrencySettings - Override concurrency settings package to play sound with
 	 * @return An audio component to manipulate the spawned sound
 	 */
 	UFUNCTION(BlueprintCallable, BlueprintCosmetic, Category="Audio", meta=( WorldContext="WorldContextObject", AdvancedDisplay = "2", UnsafeDuringActorConstruction = "true", Keywords = "play" ))
-	static UAudioComponent* SpawnSound2D(UObject* WorldContextObject, class USoundBase* Sound, float VolumeMultiplier = 1.f, float PitchMultiplier = 1.f, float StartTime = 0.f);
+	static UAudioComponent* SpawnSound2D(UObject* WorldContextObject, class USoundBase* Sound, float VolumeMultiplier = 1.f, float PitchMultiplier = 1.f, float StartTime = 0.f, class USoundConcurrency* ConcurrencySettings = nullptr);
 
 	/**
 	 * Plays a sound at the given location. This is a fire and forget sound and does not travel with any actor. Replication is also not handled at this point.
@@ -350,13 +352,14 @@ class ENGINE_API UGameplayStatics : public UBlueprintFunctionLibrary
 	 * @param PitchMultiplier - PitchMultiplier
 	 * @param StartTime - How far in to the sound to begin playback at
 	 * @param AttenuationSettings - Override attenuation settings package to play sound with
+	 * @param ConcurrencySettings - Override concurrency settings package to play sound with
 	 */
 	UFUNCTION(BlueprintCallable, Category="Audio", meta=(WorldContext="WorldContextObject", AdvancedDisplay = "3", UnsafeDuringActorConstruction = "true"))
-	static void PlaySoundAtLocation(UObject* WorldContextObject, class USoundBase* Sound, FVector Location, FRotator Rotation, float VolumeMultiplier = 1.f, float PitchMultiplier = 1.f, float StartTime = 0.f, class USoundAttenuation* AttenuationSettings = NULL);
+	static void PlaySoundAtLocation(UObject* WorldContextObject, class USoundBase* Sound, FVector Location, FRotator Rotation, float VolumeMultiplier = 1.f, float PitchMultiplier = 1.f, float StartTime = 0.f, class USoundAttenuation* AttenuationSettings = nullptr, class USoundConcurrency* ConcurrencySettings = nullptr);
 
-	static void PlaySoundAtLocation(UObject* WorldContextObject, class USoundBase* Sound, FVector Location, float VolumeMultiplier = 1.f, float PitchMultiplier = 1.f, float StartTime = 0.f, class USoundAttenuation* AttenuationSettings = NULL)
+	static void PlaySoundAtLocation(UObject* WorldContextObject, class USoundBase* Sound, FVector Location, float VolumeMultiplier = 1.f, float PitchMultiplier = 1.f, float StartTime = 0.f, class USoundAttenuation* AttenuationSettings = nullptr, class USoundConcurrency* ConcurrencySettings = nullptr)
 	{
-		PlaySoundAtLocation(WorldContextObject, Sound, Location, FRotator::ZeroRotator, VolumeMultiplier, PitchMultiplier, StartTime, AttenuationSettings);
+		PlaySoundAtLocation(WorldContextObject, Sound, Location, FRotator::ZeroRotator, VolumeMultiplier, PitchMultiplier, StartTime, AttenuationSettings, ConcurrencySettings);
 	}
 
 	/**
@@ -368,10 +371,11 @@ class ENGINE_API UGameplayStatics : public UBlueprintFunctionLibrary
 	 * @param PitchMultiplier - PitchMultiplier
 	 * @param StartTime - How far in to the sound to begin playback at
 	 * @param AttenuationSettings - Override attenuation settings package to play sound with
+	 * @param ConcurrencySettings - Override concurrency settings package to play sound with
 	 * @return An audio component to manipulate the spawned sound
 	 */
 	UFUNCTION(BlueprintCallable, Category="Audio", meta=(WorldContext="WorldContextObject", AdvancedDisplay = "3", UnsafeDuringActorConstruction = "true", Keywords = "play"))
-	static UAudioComponent* SpawnSoundAtLocation(UObject* WorldContextObject, class USoundBase* Sound, FVector Location, FRotator Rotation = FRotator::ZeroRotator, float VolumeMultiplier = 1.f, float PitchMultiplier = 1.f, float StartTime = 0.f, class USoundAttenuation* AttenuationSettings = NULL);
+	static UAudioComponent* SpawnSoundAtLocation(UObject* WorldContextObject, class USoundBase* Sound, FVector Location, FRotator Rotation = FRotator::ZeroRotator, float VolumeMultiplier = 1.f, float PitchMultiplier = 1.f, float StartTime = 0.f, class USoundAttenuation* AttenuationSettings = nullptr, class USoundConcurrency* ConcurrencySettings = nullptr);
 
 	/** Plays a sound attached to and following the specified component. This is a fire and forget sound. Replication is also not handled at this point.
 	 * @param Sound - sound to play
@@ -385,20 +389,21 @@ class ENGINE_API UGameplayStatics : public UBlueprintFunctionLibrary
 	 * @param PitchMultiplier - PitchMultiplier	 
 	 * @param StartTime - How far in to the sound to begin playback at
 	 * @param AttenuationSettings - Override attenuation settings package to play sound with
+	 * @param ConcurrencySettings - Override concurrency settings package to play sound with
 	 * @return An audio component to manipulate the spawned sound
 	 */
 	UFUNCTION(BlueprintCallable, Category="Audio", meta=(AdvancedDisplay = "2", UnsafeDuringActorConstruction = "true", Keywords = "play"))
-	static class UAudioComponent* SpawnSoundAttached(class USoundBase* Sound, class USceneComponent* AttachToComponent, FName AttachPointName = NAME_None, FVector Location = FVector(ForceInit), FRotator Rotation = FRotator::ZeroRotator, EAttachLocation::Type LocationType = EAttachLocation::KeepRelativeOffset, bool bStopWhenAttachedToDestroyed = false, float VolumeMultiplier = 1.f, float PitchMultiplier = 1.f, float StartTime = 0.f, class USoundAttenuation* AttenuationSettings = NULL);
+	static class UAudioComponent* SpawnSoundAttached(class USoundBase* Sound, class USceneComponent* AttachToComponent, FName AttachPointName = NAME_None, FVector Location = FVector(ForceInit), FRotator Rotation = FRotator::ZeroRotator, EAttachLocation::Type LocationType = EAttachLocation::KeepRelativeOffset, bool bStopWhenAttachedToDestroyed = false, float VolumeMultiplier = 1.f, float PitchMultiplier = 1.f, float StartTime = 0.f, class USoundAttenuation* AttenuationSettings = nullptr, class USoundConcurrency* ConcurrencySettings = nullptr);
 
 	DEPRECATED(4.9, "PlaySoundAttached has been renamed SpawnSoundAttached")
-	static class UAudioComponent* PlaySoundAttached(class USoundBase* Sound, class USceneComponent* AttachToComponent, FName AttachPointName = NAME_None, FVector Location = FVector(ForceInit), EAttachLocation::Type LocationType = EAttachLocation::KeepRelativeOffset, bool bStopWhenAttachedToDestroyed = false, float VolumeMultiplier = 1.f, float PitchMultiplier = 1.f, float StartTime = 0.f, class USoundAttenuation* AttenuationSettings = NULL)
+	static class UAudioComponent* PlaySoundAttached(class USoundBase* Sound, class USceneComponent* AttachToComponent, FName AttachPointName = NAME_None, FVector Location = FVector(ForceInit), EAttachLocation::Type LocationType = EAttachLocation::KeepRelativeOffset, bool bStopWhenAttachedToDestroyed = false, float VolumeMultiplier = 1.f, float PitchMultiplier = 1.f, float StartTime = 0.f, class USoundAttenuation* AttenuationSettings = nullptr)
 	{
 		return SpawnSoundAttached(Sound, AttachToComponent, AttachPointName, Location, FRotator::ZeroRotator, LocationType, bStopWhenAttachedToDestroyed, VolumeMultiplier, PitchMultiplier, StartTime, AttenuationSettings);
 	}
 
-	static class UAudioComponent* SpawnSoundAttached(class USoundBase* Sound, class USceneComponent* AttachToComponent, FName AttachPointName, FVector Location, EAttachLocation::Type LocationType = EAttachLocation::KeepRelativeOffset, bool bStopWhenAttachedToDestroyed = false, float VolumeMultiplier = 1.f, float PitchMultiplier = 1.f, float StartTime = 0.f, class USoundAttenuation* AttenuationSettings = NULL)
+	static class UAudioComponent* SpawnSoundAttached(class USoundBase* Sound, class USceneComponent* AttachToComponent, FName AttachPointName, FVector Location, EAttachLocation::Type LocationType = EAttachLocation::KeepRelativeOffset, bool bStopWhenAttachedToDestroyed = false, float VolumeMultiplier = 1.f, float PitchMultiplier = 1.f, float StartTime = 0.f, class USoundAttenuation* AttenuationSettings = nullptr, class USoundConcurrency* ConcurrencySettings = nullptr)
 	{
-		return SpawnSoundAttached(Sound, AttachToComponent, AttachPointName, Location, FRotator::ZeroRotator, LocationType, bStopWhenAttachedToDestroyed, VolumeMultiplier, PitchMultiplier, StartTime, AttenuationSettings);
+		return SpawnSoundAttached(Sound, AttachToComponent, AttachPointName, Location, FRotator::ZeroRotator, LocationType, bStopWhenAttachedToDestroyed, VolumeMultiplier, PitchMultiplier, StartTime, AttenuationSettings, ConcurrencySettings);
 	}
 
 	/**
diff --git a/Engine/Source/Runtime/Engine/Classes/Sound/AudioSettings.h b/Engine/Source/Runtime/Engine/Classes/Sound/AudioSettings.h
--- a/Engine/Source/Runtime/Engine/Classes/Sound/AudioSettings.h
+++ b/Engine/Source/Runtime/Engine/Classes/Sound/AudioSettings.h
@@ -41,6 +41,10 @@ class ENGINE_API UAudioSettings : public UDeveloperSettings
 	UPROPERTY(config, EditAnywhere, Category="Audio", meta=(AllowedClasses="SoundClass", DisplayName="Default Sound Class"))
 	FStringAssetReference DefaultSoundClassName;
 
+	/** The SoundConcurrency assigned to newly created sounds */
+	UPROPERTY(config, EditAnywhere, Category = "Audio", meta = (AllowedClasses = "SoundConcurrency", DisplayName = "Default Sound Concurrency"))
+	FStringAssetReference DefaultSoundConcurrencyName;
+
 	/** The SoundMix to use as base when no other system has specified a Base SoundMix */
 	UPROPERTY(config, EditAnywhere, Category="Audio", meta=(AllowedClasses="SoundMix"))
 	FStringAssetReference DefaultBaseSoundMix;
diff --git a/Engine/Source/Runtime/Engine/Classes/Sound/SoundBase.h b/Engine/Source/Runtime/Engine/Classes/Sound/SoundBase.h
--- a/Engine/Source/Runtime/Engine/Classes/Sound/SoundBase.h
+++ b/Engine/Source/Runtime/Engine/Classes/Sound/SoundBase.h
@@ -6,27 +6,19 @@
  * The base class for a playable sound object 
  */
 
+#include "Sound/SoundConcurrency.h"
 #include "SoundBase.generated.h"
 
+
 struct FActiveSound;
 struct FSoundParseParameters;
 struct FWaveInstance;
 struct FAttenuationSettings; 
+struct FSoundConcurrencySettings;
+
 class USoundClass;
 class USoundAttenuation;
-
-
-UENUM()
-namespace EMaxConcurrentResolutionRule
-{
-	enum Type
-	{
-		PreventNew					UMETA(ToolTip = "When Max Concurrent sounds are active do not start a new sound."),
-		StopOldest					UMETA(ToolTip = "When Max Concurrent sounds are active stop the oldest and start a new one."),
-		StopFarthestThenPreventNew	UMETA(ToolTip = "When Max Concurrent sounds are active stop the furthest sound.  If all sounds are the same distance then do not start a new sound."),
-		StopFarthestThenOldest		UMETA(ToolTip = "When Max Concurrent sounds are active stop the furthest sound.  If all sounds are the same distance then stop the oldest.")
-	};
-}
+class USoundConcurrency;
 
 UCLASS(config=Engine, hidecategories=Object, abstract, editinlinenew, BlueprintType)
 class ENGINE_API USoundBase : public UObject
@@ -40,16 +32,26 @@ protected:
 
 public:
 	/** When "stat sounds -debug" has been specified, draw this sound's attenuation shape when the sound is audible. For debugging purpose only. */
-	UPROPERTY(EditAnywhere, Category=Playback)
+	UPROPERTY(EditAnywhere, Category = Debug)
 	uint32 bDebug:1;
 
-	/** If we try to play a new version of this sound when at the max concurrent count how should it be resolved. */
-	UPROPERTY(EditAnywhere, Category=Playback)
-	TEnumAsByte<EMaxConcurrentResolutionRule::Type> MaxConcurrentResolutionRule;
+	UPROPERTY(EditAnywhere, BlueprintReadWrite, Category = Concurrency)
+	uint32 bOverrideConcurrency:1;
+
+	/** If bOverridePlayback is false, the sound concurrency settings to use for this sound. */
+	UPROPERTY(EditAnywhere, BlueprintReadWrite, Category = Concurrency, meta = (EditCondition = "!bOverrideConcurrency"))
+	class USoundConcurrency* SoundConcurrencySettings;
+
+	/** If bOverridePlayback is true, concurrency settings to use. */
+	UPROPERTY(EditAnywhere, BlueprintReadWrite, Category = Concurrency, meta = (EditCondition = "bOverrideConcurrency"))
+	struct FSoundConcurrencySettings ConcurrencyOverrides;
+
+	UPROPERTY()
+	TEnumAsByte<enum EMaxConcurrentResolutionRule::Type> MaxConcurrentResolutionRule_DEPRECATED;
 
 	/** Maximum number of times this sound can be played concurrently. */
-	UPROPERTY(EditAnywhere, Category=Playback)
-	int32 MaxConcurrentPlayCount;
+	UPROPERTY()
+	int32 MaxConcurrentPlayCount_DEPRECATED;
 
 	/** Duration of sound in seconds. */
 	UPROPERTY(Category=Info, AssetRegistrySearchable, VisibleAnywhere, BlueprintReadOnly)
@@ -65,7 +67,8 @@ public:
 
 	//~ Begin UObject Interface.
 	virtual void PostInitProperties() override;
-	//~ End UObject Interface.
+	virtual void PostLoad() override;
+	//~ End UObject interface.
 	
 	/** Returns whether the sound base is set up in a playable manner */
 	virtual bool IsPlayable() const;
@@ -108,5 +111,11 @@ public:
 	 * Returns the SoundClass used for this sound
 	 */
 	virtual USoundClass* GetSoundClass() const;
+
+	/** Returns the FSoundConcurrencySettings struct to use */
+	const FSoundConcurrencySettings* GetSoundConcurrencySettingsToApply();
+
+	/** Returns the sound concurrency object ID if it exists. If it doesn't exist, returns 0. */
+	uint32 GetSoundConcurrencyObjectID() const;
 };
 
diff --git a/Engine/Source/Runtime/Engine/Classes/Sound/SoundConcurrency.h b/Engine/Source/Runtime/Engine/Classes/Sound/SoundConcurrency.h
new file mode 100644
--- /dev/null
+++ b/Engine/Source/Runtime/Engine/Classes/Sound/SoundConcurrency.h
@@ -0,0 +1,189 @@
+// Copyright 1998-2015 Epic Games, Inc. All Rights Reserved.
+
+#pragma once
+
+#include "CoreUObject.h"
+#include "SoundConcurrency.generated.h"
+
+UENUM()
+namespace EMaxConcurrentResolutionRule
+{
+	enum Type
+	{
+		/** When Max Concurrent sounds are active do not start a new sound. */
+		PreventNew,
+
+		/** When Max Concurrent sounds are active stop the oldest and start a new one. */
+		StopOldest,
+
+		/** When Max Concurrent sounds are active stop the furthest sound.  If all sounds are the same distance then do not start a new sound. */
+		StopFarthestThenPreventNew,
+
+		/** When Max Concurrent sounds are active stop the furthest sound.  If all sounds are the same distance then stop the oldest. */
+		StopFarthestThenOldest,
+
+		/** Stop the lowest volume-weighted priority in the group. */
+		StopLowestPriority,
+
+		/** Stop the sound that is quietest in the group. */
+		StopQuietist,
+	};
+}
+
+
+USTRUCT(BlueprintType)
+struct ENGINE_API FSoundConcurrencySettings
+{
+	GENERATED_USTRUCT_BODY()
+
+	/** The max number of allowable concurrent active voices for voices playing in this concurrency group. */
+	UPROPERTY(EditAnywhere, BlueprintReadWrite, Category = Concurrency, meta = (ClampMin = "1"))
+	int32 MaxCount;
+
+	/* Whether or not to limit the concurrency to per sound owner (i.e. the actor that plays the sound). If the sound doesn't have an owner, it falls back to global concurrency. */
+	UPROPERTY(EditAnywhere, BlueprintReadWrite, Category = Concurrency)
+	uint32 bLimitToOwner:1;
+
+	/** Which concurrency resolution policy to use if max voice count is reached. */
+	UPROPERTY(EditAnywhere, BlueprintReadWrite, Category = Concurrency)
+	TEnumAsByte<enum EMaxConcurrentResolutionRule::Type> ResolutionRule;
+
+	/** The amount of attenuation to apply to each voice instance in this concurrency group. */
+	UPROPERTY(EditAnywhere, BlueprintReadWrite, Category = Concurrency, meta = (ClampMin = "0", ClampMax = "1"))
+	float VolumeScale;
+};
+
+UCLASS(BlueprintType, hidecategories=Object, editinlinenew, MinimalAPI)
+class USoundConcurrency : public UObject
+{
+	GENERATED_UCLASS_BODY()
+
+	UPROPERTY(EditAnywhere, BlueprintReadOnly, Category = Settings)
+	FSoundConcurrencySettings Concurrency;
+};
+
+struct FActiveSound;
+
+/** USoundConcurrency unique object IDs. */
+typedef uint32 FConcurrencyObjectID;
+
+/** Sound owner object IDs */
+typedef uint32 FSoundOwnerObjectID;
+
+/** Sound instance (USoundBase) object ID. */
+typedef uint32 FSoundObjectID;
+
+typedef uint32 FConcurrencyGroupID;
+
+/** Struct which is an array of active sound pointers for tracking concurrency */
+class FConcurrencyGroup
+{
+	/** Array of active sounds for this concurrency group. */
+	TArray<FActiveSound*> ActiveSounds;
+	int32 MaxActiveSounds;
+	FConcurrencyGroupID ConcurrencyGroupID;
+	EMaxConcurrentResolutionRule::Type ResolutionRule;
+
+public:
+	/** Constructor for the max concurrency active sound entry. */
+	FConcurrencyGroup();
+	FConcurrencyGroup(struct FActiveSound* ActiveSound);
+	FConcurrencyGroup(const FConcurrencyGroup& Other);
+	FConcurrencyGroup(const FConcurrencyGroup&& Other);
+	FConcurrencyGroup& operator=(const FConcurrencyGroup& Other);
+	FConcurrencyGroup& operator=(const FConcurrencyGroup&& Other);
+
+	/** Retrieves the active sounds array. */
+	TArray<FActiveSound*>& GetActiveSounds();
+
+	/** Adds an active sound to the active sound array. */
+	void AddActiveSound(struct FActiveSound* ActiveSound);
+
+	/** Sorts the active sound array by volume */
+	void StopQuietSoundsDueToMaxConcurrency();
+
+	FConcurrencyGroupID GetID() const { return ConcurrencyGroupID; }
+};
+
+typedef TMap<FConcurrencyGroupID, FConcurrencyGroup> FConcurrencyGroups;
+
+struct FSoundInstanceEntry
+{
+	TMap<FSoundObjectID, FConcurrencyGroupID> SoundInstanceToConcurrencyGroup;
+
+	FSoundInstanceEntry(FSoundObjectID SoundObjectID, FConcurrencyGroupID GroupID)
+	{
+		SoundInstanceToConcurrencyGroup.Add(SoundObjectID, GroupID);
+	}
+};
+
+/** Type for mapping an object id to a concurrency entry. */
+typedef TMap<FConcurrencyObjectID, FConcurrencyGroupID> FConcurrencyMap;
+
+struct FOwnerConcurrencyMapEntry
+{
+	FConcurrencyMap ConcurrencyObjectToConcurrencyGroup;
+
+	FOwnerConcurrencyMapEntry(FConcurrencyObjectID ConcurrencyObjectID, FConcurrencyGroupID GroupID)
+	{
+		ConcurrencyObjectToConcurrencyGroup.Add(ConcurrencyObjectID, GroupID);
+	}
+};
+
+/** Maps owners to concurrency maps */
+typedef TMap<FSoundOwnerObjectID, FOwnerConcurrencyMapEntry> FOwnerConcurrencyMap;
+
+/** Maps owners to sound instances */
+typedef TMap<FSoundOwnerObjectID, FSoundInstanceEntry> FOwnerPerSoundConcurrencyMap;
+
+/** Maps sound object ids to active sound array for global concurrency limiting */
+typedef TMap<FSoundObjectID, FConcurrencyGroupID> FPerSoundToActiveSoundsMap;
+
+class FSoundConcurrencyManager
+{
+public:
+	FSoundConcurrencyManager(class FAudioDevice* InAudioDevice);
+	~FSoundConcurrencyManager();
+
+	/** Returns a newly allocated active sound given the input active sound struct. Will return nullptr if the active sound concurrency evaluation doesn't allow for it. */
+	FActiveSound* CreateNewActiveSound(const FActiveSound& NewActiveSound);
+
+	/** Removes the active sound from the manager to remove it from concurrency tracking. */
+	void RemoveActiveSound(FActiveSound* ActiveSound);
+
+	/** Stops any active sounds due to max concurrency quietest sound resolution rule */
+	void StopQuietSoundsDueToMaxConcurrency();
+
+private: // Methods
+	/** This function handles the concurrency evaluation that happens per USoundConcurrencyObject */
+	FActiveSound* HandleConcurrencyEvaluation(const FActiveSound& NewActiveSound);
+
+	/** This function handles the concurrency evaluation that happens per-voice rather than per USoundConcurrencyObject */
+	FActiveSound* HandleConcurrencyEvaluationOverride(const FActiveSound& NewActiveSound);
+
+	/** Resolves the concurrency resolution rule for a given sound and its concurrency group */
+	FActiveSound* ResolveConcurrency(const FActiveSound& NewActiveSound, FConcurrencyGroup& ConcurrencyGroup);
+
+	FActiveSound* MakeNewActiveSound(const FActiveSound& NewActiveSound);
+
+	FConcurrencyGroupID MakeNewConcurrencyGroupAndSound(FActiveSound** OutActiveSound, const FActiveSound& NewActiveSound);
+
+private: // Data
+
+	/** Owning audio device ptr for the concurrency manager. */
+	FAudioDevice* AudioDevice;
+
+	/** Global concurrency map that maps individual sounds instances to shared USoundConcurrency UObjects. */
+	FConcurrencyMap ConcurrencyMap;
+
+	FOwnerConcurrencyMap OwnerConcurrencyMap;
+
+	/** A map of owners to concurrency maps for sounds which are concurrency-limited per sound owner. */
+	FOwnerPerSoundConcurrencyMap OwnerPerSoundConcurrencyMap;
+
+	FPerSoundToActiveSoundsMap SoundObjectToActiveSounds;
+
+	/** A map of concurrency active sound ID to concurrency active sounds */
+	TMap<FConcurrencyGroupID, FConcurrencyGroup> ConcurrencyGroups;
+
+};
\ No newline at end of file
diff --git a/Engine/Source/Runtime/Engine/Private/ActiveSound.cpp b/Engine/Source/Runtime/Engine/Private/ActiveSound.cpp
--- a/Engine/Source/Runtime/Engine/Private/ActiveSound.cpp
+++ b/Engine/Source/Runtime/Engine/Private/ActiveSound.cpp
@@ -9,17 +9,20 @@
 #include "SubtitleManager.h"
 
 FActiveSound::FActiveSound()
-	: Sound(NULL)
-	, World(NULL)
-	, AudioComponent(NULL)
-	, AudioComponentIndex(0)
-	, SoundClassOverride(NULL)
+	: Sound(nullptr)
+	, World(nullptr)
+	, AudioComponent(nullptr)
+	, AudioDevice(nullptr)
+	, ConcurrencyGroupID(0)
+	, ConcurrencySettings(nullptr)
+	, SoundClassOverride(nullptr)
 	, bOccluded(false)
 	, bAllowSpatialization(true)
 	, bHasAttenuationSettings(false)
 	, bShouldRemainActiveIfDropped(false)
 	, bFadingOut(false)
 	, bFinished(false)
+	, bShouldStopDueToMaxConcurrency(false)
 	, bRadioFilterSelected(false)
 	, bApplyRadioFilter(false)
 	, bHandleSubtitles(true)
@@ -44,12 +47,14 @@ FActiveSound::FActiveSound()
 	, VolumeMultiplier(1.f)
 	, PitchMultiplier(1.f)
 	, HighFrequencyGainMultiplier(1.f)
+	, ConcurrencyVolumeScale(1.f)
 	, SubtitlePriority(0.f)
 	, VolumeWeightedPriorityScale(1.0f)
+	, VolumeConcurrency(0.0f)
 	, OcclusionCheckInterval(0.f)
 	, LastOcclusionCheckTime(0.f)
 	, LastLocation(FVector::ZeroVector)
-	, LastAudioVolume(NULL)
+	, LastAudioVolume(nullptr)
 	, LastUpdateTime(0.f)
 	, SourceInteriorVolume(1.f)
 	, SourceInteriorLPF(1.f)
@@ -124,6 +129,46 @@ int32 FActiveSound::FindClosestListener( const TArray<FListener>& InListeners )
 	return ClosestListenerIndex;
 }
 
+uint32 FActiveSound::TryGetOwnerID() const
+{
+	// Only have an owner of the active sound if the audio component is valid and if the audio component has an owner.
+	if (UAudioComponent* AudioComponentPtr = AudioComponent.Get())
+	{
+		AActor* Owner = AudioComponentPtr->GetOwner();
+		if (Owner)
+		{
+			return Owner->GetUniqueID();
+		}
+	}
+
+	return 0;
+}
+
+const FSoundConcurrencySettings* FActiveSound::GetSoundConcurrencySettingsToApply() const
+{
+	if (ConcurrencySettings)
+	{
+		return &ConcurrencySettings->Concurrency;
+	}
+	else if (Sound)
+	{
+		return Sound->GetSoundConcurrencySettingsToApply();
+	}
+	return nullptr;
+}
+
+uint32 FActiveSound::GetSoundConcurrencyObjectID() const
+{
+	if (ConcurrencySettings)
+	{
+		return ConcurrencySettings->GetUniqueID();
+	}
+	else
+	{
+		return Sound->GetSoundConcurrencyObjectID();
+	}
+}
+
 
 void FActiveSound::UpdateWaveInstances( FAudioDevice* AudioDevice, TArray<FWaveInstance*> &InWaveInstances, const float DeltaTime )
 {
@@ -175,7 +220,7 @@ void FActiveSound::UpdateWaveInstances( FAudioDevice* AudioDevice, TArray<FWaveI
 	// (even after the Sound has started playing, and this line takes them all into account and gives us
 	// final value that is correct
 	UpdateAdjustVolumeMultiplier(DeltaTime);
-	ParseParams.VolumeMultiplier = VolumeMultiplier * Sound->GetVolumeMultiplier() * CurrentAdjustVolumeMultiplier * AudioDevice->TransientMasterVolume * (bOccluded ? 0.5f : 1.0f);
+	ParseParams.VolumeMultiplier = VolumeMultiplier * Sound->GetVolumeMultiplier() * CurrentAdjustVolumeMultiplier * AudioDevice->TransientMasterVolume * (bOccluded ? 0.5f : 1.0f) * ConcurrencyVolumeScale;
 	ParseParams.VolumeWeightedPriorityScale = VolumeWeightedPriorityScale;
 	ParseParams.Pitch *= PitchMultiplier * Sound->GetPitchMultiplier();
 	ParseParams.HighFrequencyGain *= HighFrequencyGainMultiplier;
@@ -201,13 +246,15 @@ void FActiveSound::UpdateWaveInstances( FAudioDevice* AudioDevice, TArray<FWaveI
 		ParseParams.Transform = ParseParams.Transform * ClosestListener.Transform.Inverse() * Listener.Transform;
 	}
 
+	TArray<FWaveInstance*> ThisSoundsWaveInstances;
+
 	// Recurse nodes, have SoundWave's create new wave instances and update bFinished unless we finished fading out.
 	bFinished = true;
-	if( !bFadingOut || ( PlaybackTime <= TargetAdjustVolumeStopTime ) )
+	if (!bFadingOut || (PlaybackTime <= TargetAdjustVolumeStopTime))
 	{
 		if (bHasAttenuationSettings)
 		{
-			AttenuationSettings.ApplyAttenuation(ParseParams.Transform, Listener.Transform.GetTranslation(), ParseParams.Volume, ParseParams.HighFrequencyGain );
+			AttenuationSettings.ApplyAttenuation(ParseParams.Transform, Listener.Transform.GetTranslation(), ParseParams.Volume, ParseParams.HighFrequencyGain);
 			ParseParams.OmniRadius = AttenuationSettings.OmniRadius;
 			ParseParams.StereoSpread = AttenuationSettings.StereoSpread;
 			ParseParams.bUseSpatialization = AttenuationSettings.bSpatialize;
@@ -221,16 +268,38 @@ void FActiveSound::UpdateWaveInstances( FAudioDevice* AudioDevice, TArray<FWaveI
 			}
 		}
 
-		Sound->Parse( AudioDevice, 0, *this, ParseParams, InWaveInstances );
+
+		Sound->Parse(AudioDevice, 0, *this, ParseParams, ThisSoundsWaveInstances);
 	}
 
-	if( bFinished )
+	if (bFinished)
+	{
+		Stop();
+	}
+	else if (ThisSoundsWaveInstances.Num() > 0)
 	{
-		Stop(AudioDevice);
+		// If this active sound is told to limit concurrency by the quietest sound
+		const FSoundConcurrencySettings* ConcurrencySettings = GetSoundConcurrencySettingsToApply();
+		if (ConcurrencySettings && ConcurrencySettings->ResolutionRule == EMaxConcurrentResolutionRule::StopQuietist)
+		{
+			check(ConcurrencyGroupID != 0);
+			// Now that we have this sound's active wave instances, lets find the loudest active wave instance to represent the "volume" of this active sound
+			VolumeConcurrency = 0.0f;
+			for (const FWaveInstance* WaveInstance : ThisSoundsWaveInstances)
+			{
+				const float WaveInstanceVolume = WaveInstance->GetActualVolume();
+				if (WaveInstanceVolume > VolumeConcurrency)
+				{
+					VolumeConcurrency = WaveInstanceVolume;
+				}
+			}
+		}
 	}
+
+	InWaveInstances.Append(ThisSoundsWaveInstances);
 }
 
-void FActiveSound::Stop(FAudioDevice* AudioDevice)
+void FActiveSound::Stop()
 {
 	check(AudioDevice);
 
@@ -261,13 +330,9 @@ void FActiveSound::Stop(FAudioDevice* AudioDevice)
 	}
 	WaveInstances.Empty();
 
-	// TODO - Audio Threading. This call would be a task back to game thread
-	if (AudioComponent.IsValid())
-	{
-		AudioComponent->PlaybackCompleted(false);
-	}
-
 	AudioDevice->RemoveActiveSound(this);
+
+	delete this;
 }
 
 FWaveInstance* FActiveSound::FindWaveInstance( const UPTRINT WaveInstanceHash )
diff --git a/Engine/Source/Runtime/Engine/Private/Audio.cpp b/Engine/Source/Runtime/Engine/Private/Audio.cpp
--- a/Engine/Source/Runtime/Engine/Private/Audio.cpp
+++ b/Engine/Source/Runtime/Engine/Private/Audio.cpp
@@ -44,6 +44,7 @@ DEFINE_STAT(STAT_AudioPrepareDecompressionTime);
 DEFINE_STAT(STAT_OpusDecompressTime);
 
 DEFINE_STAT(STAT_AudioUpdateEffects);
+DEFINE_STAT(STAT_AudioEvaluateConcurrency);
 DEFINE_STAT(STAT_AudioUpdateSources);
 DEFINE_STAT(STAT_AudioResourceCreationTime);
 DEFINE_STAT(STAT_AudioSourceInitTime);
@@ -464,7 +465,10 @@ float FWaveInstance::GetActualVolume() const
 
 float FWaveInstance::GetVolumeWeightedPriority() const
 {
-	return GetActualVolume() * VolumeWeightedPriorityScale;
+	check(ActiveSound);
+	// If this wave instance's active sound should stop due to max concurrency, return a large negative weighted priority so 
+	// that it will always be sorted to the bottom of the WaveInstance list and stopped.
+	return ActiveSound->bShouldStopDueToMaxConcurrency ? -100.0f : GetActualVolume() * VolumeWeightedPriorityScale;
 }
 
 bool FWaveInstance::IsStreaming() const
diff --git a/Engine/Source/Runtime/Engine/Private/AudioDevice.cpp b/Engine/Source/Runtime/Engine/Private/AudioDevice.cpp
--- a/Engine/Source/Runtime/Engine/Private/AudioDevice.cpp
+++ b/Engine/Source/Runtime/Engine/Private/AudioDevice.cpp
@@ -36,6 +36,9 @@ void FDynamicParameter::Set(float Value, float InDuration)
 	}
 	else
 	{
+		StartValue = Value;
+		DeltaValue = 0.0f;
+		DurationSec = 0.0f;
 		CurrValue = Value;
 	}
 }
@@ -95,6 +98,7 @@ FAudioDevice::FAudioDevice()
 	, bHRTFEnabledForAll(false)
 	, bIsDeviceMuted(false)
 	, bIsInitialized(false)
+	, ConcurrencyManager(this)
 {
 }
 
@@ -1779,51 +1783,63 @@ int32 FAudioDevice::GetSortedActiveWaveInstances(TArray<FWaveInstance*>& WaveIns
 			continue;
 		}
 		
-		if( !ActiveSound->Sound )
+		if (!ActiveSound->Sound)
 		{
 			// No sound - cleanup and remove
-			ActiveSound->Stop(this);
+			ActiveSound->Stop();
 		}
 		// If the world scene allows audio - tick wave instances.
-		else if( ActiveSound->World == NULL || ActiveSound->World->AllowAudioPlayback() )
+		else 
 		{
-#if !(UE_BUILD_SHIPPING || UE_BUILD_TEST)
-			if ( !ensureMsgf(ActiveSound->Sound->IsValidLowLevel(), TEXT("ActiveSound with INVALID sound. AudioComponent=%s. DebugOriginalSoundName=%s"),
-				 ActiveSound->GetAudioComponent() ? *ActiveSound->GetAudioComponent()->GetPathName() : TEXT("NO COMPONENT"),
-				 *ActiveSound->DebugOriginalSoundName.ToString() ))
-			{
-				// Sound was not valid, stop playing it.
-				ActiveSound->Stop(this);
-			}
-			else
-#endif
+			UWorld* ActiveSoundWorldPtr = ActiveSound->World.Get();
+			if (ActiveSoundWorldPtr == nullptr || ActiveSoundWorldPtr->AllowAudioPlayback())
 			{
-				const float Duration = ActiveSound->Sound->GetDuration();
-				// Divide by minimum pitch for longest possible duration
-				if( Duration < INDEFINITELY_LOOPING_DURATION && ActiveSound->PlaybackTime > Duration / MIN_PITCH )
+
+#if !(UE_BUILD_SHIPPING || UE_BUILD_TEST)
+				if (!ensureMsgf(ActiveSound->Sound->IsValidLowLevel(), TEXT("ActiveSound with INVALID sound. AudioComponent=%s. DebugOriginalSoundName=%s"),
+					 ActiveSound->GetAudioComponent() ? *ActiveSound->GetAudioComponent()->GetPathName() : TEXT("NO COMPONENT"),
+					 *ActiveSound->DebugOriginalSoundName.ToString()))
 				{
-					UE_LOG(LogAudio, Log, TEXT( "Sound stopped due to duration: %g > %g : %s %s" ), 
-						ActiveSound->PlaybackTime, 
-						Duration, 
-						*ActiveSound->Sound->GetName(), 
-						(ActiveSound->GetAudioComponent() ? *ActiveSound->GetAudioComponent()->GetName() : TEXT("NO COMPONENT")));
-					ActiveSound->Stop(this);
+					// Sound was not valid, stop playing it.
+					ActiveSound->Stop();
 				}
 				else
+#endif
 				{
-					// If not in game, do not advance sounds unless they are UI sounds.
-					float UsedDeltaTime = FApp::GetDeltaTime();
-					if (GetType == ESortedActiveWaveGetType::QueryOnly || (GetType == ESortedActiveWaveGetType::PausedUpdate && !ActiveSound->bIsUISound))
+					const float Duration = ActiveSound->Sound->GetDuration();
+
+					// Divide by minimum pitch for longest possible duration
+					if (Duration < INDEFINITELY_LOOPING_DURATION && ActiveSound->PlaybackTime > Duration / MIN_PITCH)
 					{
-						UsedDeltaTime = 0.0f;
+						UE_LOG(LogAudio, Log, TEXT("Sound stopped due to duration: %g > %g : %s %s"),
+							ActiveSound->PlaybackTime,
+							Duration,
+							*ActiveSound->Sound->GetName(),
+							(ActiveSound->GetAudioComponent() ? *ActiveSound->GetAudioComponent()->GetName() : TEXT("NO COMPONENT")));
+						ActiveSound->Stop();
+					}
+					else
+					{
+						// If not in game, do not advance sounds unless they are UI sounds.
+						float UsedDeltaTime = FApp::GetDeltaTime();
+						if (GetType == ESortedActiveWaveGetType::QueryOnly || (GetType == ESortedActiveWaveGetType::PausedUpdate && !ActiveSound->bIsUISound))
+						{
+							UsedDeltaTime = 0.0f;
+						}
+
+						ActiveSound->UpdateWaveInstances(this, WaveInstances, UsedDeltaTime);
 					}
-
-					ActiveSound->UpdateWaveInstances( this, WaveInstances, UsedDeltaTime );
 				}
 			}
 		}
 	}
 
+	// Now stop any sounds that are active that are in concurrency resolution groups that resolve by stopping quietest
+	{
+		SCOPE_CYCLE_COUNTER(STAT_AudioEvaluateConcurrency);
+		ConcurrencyManager.StopQuietSoundsDueToMaxConcurrency();
+	}
+
 	// Helper function for "Sort" (higher priority sorts last).
 	struct FCompareFWaveInstanceByPlayPriority
 	{
@@ -2014,6 +2030,13 @@ void FAudioDevice::Update( bool bGameTicking )
 		Listener.UpdateCurrentInteriorSettings();
 	}
 
+	// Stop any pending active sounds that need to be stopped
+	while (PendingSoundsToStop.Num() > 0)
+	{
+		FActiveSound* SoundToStop = PendingSoundsToStop.Pop();
+		SoundToStop->Stop();
+	}
+
 	if (Sources.Num())
 	{
 		// Kill any sources that have finished
@@ -2072,7 +2095,7 @@ void FAudioDevice::StopAllSounds( bool bShouldStopUISounds )
 
 		if (bShouldStopUISounds)
 		{
-			ActiveSound->Stop(this);
+			ActiveSound->Stop();
 		}
 		// If we're allowing UI sounds to continue then first filter on the active sounds state
 		else if (!ActiveSound->bIsUISound)
@@ -2085,7 +2108,7 @@ void FAudioDevice::StopAllSounds( bool bShouldStopUISounds )
 				FWaveInstance* WaveInstance = WaveInstanceIt.Value();
 				if (WaveInstance && !WaveInstance->bIsUISound)
 				{
-					ActiveSound->Stop(this);
+					ActiveSound->Stop();
 					break;
 				}
 			}
@@ -2093,84 +2116,23 @@ void FAudioDevice::StopAllSounds( bool bShouldStopUISounds )
 	}
 }
 
-void FAudioDevice::AddNewActiveSound( const FActiveSound& NewActiveSound )
+void FAudioDevice::AddNewActiveSound(const FActiveSound& NewActiveSound)
 {
 	check(NewActiveSound.Sound);
 
-	// if we have gone over the MaxConcurrentPlay
-	// TODO: Consider if we should handle the bShouldRemainActiveIfDropped case
-	if( ( NewActiveSound.Sound->MaxConcurrentPlayCount != 0 ) && ( NewActiveSound.Sound->CurrentPlayCount >= NewActiveSound.Sound->MaxConcurrentPlayCount ) )
-	{
-		FActiveSound* SoundToStop = NULL;
-
-		switch (NewActiveSound.Sound->MaxConcurrentResolutionRule)
-		{
-		case EMaxConcurrentResolutionRule::StopOldest:
-
-			for (int32 SoundIndex = 0; SoundIndex < ActiveSounds.Num(); ++SoundIndex)
-			{
-				FActiveSound* ActiveSound = ActiveSounds[SoundIndex];
-				if (ActiveSound->Sound == NewActiveSound.Sound && (SoundToStop == NULL || ActiveSound->PlaybackTime > SoundToStop->PlaybackTime))
-				{
-					SoundToStop = ActiveSound;
-				}
-			}
-			break;
+	// Evaluate concurrency. This will create an ActiveSound ptr which is a copy of NewActiveSound if the sound can play.
+	FActiveSound* ActiveSound = nullptr;
 
-		case EMaxConcurrentResolutionRule::StopFarthestThenOldest:
-		case EMaxConcurrentResolutionRule::StopFarthestThenPreventNew:
-		{
-			int32 ClosestListenerIndex = NewActiveSound.FindClosestListener(Listeners);
-			float DistanceToStopSoundSq = FVector::DistSquared(Listeners[ClosestListenerIndex].Transform.GetTranslation(), NewActiveSound.Transform.GetTranslation());
-
-			for (int32 SoundIndex = 0; SoundIndex < ActiveSounds.Num(); ++SoundIndex)
-			{
-				FActiveSound* ActiveSound = ActiveSounds[SoundIndex];
-				if (ActiveSound->Sound == NewActiveSound.Sound)
-				{
-					ClosestListenerIndex = ActiveSound->FindClosestListener(Listeners);
-					const float DistanceToActiveSoundSq = FVector::DistSquared(Listeners[ClosestListenerIndex].Transform.GetTranslation(), ActiveSound->Transform.GetTranslation());
-
-					if (DistanceToActiveSoundSq > DistanceToStopSoundSq)
-					{
-						SoundToStop = ActiveSound;
-						DistanceToStopSoundSq = DistanceToActiveSoundSq;
-					}
-					else if (   NewActiveSound.Sound->MaxConcurrentResolutionRule == EMaxConcurrentResolutionRule::StopFarthestThenOldest 
-							 && DistanceToActiveSoundSq == DistanceToStopSoundSq
-							 && (SoundToStop == NULL || ActiveSound->PlaybackTime > SoundToStop->PlaybackTime))
-					{
-						SoundToStop = ActiveSound;
-						DistanceToStopSoundSq = DistanceToActiveSoundSq;
-					}
-				}
-			}
-			break;
-		}
-
-		case EMaxConcurrentResolutionRule::PreventNew:
-		default:
-			break;
-
-		}
+	{
+		SCOPE_CYCLE_COUNTER(STAT_AudioEvaluateConcurrency);
 
-		// If we found a sound to stop, then stop it and allow the new sound to play.  
-		// Otherwise inform the system that the sound failed to start.
-		if (SoundToStop)
-		{
-			SoundToStop->Stop(this);
-		}
-		else
-		{
-			// TODO - Audio Threading. This call would be a task back to game thread
-			if (UAudioComponent* AudioComponent = NewActiveSound.GetAudioComponent())
-			{
-				AudioComponent->PlaybackCompleted(true);
-			}
-			//UE_LOG(LogAudio, Verbose, TEXT( "   %g: MaxConcurrentPlayCount AudioComponent : '%s' with Sound: '%s' Max: %d   Curr: %d " ), NewActiveSound.World ? NewActiveSound.World->GetAudioTimeSeconds() : 0.0f, *GetFullName(), Sound ? *Sound->GetName() : TEXT( "NULL" ), Sound->MaxConcurrentPlayCount, Sound->CurrentPlayCount );
-			return;
+		// Try to create a new active sound. This returns nullptr if too many sounds are playing with this sound's concurrency setting
+		ActiveSound = ConcurrencyManager.CreateNewActiveSound(NewActiveSound);
+	}
 
-		}
+	if (!ActiveSound)
+	{
+		return;
 	}
 
 	++NewActiveSound.Sound->CurrentPlayCount;
@@ -2181,7 +2143,8 @@ void FAudioDevice::AddNewActiveSound( const FActiveSound& NewActiveSound )
 	UE_LOG(LogAudio, VeryVerbose, TEXT("New ActiveSound %s Comp: %s Loc: %s"), *NewActiveSound.Sound->GetName(), (AudioComponent ? *AudioComponent->GetFullName() : TEXT("No AudioComponent")), *NewActiveSound.Transform.GetTranslation().ToString() );
 #endif // !(UE_BUILD_SHIPPING || UE_BUILD_TEST)
 
-	FActiveSound* ActiveSound = new FActiveSound(NewActiveSound);
+	check(ActiveSound);
+
 #if !(UE_BUILD_SHIPPING || UE_BUILD_TEST)
 	if (ActiveSound->Sound)
 	{
@@ -2206,11 +2169,18 @@ void FAudioDevice::AddNewActiveSound( const FActiveSound& NewActiveSound )
 		}
 	}
 #endif
+
 	ActiveSounds.Add(ActiveSound);
 	if (AudioComponent)
 	{
 		AudioComponentToActiveSoundMap.Add((UPTRINT)AudioComponent, ActiveSound);
 	}
+
+}
+
+void FAudioDevice::AddSoundToStop(FActiveSound* SoundToStop)
+{
+	PendingSoundsToStop.Add(SoundToStop);
 }
 
 void FAudioDevice::StopActiveSound( UAudioComponent* AudioComponent )
@@ -2220,7 +2190,7 @@ void FAudioDevice::StopActiveSound( UAudioComponent* AudioComponent )
 	FActiveSound* ActiveSound = FindActiveSound(AudioComponent);
 	if (ActiveSound)
 	{
-		ActiveSound->Stop(this);
+		ActiveSound->Stop();
 	}
 }
 
@@ -2237,20 +2207,23 @@ FActiveSound* FAudioDevice::FindActiveSound( UAudioComponent* AudioComponent )
 
 void FAudioDevice::RemoveActiveSound(FActiveSound* ActiveSound)
 {
-	for( int32 Index = 0; Index < ActiveSounds.Num(); ++Index )
+	ConcurrencyManager.RemoveActiveSound(ActiveSound);
+
+	// Perform the notification
+	if (UAudioComponent* AudioComponent = ActiveSound->GetAudioComponent())
 	{
-		if (ActiveSound == ActiveSounds[Index])
-		{
-			const UPTRINT AudioComponentIndex = ActiveSound->GetAudioComponentIndex();
-			ActiveSounds.RemoveAt(Index);
-			if (AudioComponentIndex > 0)
-			{
-				AudioComponentToActiveSoundMap.Remove(AudioComponentIndex);
-			}
-			delete ActiveSound;
-			break;
-		}
+		AudioComponent->PlaybackCompleted(false);
 	}
+
+	int32 NumRemoved = ActiveSounds.Remove(ActiveSound);
+	check(NumRemoved == 1);
+
+	const UPTRINT AudioComponentIndex = ActiveSound->GetAudioComponentIndex();
+	if (AudioComponentIndex > 0)
+	{
+		AudioComponentToActiveSoundMap.Remove(AudioComponentIndex);
+	}
+
 }
 
 bool FAudioDevice::LocationIsAudible( FVector Location, float MaxDistance )
@@ -2272,7 +2245,7 @@ bool FAudioDevice::LocationIsAudible( FVector Location, float MaxDistance )
 	return( false );
 }
 
-UAudioComponent* FAudioDevice::CreateComponent(USoundBase* Sound, UWorld* World, AActor* Actor, bool bPlay, bool bStopWhenOwnerDestroyed, const FVector* Location, USoundAttenuation* AttenuationSettings)
+UAudioComponent* FAudioDevice::CreateComponent(USoundBase* Sound, UWorld* World, AActor* Actor, bool bPlay, bool bStopWhenOwnerDestroyed, const FVector* Location, USoundAttenuation* AttenuationSettings, USoundConcurrency* ConcurrencySettings)
 {
 	UAudioComponent* AudioComponent = NULL;
 
@@ -2327,6 +2300,8 @@ UAudioComponent* FAudioDevice::CreateComponent(USoundBase* Sound, UWorld* World,
 			AudioComponent->bVisualizeComponent	= false;
 #endif
 			AudioComponent->AttenuationSettings = AttenuationSettings;
+			AudioComponent->ConcurrencySettings = ConcurrencySettings;
+
 			if (Location)
 			{
 				AudioComponent->SetWorldLocation(*Location);
@@ -2366,14 +2341,14 @@ void FAudioDevice::Flush( UWorld* WorldToFlush, bool bClearActivatedReverb )
 		{
 			if (WorldToFlush == nullptr)
 			{
-				ActiveSound->Stop(this);
+				ActiveSound->Stop();
 			}
 			else
 			{
 				UWorld* ActiveSoundWorld = ActiveSound->World.Get();
 				if (ActiveSoundWorld == nullptr || ActiveSoundWorld == WorldToFlush)
 				{
-					ActiveSound->Stop(this);
+					ActiveSound->Stop();
 				}
 			}
 		}
@@ -2662,7 +2637,7 @@ void FAudioDevice::StopSoundsUsingResource(USoundWave* SoundWave, TArray<UAudioC
 				{
 					StoppedComponents.Add(AudioComponent);
 				}
-				ActiveSound->Stop(this);
+				ActiveSound->Stop();
 				bStoppedSounds = true;
 				break;
 			}
diff --git a/Engine/Source/Runtime/Engine/Private/Components/AudioComponent.cpp b/Engine/Source/Runtime/Engine/Private/Components/AudioComponent.cpp
--- a/Engine/Source/Runtime/Engine/Private/Components/AudioComponent.cpp
+++ b/Engine/Source/Runtime/Engine/Private/Components/AudioComponent.cpp
@@ -143,6 +143,7 @@ void UAudioComponent::PlayInternal(const float StartTime, const float FadeInDura
 			NewActiveSound.World = GetWorld();
 			NewActiveSound.Sound = Sound;
 			NewActiveSound.SoundClassOverride = SoundClassOverride;
+			NewActiveSound.ConcurrencySettings = ConcurrencySettings;
 
 			NewActiveSound.VolumeMultiplier = (VolumeModulationMax + ((VolumeModulationMin - VolumeModulationMax) * FMath::SRand())) * VolumeMultiplier;
 			NewActiveSound.VolumeWeightedPriorityScale = VolumeWeightedPriorityScale;
@@ -191,6 +192,9 @@ void UAudioComponent::PlayInternal(const float StartTime, const float FadeInDura
 				NewActiveSound.CurrentAdjustVolumeMultiplier = FadeVolumeLevel;
 			}
 
+			FName SoundName = NewActiveSound.Sound->GetFName();
+			UE_LOG(LogTemp, Log, TEXT("Playing ActiveSound: %s"), *SoundName.GetPlainNameString());
+
 			// TODO - Audio Threading. This call would be a task call to dispatch to the audio thread
 			AudioDevice->AddNewActiveSound(NewActiveSound);
 
diff --git a/Engine/Source/Runtime/Engine/Private/GameplayStatics.cpp b/Engine/Source/Runtime/Engine/Private/GameplayStatics.cpp
--- a/Engine/Source/Runtime/Engine/Private/GameplayStatics.cpp
+++ b/Engine/Source/Runtime/Engine/Private/GameplayStatics.cpp
@@ -756,7 +756,7 @@ void UGameplayStatics::SetGlobalPitchModulation(UObject* WorldContextObject, flo
 }
 
 
-void UGameplayStatics::PlaySound2D(UObject* WorldContextObject, class USoundBase* Sound, float VolumeMultiplier, float PitchMultiplier, float StartTime)
+void UGameplayStatics::PlaySound2D(UObject* WorldContextObject, class USoundBase* Sound, float VolumeMultiplier, float PitchMultiplier, float StartTime, class USoundConcurrency* ConcurrencySettings)
 {
 	if (!Sound || !GEngine || !GEngine->UseSound())
 	{
@@ -781,12 +781,13 @@ void UGameplayStatics::PlaySound2D(UObject* WorldContextObject, class USoundBase
 		NewActiveSound.RequestedStartTime = FMath::Max(0.f, StartTime);
 
 		NewActiveSound.bIsUISound = true;
+		NewActiveSound.ConcurrencySettings = ConcurrencySettings;
 
 		AudioDevice->AddNewActiveSound(NewActiveSound);
 	}
 }
 
-UAudioComponent* UGameplayStatics::SpawnSound2D(UObject* WorldContextObject, class USoundBase* Sound, float VolumeMultiplier, float PitchMultiplier, float StartTime)
+UAudioComponent* UGameplayStatics::SpawnSound2D(UObject* WorldContextObject, class USoundBase* Sound, float VolumeMultiplier, float PitchMultiplier, float StartTime, class USoundConcurrency* ConcurrencySettings)
 {
 	if (!Sound || !GEngine || !GEngine->UseSound())
 	{
@@ -799,11 +800,10 @@ UAudioComponent* UGameplayStatics::SpawnSound2D(UObject* WorldContextObject, cla
 		return nullptr;
 	}
 
-	UAudioComponent* AudioComponent = FAudioDevice::CreateComponent(Sound, ThisWorld, ThisWorld->GetWorldSettings(), false);
+	UAudioComponent* AudioComponent = FAudioDevice::CreateComponent(Sound, ThisWorld, ThisWorld->GetWorldSettings(), false, false, nullptr, nullptr, ConcurrencySettings);
 	if (AudioComponent)
 	{
 		const bool bIsInGameWorld = AudioComponent->GetWorld()->IsGameWorld();
-
 		AudioComponent->SetVolumeMultiplier(VolumeMultiplier);
 		AudioComponent->SetPitchMultiplier(PitchMultiplier);
 		AudioComponent->bAllowSpatialization	= false;
@@ -815,7 +815,7 @@ UAudioComponent* UGameplayStatics::SpawnSound2D(UObject* WorldContextObject, cla
 	return AudioComponent;
 }
 
-void UGameplayStatics::PlaySoundAtLocation(UObject* WorldContextObject, class USoundBase* Sound, FVector Location, FRotator Rotation, float VolumeMultiplier, float PitchMultiplier, float StartTime, class USoundAttenuation* AttenuationSettings)
+void UGameplayStatics::PlaySoundAtLocation(UObject* WorldContextObject, class USoundBase* Sound, FVector Location, FRotator Rotation, float VolumeMultiplier, float PitchMultiplier, float StartTime, class USoundAttenuation* AttenuationSettings, class USoundConcurrency* ConcurrencySettings)
 {
 	if (!Sound || !GEngine || !GEngine->UseSound())
 	{
@@ -859,6 +859,8 @@ void UGameplayStatics::PlaySoundAtLocation(UObject* WorldContextObject, class US
 				NewActiveSound.AttenuationSettings = *AttenuationSettingsToApply;
 			}
 
+			NewActiveSound.ConcurrencySettings = ConcurrencySettings;
+
 			// TODO - Audio Threading. This call would be a task call to dispatch to the audio thread
 			AudioDevice->AddNewActiveSound(NewActiveSound);
 		}
@@ -870,7 +872,7 @@ void UGameplayStatics::PlaySoundAtLocation(UObject* WorldContextObject, class US
 	}
 }
 
-UAudioComponent* UGameplayStatics::SpawnSoundAtLocation(UObject* WorldContextObject, class USoundBase* Sound, FVector Location, FRotator Rotation, float VolumeMultiplier, float PitchMultiplier, float StartTime, class USoundAttenuation* AttenuationSettings)
+UAudioComponent* UGameplayStatics::SpawnSoundAtLocation(UObject* WorldContextObject, class USoundBase* Sound, FVector Location, FRotator Rotation, float VolumeMultiplier, float PitchMultiplier, float StartTime, class USoundAttenuation* AttenuationSettings, class USoundConcurrency* ConcurrencySettings)
 {
 	if (!Sound || !GEngine || !GEngine->UseSound())
 	{
@@ -885,7 +887,7 @@ UAudioComponent* UGameplayStatics::SpawnSoundAtLocation(UObject* WorldContextObj
 
 	const bool bIsInGameWorld = ThisWorld->IsGameWorld();
 
-	UAudioComponent* AudioComponent = FAudioDevice::CreateComponent(Sound, ThisWorld, ThisWorld->GetWorldSettings(), false, false, &Location, AttenuationSettings);
+	UAudioComponent* AudioComponent = FAudioDevice::CreateComponent(Sound, ThisWorld, ThisWorld->GetWorldSettings(), false, false, &Location, AttenuationSettings, ConcurrencySettings);
 
 	if (AudioComponent)
 	{
@@ -902,7 +904,7 @@ UAudioComponent* UGameplayStatics::SpawnSoundAtLocation(UObject* WorldContextObj
 	return AudioComponent;
 }
 
-class UAudioComponent* UGameplayStatics::SpawnSoundAttached(class USoundBase* Sound, class USceneComponent* AttachToComponent, FName AttachPointName, FVector Location, FRotator Rotation, EAttachLocation::Type LocationType, bool bStopWhenAttachedToDestroyed, float VolumeMultiplier, float PitchMultiplier, float StartTime, class USoundAttenuation* AttenuationSettings)
+class UAudioComponent* UGameplayStatics::SpawnSoundAttached(class USoundBase* Sound, class USceneComponent* AttachToComponent, FName AttachPointName, FVector Location, FRotator Rotation, EAttachLocation::Type LocationType, bool bStopWhenAttachedToDestroyed, float VolumeMultiplier, float PitchMultiplier, float StartTime, class USoundAttenuation* AttenuationSettings, class USoundConcurrency* ConcurrencySettings)
 {
 	if (!Sound)
 	{
@@ -922,7 +924,7 @@ class UAudioComponent* UGameplayStatics::SpawnSoundAttached(class USoundBase* So
 		TestLocation = AttachToComponent->GetRelativeTransform().TransformPosition(Location);
 	}
 
-	UAudioComponent* AudioComponent = FAudioDevice::CreateComponent(Sound, AttachToComponent->GetWorld(), AttachToComponent->GetOwner(), false, bStopWhenAttachedToDestroyed, &TestLocation, AttenuationSettings);
+	UAudioComponent* AudioComponent = FAudioDevice::CreateComponent(Sound, AttachToComponent->GetWorld(), AttachToComponent->GetOwner(), false, bStopWhenAttachedToDestroyed, &TestLocation, AttenuationSettings, ConcurrencySettings);
 	if (AudioComponent)
 	{
 		const bool bIsInGameWorld = AudioComponent->GetWorld()->IsGameWorld();
diff --git a/Engine/Source/Runtime/Engine/Private/SoundBase.cpp b/Engine/Source/Runtime/Engine/Private/SoundBase.cpp
--- a/Engine/Source/Runtime/Engine/Private/SoundBase.cpp
+++ b/Engine/Source/Runtime/Engine/Private/SoundBase.cpp
@@ -8,7 +8,7 @@
 USoundBase::USoundBase(const FObjectInitializer& ObjectInitializer)
 	: Super(ObjectInitializer)
 {
-	MaxConcurrentPlayCount = 16;
+	MaxConcurrentPlayCount_DEPRECATED = 16;
 }
 
 void USoundBase::PostInitProperties()
@@ -106,3 +106,42 @@ USoundClass* USoundBase::GetSoundClass() const
 {
 	return SoundClassObject;
 }
+
+const FSoundConcurrencySettings* USoundBase::GetSoundConcurrencySettingsToApply()
+{
+	if (bOverrideConcurrency)
+	{
+		return &ConcurrencyOverrides;
+	}
+	else if (SoundConcurrencySettings)
+	{
+		return &SoundConcurrencySettings->Concurrency;
+	}
+	return nullptr;
+}
+
+uint32 USoundBase::GetSoundConcurrencyObjectID() const
+{
+	if (SoundConcurrencySettings != nullptr && !bOverrideConcurrency)
+	{
+		return SoundConcurrencySettings->GetUniqueID();
+	}
+	return 0;
+}
+
+void USoundBase::PostLoad()
+{
+	Super::PostLoad();
+
+	const int32 LinkerUE4Version = GetLinkerUE4Version();
+
+	if (LinkerUE4Version < VER_UE4_SOUND_CONCURRENCY_PACKAGE)
+	{
+		bOverrideConcurrency = true;
+		ConcurrencyOverrides.bLimitToOwner = false;
+		ConcurrencyOverrides.MaxCount = MaxConcurrentPlayCount_DEPRECATED;
+		ConcurrencyOverrides.ResolutionRule = MaxConcurrentResolutionRule_DEPRECATED;
+		ConcurrencyOverrides.VolumeScale = 1.0f;
+	}
+}
+
diff --git a/Engine/Source/Runtime/Engine/Private/SoundConcurrency.cpp b/Engine/Source/Runtime/Engine/Private/SoundConcurrency.cpp
new file mode 100644
--- /dev/null
+++ b/Engine/Source/Runtime/Engine/Private/SoundConcurrency.cpp
@@ -0,0 +1,505 @@
+// Copyright 1998-2015 Epic Games, Inc. All Rights Reserved.
+
+#include "EnginePrivate.h"
+#include "SoundDefinitions.h"
+#include "Sound/SoundConcurrency.h"
+#include "Sound/SoundBase.h"
+#include "ActiveSound.h"
+
+/************************************************************************/
+/* USoundConcurrency													*/
+/************************************************************************/
+
+USoundConcurrency::USoundConcurrency(const FObjectInitializer& ObjectInitializer)
+	: Super(ObjectInitializer)
+{
+	Concurrency.MaxCount = 16;
+	Concurrency.ResolutionRule = EMaxConcurrentResolutionRule::PreventNew;
+	Concurrency.VolumeScale = 1.0f;
+}
+
+/************************************************************************/
+/* FConcurrencyActiveSounds												*/
+/************************************************************************/
+
+FConcurrencyGroup::FConcurrencyGroup()
+	: MaxActiveSounds(0)
+	, ConcurrencyGroupID(0)
+	, ResolutionRule(EMaxConcurrentResolutionRule::PreventNew)
+{
+}
+
+FConcurrencyGroup::FConcurrencyGroup(FActiveSound* ActiveSound)
+{
+	static FConcurrencyGroupID ConcurrencyGroupIDs = 1;
+	ConcurrencyGroupID = ConcurrencyGroupIDs++;
+
+	check(ActiveSound);
+	const FSoundConcurrencySettings* ConcurrencySettings = ActiveSound->GetSoundConcurrencySettingsToApply();
+	check(ConcurrencySettings);
+
+	MaxActiveSounds = ConcurrencySettings->MaxCount;
+	ResolutionRule = ConcurrencySettings->ResolutionRule;
+
+	ActiveSounds.Add(ActiveSound);
+	ActiveSound->ConcurrencyGroupID = ConcurrencyGroupID;
+}
+
+FConcurrencyGroup::FConcurrencyGroup(const FConcurrencyGroup& Other)
+	: ActiveSounds(Other.ActiveSounds)
+	, MaxActiveSounds(Other.MaxActiveSounds)
+	, ConcurrencyGroupID(Other.ConcurrencyGroupID)
+	, ResolutionRule(Other.ResolutionRule)
+{
+	check(ConcurrencyGroupID != 0);
+}
+
+FConcurrencyGroup::FConcurrencyGroup(const FConcurrencyGroup&& Other)
+	: ActiveSounds(MoveTemp(Other.ActiveSounds))
+	, MaxActiveSounds(Other.MaxActiveSounds)
+	, ConcurrencyGroupID(Other.ConcurrencyGroupID)
+	, ResolutionRule(Other.ResolutionRule)
+{
+	check(ConcurrencyGroupID != 0);
+}
+
+FConcurrencyGroup& FConcurrencyGroup::operator=(const FConcurrencyGroup& Other)
+{
+	ActiveSounds = Other.ActiveSounds;
+	MaxActiveSounds = Other.MaxActiveSounds;
+	ConcurrencyGroupID = Other.ConcurrencyGroupID;
+	ResolutionRule = Other.ResolutionRule;
+	check(ConcurrencyGroupID != 0);
+	return *this;
+}
+
+FConcurrencyGroup& FConcurrencyGroup::operator=(const FConcurrencyGroup&& Other)
+{
+	ActiveSounds = MoveTemp(Other.ActiveSounds);
+	MaxActiveSounds = Other.MaxActiveSounds;
+	ConcurrencyGroupID = Other.ConcurrencyGroupID;
+	ResolutionRule = Other.ResolutionRule;
+	check(ConcurrencyGroupID != 0);
+	return *this;
+}
+
+TArray<FActiveSound*>& FConcurrencyGroup::GetActiveSounds()
+{
+	return ActiveSounds;
+}
+
+void FConcurrencyGroup::AddActiveSound(FActiveSound* ActiveSound)
+{
+	check(ConcurrencyGroupID != 0);
+	ActiveSounds.Add(ActiveSound);
+	ActiveSound->ConcurrencyGroupID = ConcurrencyGroupID;
+}
+
+void FConcurrencyGroup::StopQuietSoundsDueToMaxConcurrency()
+{
+	// Nothing to do if our active sound count is less than or equal to our max active sounds
+	if (ResolutionRule != EMaxConcurrentResolutionRule::StopQuietist || ActiveSounds.Num() <= MaxActiveSounds)
+	{
+		return;
+	}
+
+	// Helper function for sort this concurrency group's active sounds according to their "volume" concurrency
+	// Quieter sounds will be at the front of the array
+	struct FCompareActiveSounds
+	{
+		FORCEINLINE bool operator()(const FActiveSound& A, const FActiveSound& B) const
+		{
+			return A.VolumeConcurrency < B.VolumeConcurrency;
+		}
+	};
+
+	ActiveSounds.Sort(FCompareActiveSounds());
+
+	const int32 NumSoundsToStop = ActiveSounds.Num() - MaxActiveSounds;
+	check(NumSoundsToStop > 0);
+
+	// Need to make a new list when stopping the sounds since the process of stopping an active sound
+	// will remove the sound from this concurrency group's ActiveSounds array.
+	for (int32 i = 0; i < NumSoundsToStop; ++i)
+	{
+		// Flag this active sound as needing to be stopped due to volume-based max concurrency.
+		// This will actually be stopped in the audio device update function.
+		ActiveSounds[i]->bShouldStopDueToMaxConcurrency = true;
+	}
+
+}
+
+
+/************************************************************************/
+/* FSoundConcurrencyManager												*/
+/************************************************************************/
+
+FSoundConcurrencyManager::FSoundConcurrencyManager(class FAudioDevice* InAudioDevice)
+	: AudioDevice(InAudioDevice)
+{
+}
+
+FSoundConcurrencyManager::~FSoundConcurrencyManager()
+{
+}
+
+FActiveSound* FSoundConcurrencyManager::CreateNewActiveSound(const FActiveSound& NewActiveSound)
+{
+	check(NewActiveSound.Sound);
+	
+	// If there are no concurrency settings associated then there is no limit on this sound
+	const FSoundConcurrencySettings* Concurrency = NewActiveSound.GetSoundConcurrencySettingsToApply();
+
+	// If there was no concurrency or the setting was zero, then we will always play this sound.
+	if (!Concurrency)
+	{
+		FActiveSound* ActiveSound = new FActiveSound(NewActiveSound);
+		ActiveSound->SetAudioDevice(AudioDevice);
+		return ActiveSound;
+	}
+
+	check(Concurrency->MaxCount > 0);
+
+	uint32 ConcurrencyObjectID = NewActiveSound.GetSoundConcurrencyObjectID();
+	if (ConcurrencyObjectID == 0)
+	{
+		return HandleConcurrencyEvaluationOverride(NewActiveSound);
+	}
+	else
+	{
+		return HandleConcurrencyEvaluation(NewActiveSound);
+	}
+}
+
+FActiveSound* FSoundConcurrencyManager::ResolveConcurrency(const FActiveSound& NewActiveSound, FConcurrencyGroup& ConcurrencyGroup)
+{
+	TArray<FActiveSound*>& ActiveSounds = ConcurrencyGroup.GetActiveSounds();
+	const FSoundConcurrencySettings* Concurrency = NewActiveSound.GetSoundConcurrencySettingsToApply();
+	check(Concurrency && Concurrency->MaxCount > 0);
+
+	check(AudioDevice);
+	TArray<struct FListener>& Listeners = AudioDevice->Listeners;
+
+	bool bCanPlay = true;
+	FActiveSound* SoundToStop = nullptr;
+
+	if (ActiveSounds.Num() >= Concurrency->MaxCount && Concurrency->ResolutionRule != EMaxConcurrentResolutionRule::StopQuietist)
+	{
+		switch (Concurrency->ResolutionRule)
+		{
+			case EMaxConcurrentResolutionRule::PreventNew:
+			bCanPlay = false;
+			break;
+
+			case EMaxConcurrentResolutionRule::StopOldest:
+			{
+				for (int32 SoundIndex = 0; SoundIndex < ActiveSounds.Num(); ++SoundIndex)
+				{
+					FActiveSound* ActiveSound = ActiveSounds[SoundIndex];
+					if (SoundToStop == nullptr || ActiveSound->PlaybackTime > SoundToStop->PlaybackTime)
+					{
+						SoundToStop = ActiveSound;
+					}
+				}
+			}
+			break;
+
+			case EMaxConcurrentResolutionRule::StopFarthestThenPreventNew:
+			case EMaxConcurrentResolutionRule::StopFarthestThenOldest:
+			{
+				int32 ClosestListenerIndex = NewActiveSound.FindClosestListener(Listeners);
+				float DistanceToStopSoundSq = FVector::DistSquared(Listeners[ClosestListenerIndex].Transform.GetTranslation(), NewActiveSound.Transform.GetTranslation());
+
+				for (FActiveSound* ActiveSound : ActiveSounds)
+				{
+					ClosestListenerIndex = ActiveSound->FindClosestListener(Listeners);
+					const float DistanceToActiveSoundSq = FVector::DistSquared(Listeners[ClosestListenerIndex].Transform.GetTranslation(), ActiveSound->Transform.GetTranslation());
+
+					if (DistanceToActiveSoundSq > DistanceToStopSoundSq)
+					{
+						SoundToStop = ActiveSound;
+						DistanceToStopSoundSq = DistanceToActiveSoundSq;
+					}
+					else if (Concurrency->ResolutionRule == EMaxConcurrentResolutionRule::StopFarthestThenOldest
+								&& DistanceToActiveSoundSq == DistanceToStopSoundSq
+								&& (SoundToStop == nullptr || ActiveSound->PlaybackTime > SoundToStop->PlaybackTime))
+					{
+						SoundToStop = ActiveSound;
+						DistanceToStopSoundSq = DistanceToActiveSoundSq;
+					}
+				}
+			}
+			break;
+
+			case EMaxConcurrentResolutionRule::StopLowestPriority:
+			{
+				// Find the current lowest priority sound.
+				for (FActiveSound* CurrSound : ActiveSounds)
+				{
+					if (SoundToStop == nullptr || CurrSound->VolumeWeightedPriorityScale < SoundToStop->VolumeWeightedPriorityScale)
+					{
+						SoundToStop = CurrSound;
+					}
+				}
+
+				// Only stop any sounds if the *lowest* priority is lower than the incoming NewActiveSound
+				if (SoundToStop->VolumeWeightedPriorityScale >= NewActiveSound.VolumeWeightedPriorityScale)
+				{
+					SoundToStop = nullptr;
+				}
+			}
+			break;
+
+			case EMaxConcurrentResolutionRule::StopQuietist:
+			// We won't do anything upfront for StopQuietist resolution rule, but instead resolve it after volumes have been computed
+			// and before we stopping sounds that play past the max overall voice count. This is because it is nigh impossible (currently)
+			// for UE4 to evaluate sound volumes *before* they play. In addition to the fact that gain-stages aren't evaluated independently, 
+			// this is primarily because of the way SoundCue nodes work.
+			check(false);
+			break;
+
+			default:
+			checkf(false, TEXT("Unknown EMaxConcurrentResolutionRule enumeration."));
+			break;
+		}
+
+		// If we found a sound to stop, then stop it and allow the new sound to play.  
+		// Otherwise inform the system that the sound failed to start.
+		if (SoundToStop == nullptr)
+		{
+			// If we didn't find any sound to stop, then we're not going to play this sound, so 
+			// immediatley do the playback complete notify
+			bCanPlay = false;
+
+			if (UAudioComponent* AudioComponent = NewActiveSound.GetAudioComponent())
+			{
+				AudioComponent->PlaybackCompleted(true);
+			}
+		}
+	}
+
+	if (bCanPlay)
+	{
+		// Make a new active sound
+		FActiveSound* OutActiveSound = MakeNewActiveSound(NewActiveSound);
+		check(OutActiveSound);
+
+		// And add it to to the concurrency group
+		ConcurrencyGroup.AddActiveSound(OutActiveSound);
+
+		// Factor in the volume-scale due to sound instance concurrency count
+		// E.g. if 3 sounds are playing in the same group, and its set to be 0.9 VolumeScale, then the 
+		// 4th sound in the concurrency group is going to be attenuated by 0.9^4
+		if (Concurrency->VolumeScale < 1.0f)
+		{
+			int32 NumActiveSoundsInConcurrencyGroup = ConcurrencyGroup.GetActiveSounds().Num();
+			OutActiveSound->ConcurrencyVolumeScale = FMath::Pow(Concurrency->VolumeScale, (float)NumActiveSoundsInConcurrencyGroup);
+		}
+
+		// Stop any sounds now if needed
+		if (SoundToStop)
+		{
+			check(AudioDevice == SoundToStop->AudioDevice);
+
+			// Remove the active sound from the concurrency manager immediately so it doesn't count towards 
+			// subsequent concurrency resolution checks (i.e. if sounds are triggered multiple times in this frame)
+			RemoveActiveSound(SoundToStop);
+
+			// Add this sound to list of sounds that need to stop but don't stop it immediately
+			AudioDevice->AddSoundToStop(SoundToStop);
+		}
+
+		return OutActiveSound;
+	}
+
+	return nullptr;
+}
+
+FActiveSound* FSoundConcurrencyManager::HandleConcurrencyEvaluationOverride(const FActiveSound& NewActiveSound)
+{
+	check(NewActiveSound.Sound);
+	const FSoundConcurrencySettings* ConcurrencySettings = NewActiveSound.GetSoundConcurrencySettingsToApply();
+	check(ConcurrencySettings);
+
+	uint32 OwnerObjectID = NewActiveSound.TryGetOwnerID();
+	FSoundObjectID SoundObjectID = NewActiveSound.Sound->GetUniqueID();
+
+	FActiveSound* ActiveSound = nullptr;
+
+	// First check to see if we're limiting the sound per-owner (i.e. told to limit per owner and we have a valid object id)
+	if (ConcurrencySettings->bLimitToOwner && OwnerObjectID != 0)
+	{
+		// Get the play count for the owner
+		FSoundInstanceEntry* OwnerPerSoundEntry = OwnerPerSoundConcurrencyMap.Find(OwnerObjectID);
+
+		// If we have an entry for this owner
+		if (OwnerPerSoundEntry)
+		{
+			// Check if we already have these sounds playing on this owner
+			const FConcurrencyGroupID* ConcurrencyGroupID = OwnerPerSoundEntry->SoundInstanceToConcurrencyGroup.Find(SoundObjectID);
+			if (ConcurrencyGroupID)
+			{
+				// Get the concurrency group
+				check(*ConcurrencyGroupID != 0);
+				FConcurrencyGroup* ConcurrencyGroup = ConcurrencyGroups.Find(*ConcurrencyGroupID);
+				check(ConcurrencyGroup);
+
+				// And perform the concurrency resolution
+				ActiveSound = ResolveConcurrency(NewActiveSound, *ConcurrencyGroup);
+			}
+			else
+			{
+				// If no sounds are playing on this owner, then make a new group and active sound
+				FConcurrencyGroupID NewConcurrencyGroupID = MakeNewConcurrencyGroupAndSound(&ActiveSound, NewActiveSound);
+
+				// Adding the new concurrency group ID to the owner for this sound id
+				OwnerPerSoundEntry->SoundInstanceToConcurrencyGroup.Add(SoundObjectID, NewConcurrencyGroupID);
+			}
+		}
+		else
+		{
+			// If we don't have an entry for this sound owner, then create a new entry with a single count for the instance object id
+			FConcurrencyGroupID NewConcurrencyGroupID = MakeNewConcurrencyGroupAndSound(&ActiveSound, NewActiveSound);
+			OwnerPerSoundConcurrencyMap.Add(OwnerObjectID, FSoundInstanceEntry(SoundObjectID, NewConcurrencyGroupID));
+		}
+	}
+	else
+	{
+		// If we're not limiting the concurrency to per-owner, then we need to limit concurrency on sound instances that are playing globally
+		FConcurrencyGroupID* ConcurrencyGroupID = SoundObjectToActiveSounds.Find(SoundObjectID);
+
+		if (ConcurrencyGroupID)
+		{
+			// Get the concurrency group
+			check(*ConcurrencyGroupID != 0);
+			FConcurrencyGroup* ConcurrencyGroup = ConcurrencyGroups.Find(*ConcurrencyGroupID);
+			check(ConcurrencyGroup);
+
+			ActiveSound = ResolveConcurrency(NewActiveSound, *ConcurrencyGroup);
+		}
+		else
+		{
+			FConcurrencyGroupID NewConcurrencyGroupID = MakeNewConcurrencyGroupAndSound(&ActiveSound, NewActiveSound);
+			SoundObjectToActiveSounds.Add(SoundObjectID, NewConcurrencyGroupID);
+		}
+	}
+
+	// If the sound wasn't able to play, this will be nullptr
+	return ActiveSound;
+}
+
+FActiveSound* FSoundConcurrencyManager::HandleConcurrencyEvaluation(const FActiveSound& NewActiveSound)
+{
+	check(NewActiveSound.Sound);
+	const FSoundConcurrencySettings* ConcurrencySettings = NewActiveSound.GetSoundConcurrencySettingsToApply();
+	check(ConcurrencySettings);
+
+	uint32 OwnerObjectID = NewActiveSound.TryGetOwnerID();
+	FConcurrencyObjectID ConcurrencyObjectID = NewActiveSound.GetSoundConcurrencyObjectID();
+	check(ConcurrencyObjectID != 0);
+
+	FActiveSound* ActiveSound = nullptr;
+
+	// Check to see if we're limiting per owner
+	if (ConcurrencySettings->bLimitToOwner && OwnerObjectID != 0)
+	{
+		// Get the map of owner to concurrency map instances
+		FOwnerConcurrencyMapEntry* ConcurrencyEntry = OwnerConcurrencyMap.Find(OwnerObjectID);
+		if (ConcurrencyEntry)
+		{
+			// Now get the map of active sounds to the concurrency object
+			FConcurrencyGroupID* ConcurrencyGroupID = ConcurrencyEntry->ConcurrencyObjectToConcurrencyGroup.Find(ConcurrencyObjectID);
+			if (ConcurrencyGroupID)
+			{
+				// Get the concurrency group
+				check(*ConcurrencyGroupID != 0);
+				FConcurrencyGroup* ConcurrencyGroup = ConcurrencyGroups.Find(*ConcurrencyGroupID);
+				check(ConcurrencyGroup);
+
+				ActiveSound = ResolveConcurrency(NewActiveSound, *ConcurrencyGroup);
+			}
+			else
+			{
+				FConcurrencyGroupID NewConcurrencyGroupID = MakeNewConcurrencyGroupAndSound(&ActiveSound, NewActiveSound);
+				ConcurrencyEntry->ConcurrencyObjectToConcurrencyGroup.Add(ConcurrencyObjectID, NewConcurrencyGroupID);
+			}
+		}
+		else
+		{
+			FConcurrencyGroupID NewConcurrencyGroupID = MakeNewConcurrencyGroupAndSound(&ActiveSound, NewActiveSound);
+			OwnerConcurrencyMap.Add(OwnerObjectID, FOwnerConcurrencyMapEntry(ConcurrencyObjectID, NewConcurrencyGroupID));
+		}
+	}
+	else
+	{
+		// If not limiting per owner, then we need to globally limit concurrency of sounds that are playing with this concurrency object ID
+		FConcurrencyGroupID* ConcurrencyGroupID = ConcurrencyMap.Find(ConcurrencyObjectID);
+		if (ConcurrencyGroupID)
+		{
+			check(*ConcurrencyGroupID != 0);
+			FConcurrencyGroup* ConcurrencyGroup = ConcurrencyGroups.Find(*ConcurrencyGroupID);
+			check(ConcurrencyGroup);
+
+			ActiveSound = ResolveConcurrency(NewActiveSound, *ConcurrencyGroup);
+		}
+		else
+		{
+			FConcurrencyGroupID NewConcurrencyGroupID = MakeNewConcurrencyGroupAndSound(&ActiveSound, NewActiveSound);
+			ConcurrencyMap.Add(ConcurrencyObjectID, NewConcurrencyGroupID);
+		}
+	}
+
+	// If the sound wasn't able to play, this will be nullptr
+	return ActiveSound;
+}
+
+FConcurrencyGroupID FSoundConcurrencyManager::MakeNewConcurrencyGroupAndSound(FActiveSound** OutActiveSound, const FActiveSound& NewActiveSound)
+{
+	// First make a new active sound
+	*OutActiveSound = MakeNewActiveSound(NewActiveSound);
+	check(*OutActiveSound);
+
+	const FSoundConcurrencySettings* ConcurrencySettings = NewActiveSound.GetSoundConcurrencySettingsToApply();
+	check(ConcurrencySettings);
+
+	// Make the new group
+	FConcurrencyGroup ConcurrencyGroup(*OutActiveSound);
+
+	// Add it to the concurrency group map
+	ConcurrencyGroups.Add(ConcurrencyGroup.GetID(), ConcurrencyGroup);
+
+	return ConcurrencyGroup.GetID();
+}
+
+void FSoundConcurrencyManager::RemoveActiveSound(FActiveSound* ActiveSound)
+{
+	if (!ActiveSound->ConcurrencyGroupID)
+	{
+		return;
+	}
+
+	// Remove this sound from it's concurrency list
+	FConcurrencyGroup* ConcurrencyGroup = ConcurrencyGroups.Find(ActiveSound->ConcurrencyGroupID);
+	check(ConcurrencyGroup);
+
+	TArray<FActiveSound*>& ActiveSounds = ConcurrencyGroup->GetActiveSounds();
+	check(ActiveSounds.Num() > 0);
+	ActiveSounds.Remove(ActiveSound);
+
+}
+
+FActiveSound* FSoundConcurrencyManager::MakeNewActiveSound(const FActiveSound& NewActiveSound)
+{
+	FActiveSound* ActiveSound = new FActiveSound(NewActiveSound);
+	ActiveSound->SetAudioDevice(AudioDevice);
+	check(AudioDevice == ActiveSound->AudioDevice);
+	return ActiveSound;
+}
+
+void FSoundConcurrencyManager::StopQuietSoundsDueToMaxConcurrency()
+{
+	for (auto& ConcurrenyGroupEntry : ConcurrencyGroups)
+	{
+		ConcurrenyGroupEntry.Value.StopQuietSoundsDueToMaxConcurrency();
+	}
+}
diff --git a/Engine/Source/Runtime/Engine/Public/ActiveSound.h b/Engine/Source/Runtime/Engine/Public/ActiveSound.h
--- a/Engine/Source/Runtime/Engine/Public/ActiveSound.h
+++ b/Engine/Source/Runtime/Engine/Public/ActiveSound.h
@@ -3,6 +3,7 @@
 #pragma once
 
 #include "Sound/AudioVolume.h"
+#include "Sound/SoundConcurrency.h"
 #include "Sound/SoundAttenuation.h"
 #include "Components/AudioComponent.h"
 
@@ -91,6 +92,19 @@ public:
 		AudioComponentIndex = (UPTRINT)Component;
 	}
 
+	void SetAudioDevice(FAudioDevice* InAudioDevice)
+	{
+		AudioDevice = InAudioDevice;
+	}
+
+	FAudioDevice* AudioDevice;
+
+	/** The group of active concurrent sounds that this sound is playing in. */
+	FConcurrencyGroupID ConcurrencyGroupID;
+
+	/** Optional USoundConcurrency to override sound */
+	USoundConcurrency* ConcurrencySettings;
+
 	/** Optional SoundClass to override Sound */
 	USoundClass* SoundClassOverride;
 
@@ -112,6 +126,9 @@ public:
 	/** Whether the current component has finished playing */
 	uint32 bFinished:1;
 
+	/** Whether or not to stop this active sound due to max concurrency */
+	uint32 bShouldStopDueToMaxConcurrency:1;
+
 	/** If true, the decision on whether to apply the radio filter has been made. */
 	uint32 bRadioFilterSelected:1;
 
@@ -166,9 +183,15 @@ public:
 	float PitchMultiplier;
 	float HighFrequencyGainMultiplier;
 
+	/** A volume scale to apply to a sound based on the concurrency count of the active sound when it started */
+	float ConcurrencyVolumeScale;
+
 	float SubtitlePriority;
 	float VolumeWeightedPriorityScale;
 
+	// The volume used to determine concurrency resolution for "quietest" active sound
+	float VolumeConcurrency;
+
 	/** Frequency with which to check for occlusion from its closest listener */
 	float OcclusionCheckInterval;
 
@@ -208,7 +231,7 @@ public:
 	// Updates the wave instances to be played.
 	void UpdateWaveInstances( FAudioDevice* AudioDevice, TArray<FWaveInstance*> &OutWaveInstances, const float DeltaTime );
 
-	void Stop(FAudioDevice* AudioDevice);
+	void Stop();
 
 	/** 
 	 * Find an existing waveinstance attached to this audio component (if any)
@@ -272,6 +295,15 @@ public:
 
 	/* Determines which listener is the closest to the sound */
 	int32 FindClosestListener( const TArray<struct FListener>& InListeners ) const;
+	
+	/** Returns the unique ID of the active sound's owner if it exists. Returns 0 if the sound doesn't have an owner. */
+	uint32 TryGetOwnerID() const;
+
+	/** Gets the sound concurrency to apply on this active sound instance */
+	const FSoundConcurrencySettings* GetSoundConcurrencySettingsToApply() const;
+
+	/** Returns the sound concurrency object ID if it exists. If it doesn't exist, returns 0. */
+	uint32 GetSoundConcurrencyObjectID() const;
 
 private:
 	void UpdateAdjustVolumeMultiplier(const float DeltaTime);
diff --git a/Engine/Source/Runtime/Engine/Public/Audio.h b/Engine/Source/Runtime/Engine/Public/Audio.h
--- a/Engine/Source/Runtime/Engine/Public/Audio.h
+++ b/Engine/Source/Runtime/Engine/Public/Audio.h
@@ -8,6 +8,7 @@
 
 #include "Sound/SoundClass.h"
 #include "Sound/SoundAttenuation.h"
+//#include "Sound/SoundConcurrency.h"
 
 ENGINE_API DECLARE_LOG_CATEGORY_EXTERN(LogAudio, Warning, All);
 
@@ -55,6 +56,7 @@ ENGINE_API DECLARE_LOG_CATEGORY_EXTERN(LogAudioDebug, Display, All);
  */
 DECLARE_CYCLE_STAT_EXTERN( TEXT( "Audio Update Time" ), STAT_AudioUpdateTime, STATGROUP_Audio , );
 DECLARE_DWORD_COUNTER_STAT_EXTERN( TEXT( "Active Sounds" ), STAT_ActiveSounds, STATGROUP_Audio , );
+DECLARE_CYCLE_STAT_EXTERN(TEXT("Audio Evaluate Concurrency"), STAT_AudioEvaluateConcurrency, STATGROUP_Audio, );
 DECLARE_DWORD_COUNTER_STAT_EXTERN( TEXT( "Audio Sources" ), STAT_AudioSources, STATGROUP_Audio , );
 DECLARE_DWORD_COUNTER_STAT_EXTERN( TEXT( "Wave Instances" ), STAT_WaveInstances, STATGROUP_Audio , );
 DECLARE_DWORD_COUNTER_STAT_EXTERN( TEXT( "Wave Instances Dropped" ), STAT_WavesDroppedDueToPriority, STATGROUP_Audio , );
diff --git a/Engine/Source/Runtime/Engine/Public/AudioDevice.h b/Engine/Source/Runtime/Engine/Public/AudioDevice.h
--- a/Engine/Source/Runtime/Engine/Public/AudioDevice.h
+++ b/Engine/Source/Runtime/Engine/Public/AudioDevice.h
@@ -361,7 +361,7 @@ public:
 	/**
 	 * Creates an audio component to handle playing a sound cue
 	 */
-	static class UAudioComponent* CreateComponent( class USoundBase* Sound, class UWorld* World, AActor*  AActor  = NULL, bool Play = true, bool bStopWhenOwnerDestroyed = false, const FVector* Location = NULL, USoundAttenuation* AttenuationSettings = NULL );
+	static class UAudioComponent* CreateComponent(class USoundBase* Sound, class UWorld* World, AActor*  AActor = nullptr, bool Play = true, bool bStopWhenOwnerDestroyed = false, const FVector* Location = nullptr, USoundAttenuation* AttenuationSettings = nullptr, USoundConcurrency* ConcurrencySettings = nullptr);
 
 	/**
 	 * Adds an active sound to the audio device
@@ -744,6 +744,8 @@ protected:
 		return bSpatializationExtensionEnabled;
 	}
 
+	void AddSoundToStop(struct FActiveSound* SoundToStop);
+
 public:
 
 	/** The maximum number of concurrent audible sounds */
@@ -847,11 +849,15 @@ public:
 private:
 
 	TArray<struct FActiveSound*> ActiveSounds;
+	TArray<struct FActiveSound*> PendingSoundsToStop;
 
 	TMap<UPTRINT, struct FActiveSound*> AudioComponentToActiveSoundMap;
 
 	/** List of passive SoundMixes active last frame */
 	TArray<class USoundMix*> PrevPassiveSoundMixModifiers;
+
+	friend class FSoundConcurrencyManager;
+	FSoundConcurrencyManager ConcurrencyManager;
 };
 
 
