Fixing shadowed variables in CL 2706947

diff --git a/Engine/Source/Runtime/Engine/Private/ActiveSound.cpp b/Engine/Source/Runtime/Engine/Private/ActiveSound.cpp
--- a/Engine/Source/Runtime/Engine/Private/ActiveSound.cpp
+++ b/Engine/Source/Runtime/Engine/Private/ActiveSound.cpp
@@ -170,9 +170,9 @@ uint32 FActiveSound::GetSoundConcurrencyObjectID() const
 }
 
 
-void FActiveSound::UpdateWaveInstances( FAudioDevice* AudioDevice, TArray<FWaveInstance*> &InWaveInstances, const float DeltaTime )
+void FActiveSound::UpdateWaveInstances( TArray<FWaveInstance*> &InWaveInstances, const float DeltaTime )
 {
-	check( AudioDevice );
+	check(AudioDevice);
 
 	// Early outs.
 	if( Sound == NULL || !Sound->IsPlayable() )
@@ -279,8 +279,8 @@ void FActiveSound::UpdateWaveInstances( FAudioDevice* AudioDevice, TArray<FWaveI
 	else if (ThisSoundsWaveInstances.Num() > 0)
 	{
 		// If this active sound is told to limit concurrency by the quietest sound
-		const FSoundConcurrencySettings* ConcurrencySettings = GetSoundConcurrencySettingsToApply();
-		if (ConcurrencySettings && ConcurrencySettings->ResolutionRule == EMaxConcurrentResolutionRule::StopQuietist)
+		const FSoundConcurrencySettings* ConcurrencySettingsToApply = GetSoundConcurrencySettingsToApply();
+		if (ConcurrencySettingsToApply && ConcurrencySettingsToApply->ResolutionRule == EMaxConcurrentResolutionRule::StopQuietist)
 		{
 			check(ConcurrencyGroupID != 0);
 			// Now that we have this sound's active wave instances, lets find the loudest active wave instance to represent the "volume" of this active sound
@@ -456,8 +456,9 @@ void FActiveSound::HandleInteriorVolumes( const FListener& Listener, FSoundParse
 	}
 }
 
-void FActiveSound::ApplyRadioFilter( FAudioDevice* AudioDevice, const FSoundParseParameters& ParseParams )
+void FActiveSound::ApplyRadioFilter(const FSoundParseParameters& ParseParams )
 {
+	check(AudioDevice);
 	if( AudioDevice->GetMixDebugState() != DEBUGSTATE_DisableRadio )
 	{
 		// Make sure the radio filter is requested
diff --git a/Engine/Source/Runtime/Engine/Private/AudioDevice.cpp b/Engine/Source/Runtime/Engine/Private/AudioDevice.cpp
--- a/Engine/Source/Runtime/Engine/Private/AudioDevice.cpp
+++ b/Engine/Source/Runtime/Engine/Private/AudioDevice.cpp
@@ -1827,7 +1827,7 @@ int32 FAudioDevice::GetSortedActiveWaveInstances(TArray<FWaveInstance*>& WaveIns
 							UsedDeltaTime = 0.0f;
 						}
 
-						ActiveSound->UpdateWaveInstances(this, WaveInstances, UsedDeltaTime);
+						ActiveSound->UpdateWaveInstances(WaveInstances, UsedDeltaTime);
 					}
 				}
 			}
diff --git a/Engine/Source/Runtime/Engine/Private/SoundWave.cpp b/Engine/Source/Runtime/Engine/Private/SoundWave.cpp
--- a/Engine/Source/Runtime/Engine/Private/SoundWave.cpp
+++ b/Engine/Source/Runtime/Engine/Private/SoundWave.cpp
@@ -576,7 +576,7 @@ void USoundWave::Parse( FAudioDevice* AudioDevice, const UPTRINT NodeWaveInstanc
 	{
 		if( !ActiveSound.bRadioFilterSelected )
 		{
-			ActiveSound.ApplyRadioFilter( AudioDevice, ParseParams );
+			ActiveSound.ApplyRadioFilter(ParseParams);
 		}
 
 		WaveInstance = HandleStart( ActiveSound, NodeWaveInstanceHash);
diff --git a/Engine/Source/Runtime/Engine/Public/ActiveSound.h b/Engine/Source/Runtime/Engine/Public/ActiveSound.h
--- a/Engine/Source/Runtime/Engine/Public/ActiveSound.h
+++ b/Engine/Source/Runtime/Engine/Public/ActiveSound.h
@@ -229,7 +229,7 @@ public:
 #endif
 
 	// Updates the wave instances to be played.
-	void UpdateWaveInstances( FAudioDevice* AudioDevice, TArray<FWaveInstance*> &OutWaveInstances, const float DeltaTime );
+	void UpdateWaveInstances( TArray<FWaveInstance*> &OutWaveInstances, const float DeltaTime );
 
 	void Stop();
 
@@ -241,7 +241,7 @@ public:
 	/** 
 	 * Check whether to apply the radio filter
 	 */
-	void ApplyRadioFilter( FAudioDevice* AudioDevice, const struct FSoundParseParameters& ParseParams );
+	void ApplyRadioFilter(const struct FSoundParseParameters& ParseParams );
 
 	/** Sets a float instance parameter for the ActiveSound */
 	void SetFloatParameter(const FName InName, const float InFloat);
