VS2015 fixes merging back from Orion

diff --git a/Engine/Source/Runtime/Engine/Private/PhysicsEngine/PhysXLibs.cpp b/Engine/Source/Runtime/Engine/Private/PhysicsEngine/PhysXLibs.cpp
--- a/Engine/Source/Runtime/Engine/Private/PhysicsEngine/PhysXLibs.cpp
+++ b/Engine/Source/Runtime/Engine/Private/PhysicsEngine/PhysXLibs.cpp
@@ -41,9 +41,13 @@ void LoadPhysXModules()
 
 	#if PLATFORM_64BITS
 
-		FString RootPhysXPath(PhysXBinariesRoot + TEXT("Win64/VS2013/"));
-		FString RootAPEXPath(APEXBinariesRoot + TEXT("Win64/VS2013/"));
-
+		#if _MSC_VER >= 1900
+			FString RootPhysXPath(PhysXBinariesRoot + TEXT("Win64/VS2015/"));
+			FString RootAPEXPath(APEXBinariesRoot + TEXT("Win64/VS2015/"));
+		#else
+			FString RootPhysXPath(PhysXBinariesRoot + TEXT("Win64/VS2013/"));
+			FString RootAPEXPath(APEXBinariesRoot + TEXT("Win64/VS2013/"));
+		#endif
 
 		#if UE_BUILD_DEBUG && !defined(NDEBUG)	// Use !defined(NDEBUG) to check to see if we actually are linking with Debug third party libraries (bDebugBuildsActuallyUseDebugCRT)
 
@@ -100,7 +104,10 @@ void LoadPhysXModules()
 		#endif	//UE_BUILD_DEBUG
 	#else	//PLATFORM_64BITS
 
-		#if _MSC_VER >= 1800
+		#if _MSC_VER >= 1900
+			FString RootPhysXPath(PhysXBinariesRoot + TEXT("Win32/VS2015/"));
+			FString RootAPEXPath(APEXBinariesRoot + TEXT("Win32/VS2015/"));
+		#elif _MSC_VER >= 1800
 			FString RootPhysXPath(PhysXBinariesRoot + TEXT("Win32/VS2013/"));
 			FString RootAPEXPath(APEXBinariesRoot + TEXT("Win32/VS2013/"));
 		#else
