UE-21530

Fix for high quality shader flag triggering a cook/cache of all quality levels.
Ensure high quality shader map id is not affected by quality override setting.

diff --git a/Engine/Source/Runtime/Engine/Private/Materials/Material.cpp b/Engine/Source/Runtime/Engine/Private/Materials/Material.cpp
--- a/Engine/Source/Runtime/Engine/Private/Materials/Material.cpp
+++ b/Engine/Source/Runtime/Engine/Private/Materials/Material.cpp
@@ -1853,17 +1853,10 @@ void UMaterial::CacheResourceShadersForCooking(EShaderPlatform ShaderPlatform, T
 	TArray<bool, TInlineAllocator<EMaterialQualityLevel::Num> > QualityLevelsUsed;
 	GetQualityLevelUsage(QualityLevelsUsed, ShaderPlatform);
 
-	bool bAnyQualityLevelUsed = false;
-
-	for (int32 QualityLevelIndex = 0; QualityLevelIndex < EMaterialQualityLevel::Num; QualityLevelIndex++)
-	{
-		bAnyQualityLevelUsed = bAnyQualityLevelUsed || QualityLevelsUsed[QualityLevelIndex];
-	}
-
 	for (int32 QualityLevelIndex = 0; QualityLevelIndex < EMaterialQualityLevel::Num; QualityLevelIndex++)
 	{
 		// Add all quality levels if multiple are needed (due to different node graphs), otherwise just add the high quality entry
-		if (bAnyQualityLevelUsed || QualityLevelIndex == EMaterialQualityLevel::High)
+		if (QualityLevelsUsed[QualityLevelIndex] || QualityLevelIndex == EMaterialQualityLevel::High)
 		{
 			FMaterialResource* NewResource = AllocateResource();
 			NewResource->SetMaterial(this, (EMaterialQualityLevel::Type)QualityLevelIndex, QualityLevelsUsed[QualityLevelIndex], (ERHIFeatureLevel::Type)TargetFeatureLevel);
@@ -2200,7 +2193,7 @@ void UMaterial::GetQualityLevelUsage(TArray<bool, TInlineAllocator<EMaterialQual
 		const UShaderPlatformQualitySettings* MaterialQualitySettings = UMaterialShaderQualitySettings::Get()->GetShaderPlatformQualitySettings(ShaderPlatform);
 		OutQualityLevelsUsed[EMaterialQualityLevel::Low] |= MaterialQualitySettings->GetQualityOverrides(EMaterialQualityLevel::Low).bEnableOverride;
 		OutQualityLevelsUsed[EMaterialQualityLevel::Medium] |= MaterialQualitySettings->GetQualityOverrides(EMaterialQualityLevel::Medium).bEnableOverride;
-		OutQualityLevelsUsed[EMaterialQualityLevel::High] |= MaterialQualitySettings->GetQualityOverrides(EMaterialQualityLevel::High).bEnableOverride;
+		// ignore EMaterialQualityLevel::High as this is always available
 	}
 }
 
diff --git a/Engine/Source/Runtime/Engine/Private/Materials/MaterialInstance.cpp b/Engine/Source/Runtime/Engine/Private/Materials/MaterialInstance.cpp
--- a/Engine/Source/Runtime/Engine/Private/Materials/MaterialInstance.cpp
+++ b/Engine/Source/Runtime/Engine/Private/Materials/MaterialInstance.cpp
@@ -1368,17 +1368,10 @@ void UMaterialInstance::CacheResourceShadersForCooking(EShaderPlatform ShaderPla
 		TArray<FMaterialResource*> ResourcesToCache;
 		ERHIFeatureLevel::Type TargetFeatureLevel = GetMaxSupportedFeatureLevel(ShaderPlatform);
 
-		bool bAnyQualityLevelUsed = false;
-
-		for (int32 QualityLevelIndex = 0; QualityLevelIndex < EMaterialQualityLevel::Num; QualityLevelIndex++)
-		{
-			bAnyQualityLevelUsed = bAnyQualityLevelUsed || QualityLevelsUsed[QualityLevelIndex];
-		}
-
 		for (int32 QualityLevelIndex = 0; QualityLevelIndex < EMaterialQualityLevel::Num; QualityLevelIndex++)
 		{
 			// Cache all quality levels, unless they are all the same (due to using the same nodes), then just cache the high quality
-			if (bAnyQualityLevelUsed || QualityLevelIndex == EMaterialQualityLevel::High)
+			if (QualityLevelsUsed[QualityLevelIndex] || QualityLevelIndex == EMaterialQualityLevel::High)
 			{
 				FMaterialResource* NewResource = new FMaterialResource();
 				NewResource->SetMaterial(BaseMaterial, (EMaterialQualityLevel::Type)QualityLevelIndex, QualityLevelsUsed[QualityLevelIndex], (ERHIFeatureLevel::Type)TargetFeatureLevel, this);
diff --git a/Engine/Source/Runtime/Engine/Private/Materials/MaterialShared.cpp b/Engine/Source/Runtime/Engine/Private/Materials/MaterialShared.cpp
--- a/Engine/Source/Runtime/Engine/Private/Materials/MaterialShared.cpp
+++ b/Engine/Source/Runtime/Engine/Private/Materials/MaterialShared.cpp
@@ -1905,9 +1905,9 @@ void FMaterial::GetReferencedTexturesHash(EShaderPlatform Platform, FSHAHash& Ou
 	// Appending the quality settings for this platform,
 	// this is being done here to avoid bumping EUnrealEngineObjectUE4Version for 4.10
 	// must be fixed for 4.11.
-	if (bHasQualityLevelUsage)
 	{
-		UMaterialShaderQualitySettings::Get()->GetShaderPlatformQualitySettings(Platform)->AppendToHashState(GetQualityLevelForShaderMapId(), HashState);
+		EMaterialQualityLevel::Type QualityLevelForHash = bHasQualityLevelUsage ? QualityLevel : EMaterialQualityLevel::High;
+		UMaterialShaderQualitySettings::Get()->GetShaderPlatformQualitySettings(Platform)->AppendToHashState(QualityLevelForHash, HashState);
 	}
 
 	HashState.Final();
