Allow games to override network travel destruction of world net driver

diff --git a/Engine/Source/Runtime/Engine/Classes/Engine/Engine.h b/Engine/Source/Runtime/Engine/Classes/Engine/Engine.h
--- a/Engine/Source/Runtime/Engine/Classes/Engine/Engine.h
+++ b/Engine/Source/Runtime/Engine/Classes/Engine/Engine.h
@@ -2597,6 +2597,13 @@ protected:
 
 	virtual void MovePendingLevel(FWorldContext &Context);
 
+	/**
+	 *	Returns true if BROWSE should shuts down the current network driver.
+	 **/
+	virtual bool ShouldShutdownWorldNetDriver()
+	{
+		return true;
+	}
 
 	bool WorldHasValidContext(UWorld *InWorld);
 
diff --git a/Engine/Source/Runtime/Engine/Private/UnrealEngine.cpp b/Engine/Source/Runtime/Engine/Private/UnrealEngine.cpp
--- a/Engine/Source/Runtime/Engine/Private/UnrealEngine.cpp
+++ b/Engine/Source/Runtime/Engine/Private/UnrealEngine.cpp
@@ -8293,6 +8293,8 @@ void UEngine::HandleNetworkFailure(UWorld *World, UNetDriver *NetDriver, ENetwor
 			break;
 		case ENetworkFailure::PendingConnectionFailure:
 			// TODO stop the connecting movie
+			bShouldTravel = false;
+			CancelPending(NetDriver);
 			break;
 		case ENetworkFailure::ConnectionLost:
 			// Hosts don't travel when clients disconnect
@@ -8736,7 +8738,7 @@ EBrowseReturnVal::Type UEngine::Browse( FWorldContext& WorldContext, FURL URL, F
 		}
 
 		// Clean up the netdriver/socket so that the pending level succeeds
-		if (WorldContext.World())
+		if (WorldContext.World() && ShouldShutdownWorldNetDriver())
 		{
 			ShutdownWorldNetDriver(WorldContext.World());
 		}
