Remove accidental merge of UT-specific code

diff --git a/Engine/Source/Runtime/Engine/Private/UnrealEngine.cpp b/Engine/Source/Runtime/Engine/Private/UnrealEngine.cpp
--- a/Engine/Source/Runtime/Engine/Private/UnrealEngine.cpp
+++ b/Engine/Source/Runtime/Engine/Private/UnrealEngine.cpp
@@ -8293,8 +8293,6 @@ void UEngine::HandleNetworkFailure(UWorld *World, UNetDriver *NetDriver, ENetwor
 			break;
 		case ENetworkFailure::PendingConnectionFailure:
 			// TODO stop the connecting movie
-			bShouldTravel = false;
-			CancelPending(NetDriver);
 			break;
 		case ENetworkFailure::ConnectionLost:
 			// Hosts don't travel when clients disconnect
