Fix UMG ensure

diff --git a/Engine/Source/Runtime/Engine/Private/TextureRenderTarget2D.cpp b/Engine/Source/Runtime/Engine/Private/TextureRenderTarget2D.cpp
--- a/Engine/Source/Runtime/Engine/Private/TextureRenderTarget2D.cpp
+++ b/Engine/Source/Runtime/Engine/Private/TextureRenderTarget2D.cpp
@@ -395,7 +395,7 @@ void FTextureRenderTarget2DResource::UpdateDeferredResource( FRHICommandListImme
 	{
 		RHICmdList.SetViewport(0, 0, 0.0f, TargetSizeX, TargetSizeY, 1.0f);
 		ensure(RenderTargetTextureRHI->GetClearColor() == ClearColor);
-		SetRenderTarget(RHICmdList, RenderTargetTextureRHI, FTextureRHIRef(), ESimpleRenderTargetMode::EClearColorExistingDepth);
+		SetRenderTarget(RHICmdList, RenderTargetTextureRHI, FTextureRHIRef(), ESimpleRenderTargetMode::EClearColorExistingDepth, FExclusiveDepthStencil::DepthWrite_StencilWrite, true);
 	}
  
 	if (Owner->bAutoGenerateMips)
