Fix issue with cooking

diff --git a/Engine/Source/Runtime/Engine/Private/Animation/AnimSequence.cpp b/Engine/Source/Runtime/Engine/Private/Animation/AnimSequence.cpp
--- a/Engine/Source/Runtime/Engine/Private/Animation/AnimSequence.cpp
+++ b/Engine/Source/Runtime/Engine/Private/Animation/AnimSequence.cpp
@@ -1687,7 +1687,9 @@ bool UAnimSequence::CompressRawAnimData(float MaxPosDiff, float MaxAngleDiff)
 	bool bRemovedKeys = false;
 #if WITH_EDITORONLY_DATA
 
-	if ( ensureMsgf (RawAnimationData.Num() > 0, TEXT("%s is trying to compress while raw animation is missing"), *GetName()) )
+	const USkeleton* MySkeleton = GetSkeleton();
+
+	if ( MySkeleton && ensureMsgf(RawAnimationData.Num() > 0, TEXT("%s is trying to compress while raw animation is missing"), *GetName()))
 	{
 		// This removes trivial keys, and this has to happen before the removing tracks
 		for(int32 TrackIndex=0; TrackIndex<RawAnimationData.Num(); TrackIndex++)
@@ -1695,36 +1697,31 @@ bool UAnimSequence::CompressRawAnimData(float MaxPosDiff, float MaxAngleDiff)
 			bRemovedKeys |= CompressRawAnimSequenceTrack(RawAnimationData[TrackIndex], MaxPosDiff, MaxAngleDiff);
 		}
 
-		const USkeleton* MySkeleton = GetSkeleton();
-
-		if(MySkeleton)
+		bool bCompressScaleKeys = false;
+		// go through remove keys if not needed
+		for(int32 TrackIndex=0; TrackIndex<RawAnimationData.Num(); TrackIndex++)
 		{
-			bool bCompressScaleKeys = false;
-			// go through remove keys if not needed
-			for(int32 TrackIndex=0; TrackIndex<RawAnimationData.Num(); TrackIndex++)
+			FRawAnimSequenceTrack const& RawData = RawAnimationData[TrackIndex];
+			if(RawData.ScaleKeys.Num() > 0)
 			{
-				FRawAnimSequenceTrack const& RawData = RawAnimationData[TrackIndex];
-				if(RawData.ScaleKeys.Num() > 0)
+				// if scale key exists, see if we can just empty it
+				if((RawData.ScaleKeys.Num() > 1) || (RawData.ScaleKeys[0].Equals(FVector(1.f)) == false))
 				{
-					// if scale key exists, see if we can just empty it
-					if((RawData.ScaleKeys.Num() > 1) || (RawData.ScaleKeys[0].Equals(FVector(1.f)) == false))
-					{
-						bCompressScaleKeys = true;
-						break;
-					}
+					bCompressScaleKeys = true;
+					break;
 				}
 			}
+		} 
 
-			// if we don't have scale, we should delete all scale keys
-			// if you have one track that has scale, we still should support scale, so compress scale
-			if(!bCompressScaleKeys)
+		// if we don't have scale, we should delete all scale keys
+		// if you have one track that has scale, we still should support scale, so compress scale
+		if(!bCompressScaleKeys)
+		{
+			// then remove all scale keys
+			for(int32 TrackIndex=0; TrackIndex<RawAnimationData.Num(); TrackIndex++)
 			{
-				// then remove all scale keys
-				for(int32 TrackIndex=0; TrackIndex<RawAnimationData.Num(); TrackIndex++)
-				{
-					FRawAnimSequenceTrack& RawData = RawAnimationData[TrackIndex];
-					RawData.ScaleKeys.Empty();
-				}
+				FRawAnimSequenceTrack& RawData = RawAnimationData[TrackIndex];
+				RawData.ScaleKeys.Empty();
 			}
 		}
 
