Fix issue where non simulating events are being welded even though they are marked as query only.

diff --git a/Engine/Source/Runtime/Engine/Private/PrimitiveComponentPhysics.cpp b/Engine/Source/Runtime/Engine/Private/PrimitiveComponentPhysics.cpp
--- a/Engine/Source/Runtime/Engine/Private/PrimitiveComponentPhysics.cpp
+++ b/Engine/Source/Runtime/Engine/Private/PrimitiveComponentPhysics.cpp
@@ -718,6 +718,13 @@ bool UPrimitiveComponent::WeldToImplementation(USceneComponent * InParent, FName
 		return false;
 	}
 
+	//Make sure that objects marked as non-simulating do not start simulating due to welding.
+	ECollisionEnabled::Type CollisionType = BI->GetCollisionEnabled();
+	if (CollisionType == ECollisionEnabled::QueryOnly || CollisionType == ECollisionEnabled::NoCollision)
+	{
+		return false;
+	}
+
 	UnWeldFromParent();	//make sure to unweld from wherever we currently are
 
 	FName SocketName;
