Fixed CastShadows being uneditable with default materials

diff --git a/Engine/Source/Runtime/Engine/Private/Components/PrimitiveComponent.cpp b/Engine/Source/Runtime/Engine/Private/Components/PrimitiveComponent.cpp
--- a/Engine/Source/Runtime/Engine/Private/Components/PrimitiveComponent.cpp
+++ b/Engine/Source/Runtime/Engine/Private/Components/PrimitiveComponent.cpp
@@ -686,13 +686,20 @@ bool UPrimitiveComponent::CanEditChange(const UProperty* InProperty) const
 			const int32 NumMaterials = GetNumMaterials();
 			for (int32 MaterialIndex = 0; (MaterialIndex < NumMaterials) && !bHasAnyLitMaterials; ++MaterialIndex)
 			{
-				if (UMaterialInterface* Material = GetMaterial(MaterialIndex))
+				UMaterialInterface* Material = GetMaterial(MaterialIndex);
+
+				if (Material)
 				{
 					if (Material->GetShadingModel() != MSM_Unlit)
 					{
 						bHasAnyLitMaterials = true;
 					}
 				}
+				else
+				{
+					// Default material is lit
+					bHasAnyLitMaterials = true;
+				}
 			}
 
 			// If there's at least one lit section it could cast shadows, so let the property be edited.
