Integrate: Increased distance field atlas max to 512Mb

diff --git a/Engine/Source/Runtime/Engine/Private/DistanceFieldAtlas.cpp b/Engine/Source/Runtime/Engine/Private/DistanceFieldAtlas.cpp
--- a/Engine/Source/Runtime/Engine/Private/DistanceFieldAtlas.cpp
+++ b/Engine/Source/Runtime/Engine/Private/DistanceFieldAtlas.cpp
@@ -15,10 +15,13 @@
 
 TGlobalResource<FDistanceFieldVolumeTextureAtlas> GDistanceFieldVolumeTextureAtlas = TGlobalResource<FDistanceFieldVolumeTextureAtlas>(PF_R16F);
 
-const int32 MaxAtlasDimension = 512;
+// 512Mb
+const int32 MaxAtlasDimensionX = 512;
+const int32 MaxAtlasDimensionY = 512;
+const int32 MaxAtlasDimensionZ = 1024;
 
 FDistanceFieldVolumeTextureAtlas::FDistanceFieldVolumeTextureAtlas(EPixelFormat InFormat) :
-	BlockAllocator(0, 0, 0, MaxAtlasDimension, MaxAtlasDimension, MaxAtlasDimension, false, false)
+	BlockAllocator(0, 0, 0, MaxAtlasDimensionX, MaxAtlasDimensionY, MaxAtlasDimensionZ, false, false)
 {
 	Generation = 0;
 	Format = InFormat;
@@ -92,7 +95,7 @@ void FDistanceFieldVolumeTextureAtlas::UpdateAllocations()
 			if (CurrentAllocations.Num() > 0)
 			{
 				// Remove all allocations from the layout so we have a clean slate
-				BlockAllocator = FTextureLayout3d(0, 0, 0, MaxAtlasDimension, MaxAtlasDimension, MaxAtlasDimension, false, false);
+				BlockAllocator = FTextureLayout3d(0, 0, 0, MaxAtlasDimensionX, MaxAtlasDimensionY, MaxAtlasDimensionZ, false, false);
 				
 				Generation++;
 
