Integrate: Fixed r.forceLOD causing components that didn't write to LODIndex (landscape) to disappear

diff --git a/Engine/Source/Runtime/Engine/Private/SceneManagement.cpp b/Engine/Source/Runtime/Engine/Private/SceneManagement.cpp
--- a/Engine/Source/Runtime/Engine/Private/SceneManagement.cpp
+++ b/Engine/Source/Runtime/Engine/Private/SceneManagement.cpp
@@ -372,7 +372,8 @@ FLODMask ComputeLODForMeshes( const TIndirectArray<class FStaticMesh>& StaticMes
 	// Handle forced LOD level first
 	if(ForcedLODLevel >= 0)
 	{
-		int8 MaxLOD = 0;
+		// Note: starting at -1 which is the default LODIndex, for cases where LODIndex didn't get set
+		int8 MaxLOD = -1;
 		for(int32 MeshIndex = 0 ; MeshIndex < StaticMeshes.Num() ; ++MeshIndex)
 		{
 			const FStaticMesh&  Mesh = StaticMeshes[MeshIndex];
