Engine: Made hitch and 'bound by' determination in fpscharts configurable by console variables

Bound by (GT, RT, GPU) calculations now uses the t.TargetFrameTimeThreshold cvar (in ms) rather than being hardcoded to only kick in when over 30 fps

New hitch detection cvars:
- t.HitchFrameTimeThreshold is the definition of a hitchy frame (in ms), defaulting to 60 ms
- t.HitchDeadTimeWindow is the minimum time that must pass before we'll record a new hitch (in ms), defaulting to 200 ms
- t.HitchVersusNonHitchRatio is the minimum ratio between the last frame time and the current frame time to consider it a hitch, defaulting to 1.5 (e.g., if the last frame time was 50, the current frame must be at least 75 to be considered a hitch)
- They can also be set in the [SystemSettings] section of a DefaultEngine.ini for games that target something other than 30 fps

[INTEGRATE] Change 2703808 by [REDACTED_PERSON][REDACTED_HANDLE] on 2015/09/24 01:28:21

diff --git a/Engine/Source/Runtime/Engine/Private/ChartCreation.cpp b/Engine/Source/Runtime/Engine/Private/ChartCreation.cpp
--- a/Engine/Source/Runtime/Engine/Private/ChartCreation.cpp
+++ b/Engine/Source/Runtime/Engine/Private/ChartCreation.cpp
@@ -10,6 +10,7 @@
 #include "RenderCore.h"
 #include "Scalability.h"
 #include "GameFramework/GameUserSettings.h"
+#include "Performance/EnginePerformanceTargets.h"
 
 DEFINE_LOG_CATEGORY_STATIC(LogChartCreation, Log, All);
 
@@ -40,7 +41,7 @@ FHitchChartEntry GHitchChart[ STAT_FPSChart_LastHitchBucketStat - STAT_FPSChart_
 
 /** Time thresholds for the various hitch buckets in milliseconds */
 const int32 GHitchThresholds[ STAT_FPSChart_LastHitchBucketStat - STAT_FPSChart_FirstHitchStat ] =
-	{ 5000, 2500, 2000, 1500, 1000, 750, 500, 300, 200, 150, 100, 60 };
+	{ 5000, 2500, 2000, 1500, 1000, 750, 500, 300, 200, 150, 100, 60, 30 };
 
 /** Number of frames for each time of <boundtype> **/
 uint32 GNumFramesBound_GameThread = 0;
@@ -116,6 +117,8 @@ static FString CreateOutputDirectory()
 */
 void UEngine::TickFPSChart( float DeltaSeconds )
 {
+	const float MSToSeconds = 1.0f / 1000.0f;
+
 	// We can't trust delta seconds if frame time clamping is enabled or if we're benchmarking so we simply
 	// calculate it ourselves.
 	static double LastTime = 0;
@@ -144,15 +147,16 @@ void UEngine::TickFPSChart( float DeltaSeconds )
 	const uint32 LocalRenderThreadTime = GRenderThreadTime;
 	const uint32 LocalGPUFrameTime = GGPUFrameTime;
 
-	// determine which Time is the greatest and less than 33.33 (as if we are above 30 fps we are happy and thus not "bounded" )
-	const float Epsilon = 0.250f;
+	// determine which pipeline time is the greatest (between game thread, render thread, and GPU)
+	const float EpsilonCycles = 0.250f;
 	uint32 MaxThreadTimeValue = FMath::Max3<uint32>( LocalRenderThreadTime, GGameThreadTime, LocalGPUFrameTime );
 	const float FrameTime = FPlatformTime::ToSeconds(MaxThreadTimeValue);
 
-	// so we want to try and guess when we are bound by the GPU even when on the xenon we do not get that info
-	// If the frametime is bigger than 35 ms we can take DeltaSeconds into account as we're not VSYNCing in that case.
+	const float EngineTargetMS = FEnginePerformanceTargets::GetTargetFrameTimeThresholdMS();
+
+	// Try to estimate a GPU time even if the current platform does not support GPU timing
 	uint32 PossibleGPUTime = LocalGPUFrameTime;
-	if( PossibleGPUTime == 0 )
+	if (PossibleGPUTime == 0)
 	{
 		// if we are over
 		PossibleGPUTime = static_cast<uint32>(FMath::Max( FrameTime, DeltaSeconds ) / FPlatformTime::GetSecondsPerCycle() );
@@ -160,109 +164,108 @@ void UEngine::TickFPSChart( float DeltaSeconds )
 	}
 
 	// Disregard frames that took longer than one second when accumulating data.
-	if( DeltaSeconds < 1.f )
+	if (DeltaSeconds < 1.0f)
 	{
-		const float CurrentFPS = 1 / DeltaSeconds;
+		const float CurrentFPS = 1.0f / DeltaSeconds;
 
-		if( CurrentFPS < 5 )
+		if (CurrentFPS < 5.0f)
 		{
-			GFPSChart[ STAT_FPSChart_0_5 - STAT_FPSChartFirstStat ].Count++;
-			GFPSChart[ STAT_FPSChart_0_5 - STAT_FPSChartFirstStat ].CummulativeTime += DeltaSeconds;
+			GFPSChart[STAT_FPSChart_0_5 - STAT_FPSChartFirstStat].Count++;
+			GFPSChart[STAT_FPSChart_0_5 - STAT_FPSChartFirstStat].CummulativeTime += DeltaSeconds;
 		}
-		else if( CurrentFPS < 10 )
+		else if (CurrentFPS < 10.0f)
 		{
-			GFPSChart[ STAT_FPSChart_5_10 - STAT_FPSChartFirstStat ].Count++;
-			GFPSChart[ STAT_FPSChart_5_10 - STAT_FPSChartFirstStat ].CummulativeTime += DeltaSeconds;
+			GFPSChart[STAT_FPSChart_5_10 - STAT_FPSChartFirstStat].Count++;
+			GFPSChart[STAT_FPSChart_5_10 - STAT_FPSChartFirstStat].CummulativeTime += DeltaSeconds;
 		}
-		else if( CurrentFPS < 15 )
+		else if (CurrentFPS < 15.0f)
 		{
-			GFPSChart[ STAT_FPSChart_10_15 - STAT_FPSChartFirstStat ].Count++;
-			GFPSChart[ STAT_FPSChart_10_15 - STAT_FPSChartFirstStat ].CummulativeTime += DeltaSeconds;
+			GFPSChart[ STAT_FPSChart_10_15 - STAT_FPSChartFirstStat].Count++;
+			GFPSChart[ STAT_FPSChart_10_15 - STAT_FPSChartFirstStat].CummulativeTime += DeltaSeconds;
 		}
-		else if( CurrentFPS < 20 )
+		else if (CurrentFPS < 20.0f)
 		{
-			GFPSChart[ STAT_FPSChart_15_20 - STAT_FPSChartFirstStat ].Count++;
-			GFPSChart[ STAT_FPSChart_15_20 - STAT_FPSChartFirstStat ].CummulativeTime += DeltaSeconds;
+			GFPSChart[STAT_FPSChart_15_20 - STAT_FPSChartFirstStat].Count++;
+			GFPSChart[STAT_FPSChart_15_20 - STAT_FPSChartFirstStat].CummulativeTime += DeltaSeconds;
 		}
-		else if( CurrentFPS < 25 )
+		else if (CurrentFPS < 25.0f)
 		{
-			GFPSChart[ STAT_FPSChart_20_25 - STAT_FPSChartFirstStat ].Count++;
-			GFPSChart[ STAT_FPSChart_20_25 - STAT_FPSChartFirstStat ].CummulativeTime += DeltaSeconds;
+			GFPSChart[STAT_FPSChart_20_25 - STAT_FPSChartFirstStat].Count++;
+			GFPSChart[STAT_FPSChart_20_25 - STAT_FPSChartFirstStat].CummulativeTime += DeltaSeconds;
 		}
-		else if( CurrentFPS < 30 )
+		else if (CurrentFPS < 30.0f)
 		{
-			GFPSChart[ STAT_FPSChart_25_30 - STAT_FPSChartFirstStat ].Count++;
-			GFPSChart[ STAT_FPSChart_25_30 - STAT_FPSChartFirstStat ].CummulativeTime += DeltaSeconds;
+			GFPSChart[STAT_FPSChart_25_30 - STAT_FPSChartFirstStat].Count++;
+			GFPSChart[STAT_FPSChart_25_30 - STAT_FPSChartFirstStat].CummulativeTime += DeltaSeconds;
 		}
-		else if( CurrentFPS < 40 )
+		else if (CurrentFPS < 40.0f)
 		{
-			GFPSChart[ STAT_FPSChart_30_40 - STAT_FPSChartFirstStat ].Count++;
-			GFPSChart[ STAT_FPSChart_30_40 - STAT_FPSChartFirstStat ].CummulativeTime += DeltaSeconds;
+			GFPSChart[STAT_FPSChart_30_40 - STAT_FPSChartFirstStat].Count++;
+			GFPSChart[STAT_FPSChart_30_40 - STAT_FPSChartFirstStat].CummulativeTime += DeltaSeconds;
 		}
-		else if( CurrentFPS < 50 )
+		else if (CurrentFPS < 50.0f)
 		{
-			GFPSChart[ STAT_FPSChart_40_50 - STAT_FPSChartFirstStat ].Count++;
-			GFPSChart[ STAT_FPSChart_40_50 - STAT_FPSChartFirstStat ].CummulativeTime += DeltaSeconds;
+			GFPSChart[STAT_FPSChart_40_50 - STAT_FPSChartFirstStat].Count++;
+			GFPSChart[STAT_FPSChart_40_50 - STAT_FPSChartFirstStat].CummulativeTime += DeltaSeconds;
 		}
-		else if( CurrentFPS < 60 )
+		else if (CurrentFPS < 60.0f)
 		{
-			GFPSChart[ STAT_FPSChart_50_60 - STAT_FPSChartFirstStat ].Count++;
-			GFPSChart[ STAT_FPSChart_50_60 - STAT_FPSChartFirstStat ].CummulativeTime += DeltaSeconds;
+			GFPSChart[STAT_FPSChart_50_60 - STAT_FPSChartFirstStat].Count++;
+			GFPSChart[STAT_FPSChart_50_60 - STAT_FPSChartFirstStat].CummulativeTime += DeltaSeconds;
 		}
-		else if (CurrentFPS < 70)
+		else if (CurrentFPS < 70.0f)
 		{
 			GFPSChart[STAT_FPSChart_60_70 - STAT_FPSChartFirstStat].Count++;
 			GFPSChart[STAT_FPSChart_60_70 - STAT_FPSChartFirstStat].CummulativeTime += DeltaSeconds;
 		}
-		else if (CurrentFPS < 80)
+		else if (CurrentFPS < 80.0f)
 		{
 			GFPSChart[STAT_FPSChart_70_80 - STAT_FPSChartFirstStat].Count++;
 			GFPSChart[STAT_FPSChart_70_80 - STAT_FPSChartFirstStat].CummulativeTime += DeltaSeconds;
 		}
-		else if (CurrentFPS < 90)
+		else if (CurrentFPS < 90.0f)
 		{
 			GFPSChart[STAT_FPSChart_80_90 - STAT_FPSChartFirstStat].Count++;
 			GFPSChart[STAT_FPSChart_80_90 - STAT_FPSChartFirstStat].CummulativeTime += DeltaSeconds;
 		}
-		else if (CurrentFPS < 100)
+		else if (CurrentFPS < 100.0f)
 		{
 			GFPSChart[STAT_FPSChart_90_100 - STAT_FPSChartFirstStat].Count++;
 			GFPSChart[STAT_FPSChart_90_100 - STAT_FPSChartFirstStat].CummulativeTime += DeltaSeconds;
 		}
-		else if (CurrentFPS < 110)
+		else if (CurrentFPS < 110.0f)
 		{
 			GFPSChart[STAT_FPSChart_100_110 - STAT_FPSChartFirstStat].Count++;
 			GFPSChart[STAT_FPSChart_100_110 - STAT_FPSChartFirstStat].CummulativeTime += DeltaSeconds;
 		}
-		else if (CurrentFPS < 120)
+		else if (CurrentFPS < 120.0f)
 		{
 			GFPSChart[STAT_FPSChart_110_120 - STAT_FPSChartFirstStat].Count++;
 			GFPSChart[STAT_FPSChart_110_120 - STAT_FPSChartFirstStat].CummulativeTime += DeltaSeconds;
 		}
 		else
 		{
-			GFPSChart[ STAT_FPSChart_120_INF - STAT_FPSChartFirstStat ].Count++;
-			GFPSChart[ STAT_FPSChart_120_INF - STAT_FPSChartFirstStat ].CummulativeTime += DeltaSeconds;
+			GFPSChart[STAT_FPSChart_120_INF - STAT_FPSChartFirstStat].Count++;
+			GFPSChart[STAT_FPSChart_120_INF - STAT_FPSChartFirstStat].CummulativeTime += DeltaSeconds;
 		}
 
 		GTotalGPUTime += FPlatformTime::ToSeconds(LocalGPUFrameTime);
 
-		// if Frametime is > 33 ms then we are bounded by something
-		if( CurrentFPS < 30 )
+		// if frame time is greater than our target then we are bounded by something
+		const float TargetThreadTimeSeconds = EngineTargetMS * MSToSeconds;
+		if (DeltaSeconds > TargetThreadTimeSeconds)
 		{
-			// If GPU time is inferred we can only determine GPU > 33 ms if we are GPU bound.
+			// If GPU time is inferred we can only determine GPU > threshold if we are GPU bound.
 			bool bAreWeGPUBoundIfInferred = true;
-			// 33.333ms
-			const float TargetThreadTimeSeconds = .0333333f;
 
-			if( FPlatformTime::ToSeconds(GGameThreadTime) >= TargetThreadTimeSeconds )
+			if (FPlatformTime::ToSeconds(GGameThreadTime) >= TargetThreadTimeSeconds)
 			{
 				GNumFramesBound_GameThread++;
 				GTotalFramesBoundTime_GameThread += DeltaSeconds;
 				bAreWeGPUBoundIfInferred = false;
 			}
 
-			if( FPlatformTime::ToSeconds(LocalRenderThreadTime) >= TargetThreadTimeSeconds )
+			if (FPlatformTime::ToSeconds(LocalRenderThreadTime) >= TargetThreadTimeSeconds)
 			{
 				GNumFramesBound_RenderThread++;
 				GTotalFramesBoundTime_RenderThread += DeltaSeconds;
@@ -270,9 +273,9 @@ void UEngine::TickFPSChart( float DeltaSeconds )
 			}			
 
 			// Consider this frame GPU bound if we have an actual measurement which is over the limit,
-			if( (LocalGPUFrameTime != 0 && FPlatformTime::ToSeconds(LocalGPUFrameTime) >= TargetThreadTimeSeconds) ||
+			if (((LocalGPUFrameTime != 0) && (FPlatformTime::ToSeconds(LocalGPUFrameTime) >= TargetThreadTimeSeconds)) ||
 				// Or if we don't have a measurement but neither of the other threads were the slowest
-				(LocalGPUFrameTime == 0 && bAreWeGPUBoundIfInferred && PossibleGPUTime == MaxThreadTimeValue) )
+				((LocalGPUFrameTime == 0) && bAreWeGPUBoundIfInferred && (PossibleGPUTime == MaxThreadTimeValue)) )
 			{
 				GTotalFramesBoundTime_GPU += DeltaSeconds;
 				GNumFramesBound_GPU++;
@@ -293,47 +296,47 @@ void UEngine::TickFPSChart( float DeltaSeconds )
 	// Check for hitches
 	{
 		// Minimum time quantum before we'll even consider this a hitch
-		const float MinFrameTimeToConsiderAsHitch = 0.06f;
-
-		// Minimum time passed before we'll record a new hitch
-		const float MinTimeBetweenHitches = 0.2f;
-
-		// For the current frame to be considered a hitch, it must have run at least this many times slower than
-		// the previous frame
-		const float HitchMultiplierAmount = 1.5f;
+		const float MinFrameTimeToConsiderAsHitch = FEnginePerformanceTargets::GetHitchFrameTimeThresholdMS() * MSToSeconds;
 
 		// Ignore frames faster than our threshold
-		if( DeltaSeconds >= MinFrameTimeToConsiderAsHitch )
+		if (DeltaSeconds >= MinFrameTimeToConsiderAsHitch)
 		{
 			// How long has it been since the last hitch we detected?
 			const float TimeSinceLastHitch = ( float )( CurrentTime - LastHitchTime );
 
+			// Minimum time passed before we'll record a new hitch
+			const float MinTimeBetweenHitches = FEnginePerformanceTargets::GetMinTimeBetweenHitchesMS() * MSToSeconds;
+
 			// Make sure at least a little time has passed since the last hitch we reported
-			if( TimeSinceLastHitch >= MinTimeBetweenHitches )
+			if (TimeSinceLastHitch >= MinTimeBetweenHitches)
 			{
+				// For the current frame to be considered a hitch, it must have run at least this many times slower than
+				// the previous frame
+				const float HitchMultiplierAmount = FEnginePerformanceTargets::GetHitchToNonHitchRatio();
+
 				// If our frame time is much larger than our last frame time, we'll count this as a hitch!
-				if( DeltaSeconds > LastDeltaSeconds * HitchMultiplierAmount )
+				if (DeltaSeconds > (LastDeltaSeconds * HitchMultiplierAmount))
 				{
 					// We have a hitch!
 					
 					// Track the hitch by bucketing it based on time severity
 					const int32 HitchBucketCount = STAT_FPSChart_LastHitchBucketStat - STAT_FPSChart_FirstHitchStat;
-					for( int32 CurBucketIndex = 0; CurBucketIndex < HitchBucketCount; ++CurBucketIndex )
+					for (int32 CurBucketIndex = 0; CurBucketIndex < HitchBucketCount; ++CurBucketIndex)
 					{
-						float HitchThresholdInSeconds = ( float )GHitchThresholds[ CurBucketIndex ] * 0.001f;
-						if( DeltaSeconds >= HitchThresholdInSeconds )
+						const float HitchThresholdInSeconds = ( float )GHitchThresholds[ CurBucketIndex ] * MSToSeconds;
+						if (DeltaSeconds >= HitchThresholdInSeconds)
 						{
 							// Increment the hitch count for this bucket
 							++GHitchChart[ CurBucketIndex ].HitchCount;
 
 						
 							// Check to see what we were limited by this frame
-							if( GGameThreadTime >= (MaxThreadTimeValue - Epsilon) )
+							if (GGameThreadTime >= (MaxThreadTimeValue - EpsilonCycles))
 							{
 								// Bound by game thread
 								++GHitchChart[ CurBucketIndex ].GameThreadBoundHitchCount;
 							}
-							else if( LocalRenderThreadTime >= (MaxThreadTimeValue - Epsilon) )
+							else if (LocalRenderThreadTime >= (MaxThreadTimeValue - EpsilonCycles))
 							{
 								// Bound by render thread
 								++GHitchChart[ CurBucketIndex ].RenderThreadBoundHitchCount;
diff --git a/Engine/Source/Runtime/Engine/Private/Performance/EnginePerformanceTargets.cpp b/Engine/Source/Runtime/Engine/Private/Performance/EnginePerformanceTargets.cpp
new file mode 100644
--- /dev/null
+++ b/Engine/Source/Runtime/Engine/Private/Performance/EnginePerformanceTargets.cpp
@@ -0,0 +1,75 @@
+// Copyright 1998-2015 Epic Games, Inc. All Rights Reserved.
+
+#include "EnginePrivate.h"
+#include "Performance/EnginePerformanceTargets.h"
+
+//////////////////////////////////////////////////////////////////////
+
+// The maximum threshold for an 'OK' frame time in miliseconds (*cosmetic only* and used for fps/stat display, should not be used in scalability code)
+// Values above this will be red, values between this and the acceptable limit will be yellow, and values below will be green.
+static TAutoConsoleVariable<float> GUnacceptableFrameTimeThresholdCVar(
+	TEXT("t.UnacceptableFrameTimeThreshold"),
+	50.0f,
+	TEXT("The frame time theshold for what is considered completely unacceptable (in ms); values above this will be drawn as red\n")
+	TEXT(" default: 50.0 ms"));
+
+// The target threshold for frame time in miliseconds (*cosmetic only* and used for fps/stat display, should not be used in scalability code)
+// Values below this will be green, values between this and the unacceptable threshold will be yellow, and values above that will be red
+TAutoConsoleVariable<float> GTargetFrameTimeThresholdCVar(
+	TEXT("t.TargetFrameTimeThreshold"),
+	33.9f,
+	TEXT("The target frame time (in ms); values below this will be drawn in green, values above will be yellow or red depending on the severity\n")
+	TEXT(" default: 33.9 ms"));
+
+// The threshold that would be considered so bad that it would cause a hitch in gameplay
+// (*cosmetic only* for reporting purposes such as FPS charts, should not be used in scalability code)
+TAutoConsoleVariable<float> GHitchFrameTimeThresholdCVar(
+	TEXT("t.HitchFrameTimeThreshold"),
+	60.0f,
+	TEXT("Definition of a hitchy frame (in ms)\n")
+	TEXT(" default: 60.0 ms"));
+
+// Minimum time passed before we'll record a new hitch
+TAutoConsoleVariable<float> GHitchDeadTimeWindowCVar(
+	TEXT("t.HitchDeadTimeWindow"),
+	200.0f,
+	TEXT("Minimum time passed before we'll record a new hitch (in ms)\n")
+	TEXT(" default: 200.0 ms"));
+
+// For the current frame to be considered a hitch, it must have run at least this many times slower than the previous frame
+TAutoConsoleVariable<float> GHitchVersusNonHitchMinimumRatioCVar(
+	TEXT("t.HitchVersusNonHitchRatio"),
+	1.5f,
+	TEXT("For the current frame to be considered a hitch, it must have run at least this many times slower than the previous frame.\nThe actual ratio is clamped to be between this and t.HitchFrameTimeThreshold/t.TargetFrameTimeThreshold\n")
+	TEXT(" default: 1.5"));
+
+//////////////////////////////////////////////////////////////////////
+// FEnginePerformanceTargets
+
+float FEnginePerformanceTargets::GetTargetFrameTimeThresholdMS()
+{
+	return GTargetFrameTimeThresholdCVar.GetValueOnGameThread();
+}
+
+float FEnginePerformanceTargets::GetUnacceptableFrameTimeThresholdMS()
+{
+	return GUnacceptableFrameTimeThresholdCVar.GetValueOnGameThread();
+}
+
+float FEnginePerformanceTargets::GetHitchFrameTimeThresholdMS()
+{
+	return GHitchFrameTimeThresholdCVar.GetValueOnGameThread();
+}
+
+float FEnginePerformanceTargets::GetMinTimeBetweenHitchesMS()
+{
+	return GHitchDeadTimeWindowCVar.GetValueOnGameThread();
+}
+
+float FEnginePerformanceTargets::GetHitchToNonHitchRatio()
+{
+	const float MinimumRatio = 1.0f;
+	const float MaximumRatio = GetHitchFrameTimeThresholdMS() / GetTargetFrameTimeThresholdMS();
+
+	return FMath::Clamp(GHitchVersusNonHitchMinimumRatioCVar.GetValueOnGameThread(), MinimumRatio, MaximumRatio);
+}
diff --git a/Engine/Source/Runtime/Engine/Private/UnrealEngine.cpp b/Engine/Source/Runtime/Engine/Private/UnrealEngine.cpp
--- a/Engine/Source/Runtime/Engine/Private/UnrealEngine.cpp
+++ b/Engine/Source/Runtime/Engine/Private/UnrealEngine.cpp
@@ -118,6 +118,7 @@
 #include "MovieSceneCaptureModule.h"
 #include "GameFramework/OnlineSession.h"
 #include "ABTesting.h"
+#include "Performance/EnginePerformanceTargets.h"
 
 #if WITH_EDITOR
 #include "Editor/UnrealEd/Public/Animation/AnimationRecorder.h"
@@ -172,22 +173,6 @@ static FAutoConsoleVariableRef GHitchThresholdCVar(
 	TEXT("Time in seconds that is considered a hitch by \"stat dumphitches\"")
 	);
 
-// The maximum threshold for an 'OK' frame time in miliseconds (*cosmetic only* and used for fps/stat display, should not be used in scalability code)
-// Values above this will be red, values between this and the acceptable limit will be yellow, and values below will be green.
-static TAutoConsoleVariable<float> GUnacceptableFrameTimeThresholdCVar(
-	TEXT("t.UnacceptableFrameTimeThreshold"),
-	50.0f,
-	TEXT("The frame time theshold for what is considered completely unacceptable; values above this will be drawn as red\n")
-	TEXT(" default: 50.0 ms"));
-
-// The target threshold for frame time in miliseconds (*cosmetic only* and used for fps/stat display, should not be used in scalability code)
-// Values below this will be green, values between this and the unacceptable threshold will be yellow, and values above that will be red
-TAutoConsoleVariable<float> GTargetFrameTimeThresholdCVar(
-	TEXT("t.TargetFrameTimeThreshold"),
-	33.9f,
-	TEXT("The target frame time; values below this will be drawn in green, values above will be yellow or red depending on the severity\n")
-	TEXT(" default: 33.9 ms"));
-
 TAutoConsoleVariable<int32> CVarAllowOneFrameThreadLag(
 	TEXT("r.OneFrameThreadLag"),
 	1,
@@ -7222,8 +7207,8 @@ void UEngine::SetAverageUnitTimes(float FrameTime, float RenderThreadTime, float
 
 FColor UEngine::GetFrameTimeDisplayColor(float FrameTimeMS) const
 {
-	const float UnacceptableTime = GUnacceptableFrameTimeThresholdCVar.GetValueOnGameThread();
-	const float TargetTime = GTargetFrameTimeThresholdCVar.GetValueOnGameThread();
+	const float UnacceptableTime = FEnginePerformanceTargets::GetUnacceptableFrameTimeThresholdMS();
+	const float TargetTime = FEnginePerformanceTargets::GetTargetFrameTimeThresholdMS();
 
 	return (FrameTimeMS > UnacceptableTime) ? FColor::Red : ((FrameTimeMS > TargetTime) ? FColor::Yellow : FColor::Green);
 }
diff --git a/Engine/Source/Runtime/Engine/Public/ChartCreation.h b/Engine/Source/Runtime/Engine/Public/ChartCreation.h
--- a/Engine/Source/Runtime/Engine/Public/ChartCreation.h
+++ b/Engine/Source/Runtime/Engine/Public/ChartCreation.h
@@ -60,6 +60,7 @@ enum FPSChartStats
 	STAT_FPSChart_Hitch_150_200,
 	STAT_FPSChart_Hitch_100_150,
 	STAT_FPSChart_Hitch_60_100,
+	STAT_FPSChart_Hitch_30_60,
 	STAT_FPSChart_LastHitchBucketStat,
 	STAT_FPSChart_TotalHitchCount,
 };
diff --git a/Engine/Source/Runtime/Engine/Public/Performance/EnginePerformanceTargets.h b/Engine/Source/Runtime/Engine/Public/Performance/EnginePerformanceTargets.h
new file mode 100644
--- /dev/null
+++ b/Engine/Source/Runtime/Engine/Public/Performance/EnginePerformanceTargets.h
@@ -0,0 +1,28 @@
+// Copyright 1998-2015 Epic Games, Inc. All Rights Reserved.
+
+#pragma once
+
+struct FEnginePerformanceTargets
+{
+public:
+	// The target threshold for frame time in miliseconds (*cosmetic only* and used for fps/stat display, should not be used in scalability code)
+	// Values below this will be green, values between this and the unacceptable threshold will be yellow, and values above that will be red
+	static float GetTargetFrameTimeThresholdMS();
+
+	// The maximum threshold for an 'OK' frame time in miliseconds (*cosmetic only* and used for fps/stat display, should not be used in scalability code)
+	// Values above this will be red, values between this and the acceptable limit will be yellow, and values below will be green.
+	static float GetUnacceptableFrameTimeThresholdMS();
+
+	// The threshold that would be considered so bad that it would cause a hitch in gameplay
+	// (*cosmetic only* for reporting purposes such as FPS charts, should not be used in scalability code)
+	static float GetHitchFrameTimeThresholdMS();
+
+	// Minimum time passed before we'll record a new hitch
+	static float GetMinTimeBetweenHitchesMS();
+
+	// For the current frame to be considered a hitch, it must have run at least this many times slower than the previous frame
+	static float GetHitchToNonHitchRatio();
+
+private:
+	FEnginePerformanceTargets() {}
+};
\ No newline at end of file
