Blueprints: Exposed a way to get the full path name of a UObject from a Blueprint (GetPathName)
[INTEGRATE] Change 2703789 by [REDACTED_PERSON][REDACTED_HANDLE] on 2015/09/24 00:40:06

diff --git a/Engine/Source/Runtime/Engine/Classes/Kismet/KismetSystemLibrary.h b/Engine/Source/Runtime/Engine/Classes/Kismet/KismetSystemLibrary.h
--- a/Engine/Source/Runtime/Engine/Classes/Kismet/KismetSystemLibrary.h
+++ b/Engine/Source/Runtime/Engine/Classes/Kismet/KismetSystemLibrary.h
@@ -84,6 +84,10 @@ class ENGINE_API UKismetSystemLibrary : public UBlueprintFunctionLibrary
 	UFUNCTION(BlueprintPure, Category = "Utilities")
 	static FString GetObjectName(const UObject* Object);
 
+	// Returns the full path to the specified object.
+	UFUNCTION(BlueprintPure, Category="Utilities")
+	static FString GetPathName(const UObject* Object);
+
 	// Returns the display name (or actor label), for displaying as a debugging aid.
 	// Note: In editor builds, this is the actor label.  In non-editor builds, this is the actual object name.  This function should not be used to uniquely identify actors!
 	// It is not localized and should not be used for display to an end user of a game.
diff --git a/Engine/Source/Runtime/Engine/Private/KismetSystemLibrary.cpp b/Engine/Source/Runtime/Engine/Private/KismetSystemLibrary.cpp
--- a/Engine/Source/Runtime/Engine/Private/KismetSystemLibrary.cpp
+++ b/Engine/Source/Runtime/Engine/Private/KismetSystemLibrary.cpp
@@ -47,6 +47,11 @@ FString UKismetSystemLibrary::GetObjectName(const UObject* Object)
 	return GetNameSafe(Object);
 }
 
+FString UKismetSystemLibrary::GetPathName(const UObject* Object)
+{
+	return GetPathNameSafe(Object);
+}
+
 FString UKismetSystemLibrary::GetDisplayName(const UObject* Object)
 {
 #if WITH_EDITOR
