Engine: Added build configuration to the output .log generated for fpscharts
[INTEGRATE] Change 2703530 and 2703536 by [REDACTED_PERSON][REDACTED_HANDLE]

diff --git a/Engine/Source/Runtime/Engine/Private/ChartCreation.cpp b/Engine/Source/Runtime/Engine/Private/ChartCreation.cpp
--- a/Engine/Source/Runtime/Engine/Private/ChartCreation.cpp
+++ b/Engine/Source/Runtime/Engine/Private/ChartCreation.cpp
@@ -715,7 +715,7 @@ void UEngine::DumpFPSChartToStatsLog( float TotalTime, float DeltaTime, int32 Nu
 
 	if( OutputFile )
 	{
-		OutputFile->Logf(TEXT("Dumping FPS chart at %s using build %s built from changelist %i"), *FDateTime::Now().ToString(), *FEngineVersion::Current().ToString(), GetChangeListNumberForPerfTesting() );
+		OutputFile->Logf(TEXT("Dumping FPS chart at %s using build %s in config %s built from changelist %i"), *FDateTime::Now().ToString(), *FEngineVersion::Current().ToString(), EBuildConfigurations::ToString(FApp::GetBuildConfiguration()), GetChangeListNumberForPerfTesting());
 
 		// Get OS info
 		FString OSMajor;
