Physics: Optimization for ConvertOverlapResults when there are several overlaps (now controlled by [REDACTED_PERSON], which defaults to 2)
[INTEGRATE] Change 2654953 and 2703820

diff --git a/Engine/Source/Runtime/Engine/Private/Collision/CollisionConversions.cpp b/Engine/Source/Runtime/Engine/Private/Collision/CollisionConversions.cpp
--- a/Engine/Source/Runtime/Engine/Private/Collision/CollisionConversions.cpp
+++ b/Engine/Source/Runtime/Engine/Private/Collision/CollisionConversions.cpp
@@ -1012,6 +1012,15 @@ bool IsBlocking(const PxShape* PShape, const PxFilterData& QueryFilter)
 	return bBlock;
 }
 
+/** Min number of overlaps required before using a TMap for deduplication */
+int32 GNumOverlapsRequiredForTMap = 2;
+
+static FAutoConsoleVariableRef GTestOverlapSpeed(
+	TEXT("Engine.MinNumOverlapsToUseTMap"),
+	GNumOverlapsRequiredForTMap,
+	TEXT("Min number of overlaps required before using a TMap for deduplication")
+	);
+
 bool ConvertOverlapResults(int32 NumOverlaps, PxOverlapHit* POverlapResults, const PxFilterData& QueryFilter, TArray<FOverlapResult>& OutOverlaps)
 {
 	SCOPE_CYCLE_COUNTER(STAT_CollisionConvertOverlap);
@@ -1019,34 +1028,17 @@ bool ConvertOverlapResults(int32 NumOverlaps, PxOverlapHit* POverlapResults, con
 	OutOverlaps.Reserve(OutOverlaps.Num() + NumOverlaps);
 	bool bBlockingFound = false;
 
-	// This number was not empirically determined, just a rough rule of thumb
-	if (OutOverlaps.Num() + NumOverlaps < 6)
+	if (OutOverlaps.Max() >= GNumOverlapsRequiredForTMap)
 	{
-		// N^2 approach, no maps
-		for (int32 i = 0; i < NumOverlaps; i++)
-		{
-			FOverlapResult NewOverlap;
-			ConvertQueryOverlap(POverlapResults[i].shape, POverlapResults[i].actor, NewOverlap, QueryFilter);
-
-			if (NewOverlap.bBlockingHit)
-			{
-				bBlockingFound = true;
-			}
-
-			AddUniqueOverlap(OutOverlaps, NewOverlap);
-		}
-	}
-	else
-	{
-		// Map from an overlap to the position in the result array
+		// Map from an overlap to the position in the result array (the index has one added to it so 0 can be a sentinel)
 		TMap<FOverlapKey, int32> OverlapMap;
-		OverlapMap.Reserve(OutOverlaps.Num());
+		OverlapMap.Reserve(OutOverlaps.Max());
 
 		// Fill in the map with existing hits
 		for (int32 ExistingIndex = 0; ExistingIndex < OutOverlaps.Num(); ++ExistingIndex)
 		{
 			const FOverlapResult& ExistingOverlap = OutOverlaps[ExistingIndex];
-			OverlapMap.Add(FOverlapKey(ExistingOverlap.Component.Get(), ExistingOverlap.ItemIndex), ExistingIndex);
+			OverlapMap.Add(FOverlapKey(ExistingOverlap.Component.Get(), ExistingOverlap.ItemIndex), ExistingIndex + 1);
 		}
 
 		for (int32 PResultIndex = 0; PResultIndex < NumOverlaps; ++PResultIndex)
@@ -1059,10 +1051,15 @@ bool ConvertOverlapResults(int32 NumOverlaps, PxOverlapHit* POverlapResults, con
 				bBlockingFound = true;
 			}
 
+			// Look for it in the map, newly added elements will start with 0, so we know we need to add it to the results array then (the index is stored as +1)
 			int32& DestinationIndex = OverlapMap.FindOrAdd(FOverlapKey(NewOverlap.Component.Get(), NewOverlap.ItemIndex));
-			if (DestinationIndex < OutOverlaps.Num())
+			if (DestinationIndex == 0)
 			{
-				FOverlapResult& ExistingOverlap = OutOverlaps[DestinationIndex];
+				DestinationIndex = OutOverlaps.Add(NewOverlap) + 1;
+			}
+			else
+			{
+				FOverlapResult& ExistingOverlap = OutOverlaps[DestinationIndex - 1];
 
 				// If we had a non-blocking overlap with this component, but now we have a blocking one, use that one instead!
 				if (!ExistingOverlap.bBlockingHit && NewOverlap.bBlockingHit)
@@ -1070,15 +1067,26 @@ bool ConvertOverlapResults(int32 NumOverlaps, PxOverlapHit* POverlapResults, con
 					ExistingOverlap = NewOverlap;
 				}
 			}
-			else
+		}
+	}
+	else
+	{
+		// N^2 approach, no maps
+		for (int32 i = 0; i < NumOverlaps; i++)
+		{
+			FOverlapResult NewOverlap;
+			ConvertQueryOverlap(POverlapResults[i].shape, POverlapResults[i].actor, NewOverlap, QueryFilter);
+
+			if (NewOverlap.bBlockingHit)
 			{
-				DestinationIndex = OutOverlaps.Add(NewOverlap);
+				bBlockingFound = true;
 			}
+
+			AddUniqueOverlap(OutOverlaps, NewOverlap);
 		}
 	}
 
 	return bBlockingFound;
 }
 
-
 #endif // WITH_PHYSX
