Add -MatchAutoStatCapture command to save stats each match.

- Each time match begins/end the engine will start/stop writing to the stats file.
- Intended to aid PerfMem QA in stat collection.

(Merging CL 2708811 to the main branch).

diff --git a/Engine/Source/Runtime/Engine/Private/GameSession.cpp b/Engine/Source/Runtime/Engine/Private/GameSession.cpp
--- a/Engine/Source/Runtime/Engine/Private/GameSession.cpp
+++ b/Engine/Source/Runtime/Engine/Private/GameSession.cpp
@@ -68,6 +68,15 @@ void AGameSession::HandleMatchHasStarted()
 		StartSessionCompleteHandle = SessionInt->AddOnStartSessionCompleteDelegate_Handle(FOnStartSessionCompleteDelegate::CreateUObject(this, &AGameSession::OnStartSessionComplete));
 		SessionInt->StartSession(SessionName);
 	}
+
+	if (STATS && !UE_BUILD_SHIPPING)
+	{
+		if (FParse::Param(FCommandLine::Get(), TEXT("MatchAutoStatCapture")))
+		{
+			UE_LOG(LogGameSession, Log, TEXT("Match has started - begin automatic stat capture"));
+			GEngine->Exec(GetWorld(), TEXT("stat startfile"));
+		}
+	}
 }
 
 void AGameSession::OnStartSessionComplete(FName InSessionName, bool bWasSuccessful)
@@ -83,6 +92,15 @@ void AGameSession::OnStartSessionComplete(FName InSessionName, bool bWasSuccessf
 
 void AGameSession::HandleMatchHasEnded()
 {
+	if (STATS && !UE_BUILD_SHIPPING)
+	{
+		if (FParse::Param(FCommandLine::Get(), TEXT("MatchAutoStatCapture")))
+		{
+			UE_LOG(LogGameSession, Log, TEXT("Match has ended - end automatic stat capture"));
+			GEngine->Exec(GetWorld(), TEXT("stat stopfile"));
+		}
+	}
+
 	UWorld* World = GetWorld();
 	IOnlineSessionPtr SessionInt = Online::GetSessionInterface(World);
 	if (SessionInt.IsValid())
