Fixed large chunk classification incorrectly picking up wrong bounds.

diff --git a/Engine/Source/Runtime/Engine/Classes/Components/DestructibleComponent.h b/Engine/Source/Runtime/Engine/Classes/Components/DestructibleComponent.h
--- a/Engine/Source/Runtime/Engine/Classes/Components/DestructibleComponent.h
+++ b/Engine/Source/Runtime/Engine/Classes/Components/DestructibleComponent.h
@@ -240,13 +240,12 @@ private:
 
 	void SetCollisionResponseForShape(physx::PxShape* Shape, int32 ChunkIdx);
 	void SetCollisionResponseForActor(physx::PxRigidDynamic* Actor, int32 ChunkIdx, const FCollisionResponseContainer* ResponseOverride = NULL);
-	void SetCollisionResponseForAllActors(const FCollisionResponseContainer& ResponseOverride);
 
 
 public:
 	/** User data wrapper for the chunks passed to physx */
 	TArray<FPhysxUserData> PhysxChunkUserData;
-	bool IsChunkLarge(int32 ChunkIdx) const;
+	bool IsChunkLarge(physx::PxRigidActor* ChunkActor) const;
 #endif
 };
 
diff --git a/Engine/Source/Runtime/Engine/Private/Components/DestructibleComponent.cpp b/Engine/Source/Runtime/Engine/Private/Components/DestructibleComponent.cpp
--- a/Engine/Source/Runtime/Engine/Private/Components/DestructibleComponent.cpp
+++ b/Engine/Source/Runtime/Engine/Private/Components/DestructibleComponent.cpp
@@ -1392,11 +1392,11 @@ bool UDestructibleComponent::IsAnySimulatingPhysics() const
 
 #if WITH_PHYSX
 
-bool UDestructibleComponent::IsChunkLarge(int32 ChunkIdx) const
+bool UDestructibleComponent::IsChunkLarge(PxRigidActor* ChunkActor) const
 {
 #if WITH_APEX
-	check(ApexDestructibleActor);
-	physx::PxBounds3 Bounds = ApexDestructibleActor->getChunkBounds(ChunkIdx);
+	check(ChunkActor);
+	physx::PxBounds3 Bounds = ChunkActor->getWorldBounds();
 	return Bounds.getExtents().maxElement() > LargeChunkThreshold;
 #else
 	return true;
@@ -1418,7 +1418,7 @@ void UDestructibleComponent::SetCollisionResponseForActor(PxRigidDynamic* Actor,
 	{
 		UDestructibleMesh* TheDestructibleMesh = GetDestructibleMesh();
 		AActor* Owner = GetOwner();
-		bool bLargeChunk = IsChunkLarge(ChunkIdx);
+		bool bLargeChunk = IsChunkLarge(Actor);
 		const FCollisionResponseContainer& UseResponse = ResponseOverride == NULL ? (bLargeChunk ? LargeChunkCollisionResponse.GetResponseContainer() : SmallChunkCollisionResponse.GetResponseContainer()) : *ResponseOverride;
 
 		physx::PxU32 SupportDepth = TheDestructibleMesh->ApexDestructibleAsset->getChunkDepth(ChunkIdx);
@@ -1449,51 +1449,6 @@ void UDestructibleComponent::SetCollisionResponseForActor(PxRigidDynamic* Actor,
 #endif
 }
 
-void UDestructibleComponent::SetCollisionResponseForAllActors(const FCollisionResponseContainer& ResponseOverride)
-{
-#if WITH_APEX
-	if (ApexDestructibleActor == NULL)
-	{
-		return;
-	}
-
-	PxRigidDynamic** PActorBuffer = NULL;
-	PxU32 PActorCount = 0;
-	if (ApexDestructibleActor->acquirePhysXActorBuffer(PActorBuffer, PActorCount))
-	{
-		PxScene* LockedScene = NULL;
-
-		while (PActorCount--)
-		{
-			PxRigidDynamic* PActor = *PActorBuffer++;
-			if (PActor != NULL)
-			{
-				FDestructibleChunkInfo* ChunkInfo = FPhysxUserData::Get<FDestructibleChunkInfo>(PActor->userData);
-				if (ChunkInfo != NULL)
-				{
-					if (!LockedScene)
-					{
-						LockedScene = PActor->getScene();
-						LockedScene->lockWrite();
-						LockedScene->lockRead();
-					}
-					SetCollisionResponseForActor(PActor, ChunkInfo->ChunkIndex, &ResponseOverride);	// ChunkIndex is the last chunk made visible.  But SetCollisionResponseForActor already doesn't respect per-chunk collision properties.
-				}
-			}
-		}
-
-		if (LockedScene)
-		{
-			LockedScene->unlockRead();
-			LockedScene->unlockWrite();
-			LockedScene = NULL;
-		}
-
-		ApexDestructibleActor->releasePhysXActorBuffer();
-	}
-#endif
-}
-
 void UDestructibleComponent::SetCollisionResponseForShape(PxShape* Shape, int32 ChunkIdx)
 {
 #if WITH_APEX
@@ -1509,7 +1464,7 @@ void UDestructibleComponent::SetCollisionResponseForShape(PxShape* Shape, int32
 	if (IsCollisionEnabled())
 	{
 		AActor* Owner = GetOwner();
-		bool bLargeChunk = IsChunkLarge(ChunkIdx);
+		bool bLargeChunk = IsChunkLarge(Shape->getActor());
 		const FCollisionResponse& ColResponse = bLargeChunk ? LargeChunkCollisionResponse : SmallChunkCollisionResponse;
 		//TODO: we currently assume chunks will not have impact damage as it's very expensive. Should look into exposing this a bit more
 		CreateShapeFilterData(MoveChannel, (Owner ? Owner->GetUniqueID() : 0), ColResponse.GetResponseContainer(), 0, ChunkIdxToBoneIdx(ChunkIdx), PQueryFilterData, PSimFilterData, BodyInstance.bUseCCD, false, false);
