Merging CL 2709066:

Added NativeUninitializeAnimation override point

diff --git a/Engine/Source/Runtime/Engine/Classes/Animation/AnimInstance.h b/Engine/Source/Runtime/Engine/Classes/Animation/AnimInstance.h
--- a/Engine/Source/Runtime/Engine/Classes/Animation/AnimInstance.h
+++ b/Engine/Source/Runtime/Engine/Classes/Animation/AnimInstance.h
@@ -842,6 +842,8 @@ public:
 	virtual bool NativeEvaluateAnimation(FPoseContext& Output);
 	// Native Post Evaluate override point
 	virtual void NativePostEvaluateAnimation();
+	// Native Uninitialize override point
+	virtual void NativeUninitializeAnimation();
 
 	// Sets up a native transition delegate between states with PrevStateName and NextStateName, in the state machine with name MachineName.
 	// Note that a transition already has to exist for this to succeed
diff --git a/Engine/Source/Runtime/Engine/Private/Animation/AnimInstance.cpp b/Engine/Source/Runtime/Engine/Private/Animation/AnimInstance.cpp
--- a/Engine/Source/Runtime/Engine/Private/Animation/AnimInstance.cpp
+++ b/Engine/Source/Runtime/Engine/Private/Animation/AnimInstance.cpp
@@ -290,6 +290,8 @@ void UAnimInstance::InitializeAnimation()
 
 void UAnimInstance::UninitializeAnimation()
 {
+	NativeUninitializeAnimation();
+
 	StopAllMontages(0.f);
 
 	for(int32 Index = 0; Index < MontageInstances.Num(); ++Index)
@@ -616,6 +618,10 @@ void UAnimInstance::NativePostEvaluateAnimation()
 {
 }
 
+void UAnimInstance::NativeUninitializeAnimation()
+{
+}
+
 void UAnimInstance::AddNativeTransitionBinding(const FName& MachineName, const FName& PrevStateName, const FName& NextStateName, const FCanTakeTransition& NativeTransitionDelegate)
 {
 	NativeTransitionBindings.Add(FNativeTransitionBinding(MachineName, PrevStateName, NextStateName, NativeTransitionDelegate));
