Fixed game specific net control messages not working if they are received during the pending phase

Merging 2708721

//depot/UE4-UT/...

to //depot/UE4/...

diff --git a/Engine/Source/Runtime/Engine/Private/DataChannel.cpp b/Engine/Source/Runtime/Engine/Private/DataChannel.cpp
--- a/Engine/Source/Runtime/Engine/Private/DataChannel.cpp
+++ b/Engine/Source/Runtime/Engine/Private/DataChannel.cpp
@@ -1190,6 +1190,14 @@ void UControlChannel::ReceivedBunch( FInBunch& Bunch )
 			{
 				Connection->Driver->World->GetGameInstance()->HandleGameNetControlMessage(Connection, MessageByte, MessageStr);
 			}
+			else
+			{
+				FWorldContext* Context = GEngine->GetWorldContextFromPendingNetGameNetDriver(Connection->Driver);
+				if (Context != NULL && Context->OwningGameInstance != NULL)
+				{
+					Context->OwningGameInstance->HandleGameNetControlMessage(Connection, MessageByte, MessageStr);
+				}
+			}
 		}
 		else if(MessageType == NMT_SecurityViolation)
 		{
