Remove inaccurate comment

diff --git a/Engine/Source/Runtime/Engine/Classes/Engine/Engine.h b/Engine/Source/Runtime/Engine/Classes/Engine/Engine.h
--- a/Engine/Source/Runtime/Engine/Classes/Engine/Engine.h
+++ b/Engine/Source/Runtime/Engine/Classes/Engine/Engine.h
@@ -2076,7 +2076,6 @@ public:
 
 	/** 
 	 * Obtain a world object pointer from an object with has a world context.
-	 * This should be be overridden to cater for game specific object types that do not derive from the Actor class.
 	 *
 	 * @param Object		Object whose owning world we require.
 	 * @param bChecked      Allows calling function to specify not to do ensure check and that a nullptr return value is acceptable
