PR #1605: Do not crash when PlayerController is given a LocalPlayer in commandlet mode (Contributed by [REDACTED_PERSON])

diff --git a/Engine/Source/Runtime/Engine/Private/PlayerController.cpp b/Engine/Source/Runtime/Engine/Private/PlayerController.cpp
--- a/Engine/Source/Runtime/Engine/Private/PlayerController.cpp
+++ b/Engine/Source/Runtime/Engine/Private/PlayerController.cpp
@@ -1366,11 +1366,14 @@ void APlayerController::EndPlay(const EEndPlayReason::Type EndPlayReason)
 			ActivateTouchInterface(nullptr);
 		}
 
-		// Stop any force feedback effects that may be active
-		IInputInterface* InputInterface = FSlateApplication::Get().GetInputInterface();
-		if (InputInterface)
+		if (FSlateApplication::IsInitialized())
 		{
-			InputInterface->SetForceFeedbackChannelValues(LocalPlayer->GetControllerId(), FForceFeedbackValues());
+			IInputInterface* InputInterface = FSlateApplication::Get().GetInputInterface();
+			if (InputInterface)
+			{
+				// Stop any force feedback effects that may be active
+				InputInterface->SetForceFeedbackChannelValues(LocalPlayer->GetControllerId(), FForceFeedbackValues());
+			}
 		}
 	}
 
@@ -3532,10 +3535,13 @@ void APlayerController::ProcessForceFeedback(const float DeltaTime, const bool b
 		}
 	}
 
-	IInputInterface* InputInterface = FSlateApplication::Get().GetInputInterface();
-	if (InputInterface)
+	if (FSlateApplication::IsInitialized())
 	{
-		InputInterface->SetForceFeedbackChannelValues(CastChecked<ULocalPlayer>(Player)->GetControllerId(), (bForceFeedbackEnabled ? ForceFeedbackValues : FForceFeedbackValues()));
+		IInputInterface* InputInterface = FSlateApplication::Get().GetInputInterface();
+		if (InputInterface)
+		{
+			InputInterface->SetForceFeedbackChannelValues(CastChecked<ULocalPlayer>(Player)->GetControllerId(), (bForceFeedbackEnabled ? ForceFeedbackValues : FForceFeedbackValues()));
+		}
 	}
 }
 
