disable stats gathering if shipping or physx release

like the pull request but also adding the condidition for the relase libs since stats would not be valid in that case.

original info:
 PR #1551: GatherPhysXStats disabled in shipping (Contributed by [REDACTED_PERSON])
 was from p4 changelist  2689654  <- old changelist number
 Jira   https://jira.ol.epicgames.net/browse/UE-21042
 # UE-21042

diff --git a/Engine/Source/Runtime/Engine/Private/PhysicsEngine/PhysScene.cpp b/Engine/Source/Runtime/Engine/Private/PhysicsEngine/PhysScene.cpp
--- a/Engine/Source/Runtime/Engine/Private/PhysicsEngine/PhysScene.cpp
+++ b/Engine/Source/Runtime/Engine/Private/PhysicsEngine/PhysScene.cpp
@@ -668,12 +668,11 @@ void FPhysScene::TickPhysScene(uint32 SceneType, FGraphEventRef& InOutCompletion
 		return;
 	}
 
-#if WITH_PHYSX
+#if ( WITH_PHYSX  && !(UE_BUILD_SHIPPING || WITH_PHYSX_RELEASE))
 	GatherPhysXStats(GetPhysXScene(SceneType), SceneType);
 #endif
 
 
-
 	/**
 	* Weight frame time according to PhysScene settings.
 	*/
