Fix crash where parallel cloth is being set as evaluating even though it's not running in parallel.

diff --git a/Engine/Source/Runtime/Engine/Private/PhysicsEngine/ClothManager.cpp b/Engine/Source/Runtime/Engine/Private/PhysicsEngine/ClothManager.cpp
--- a/Engine/Source/Runtime/Engine/Private/PhysicsEngine/ClothManager.cpp
+++ b/Engine/Source/Runtime/Engine/Private/PhysicsEngine/ClothManager.cpp
@@ -18,11 +18,12 @@ void FClothManagerData::PrepareCloth(float DeltaTime, FTickFunction& TickFunctio
 	QUICK_SCOPE_CYCLE_COUNTER(STAT_FClothManager_PrepareCloth);
 	if (SkeletalMeshComponents.Num())
 	{
-		IsPreparingCloth.AtomicSet(true);
+		const bool bParallelCloth = !!CVarParallelCloth.GetValueOnAnyThread();
+		IsPreparingClothInParallel.AtomicSet(bParallelCloth);
 		for (USkeletalMeshComponent* SkeletalMeshComponent : SkeletalMeshComponents)
 		{
 			FClothSimulationContext& ClothSimulationContext = SkeletalMeshComponent->InternalClothSimulationContext;
-			if(CVarParallelCloth.GetValueOnAnyThread())
+			if (bParallelCloth)
 			{
 				SkeletalMeshComponent->ParallelTickClothing(DeltaTime, ClothSimulationContext);
 			}else
@@ -33,7 +34,7 @@ void FClothManagerData::PrepareCloth(float DeltaTime, FTickFunction& TickFunctio
 			
 		}
 
-		IsPreparingCloth.AtomicSet(false);
+		IsPreparingClothInParallel.AtomicSet(false);
 	}
 #endif
 }
diff --git a/Engine/Source/Runtime/Engine/Private/PhysicsEngine/ClothManager.h b/Engine/Source/Runtime/Engine/Private/PhysicsEngine/ClothManager.h
--- a/Engine/Source/Runtime/Engine/Private/PhysicsEngine/ClothManager.h
+++ b/Engine/Source/Runtime/Engine/Private/PhysicsEngine/ClothManager.h
@@ -23,7 +23,7 @@ class FClothManagerData
 {
 public:
 	TArray<USkeletalMeshComponent*> SkeletalMeshComponents;
-	FThreadSafeBool IsPreparingCloth;
+	FThreadSafeBool IsPreparingClothInParallel;
 	FGraphEventRef PrepareCompletion;
 
 	/** Go through all SkeletalMeshComponents registered for work this frame and call the solver as needed*/
@@ -76,7 +76,7 @@ public:
 		bool bIsPreparingCloth = false;
 		for(const FClothManagerData& PrepareCloth : PrepareClothDataArray)
 		{
-			return bIsPreparingCloth |= PrepareCloth.IsPreparingCloth;
+			return bIsPreparingCloth |= PrepareCloth.IsPreparingClothInParallel;
 		}
 
 		return bIsPreparingCloth;
