Don't tick matinee actor unless it is playing and not paused
Fix issues with registering tick functions disabling a tick function that generally does not start enabled but has been enabled/activated before register all tick functions occurs.
- Since FTickFunction by default starts enabled, we need to mark it disabled for Actors and ActorComponents so that we can determine whether someone has intentionally activated it prior to register tick functions being called, otherwise the start with tick enabled being false will cause the tick function to be disabled when it was intended to be enabled.

diff --git a/Engine/Source/Runtime/Engine/Private/Actor.cpp b/Engine/Source/Runtime/Engine/Private/Actor.cpp
--- a/Engine/Source/Runtime/Engine/Private/Actor.cpp
+++ b/Engine/Source/Runtime/Engine/Private/Actor.cpp
@@ -62,6 +62,7 @@ void AActor::InitializeDefaults()
 	// Default to no tick function, but if we set 'never ticks' to false (so there is a tick function) it is enabled by default
 	PrimaryActorTick.bCanEverTick = false;
 	PrimaryActorTick.bStartWithTickEnabled = true;
+	PrimaryActorTick.SetTickFunctionEnable(false); 
 
 	CustomTimeDilation = 1.0f;
 
@@ -636,7 +637,7 @@ void AActor::RegisterActorTickFunctions(bool bRegister)
 		if(PrimaryActorTick.bCanEverTick)
 		{
 			PrimaryActorTick.Target = this;
-			PrimaryActorTick.SetTickFunctionEnable(PrimaryActorTick.bStartWithTickEnabled);
+			PrimaryActorTick.SetTickFunctionEnable(PrimaryActorTick.bStartWithTickEnabled || PrimaryActorTick.IsTickFunctionEnabled());
 			PrimaryActorTick.RegisterTickFunction(GetLevel());
 		}
 	}
diff --git a/Engine/Source/Runtime/Engine/Private/Components/ActorComponent.cpp b/Engine/Source/Runtime/Engine/Private/Components/ActorComponent.cpp
--- a/Engine/Source/Runtime/Engine/Private/Components/ActorComponent.cpp
+++ b/Engine/Source/Runtime/Engine/Private/Components/ActorComponent.cpp
@@ -116,6 +116,7 @@ UActorComponent::UActorComponent(const FObjectInitializer& ObjectInitializer /*=
 	PrimaryComponentTick.TickGroup = TG_DuringPhysics;
 	PrimaryComponentTick.bStartWithTickEnabled = true;
 	PrimaryComponentTick.bCanEverTick = false;
+	PrimaryComponentTick.SetTickFunctionEnable(false);
 
 	CreationMethod = EComponentCreationMethod::Native;
 
@@ -703,7 +704,7 @@ bool UActorComponent::SetupActorComponentTickFunction(struct FTickFunction* Tick
 		if (!MyOwner || !MyOwner->IsTemplate())
 		{
 			ULevel* ComponentLevel = (MyOwner ? MyOwner->GetLevel() : GetWorld()->PersistentLevel);
-			TickFunction->SetTickFunctionEnable(TickFunction->bStartWithTickEnabled);
+			TickFunction->SetTickFunctionEnable(TickFunction->bStartWithTickEnabled || TickFunction->IsTickFunctionEnabled());
 			TickFunction->RegisterTickFunction(ComponentLevel);
 			return true;
 		}
diff --git a/Engine/Source/Runtime/Engine/Private/Interpolation.cpp b/Engine/Source/Runtime/Engine/Private/Interpolation.cpp
--- a/Engine/Source/Runtime/Engine/Private/Interpolation.cpp
+++ b/Engine/Source/Runtime/Engine/Private/Interpolation.cpp
@@ -361,6 +361,7 @@ AMatineeActor::AMatineeActor(const FObjectInitializer& ObjectInitializer)
 #endif // WITH_EDITORONLY_DATA
 
 	PrimaryActorTick.bCanEverTick = true;
+	PrimaryActorTick.bStartWithTickEnabled = false;
 	PrimaryActorTick.TickGroup = TG_PrePhysics;
 	bPlayOnLevelLoad = false;
 #if WITH_EDITORONLY_DATA
@@ -471,6 +472,7 @@ void AMatineeActor::Play()
 	bReversePlayback = false;
 	bIsPlaying = true;
 	bPaused = false;
+	SetActorTickEnabled(true);
 }
 
 void AMatineeActor::Reverse()
@@ -484,6 +486,7 @@ void AMatineeActor::Reverse()
 	bReversePlayback = true;
 	bIsPlaying = true;
 	bPaused = false;
+	SetActorTickEnabled(true);
 }
 
 void AMatineeActor::Stop()
@@ -493,6 +496,7 @@ void AMatineeActor::Stop()
 
 	bIsPlaying = false;
 	bPaused = false;
+	SetActorTickEnabled(false);
 
 	if (GetWorld()->IsGameWorld())
 	{
@@ -507,6 +511,7 @@ void AMatineeActor::Pause()
 	{
 		EnableRadioFilter();
 		bPaused = !bPaused;
+		SetActorTickEnabled(!bPaused);
 	}
 }
 
@@ -515,6 +520,7 @@ void AMatineeActor::ChangePlaybackDirection()
 	bReversePlayback = !bReversePlayback;
 	bIsPlaying = true;
 	bPaused = false;
+	SetActorTickEnabled(true);
 }
 
 void AMatineeActor::SetLoopingState(bool bNewLooping)
diff --git a/Engine/Source/Runtime/Engine/Private/Timeline.cpp b/Engine/Source/Runtime/Engine/Private/Timeline.cpp
--- a/Engine/Source/Runtime/Engine/Private/Timeline.cpp
+++ b/Engine/Source/Runtime/Engine/Private/Timeline.cpp
@@ -589,16 +589,12 @@ void UTimelineComponent::TickComponent(float DeltaTime, enum ELevelTick TickType
 void UTimelineComponent::Activate(bool bReset)
 {
 	Super::Activate(bReset);
-
-	PrimaryComponentTick.bStartWithTickEnabled = true;
 	PrimaryComponentTick.SetTickFunctionEnable(true);
 }
 
 void UTimelineComponent::Deactivate()
 {
 	Super::Deactivate();
-
-	PrimaryComponentTick.bStartWithTickEnabled = false;
 	PrimaryComponentTick.SetTickFunctionEnable(false);
 }
 
