#merge: Adding logs to help diagnose timeouts

diff --git a/Engine/Source/Runtime/Engine/Classes/Engine/NetConnection.h b/Engine/Source/Runtime/Engine/Classes/Engine/NetConnection.h
--- a/Engine/Source/Runtime/Engine/Classes/Engine/NetConnection.h
+++ b/Engine/Source/Runtime/Engine/Classes/Engine/NetConnection.h
@@ -619,6 +619,11 @@ public:
 	 */
 	ENGINE_API bool TrackLogsPerSecond();
 
+	/**
+	* Return current timeout value that should be used
+	*/
+	float GetTimeoutValue();
+
 protected:
 
 	void CleanupDormantActorState();
diff --git a/Engine/Source/Runtime/Engine/Classes/Engine/NetDriver.h b/Engine/Source/Runtime/Engine/Classes/Engine/NetDriver.h
--- a/Engine/Source/Runtime/Engine/Classes/Engine/NetDriver.h
+++ b/Engine/Source/Runtime/Engine/Classes/Engine/NetDriver.h
@@ -200,7 +200,11 @@ public:
 	/** Accumulated time for the net driver, updated by Tick */
 	UPROPERTY()
 	float						Time;
-	/** If true then client connections are to other client peers */ 
+
+	/** Last realtime a tick dispatch occurred. Used currently to try and diagnose timeout issues */
+	double						LastTickDispatchRealtime;
+
+	/** If true then client connections are to other client peers */
 	bool						bIsPeer;
 	/** @todo document */
 	bool						ProfileStats;
diff --git a/Engine/Source/Runtime/Engine/Private/NetConnection.cpp b/Engine/Source/Runtime/Engine/Private/NetConnection.cpp
--- a/Engine/Source/Runtime/Engine/Private/NetConnection.cpp
+++ b/Engine/Source/Runtime/Engine/Private/NetConnection.cpp
@@ -84,14 +84,15 @@ void UNetConnection::InitBase(UNetDriver* InDriver,class FSocket* InSocket, cons
 	Driver = InDriver;
 
 	// Stats
-	StatUpdateTime = Driver->Time;
-	LastReceiveTime = Driver->Time;
-	LastReceiveRealtime = FPlatformTime::Seconds();
-	LastGoodPacketRealtime = FPlatformTime::Seconds();
-	LastSendTime = Driver->Time;
-	LastTickTime = Driver->Time;
-	LastRecvAckTime = Driver->Time;
-	ConnectTime = Driver->Time;
+	StatUpdateTime			= Driver->Time;
+	LastReceiveTime			= Driver->Time;
+	LastReceiveRealtime		= FPlatformTime::Seconds();
+	LastGoodPacketRealtime	= FPlatformTime::Seconds();
+	LastTime				= FPlatformTime::Seconds();
+	LastSendTime			= Driver->Time;
+	LastTickTime			= Driver->Time;
+	LastRecvAckTime			= Driver->Time;
+	ConnectTime				= Driver->Time;
 
 	// Current state
 	State = InState;
@@ -665,9 +666,20 @@ void UNetConnection::ReceivedPacket( FBitReader& Reader )
 
 	ValidateSendBuffer();
 
+	const double RealtimeSeconds	= FPlatformTime::Seconds();
+	const float DeltaTime			= Driver->Time - LastReceiveTime;
+	const float DeltaRealtime		= RealtimeSeconds - LastReceiveRealtime;
+
 	// Update receive time to avoid timeout.
-	LastReceiveTime = Driver->Time;
-	LastReceiveRealtime = FPlatformTime::Seconds();
+	LastReceiveTime		= Driver->Time;
+	LastReceiveRealtime = RealtimeSeconds;
+
+	const float Timeout = GetTimeoutValue();
+
+	if ( DeltaTime > Timeout || DeltaRealtime > Timeout )
+	{
+		UE_LOG( LogNetTraffic, Warning, TEXT( "UNetConnection::ReceivedPacket: Large time between packets. DeltaTime: %2.2f, Realtime: %2.2f, DriverTime: %2.2f, Threshold: %2.2f %s" ), DeltaTime, DeltaRealtime, Driver->Time, Timeout, *Describe() );
+	}
 
 	// Check packet ordering.
 	const int32 PacketId = InternalAck ? InPacketId + 1 : MakeRelative(Reader.ReadInt(MAX_PACKETID),InPacketId,MAX_PACKETID);
@@ -677,7 +689,7 @@ void UNetConnection::ReceivedPacket( FBitReader& Reader )
 		
 		if ( PacketsLost > 10 )
 		{
-			UE_LOG( LogNetTraffic, Warning, TEXT( "High single frame packet loss: %i" ), PacketsLost );
+			UE_LOG( LogNetTraffic, Warning, TEXT( "High single frame packet loss. PacketsLost: %i %s" ), PacketsLost, *Describe() );
 		}
 
 		InPacketsLost += PacketsLost;
@@ -1323,6 +1335,17 @@ UVoiceChannel* UNetConnection::GetVoiceChannel()
 		(UVoiceChannel*)Channels[VOICE_CHANNEL_INDEX] : NULL;
 }
 
+float UNetConnection::GetTimeoutValue()
+{
+	float Timeout = Driver->InitialConnectTimeout;
+	if ( ( State != USOCK_Pending ) && ( bPendingDestroy || ( OwningActor && OwningActor->UseShortConnectTimeout() ) ) )
+	{
+		Timeout = bPendingDestroy ? 2.f : Driver->ConnectionTimeout;
+	}
+
+	return Timeout;
+}
+
 void UNetConnection::Tick()
 {
 	AssertValid();
@@ -1399,16 +1422,21 @@ void UNetConnection::Tick()
 	}
 
 	// Compute time passed since last update.
-	float DeltaTime     = Driver->Time - LastTickTime;
-	LastTickTime        = Driver->Time;
+	const float DeltaTime	= Driver->Time - LastTickTime;
+	LastTickTime			= Driver->Time;
 
-	// Handle timeouts.
-	float Timeout = Driver->InitialConnectTimeout;
-	if ((State != USOCK_Pending) && (bPendingDestroy || (OwningActor && OwningActor->UseShortConnectTimeout())))
+	// Check to see if too much time is passing between ticks
+	// Setting this to somewhat large value for now, but small enough to catch blocking calls that are causing timeouts
+	const float TickWarnThreshold = 5.0f;
+
+	if ( DeltaTime > TickWarnThreshold || FrameTime > TickWarnThreshold )
 	{
-		Timeout = bPendingDestroy ? 2.f : Driver->ConnectionTimeout;
+		UE_LOG( LogNet, Warning, TEXT( "UNetConnection::Tick: Very long time between ticks. DeltaTime: %2.2f, Realtime: %2.2f %s" ), DeltaTime, FrameTime, *Describe() );
 	}
 
+	// Handle timeouts.
+	const float Timeout = GetTimeoutValue();
+
 #if !UE_BUILD_SHIPPING
 	if (!Driver->bNoTimeouts && (Driver->Time - LastReceiveTime) > Timeout)
 #else
@@ -1422,10 +1450,11 @@ void UNetConnection::Tick()
 		const float GoodRealtimeDelta		= Seconds - LastGoodPacketRealtime;
 
 		// Timeout.
-		FString Error = FString::Printf(TEXT("UNetConnection::Tick: Connection TIMED OUT. Closing connection. Elapsed: %2.2f, Real: %2.2f, Good: %2.2f, Threshold: %2.2f, %s"),
+		FString Error = FString::Printf(TEXT("UNetConnection::Tick: Connection TIMED OUT. Closing connection. Elapsed: %2.2f, Real: %2.2f, Good: %2.2f, DriverTime: %2.2f, Threshold: %2.2f, %s"),
 			Driver->Time - LastReceiveTime,
 			ReceiveRealtimeDelta,
 			GoodRealtimeDelta,
+			Driver->Time,
 			Timeout,
 			*Describe());
 		UE_LOG(LogNet, Warning, TEXT("%s"), *Error);
diff --git a/Engine/Source/Runtime/Engine/Private/NetworkDriver.cpp b/Engine/Source/Runtime/Engine/Private/NetworkDriver.cpp
--- a/Engine/Source/Runtime/Engine/Private/NetworkDriver.cpp
+++ b/Engine/Source/Runtime/Engine/Private/NetworkDriver.cpp
@@ -116,6 +116,7 @@ UNetDriver::UNetDriver(const FObjectInitializer& ObjectInitializer)
 ,   World(nullptr)
 ,   Notify(nullptr)
 ,	Time( 0.f )
+,	LastTickDispatchRealtime( 0.f )
 ,   bIsPeer(false)
 ,	InBytes(0)
 ,	OutBytes(0)
@@ -601,6 +602,7 @@ bool UNetDriver::InitConnectionClass(void)
 
 bool UNetDriver::InitBase(bool bInitAsClient, FNetworkNotify* InNotify, const FURL& URL, bool bReuseAddressAndPort, FString& Error)
 {
+	LastTickDispatchRealtime = FPlatformTime::Seconds();
 	bool bSuccess = InitConnectionClass();
 	Notify = InNotify;
 	return bSuccess;
@@ -696,6 +698,21 @@ void UNetDriver::TickDispatch( float DeltaTime )
 {
 	SendCycles=RecvCycles=0;
 
+	const double CurrentRealtime = FPlatformTime::Seconds();
+
+	const float DeltaRealtime = CurrentRealtime - LastTickDispatchRealtime;
+
+	LastTickDispatchRealtime = CurrentRealtime;
+
+	// Check to see if too much time is passing between ticks
+	// Setting this to somewhat large value for now, but small enough to catch blocking calls that are causing timeouts
+	const float TickWarnThreshold = 5.0f;
+
+	if ( DeltaTime > TickWarnThreshold || DeltaRealtime > TickWarnThreshold )
+	{
+		UE_LOG( LogNet, Warning, TEXT( "UNetDriver::TickDispatch: Very long time between ticks. DeltaTime: %2.2f, Realtime: %2.2f. %s" ), DeltaTime, DeltaRealtime, *GetName() );
+	}
+
 	// Get new time.
 	Time += DeltaTime;
 
