Engine: Fixed uninitialized variable in FAnimUpdateRateParameters (OptimizeMode now defaults to TrailMode)
[INTEGRATE] Change 2710571 by [REDACTED_PERSON][REDACTED_HANDLE] on 2015/09/29 22:03:10

diff --git a/Engine/Source/Runtime/Engine/Classes/Engine/EngineTypes.h b/Engine/Source/Runtime/Engine/Classes/Engine/EngineTypes.h
--- a/Engine/Source/Runtime/Engine/Classes/Engine/EngineTypes.h
+++ b/Engine/Source/Runtime/Engine/Classes/Engine/EngineTypes.h
@@ -1933,7 +1933,8 @@ public:
 
 	/** Default constructor. */
 	FAnimUpdateRateParameters()
-		: UpdateRate(1)
+		: OptimizeMode(TrailMode)
+		, UpdateRate(1)
 		, EvaluationRate(1)
 		, bInterpolateSkippedFrames(false)
 		, bSkipUpdate(false)
