Engine: Switched 'stat hitches' display to use the t.HitchFrameTimeThreshold value, and updated coloring to use the related coloring thresholds (at 25% of the hitch value) instead of using hardcoded yellow at 100 ms and red at 200 ms
[INTEGRATE] Change 2710575 by [REDACTED_PERSON][REDACTED_HANDLE] on 2015/09/29 22:08:29

diff --git a/Engine/Source/Runtime/Engine/Private/UnrealClient.cpp b/Engine/Source/Runtime/Engine/Private/UnrealClient.cpp
--- a/Engine/Source/Runtime/Engine/Private/UnrealClient.cpp
+++ b/Engine/Source/Runtime/Engine/Private/UnrealClient.cpp
@@ -21,6 +21,7 @@
 #include "SceneUtils.h"
 #include "MovieSceneCaptureModule.h"
 #include "NotificationManager.h"
+#include "Performance/EnginePerformanceTargets.h"
 
 DEFINE_LOG_CATEGORY_STATIC(LogClient, Log, All);
 
@@ -619,8 +620,10 @@ int32 FStatHitchesData::DrawStat(FViewport* InViewport, FCanvas* InCanvas, int32
 	const double CurrentTime = FPlatformTime::Seconds();
 	if (LastTime > 0)
 	{
-		float DeltaSeconds = CurrentTime - LastTime;
-		if (DeltaSeconds > GHitchThreshold)
+		const float HitchThresholdSecs = FEnginePerformanceTargets::GetHitchFrameTimeThresholdMS() * 0.001f;
+
+		const float DeltaSeconds = CurrentTime - LastTime;
+		if (DeltaSeconds > HitchThresholdSecs)
 		{
 			Hitches[OverwriteIndex] = DeltaSeconds;
 			When[OverwriteIndex] = CurrentTime;
@@ -639,17 +642,21 @@ int32 FStatHitchesData::DrawStat(FViewport* InViewport, FCanvas* InCanvas, int32
 				UE_LOG(LogClient, Warning, TEXT("HITCH %d              running cnt = %5d"), int32(DeltaSeconds * 1000), Count++);
 			}
 		}
-		int32	MaxY = InViewport->GetSizeXY().Y;
+
+		const int32 MaxY = InViewport->GetSizeXY().Y;
 		static const double TravelTime = 4.2;
 		for (int32 i = 0; i < NumHitches; i++)
 		{
 			if (When[i] > 0 && When[i] <= CurrentTime && When[i] >= CurrentTime - TravelTime)
 			{
-				FColor MyColor = FColor::Green;
-				if (Hitches[i] > 0.2f) MyColor = FColor::Yellow;
-				if (Hitches[i] > 0.3f) MyColor = FColor::Red;
-				int32 MyY = InY + int32(float(MaxY - InY) * float((CurrentTime - When[i]) / TravelTime));
-				FString Hitch = FString::Printf(TEXT("%5d"), int32(Hitches[i] * 1000.0f));
+				const float MyHitchSecs = Hitches[i];
+				const float MyHitchMS = MyHitchSecs * 1000.0f;
+
+				// Scale the time before passing in so that hitches aren't all red
+				const FColor MyColor = GEngine->GetFrameTimeDisplayColor(MyHitchMS * 0.25f);
+
+				const int32 MyY = InY + int32(float(MaxY - InY) * float((CurrentTime - When[i]) / TravelTime));
+				const FString Hitch = FString::Printf(TEXT("%5d"), int32(MyHitchMS));
 				InCanvas->DrawShadowedString(InX, MyY, *Hitch, GEngine->GetSmallFont(), MyColor);
 			}
 		}
