Engine: Switched on-screen cycle counter stat display to scale the bar based on t.TargetFrameTimeThreshold rather than a hardcoded 33.3 ms
[INTEGRATE] Change 2710576 by [REDACTED_PERSON][REDACTED_HANDLE] on 2015/09/29 22:09:01

diff --git a/Engine/Source/Runtime/Engine/Private/StatsRender2.cpp b/Engine/Source/Runtime/Engine/Private/StatsRender2.cpp
--- a/Engine/Source/Runtime/Engine/Private/StatsRender2.cpp
+++ b/Engine/Source/Runtime/Engine/Private/StatsRender2.cpp
@@ -10,6 +10,7 @@
 #if STATS
 
 #include "StatsData.h"
+#include "Performance/EnginePerformanceTargets.h"
 
 /** Stats rendering constants. */
 enum class EStatRenderConsts
@@ -296,7 +297,9 @@ static int32 RenderCycle( const FComplexStatMessage& Item, class FCanvas* Canvas
 		CounterName.RemoveFromStart(TEXT("STAT_"), ESearchCase::CaseSensitive);
 		GEngine->GetStatValueColoration(CounterName, InMs, Color);
 
-		const float MaxMeter = 33.3f; // the time of a "full bar" in ms
+		// the time of a "full bar" in ms
+		const float MaxMeter = FEnginePerformanceTargets::GetTargetFrameTimeThresholdMS();
+
 		const int32 MeterWidth = Globals.AfterNameColumnOffset;
 
 		int32 BarWidth = int32((InMs / MaxMeter) * MeterWidth);
