UE4 - Added configurability to the texture streamer update frequency.

diff --git a/Engine/Source/Runtime/Engine/Private/ContentStreaming.cpp b/Engine/Source/Runtime/Engine/Private/ContentStreaming.cpp
--- a/Engine/Source/Runtime/Engine/Private/ContentStreaming.cpp
+++ b/Engine/Source/Runtime/Engine/Private/ContentStreaming.cpp
@@ -3234,8 +3234,6 @@ void FStreamingManagerTexture::UpdateStreamingTextures( FStreamingContext& Conte
 	{
 		FStreamingTexture& StreamingTexture = StreamingTextures[ Index ];
 		FPlatformMisc::Prefetch( &StreamingTexture + 1 );
-		FPlatformMisc::Prefetch( StreamingTexture.Texture );
-		FPlatformMisc::Prefetch( StreamingTexture.Texture, CACHE_LINE_SIZE );
 
 		// Is this texture marked for removal?
 		if ( StreamingTexture.Texture == NULL )
@@ -3686,6 +3684,11 @@ void FStreamingManagerTexture::CheckUserSettings()
  * @param DeltaTime				Time since last call in seconds
  * @param bProcessEverything	[opt] If true, process all resources with no throttling limits
  */
+static TAutoConsoleVariable<int32> CVarFramesForFullUpdate(
+	TEXT("r.Streaming.FramesForFullUpdate"),
+	5,
+	TEXT("Texture streaming is time sliced per frame. This values gives the number of frames to visit all textures."));
+
 void FStreamingManagerTexture::UpdateResourceStreaming( float DeltaTime, bool bProcessEverything/*=false*/ )
 {
 	SCOPE_CYCLE_COUNTER(STAT_GameThreadUpdateTime);
@@ -3708,10 +3711,18 @@ void FStreamingManagerTexture::UpdateResourceStreaming( float DeltaTime, bool bP
 		bWasLocationOveridden = bIsLocationOverridden;
 	}
 #endif
-
+	int32 NewNumTextureProcessingStages = CVarFramesForFullUpdate.GetValueOnGameThread();
+	if (NewNumTextureProcessingStages > 0 && NewNumTextureProcessingStages != NumTextureProcessingStages)
+	{
+		QUICK_SCOPE_CYCLE_COUNTER(STAT_FStreamingManagerTexture_UpdateResourceStreaming_EnsureCompletion_E);
+		AsyncWork->EnsureCompletion();
+		ProcessingStage = 0;
+		NumTextureProcessingStages = NewNumTextureProcessingStages;
+	}
 	int32 OldNumTextureProcessingStages = NumTextureProcessingStages;
 	if ( bProcessEverything || IndividualStreamingTexture )
 	{
+		QUICK_SCOPE_CYCLE_COUNTER(STAT_FStreamingManagerTexture_UpdateResourceStreaming_EnsureCompletion_A);
 		AsyncWork->EnsureCompletion();
 		ProcessingStage = 0;
 		NumTextureProcessingStages = 1;
@@ -3732,17 +3743,30 @@ void FStreamingManagerTexture::UpdateResourceStreaming( float DeltaTime, bool bP
 		// Is the AsyncWork is running for some reason? (E.g. we reset the system by simply setting ProcessingStage to 0.)
 		if ( AsyncWork->IsDone() == false )
 		{
+			QUICK_SCOPE_CYCLE_COUNTER(STAT_FStreamingManagerTexture_UpdateResourceStreaming_EnsureCompletion_B);
 			AsyncWork->EnsureCompletion();
 		}
 
-		CheckUserSettings();
+		{
+			QUICK_SCOPE_CYCLE_COUNTER(STAT_FStreamingManagerTexture_UpdateResourceStreaming_CheckUserSettings);
+			CheckUserSettings();
+		}
 
-		ResetStreamingStats();
-		UpdateThreadData();
+		{
+			QUICK_SCOPE_CYCLE_COUNTER(STAT_FStreamingManagerTexture_UpdateResourceStreaming_ResetStreamingStats);
+			ResetStreamingStats();
+		}
+		{
+			QUICK_SCOPE_CYCLE_COUNTER(STAT_FStreamingManagerTexture_UpdateResourceStreaming_UpdateThreadData);
+			UpdateThreadData();
+		}
 
 		//@TODO: Move this to the worker thread once it's thread-safe.
 		//Note: Must be done with updated primitives (right after UpdateThreadData), because we can't ignore a primitive just because it moved.
-		CalcDynamicWantedMips();
+		{
+			QUICK_SCOPE_CYCLE_COUNTER(STAT_FStreamingManagerTexture_UpdateResourceStreaming_CalcDynamicWantedMips);
+			CalcDynamicWantedMips();
+		}
 
 		if ( bTriggerDumpTextureGroupStats )
 		{
@@ -3760,8 +3784,14 @@ void FStreamingManagerTexture::UpdateResourceStreaming( float DeltaTime, bool bP
 	{
 		// Setup a context for this run. Note that we only (potentially) collect texture stats from the AsyncWork.
 		FStreamingContext Context( bProcessEverything, IndividualStreamingTexture, false );
-		UpdateStreamingTextures( Context, ProcessingStage, NumDataCollectionStages );
-		UpdateStreamingStats( Context, false );
+		{
+			QUICK_SCOPE_CYCLE_COUNTER(STAT_FStreamingManagerTexture_UpdateResourceStreaming_UpdateStreamingTextures);
+			UpdateStreamingTextures( Context, ProcessingStage, NumDataCollectionStages );
+		}
+		{
+			QUICK_SCOPE_CYCLE_COUNTER(STAT_FStreamingManagerTexture_UpdateResourceStreaming_UpdateStreamingStats);
+			UpdateStreamingStats( Context, false );
+		}
 	}
 
 	// Start async task after the last data collection stage (if we're not paused).
@@ -3770,17 +3800,20 @@ void FStreamingManagerTexture::UpdateResourceStreaming( float DeltaTime, bool bP
 		// Is the AsyncWork is running for some reason? (E.g. we reset the system by simply setting ProcessingStage to 0.)
 		if ( AsyncWork->IsDone() == false )
 		{
+			QUICK_SCOPE_CYCLE_COUNTER(STAT_FStreamingManagerTexture_UpdateResourceStreaming_EnsureCompletion_C);
 			AsyncWork->EnsureCompletion();
 		}
 
 		AsyncWork->GetTask().Reset(bCollectTextureStats);
 		if ( NumTextureProcessingStages > 1 )
 		{
+			QUICK_SCOPE_CYCLE_COUNTER(STAT_FStreamingManagerTexture_UpdateResourceStreaming_StartBackgroundTask);
 			AsyncWork->StartBackgroundTask();
 		}
 		else
 		{
 			// Perform the work synchronously on this thread.
+			QUICK_SCOPE_CYCLE_COUNTER(STAT_FStreamingManagerTexture_UpdateResourceStreaming_StartSynchronousTask);
 			AsyncWork->StartSynchronousTask();
 		}
 	}
@@ -3799,6 +3832,7 @@ void FStreamingManagerTexture::UpdateResourceStreaming( float DeltaTime, bool bP
 	else
 	{
 		// All priorities have been calculated, do all streaming.
+		QUICK_SCOPE_CYCLE_COUNTER(STAT_FStreamingManagerTexture_UpdateResourceStreaming_StreamTextures);
 		StreamTextures( bProcessEverything );
 		ProcessingStage = 0;
 	}
@@ -5343,24 +5377,13 @@ void FStreamingManagerTexture::CalcDynamicWantedMips()
 	STAT( double StartTime = FPlatformTime::Seconds() );
 
 	// Reset the dynamic variables for all textures.
-	const int32 PrefetchCount = AlignArbitrary( 4*CACHE_LINE_SIZE, sizeof(FStreamingTexture) ) / sizeof(FStreamingTexture);
-	for ( int32 Index=0; Index < StreamingTextures.Num() - PrefetchCount; ++Index )
+	for ( int32 Index=0; Index < StreamingTextures.Num(); ++Index )
 	{
 		FStreamingTexture& StreamingTexture = StreamingTextures[ Index ];
-		FPlatformMisc::Prefetch( &StreamingTexture, CACHE_LINE_SIZE );
-		FPlatformMisc::Prefetch( &StreamingTexture, CACHE_LINE_SIZE*2 );
-		FPlatformMisc::Prefetch( &StreamingTexture, CACHE_LINE_SIZE*3 );
+		FPlatformMisc::Prefetch( &StreamingTexture, CACHE_LINE_SIZE*4 );
 		StreamingTexture.DynamicScreenSize = 0.0f;
 		StreamingTexture.DynamicMinDistanceSq = FLT_MAX;
 	}
-	// The last ones
-	for ( int32 Index=FMath::Max(StreamingTextures.Num() - PrefetchCount, 0); Index < StreamingTextures.Num(); ++Index )
-	{
-		FStreamingTexture& StreamingTexture = StreamingTextures[ Index ];
-		StreamingTexture.DynamicScreenSize = 0.0f;
-		StreamingTexture.DynamicMinDistanceSq = FLT_MAX;
-	}
-
 	// Iterate over all dynamic primitives.
 	for ( TMap<const UPrimitiveComponent*,FSpawnedPrimitiveData>::TIterator It(ThreadSettings.SpawnedPrimitives); It; ++It )
 	{
