Fix crash when something other than USkeletalMeshComponent has a MasterPoseComponent

diff --git a/Engine/Source/Runtime/Engine/Private/SkeletalRenderGPUSkin.cpp b/Engine/Source/Runtime/Engine/Private/SkeletalRenderGPUSkin.cpp
--- a/Engine/Source/Runtime/Engine/Private/SkeletalRenderGPUSkin.cpp
+++ b/Engine/Source/Runtime/Engine/Private/SkeletalRenderGPUSkin.cpp
@@ -1128,7 +1128,7 @@ bool FDynamicSkelMeshObjectDataGPUSkin::UpdateClothSimulationData(USkinnedMeshCo
 	USkeletalMeshComponent * SkelMeshComponent = Cast<USkeletalMeshComponent>(InMeshComponent);
 
 #if WITH_APEX_CLOTHING
-	if(InMeshComponent->MasterPoseComponent.IsValid() && SkelMeshComponent->IsClothBoundToMasterComponent())
+	if(InMeshComponent->MasterPoseComponent.IsValid() && SkelMeshComponent && SkelMeshComponent->IsClothBoundToMasterComponent())
 	{
 		USkeletalMeshComponent* OriginalComponent = Cast<USkeletalMeshComponent>(InMeshComponent);
 		SkelMeshComponent = Cast<USkeletalMeshComponent>(InMeshComponent->MasterPoseComponent.Get());
