Fix crash when calling get bone functions on a UPoseableMeshComponent when it has a MasterPoseComponent

diff --git a/Engine/Source/Runtime/Engine/Private/Components/PoseableMeshComponent.cpp b/Engine/Source/Runtime/Engine/Private/Components/PoseableMeshComponent.cpp
--- a/Engine/Source/Runtime/Engine/Private/Components/PoseableMeshComponent.cpp
+++ b/Engine/Source/Runtime/Engine/Private/Components/PoseableMeshComponent.cpp
@@ -129,6 +129,8 @@ void UPoseableMeshComponent::SetBoneTransformByName(FName BoneName, const FTrans
 		return;
 	}
 
+	check(!MasterPoseComponent.IsValid()); //Shouldn't call set bone functions when we are using MasterPoseComponent
+
 	int32 BoneIndex = GetBoneIndex(BoneName);
 	if(BoneIndex >=0 && BoneIndex < LocalAtoms.Num())
 	{
@@ -178,15 +180,11 @@ void UPoseableMeshComponent::SetBoneScaleByName(FName BoneName, FVector InScale3
 	SetBoneTransformByName(BoneName, CurrentTransform, BoneSpace);
 }
 
-FTransform UPoseableMeshComponent::GetBoneTransformByName(FName BoneName, EBoneSpaces::Type BoneSpace)
+template<class CompType>
+FTransform GetBoneTransformByNameHelper(FName BoneName, EBoneSpaces::Type BoneSpace, FBoneContainer& RequiredBones, CompType* Component)
 {
-	if( !SkeletalMesh || !RequiredBones.IsValid() )
-	{
-		return FTransform();
-	}
-
-	int32 BoneIndex = GetBoneIndex(BoneName);
-	if( BoneIndex == INDEX_NONE)
+	int32 BoneIndex = Component->GetBoneIndex(BoneName);
+	if (BoneIndex == INDEX_NONE)
 	{
 		FString Message = FString::Printf(TEXT("Invalid Bone Name '%s'"), *BoneName.ToString());
 		FFrame::KismetExecutionMessage(*Message, ELogVerbosity::Warning);
@@ -195,20 +193,47 @@ FTransform UPoseableMeshComponent::GetBoneTransformByName(FName BoneName, EBoneS
 
 	/*if(BoneSpace == EBoneSpaces::LocalSpace)
 	{
-		return LocalAtoms[i];
+		return Component->LocalAtoms[i];
 	}*/
 
 	FA2CSPose CSPose;
-	CSPose.AllocateLocalPoses(RequiredBones, LocalAtoms);
+	CSPose.AllocateLocalPoses(RequiredBones, Component->LocalAtoms);
 
-	if(BoneSpace == EBoneSpaces::ComponentSpace)
+	if (BoneSpace == EBoneSpaces::ComponentSpace)
 	{
 		return CSPose.GetComponentSpaceTransform(BoneIndex);
 	}
 	else
 	{
-		return CSPose.GetComponentSpaceTransform(BoneIndex) * ComponentToWorld;
+		return CSPose.GetComponentSpaceTransform(BoneIndex) * Component->ComponentToWorld;
+	}
+}
+
+FTransform UPoseableMeshComponent::GetBoneTransformByName(FName BoneName, EBoneSpaces::Type BoneSpace)
+{
+	if( !SkeletalMesh || !RequiredBones.IsValid() )
+	{
+		return FTransform();
+	}
+
+	USkinnedMeshComponent* MPCPtr = MasterPoseComponent.Get();
+	if (MPCPtr)
+	{
+		if (USkeletalMeshComponent* SMC = Cast<USkeletalMeshComponent>(MPCPtr))
+		{
+			if (UAnimInstance* AnimInstance = SMC->GetAnimInstance())
+			{
+				return GetBoneTransformByNameHelper(BoneName, BoneSpace, AnimInstance->RequiredBones, SMC);
+			}
+			FString Message = FString::Printf(TEXT("Cannot return valid bone transform. Master Pose Component has no anim instance"));
+			FFrame::KismetExecutionMessage(*Message, ELogVerbosity::Warning);
+			return FTransform();
+		}
+		FString Message = FString::Printf(TEXT("Cannot return valid bone transform. Master Pose Component is not of type USkeletalMeshComponent"));
+		FFrame::KismetExecutionMessage(*Message, ELogVerbosity::Warning);
+		return FTransform();
 	}
+	return GetBoneTransformByNameHelper(BoneName, BoneSpace, RequiredBones, this);
 }
 
 FVector UPoseableMeshComponent::GetBoneLocationByName(FName BoneName, EBoneSpaces::Type BoneSpace)
