Fix deprecation warning caused by automatically generated operator= for FActorSpawnParameters.
This fixes the default parameter assignments causing warnings in the 2015 toolchain.

diff --git a/Engine/Source/Runtime/Engine/Classes/Engine/World.h b/Engine/Source/Runtime/Engine/Classes/Engine/World.h
--- a/Engine/Source/Runtime/Engine/Classes/Engine/World.h
+++ b/Engine/Source/Runtime/Engine/Classes/Engine/World.h
@@ -421,6 +421,9 @@ struct ENGINE_API FActorSpawnParameters
 {
 	FActorSpawnParameters();
 
+	// Assignment operator overriden to allow disabling of deprecation warning when copying bNoCollisionFail
+	FActorSpawnParameters& operator=(const FActorSpawnParameters& Other);
+
 	/* A name to assign as the Name of the Actor being spawned. If no value is specified, the name of the spawned Actor will be automatically generated using the form [Class]_[Number]. */
 	FName Name;
 
diff --git a/Engine/Source/Runtime/Engine/Private/World.cpp b/Engine/Source/Runtime/Engine/Private/World.cpp
--- a/Engine/Source/Runtime/Engine/Private/World.cpp
+++ b/Engine/Source/Runtime/Engine/Private/World.cpp
@@ -111,6 +111,25 @@ FActorSpawnParameters::FActorSpawnParameters()
 
 PRAGMA_ENABLE_DEPRECATION_WARNINGS
 
+FActorSpawnParameters& FActorSpawnParameters::operator=(const FActorSpawnParameters& Other)
+{
+	Name = Other.Name;
+	Template = Other.Template;
+	Owner = Other.Owner;
+	Instigator = Other.Instigator;
+	OverrideLevel = Other.OverrideLevel;
+	SpawnCollisionHandlingOverride = Other.SpawnCollisionHandlingOverride;
+PRAGMA_DISABLE_DEPRECATION_WARNINGS
+	bNoCollisionFail = Other.bNoCollisionFail;
+PRAGMA_ENABLE_DEPRECATION_WARNINGS
+	bRemoteOwned = Other.bRemoteOwned;
+	bNoFail = Other.bNoFail;
+	bDeferConstruction = Other.bDeferConstruction;
+	bAllowDuringConstructionScript = Other.bAllowDuringConstructionScript;
+	ObjectFlags = Other.ObjectFlags;
+	return *this;
+}
+
 /*-----------------------------------------------------------------------------
 	UWorld implementation.
 -----------------------------------------------------------------------------*/
