Fixed copy/paste error in ConsoleVariable meta tag.
PR#1607

diff --git a/Engine/Source/Runtime/Engine/Classes/Engine/RendererSettings.h b/Engine/Source/Runtime/Engine/Classes/Engine/RendererSettings.h
--- a/Engine/Source/Runtime/Engine/Classes/Engine/RendererSettings.h
+++ b/Engine/Source/Runtime/Engine/Classes/Engine/RendererSettings.h
@@ -126,7 +126,7 @@ class ENGINE_API URendererSettings : public UDeveloperSettings
 	float MinScreenRadiusForEarlyZPass;
 
 	UPROPERTY(config, EditAnywhere, Category=Culling, meta=(
-		ConsoleVariable="r.MinScreenRadiusForDepthPrepass",DisplayName="Min Screen Radius for Cascaded Shadow Maps",
+		ConsoleVariable="r.MinScreenRadiusForCSMDepth",DisplayName="Min Screen Radius for Cascaded Shadow Maps",
 		ToolTip="Screen radius at which objects are culled for cascaded shadow map depth passes. Larger values can improve performance but can cause artifacts as objects stop casting shadows."))
 	float MinScreenRadiusForCSMdepth;
 	
