Fixed issue with infinite loop when playing negative rate with blending trigger time

diff --git a/Engine/Source/Runtime/Engine/Private/Animation/AnimMontage.cpp b/Engine/Source/Runtime/Engine/Private/Animation/AnimMontage.cpp
--- a/Engine/Source/Runtime/Engine/Private/Animation/AnimMontage.cpp
+++ b/Engine/Source/Runtime/Engine/Private/Animation/AnimMontage.cpp
@@ -1490,6 +1490,12 @@ void FAnimMontageInstance::Advance(float DeltaTime, struct FRootMotionMovementPa
 							Position = bPlayingForward ? LatestNextSectionStartTime : (LatestNextSectionEndTime - EndOffset);
 						}
 					}
+
+					if (!bHaveMoved)
+					{
+						// If it hasn't moved, there is nothing much to do but weight update
+						break;
+					}
 				}
 				else
 				{
