Don't do unnecessary work related to hit box detection if none are being created by the draw hud calls.
Use Reset instead of Empty on persistent HUD arrays

diff --git a/Engine/Source/Runtime/Engine/Private/HUD.cpp b/Engine/Source/Runtime/Engine/Private/HUD.cpp
--- a/Engine/Source/Runtime/Engine/Private/HUD.cpp
+++ b/Engine/Source/Runtime/Engine/Private/HUD.cpp
@@ -140,44 +140,57 @@ void AHUD::PostRender()
 	{
 		DrawHUD();
 		
-		ULocalPlayer* LocalPlayer = GetOwningPlayerController() ? Cast<ULocalPlayer>(GetOwningPlayerController()->Player) : NULL;
-
-		if (LocalPlayer && LocalPlayer->ViewportClient)
+		// No need to do work to determine hit box candidates if there will never be any
+		if (HitBoxMap.Num() > 0)
 		{
-			TArray<FVector2D> ContactPoints;
+			ULocalPlayer* LocalPlayer = GetOwningPlayerController() ? Cast<ULocalPlayer>(GetOwningPlayerController()->Player) : NULL;
 
-			if (!FSlateApplication::Get().IsFakingTouchEvents())
+			if (LocalPlayer && LocalPlayer->ViewportClient)
 			{
-				FVector2D MousePosition;
-				if (LocalPlayer->ViewportClient->GetMousePosition(MousePosition))
+				TArray<FVector2D> ContactPoints;
+
+				if (!FSlateApplication::Get().IsFakingTouchEvents())
 				{
-					ContactPoints.Add(MousePosition);
+					FVector2D MousePosition;
+					if (LocalPlayer->ViewportClient->GetMousePosition(MousePosition))
+					{
+						ContactPoints.Add(MousePosition);
+					}
 				}
-			}
 
-			for (int32 FingerIndex = 0; FingerIndex < EKeys::NUM_TOUCH_KEYS; ++FingerIndex)
-			{
-				FVector2D TouchLocation;
-				bool bPressed = false;
+				for (int32 FingerIndex = 0; FingerIndex < EKeys::NUM_TOUCH_KEYS; ++FingerIndex)
+				{
+					FVector2D TouchLocation;
+					bool bPressed = false;
 
-				GetOwningPlayerController()->GetInputTouchState((ETouchIndex::Type)FingerIndex, TouchLocation.X, TouchLocation.Y, bPressed);
+					GetOwningPlayerController()->GetInputTouchState((ETouchIndex::Type)FingerIndex, TouchLocation.X, TouchLocation.Y, bPressed);
 
-				if (bPressed)
-				{
-					ContactPoints.Add(TouchLocation);
+					if (bPressed)
+					{
+						ContactPoints.Add(TouchLocation);
+					}
 				}
-			}
 
-			FVector2D ContactPointOffset = GetCoordinateOffset();
+				const FVector2D ContactPointOffset = GetCoordinateOffset();
 
-			if (!ContactPointOffset.IsZero())
-			{
-				for (FVector2D& ContactPoint : ContactPoints)
+				if (!ContactPointOffset.IsZero())
 				{
-					ContactPoint += ContactPointOffset;
+					for (FVector2D& ContactPoint : ContactPoints)
+					{
+						ContactPoint += ContactPointOffset;
+					}
 				}
+				UpdateHitBoxCandidates( ContactPoints );
+			}
+		}
+		else if (HitBoxesOver.Num() > 0)
+		{
+			// We still need to dispatch any end cursor over messages even if we don't have any hitboxes anymore
+			for (const FName HitBoxName : HitBoxesOver)
+			{
+				NotifyHitBoxEndCursorOver(HitBoxName);
 			}
-			UpdateHitBoxCandidates( ContactPoints );
+			HitBoxesOver.Reset();
 		}
 	}
 	
@@ -263,7 +276,7 @@ void AHUD::ShowDebug(FName DebugType)
 	}
 	else if( DebugType == NAME_Reset )
 	{
-		DebugDisplay.Empty();
+		DebugDisplay.Reset();
 		bShowDebugInfo = false;
 		SaveConfig();
 	}
@@ -293,7 +306,7 @@ void AHUD::ShowDebugToggleSubCategory(FName Category)
 {
 	if( Category == NAME_Reset )
 	{
-		ToggledDebugCategories.Empty();
+		ToggledDebugCategories.Reset();
 		SaveConfig();
 	}
 	else
@@ -366,8 +379,8 @@ void AHUD::ShowDebugInfo(float& YL, float& YPos)
 
 void AHUD::DrawHUD()
 {
-	HitBoxMap.Empty();
-	HitBoxHits.Empty();
+	HitBoxMap.Reset();
+	HitBoxHits.Reset();
 	if ( bShowOverlays && (PlayerOwner != NULL) )
 	{
 		FVector ViewPoint;
@@ -575,7 +588,7 @@ void AHUD::RemoveDebugText_Implementation(AActor* SrcActor, bool bLeaveDurationT
 /** Remove all debug text */
 void AHUD::RemoveAllDebugStrings_Implementation()
 {
-	DebugTextList.Empty();
+	DebugTextList.Reset();
 }
 
 void AHUD::NotifyHitBoxClick(FName BoxName)
@@ -720,7 +733,7 @@ void AHUD::GetActorsInSelectionRectangle(TSubclassOf<class AActor> ClassFilter,
 {
 	// Because this is a HUD function it is likely to get called each tick,
 	// so make sure any previous contents of the out actor array have been cleared!
-	OutActors.Empty();
+	OutActors.Reset();
 
 	//Create Selection Rectangle from Points
 	FBox2D SelectionRectangle(0);
@@ -835,7 +848,7 @@ void AHUD::RenderHitBoxes( FCanvas* InCanvas )
 
 void AHUD::UpdateHitBoxCandidates( TArray<FVector2D> InContactPoints )
 {
-	HitBoxHits.Empty();
+	HitBoxHits.Reset();
 	for (FHUDHitBox& HitBox : HitBoxMap)
 	{
 		bool bAdded = false;
@@ -882,17 +895,15 @@ void AHUD::UpdateHitBoxCandidates( TArray<FVector2D> InContactPoints )
 	}
 
 	// Dispatch the end cursor over messages
-	for (auto It = NotOverHitBoxes.CreateConstIterator(); It; ++It)
+	for (const FName HitBoxName : NotOverHitBoxes)
 	{
-		const FName HitBoxName = *It;
 		NotifyHitBoxEndCursorOver(HitBoxName);
 		HitBoxesOver.Remove(HitBoxName);
 	}
 
 	// Dispatch the newly over hitbox messages
-	for (int32 OverIndex = 0; OverIndex < NewlyOverHitBoxes.Num(); ++OverIndex)
+	for (const FName HitBoxName : NewlyOverHitBoxes)
 	{
-		const FName HitBoxName = NewlyOverHitBoxes[OverIndex];
 		NotifyHitBoxBeginCursorOver(HitBoxName);
 		HitBoxesOver.Add(HitBoxName);
 	}
@@ -900,42 +911,49 @@ void AHUD::UpdateHitBoxCandidates( TArray<FVector2D> InContactPoints )
 
 const FHUDHitBox* AHUD::GetHitBoxAtCoordinates( FVector2D InHitLocation, const bool bIsConsumingInput ) const
 {
-	InHitLocation -= GetCoordinateOffset();
-
-	for (const FHUDHitBox& HitBox : HitBoxMap)
+	if (HitBoxMap.Num() > 0)
 	{
-		if( (!bIsConsumingInput || HitBox.ConsumesInput()) && HitBox.Contains( InHitLocation ) )
+		InHitLocation -= GetCoordinateOffset();
+
+		for (const FHUDHitBox& HitBox : HitBoxMap)
 		{
-			return &HitBox;
+			if( (!bIsConsumingInput || HitBox.ConsumesInput()) && HitBox.Contains( InHitLocation ) )
+			{
+				return &HitBox;
+			}
 		}
 	}
-	return NULL;
+	return nullptr;
 }
 
 void AHUD::GetHitBoxesAtCoordinates(FVector2D InHitLocation, TArray<const FHUDHitBox*>& OutHitBoxes) const
 {
-	InHitLocation -= GetCoordinateOffset();
-	OutHitBoxes.Empty();
+	OutHitBoxes.Reset();
 
-	for (const FHUDHitBox& HitBox : HitBoxMap)
+	if (HitBoxMap.Num() > 0)
 	{
-		if (HitBox.Contains(InHitLocation))
+		InHitLocation -= GetCoordinateOffset();
+
+		for (const FHUDHitBox& HitBox : HitBoxMap)
 		{
-			OutHitBoxes.Add(&HitBox);
+			if (HitBox.Contains(InHitLocation))
+			{
+				OutHitBoxes.Add(&HitBox);
+			}
 		}
 	}
 }
 
 const FHUDHitBox* AHUD::GetHitBoxWithName( const FName InName ) const
 {
-	for (int32 iBox = 0; iBox < HitBoxMap.Num() ; iBox++)
+	for (const FHUDHitBox& HitBox : HitBoxMap)
 	{
-		if( HitBoxMap[iBox].GetName() == InName )
+		if( HitBox.GetName() == InName )
 		{
-			return &HitBoxMap[iBox];
+			return &HitBox;
 		}
 	}
-	return NULL;
+	return nullptr;
 }
 
 bool AHUD::AnyCurrentHitBoxHits() const
@@ -945,9 +963,16 @@ bool AHUD::AnyCurrentHitBoxHits() const
 
 bool AHUD::UpdateAndDispatchHitBoxClickEvents(FVector2D ClickLocation, const EInputEvent InEventType)
 {
+	const bool bIsClickEvent = (InEventType == IE_Pressed || InEventType == IE_DoubleClick);
+
+	// Early out to avoid unnecessary expense of calling GetCoordinateOffset
+	if ((bIsClickEvent && HitBoxMap.Num() == 0) || (!bIsClickEvent && HitBoxHits.Num() == 0))
+	{
+		return false;
+	}
+
 	ClickLocation += GetCoordinateOffset();
 
-	const bool bIsClickEvent = (InEventType == IE_Pressed || InEventType == IE_DoubleClick);
 	bool bHit = false;
 
 	// If this is a click event we may not have the hit box in the hit list yet (particularly for touch events) so we need to check all HitBoxes
@@ -993,7 +1018,7 @@ bool AHUD::UpdateAndDispatchHitBoxClickEvents(FVector2D ClickLocation, const EIn
 
 void AHUD::AddHitBox(FVector2D Position, FVector2D Size, FName Name, bool bConsumesInput, int32 Priority)
 {	
-	if( GetHitBoxWithName(Name) == NULL )
+	if( GetHitBoxWithName(Name) == nullptr )
 	{
 		bool bAdded = false;
 		for (int32 Index = 0; Index < HitBoxMap.Num(); ++Index)
