Added warning to resave the source package of a blueprint when a child blueprint function / dependency is out of date.
Resave is required to resolve deterministic cooking issue.

diff --git a/Engine/Source/Runtime/Engine/Private/EdGraph/EdGraphNode.cpp b/Engine/Source/Runtime/Engine/Private/EdGraph/EdGraphNode.cpp
--- a/Engine/Source/Runtime/Engine/Private/EdGraph/EdGraphNode.cpp
+++ b/Engine/Source/Runtime/Engine/Private/EdGraph/EdGraphNode.cpp
@@ -279,7 +279,7 @@ void UEdGraphNode::PostLoad()
 	// Create Guid if not present (and not CDO)
 	if(!NodeGuid.IsValid() && !IsTemplate() && GetLinker() && GetLinker()->IsPersistent() && GetLinker()->IsLoading())
 	{
-		UE_LOG(LogBlueprint, Warning, TEXT("Node '%s' missing NodeGuid."), *GetPathName());
+		UE_LOG(LogBlueprint, Warning, TEXT("Node '%s' missing NodeGuid, this can cause deterministic cooking issues please resave package."), *GetPathName());
 
 		// Generate new one
 		CreateNewGuid();
@@ -288,6 +288,8 @@ void UEdGraphNode::PostLoad()
 	// Duplicating a Blueprint needs to have a new Node Guid generated, which was not occuring before this version
 	if(GetLinkerUE4Version() < VER_UE4_POST_DUPLICATE_NODE_GUID)
 	{
+		UE_LOG(LogBlueprint, Warning, TEXT("Node '%s' missing NodeGuid because of upgrade from old package version, this can cause deterministic cooking issues please resave package."), *GetPathName());
+
 		// Generate new one
 		CreateNewGuid();
 	}
