Added warning to resave material package when a dependent texture becomes out of date.
Resave is required to resolve deterministic cooking issue.

diff --git a/Engine/Source/Runtime/Engine/Private/Materials/Material.cpp b/Engine/Source/Runtime/Engine/Private/Materials/Material.cpp
--- a/Engine/Source/Runtime/Engine/Private/Materials/Material.cpp
+++ b/Engine/Source/Runtime/Engine/Private/Materials/Material.cpp
@@ -3451,6 +3451,10 @@ bool UMaterial::UpdateLightmassTextureTracking()
 
 	if ( bTexturesHaveChanged )
 	{
+		if (IsRunningCommandlet())
+		{
+			UE_LOG(LogMaterial, Warning, TEXT("UMaterial's dependent textures have changed, resave package %s to avoid deterministic cooking issues."), *GetPathName());
+		}
 		// This will invalidate any cached Lightmass material exports
 		SetLightingGuid();
 	}
