Updated async trace handling to accept trace requests more widely. Avoids some asserts when poorly-timed requests would come in.

diff --git a/Engine/Source/Runtime/Engine/Private/Collision/WorldCollisionAsync.cpp b/Engine/Source/Runtime/Engine/Private/Collision/WorldCollisionAsync.cpp
--- a/Engine/Source/Runtime/Engine/Private/Collision/WorldCollisionAsync.cpp
+++ b/Engine/Source/Runtime/Engine/Private/Collision/WorldCollisionAsync.cpp
@@ -292,6 +292,8 @@ FWorldAsyncTraceState::FWorldAsyncTraceState()
 	, NextAvailableTraceIndex  (0)
 	, NextAvailableOverlapIndex(0)
 {
+	// initial buffer is open for business
+	DataBuffer[CurrentFrame].bAsyncAllowed = true;
 }
 
 FTraceHandle UWorld::AsyncLineTraceByChannel(const FVector& Start,const FVector& End, ECollisionChannel TraceChannel, const FCollisionQueryParams& Params /* = FCollisionQueryParams::DefaultQueryParam */, const FCollisionResponseParams& ResponseParam /* = FCollisionResponseParams::DefaultResponseParam */, FTraceDelegate * InDelegate/* =NULL */, uint32 UserData /* = 0 */, bool bMultiTrace/* =false */ )
@@ -406,21 +408,16 @@ void UWorld::ResetAsyncTrace()
 	WaitForAllAsyncTraceTasks();
 
 	// do run delegates before starting next round
-	for (int32 Idx = 0; Idx != AsyncTraceState.NextAvailableTraceIndex; ++Idx)
+	for (int32 Idx = 0; Idx != DataBufferExecuted.TraceData.Num(); ++Idx)
 	{
 		auto& TraceData = FBufferIndexPair(Idx).DatumLookupChecked(DataBufferExecuted.TraceData);
 		TraceData.Delegate.ExecuteIfBound(FTraceHandle(TraceData.FrameNumber, Idx), TraceData);
 	}
-	for (int32 Idx = 0; Idx != AsyncTraceState.NextAvailableOverlapIndex; ++Idx)
+	for (int32 Idx = 0; Idx != DataBufferExecuted.OverlapData.Num(); ++Idx)
 	{
 		auto& TraceData = FBufferIndexPair(Idx).DatumLookupChecked(DataBufferExecuted.OverlapData);
 		TraceData.Delegate.ExecuteIfBound(FTraceHandle(TraceData.FrameNumber, Idx), TraceData);
 	}
-
-	// re-initialize all variables
-	AsyncTraceState.GetBufferForCurrentFrame().bAsyncAllowed = true;
-	AsyncTraceState.NextAvailableTraceIndex   = 0;
-	AsyncTraceState.NextAvailableOverlapIndex = 0;
 }
 
 void UWorld::FinishAsyncTrace()
@@ -434,4 +431,10 @@ void UWorld::FinishAsyncTrace()
 
 	// increase buffer index to next one
 	++AsyncTraceState.CurrentFrame;
+
+	// set up new buffer to accept trace requests
+	AsyncTraceState.GetBufferForCurrentFrame().bAsyncAllowed = true;
+	AsyncTraceState.NextAvailableTraceIndex = 0;
+	AsyncTraceState.NextAvailableOverlapIndex = 0;
+
 }
