#merge: Lower some debugging verbosity now that we understand what's going on a little better

diff --git a/Engine/Source/Runtime/Engine/Private/NetConnection.cpp b/Engine/Source/Runtime/Engine/Private/NetConnection.cpp
--- a/Engine/Source/Runtime/Engine/Private/NetConnection.cpp
+++ b/Engine/Source/Runtime/Engine/Private/NetConnection.cpp
@@ -685,20 +685,9 @@ void UNetConnection::ReceivedPacket( FBitReader& Reader )
 
 	ValidateSendBuffer();
 
-	const double RealtimeSeconds	= FPlatformTime::Seconds();
-	const float DeltaTime			= Driver->Time - LastReceiveTime;
-	const float DeltaRealtime		= RealtimeSeconds - LastReceiveRealtime;
-
 	// Update receive time to avoid timeout.
 	LastReceiveTime		= Driver->Time;
-	LastReceiveRealtime = RealtimeSeconds;
-
-	const float Timeout = GetTimeoutValue();
-
-	if ( DeltaTime > Timeout || DeltaRealtime > Timeout )
-	{
-		UE_LOG( LogNetTraffic, Warning, TEXT( "UNetConnection::ReceivedPacket: Large time between packets. DeltaTime: %2.2f, Realtime: %2.2f, DriverTime: %2.2f, Threshold: %2.2f %s" ), DeltaTime, DeltaRealtime, Driver->Time, Timeout, *Describe() );
-	}
+	LastReceiveRealtime = FPlatformTime::Seconds();
 
 	// Check packet ordering.
 	const int32 PacketId = InternalAck ? InPacketId + 1 : MakeRelative(Reader.ReadInt(MAX_PACKETID),InPacketId,MAX_PACKETID);
