Integrated PR #1495: Added Pawn::IsPlayerControlled() to compliment IsLocallyControlled() (thanks EverNewJoy!)

diff --git a/Engine/Source/Runtime/Engine/Classes/GameFramework/Pawn.h b/Engine/Source/Runtime/Engine/Classes/GameFramework/Pawn.h
--- a/Engine/Source/Runtime/Engine/Classes/GameFramework/Pawn.h
+++ b/Engine/Source/Runtime/Engine/Classes/GameFramework/Pawn.h
@@ -276,9 +276,13 @@ public:
 	void ReceiveUnpossessed(AController* OldController);
 
 	/** @return true if controlled by a local (not network) Controller.	 */
-	UFUNCTION(BlueprintCallable, Category="Pawn")
+	UFUNCTION(BlueprintPure, Category="Pawn")
 	virtual bool IsLocallyControlled() const;
-
+  
+	/** @return true if controlled by a human player (possessed by a PlayerController).	 */
+	UFUNCTION(BlueprintPure, Category="Pawn")
+	virtual bool IsPlayerControlled() const;
+	
 	/**
 	 * Get the view rotation of the Pawn (direction they are looking, normally Controller->ControlRotation).
 	 * @return The view rotation of the Pawn.
diff --git a/Engine/Source/Runtime/Engine/Private/Pawn.cpp b/Engine/Source/Runtime/Engine/Private/Pawn.cpp
--- a/Engine/Source/Runtime/Engine/Private/Pawn.cpp
+++ b/Engine/Source/Runtime/Engine/Private/Pawn.cpp
@@ -204,6 +204,10 @@ bool APawn::IsLocallyControlled() const
 {
 	return ( Controller && Controller->IsLocalController() );
 }
+bool APawn::IsPlayerControlled() const
+{
+	return PlayerState && !PlayerState->bIsABot;
+}
 
 bool APawn::ReachedDesiredRotation()
 {
