Integrated PR #1453: Fixed potential startup crash in GameSession if PlayerStateClass is null. (thanks EverNewJoy!)

diff --git a/Engine/Source/Runtime/Engine/Private/GameSession.cpp b/Engine/Source/Runtime/Engine/Private/GameSession.cpp
--- a/Engine/Source/Runtime/Engine/Private/GameSession.cpp
+++ b/Engine/Source/Runtime/Engine/Private/GameSession.cpp
@@ -143,7 +143,16 @@ void AGameSession::InitOptions( const FString& Options )
 
 	MaxPlayers = UGameplayStatics::GetIntOption( Options, TEXT("MaxPlayers"), MaxPlayers );
 	MaxSpectators = UGameplayStatics::GetIntOption( Options, TEXT("MaxSpectators"), MaxSpectators );
-	SessionName = GetDefault<APlayerState>(GameMode->PlayerStateClass)->SessionName;
+	
+	APlayerState const* const DefaultPlayerState = GameMode ? GetDefault<APlayerState>(GameMode->PlayerStateClass) : nullptr;
+	if (DefaultPlayerState)
+	{
+		SessionName = DefaultPlayerState->SessionName;
+	}
+	else
+	{
+		UE_LOG(LogGameSession, Error, TEXT("Player State class is invalid for game mode: %s!"), *GameMode->GetName());
+	}
 }
 
 bool AGameSession::ProcessAutoLogin()
