Further hardening of GameSession::InitOptions against potential crashes.

diff --git a/Engine/Source/Runtime/Engine/Private/GameSession.cpp b/Engine/Source/Runtime/Engine/Private/GameSession.cpp
--- a/Engine/Source/Runtime/Engine/Private/GameSession.cpp
+++ b/Engine/Source/Runtime/Engine/Private/GameSession.cpp
@@ -137,21 +137,24 @@ bool AGameSession::HandleStartMatchRequest()
 
 void AGameSession::InitOptions( const FString& Options )
 {
-	UWorld* World = GetWorld();
+	UWorld* const World = GetWorld();
 	check(World);
-	AGameMode* const GameMode = World->GetAuthGameMode();
+	AGameMode* const GameMode = World ? World->GetAuthGameMode() : nullptr;
 
 	MaxPlayers = UGameplayStatics::GetIntOption( Options, TEXT("MaxPlayers"), MaxPlayers );
 	MaxSpectators = UGameplayStatics::GetIntOption( Options, TEXT("MaxSpectators"), MaxSpectators );
 	
-	APlayerState const* const DefaultPlayerState = GameMode ? GetDefault<APlayerState>(GameMode->PlayerStateClass) : nullptr;
-	if (DefaultPlayerState)
-	{
-		SessionName = DefaultPlayerState->SessionName;
-	}
-	else
+	if (GameMode)
 	{
-		UE_LOG(LogGameSession, Error, TEXT("Player State class is invalid for game mode: %s!"), *GameMode->GetName());
+		APlayerState const* const DefaultPlayerState = GetDefault<APlayerState>(GameMode->PlayerStateClass);
+		if (DefaultPlayerState)
+		{
+			SessionName = DefaultPlayerState->SessionName;
+		}
+		else
+		{
+			UE_LOG(LogGameSession, Error, TEXT("Player State class is invalid for game mode: %s!"), *GameMode->GetName());
+		}
 	}
 }
 
