#ue4 - Add console commands to list components/actors with a given response to any specified profile, and to list profiles with a given response to any specified collision channel.  This processes only those assets which have been loaded in to memory so far.

New commands are:

- Collision.ListProfiles
- Collision.ListChannels
- Collision.ListComponentsWithResponseToProfile <Response> <Profile>
- Collision.ListProfilesWithResponseToChannel <Response> <Channel>

Response is one of: Ignore, Overlap, Block.
Profile is either an index or name from the ListProfiles command.
Channel is either an index or name from the ListChannels command.

--------

diff --git a/Engine/Source/Runtime/Engine/Private/Collision/Collision.cpp b/Engine/Source/Runtime/Engine/Private/Collision/Collision.cpp
--- a/Engine/Source/Runtime/Engine/Private/Collision/Collision.cpp
+++ b/Engine/Source/Runtime/Engine/Private/Collision/Collision.cpp
@@ -7,6 +7,9 @@
 #include "EnginePrivate.h"
 #include "Collision.h"
 
+// TEMP until crash is fixed in IsCollisionEnabled().
+#include "LandscapeHeightfieldCollisionComponent.h"
+
 //////////////////////////////////////////////////////////////////////////
 // FHitResult
 
@@ -511,3 +514,400 @@ bool FSeparatingAxisPointCheck::FindSeparatingAxisGeneric()
 	return true;
 }
 
+
+
+#if !(UE_BUILD_SHIPPING || UE_BUILD_TEST)
+
+DECLARE_LOG_CATEGORY_EXTERN(LogCollisionCommands, Log, All);
+DEFINE_LOG_CATEGORY(LogCollisionCommands);
+
+namespace CollisionResponseConsoleCommands
+{
+	static const FString ResponseStrings[] = {TEXT("Ignore"), TEXT("Overlap"), TEXT("Block")};
+
+	FString GetDisplayNameText(const UEnum* Enum, int32 Index, const FString& Fallback)
+	{
+		if (Enum)
+		{
+			#if WITH_EDITOR
+				return Enum->GetDisplayNameText(Index).ToString();
+			#else
+				return FName::NameToDisplayString(Enum->GetEnumName(Index), false);
+			#endif
+		}
+		return Fallback;
+	}
+
+	FString FormatObjectName(UObject* Obj)
+	{
+		FString Result;
+		if (Obj)
+		{
+			if (Obj->GetOuter())
+			{
+				Result = TEXT("'") + Obj->GetPathName(Obj->GetOuter()->GetOuter()) + TEXT("'");
+			}
+			else
+			{
+				Result = TEXT("'") + Obj->GetPathName(Obj->GetOuter()) + TEXT("'");
+			}
+		}
+		return Result;
+	}
+
+	void ListCollisionProfileNames()
+	{
+		TArray<TSharedPtr<FName>> ProfileNameList;
+		UCollisionProfile::Get()->GetProfileNames(ProfileNameList);
+		int32 Index = 0;
+		for (TSharedPtr<FName> NamePtr : ProfileNameList)
+		{
+			const FName TemplateName = *NamePtr;
+			UE_LOG(LogCollisionCommands, Log, TEXT("%2d: %s"), Index++, *TemplateName.ToString());
+		}
+	}
+
+	void ListCollisionChannelNames()
+	{
+		const UEnum *ChannelEnum = FindObject<UEnum>(ANY_PACKAGE, TEXT("ECollisionChannel"), true);
+		if (ChannelEnum)
+		{
+			for (int32 ChannelIndex = 0; ChannelIndex < ECollisionChannel::ECC_MAX; ChannelIndex++)
+			{
+				const FString ChannelShortName = ChannelEnum->GetEnumName(ChannelIndex);
+				const FString ChannelDisplayName = GetDisplayNameText(ChannelEnum, ChannelIndex, ChannelShortName);
+				UE_LOG(LogCollisionCommands, Log, TEXT("%2d: %s (%s)"), ChannelIndex, *ChannelShortName, *ChannelDisplayName);
+			}
+		}
+	}
+
+	struct FSortComponentsWithResponseToProfile
+	{
+		FSortComponentsWithResponseToProfile(ECollisionResponse InRequiredResponse)
+		: RequiredResponse(InRequiredResponse)
+		{
+		}
+
+		bool operator()(const UPrimitiveComponent& A, const UPrimitiveComponent& B) const
+		{
+			UObject* AOwner = A.GetOuter();
+			UObject* BOwner = B.GetOuter();
+			if (AOwner && BOwner)
+			{
+				// For overlaps, sort first by bGenerateOverlapEvents
+				if (RequiredResponse == ECollisionResponse::ECR_Overlap)
+				{
+					if (A.bGenerateOverlapEvents != B.bGenerateOverlapEvents)
+					{
+						return A.bGenerateOverlapEvents;
+					}
+				}
+
+				// Sort by name.
+				return AOwner->GetName() < BOwner->GetName();
+			}
+			else
+			{
+				return false;
+			}
+		}
+
+	private:
+		ECollisionResponse RequiredResponse;
+	};
+
+	ECollisionResponse StringToCollisionResponse(const FString& InString)
+	{
+		int32 TestIndex = 0;
+		for (const FString& CurrentString : ResponseStrings)
+		{
+			if (CurrentString == InString)
+			{
+				return ECollisionResponse(TestIndex);
+			}
+			TestIndex++;
+		}
+
+		return ECollisionResponse::ECR_MAX;
+	}
+
+	ECollisionChannel StringToCollisionChannel(const FString& InString)
+	{
+		const UEnum *ChannelEnum = FindObject<UEnum>(ANY_PACKAGE, TEXT("ECollisionChannel"), true);
+		if (ChannelEnum)
+		{
+			int32 ChannelInt = INDEX_NONE;
+
+			// Try the name they gave us.
+			ChannelInt = ChannelEnum->GetValueByName(FName(*InString));
+			if (ChannelInt != INDEX_NONE)
+			{
+				return ECollisionChannel(ChannelInt);
+			}
+
+			// Try with adding the prefix, ie "ECC_".
+			const FString WithPrefixName = ChannelEnum->GenerateEnumPrefix() + TEXT("_") + InString;
+			ChannelInt = ChannelEnum->GetValueByName(FName(*WithPrefixName));
+			if (ChannelInt != INDEX_NONE)
+			{
+				return ECollisionChannel(ChannelInt);
+			}
+		}
+
+		// Try parsing a digit, as in from the ListChannels command
+		const int32 ChannelIndex = FCString::IsNumeric(*InString) ? FCString::Atoi(*InString) : INDEX_NONE;
+		return (ChannelIndex >= 0 && ChannelIndex < ECollisionChannel::ECC_MAX) ? ECollisionChannel(ChannelIndex) : ECollisionChannel::ECC_MAX;
+	}
+
+	FName StringToCollisionProfile(const FString& InString)
+	{
+		const FName InStringAsName(*InString);
+		FCollisionResponseTemplate Template;
+		if (UCollisionProfile::Get()->GetProfileTemplate(InStringAsName, Template))
+		{
+			return InStringAsName;
+		}
+
+		if (FCString::IsNumeric(*InString))
+		{
+			const int32 ProfileIndex = FCString::Atoi(*InString);
+			const FCollisionResponseTemplate* TemplateByIndex = UCollisionProfile::Get()->GetProfileByIndex(ProfileIndex);
+			if (TemplateByIndex)
+			{
+				return TemplateByIndex->Name;
+			}
+		}
+
+		return NAME_None;
+	}
+
+	//////////////////////////////////////////////////////////////////////////
+	void ListComponentsWithResponseToProfile(const ECollisionResponse RequiredResponse, const FName& ProfileToCheck)
+	{
+		// Look at all components and check response to profile
+		TSet<UPrimitiveComponent*> Results;
+		FCollisionResponseTemplate Template;
+		if (UCollisionProfile::Get()->GetProfileTemplate(ProfileToCheck, Template))
+		{
+			for (TObjectIterator<UPrimitiveComponent> Iter(RF_PendingKill); Iter; ++Iter)
+			{
+				UPrimitiveComponent* Comp = *Iter;
+				
+				// TEMP CRASH WORKAROUND: IsCollisionEnabled() fails on ULandscapeComponent CDO.
+				if (Cast<ULandscapeHeightfieldCollisionComponent>(Comp))
+				{
+					continue;
+				}
+
+				if (IsQueryCollisionEnabled(Comp))
+				{
+					const ECollisionResponse CompResponse = Comp->GetCollisionResponseToChannel(Template.ObjectType.GetValue());
+					const ECollisionResponse TemplateResponse = Template.ResponseToChannels.GetResponse(Comp->GetCollisionObjectType());
+					const ECollisionResponse MinResponse = FMath::Min<ECollisionResponse>(CompResponse, TemplateResponse);
+					if (MinResponse == RequiredResponse)
+					{
+						Results.Add(Comp);
+					}
+				}
+			}
+		}
+		else
+		{
+			UE_LOG(LogCollisionCommands, Warning, TEXT("Could not find collision profile '%s'. Use 'Collision.ListProfiles' to see a full list of available profiles."), *ProfileToCheck.ToString());
+			return;
+		}
+
+		// Log results.
+		if (Results.Num() > 0)
+		{
+			Results.Sort(FSortComponentsWithResponseToProfile(RequiredResponse));
+
+			// Get max column widths for some data
+			int32 MaxWidth = 0;
+			for (UPrimitiveComponent* Comp : Results)
+			{
+				UObject* Outer = Comp->GetOuter();
+				if (Outer)
+				{
+					const FString PathName = FormatObjectName(Comp);
+					MaxWidth = FMath::Max<int32>(MaxWidth, PathName.Len());
+				}
+			}
+
+			// Column headings
+			if (RequiredResponse == ECollisionResponse::ECR_Overlap)
+			{
+				UE_LOG(LogCollisionCommands, Log, TEXT("  #, GenerateEvents, %-*s, %-16s, Path"), MaxWidth, TEXT("Component"), TEXT("ObjectType"));
+			}
+			else
+			{
+				UE_LOG(LogCollisionCommands, Log, TEXT("  #, %-*s, %-16s, Path"), MaxWidth, TEXT("Component"), TEXT("ObjectType"));
+			}
+
+			// Data
+			int32 Index=0;
+			const UEnum *ChannelEnum = FindObject<UEnum>(ANY_PACKAGE, TEXT("ECollisionChannel"), true);
+			for (UPrimitiveComponent* Comp : Results)
+			{
+				const FString ChannelName = (ChannelEnum ? ChannelEnum->GetNameByIndex((int32)Comp->GetCollisionObjectType()).ToString() : TEXT("<unknown>"));
+				const FString ChannelDisplayName = GetDisplayNameText(ChannelEnum, (int32)Comp->GetCollisionObjectType(), ChannelName);
+				UObject* Outer = Comp->GetOuter();
+				if (Outer)
+				{
+					const FString PathName = FormatObjectName(Comp);
+					if (RequiredResponse == ECollisionResponse::ECR_Overlap)
+					{
+						UE_LOG(LogCollisionCommands, Log, TEXT("%3d, %14s, %-*s, %-16s, %s"),
+							   Index, Comp->bGenerateOverlapEvents ? TEXT("true"):TEXT("false"), MaxWidth, *PathName, *ChannelDisplayName, Outer->GetOuter() ? *GetPathNameSafe(Outer->GetOuter()) : *GetPathNameSafe(Outer));
+					}
+					else
+					{
+						UE_LOG(LogCollisionCommands, Log, TEXT("%3d, %-*s, %-16s, %s"),
+							   Index, MaxWidth, *PathName, *ChannelDisplayName, Outer->GetOuter() ? *GetPathNameSafe(Outer->GetOuter()) : *GetPathNameSafe(Outer));
+					}
+					Index++;
+				}
+			}
+		}
+		check(RequiredResponse < ECollisionResponse::ECR_MAX);
+		UE_LOG(LogCollisionCommands, Log, TEXT("Found %d components with '%s' response to profile '%s'."), Results.Num(), *ResponseStrings[(int32)RequiredResponse], *ProfileToCheck.ToString());
+	}
+
+	//////////////////////////////////////////////////////////////////////////
+	// Args: <Response> <Profile>
+	void Parse_ListComponentsWithResponseToProfile(const TArray<FString>& Args, UWorld* World)
+	{
+		if (Args.Num() < 2)
+		{
+			UE_LOG(LogCollisionCommands, Warning, TEXT("Usage: 'Collision.ListComponentsWithResponseToProfile <Response> <Profile>'."));
+			UE_LOG(LogCollisionCommands, Warning, TEXT("  Response: Ignore, Overlap, Block"));
+			UE_LOG(LogCollisionCommands, Warning, TEXT("  Profile:  Collision profile name or index. Use 'Collision.ListProfiles' to see a full list."));
+			return;
+		}
+
+		// Arg0 : Response
+		const FString& ResponseString = Args[0];
+		ECollisionResponse RequiredResponse = StringToCollisionResponse(ResponseString);
+		if (RequiredResponse == ECollisionResponse::ECR_MAX)
+		{
+			UE_LOG(LogCollisionCommands, Warning, TEXT("Unknown response '%s'. Must be Ignore, Overlap, or Block."), *ResponseString);
+			return;
+		}
+
+		// Arg1 : Profile
+		const FString& ProfileNameString = Args[1];
+		const FName ProfileToCheck = StringToCollisionProfile(ProfileNameString);
+		if (ProfileToCheck == NAME_None)
+		{
+			UE_LOG(LogCollisionCommands, Warning, TEXT("Could not find collision profile '%s'. Use 'Collision.ListProfiles' to see a full list of available profiles."), *ProfileToCheck.ToString());
+			return;
+		}
+
+		ListComponentsWithResponseToProfile(RequiredResponse, ProfileToCheck);
+	}
+
+	//////////////////////////////////////////////////////////////////////////
+	void ListProfilesWithResponseToChannel(const ECollisionResponse RequiredResponse, const ECollisionChannel TestChannel)
+	{
+		TSet<FName> Results;
+		TArray<TSharedPtr<FName>> ProfileNameList;
+		UCollisionProfile::Get()->GetProfileNames(ProfileNameList);
+		int32 Index = 0;
+		for (TSharedPtr<FName> NamePtr : ProfileNameList)
+		{
+			const FName TemplateName = *NamePtr;
+			FCollisionResponseTemplate Template;
+			if (UCollisionProfile::Get()->GetProfileTemplate(TemplateName, Template))
+			{
+				const ECollisionResponse TemplateResponse = Template.ResponseToChannels.GetResponse(TestChannel);
+				if (TemplateResponse == RequiredResponse)
+				{
+					Results.Add(TemplateName);
+				}
+			}
+			++Index;
+		}
+
+		if (Results.Num() > 0)
+		{
+			Results.Sort([](const FName& A, const FName& B) { return A < B; });
+			for (FName& ResultName : Results)
+			{
+				UE_LOG(LogCollisionCommands, Log, TEXT("%s"), *ResultName.ToString());
+			}
+		}
+		check(RequiredResponse < ECollisionResponse::ECR_MAX);
+		const UEnum *ChannelEnum = FindObject<UEnum>(ANY_PACKAGE, TEXT("ECollisionChannel"), true);
+		const FString ChannelName = (ChannelEnum ? ChannelEnum->GetEnumName(TestChannel) : TEXT("<unknown>"));
+		const FString ChannelDisplayName = GetDisplayNameText(ChannelEnum, TestChannel, ChannelName);
+		UE_LOG(LogCollisionCommands, Log, TEXT("Found %d profiles with '%s' response to channel '%s' ('%s')"), Results.Num(), *ResponseStrings[(int32)RequiredResponse], *ChannelName, *ChannelDisplayName);
+	}
+
+	//////////////////////////////////////////////////////////////////////////
+	// Args: <Response> <Channel>
+	void Parse_ListProfilesWithResponseToChannel(const TArray<FString>& Args, UWorld* World)
+	{
+		if (Args.Num() < 2)
+		{
+			UE_LOG(LogCollisionCommands, Warning, TEXT("Usage: 'Collision.ListProfilesWithResponseToChannel <Response> <Channel>'."));
+			UE_LOG(LogCollisionCommands, Warning, TEXT("  Response: Ignore, Overlap, Block"));
+			UE_LOG(LogCollisionCommands, Warning, TEXT("  Profile:  Collision channel name or index. Use 'Collision.ListChannels' to see a full list."));
+			return;
+		}
+
+		// Arg0 : Response
+		const FString& ResponseString = Args[0];
+		const ECollisionResponse RequiredResponse = StringToCollisionResponse(ResponseString);
+		if (RequiredResponse == ECollisionResponse::ECR_MAX)
+		{
+			UE_LOG(LogCollisionCommands, Warning, TEXT("Unknown response '%s'. Must be Ignore, Overlap, or Block."), *ResponseString);
+			return;
+		}
+
+		// Arg1 : Channel
+		const FString& ChannelNameString = Args[1];
+		const ECollisionChannel Channel = StringToCollisionChannel(ChannelNameString);
+		if (Channel == ECollisionChannel::ECC_MAX)
+		{
+			UE_LOG(LogCollisionCommands, Warning, TEXT("Unknown channel '%s. Use 'Collision.ListChannels' to see a full list.'"), *ChannelNameString);
+			return;
+		}
+
+		ListProfilesWithResponseToChannel(RequiredResponse, Channel);
+	}
+
+
+	//////////////////////////////////////////////////////////////////////////
+	// Console commands
+
+	static FAutoConsoleCommand ListProfilesCommand(
+		TEXT("Collision.ListProfiles"),
+		TEXT("ListProfiles"),
+		FConsoleCommandDelegate::CreateStatic(ListCollisionProfileNames)
+	);
+
+	static FAutoConsoleCommand ListChannelsCommand(
+		TEXT("Collision.ListChannels"),
+		TEXT("ListChannels"),
+		FConsoleCommandDelegate::CreateStatic(ListCollisionChannelNames)
+	);
+
+	static FAutoConsoleCommandWithWorldAndArgs ListComponentsWithResponseToProfileCommand(
+		TEXT("Collision.ListComponentsWithResponseToProfile"),
+		TEXT(""),
+		FConsoleCommandWithWorldAndArgsDelegate::CreateStatic(CollisionResponseConsoleCommands::Parse_ListComponentsWithResponseToProfile)
+	);
+
+	static FAutoConsoleCommandWithWorldAndArgs ListProfilesWithResponseToChannelCommand(
+		TEXT("Collision.ListProfilesWithResponseToChannel"),
+		TEXT(""),
+		FConsoleCommandWithWorldAndArgsDelegate::CreateStatic(CollisionResponseConsoleCommands::Parse_ListProfilesWithResponseToChannel)
+	);
+
+}
+
+#endif //!(UE_BUILD_SHIPPING || UE_BUILD_TEST)
+
+
+
+
