Core: Switched GHitchThreshold to be GHitchThresholdMS (now set via t.HitchFrameTimeThreshold, rather than the old t.HitchThreshold) so everything is using the same values on the same scale

[INTEGRATE] Change 2712345 by [REDACTED_PERSON][REDACTED_HANDLE] on 2015/09/30 22:04:15

diff --git a/Engine/Source/Runtime/Core/Private/Misc/Core.cpp b/Engine/Source/Runtime/Core/Private/Misc/Core.cpp
--- a/Engine/Source/Runtime/Core/Private/Misc/Core.cpp
+++ b/Engine/Source/Runtime/Core/Private/Misc/Core.cpp
@@ -216,8 +216,8 @@ uint32					GFrameNumberRenderThread		= 1;
 bool					GIsFirstInstance				= true;
 #endif
 #endif
-/** Threshold for a frame to be considered a hitch (in seconds. */
-float GHitchThreshold = 0.075f;
+/** Threshold for a frame to be considered a hitch (in milliseconds). */
+float GHitchThresholdMS = 60.0f;
 /** Size to break up data into when saving compressed data													*/
 int32					GSavingCompressionChunkSize		= SAVING_COMPRESSION_CHUNK_SIZE;
 /** Whether we are using the seekfree/ cooked loading codepath.												*/
diff --git a/Engine/Source/Runtime/Core/Private/Stats/StatsCommand.cpp b/Engine/Source/Runtime/Core/Private/Stats/StatsCommand.cpp
--- a/Engine/Source/Runtime/Core/Private/Stats/StatsCommand.cpp
+++ b/Engine/Source/Runtime/Core/Private/Stats/StatsCommand.cpp
@@ -509,8 +509,9 @@ static void DumpHitch(int64 Frame)
 
 	const float GameThreadTime = FPlatformTime::ToSeconds(Stats.GetFastThreadFrameTime(Frame, EThreadType::Game));
 	const float RenderThreadTime = FPlatformTime::ToSeconds(Stats.GetFastThreadFrameTime(Frame, EThreadType::Renderer));
+	const float HitchThresholdSecs = GHitchThresholdMS * 0.001f;
 
-	if( GameThreadTime > GHitchThreshold || RenderThreadTime > GHitchThreshold )
+	if ((GameThreadTime > HitchThresholdSecs) || (RenderThreadTime > HitchThresholdSecs))
 	{
 		HitchIndex++;
 		float ThisHitch = FMath::Max<float>(GameThreadTime, RenderThreadTime) * 1000.0f;
@@ -521,7 +522,8 @@ static void DumpHitch(int64 Frame)
 		Stack.AddNameHierarchy();
 		Stack.AddSelf();
 
-		int64 MinCycles = int64( FMath::Min<float>( FMath::Max<float>( GHitchThreshold - 33.3f / 1000.0f, 1.0f / 1000.0f ), 1.0f / 1000.0f ) / FPlatformTime::GetSecondsPerCycle() );
+		const float MinTimeToReportInSecs = 1.0f / 1000.0f;
+		const int64 MinCycles = int64(MinTimeToReportInSecs / FPlatformTime::GetSecondsPerCycle());
 		FRawStatStackNode* GameThread = NULL;
 		FRawStatStackNode* RenderThread = NULL;
 
diff --git a/Engine/Source/Runtime/Core/Public/CoreGlobals.h b/Engine/Source/Runtime/Core/Public/CoreGlobals.h
--- a/Engine/Source/Runtime/Core/Public/CoreGlobals.h
+++ b/Engine/Source/Runtime/Core/Public/CoreGlobals.h
@@ -252,8 +252,8 @@ extern CORE_API bool GIsFirstInstance;
 
 #endif
 
-/** Threshold for a frame to be considered a hitch (in seconds. */
-extern CORE_API float GHitchThreshold;
+/** Threshold for a frame to be considered a hitch (in milliseconds). */
+extern CORE_API float GHitchThresholdMS;
 
 /** Size to break up data into when saving compressed data */
 extern CORE_API int32 GSavingCompressionChunkSize;
diff --git a/Engine/Source/Runtime/Engine/Private/Performance/EnginePerformanceTargets.cpp b/Engine/Source/Runtime/Engine/Private/Performance/EnginePerformanceTargets.cpp
--- a/Engine/Source/Runtime/Engine/Private/Performance/EnginePerformanceTargets.cpp
+++ b/Engine/Source/Runtime/Engine/Private/Performance/EnginePerformanceTargets.cpp
@@ -23,9 +23,9 @@ TAutoConsoleVariable<float> GTargetFrameTimeThresholdCVar(
 
 // The threshold that would be considered so bad that it would cause a hitch in gameplay
 // (*cosmetic only* for reporting purposes such as FPS charts, should not be used in scalability code)
-TAutoConsoleVariable<float> GHitchFrameTimeThresholdCVar(
+FAutoConsoleVariableRef GHitchFrameTimeThresholdCVar(
 	TEXT("t.HitchFrameTimeThreshold"),
-	60.0f,
+	GHitchThresholdMS,
 	TEXT("Definition of a hitchy frame (in ms)\n")
 	TEXT(" default: 60.0 ms"));
 
@@ -58,7 +58,7 @@ float FEnginePerformanceTargets::GetUnacceptableFrameTimeThresholdMS()
 
 float FEnginePerformanceTargets::GetHitchFrameTimeThresholdMS()
 {
-	return GHitchFrameTimeThresholdCVar.GetValueOnGameThread();
+	return GHitchThresholdMS;
 }
 
 float FEnginePerformanceTargets::GetMinTimeBetweenHitchesMS()
diff --git a/Engine/Source/Runtime/Engine/Private/UnrealEngine.cpp b/Engine/Source/Runtime/Engine/Private/UnrealEngine.cpp
--- a/Engine/Source/Runtime/Engine/Private/UnrealEngine.cpp
+++ b/Engine/Source/Runtime/Engine/Private/UnrealEngine.cpp
@@ -166,13 +166,6 @@ ENGINE_API uint32 GGPUFrameTime = 0;
 /** System resolution instance */
 FSystemResolution GSystemResolution;
 
-/** Threshold for a frame to be considered a hitch (in seconds. */
-static FAutoConsoleVariableRef GHitchThresholdCVar(
-	TEXT("t.HitchThreshold"),
-	GHitchThreshold,
-	TEXT("Time in seconds that is considered a hitch by \"stat dumphitches\"")
-	);
-
 TAutoConsoleVariable<int32> CVarAllowOneFrameThreadLag(
 	TEXT("r.OneFrameThreadLag"),
 	1,
@@ -6459,19 +6452,19 @@ float UEngine::GetMaxTickRate(float DeltaTime, bool bAllowFrameRateSmoothing) co
 
 	if (bAllowFrameRateSmoothing && IsAllowedFramerateSmoothing())
 	{
-			// Work in FPS domain as that is what the function will return.
-			MaxTickRate = 1.f / RunningAverageDeltaTime;
+		// Work in FPS domain as that is what the function will return.
+		MaxTickRate = 1.f / RunningAverageDeltaTime;
 
-			// Clamp FPS into ini defined min/ max range.
-			if (SmoothedFrameRateRange.HasLowerBound())
-			{
-				MaxTickRate = FMath::Max( MaxTickRate, SmoothedFrameRateRange.GetLowerBoundValue() );
-			}
-			if (SmoothedFrameRateRange.HasUpperBound())
-			{
-				MaxTickRate = FMath::Min( MaxTickRate, SmoothedFrameRateRange.GetUpperBoundValue() );
-			}
+		// Clamp FPS into ini defined min/ max range.
+		if (SmoothedFrameRateRange.HasLowerBound())
+		{
+			MaxTickRate = FMath::Max( MaxTickRate, SmoothedFrameRateRange.GetLowerBoundValue() );
 		}
+		if (SmoothedFrameRateRange.HasUpperBound())
+		{
+			MaxTickRate = FMath::Min( MaxTickRate, SmoothedFrameRateRange.GetUpperBoundValue() );
+		}
+	}
 
 	if (CVarCauseHitches.GetValueOnGameThread())
 	{
