Integrate CLs 2705944 & 2709164 from Orion to Main
Add null check  in USkeletalMeshComponent::UpdateClothTransform()

diff --git a/Engine/Source/Runtime/Engine/Private/SkeletalMeshComponentPhysics.cpp b/Engine/Source/Runtime/Engine/Private/SkeletalMeshComponentPhysics.cpp
--- a/Engine/Source/Runtime/Engine/Private/SkeletalMeshComponentPhysics.cpp
+++ b/Engine/Source/Runtime/Engine/Private/SkeletalMeshComponentPhysics.cpp
@@ -3190,7 +3190,10 @@ void USkeletalMeshComponent::UpdateClothTransform()
 		if(NxClothingActor* ApexClothingActor = ClothingActor.ApexClothingActor)
 		{
 			NxParameterized::Interface* ActorDesc = ApexClothingActor->getActorDesc();
-			verify(NxParameterized::setParamMat44(*ActorDesc, "globalPose", PxGlobalPose));
+			if(ActorDesc != nullptr)
+			{
+				verify(NxParameterized::setParamMat44(*ActorDesc, "globalPose", PxGlobalPose));
+			}
 		}
 	}
 }
