Integrate 2709161 from Orion to Main
Add bWantsOnUpdateTransform so we only call OnUpdateTransform for components that want it

diff --git a/Engine/Source/Runtime/Engine/Classes/Components/SceneComponent.h b/Engine/Source/Runtime/Engine/Classes/Components/SceneComponent.h
--- a/Engine/Source/Runtime/Engine/Classes/Components/SceneComponent.h
+++ b/Engine/Source/Runtime/Engine/Classes/Components/SceneComponent.h
@@ -173,6 +173,9 @@ protected:
 	// Transient flag that temporarily disables UpdateOverlaps within DetachFromParent().
 	uint32 bDisableDetachmentUpdateOverlaps:1;
 
+	/** If true, OnUpdateTransform virtual will be called each time this component is moved. */
+	uint32 bWantsOnUpdateTransform:1;
+
 private:
 	/** Physics Volume in which this SceneComponent is located **/
 	UPROPERTY(transient)
diff --git a/Engine/Source/Runtime/Engine/Private/Components/AudioComponent.cpp b/Engine/Source/Runtime/Engine/Private/Components/AudioComponent.cpp
--- a/Engine/Source/Runtime/Engine/Private/Components/AudioComponent.cpp
+++ b/Engine/Source/Runtime/Engine/Private/Components/AudioComponent.cpp
@@ -17,6 +17,7 @@ UAudioComponent::UAudioComponent(const FObjectInitializer& ObjectInitializer)
 	bAllowSpatialization = true;
 	bStopWhenOwnerDestroyed = true;
 	bNeverNeedsRenderUpdate = true;
+	bWantsOnUpdateTransform = true;
 #if WITH_EDITORONLY_DATA
 	bVisualizeComponent = true;
 #endif
diff --git a/Engine/Source/Runtime/Engine/Private/Components/PrimitiveComponent.cpp b/Engine/Source/Runtime/Engine/Private/Components/PrimitiveComponent.cpp
--- a/Engine/Source/Runtime/Engine/Private/Components/PrimitiveComponent.cpp
+++ b/Engine/Source/Runtime/Engine/Private/Components/PrimitiveComponent.cpp
@@ -185,6 +185,8 @@ UPrimitiveComponent::UPrimitiveComponent(const FObjectInitializer& ObjectInitial
 	bCanEverAffectNavigation = false;
 	bNavigationRelevant = false;
 
+	bWantsOnUpdateTransform = true;
+
 	bCachedAllCollideableDescendantsRelative = false;
 	LastCheckedAllCollideableDescendantsTime = 0.f;
 }
diff --git a/Engine/Source/Runtime/Engine/Private/Components/SceneComponent.cpp b/Engine/Source/Runtime/Engine/Private/Components/SceneComponent.cpp
--- a/Engine/Source/Runtime/Engine/Private/Components/SceneComponent.cpp
+++ b/Engine/Source/Runtime/Engine/Private/Components/SceneComponent.cpp
@@ -442,15 +442,16 @@ void USceneComponent::PropagateTransformUpdate(bool bTransformChanged, bool bSki
 	{
 		//QUICK_SCOPE_CYCLE_COUNTER(STAT_USceneComponent_PropagateTransformUpdate_TransformChanged);
 		{
-		// Then update bounds
+			// Then update bounds
 			//QUICK_SCOPE_CYCLE_COUNTER(STAT_USceneComponent_PropagateTransformUpdate_UpdateBounds);
-		UpdateBounds();
+			UpdateBounds();
 		}
 
+		// Call OnUpdateTransform if this components wants it
+		if(bWantsOnUpdateTransform)
 		{
 			//QUICK_SCOPE_CYCLE_COUNTER(STAT_USceneComponent_PropagateTransformUpdate_OnUpdateTransform);
-		// Always send new transform to physics
-		OnUpdateTransform(bSkipPhysicsMove, Teleport);
+			OnUpdateTransform(bSkipPhysicsMove, Teleport);
 		}
 
 		// Flag render transform as dirty
@@ -458,9 +459,9 @@ void USceneComponent::PropagateTransformUpdate(bool bTransformChanged, bool bSki
 		
 		{
 			//QUICK_SCOPE_CYCLE_COUNTER(STAT_USceneComponent_PropagateTransformUpdate_UpdateChildTransforms);
-		// Now go and update children
-		//Do not pass skip physics to children. This is only used when physics updates us, but in that case we really do need to update the attached children since they are kinematic
-		UpdateChildTransforms(false, Teleport);
+			// Now go and update children
+			//Do not pass skip physics to children. This is only used when physics updates us, but in that case we really do need to update the attached children since they are kinematic
+			UpdateChildTransforms(false, Teleport);
 		}
 
 		// Refresh navigation
@@ -471,23 +472,23 @@ void USceneComponent::PropagateTransformUpdate(bool bTransformChanged, bool bSki
 		//QUICK_SCOPE_CYCLE_COUNTER(STAT_USceneComponent_PropagateTransformUpdate_NOT_TransformChanged);
 		{
 			//QUICK_SCOPE_CYCLE_COUNTER(STAT_USceneComponent_PropagateTransformUpdate_UpdateBounds);
-		// We update bounds even if transform doesn't change, as shape/mesh etc might have done
-		UpdateBounds();
+			// We update bounds even if transform doesn't change, as shape/mesh etc might have done
+			UpdateBounds();
 		}
 
 		{
 			//QUICK_SCOPE_CYCLE_COUNTER(STAT_USceneComponent_PropagateTransformUpdate_UpdateChildTransforms);
-		// Now go and update children
-		UpdateChildTransforms();
+			// Now go and update children
+			UpdateChildTransforms();
 		}
 
 		{
 			//QUICK_SCOPE_CYCLE_COUNTER(STAT_USceneComponent_PropagateTransformUpdate_MarkRenderTransformDirty);
-		// Need to flag as dirty so new bounds are sent to render thread
-		MarkRenderTransformDirty();
+			// Need to flag as dirty so new bounds are sent to render thread
+			MarkRenderTransformDirty();
+		}
 	}
 }
-}
 
 
 void USceneComponent::EndScopedMovementUpdate(class FScopedMovementUpdate& CompletedScope)
diff --git a/Engine/Source/Runtime/Engine/Private/Particles/ParticleComponents.cpp b/Engine/Source/Runtime/Engine/Private/Particles/ParticleComponents.cpp
--- a/Engine/Source/Runtime/Engine/Private/Particles/ParticleComponents.cpp
+++ b/Engine/Source/Runtime/Engine/Private/Particles/ParticleComponents.cpp
@@ -2756,6 +2756,9 @@ UParticleSystemComponent::UParticleSystemComponent(const FObjectInitializer& Obj
 	// Disable receiving decals by default.
 	bReceivesDecals = false;
 
+	// Don't need to call OnUpdateTransform, no physics state to update
+	bWantsOnUpdateTransform = false;
+
 	SavedAutoAttachRelativeScale3D = FVector(1.f, 1.f, 1.f);
 }
 
