Fix issue with Root Motion from Everything being calculated incorrectly when animations are part of a sync group.

diff --git a/Engine/Source/Runtime/Engine/Private/Animation/AnimInstance.cpp b/Engine/Source/Runtime/Engine/Private/Animation/AnimInstance.cpp
--- a/Engine/Source/Runtime/Engine/Private/Animation/AnimInstance.cpp
+++ b/Engine/Source/Runtime/Engine/Private/Animation/AnimInstance.cpp
@@ -494,6 +494,7 @@ void UAnimInstance::UpdateAnimation(float DeltaSeconds)
 							{
 								QUICK_SCOPE_CYCLE_COUNTER(STAT_UAnimInstance_UpdateAnimation_TickAssetPlayerInstance);
 								FScopeCycleCounterUObject Scope(AssetPlayer.SourceAsset);
+								TickContext.RootMotionMovementParams.Clear();
 								AssetPlayer.SourceAsset->TickAssetPlayerInstance(AssetPlayer, this, TickContext);
 							}
 							if (RootMotionMode == ERootMotionMode::RootMotionFromEverything && TickContext.RootMotionMovementParams.bHasRootMotion)
