Merging using Orion->UE4

CL# 2712900
safe iteration in RouteActorInitialize

diff --git a/Engine/Source/Runtime/Engine/Private/Level.cpp b/Engine/Source/Runtime/Engine/Private/Level.cpp
--- a/Engine/Source/Runtime/Engine/Private/Level.cpp
+++ b/Engine/Source/Runtime/Engine/Private/Level.cpp
@@ -1688,8 +1688,9 @@ void ULevel::RouteActorInitialize()
 	TArray<AActor *> ActorsToBeginPlay;
 
 	// Send InitializeComponents on components and PostInitializeComponents.
-	for( AActor* const Actor : Actors )
+	for( int32 Index = 0; Index < Actors.Num(); ++Index )
 	{
+		AActor* const Actor = Actors[Index];
 		if( Actor )
 		{
 			if( !Actor->IsActorInitialized() )
