Remove unused function (equivalent Particle System function removed in CL#1404459)

diff --git a/Engine/Source/Runtime/Engine/Classes/Animation/AnimInstance.h b/Engine/Source/Runtime/Engine/Classes/Animation/AnimInstance.h
--- a/Engine/Source/Runtime/Engine/Classes/Animation/AnimInstance.h
+++ b/Engine/Source/Runtime/Engine/Classes/Animation/AnimInstance.h
@@ -430,9 +430,6 @@ public:
 
 	bool CanTransitionSignature() const;
 	
-	UFUNCTION()
-	void AnimNotify_Sound(UAnimNotify* Notify);
-
 	/*********************************************************************************************
 	* SlotAnimation
 	********************************************************************************************* */
diff --git a/Engine/Source/Runtime/Engine/Private/Animation/AnimInstance.cpp b/Engine/Source/Runtime/Engine/Private/Animation/AnimInstance.cpp
--- a/Engine/Source/Runtime/Engine/Private/Animation/AnimInstance.cpp
+++ b/Engine/Source/Runtime/Engine/Private/Animation/AnimInstance.cpp
@@ -1389,11 +1389,6 @@ void UAnimInstance::TriggerSingleAnimNotify(const FAnimNotifyEvent* AnimNotifyEv
 	}
 }
 
-void UAnimInstance::AnimNotify_Sound(UAnimNotify* AnimNotify)
-{
-	AnimNotify->Notify(GetSkelMeshComponent(), NULL);
-}
-
 //to debug montage weight
 #define DEBUGMONTAGEWEIGHT 0
 
