#ue4 - Optimize UObject::IsNameStableForNetworking() to check cheap flags first before expensive IsDefaultSubObject(). Optimize UActorComponent::IsSupportedForNetworking() to check GetIsReplicated() before IsNameStableForNetworking(), and inlined GetIsReplicated().

--------

diff --git a/Engine/Source/Runtime/CoreUObject/Private/UObject/Obj.cpp b/Engine/Source/Runtime/CoreUObject/Private/UObject/Obj.cpp
--- a/Engine/Source/Runtime/CoreUObject/Private/UObject/Obj.cpp
+++ b/Engine/Source/Runtime/CoreUObject/Private/UObject/Obj.cpp
@@ -3847,7 +3847,7 @@ void UObject::PreDestroyFromReplication()
 /** IsNameStableForNetworking means an object can be referred to its path name (relative to outer) over the network */
 bool UObject::IsNameStableForNetworking() const
 {
-	return IsDefaultSubobject() || HasAnyFlags( RF_WasLoaded | RF_DefaultSubObject | RF_Native );
+	return HasAnyFlags(RF_WasLoaded | RF_DefaultSubObject | RF_Native) || IsDefaultSubobject();
 }
 
 /** IsFullNameStableForNetworking means an object can be referred to its full path name over the network */
diff --git a/Engine/Source/Runtime/Engine/Classes/Components/ActorComponent.h b/Engine/Source/Runtime/Engine/Classes/Components/ActorComponent.h
--- a/Engine/Source/Runtime/Engine/Classes/Components/ActorComponent.h
+++ b/Engine/Source/Runtime/Engine/Classes/Components/ActorComponent.h
@@ -274,7 +274,10 @@ public:
 	void SetIsReplicated(bool ShouldReplicate);
 
 	/** Returns whether replication is enabled or not. */
-	bool GetIsReplicated() const;
+	FORCEINLINE bool GetIsReplicated() const
+	{
+		return bReplicates;
+	}
 
 	/** Allows a component to replicate other subobject on the actor  */
 	virtual bool ReplicateSubobjects(class UActorChannel *Channel, class FOutBunch *Bunch, FReplicationFlags *RepFlags);
diff --git a/Engine/Source/Runtime/Engine/Private/Components/ActorComponent.cpp b/Engine/Source/Runtime/Engine/Private/Components/ActorComponent.cpp
--- a/Engine/Source/Runtime/Engine/Private/Components/ActorComponent.cpp
+++ b/Engine/Source/Runtime/Engine/Private/Components/ActorComponent.cpp
@@ -1449,7 +1449,7 @@ bool UActorComponent::IsNameStableForNetworking() const
 
 bool UActorComponent::IsSupportedForNetworking() const
 {
-	return IsNameStableForNetworking() || GetIsReplicated();
+	return GetIsReplicated() || IsNameStableForNetworking();
 }
 
 void UActorComponent::SetIsReplicated(bool ShouldReplicate)
@@ -1463,11 +1463,6 @@ void UActorComponent::SetIsReplicated(bool ShouldReplicate)
 	}
 }
 
-bool UActorComponent::GetIsReplicated() const
-{
-	return bReplicates;
-}
-
 bool UActorComponent::ReplicateSubobjects(class UActorChannel *Channel, class FOutBunch *Bunch, FReplicationFlags *RepFlags)
 {
 	return false;
