#ue4 - Added 'Profile' to Collision.ListComponentsWithRepsonseToProfile.

--------

diff --git a/Engine/Source/Runtime/Engine/Private/Collision/Collision.cpp b/Engine/Source/Runtime/Engine/Private/Collision/Collision.cpp
--- a/Engine/Source/Runtime/Engine/Private/Collision/Collision.cpp
+++ b/Engine/Source/Runtime/Engine/Private/Collision/Collision.cpp
@@ -525,6 +525,20 @@ namespace CollisionResponseConsoleCommands
 {
 	static const FString ResponseStrings[] = {TEXT("Ignore"), TEXT("Overlap"), TEXT("Block")};
 
+	FString FillString(TCHAR Char, int32 Count)
+	{
+		FString Result;
+		if (Count > 0)
+		{
+			Result.Reserve(Count);
+			for (int32 i=0; i < Count; i++)
+			{
+				Result.AppendChar(TEXT('-'));
+			}
+		}
+		return Result;
+	}
+
 	FString GetDisplayNameText(const UEnum* Enum, int32 Index, const FString& Fallback)
 	{
 		if (Enum)
@@ -721,32 +735,48 @@ namespace CollisionResponseConsoleCommands
 		if (Results.Num() > 0)
 		{
 			Results.Sort(FSortComponentsWithResponseToProfile(RequiredResponse));
+			const UEnum *ChannelEnum = FindObject<UEnum>(ANY_PACKAGE, TEXT("ECollisionChannel"), true);
 
 			// Get max column widths for some data
-			int32 MaxWidth = 0;
+			int32 MaxNameWidth = 0;
+			int32 MaxChannelWidth = 0;
+			int32 MaxProfileWidth = 0;
 			for (UPrimitiveComponent* Comp : Results)
 			{
 				UObject* Outer = Comp->GetOuter();
 				if (Outer)
 				{
 					const FString PathName = FormatObjectName(Comp);
-					MaxWidth = FMath::Max<int32>(MaxWidth, PathName.Len());
+					MaxNameWidth = FMath::Max<int32>(MaxNameWidth, PathName.Len());
 				}
+
+				if (ChannelEnum)
+				{
+					const FString ChannelName = ChannelEnum->GetNameByIndex((int32)Comp->GetCollisionObjectType()).ToString();
+					const FString ChannelDisplayName = GetDisplayNameText(ChannelEnum, (int32)Comp->GetCollisionObjectType(), ChannelName);
+					MaxChannelWidth = FMath::Max<int32>(MaxChannelWidth, ChannelDisplayName.Len());
+				}
+
+				MaxProfileWidth = FMath::Max<int32>(MaxProfileWidth, Comp->GetCollisionProfileName().ToString().Len());
 			}
 
 			// Column headings
+			FString Output;
 			if (RequiredResponse == ECollisionResponse::ECR_Overlap)
 			{
-				UE_LOG(LogCollisionCommands, Log, TEXT("  #, GenerateEvents, %-*s, %-16s, Path"), MaxWidth, TEXT("Component"), TEXT("ObjectType"));
+				Output = FString::Printf(TEXT("  #, GenerateEvents, %-*s, %-*s, %-*s, Path"), MaxNameWidth, TEXT("Component"), MaxChannelWidth, TEXT("ObjectType"), MaxProfileWidth, TEXT("Profile"));
 			}
 			else
 			{
-				UE_LOG(LogCollisionCommands, Log, TEXT("  #, %-*s, %-16s, Path"), MaxWidth, TEXT("Component"), TEXT("ObjectType"));
+				Output = FString::Printf(TEXT("  #, %-*s, %-*s, %-*s, Path"), MaxNameWidth, TEXT("Component"), MaxChannelWidth, TEXT("ObjectType"), MaxProfileWidth, TEXT("Profile"));
 			}
+			UE_LOG(LogCollisionCommands, Log, TEXT("%s"), *Output);
+			const int32 TotalLen = Output.Len() + 16;
+			const FString LineMarker = FillString(TCHAR('-'), TotalLen);
+			UE_LOG(LogCollisionCommands, Log, TEXT("%s"), *LineMarker);
 
 			// Data
 			int32 Index=0;
-			const UEnum *ChannelEnum = FindObject<UEnum>(ANY_PACKAGE, TEXT("ECollisionChannel"), true);
 			for (UPrimitiveComponent* Comp : Results)
 			{
 				const FString ChannelName = (ChannelEnum ? ChannelEnum->GetNameByIndex((int32)Comp->GetCollisionObjectType()).ToString() : TEXT("<unknown>"));
@@ -754,20 +784,22 @@ namespace CollisionResponseConsoleCommands
 				UObject* Outer = Comp->GetOuter();
 				if (Outer)
 				{
+					const FString OtherProfileName = Comp->GetCollisionProfileName().ToString();
 					const FString PathName = FormatObjectName(Comp);
 					if (RequiredResponse == ECollisionResponse::ECR_Overlap)
 					{
-						UE_LOG(LogCollisionCommands, Log, TEXT("%3d, %-14s, %-*s, %-16s, %s"),
-							   Index, Comp->bGenerateOverlapEvents ? TEXT("true"):TEXT("false"), MaxWidth, *PathName, *ChannelDisplayName, Outer->GetOuter() ? *GetPathNameSafe(Outer->GetOuter()) : *GetPathNameSafe(Outer));
+						UE_LOG(LogCollisionCommands, Log, TEXT("%3d, %-14s, %-*s, %-*s, %-*s, %s"),
+							   Index, Comp->bGenerateOverlapEvents ? TEXT("true"):TEXT("false"), MaxNameWidth, *PathName, MaxChannelWidth, *ChannelDisplayName, MaxProfileWidth, *OtherProfileName, Outer->GetOuter() ? *GetPathNameSafe(Outer->GetOuter()) : *GetPathNameSafe(Outer));
 					}
 					else
 					{
-						UE_LOG(LogCollisionCommands, Log, TEXT("%3d, %-*s, %-16s, %s"),
-							   Index, MaxWidth, *PathName, *ChannelDisplayName, Outer->GetOuter() ? *GetPathNameSafe(Outer->GetOuter()) : *GetPathNameSafe(Outer));
+						UE_LOG(LogCollisionCommands, Log, TEXT("%3d, %-*s, %-*s, %-*s, %s"),
+							   Index, MaxNameWidth, *PathName, MaxChannelWidth, *ChannelDisplayName, MaxProfileWidth, *OtherProfileName, Outer->GetOuter() ? *GetPathNameSafe(Outer->GetOuter()) : *GetPathNameSafe(Outer));
 					}
 					Index++;
 				}
 			}
+			UE_LOG(LogCollisionCommands, Log, TEXT("%s"), *LineMarker);
 		}
 		check(RequiredResponse < ECollisionResponse::ECR_MAX);
 		UE_LOG(LogCollisionCommands, Log, TEXT("Found %d components with '%s' response to profile '%s'."), Results.Num(), *ResponseStrings[(int32)RequiredResponse], *ProfileToCheck.ToString());
@@ -840,6 +872,7 @@ namespace CollisionResponseConsoleCommands
 		const UEnum *ChannelEnum = FindObject<UEnum>(ANY_PACKAGE, TEXT("ECollisionChannel"), true);
 		const FString ChannelName = (ChannelEnum ? ChannelEnum->GetEnumName(TestChannel) : TEXT("<unknown>"));
 		const FString ChannelDisplayName = GetDisplayNameText(ChannelEnum, TestChannel, ChannelName);
+		UE_LOG(LogCollisionCommands, Log, TEXT("----------------------------------------------------------------------"));
 		UE_LOG(LogCollisionCommands, Log, TEXT("Found %d profiles with '%s' response to channel '%s' ('%s')"), Results.Num(), *ResponseStrings[(int32)RequiredResponse], *ChannelName, *ChannelDisplayName);
 	}
 
