#ue4 - UActorComponent::ConditionalTickComponent() is now private as it should have been. It is now inlined within the source file for the single use of it, with a redundant IsPendingKill() check removed.

--------

diff --git a/Engine/Source/Runtime/Engine/Classes/Components/ActorComponent.h b/Engine/Source/Runtime/Engine/Classes/Components/ActorComponent.h
--- a/Engine/Source/Runtime/Engine/Classes/Components/ActorComponent.h
+++ b/Engine/Source/Runtime/Engine/Classes/Components/ActorComponent.h
@@ -465,6 +465,7 @@ public:
 	 */
 	void RegisterComponentWithWorld(UWorld* InWorld);
 	
+private:
 	/**
 	 * Conditionally calls Tick if bRegistered == true and a bunch of other criteria are met
 	 * @param DeltaTime - The time since the last tick.
@@ -473,6 +474,11 @@ public:
 	 */
 	void ConditionalTickComponent(float DeltaTime, enum ELevelTick TickType, FActorComponentTickFunction &ThisTickFunction);
 
+	/** Friend for access to ConditionalTickComponent(). */
+	friend struct FActorComponentTickFunction;
+
+public:
+
 	/**
 	 * Returns whether the component's owner is selected.
 	 */
diff --git a/Engine/Source/Runtime/Engine/Private/Components/ActorComponent.cpp b/Engine/Source/Runtime/Engine/Private/Components/ActorComponent.cpp
--- a/Engine/Source/Runtime/Engine/Private/Components/ActorComponent.cpp
+++ b/Engine/Source/Runtime/Engine/Private/Components/ActorComponent.cpp
@@ -1157,9 +1157,9 @@ void UActorComponent::RecreatePhysicsState()
 	}
 }
 
-void UActorComponent::ConditionalTickComponent(float DeltaTime, enum ELevelTick TickType, FActorComponentTickFunction &ThisTickFunction)
+FORCEINLINE_DEBUGGABLE void UActorComponent::ConditionalTickComponent(float DeltaTime, enum ELevelTick TickType, FActorComponentTickFunction &ThisTickFunction)
 {
-	if(bRegistered && !IsPendingKill())
+	if(bRegistered)
 	{
 		AActor* MyOwner = GetOwner();
 		//@optimization, I imagine this is all unnecessary in a shipping game with no editor
