Fixed missing streaming levels

diff --git a/Engine/Source/Runtime/Engine/Private/LevelStreaming.cpp b/Engine/Source/Runtime/Engine/Private/LevelStreaming.cpp
--- a/Engine/Source/Runtime/Engine/Private/LevelStreaming.cpp
+++ b/Engine/Source/Runtime/Engine/Private/LevelStreaming.cpp
@@ -767,7 +767,8 @@ void ULevelStreaming::SetWorldAsset(const TAssetPtr<UWorld>& NewWorldAsset)
 
 FString ULevelStreaming::GetWorldAssetPackageName() const
 {
-	return GetWorldAssetPackageFName().ToString();
+	const FString WorldAssetPath = WorldAsset.ToStringReference().ToString();
+	return FPackageName::ObjectPathToPackageName(WorldAssetPath);
 }
 
 FName ULevelStreaming::GetWorldAssetPackageFName() const
