PR #1616: [OPT] DispatchPhysNotifications_AssumesLocked less memory allocations (Contributed by [REDACTED_PERSON])

diff --git a/Engine/Source/Runtime/Engine/Private/PhysicsEngine/PhysScene.cpp b/Engine/Source/Runtime/Engine/Private/PhysicsEngine/PhysScene.cpp
--- a/Engine/Source/Runtime/Engine/Private/PhysicsEngine/PhysScene.cpp
+++ b/Engine/Source/Runtime/Engine/Private/PhysicsEngine/PhysScene.cpp
@@ -980,7 +980,7 @@ void FPhysScene::DispatchPhysNotifications_AssumesLocked()
 				}
 			}
 		}
-		PendingCollisionNotifies.Empty();
+		PendingCollisionNotifies.Reset();
 	}
 
 #if WITH_APEX
