Option to create CompositeNavModifiers from convex and sphear shapes added #UE4

github PR#1619 by lion03

UE-21734

diff --git a/Engine/Source/Runtime/Engine/Private/AI/Navigation/NavigationModifier.cpp b/Engine/Source/Runtime/Engine/Private/AI/Navigation/NavigationModifier.cpp
--- a/Engine/Source/Runtime/Engine/Private/AI/Navigation/NavigationModifier.cpp
+++ b/Engine/Source/Runtime/Engine/Private/AI/Navigation/NavigationModifier.cpp
@@ -655,7 +655,21 @@ void FCompositeNavModifier::CreateAreaModifiers(const UPrimitiveComponent* PrimC
 		Add(AreaMod);
 	}
 
-	// convex elements support?
+	for (int32 Idx = 0; Idx < BodySetup->AggGeom.ConvexElems.Num(); Idx++)
+	{
+		const FKConvexElem& ConvexElem = BodySetup->AggGeom.ConvexElems[Idx];
+		
+		FAreaNavModifier AreaMod(ConvexElem.VertexData, 0, ConvexElem.VertexData.Num(), ENavigationCoordSystem::Unreal, PrimComp->ComponentToWorld, AreaClass);
+		Add(AreaMod);
+	}
+	
+	for (int32 Idx = 0; Idx < BodySetup->AggGeom.SphereElems.Num(); Idx++)
+	{
+		const FKSphereElem& SphereElem = BodySetup->AggGeom.SphereElems[Idx];
+		
+		FAreaNavModifier AreaMod(SphereElem.Radius, SphereElem.Radius, PrimComp->ComponentToWorld, AreaClass);
+		Add(AreaMod);
+	}
 }
 
 uint32 FCompositeNavModifier::GetAllocatedSize() const
