Using "Draw Debug Float History" and "Draw Debug Transform History" in Blueprints will no longer crash.

Added metadata to properly hook up the WorldContextObject parameter.

diff --git a/Engine/Source/Runtime/Engine/Classes/Kismet/KismetSystemLibrary.h b/Engine/Source/Runtime/Engine/Classes/Kismet/KismetSystemLibrary.h
--- a/Engine/Source/Runtime/Engine/Classes/Kismet/KismetSystemLibrary.h
+++ b/Engine/Source/Runtime/Engine/Classes/Kismet/KismetSystemLibrary.h
@@ -1312,11 +1312,11 @@ class ENGINE_API UKismetSystemLibrary : public UBlueprintFunctionLibrary
 	static void DrawDebugCamera(const ACameraActor* CameraActor, FLinearColor CameraColor = FLinearColor::White, float Duration=0.f);
 
 	/* Draws a 2D Histogram of size 'DrawSize' based FDebugFloatHistory struct, using DrawTransform for the position in the world. */
-	UFUNCTION(BlueprintCallable, Category = "Rendering|Debug")
+	UFUNCTION(BlueprintCallable, Category = "Rendering|Debug", meta=(WorldContext="WorldContextObject"))
 	static void DrawDebugFloatHistoryTransform(UObject* WorldContextObject, const FDebugFloatHistory& FloatHistory, const FTransform& DrawTransform, FVector2D DrawSize, FLinearColor DrawColor = FLinearColor::White, float Duration = 0.f);
 
 	/* Draws a 2D Histogram of size 'DrawSize' based FDebugFloatHistory struct, using DrawLocation for the location in the world, rotation will face camera of first player. */
-	UFUNCTION(BlueprintCallable, Category = "Rendering|Debug")
+	UFUNCTION(BlueprintCallable, Category = "Rendering|Debug", meta=(WorldContext="WorldContextObject"))
 	static void DrawDebugFloatHistoryLocation(UObject* WorldContextObject, const FDebugFloatHistory& FloatHistory, FVector DrawLocation, FVector2D DrawSize, FLinearColor DrawColor = FLinearColor::White, float Duration = 0.f);
 
 	UFUNCTION(BlueprintPure, Category = "Rendering|Debug")
