#ue4 - Add ability to match by channel displayname in Collision.ListProfilesWithResponseToChannel

--------

diff --git a/Engine/Source/Runtime/Engine/Private/Collision/Collision.cpp b/Engine/Source/Runtime/Engine/Private/Collision/Collision.cpp
--- a/Engine/Source/Runtime/Engine/Private/Collision/Collision.cpp
+++ b/Engine/Source/Runtime/Engine/Private/Collision/Collision.cpp
@@ -666,6 +666,17 @@ namespace CollisionResponseConsoleCommands
 			{
 				return ECollisionChannel(ChannelInt);
 			}
+
+			// Try matching the display name
+			const FString NullString(TEXT(""));
+			for (int32 ChannelIndex = 0; ChannelIndex < ChannelEnum->NumEnums(); ChannelIndex++)
+			{
+				const FString ChannelDisplayName = GetDisplayNameText(ChannelEnum, ChannelIndex, NullString);
+				if (ChannelDisplayName == InString)
+				{
+					return ECollisionChannel(ChannelEnum->GetValueByIndex(ChannelIndex));
+				}
+			}
 		}
 
 		// Try parsing a digit, as in from the ListChannels command
