Fix resave message

CL 2710016

diff --git a/Engine/Source/Runtime/Engine/Private/Animation/AnimSequence.cpp b/Engine/Source/Runtime/Engine/Private/Animation/AnimSequence.cpp
--- a/Engine/Source/Runtime/Engine/Private/Animation/AnimSequence.cpp
+++ b/Engine/Source/Runtime/Engine/Private/Animation/AnimSequence.cpp
@@ -529,22 +529,30 @@ void UAnimSequence::PostLoad()
 }
 
 #if WITH_EDITOR
-void UAnimSequence::VerifyTrackMap(USkeleton* MySkeleton)
+void ShowResaveMessage(const UAnimSequence* Sequence)
 {
-	USkeleton* UseSkeleton = (MySkeleton)? MySkeleton: GetSkeleton();
-
-	if( AnimationTrackNames.Num() != TrackToSkeletonMapTable.Num() && UseSkeleton!=NULL)
+	if (!IsRunningGame())
 	{
-		UE_LOG(LogAnimation, Warning, TEXT("RESAVE ANIMATION NEEDED(%s): Fixing track names."), *GetName());
+		UE_LOG(LogAnimation, Warning, TEXT("RESAVE ANIMATION NEEDED(%s): Fixing track data."), *GetNameSafe(Sequence));
 
 		static FName NAME_LoadErrors("LoadErrors");
 		FMessageLog LoadErrors(NAME_LoadErrors);
 
 		TSharedRef<FTokenizedMessage> Message = LoadErrors.Warning();
 		Message->AddToken(FTextToken::Create(LOCTEXT("AnimationNeedsResave1", "The Animation ")));
-		Message->AddToken(FAssetNameToken::Create(GetPathName(), FText::FromString(GetNameSafe(this))));
+		Message->AddToken(FAssetNameToken::Create(Sequence->GetPathName(), FText::FromString(GetNameSafe(Sequence))));
 		Message->AddToken(FTextToken::Create(LOCTEXT("AnimationNeedsResave2", " needs resave.")));
-		LoadErrors.Open();
+		LoadErrors.Notify();
+	}
+}
+
+void UAnimSequence::VerifyTrackMap(USkeleton* MySkeleton)
+{
+	USkeleton* UseSkeleton = (MySkeleton)? MySkeleton: GetSkeleton();
+
+	if( AnimationTrackNames.Num() != TrackToSkeletonMapTable.Num() && UseSkeleton!=NULL)
+	{
+		ShowResaveMessage(this);
 
 		const TArray<FBoneNode>& BoneTree = UseSkeleton->GetBoneTree();
 		AnimationTrackNames.Empty();
@@ -580,17 +588,7 @@ void UAnimSequence::VerifyTrackMap(USkeleton* MySkeleton)
 
 			if(bNeedsFixing)
 			{
-				UE_LOG(LogAnimation, Warning, TEXT("RESAVE ANIMATION NEEDED(%s): Fixing track index."), *GetName());
-
-				static FName NAME_LoadErrors("LoadErrors");
-				FMessageLog LoadErrors(NAME_LoadErrors);
-
-				TSharedRef<FTokenizedMessage> Message = LoadErrors.Warning();
-				Message->AddToken(FTextToken::Create(LOCTEXT("AnimationNeedsResave1", "Skeleton has been changed. The Animation ")));
-				Message->AddToken(FAssetNameToken::Create(GetPathName(), FText::FromString(GetNameSafe(this))));
-				Message->AddToken(FTextToken::Create(LOCTEXT("AnimationNeedsResave2", " needs resave.")));
-
-				LoadErrors.Open();
+				ShowResaveMessage(this);
 
 				const TArray<FBoneNode>& BoneTree = UseSkeleton->GetBoneTree();
 				for(int32 I=NumTracks-1; I>=0; --I)
