Added rebuilding index buffer for static LOD model

Merging

//depot/UE4-Orion/Engine/Source/...

to //depot/UE4/Engine/Source/...

diff --git a/Engine/Source/Developer/MeshBoneReduction/Private/MeshBoneReduction.cpp b/Engine/Source/Developer/MeshBoneReduction/Private/MeshBoneReduction.cpp
--- a/Engine/Source/Developer/MeshBoneReduction/Private/MeshBoneReduction.cpp
+++ b/Engine/Source/Developer/MeshBoneReduction/Private/MeshBoneReduction.cpp
@@ -251,9 +251,10 @@ public:
 		SrcModel.LegacyRawPointIndices.Unlock();
 
 		// The index buffer needs to be rebuilt on copy.
-		FMultiSizeIndexContainerData IndexBufferData;
+		FMultiSizeIndexContainerData IndexBufferData, AdjacencyIndexBufferData;
 		SrcModel.MultiSizeIndexContainer.GetIndexBufferData( IndexBufferData );
-		NewModel->MultiSizeIndexContainer.RebuildIndexBuffer( IndexBufferData );
+		SrcModel.AdjacencyMultiSizeIndexContainer.GetIndexBufferData( AdjacencyIndexBufferData );
+		NewModel->RebuildIndexBuffer( &IndexBufferData, &AdjacencyIndexBufferData );
 
 		// Required bones are recalculated later on.
 		NewModel->RequiredBones.Empty();
diff --git a/Engine/Source/Developer/SimplygonMeshReduction/Private/SimplygonMeshReduction.cpp b/Engine/Source/Developer/SimplygonMeshReduction/Private/SimplygonMeshReduction.cpp
--- a/Engine/Source/Developer/SimplygonMeshReduction/Private/SimplygonMeshReduction.cpp
+++ b/Engine/Source/Developer/SimplygonMeshReduction/Private/SimplygonMeshReduction.cpp
@@ -325,9 +325,10 @@ public:
 			SrcModel->LegacyRawPointIndices.Unlock();
 
 			// The index buffer needs to be rebuilt on copy.
-			FMultiSizeIndexContainerData IndexBufferData;
-			SrcModel->MultiSizeIndexContainer.GetIndexBufferData( IndexBufferData );
-			NewSrcModel->MultiSizeIndexContainer.RebuildIndexBuffer( IndexBufferData );
+			FMultiSizeIndexContainerData IndexBufferData, AdjacencyIndexBufferData;
+			SrcModel->MultiSizeIndexContainer.GetIndexBufferData(IndexBufferData);
+			SrcModel->AdjacencyMultiSizeIndexContainer.GetIndexBufferData(AdjacencyIndexBufferData);
+			NewSrcModel->RebuildIndexBuffer(&IndexBufferData, &AdjacencyIndexBufferData);
 
 			// now fix up SrcModel to NewSrcModel
 			SrcModel = NewSrcModel;
@@ -365,9 +366,10 @@ public:
 			SrcModel->LegacyRawPointIndices.Unlock();
 
 			// The index buffer needs to be rebuilt on copy.
-			FMultiSizeIndexContainerData IndexBufferData;
-			SrcModel->MultiSizeIndexContainer.GetIndexBufferData( IndexBufferData );
-			NewModel->MultiSizeIndexContainer.RebuildIndexBuffer( IndexBufferData );
+			FMultiSizeIndexContainerData IndexBufferData, AdjacencyIndexBufferData;
+			SrcModel->MultiSizeIndexContainer.GetIndexBufferData(IndexBufferData);
+			SrcModel->AdjacencyMultiSizeIndexContainer.GetIndexBufferData(AdjacencyIndexBufferData);
+			NewModel->RebuildIndexBuffer(&IndexBufferData, &AdjacencyIndexBufferData);
 
 			// Required bones are recalculated later on.
 			NewModel->RequiredBones.Empty();
diff --git a/Engine/Source/Editor/UnrealEd/Private/Factories/SkeletalMeshImport.cpp b/Engine/Source/Editor/UnrealEd/Private/Factories/SkeletalMeshImport.cpp
--- a/Engine/Source/Editor/UnrealEd/Private/Factories/SkeletalMeshImport.cpp
+++ b/Engine/Source/Editor/UnrealEd/Private/Factories/SkeletalMeshImport.cpp
@@ -888,8 +888,7 @@ void RestoreExistingSkelMeshData(ExistingSkelMeshData* MeshData, USkeletalMesh*
 				{
 					FStaticLODModel* NewLODModel = new(SkeletalMesh->GetImportedResource()->LODModels) FStaticLODModel( LODModel );
 			
-					NewLODModel->MultiSizeIndexContainer.RebuildIndexBuffer( MeshData->ExistingIndexBufferData[i] );
-					NewLODModel->AdjacencyMultiSizeIndexContainer.RebuildIndexBuffer( MeshData->ExistingAdjacencyIndexBufferData[i] );
+					NewLODModel->RebuildIndexBuffer( &MeshData->ExistingIndexBufferData[i], &MeshData->ExistingAdjacencyIndexBufferData[i]);
 
 					SkeletalMesh->LODInfo.Add( LODInfo );
 				}
diff --git a/Engine/Source/Editor/UnrealEd/Private/Fbx/FbxSkeletalMeshImport.cpp b/Engine/Source/Editor/UnrealEd/Private/Fbx/FbxSkeletalMeshImport.cpp
--- a/Engine/Source/Editor/UnrealEd/Private/Fbx/FbxSkeletalMeshImport.cpp
+++ b/Engine/Source/Editor/UnrealEd/Private/Fbx/FbxSkeletalMeshImport.cpp
@@ -2465,12 +2465,13 @@ void UnFbx::FFbxImporter::InsertNewLODToBaseSkeletalMesh(USkeletalMesh* InSkelet
 	// Index buffer will be destroyed when we copy the LOD model so we must copy the index buffer and reinitialize it after the copy
 	FMultiSizeIndexContainerData Data;
 	NewLODModel.MultiSizeIndexContainer.GetIndexBufferData(Data);
+	FMultiSizeIndexContainerData AdjacencyData;
+	NewLODModel.AdjacencyMultiSizeIndexContainer.GetIndexBufferData(AdjacencyData);
 
 	// Assign new FStaticLODModel to desired slot in selected skeletal mesh.
 	DestImportedResource->LODModels[DesiredLOD] = NewLODModel;
 
-	DestImportedResource->LODModels[DesiredLOD].MultiSizeIndexContainer.RebuildIndexBuffer(Data);
-
+	DestImportedResource->LODModels[DesiredLOD].RebuildIndexBuffer(&Data, &AdjacencyData);
 	// rebuild vertex buffers and reinit RHI resources
 	BaseSkeletalMesh->PostEditChange();
 }
diff --git a/Engine/Source/Runtime/Engine/Private/SkeletalMesh.cpp b/Engine/Source/Runtime/Engine/Private/SkeletalMesh.cpp
--- a/Engine/Source/Runtime/Engine/Private/SkeletalMesh.cpp
+++ b/Engine/Source/Runtime/Engine/Private/SkeletalMesh.cpp
@@ -1667,6 +1667,31 @@ SIZE_T FStaticLODModel::GetResourceSize() const
 
 	return ResourceSize;
 }
+
+void FStaticLODModel::RebuildIndexBuffer(FMultiSizeIndexContainerData* IndexBufferData, FMultiSizeIndexContainerData* AdjacencyIndexBufferData)
+{
+	if (IndexBufferData)
+	{
+		MultiSizeIndexContainer.RebuildIndexBuffer(*IndexBufferData);
+	}
+
+	if (AdjacencyIndexBufferData)
+	{
+		AdjacencyMultiSizeIndexContainer.RebuildIndexBuffer(*AdjacencyIndexBufferData);
+	}
+}
+
+void FStaticLODModel::RebuildIndexBuffer()
+{
+	// The index buffer needs to be rebuilt on copy.
+	FMultiSizeIndexContainerData IndexBufferData;
+	MultiSizeIndexContainer.GetIndexBufferData(IndexBufferData);
+
+	FMultiSizeIndexContainerData AdjacencyIndexBufferData;
+	AdjacencyMultiSizeIndexContainer.GetIndexBufferData(AdjacencyIndexBufferData);
+
+	RebuildIndexBuffer(&IndexBufferData, &AdjacencyIndexBufferData);
+}
 /*-----------------------------------------------------------------------------
 FStaticMeshSourceData
 -----------------------------------------------------------------------------*/
@@ -1734,9 +1759,10 @@ void FSkeletalMeshSourceData::Init( const USkeletalMesh* SkeletalMesh, FStaticLO
 	InLODModel.LegacyRawPointIndices.Unlock();
 
 	/** The index buffer needs to be rebuilt on copy. */
-	FMultiSizeIndexContainerData IndexBufferData;
+	FMultiSizeIndexContainerData IndexBufferData, AdjacencyIndexBufferData;
 	InLODModel.MultiSizeIndexContainer.GetIndexBufferData( IndexBufferData );
-	LODModel->MultiSizeIndexContainer.RebuildIndexBuffer( IndexBufferData );
+	InLODModel.AdjacencyMultiSizeIndexContainer.GetIndexBufferData(AdjacencyIndexBufferData);
+	LODModel->RebuildIndexBuffer( &IndexBufferData, &AdjacencyIndexBufferData );
 
 	/** Vertex buffers also need to be rebuilt. Source data is always stored with full precision position data. */
 	LODModel->BuildVertexBuffers(SkeletalMesh->GetVertexBufferFlags());
diff --git a/Engine/Source/Runtime/Engine/Public/SkeletalMeshTypes.h b/Engine/Source/Runtime/Engine/Public/SkeletalMeshTypes.h
--- a/Engine/Source/Runtime/Engine/Public/SkeletalMeshTypes.h
+++ b/Engine/Source/Runtime/Engine/Public/SkeletalMeshTypes.h
@@ -1663,6 +1663,10 @@ public:
 	 * Get Resource Size
 	 */
 	SIZE_T GetResourceSize() const;
+
+	/** Rebuild index buffer for everything **/
+	ENGINE_API void RebuildIndexBuffer();
+	ENGINE_API void RebuildIndexBuffer(FMultiSizeIndexContainerData* IndexBufferData, FMultiSizeIndexContainerData* AdjacencyData);
 };
 
 /**
