Support Base LOD for simplygon ops

Merging

//depot/UE4-Orion/Engine/Source/...

to //depot/UE4/Engine/Source/...

diff --git a/Engine/Source/Developer/SimplygonMeshReduction/Private/SimplygonMeshReduction.cpp b/Engine/Source/Developer/SimplygonMeshReduction/Private/SimplygonMeshReduction.cpp
--- a/Engine/Source/Developer/SimplygonMeshReduction/Private/SimplygonMeshReduction.cpp
+++ b/Engine/Source/Developer/SimplygonMeshReduction/Private/SimplygonMeshReduction.cpp
@@ -266,6 +266,12 @@ public:
 
 		FStaticLODModel* SrcModel = &SkeletalMesh->PreModifyMesh();
 
+		const int32 BaseLOD = Settings.BaseLOD;
+		if (BaseLOD> 0 && SkeletalMeshResource->LODModels.IsValidIndex(BaseLOD))
+		{
+			SrcModel = &SkeletalMeshResource->LODModels[BaseLOD];
+		}
+
 		TComponentReregisterContext<USkinnedMeshComponent> ReregisterContext;
 		SkeletalMesh->ReleaseResources();
 		SkeletalMesh->ReleaseResourcesFence.Wait();
diff --git a/Engine/Source/Editor/Persona/Private/PersonaMeshDetails.cpp b/Engine/Source/Editor/Persona/Private/PersonaMeshDetails.cpp
--- a/Engine/Source/Editor/Persona/Private/PersonaMeshDetails.cpp
+++ b/Engine/Source/Editor/Persona/Private/PersonaMeshDetails.cpp
@@ -299,6 +299,25 @@ void FSkelMeshReductionSettingsLayout::GenerateChildContent(IDetailChildrenBuild
 
 	}
 
+	{
+		ChildrenBuilder.AddChildContent(LOCTEXT("BaseLOD", "Base LOD"))
+			.NameContent()
+			[
+				SNew(STextBlock)
+				.Font(IDetailLayoutBuilder::GetDetailFont())
+				.Text(LOCTEXT("BaseLODTitle", "Base LOD"))
+			]
+		.ValueContent()
+			[
+				SNew(SSpinBox<int32>)
+				.Font(IDetailLayoutBuilder::GetDetailFont())
+				.MinValue(0)
+				.Value(this, &FSkelMeshReductionSettingsLayout::GetBaseLOD)
+				.OnValueChanged(this, &FSkelMeshReductionSettingsLayout::OnBaseLODChanged)
+			];
+
+	}
+
 	{
 		ChildrenBuilder.AddChildContent(LOCTEXT("ApplyChanges", "Apply Changes"))
 			.ValueContent()
@@ -371,6 +390,11 @@ int32 FSkelMeshReductionSettingsLayout::GetMaxBonesPerVertex() const
 	return ReductionSettings.MaxBonesPerVertex;
 }
 
+int32 FSkelMeshReductionSettingsLayout::GetBaseLOD() const
+{
+	return ReductionSettings.BaseLOD;
+}
+
 void FSkelMeshReductionSettingsLayout::OnPercentTrianglesChanged(float NewValue)
 {
 	// Percentage -> fraction.
@@ -402,6 +426,11 @@ void FSkelMeshReductionSettingsLayout::OnMaxBonesPerVertexChanged(int32 NewValue
 	ReductionSettings.MaxBonesPerVertex = NewValue;
 }
 
+void FSkelMeshReductionSettingsLayout::OnBaseLODChanged(int32 NewLOD)
+{
+	ReductionSettings.BaseLOD = NewLOD;
+}
+
 void FSkelMeshReductionSettingsLayout::OnSilhouetteImportanceChanged(TSharedPtr<FString> NewValue, ESelectInfo::Type SelectInfo)
 {
 	ReductionSettings.SilhouetteImportance = (SkeletalMeshOptimizationImportance)ImportanceOptions.Find(NewValue);
diff --git a/Engine/Source/Editor/Persona/Private/PersonaMeshDetails.h b/Engine/Source/Editor/Persona/Private/PersonaMeshDetails.h
--- a/Engine/Source/Editor/Persona/Private/PersonaMeshDetails.h
+++ b/Engine/Source/Editor/Persona/Private/PersonaMeshDetails.h
@@ -63,6 +63,7 @@ private:
 	ECheckBoxState ShouldRecalculateNormals() const;
 	float GetHardAngleThreshold() const;
 	int32 GetMaxBonesPerVertex() const;
+	int32 GetBaseLOD() const;
 
 	void OnPercentTrianglesChanged(float NewValue);
 	void OnMaxDeviationChanged(float NewValue);
@@ -71,6 +72,7 @@ private:
 	void OnWeldingThresholdChanged(float NewValue);
 	void OnHardAngleThresholdChanged(float NewValue);
 	void OnMaxBonesPerVertexChanged(int32 NewValue);
+	void OnBaseLODChanged(int32 NewBasedLOD);
 
 	void OnSilhouetteImportanceChanged(TSharedPtr<FString> NewValue, ESelectInfo::Type SelectInfo);
 	void OnTextureImportanceChanged(TSharedPtr<FString> NewValue, ESelectInfo::Type SelectInfo);
diff --git a/Engine/Source/Runtime/Engine/Classes/Engine/SkeletalMesh.h b/Engine/Source/Runtime/Engine/Classes/Engine/SkeletalMesh.h
--- a/Engine/Source/Runtime/Engine/Classes/Engine/SkeletalMesh.h
+++ b/Engine/Source/Runtime/Engine/Classes/Engine/SkeletalMesh.h
@@ -166,58 +166,61 @@ struct FSkeletalMeshOptimizationSettings
 {
 	GENERATED_USTRUCT_BODY()
 
-		/** The method to use when optimizing the skeletal mesh LOD */
-		UPROPERTY()
-		TEnumAsByte<enum SkeletalMeshOptimizationType> ReductionMethod;
+	/** The method to use when optimizing the skeletal mesh LOD */
+	UPROPERTY()
+	TEnumAsByte<enum SkeletalMeshOptimizationType> ReductionMethod;
 
 	/** If ReductionMethod equals SMOT_NumOfTriangles this value is the ratio of triangles [0-1] to remove from the mesh */
 	UPROPERTY()
-		float NumOfTrianglesPercentage;
+	float NumOfTrianglesPercentage;
 
 	/**If ReductionMethod equals SMOT_MaxDeviation this value is the maximum deviation from the base mesh as a percentage of the bounding sphere. */
 	UPROPERTY()
-		float MaxDeviationPercentage;
+	float MaxDeviationPercentage;
 
 	/** The welding threshold distance. Vertices under this distance will be welded. */
 	UPROPERTY()
-		float WeldingThreshold;
+	float WeldingThreshold;
 
 	/** Whether Normal smoothing groups should be preserved. If false then NormalsThreshold is used **/
 	UPROPERTY()
-		bool bRecalcNormals;
+	bool bRecalcNormals;
 
 	/** If the angle between two triangles are above this value, the normals will not be
 	smooth over the edge between those two triangles. Set in degrees. This is only used when PreserveNormals is set to false*/
 	UPROPERTY()
-		float NormalsThreshold;
+	float NormalsThreshold;
 
 	/** How important the shape of the geometry is. */
 	UPROPERTY()
-		TEnumAsByte<enum SkeletalMeshOptimizationImportance> SilhouetteImportance;
+	TEnumAsByte<enum SkeletalMeshOptimizationImportance> SilhouetteImportance;
 
 	/** How important texture density is. */
 	UPROPERTY()
-		TEnumAsByte<enum SkeletalMeshOptimizationImportance> TextureImportance;
+	TEnumAsByte<enum SkeletalMeshOptimizationImportance> TextureImportance;
 
 	/** How important shading quality is. */
 	UPROPERTY()
-		TEnumAsByte<enum SkeletalMeshOptimizationImportance> ShadingImportance;
+	TEnumAsByte<enum SkeletalMeshOptimizationImportance> ShadingImportance;
 
 	/** How important skinning quality is. */
 	UPROPERTY()
-		TEnumAsByte<enum SkeletalMeshOptimizationImportance> SkinningImportance;
+	TEnumAsByte<enum SkeletalMeshOptimizationImportance> SkinningImportance;
 
 	/** The ratio of bones that will be removed from the mesh */
 	UPROPERTY()
-		float BoneReductionRatio;
+	float BoneReductionRatio;
 
 	/** Maximum number of bones that can be assigned to each vertex. */
 	UPROPERTY()
-		int32 MaxBonesPerVertex;
+	int32 MaxBonesPerVertex;
 
 	UPROPERTY(EditAnywhere, Category = ReductionSettings)
-		TArray<FBoneReference> BonesToRemove;
+	TArray<FBoneReference> BonesToRemove;
 
+	/** Maximum number of bones that can be assigned to each vertex. */
+	UPROPERTY()
+	int32 BaseLOD;
 
 	FSkeletalMeshOptimizationSettings()
 		: ReductionMethod(SMOT_MaxDeviation)
@@ -232,6 +235,7 @@ struct FSkeletalMeshOptimizationSettings
 		, SkinningImportance(SMOI_Normal)
 		, BoneReductionRatio(100.0f)
 		, MaxBonesPerVertex(4)
+		, BaseLOD(0)
 	{
 	}
 
@@ -271,7 +275,8 @@ struct FSkeletalMeshOptimizationSettings
 			&& SkinningImportance == Other.SkinningImportance
 			&& bRecalcNormals == Other.bRecalcNormals
 			&& BoneReductionRatio == Other.BoneReductionRatio
-			&& MaxBonesPerVertex == Other.MaxBonesPerVertex;
+			&& MaxBonesPerVertex == Other.MaxBonesPerVertex
+			&& BaseLOD == Other.BaseLOD;
 	}
 
 	/** Inequality. */
