Integrate - Material inputs are now greyed based on translucency lighting mode

diff --git a/Engine/Source/Runtime/Engine/Private/Materials/Material.cpp b/Engine/Source/Runtime/Engine/Private/Materials/Material.cpp
--- a/Engine/Source/Runtime/Engine/Private/Materials/Material.cpp
+++ b/Engine/Source/Runtime/Engine/Private/Materials/Material.cpp
@@ -4070,6 +4070,13 @@ bool UMaterial::IsPropertyActive(EMaterialProperty InProperty) const
 		}
 	}
 
+	const bool bIsTranslucentBlendMode = IsTranslucentBlendMode((EBlendMode)BlendMode);
+	const bool bIsNonDirectionalTranslucencyLightingMode = TranslucencyLightingMode == TLM_VolumetricNonDirectional || TranslucencyLightingMode == TLM_VolumetricPerVertexNonDirectional;
+	const bool bIsVolumetricTranslucencyLightingMode = TranslucencyLightingMode == TLM_VolumetricNonDirectional 
+		|| TranslucencyLightingMode == TLM_VolumetricDirectional 
+		|| TranslucencyLightingMode == TLM_VolumetricPerVertexNonDirectional 
+		|| TranslucencyLightingMode == TLM_VolumetricPerVertexDirectional;
+
 	bool Active = true;
 
 	switch (InProperty)
@@ -4079,10 +4086,10 @@ bool UMaterial::IsPropertyActive(EMaterialProperty InProperty) const
 		Active = false;
 		break;
 	case MP_Refraction:
-		Active = IsTranslucentBlendMode((EBlendMode)BlendMode) && BlendMode != BLEND_Modulate;
+		Active =bIsTranslucentBlendMode && BlendMode != BLEND_Modulate;
 		break;
 	case MP_Opacity:
-		Active = IsTranslucentBlendMode((EBlendMode)BlendMode) && BlendMode != BLEND_Modulate;
+		Active = bIsTranslucentBlendMode && BlendMode != BLEND_Modulate;
 		if (IsSubsurfaceShadingModel(ShadingModel))
 		{
 			Active = true;
@@ -4092,17 +4099,19 @@ bool UMaterial::IsPropertyActive(EMaterialProperty InProperty) const
 		Active = BlendMode == BLEND_Masked;
 		break;
 	case MP_BaseColor:
-	case MP_Specular:
-	case MP_Roughness:
 	case MP_AmbientOcclusion:
 		Active = ShadingModel != MSM_Unlit;
 		break;
+	case MP_Specular:
+	case MP_Roughness:
+		Active = ShadingModel != MSM_Unlit && (!bIsTranslucentBlendMode || !bIsVolumetricTranslucencyLightingMode);
+		break;
 	case MP_Metallic:
 		// Subsurface models store opacity in place of Metallic in the GBuffer
-		Active = ShadingModel != MSM_Unlit;
+		Active = ShadingModel != MSM_Unlit && (!bIsTranslucentBlendMode || !bIsVolumetricTranslucencyLightingMode);
 		break;
 	case MP_Normal:
-		Active = ShadingModel != MSM_Unlit || Refraction.IsConnected();
+		Active = ShadingModel != (MSM_Unlit && (!bIsTranslucentBlendMode || !bIsNonDirectionalTranslucencyLightingMode)) || Refraction.IsConnected();
 		break;
 	case MP_SubsurfaceColor:
 		Active = ShadingModel == MSM_Subsurface || ShadingModel == MSM_PreintegratedSkin || ShadingModel == MSM_TwoSidedFoliage;
@@ -4123,7 +4132,7 @@ bool UMaterial::IsPropertyActive(EMaterialProperty InProperty) const
 		Active = true;
 		break;
 	case MP_PixelDepthOffset:
-		Active = !IsTranslucentBlendMode((EBlendMode)BlendMode);
+		Active = !bIsTranslucentBlendMode;
 		break;
 	case MP_MaterialAttributes:
 	default:
