Integrate - The static shadow depth map used to shadow translucency from Stationary lights is now saved properly during a blueprint construction script

diff --git a/Engine/Source/Editor/UnrealEd/Private/Lightmass/Lightmass.cpp b/Engine/Source/Editor/UnrealEd/Private/Lightmass/Lightmass.cpp
--- a/Engine/Source/Editor/UnrealEd/Private/Lightmass/Lightmass.cpp
+++ b/Engine/Source/Editor/UnrealEd/Private/Lightmass/Lightmass.cpp
@@ -3599,11 +3599,11 @@ void FLightmassProcessor::ImportStaticShadowDepthMap(ULightComponent* Light)
 		BeginReleaseResource(&DepthMap);
 		DepthMap.Empty();
 
-		DepthMap.WorldToLight = ShadowMapData.WorldToLight;
-		DepthMap.ShadowMapSizeX = ShadowMapData.ShadowMapSizeX;
-		DepthMap.ShadowMapSizeY = ShadowMapData.ShadowMapSizeY;
+		DepthMap.Data.WorldToLight = ShadowMapData.WorldToLight;
+		DepthMap.Data.ShadowMapSizeX = ShadowMapData.ShadowMapSizeX;
+		DepthMap.Data.ShadowMapSizeY = ShadowMapData.ShadowMapSizeY;
 
-		ReadArray(Channel, DepthMap.DepthSamples);
+		ReadArray(Channel, DepthMap.Data.DepthSamples);
 		Swarm.CloseChannel(Channel);
 
 		DepthMap.InitializeAfterImport();
diff --git a/Engine/Source/Runtime/Engine/Classes/Components/LightComponent.h b/Engine/Source/Runtime/Engine/Classes/Components/LightComponent.h
--- a/Engine/Source/Runtime/Engine/Classes/Components/LightComponent.h
+++ b/Engine/Source/Runtime/Engine/Classes/Components/LightComponent.h
@@ -7,11 +7,7 @@
 #include "EngineDefines.h"
 #include "LightComponent.generated.h"
 
-/** 
- * A texture containing depth values of static objects that was computed during the lighting build.
- * Used by Stationary lights to shadow translucency.
- */
-class FStaticShadowDepthMap : public FTexture
+class FStaticShadowDepthMapData
 {
 public:
 	/** Transform from world space to the coordinate space that DepthSamples are stored in. */
@@ -22,12 +18,24 @@ public:
 	/** Shadowmap depth values */
 	TArray<FFloat16> DepthSamples;
 
-	FStaticShadowDepthMap() :
+	FStaticShadowDepthMapData() :
 		WorldToLight(FMatrix::Identity),
 		ShadowMapSizeX(0),
 		ShadowMapSizeY(0)
 	{}
 
+	friend FArchive& operator<<(FArchive& Ar, FStaticShadowDepthMapData& ShadowMap);
+};
+
+/** 
+ * A texture containing depth values of static objects that was computed during the lighting build.
+ * Used by Stationary lights to shadow translucency.
+ */
+class FStaticShadowDepthMap : public FTexture
+{
+public:
+	FStaticShadowDepthMapData Data;
+
 	virtual void InitRHI();
 
 	/** Frees the CPU backing of the shadowmap. */
diff --git a/Engine/Source/Runtime/Engine/Private/Components/LightComponent.cpp b/Engine/Source/Runtime/Engine/Private/Components/LightComponent.cpp
--- a/Engine/Source/Runtime/Engine/Private/Components/LightComponent.cpp
+++ b/Engine/Source/Runtime/Engine/Private/Components/LightComponent.cpp
@@ -17,21 +17,31 @@
 #include "Components/LightComponent.h"
 #include "Components/DirectionalLightComponent.h"
 
+FArchive& operator<<(FArchive& Ar, FStaticShadowDepthMapData& ShadowMapData)
+{
+	Ar << ShadowMapData.WorldToLight;
+	Ar << ShadowMapData.ShadowMapSizeX;
+	Ar << ShadowMapData.ShadowMapSizeY;
+	Ar << ShadowMapData.DepthSamples;
+
+	return Ar;
+}
+
 void FStaticShadowDepthMap::InitRHI()
 {
-	if (ShadowMapSizeX > 0 && ShadowMapSizeY > 0 && GMaxRHIFeatureLevel >= ERHIFeatureLevel::SM4)
+	if (Data.ShadowMapSizeX > 0 && Data.ShadowMapSizeY > 0 && GMaxRHIFeatureLevel >= ERHIFeatureLevel::SM4)
 	{
 		FRHIResourceCreateInfo CreateInfo;
-		FTexture2DRHIRef Texture2DRHI = RHICreateTexture2D(ShadowMapSizeX, ShadowMapSizeY, PF_R16F, 1, 1, 0, CreateInfo);
+		FTexture2DRHIRef Texture2DRHI = RHICreateTexture2D(Data.ShadowMapSizeX, Data.ShadowMapSizeY, PF_R16F, 1, 1, 0, CreateInfo);
 		TextureRHI = Texture2DRHI;
 
 		uint32 DestStride = 0;
 		uint8* TextureData = (uint8*)RHILockTexture2D(Texture2DRHI, 0, RLM_WriteOnly, DestStride, false);
-		uint32 RowSize = ShadowMapSizeX * GPixelFormats[PF_R16F].BlockBytes;
+		uint32 RowSize = Data.ShadowMapSizeX * GPixelFormats[PF_R16F].BlockBytes;
 
-		for (int32 Y = 0; Y < ShadowMapSizeY; Y++)
+		for (int32 Y = 0; Y < Data.ShadowMapSizeY; Y++)
 		{
-			FMemory::Memcpy(TextureData + DestStride * Y, ((uint8*)DepthSamples.GetData()) + RowSize * Y, RowSize);
+			FMemory::Memcpy(TextureData + DestStride * Y, ((uint8*)Data.DepthSamples.GetData()) + RowSize * Y, RowSize);
 		}
 
 		RHIUnlockTexture2D(Texture2DRHI, 0, false);
@@ -40,29 +50,26 @@ void FStaticShadowDepthMap::InitRHI()
 
 void FStaticShadowDepthMap::Empty()
 {
-	DEC_DWORD_STAT_BY(STAT_PrecomputedShadowDepthMapMemory, DepthSamples.GetAllocatedSize());
+	DEC_DWORD_STAT_BY(STAT_PrecomputedShadowDepthMapMemory, Data.DepthSamples.GetAllocatedSize());
 
-	ShadowMapSizeX = 0;
-	ShadowMapSizeY = 0;
-	DepthSamples.Empty();
+	Data.ShadowMapSizeX = 0;
+	Data.ShadowMapSizeY = 0;
+	Data.DepthSamples.Empty();
 }
 
 void FStaticShadowDepthMap::InitializeAfterImport()
 {
-	INC_DWORD_STAT_BY(STAT_PrecomputedShadowDepthMapMemory, DepthSamples.GetAllocatedSize());
+	INC_DWORD_STAT_BY(STAT_PrecomputedShadowDepthMapMemory, Data.DepthSamples.GetAllocatedSize());
 	BeginInitResource(this);
 }
 
 FArchive& operator<<(FArchive& Ar, FStaticShadowDepthMap& ShadowMap)
 {
-	Ar << ShadowMap.WorldToLight;
-	Ar << ShadowMap.ShadowMapSizeX;
-	Ar << ShadowMap.ShadowMapSizeY;
-	Ar << ShadowMap.DepthSamples;
+	Ar << ShadowMap.Data;
 
 	if (Ar.IsLoading())
 	{
-		INC_DWORD_STAT_BY(STAT_PrecomputedShadowDepthMapMemory, ShadowMap.DepthSamples.GetAllocatedSize());
+		INC_DWORD_STAT_BY(STAT_PrecomputedShadowDepthMapMemory, ShadowMap.Data.DepthSamples.GetAllocatedSize());
 	}
 
 	return Ar;
@@ -943,6 +950,7 @@ public:
 		, ShadowMapChannel(SourceComponent->ShadowMapChannel)
 		, PreviewShadowMapChannel(SourceComponent->PreviewShadowMapChannel)
 		, bPrecomputedLightingIsValid(SourceComponent->bPrecomputedLightingIsValid)
+		, StaticShadowDepthMapData(SourceComponent->StaticShadowDepthMap.Data)
 	{}
 
 	virtual void ApplyToComponent(UActorComponent* Component, const ECacheApplyPhase CacheApplyPhase) override
@@ -956,6 +964,7 @@ public:
 	int32 ShadowMapChannel;
 	int32 PreviewShadowMapChannel;
 	bool bPrecomputedLightingIsValid;
+	FStaticShadowDepthMapData StaticShadowDepthMapData;
 };
 
 FActorComponentInstanceData* ULightComponent::GetComponentInstanceData() const
@@ -977,6 +986,12 @@ void ULightComponent::ApplyComponentInstanceData(FPrecomputedLightInstanceData*
 	ShadowMapChannel = LightMapData->ShadowMapChannel;
 	PreviewShadowMapChannel = LightMapData->PreviewShadowMapChannel;
 	bPrecomputedLightingIsValid = LightMapData->bPrecomputedLightingIsValid;
+	StaticShadowDepthMap.Data = LightMapData->StaticShadowDepthMapData;
+
+	if (HasStaticShadowing() && !HasStaticLighting())
+	{
+		BeginUpdateResourceRHI(&StaticShadowDepthMap);
+	}
 
 	MarkRenderStateDirty();
 
diff --git a/Engine/Source/Runtime/Renderer/Private/TranslucentLighting.cpp b/Engine/Source/Runtime/Renderer/Private/TranslucentLighting.cpp
--- a/Engine/Source/Runtime/Renderer/Private/TranslucentLighting.cpp
+++ b/Engine/Source/Runtime/Renderer/Private/TranslucentLighting.cpp
@@ -951,7 +951,7 @@ public:
 		const FStaticShadowDepthMap* StaticShadowDepthMap = LightSceneInfo->Proxy->GetStaticShadowDepthMap();
 		const uint32 bStaticallyShadowedValue = LightSceneInfo->IsPrecomputedLightingValid() && StaticShadowDepthMap && StaticShadowDepthMap->TextureRHI ? 1 : 0;
 		FTextureRHIParamRef StaticShadowDepthMapTexture = bStaticallyShadowedValue ? StaticShadowDepthMap->TextureRHI : GWhiteTexture->TextureRHI;
-		const FMatrix WorldToStaticShadow = bStaticallyShadowedValue ? StaticShadowDepthMap->WorldToLight : FMatrix::Identity;
+		const FMatrix WorldToStaticShadow = bStaticallyShadowedValue ? StaticShadowDepthMap->Data.WorldToLight : FMatrix::Identity;
 
 		SetShaderValue(RHICmdList, ShaderRHI, bStaticallyShadowed, bStaticallyShadowedValue);
 
