Unshelved from pending changelist '2683851':

Fix point wind source setting incorrect weight

diff --git a/Engine/Source/Runtime/Engine/Classes/Components/WindDirectionalSourceComponent.h b/Engine/Source/Runtime/Engine/Classes/Components/WindDirectionalSourceComponent.h
--- a/Engine/Source/Runtime/Engine/Classes/Components/WindDirectionalSourceComponent.h
+++ b/Engine/Source/Runtime/Engine/Classes/Components/WindDirectionalSourceComponent.h
@@ -23,6 +23,12 @@ class ENGINE_API UWindDirectionalSourceComponent : public USceneComponent
 	UPROPERTY(interp, Category = WindDirectionalSourceComponent)
 	float MaxGustAmount;
 
+	UPROPERTY(interp, Category = WindDirectionalSourceComponent, meta = (editcondition = "bSimulatePhysics", ClampMin = "0.1", UIMin = "0.1"))
+	float Radius;
+
+	UPROPERTY(EditAnywhere, Category = WindDirectionalSourceComponent)
+	uint32 bPointWind : 1;
+
 public:
 	class FWindSourceSceneProxy* SceneProxy;
 
diff --git a/Engine/Source/Runtime/Engine/Private/WindDirectionalSource.cpp b/Engine/Source/Runtime/Engine/Private/WindDirectionalSource.cpp
--- a/Engine/Source/Runtime/Engine/Private/WindDirectionalSource.cpp
+++ b/Engine/Source/Runtime/Engine/Private/WindDirectionalSource.cpp
@@ -94,7 +94,7 @@ bool FWindSourceSceneProxy::GetWindParameters(const FVector& EvaluatePosition, F
 			WindData.MinGustAmt = MinGustAmt * RadialFalloff;
 			WindData.MaxGustAmt = MaxGustAmt * RadialFalloff;
 
-			Weight = Distance / Radius * Strength;
+			Weight = (Radius - Distance) / Radius * Strength;
 			return true;
 		}
 		Weight = 0.f;
@@ -168,13 +168,27 @@ void UWindDirectionalSourceComponent::DestroyRenderState_Concurrent()
 
 FWindSourceSceneProxy* UWindDirectionalSourceComponent::CreateSceneProxy() const
 {
-	return new FWindSourceSceneProxy(
-		ComponentToWorld.GetUnitAxis( EAxis::X ),
-		Strength,
-		Speed,
-		MinGustAmount,
-		MaxGustAmount
-		);
+	if (bPointWind)
+	{
+		return new FWindSourceSceneProxy(
+			ComponentToWorld.GetLocation(),
+			Strength,
+			Speed,
+			MinGustAmount,
+			MaxGustAmount,
+			Radius
+			);
+	}
+	else
+	{
+		return new FWindSourceSceneProxy(
+			ComponentToWorld.GetUnitAxis(EAxis::X),
+			Strength,
+			Speed,
+			MinGustAmount,
+			MaxGustAmount
+			);
+	}
 }
 
 
