Fix compile error.

diff --git a/Engine/Source/Runtime/Engine/Private/SkeletalMesh.cpp b/Engine/Source/Runtime/Engine/Private/SkeletalMesh.cpp
--- a/Engine/Source/Runtime/Engine/Private/SkeletalMesh.cpp
+++ b/Engine/Source/Runtime/Engine/Private/SkeletalMesh.cpp
@@ -1681,6 +1681,7 @@ void FStaticLODModel::RebuildIndexBuffer(FMultiSizeIndexContainerData* IndexBuff
 	}
 }
 
+#if WITH_EDITOR
 void FStaticLODModel::RebuildIndexBuffer()
 {
 	// The index buffer needs to be rebuilt on copy.
@@ -1692,6 +1693,7 @@ void FStaticLODModel::RebuildIndexBuffer()
 
 	RebuildIndexBuffer(&IndexBufferData, &AdjacencyIndexBufferData);
 }
+#endif
 /*-----------------------------------------------------------------------------
 FStaticMeshSourceData
 -----------------------------------------------------------------------------*/
diff --git a/Engine/Source/Runtime/Engine/Public/SkeletalMeshTypes.h b/Engine/Source/Runtime/Engine/Public/SkeletalMeshTypes.h
--- a/Engine/Source/Runtime/Engine/Public/SkeletalMeshTypes.h
+++ b/Engine/Source/Runtime/Engine/Public/SkeletalMeshTypes.h
@@ -1665,7 +1665,9 @@ public:
 	SIZE_T GetResourceSize() const;
 
 	/** Rebuild index buffer for everything **/
+#if WITH_EDITOR
 	ENGINE_API void RebuildIndexBuffer();
+#endif
 	ENGINE_API void RebuildIndexBuffer(FMultiSizeIndexContainerData* IndexBufferData, FMultiSizeIndexContainerData* AdjacencyData);
 };
 
