#merge: Null checks

diff --git a/Engine/Source/Runtime/Engine/Private/PackageMapClient.cpp b/Engine/Source/Runtime/Engine/Private/PackageMapClient.cpp
--- a/Engine/Source/Runtime/Engine/Private/PackageMapClient.cpp
+++ b/Engine/Source/Runtime/Engine/Private/PackageMapClient.cpp
@@ -2021,7 +2021,7 @@ bool FNetGUIDCache::ShouldIgnoreWhenMissing( const FNetworkGUID& NetGUID ) const
 
 	const FNetGuidCacheObject* OutermostCacheObject = CacheObject;
 
-	while ( OutermostCacheObject->OuterGUID.IsValid() )
+	while ( OutermostCacheObject != NULL && OutermostCacheObject->OuterGUID.IsValid() )
 	{
 		OutermostCacheObject = ObjectLookup.Find( OutermostCacheObject->OuterGUID );
 	}
