Merge build fixes back from //UE4/Main to //depot/UE4

diff --git a/Engine/Shaders/Common.usf b/Engine/Shaders/Common.usf
--- a/Engine/Shaders/Common.usf
+++ b/Engine/Shaders/Common.usf
@@ -943,20 +943,6 @@ float CalcSceneDepth( float2 ScreenUV )
 			float SceneW = ConvertFromDeviceZ(Texture2DSampleLevel(SceneDepthTexture, SceneDepthTextureSampler, ScreenUV, 0).r);		
 			return DepthbufferFetchES2(SceneW, View.InvDeviceZToWorldZTransform[2], View.InvDeviceZToWorldZTransform[3]);
 		#endif 
-{
-#if FEATURE_LEVEL > FEATURE_LEVEL_ES3_1
-	return ConvertFromDeviceZ(Texture2DSampleLevel(SceneDepthTexture, SceneDepthTextureSampler, ScreenUV, 0).r);
-#else
-	#if COMPILER_GLSL_ES2
-		#if IOS
-			// Only call FramebufferFetch when actually compiling for IOS ES2.
-			return FramebufferFetchES2().w;
-		#elif WEBGL
-			return Texture2DSampleLevel(SceneAlphaCopyTexture, SceneAlphaCopyTextureSampler, ScreenUV, 0).r;
-		#else 
-			float SceneW = Texture2DSampleLevel(SceneColorTexture, SceneColorTextureSampler, ScreenUV, 0).a;		
-			return DepthbufferFetchES2(SceneW, View.InvDeviceZToWorldZTransform[2], View.InvDeviceZToWorldZTransform[3]);
-		#endif 
 	#elif METAL_PROFILE
 			return FramebufferFetchES2().w;
 	#else
diff --git a/Engine/Source/Developer/Apple/MetalShaderFormat/Private/MetalShaderCompiler.cpp b/Engine/Source/Developer/Apple/MetalShaderFormat/Private/MetalShaderCompiler.cpp
--- a/Engine/Source/Developer/Apple/MetalShaderFormat/Private/MetalShaderCompiler.cpp
+++ b/Engine/Source/Developer/Apple/MetalShaderFormat/Private/MetalShaderCompiler.cpp
@@ -280,7 +280,7 @@ static void BuildMetalShaderOutput(
 	// Build the SRT for this shader.
 	{
 		// Build the generic SRT for this shader.
-		FShaderResourceTable GenericSRT;
+		FShaderCompilerResourceTable GenericSRT;
 		BuildResourceTableMapping(ShaderInput.Environment.ResourceTableMap, ShaderInput.Environment.ResourceTableLayoutHashes, UsedUniformBufferSlots, ShaderOutput.ParameterMap, GenericSRT);
 
 		// Copy over the bits indicating which resource tables are active.
diff --git a/Engine/Source/Developer/FriendsAndChat/Private/FriendsChatStyle.cpp b/Engine/Source/Developer/FriendsAndChat/Private/FriendsChatStyle.cpp
--- a/Engine/Source/Developer/FriendsAndChat/Private/FriendsChatStyle.cpp
+++ b/Engine/Source/Developer/FriendsAndChat/Private/FriendsChatStyle.cpp
@@ -100,29 +100,6 @@ FFriendsChatStyle& FFriendsChatStyle::SetPartyHyperlinkChatColor(const FLinearCo
 	return *this;
 }
 
-{
-	WhisperHyperlinkChatColor = InFontColor;
-	return *this;
-}
-
-FFriendsChatStyle& FFriendsChatStyle::SetGlobalHyperlinkChatColor(const FLinearColor& InFontColor)
-{
-	GlobalHyperlinkChatColor = InFontColor;
-	return *this;
-}
-
-FFriendsChatStyle& FFriendsChatStyle::SetGameHyperlinkChatColor(const FLinearColor& InFontColor)
-{
-	GameHyperlinkChatColor = InFontColor;
-	return *this;
-}
-
-FFriendsChatStyle& FFriendsChatStyle::SetPartyHyperlinkChatColor(const FLinearColor& InFontColor)
-{
-	PartyHyperlinkChatColor = InFontColor;
-	return *this;
-}
-
 FFriendsChatStyle& FFriendsChatStyle::SetChatInvalidBrush(const FSlateBrush& Brush)
 {
 	ChatInvalidBrush = Brush;
diff --git a/Engine/Source/Developer/FriendsAndChat/Public/FriendsChatStyle.h b/Engine/Source/Developer/FriendsAndChat/Public/FriendsChatStyle.h
--- a/Engine/Source/Developer/FriendsAndChat/Public/FriendsChatStyle.h
+++ b/Engine/Source/Developer/FriendsAndChat/Public/FriendsChatStyle.h
@@ -109,24 +109,6 @@ struct FRIENDSANDCHAT_API FFriendsChatStyle
 	FLinearColor PartyHyperlinkChatColor;
 	FFriendsChatStyle& SetPartyHyperlinkChatColor(const FLinearColor& InFontColor);
 
-	UPROPERTY(EditAnywhere, Category=Appearance)
-
-	UPROPERTY(EditAnywhere, Category=Appearance)
-	FLinearColor WhisperHyperlinkChatColor;
-	FFriendsChatStyle& SetWhisplerHyperlinkChatColor(const FLinearColor& InFontColor);
-
-	UPROPERTY(EditAnywhere, Category = Appearance)
-	FLinearColor GlobalHyperlinkChatColor;
-	FFriendsChatStyle& SetGlobalHyperlinkChatColor(const FLinearColor& InFontColor);
-
-	UPROPERTY(EditAnywhere, Category=Appearance)
-	FLinearColor GameHyperlinkChatColor;
-	FFriendsChatStyle& SetGameHyperlinkChatColor(const FLinearColor& InFontColor);
-
-	UPROPERTY(EditAnywhere, Category = Appearance)
-	FLinearColor PartyHyperlinkChatColor;
-	FFriendsChatStyle& SetPartyHyperlinkChatColor(const FLinearColor& InFontColor);
-
 	UPROPERTY(EditAnywhere, Category = Appearance)
 	FSlateBrush ChatInvalidBrush;
 	FFriendsChatStyle& SetChatInvalidBrush(const FSlateBrush& Brush);
diff --git a/Engine/Source/Editor/HierarchicalLODOutliner/Private/HLODTreeWidgetItem.cpp b/Engine/Source/Editor/HierarchicalLODOutliner/Private/HLODTreeWidgetItem.cpp
--- a/Engine/Source/Editor/HierarchicalLODOutliner/Private/HLODTreeWidgetItem.cpp
+++ b/Engine/Source/Editor/HierarchicalLODOutliner/Private/HLODTreeWidgetItem.cpp
@@ -173,7 +173,7 @@ END_SLATE_FUNCTION_BUILD_OPTIMIZATION
 		{
 			FDragValidationInfo ValidationInfo = FDragValidationInfo::Invalid();
 
-			HandleDrop(WeakTableViewBase, DragDropEvent, *TreeItem, ValidationInfo, false);
+			HLODOutliner::HandleDrop(WeakTableViewBase, DragDropEvent, *TreeItem, ValidationInfo, false);
 			UpdateOperationDecorator(DragDropEvent, ValidationInfo);
 		}
 	}
@@ -193,7 +193,7 @@ END_SLATE_FUNCTION_BUILD_OPTIMIZATION
 		if (TreeItem)
 		{
 			FDragValidationInfo ValidationInfo = FDragValidationInfo::Invalid();
-			return HandleDrop(WeakTableViewBase, DragDropEvent, *TreeItem, ValidationInfo, this, true);
+			return HLODOutliner::HandleDrop(WeakTableViewBase, DragDropEvent, *TreeItem, ValidationInfo, this, true);
 		}
 
 		return FReply::Unhandled();
diff --git a/Engine/Source/Programs/AutomationTool/Scripts/CookCommand.Automation.cs b/Engine/Source/Programs/AutomationTool/Scripts/CookCommand.Automation.cs
--- a/Engine/Source/Programs/AutomationTool/Scripts/CookCommand.Automation.cs
+++ b/Engine/Source/Programs/AutomationTool/Scripts/CookCommand.Automation.cs
@@ -282,7 +282,7 @@ public partial class Project : CommandUtils
     private static void DiffCookedContent( ProjectParams Params)
     {
         List<UnrealTargetPlatform> PlatformsToCook = Params.ClientTargetPlatforms;
-        string ProjectPath = Path.GetFullPath(Params.RawProjectPath);
+        string ProjectPath = Params.RawProjectPath.FullName;
 
         var CookedSandboxesPath = CombinePaths(GetDirectoryName(ProjectPath), "Saved", "Cooked");
 
diff --git a/Engine/Source/Runtime/AIModule/Private/AIController.cpp b/Engine/Source/Runtime/AIModule/Private/AIController.cpp
--- a/Engine/Source/Runtime/AIModule/Private/AIController.cpp
+++ b/Engine/Source/Runtime/AIModule/Private/AIController.cpp
@@ -783,7 +783,7 @@ FAIRequestID AAIController::RequestPathAndMove(const FAIMoveRequest& MoveRequest
 		else
 		{
 			UE_VLOG(this, LogAINavigation, Error, TEXT("Trying to find path to %s resulted in Error")
-				, MoveRequest.HasGoalActor() ? *GetNameSafe(MoveRequest.GetGoalActor()) : *MoveRequest.GetGoalLocation().ToString());
+				, MoveRequest.IsMoveToActorRequest() ? *GetNameSafe(MoveRequest.GetGoalActor()) : *MoveRequest.GetGoalLocation().ToString());
 			UE_VLOG_SEGMENT(this, LogAINavigation, Error, GetPawn() ? GetPawn()->GetActorLocation() : FAISystem::InvalidLocation
 				, MoveRequest.GetGoalLocation(), FColor::Red, TEXT("Failed move to %s"), *GetNameSafe(MoveRequest.GetGoalActor()));
 		}
diff --git a/Engine/Source/Runtime/AIModule/Private/EnvironmentQuery/Generators/EnvQueryGenerator_OnCircle.cpp b/Engine/Source/Runtime/AIModule/Private/EnvironmentQuery/Generators/EnvQueryGenerator_OnCircle.cpp
--- a/Engine/Source/Runtime/AIModule/Private/EnvironmentQuery/Generators/EnvQueryGenerator_OnCircle.cpp
+++ b/Engine/Source/Runtime/AIModule/Private/EnvironmentQuery/Generators/EnvQueryGenerator_OnCircle.cpp
@@ -216,26 +216,6 @@ void UEnvQueryGenerator_OnCircle::GenerateItems(FEnvQueryInstance& QueryInstance
 		uint8* RawData = ContextData.RawData.GetData();
 		const FVector CircleCenterOffset = FVector(0.f, 0.f, CircleCenterZOffset.GetValue());
 
-		TArray<AActor*> ContextActorsToIgnoreWhenGeneratingCircle;
-		if (bIgnoreAnyContextActorsWhenGeneratingCircle)
-		{
-			UEnvQueryItemType_ActorBase* DefEnvQueryActorBaseItemTypeOb = Cast<UEnvQueryItemType_ActorBase>(ContextData.ValueType->GetDefaultObject());
-			if (DefEnvQueryActorBaseItemTypeOb != nullptr)
-			{
-				uint8* ActorRawData = ContextData.RawData.GetData();
-
-				for (int32 ValueIndex = 0; ValueIndex < ContextData.NumValues; ValueIndex++)
-				{
-					AActor* Actor = DefEnvQueryActorBaseItemTypeOb->GetActor(ActorRawData);
-					if (Actor)
-					{
-						ContextActorsToIgnoreWhenGeneratingCircle.Add(Actor);
-					}
-					ActorRawData += DefTypeValueSize;
-				}
-			}
-		}
-
 		for (int32 ValueIndex = 0; ValueIndex < ContextData.NumValues; ValueIndex++)
 		{
 			const FVector ContextItemLocation = DefTypeOb->GetItemLocation(RawData) + CircleCenterOffset;
diff --git a/Engine/Source/Runtime/AutomationMessages/Classes/AutomationWorkerMessages.h b/Engine/Source/Runtime/AutomationMessages/Classes/AutomationWorkerMessages.h
--- a/Engine/Source/Runtime/AutomationMessages/Classes/AutomationWorkerMessages.h
+++ b/Engine/Source/Runtime/AutomationMessages/Classes/AutomationWorkerMessages.h
@@ -230,6 +230,10 @@ struct FAutomationWorkerRunTests
 	UPROPERTY(EditAnywhere, Category="Message")
 	bool bUseFullSizeScreenShots;
 
+	/** If true, send results to analytics when complete */
+	UPROPERTY()
+	bool bSendAnalytics;
+
 	/** Default constructor. */
 	FAutomationWorkerRunTests( ) { }
 
diff --git a/Engine/Source/Runtime/AutomationWorker/Private/AutomationAnalytics.cpp b/Engine/Source/Runtime/AutomationWorker/Private/AutomationAnalytics.cpp
--- a/Engine/Source/Runtime/AutomationWorker/Private/AutomationAnalytics.cpp
+++ b/Engine/Source/Runtime/AutomationWorker/Private/AutomationAnalytics.cpp
@@ -296,7 +296,7 @@ void FAutomationAnalytics::FireEvent_AutomationTestResults(const FAutomationWork
 		ParamArray.Add(FAnalyticsEventAttribute(GetAutomationParamName(EAutomationAnalyticParam::ErrorCount), TestResults->Errors.Num()));
 		ParamArray.Add(FAnalyticsEventAttribute(GetAutomationParamName(EAutomationAnalyticParam::WarningCount), TestResults->Warnings.Num()));
 
-		const int32 ChangeList = GEngineVersion.GetChangelist();
+		const int32 ChangeList = FEngineVersion::Current().GetChangelist();
 		FString ChangeListString = FString::FromInt(ChangeList);
 		ParamArray.Add(FAnalyticsEventAttribute(GetAutomationParamName(EAutomationAnalyticParam::CL), ChangeListString));
 
diff --git a/Engine/Source/Runtime/Core/Private/GenericPlatform/GenericPlatformMisc.cpp b/Engine/Source/Runtime/Core/Private/GenericPlatform/GenericPlatformMisc.cpp
--- a/Engine/Source/Runtime/Core/Private/GenericPlatform/GenericPlatformMisc.cpp
+++ b/Engine/Source/Runtime/Core/Private/GenericPlatform/GenericPlatformMisc.cpp
@@ -710,7 +710,7 @@ const TCHAR* FGenericPlatformMisc::GamePersistentDownloadDir()
 	return *GamePersistentDownloadDir;
 }
 
-uint32 FGenericPlatformMisc::GetStandardPrintableKeyMap(uint16* KeyCodes, FString* KeyNames, uint32 MaxMappings, bool bMapUppercaseKeys, bool bMapLowercaseKeys)
+uint32 FGenericPlatformMisc::GetStandardPrintableKeyMap(uint32* KeyCodes, FString* KeyNames, uint32 MaxMappings, bool bMapUppercaseKeys, bool bMapLowercaseKeys)
 {
 	uint32 NumMappings = 0;
 
diff --git a/Engine/Source/Runtime/Core/Private/Math/UnrealMath.cpp b/Engine/Source/Runtime/Core/Private/Math/UnrealMath.cpp
--- a/Engine/Source/Runtime/Core/Private/Math/UnrealMath.cpp
+++ b/Engine/Source/Runtime/Core/Private/Math/UnrealMath.cpp
@@ -787,47 +787,6 @@ FQuat FQuat::FindBetweenVectors(const FVector& A, const FVector& B)
 	return FindBetween_Helper(A, B, NormAB);
 }
 
-
-//
-// Based on:
-// http://lolengine.net/blog/2014/02/24/quaternion-from-two-vectors-final
-// http://www.euclideanspace.com/maths/algebra/vectors/angleBetween/index.htm
-//
-FORCEINLINE_DEBUGGABLE FQuat FindBetween_Helper(const FVector& A, const FVector& B, float NormAB)
-{
-	float W = NormAB + FVector::DotProduct(A, B);
-	FQuat Result;
-
-	if (W >= 1e-6f * NormAB)
-	{
-		//Axis = FVector::CrossProduct(A, B);
-		Result = FQuat(A.Y * B.Z - A.Z * B.Y,
-			A.Z * B.X - A.X * B.Z,
-			A.X * B.Y - A.Y * B.X,
-			W);
-	}
-	else
-	{
-		// A and B point in opposite directions
-		W = 0.f;
-		Result = FMath::Abs(A.X) > FMath::Abs(A.Y)
-			? FQuat(-A.Z, 0.f, A.X, W)
-			: FQuat(0.f, -A.Z, A.Y, W);
-	}
-
-	Result.Normalize();
-	return Result;
-}
-
-
-
-FQuat FQuat::FindBetweenNormals(const FVector& A, const FVector& B)
-{
-	const float NormAB = 1.f;
-	return FindBetween_Helper(A, B, NormAB);
-}
-
-
 FQuat FQuat::Log() const
 {
 	FQuat Result;
diff --git a/Engine/Source/Runtime/CoreUObject/Private/UObject/SavePackage.cpp b/Engine/Source/Runtime/CoreUObject/Private/UObject/SavePackage.cpp
--- a/Engine/Source/Runtime/CoreUObject/Private/UObject/SavePackage.cpp
+++ b/Engine/Source/Runtime/CoreUObject/Private/UObject/SavePackage.cpp
@@ -3102,12 +3102,58 @@ TMap<UObject*, UObject*> UnmarkExportTagFromDuplicates()
 
 	return RedirectDuplicatesToOriginals;
 }
-#endif // WITH_EDITOR
-
 
-#if WITH_EDITOR
 COREUOBJECT_API extern bool GOutputCookingWarnings;
 
+/**
+ * Find imports to converted BPGC, that should replaced by a generated native class. 
+ * For each converted class a stub is generated. The import object is replaced with the stub.
+ */
+struct FImportReplacementsHelper
+{
+	TSet<UObject*> UniqueTagImpObjects;
+
+	FImportReplacementsHelper(const TArray<UObject*>& InTagImpObjects)
+	{
+		if (FReplaceCookedBPGC::Get().IsEnabled())
+		{
+			UniqueTagImpObjects = TSet<UObject*>(InTagImpObjects);
+		}
+	}
+
+	UObject* Handle(UObject* Obj, TArray<FObjectImport>& ImportMap)
+	{
+		if (FReplaceCookedBPGC::Get().IsEnabled())
+		{
+			UObject* ReplacedObjectToImport = FReplaceCookedBPGC::Get().FindReplacementStub(Obj);
+			Obj = ReplacedObjectToImport ? ReplacedObjectToImport : Obj;
+			if (ReplacedObjectToImport)
+			{
+				{
+					ensure(ReplacedObjectToImport->GetOuter() == ReplacedObjectToImport->GetOutermost());
+					bool bAlreadyAdded = false;
+					UniqueTagImpObjects.Add(ReplacedObjectToImport->GetOuter(), &bAlreadyAdded);
+					if (!bAlreadyAdded)
+					{
+						new(ImportMap)FObjectImport(ReplacedObjectToImport->GetOuter());
+					}
+				}
+
+									{
+										bool bAlreadyAdded = false;
+										UniqueTagImpObjects.Add(ReplacedObjectToImport->GetClass(), &bAlreadyAdded);
+										if (!bAlreadyAdded)
+										{
+											new(ImportMap)FObjectImport(ReplacedObjectToImport->GetClass());
+										}
+									}
+			}
+		}
+		return Obj;
+	}
+
+};
+
 class FDiffSerializeArchive : public FBufferArchive
 {
 private:
diff --git a/Engine/Source/Runtime/Engine/Classes/AI/Navigation/NavigationSystem.h b/Engine/Source/Runtime/Engine/Classes/AI/Navigation/NavigationSystem.h
--- a/Engine/Source/Runtime/Engine/Classes/AI/Navigation/NavigationSystem.h
+++ b/Engine/Source/Runtime/Engine/Classes/AI/Navigation/NavigationSystem.h
@@ -596,7 +596,7 @@ public:
 	static void UpdateActorInNavOctree(AActor& Actor);
 	static void UpdateComponentInNavOctree(UActorComponent& Comp);
 	/** update all navoctree entries for actor and its components */
-	static void UpdateNavOctreeAll(AActor& Actor, bool bUpdateAttachedActors = true);
+	static void UpdateActorAndComponentsInNavOctree(AActor& Actor, bool bUpdateAttachedActors = true);
 	/** update all navoctree entries for actor and its non scene components after root movement */
 	static void UpdateNavOctreeAfterMove(USceneComponent* Comp);
 
@@ -988,7 +988,7 @@ public:
 	static void UpdateNavOctree(AActor* Actor);
 	DEPRECATED(4.11, "UpdateNavOctree is deprecated. Use UpdateComponentInNavOctree")
 	static void UpdateNavOctree(UActorComponent* Comp);
-	DEPRECATED(4.11, "UpdateNavOctreeAll is deprecated. Use UpdateComponentInNavOctree")
+	DEPRECATED(4.11, "UpdateNavOctreeAll is deprecated. Use UpdateActorAndComponentsInNavOctree")
 	static void UpdateNavOctreeAll(AActor* Actor);
 };
 
diff --git a/Engine/Source/Runtime/Engine/Classes/Animation/AnimInstance.h b/Engine/Source/Runtime/Engine/Classes/Animation/AnimInstance.h
--- a/Engine/Source/Runtime/Engine/Classes/Animation/AnimInstance.h
+++ b/Engine/Source/Runtime/Engine/Classes/Animation/AnimInstance.h
@@ -808,14 +808,6 @@ public:
 		return (NodeType*)GetNodeFromIndexUntyped(NodeIdx, NodeType::StaticStruct());
 	}
 
-	/** 
-	 * Get the index of the specified instance asset player. Useful to pass to GetInstanceAssetPlayerLength (etc.).
-	 * Passing NAME_None to InstanceName will return the first (assumed only) player instance index found.
-	 */
-	int32 GetInstanceAssetPlayerIndex(FName MachineName, FName StateName, FName InstanceName = NAME_None);
-
-public:
-
 	/** Gets the runtime instance of the specified state machine by Name */
 	FAnimNode_StateMachine* GetStateMachineInstanceFromName(FName MachineName);
 
@@ -827,23 +819,15 @@ protected:
 	/** Gets the runtime instance of the specified state machine */
 	FAnimNode_StateMachine* GetStateMachineInstance(int32 MachineIndex);
 
-	/** Gets the runtime instance desc of the state machine specified by name */
-	const FBakedAnimationStateMachine* GetStateMachineInstanceDesc(FName MachineName);
-
 	/** 
 	 * Get the index of the specified instance asset player. Useful to pass to GetInstanceAssetPlayerLength (etc.).
 	 * Passing NAME_None to InstanceName will return the first (assumed only) player instance index found.
 	 */
 	int32 GetInstanceAssetPlayerIndex(FName MachineName, FName StateName, FName InstanceName = NAME_None);
 
-protected:
-
 	/** Gets the runtime instance desc of the state machine specified by name */
 	const FBakedAnimationStateMachine* GetStateMachineInstanceDesc(FName MachineName);
 
-	/** Get the machine description for the specified instance. Does not rely on PRIVATE_MachineDescription being initialized */
-	const FBakedAnimationStateMachine* GetMachineDescription(UAnimBlueprintGeneratedClass* AnimBlueprintClass, FAnimNode_StateMachine* MachineInstance);
-
 	/** Gets the most relevant asset player in a specified state */
 	FAnimNode_AssetPlayerBase* GetRelevantAssetPlayerFromState(int32 MachineIndex, int32 StateIndex);
 
diff --git a/Engine/Source/Runtime/Engine/Classes/Engine/DemoNetDriver.h b/Engine/Source/Runtime/Engine/Classes/Engine/DemoNetDriver.h
--- a/Engine/Source/Runtime/Engine/Classes/Engine/DemoNetDriver.h
+++ b/Engine/Source/Runtime/Engine/Classes/Engine/DemoNetDriver.h
@@ -3,6 +3,7 @@
 #pragma once
 
 #include "NetDriver.h"
+#include "NetworkReplayStreaming.h"
 #include "DemoNetDriver.generated.h"
 
 DECLARE_LOG_CATEGORY_EXTERN( LogDemo, Log, All );
@@ -157,6 +158,8 @@ public:
 	void SpawnDemoRecSpectator( UNetConnection* Connection, const FURL& ListenURL );
 	void ResetDemoState();
 	void JumpToEndOfLiveReplay();
+	void AddEvent(const FString& Group, const FString& Meta, const TArray<uint8>& Data);
+	void EnumerateEvents(const FString& Group, FEnumerateEventsCompleteDelegate& EnumerationCompleteDelegate);
 	void RequestEventData(const FString& EventID, FOnRequestEventDataComplete& RequestEventDataCompleteDelegate);
 	virtual bool IsFastForwarding() { return bIsFastForwarding; }
 
diff --git a/Engine/Source/Runtime/Engine/Engine.Build.cs b/Engine/Source/Runtime/Engine/Engine.Build.cs
--- a/Engine/Source/Runtime/Engine/Engine.Build.cs
+++ b/Engine/Source/Runtime/Engine/Engine.Build.cs
@@ -31,8 +31,8 @@ public class Engine : ModuleRules
 				"Advertising",
 				"NetworkReplayStreaming",
 				"AutomationWorker",
-                		"Analytics"
-				"MovieSceneCapture"
+                "Analytics",
+				"MovieSceneCapture",
 			}
 		);
 
diff --git a/Engine/Source/Runtime/Engine/Private/AI/Navigation/AvoidanceManager.cpp b/Engine/Source/Runtime/Engine/Private/AI/Navigation/AvoidanceManager.cpp
--- a/Engine/Source/Runtime/Engine/Private/AI/Navigation/AvoidanceManager.cpp
+++ b/Engine/Source/Runtime/Engine/Private/AI/Navigation/AvoidanceManager.cpp
@@ -495,7 +495,7 @@ FVector UAvoidanceManager::GetAvoidanceVelocity_Internal(const FNavAvoidanceData
 			BestScorePotential = (VelSpacePoint|ReturnVelocity) * (VelSpacePoint|VelSpacePoint);
 			if (BestScorePotential > BestScore)
 			{
-				const bool bAvoidsNavEdges = NavEdges.Num() > 0 ? AvoidsNavEdges(inAvoidanceData.Center, VelSpacePoint, NavEdges, inAvoidanceData.Height) : true;
+				const bool bAvoidsNavEdges = NavEdges.Num() > 0 ? AvoidsNavEdges(inAvoidanceData.Center, VelSpacePoint, NavEdges, inAvoidanceData.HalfHeight) : true;
 				if (bAvoidsNavEdges)
 				{
 					FVector CandidateVelocity = AvoidCones(AllCones, FVector::ZeroVector, VelSpacePoint, AllCones.Num());
diff --git a/Engine/Source/Runtime/Engine/Private/AI/Navigation/NavigationSystem.cpp b/Engine/Source/Runtime/Engine/Private/AI/Navigation/NavigationSystem.cpp
--- a/Engine/Source/Runtime/Engine/Private/AI/Navigation/NavigationSystem.cpp
+++ b/Engine/Source/Runtime/Engine/Private/AI/Navigation/NavigationSystem.cpp
@@ -2517,13 +2517,7 @@ void UNavigationSystem::UpdateComponentInNavOctree(UActorComponent& Comp)
 	}
 }
 
->>>> ORIGINAL //depot/UE4-Orion/Engine/Source/Runtime/Engine/Private/AI/Navigation/NavigationSystem.cpp#8
-void UNavigationSystem::UpdateNavOctreeAll(AActor* Actor)
-==== THEIRS //depot/UE4-Orion/Engine/Source/Runtime/Engine/Private/AI/Navigation/NavigationSystem.cpp#14
-void UNavigationSystem::UpdateActorAndComponentsInNavOctree(AActor& Actor)
-==== YOURS //Marc.Audy_Z2487/Engine/Source/Runtime/Engine/Private/AI/Navigation/NavigationSystem.cpp
-void UNavigationSystem::UpdateNavOctreeAll(AActor* Actor, bool bUpdateAttachedActors)
-<<<<
+void UNavigationSystem::UpdateActorAndComponentsInNavOctree(AActor& Actor, bool bUpdateAttachedActors)
 {
 	UpdateActorInNavOctree(Actor);
 		
@@ -2536,37 +2530,9 @@ void UNavigationSystem::UpdateNavOctreeAll(AActor* Actor, bool bUpdateAttachedAc
 		UpdateComponentInNavOctree(*Components[ComponentIndex]);
 	}
 
->>>> ORIGINAL //depot/UE4-Orion/Engine/Source/Runtime/Engine/Private/AI/Navigation/NavigationSystem.cpp#8
-		if (Actor->GetRootComponent())
-==== THEIRS //depot/UE4-Orion/Engine/Source/Runtime/Engine/Private/AI/Navigation/NavigationSystem.cpp#14
-	if (Actor.GetRootComponent())
+	if (bUpdateAttachedActors)
 	{
-		for (int32 RootChildIndex = 0; RootChildIndex < Actor.GetRootComponent()->AttachChildren.Num(); RootChildIndex++)
-==== YOURS //Marc.Audy_Z2487/Engine/Source/Runtime/Engine/Private/AI/Navigation/NavigationSystem.cpp
-		if (bUpdateAttachedActors)
-<<<<
-		{
->>>> ORIGINAL //depot/UE4-Orion/Engine/Source/Runtime/Engine/Private/AI/Navigation/NavigationSystem.cpp#8
-			for (int32 RootChildIndex = 0; RootChildIndex < Actor->GetRootComponent()->AttachChildren.Num(); RootChildIndex++)
-			{
-				if (Actor->GetRootComponent()->AttachChildren[RootChildIndex] && Actor->GetRootComponent()->AttachChildren[RootChildIndex]->GetOuter() != Actor)
-				{
-					UpdateNavOctreeAll(Cast<AActor>(Actor->GetRootComponent()->AttachChildren[RootChildIndex]->GetOuter()));
-				}
-			}
-==== THEIRS //depot/UE4-Orion/Engine/Source/Runtime/Engine/Private/AI/Navigation/NavigationSystem.cpp#14
-			if (Actor.GetRootComponent()->AttachChildren[RootChildIndex] && Actor.GetRootComponent()->AttachChildren[RootChildIndex]->GetOuter() != &Actor)
-			{
-				AActor* AttachedActor = Cast<AActor>(Actor.GetRootComponent()->AttachChildren[RootChildIndex]->GetOuter());
-				if (AttachedActor)
-				{
-					UpdateActorAndComponentsInNavOctree(*AttachedActor);
-				}
-			}
-==== YOURS //Marc.Audy_Z2487/Engine/Source/Runtime/Engine/Private/AI/Navigation/NavigationSystem.cpp
-			UpdateAttachedActorsInNavOctree(*Actor);
-<<<<
-		}
+		UpdateAttachedActorsInNavOctree(Actor);
 	}
 }
 
@@ -2608,31 +2574,15 @@ void UNavigationSystem::UpdateAttachedActorsInNavOctree(AActor& RootActor)
 		
 		for (int32 AttachmentIndex = 0; AttachmentIndex < TempAttachedActors.Num(); ++AttachmentIndex)
 		{
->>>> ORIGINAL //depot/UE4-Orion/Engine/Source/Runtime/Engine/Private/AI/Navigation/NavigationSystem.cpp#8
-			if (Comp->AttachChildren[RootChildIndex] && Comp->AttachChildren[RootChildIndex]->GetOuter() != OwnerActor)
-			{
-				UpdateNavOctreeAll(Cast<AActor>(Comp->AttachChildren[RootChildIndex]->GetOuter()));
-			}
-==== THEIRS //depot/UE4-Orion/Engine/Source/Runtime/Engine/Private/AI/Navigation/NavigationSystem.cpp#14
-			if (Comp->AttachChildren[RootChildIndex] && Comp->AttachChildren[RootChildIndex]->GetOuter() != OwnerActor)
-			{
-				AActor* AttachedActor = Cast<AActor>(Comp->AttachChildren[RootChildIndex]->GetOuter());
-				if (AttachedActor)
-				{
-					UpdateActorAndComponentsInNavOctree(*AttachedActor);
-				}
-			}
-==== YOURS //Marc.Audy_Z2487/Engine/Source/Runtime/Engine/Private/AI/Navigation/NavigationSystem.cpp
 			// and store the ones we don't know about yet
 			UniqueAttachedActors.AddUnique(TempAttachedActors[AttachmentIndex]);
-<<<<
 		}
 	}
 	
 	// skipping the first item since that's the root, and we just care about the attached actors
 	for (int32 ActorIndex = 1; ActorIndex < UniqueAttachedActors.Num(); ++ActorIndex)
 	{
-		UpdateNavOctreeAll(UniqueAttachedActors[ActorIndex], /*bUpdateAttachedActors = */false);
+		UpdateActorAndComponentsInNavOctree(*UniqueAttachedActors[ActorIndex], /*bUpdateAttachedActors = */false);
 	}
 }
 
@@ -4058,9 +4008,6 @@ bool UNavigationSystem::GetRandomPointInRadius(const FVector& Origin, float Radi
 	return GetRandomReachablePointInRadius(Origin, Radius, ResultLocation, NavData, QueryFilter);
 }
 
->>>> ORIGINAL //depot/UE4-Orion/Engine/Source/Runtime/Engine/Private/AI/Navigation/NavigationSystem.cpp#8
-}
-==== THEIRS //depot/UE4-Orion/Engine/Source/Runtime/Engine/Private/AI/Navigation/NavigationSystem.cpp#14
 void UNavigationSystem::UpdateNavOctree(AActor* Actor)
 {
 	if (Actor)
@@ -4084,5 +4031,5 @@ void UNavigationSystem::UpdateNavOctreeAll(AActor* Actor)
 		UpdateActorAndComponentsInNavOctree(*Actor);
 	}
 }
-==== YOURS //Marc.Audy_Z2487/Engine/Source/Runtime/Engine/Private/AI/Navigation/NavigationSystem.cpp
+
 #undef LOCTEXT_NAMESPACE
diff --git a/Engine/Source/Runtime/Engine/Private/ActorEditor.cpp b/Engine/Source/Runtime/Engine/Private/ActorEditor.cpp
--- a/Engine/Source/Runtime/Engine/Private/ActorEditor.cpp
+++ b/Engine/Source/Runtime/Engine/Private/ActorEditor.cpp
@@ -174,8 +174,8 @@ void AActor::PostEditMove(bool bFinished)
 			UNavigationSystem::UpdateNavOctreeBounds(ParentedActors[Idx]);
 		}
 
-		// not doing manual update of all attached actors since UpdateNavOctreeAll should take care of it
-		UNavigationSystem::UpdateNavOctreeAll(this);
+		// not doing manual update of all attached actors since UpdateActorAndComponentsInNavOctree should take care of it
+		UNavigationSystem::UpdateActorAndComponentsInNavOctree(*this);
 	}
 }
 
diff --git a/Engine/Source/Runtime/Engine/Private/Animation/AnimMontage.cpp b/Engine/Source/Runtime/Engine/Private/Animation/AnimMontage.cpp
--- a/Engine/Source/Runtime/Engine/Private/Animation/AnimMontage.cpp
+++ b/Engine/Source/Runtime/Engine/Private/Animation/AnimMontage.cpp
@@ -11,6 +11,8 @@
 #include "Animation/AnimNotifies/AnimNotifyState.h"
 #include "Animation/AnimNotifies/AnimNotify.h"
 #include "Animation/AnimSequenceBase.h"
+#include "Animation/AnimSingleNodeInstance.h"
+#include "Animation/AnimStats.h"
 
 ///////////////////////////////////////////////////////////////////////////
 //
@@ -1660,7 +1662,7 @@ void FAnimMontageInstance::SetMatineeAnimPositionInner(FName SlotName, USkeletal
 	UAnimInstance* AnimInst = SkeletalMeshComponent->GetAnimInstance();
 	if(AnimInst)
 	{
-		UAnimSingleNodeInstance * SingleNodeInst = SkeletalMeshComponent->GetSingleNodeInstance();
+		UAnimSingleNodeInstance* SingleNodeInst = SkeletalMeshComponent->GetSingleNodeInstance();
 		if(SingleNodeInst)
 		{
 			if(SingleNodeInst->CurrentAsset != InAnimSequence)
@@ -1770,7 +1772,7 @@ void FAnimMontageInstance::PreviewMatineeSetAnimPositionInner(FName SlotName, US
 
 			float OldMontagePosition = AnimInst->Montage_GetPosition(CurrentlyPlayingMontage.Get());
 			AnimInst->Montage_SetPosition(CurrentlyPlayingMontage.Get(), InPosition);
-			AnimInst->UpdateAnimation(DeltaTime);
+			AnimInst->UpdateAnimation(DeltaTime, false);
 
 			// since we don't advance montage in the tick, we manually have to handle notifies
 			AnimMontageInst->HandleEvents(OldMontagePosition, InPosition, NULL);
diff --git a/Engine/Source/Runtime/Engine/Private/Components/SceneComponent.cpp b/Engine/Source/Runtime/Engine/Private/Components/SceneComponent.cpp
--- a/Engine/Source/Runtime/Engine/Private/Components/SceneComponent.cpp
+++ b/Engine/Source/Runtime/Engine/Private/Components/SceneComponent.cpp
@@ -288,7 +288,7 @@ void USceneComponent::PostEditChangeProperty(FPropertyChangedEvent& PropertyChan
 
 	if (PropertyName == LocationName || PropertyName == RotationName || PropertyName == ScaleName)
 	{
-		UNavigationSystem::UpdateNavOctree(this);
+		UNavigationSystem::UpdateComponentInNavOctree(*this);
 	}
 }
 
diff --git a/Engine/Source/Runtime/Engine/Private/NetConnection.cpp b/Engine/Source/Runtime/Engine/Private/NetConnection.cpp
--- a/Engine/Source/Runtime/Engine/Private/NetConnection.cpp
+++ b/Engine/Source/Runtime/Engine/Private/NetConnection.cpp
@@ -15,6 +15,8 @@
 #include "Runtime/PacketHandlers/PacketHandler/Public/PacketHandler.h"
 
 #include "PerfCountersHelpers.h"
+#if WITH_EDITOR
+#include "UnrealEd.h"
 #endif
 
 #if !UE_BUILD_SHIPPING
diff --git a/Engine/Source/Runtime/Engine/Private/Particles/ParticleSystemRender.cpp b/Engine/Source/Runtime/Engine/Private/Particles/ParticleSystemRender.cpp
--- a/Engine/Source/Runtime/Engine/Private/Particles/ParticleSystemRender.cpp
+++ b/Engine/Source/Runtime/Engine/Private/Particles/ParticleSystemRender.cpp
@@ -6700,7 +6700,7 @@ void FParticleSystemSceneProxy::GetObjectPositionAndScale(const FSceneView& View
 /**
 * @return Relevance for rendering the particle system primitive component in the given View
 */
-FPrimitiveViewRelevance FParticleSystemSceneProxy::GetViewRelevance(const FSceneView* View)
+FPrimitiveViewRelevance FParticleSystemSceneProxy::GetViewRelevance(const FSceneView* View) const
 {
 	FPrimitiveViewRelevance Result;
 	Result.bDrawRelevance = IsShown(View) && View->Family->EngineShowFlags.Particles;
diff --git a/Engine/Source/Runtime/Engine/Private/UnrealEngine.cpp b/Engine/Source/Runtime/Engine/Private/UnrealEngine.cpp
--- a/Engine/Source/Runtime/Engine/Private/UnrealEngine.cpp
+++ b/Engine/Source/Runtime/Engine/Private/UnrealEngine.cpp
@@ -6796,7 +6796,7 @@ void UEngine::PerformanceCapture(UWorld* World, const FString& CaptureName)
 void UEngine::LogPerformanceCapture(UWorld* World, const FString& CaptureName)
 {
 	const FString EventType = TEXT("PERF");
-	const int32 ChangeList = GEngineVersion.GetChangelist();
+	const int32 ChangeList = FEngineVersion::Current().GetChangelist();
 
 	TArray<FString> PathArray;
 	CaptureName.ParseIntoArray(PathArray, TEXT("/"), true);
diff --git a/Engine/Source/Runtime/Engine/Public/ParticleHelper.h b/Engine/Source/Runtime/Engine/Public/ParticleHelper.h
--- a/Engine/Source/Runtime/Engine/Public/ParticleHelper.h
+++ b/Engine/Source/Runtime/Engine/Public/ParticleHelper.h
@@ -2450,7 +2450,7 @@ public:
 		return false;
 	}
 	virtual void GetDynamicMeshElements(const TArray<const FSceneView*>& Views, const FSceneViewFamily& ViewFamily, uint32 VisibilityMap, FMeshElementCollector& Collector) const override;
-	virtual FPrimitiveViewRelevance GetViewRelevance(const FSceneView* View) override;
+	virtual FPrimitiveViewRelevance GetViewRelevance(const FSceneView* View) const override;
 	virtual void OnActorPositionChanged() override;
 	virtual void OnTransformChanged() override;
 
diff --git a/Engine/Source/Runtime/GameplayTags/Private/GameplayTagsManager.cpp b/Engine/Source/Runtime/GameplayTags/Private/GameplayTagsManager.cpp
--- a/Engine/Source/Runtime/GameplayTags/Private/GameplayTagsManager.cpp
+++ b/Engine/Source/Runtime/GameplayTags/Private/GameplayTagsManager.cpp
@@ -861,20 +861,16 @@ int32 UGameplayTagsManager::GameplayTagsMatchDepth(const FGameplayTag& GameplayT
 	TSet<FName> Tags1;
 	TSet<FName> Tags2;
 
+	const TSharedPtr<FGameplayTagNode>* TagNode = GameplayTagNodeMap.Find(GameplayTagOne);
+	if (TagNode)
 	{
-		FScopeLock Lock(&GameplayTagNodeMapCritical);
-
-		const TSharedPtr<FGameplayTagNode>* TagNode = GameplayTagNodeMap.Find(GameplayTagOne);
-		if (TagNode)
-		{
-			GetAllParentNodeNames(Tags1, *TagNode);
-		}
+		GetAllParentNodeNames(Tags1, *TagNode);
+	}
 
-		TagNode = GameplayTagNodeMap.Find(GameplayTagTwo);
-		if (TagNode)
-		{
-			GetAllParentNodeNames(Tags2, *TagNode);
-		}
+	TagNode = GameplayTagNodeMap.Find(GameplayTagTwo);
+	if (TagNode)
+	{
+		GetAllParentNodeNames(Tags2, *TagNode);
 	}
 
 	return Tags1.Intersect(Tags2).Num();
diff --git a/Engine/Source/Runtime/NetworkReplayStreaming/InMemoryNetworkReplayStreaming/InMemoryNetworkReplayStreaming.build.cs b/Engine/Source/Runtime/NetworkReplayStreaming/InMemoryNetworkReplayStreaming/InMemoryNetworkReplayStreaming.build.cs
--- a/Engine/Source/Runtime/NetworkReplayStreaming/InMemoryNetworkReplayStreaming/InMemoryNetworkReplayStreaming.build.cs
+++ b/Engine/Source/Runtime/NetworkReplayStreaming/InMemoryNetworkReplayStreaming/InMemoryNetworkReplayStreaming.build.cs
@@ -8,13 +8,19 @@ namespace UnrealBuildTool.Rules
 		{
 			PublicIncludePathModuleNames.Add("Engine");
 
+			PublicDependencyModuleNames.AddRange(
+				new string[]
+				{
+					"Engine"
+				}
+			);
+
 			PrivateDependencyModuleNames.AddRange(
 				new string[]
 				{
 					"Core",
 					"NetworkReplayStreaming",
-                    "Json",
-					"Engine"
+                    "Json"
 				}
 			);
 		}
diff --git a/Engine/Source/Runtime/Renderer/Private/PostProcess/SceneRenderTargets.cpp b/Engine/Source/Runtime/Renderer/Private/PostProcess/SceneRenderTargets.cpp
--- a/Engine/Source/Runtime/Renderer/Private/PostProcess/SceneRenderTargets.cpp
+++ b/Engine/Source/Runtime/Renderer/Private/PostProcess/SceneRenderTargets.cpp
@@ -1264,7 +1264,7 @@ bool FSceneRenderTargets::BeginRenderingSeparateTranslucency(FRHICommandList& RH
 			SeparateTranslucency = &ViewState->GetSeparateTranslucency(RHICmdList, ScaledSize);
 		}
 		FSceneViewState* ViewState = (FSceneViewState*)View.State;
-		const FTexture2DRHIRef &SeparateTranslucencyDepth = (const FTexture2DRHIRef&)ViewState->GetSeparateTranslucencyDepth(GetBufferSizeXY())->GetRenderTargetItem().TargetableTexture;
+		const FTexture2DRHIRef &SeparateTranslucencyDepth = (const FTexture2DRHIRef&)ViewState->GetSeparateTranslucencyDepth(RHICmdList, GetBufferSizeXY())->GetRenderTargetItem().TargetableTexture;
 
 		check((*SeparateTranslucency)->GetRenderTargetItem().TargetableTexture->GetClearColor() == FLinearColor::Black);
 		// clear the render target the first time, re-use afterwards
@@ -1304,7 +1304,7 @@ void FSceneRenderTargets::FinishRenderingSeparateTranslucency(FRHICommandList& R
 			static IConsoleVariable* STSP_CVar = IConsoleManager::Get().FindConsoleVariable(TEXT("r.SeparateTranslucencyScreenPercentage"));
 			FSceneViewState* ViewState = (FSceneViewState*)View.State;
 			FIntPoint ScaledSize(GetBufferSizeXY().X * (STSP_CVar->GetInt() / 100.0f), GetBufferSizeXY().Y * (STSP_CVar->GetInt() / 100.0f));
-			SeparateTranslucency = &ViewState->GetSeparateTranslucency( ScaledSize );
+			SeparateTranslucency = &ViewState->GetSeparateTranslucency( RHICmdList, ScaledSize );
 			SeparateTranslucencyDepth = &ViewState->GetSeparateTranslucencyDepth(RHICmdList, GetBufferSizeXY());
 		}
 
diff --git a/Engine/Source/Runtime/Renderer/Private/PrimitiveSceneInfo.cpp b/Engine/Source/Runtime/Renderer/Private/PrimitiveSceneInfo.cpp
--- a/Engine/Source/Runtime/Renderer/Private/PrimitiveSceneInfo.cpp
+++ b/Engine/Source/Runtime/Renderer/Private/PrimitiveSceneInfo.cpp
@@ -403,13 +403,8 @@ void FPrimitiveSceneInfo::UnlinkAttachmentGroup()
 	}
 }
 
-// Temporary:  Used to gather data for a bug QA hasn't been able to reproduce internally.
-#define TEST_BBOX_NANS_UE_21208
-
 void FPrimitiveSceneInfo::GatherLightingAttachmentGroupPrimitives(TArray<FPrimitiveSceneInfo*, SceneRenderingAllocator>& OutChildSceneInfos)
 {
->>>> ORIGINAL //depot/UE4-Orion/Engine/Source/Runtime/Renderer/Private/PrimitiveSceneInfo.cpp#8
-==== THEIRS //depot/UE4-Orion/Engine/Source/Runtime/Renderer/Private/PrimitiveSceneInfo.cpp#9
 #if ENABLE_NAN_DIAGNOSTIC
 	// local function that returns full name of object
 	auto GetObjectName = [](const UPrimitiveComponent* InPrimitive)->FString
@@ -432,35 +427,7 @@ void FPrimitiveSceneInfo::GatherLightingAttachmentGroupPrimitives(TArray<FPrimit
 
 #else 
 	// add self at the head of this queue
-==== YOURS //Marc.Audy_Z2487/Engine/Source/Runtime/Renderer/Private/PrimitiveSceneInfo.cpp
-
-#ifdef TEST_BBOX_NANS_UE_21208
-	
-	// local function that returns the name of object
-	auto GetObjectName = [](const UPrimitiveComponent* InPrimitive)->FString
-	{
-		return (InPrimitive) ? InPrimitive->GetFullName() : FString(TEXT("Unknown Object"));
-	};
-
-	// verify that the current object has a valid bbox before adding it
-	const float& BoundsRadius = this->Proxy->GetBounds().SphereRadius;
-	if (ensureMsgf(!FMath::IsNaN(BoundsRadius) && FMath::IsFinite(BoundsRadius),
-		TEXT("%s had an ill-formed bbox and was skipped during shadow setup, contact DavidH."), *GetObjectName(this->ComponentForDebuggingOnly)))
-	{
-
-		OutChildSceneInfos.Add(this);
-	} 
-	else
-	{
-		// return, leaving the TArray empty
-		return;
-	}
-
-#else
-	// add self at the head of the quie
-<<<<
 	OutChildSceneInfos.Add(this);
-
 #endif
 
 	if (!LightingAttachmentRoot.IsValid() && Proxy->LightAttachmentsAsGroup())
@@ -473,8 +440,6 @@ void FPrimitiveSceneInfo::GatherLightingAttachmentGroupPrimitives(TArray<FPrimit
 			for (int32 ChildIndex = 0, ChildIndexMax = AttachmentGroup->Primitives.Num(); ChildIndex < ChildIndexMax; ChildIndex++)
 			{
 				FPrimitiveSceneInfo* ShadowChild = AttachmentGroup->Primitives[ChildIndex];
->>>> ORIGINAL //depot/UE4-Orion/Engine/Source/Runtime/Renderer/Private/PrimitiveSceneInfo.cpp#8
-==== THEIRS //depot/UE4-Orion/Engine/Source/Runtime/Renderer/Private/PrimitiveSceneInfo.cpp#9
 #if ENABLE_NAN_DIAGNOSTIC
 				// Only enqueue objects with valid bounds using the normality of the SphereRaduis as criteria.
 
@@ -486,24 +451,8 @@ void FPrimitiveSceneInfo::GatherLightingAttachmentGroupPrimitives(TArray<FPrimit
 					checkSlow(!OutChildSceneInfos.Contains(ShadowChild))
 				    OutChildSceneInfos.Add(ShadowChild);
 				}
-#else
-==== YOURS //Marc.Audy_Z2487/Engine/Source/Runtime/Renderer/Private/PrimitiveSceneInfo.cpp
-
-#ifdef TEST_BBOX_NANS_UE_21208
-				// Only enqueue objects with valid bounds using the normality of the SphereRaduis as criteria.
-
-				const float& ShadowChildBoundsRadius = ShadowChild->Proxy->GetBounds().SphereRadius;
-
-				if (ensureMsgf(!FMath::IsNaN(ShadowChildBoundsRadius) && FMath::IsFinite(ShadowChildBoundsRadius),
-					TEXT("%s had an ill-formed bbox and was skipped during shadow setup, contact DavidH."), *GetObjectName(ShadowChild->ComponentForDebuggingOnly)))
-				{
-						checkSlow(!OutChildSceneInfos.Contains(ShadowChild))
-							OutChildSceneInfos.Add(ShadowChild);
-				}
 #else
 				// enqueue all objects.
-
-<<<<
 				checkSlow(!OutChildSceneInfos.Contains(ShadowChild))
 			    OutChildSceneInfos.Add(ShadowChild);
 #endif
@@ -512,8 +461,6 @@ void FPrimitiveSceneInfo::GatherLightingAttachmentGroupPrimitives(TArray<FPrimit
 	}
 }
 
-#undef TEST_BBOX_NANS_UE_21208
-
 FBoxSphereBounds FPrimitiveSceneInfo::GetAttachmentGroupBounds() const
 {
 	FBoxSphereBounds Bounds = Proxy->GetBounds();
diff --git a/Engine/Source/Runtime/Renderer/Private/ScenePrivate.h b/Engine/Source/Runtime/Renderer/Private/ScenePrivate.h
--- a/Engine/Source/Runtime/Renderer/Private/ScenePrivate.h
+++ b/Engine/Source/Runtime/Renderer/Private/ScenePrivate.h
@@ -705,13 +705,13 @@ public:
 		return SeparateTranslucencyDepthRT != nullptr;
 	}
 
-	TRefCountPtr<IPooledRenderTarget>& GetSeparateTranslucencyDepth(FIntPoint Size)
+	TRefCountPtr<IPooledRenderTarget>& GetSeparateTranslucencyDepth(FRHICommandList& RHICmdList, FIntPoint Size)
 	{
 		if (!SeparateTranslucencyDepthRT || SeparateTranslucencyDepthRT->GetDesc().Extent != Size)
 		{
 			// Create the SeparateTranslucency render target (alpha is needed to lerping)
 			FPooledRenderTargetDesc Desc(FPooledRenderTargetDesc::Create2DDesc(Size, PF_DepthStencil, FClearValueBinding::None, TexCreate_None, TexCreate_DepthStencilTargetable, true));
-			GRenderTargetPool.FindFreeElement(Desc, SeparateTranslucencyDepthRT, TEXT("SeparateTranslucencyDepth"));
+			GRenderTargetPool.FindFreeElement(RHICmdList, Desc, SeparateTranslucencyDepthRT, TEXT("SeparateTranslucencyDepth"));
 		}
 		return SeparateTranslucencyDepthRT;
 	}
diff --git a/Engine/Source/Runtime/Renderer/Private/TranslucentRendering.cpp b/Engine/Source/Runtime/Renderer/Private/TranslucentRendering.cpp
--- a/Engine/Source/Runtime/Renderer/Private/TranslucentRendering.cpp
+++ b/Engine/Source/Runtime/Renderer/Private/TranslucentRendering.cpp
@@ -937,9 +937,9 @@ void FDeferredShadingSceneRenderer::RenderTranslucencyParallel(FRHICommandListIm
 				// we need to allocate this now so it ends up in the snapshot
 				static IConsoleVariable* STSP_CVar = IConsoleManager::Get().FindConsoleVariable(TEXT("r.SeparateTranslucencyScreenPercentage"));
 				FIntPoint ScaledSize(SceneContext.GetBufferSizeXY().X * (STSP_CVar->GetInt() / 100.0f), SceneContext.GetBufferSizeXY().Y * (STSP_CVar->GetInt() / 100.0f) );
-				ViewState->GetSeparateTranslucency(ScaledSize);
+				ViewState->GetSeparateTranslucency(RHICmdList, ScaledSize);
 
-				ViewState->GetSeparateTranslucencyDepth(SceneContext.GetBufferSizeXY());
+				ViewState->GetSeparateTranslucencyDepth(RHICmdList, SceneContext.GetBufferSizeXY());
 				DownsampleDepthSurface(RHICmdList, ViewState->GetSeparateTranslucencyDepthSurface(), View, STSP_CVar->GetInt() / 100.0f);
 			}
 			FTranslucencyPassParallelCommandListSet ParallelCommandListSet(View, RHICmdList, 
@@ -1109,7 +1109,7 @@ void FDeferredShadingSceneRenderer::RenderTranslucency(FRHICommandListImmediate&
 				if (ViewState)
 				{
 					static IConsoleVariable* STSP_CVar = IConsoleManager::Get().FindConsoleVariable(TEXT("r.SeparateTranslucencyScreenPercentage"));
-					ViewState->GetSeparateTranslucencyDepth(SceneContext.GetBufferSizeXY());
+					ViewState->GetSeparateTranslucencyDepth(RHICmdList, SceneContext.GetBufferSizeXY());
 					DownsampleDepthSurface(RHICmdList, ViewState->GetSeparateTranslucencyDepthSurface(), View, STSP_CVar->GetInt() / 100.0f);
 				}
 				bool bSetupTranslucency = SceneContext.BeginRenderingSeparateTranslucency(RHICmdList, View, true);
