#ue4 - Fix compile error on XB One.

diff --git a/Engine/Source/Runtime/Engine/Private/Animation/AnimPhysicsSolver.cpp b/Engine/Source/Runtime/Engine/Private/Animation/AnimPhysicsSolver.cpp
--- a/Engine/Source/Runtime/Engine/Private/Animation/AnimPhysicsSolver.cpp
+++ b/Engine/Source/Runtime/Engine/Private/Animation/AnimPhysicsSolver.cpp
@@ -505,7 +505,7 @@ void FAnimPhys::ConstrainAlongDirection(float DeltaTime, TArray<FAnimPhysLinearL
 	if (FMath::Abs(Limits.X - Limits.Y) < SMALL_NUMBER)
 	{
 		// Fully locked axis, just generate one limit
-		LimitContainer.Add(FAnimPhysLinearLimit(FirstBody, SecondBody, FirstPosition, SecondPosition, AxisToConstrain, Distance / DeltaTime, Distance / DeltaTime, { MinimumForce, MaximumForce }));
+		LimitContainer.Add(FAnimPhysLinearLimit(FirstBody, SecondBody, FirstPosition, SecondPosition, AxisToConstrain, Distance / DeltaTime, Distance / DeltaTime, FVector2D(MinimumForce, MaximumForce)));
 	}
 	else
 	{
