Disabled warning about resaving material packages as it doesn't seem to be working correctly.

diff --git a/Engine/Source/Runtime/Engine/Private/Materials/Material.cpp b/Engine/Source/Runtime/Engine/Private/Materials/Material.cpp
--- a/Engine/Source/Runtime/Engine/Private/Materials/Material.cpp
+++ b/Engine/Source/Runtime/Engine/Private/Materials/Material.cpp
@@ -3451,10 +3451,6 @@ bool UMaterial::UpdateLightmassTextureTracking()
 
 	if ( bTexturesHaveChanged )
 	{
-		if (IsRunningCommandlet())
-		{
-			UE_LOG(LogMaterial, Warning, TEXT("UMaterial's dependent textures have changed, resave package %s to avoid deterministic cooking issues."), *GetPathName());
-		}
 		// This will invalidate any cached Lightmass material exports
 		SetLightingGuid();
 	}
