Fix AnyKey not working in ActionBindings
Fix allowing AnyKey to use keys that have already been consumed
Ensure that AnyKey is reported as handling an Action when a valid key comes through

diff --git a/Engine/Source/Runtime/Engine/Private/UserInterface/PlayerInput.cpp b/Engine/Source/Runtime/Engine/Private/UserInterface/PlayerInput.cpp
--- a/Engine/Source/Runtime/Engine/Private/UserInterface/PlayerInput.cpp
+++ b/Engine/Source/Runtime/Engine/Private/UserInterface/PlayerInput.cpp
@@ -672,11 +672,11 @@ void UPlayerInput::GetChordsForAction(const FInputActionBinding& ActionBinding,
 			{
 				for (auto KeyStateIt(KeyStateMap.CreateConstIterator()); KeyStateIt; ++KeyStateIt)
 				{
-					if (!KeyStateIt.Key().IsFloatAxis() && !KeyStateIt.Key().IsVectorAxis())
+					if (!KeyStateIt.Key().IsFloatAxis() && !KeyStateIt.Key().IsVectorAxis() && !IsKeyConsumed(KeyStateIt.Key())
 					{
 						FInputActionKeyMapping SubKeyMapping(KeyMapping);
 						SubKeyMapping.Key = KeyStateIt.Key();
-						GetChordsForKeyMapping(KeyMapping, ActionBinding, bGamePaused, FoundChords, KeysToConsume);
+						GetChordsForKeyMapping(SubKeyMapping, ActionBinding, bGamePaused, FoundChords, KeysToConsume);
 					}
 				}
 			}
@@ -696,7 +696,7 @@ void UPlayerInput::GetChordForKey(const FInputKeyBinding& KeyBinding, const bool
 	{
 		for (auto KeyStateIt(KeyStateMap.CreateConstIterator()); KeyStateIt; ++KeyStateIt)
 		{
-			if (!KeyStateIt.Key().IsFloatAxis() && !KeyStateIt.Key().IsVectorAxis())
+			if (!KeyStateIt.Key().IsFloatAxis() && !KeyStateIt.Key().IsVectorAxis() && !IsKeyConsumed(KeyStateIt.Key()))
 			{
 				FInputKeyBinding SubKeyBinding(KeyBinding);
 				SubKeyBinding.Chord.Key = KeyStateIt.Key();
@@ -1563,7 +1563,7 @@ bool UPlayerInput::IsKeyHandledByAction( FKey Key ) const
 {
 	for (const FInputActionKeyMapping& Mapping : ActionMappings)
 	{
-		if( Mapping.Key == Key && 
+		if( (Mapping.Key == Key || Mapping.Key == EKeys::AnyKey) && 
 			(Mapping.bAlt == false || IsAltPressed()) &&
 			(Mapping.bCtrl == false || IsCtrlPressed()) &&
 			(Mapping.bShift == false || IsShiftPressed()) &&
