Disable packet handlers until we fix the real issue + disable previous attempt at fix

diff --git a/Engine/Source/Runtime/Engine/Private/NetConnection.cpp b/Engine/Source/Runtime/Engine/Private/NetConnection.cpp
--- a/Engine/Source/Runtime/Engine/Private/NetConnection.cpp
+++ b/Engine/Source/Runtime/Engine/Private/NetConnection.cpp
@@ -90,7 +90,9 @@ void UNetConnection::InitBase(UNetDriver* InDriver,class FSocket* InSocket, cons
 	// Reset Handler
 	Handler.Reset(NULL);
 
+#if ENABLER_HANDLERS == 1
 	Handler = MakeUnique<PacketHandler>();
+#endif
 
 	if(Handler.IsValid())
 	{
@@ -157,10 +159,12 @@ void UNetConnection::InitConnection(UNetDriver* InDriver, EConnectionState InSta
 {
 	Driver = InDriver;
 
+#if ENABLER_HANDLERS == 1
 	if (!Handler.IsValid())
 	{
 		Handler = MakeUnique<PacketHandler>();
 	}
+#endif
 
 	// We won't be sending any packets, so use a default size
 	MaxPacket = (InMaxPacket == 0 || InMaxPacket > MAX_PACKET_SIZE) ? MAX_PACKET_SIZE : InMaxPacket;
diff --git a/Engine/Source/Runtime/Online/OnlineSubsystemUtils/Private/IpNetDriver.cpp b/Engine/Source/Runtime/Online/OnlineSubsystemUtils/Private/IpNetDriver.cpp
--- a/Engine/Source/Runtime/Online/OnlineSubsystemUtils/Private/IpNetDriver.cpp
+++ b/Engine/Source/Runtime/Online/OnlineSubsystemUtils/Private/IpNetDriver.cpp
@@ -171,7 +171,7 @@ void UIpNetDriver::TickDispatch( float DeltaTime )
 	const double StartReceiveTime = FPlatformTime::Seconds();
 
 	// Process all incoming packets.
-	uint8 Data[MAX_PACKET_SIZE * 2];		// x 2 for now until we resolve some issues with the sending side incorrectly calculating fragmentation size
+	uint8 Data[MAX_PACKET_SIZE];
 	TSharedRef<FInternetAddr> FromAddr = SocketSubsystem->CreateInternetAddr();
 	for( ; Socket != NULL; )
 	{
