Fix compile error

diff --git a/Engine/Source/Runtime/Engine/Private/UserInterface/PlayerInput.cpp b/Engine/Source/Runtime/Engine/Private/UserInterface/PlayerInput.cpp
--- a/Engine/Source/Runtime/Engine/Private/UserInterface/PlayerInput.cpp
+++ b/Engine/Source/Runtime/Engine/Private/UserInterface/PlayerInput.cpp
@@ -672,7 +672,7 @@ void UPlayerInput::GetChordsForAction(const FInputActionBinding& ActionBinding,
 			{
 				for (auto KeyStateIt(KeyStateMap.CreateConstIterator()); KeyStateIt; ++KeyStateIt)
 				{
-					if (!KeyStateIt.Key().IsFloatAxis() && !KeyStateIt.Key().IsVectorAxis() && !IsKeyConsumed(KeyStateIt.Key())
+					if (!KeyStateIt.Key().IsFloatAxis() && !KeyStateIt.Key().IsVectorAxis() && !IsKeyConsumed(KeyStateIt.Key()))
 					{
 						FInputActionKeyMapping SubKeyMapping(KeyMapping);
 						SubKeyMapping.Key = KeyStateIt.Key();
