Improved error messages for missing UTexture source mips.

diff --git a/Engine/Source/Runtime/Engine/Private/TextureDerivedData.cpp b/Engine/Source/Runtime/Engine/Private/TextureDerivedData.cpp
--- a/Engine/Source/Runtime/Engine/Private/TextureDerivedData.cpp
+++ b/Engine/Source/Runtime/Engine/Private/TextureDerivedData.cpp
@@ -708,7 +708,7 @@ class FTextureCacheDerivedDataWorker : public FNonAbandonableTask
 		if ( Texture.Source.HasHadBulkDataCleared() )
 		{
 			// don't do any work we can't reload this
-			UE_LOG(LogTexture, Error, TEXT("Unable to load texture source data could be because it has been cleared. %s"), *Texture.GetName())
+			UE_LOG(LogTexture, Error, TEXT("Unable to load texture source data could be because it has been cleared. %s"), *Texture.GetPathName())
 			return;
 		}
 
