UE-22084 Fixing crash when stopping sounds.

- Moving all active sound stops to the pending stop list
- Changed the list from a TArray to a TSet to avoid duplicate adds
- Moved ActiveSound::Stop to private to reduce chance of accidentally stopping it outside of audio device

diff --git a/Engine/Source/Runtime/Engine/Private/ActiveSound.cpp b/Engine/Source/Runtime/Engine/Private/ActiveSound.cpp
--- a/Engine/Source/Runtime/Engine/Private/ActiveSound.cpp
+++ b/Engine/Source/Runtime/Engine/Private/ActiveSound.cpp
@@ -274,7 +274,7 @@ void FActiveSound::UpdateWaveInstances( TArray<FWaveInstance*> &InWaveInstances,
 
 	if (bFinished)
 	{
-		Stop();
+		AudioDevice->StopActiveSound(this);
 	}
 	else if (ThisSoundsWaveInstances.Num() > 0)
 	{
diff --git a/Engine/Source/Runtime/Engine/Private/AudioDevice.cpp b/Engine/Source/Runtime/Engine/Private/AudioDevice.cpp
--- a/Engine/Source/Runtime/Engine/Private/AudioDevice.cpp
+++ b/Engine/Source/Runtime/Engine/Private/AudioDevice.cpp
@@ -1786,7 +1786,7 @@ int32 FAudioDevice::GetSortedActiveWaveInstances(TArray<FWaveInstance*>& WaveIns
 		if (!ActiveSound->Sound)
 		{
 			// No sound - cleanup and remove
-			ActiveSound->Stop();
+			AddSoundToStop(ActiveSound);
 		}
 		// If the world scene allows audio - tick wave instances.
 		else 
@@ -1801,7 +1801,7 @@ int32 FAudioDevice::GetSortedActiveWaveInstances(TArray<FWaveInstance*>& WaveIns
 					 *ActiveSound->DebugOriginalSoundName.ToString()))
 				{
 					// Sound was not valid, stop playing it.
-					ActiveSound->Stop();
+					AddSoundToStop(ActiveSound);
 				}
 				else
 #endif
@@ -1816,7 +1816,7 @@ int32 FAudioDevice::GetSortedActiveWaveInstances(TArray<FWaveInstance*>& WaveIns
 							Duration,
 							*ActiveSound->Sound->GetName(),
 							(ActiveSound->GetAudioComponent() ? *ActiveSound->GetAudioComponent()->GetName() : TEXT("NO COMPONENT")));
-						ActiveSound->Stop();
+						AddSoundToStop(ActiveSound);
 					}
 					else
 					{
@@ -2030,12 +2030,7 @@ void FAudioDevice::Update( bool bGameTicking )
 		Listener.UpdateCurrentInteriorSettings();
 	}
 
-	// Stop any pending active sounds that need to be stopped
-	while (PendingSoundsToStop.Num() > 0)
-	{
-		FActiveSound* SoundToStop = PendingSoundsToStop.Pop();
-		SoundToStop->Stop();
-	}
+	ProcessingPendingActiveSoundStops();
 
 	if (Sources.Num())
 	{
@@ -2097,7 +2092,7 @@ void FAudioDevice::StopAllSounds( bool bShouldStopUISounds )
 
 		if (bShouldStopUISounds)
 		{
-			ActiveSound->Stop();
+			AddSoundToStop(ActiveSound);
 		}
 		// If we're allowing UI sounds to continue then first filter on the active sounds state
 		else if (!ActiveSound->bIsUISound)
@@ -2110,12 +2105,15 @@ void FAudioDevice::StopAllSounds( bool bShouldStopUISounds )
 				FWaveInstance* WaveInstance = WaveInstanceIt.Value();
 				if (WaveInstance && !WaveInstance->bIsUISound)
 				{
-					ActiveSound->Stop();
+					AddSoundToStop(ActiveSound);
 					break;
 				}
 			}
 		}
 	}
+
+	// Immediately process stopping sounds
+	ProcessingPendingActiveSoundStops();
 }
 
 void FAudioDevice::AddNewActiveSound(const FActiveSound& NewActiveSound)
@@ -2180,9 +2178,27 @@ void FAudioDevice::AddNewActiveSound(const FActiveSound& NewActiveSound)
 
 }
 
+void FAudioDevice::ProcessingPendingActiveSoundStops()
+{
+	// Stop any pending active sounds that need to be stopped
+	for (FActiveSound* ActiveSound : PendingSoundsToStop)
+	{
+		check(ActiveSound);
+		ActiveSound->Stop();
+	}
+	PendingSoundsToStop.Reset();
+}
+
+
 void FAudioDevice::AddSoundToStop(FActiveSound* SoundToStop)
 {
-	PendingSoundsToStop.Add(SoundToStop);
+	check(SoundToStop);
+	bool bIsAlreadyInSet = false;
+	PendingSoundsToStop.Add(SoundToStop, &bIsAlreadyInSet);
+	if (bIsAlreadyInSet)
+	{
+		UE_LOG(LogAudio, Verbose, TEXT("Stopping sound which was already in the process of stopping"));
+	}
 }
 
 void FAudioDevice::StopActiveSound( UAudioComponent* AudioComponent )
@@ -2192,10 +2208,15 @@ void FAudioDevice::StopActiveSound( UAudioComponent* AudioComponent )
 	FActiveSound* ActiveSound = FindActiveSound(AudioComponent);
 	if (ActiveSound)
 	{
-		ActiveSound->Stop();
+		StopActiveSound(ActiveSound);
 	}
 }
 
+void FAudioDevice::StopActiveSound(FActiveSound* ActiveSound)
+{
+	AddSoundToStop(ActiveSound);
+}
+
 FActiveSound* FAudioDevice::FindActiveSound( UAudioComponent* AudioComponent )
 {
 	// find the active sound corresponding to this audio component
@@ -2343,19 +2364,22 @@ void FAudioDevice::Flush( UWorld* WorldToFlush, bool bClearActivatedReverb )
 		{
 			if (WorldToFlush == nullptr)
 			{
-				ActiveSound->Stop();
+				AddSoundToStop(ActiveSound);
 			}
 			else
 			{
 				UWorld* ActiveSoundWorld = ActiveSound->World.Get();
 				if (ActiveSoundWorld == nullptr || ActiveSoundWorld == WorldToFlush)
 				{
-					ActiveSound->Stop();
+					AddSoundToStop(ActiveSound);
 				}
 			}
 		}
 	}
 
+	// Immediately stop all pending active sounds
+	ProcessingPendingActiveSoundStops();
+
 	// Anytime we flush, make sure to clear all the listeners.  We'll get the right ones soon enough.
 	Listeners.Empty();
 	Listeners.Add(FListener());
@@ -2639,7 +2663,7 @@ void FAudioDevice::StopSoundsUsingResource(USoundWave* SoundWave, TArray<UAudioC
 				{
 					StoppedComponents.Add(AudioComponent);
 				}
-				ActiveSound->Stop();
+				AddSoundToStop(ActiveSound);
 				bStoppedSounds = true;
 				break;
 			}
diff --git a/Engine/Source/Runtime/Engine/Public/ActiveSound.h b/Engine/Source/Runtime/Engine/Public/ActiveSound.h
--- a/Engine/Source/Runtime/Engine/Public/ActiveSound.h
+++ b/Engine/Source/Runtime/Engine/Public/ActiveSound.h
@@ -231,8 +231,6 @@ public:
 	// Updates the wave instances to be played.
 	void UpdateWaveInstances( TArray<FWaveInstance*> &OutWaveInstances, const float DeltaTime );
 
-	void Stop();
-
 	/** 
 	 * Find an existing waveinstance attached to this audio component (if any)
 	 */
@@ -306,6 +304,13 @@ public:
 	uint32 GetSoundConcurrencyObjectID() const;
 
 private:
+
+	/** This is a friend so the audio device can call Stop() on the active sound. */
+	friend class FAudioDevice;
+
+	/** Stops the active sound. Can only be called from the owning audio device. */
+	void Stop();
+
 	void UpdateAdjustVolumeMultiplier(const float DeltaTime);
 
 	/** if OcclusionCheckInterval > 0.0, checks if the sound has become (un)occluded during playback
@@ -318,5 +323,4 @@ private:
 
 	/** Apply the interior settings to the ambient sound as appropriate */
 	void HandleInteriorVolumes( const FListener& Listener, struct FSoundParseParameters& ParseParams );
-
 };
\ No newline at end of file
diff --git a/Engine/Source/Runtime/Engine/Public/AudioDevice.h b/Engine/Source/Runtime/Engine/Public/AudioDevice.h
--- a/Engine/Source/Runtime/Engine/Public/AudioDevice.h
+++ b/Engine/Source/Runtime/Engine/Public/AudioDevice.h
@@ -373,6 +373,11 @@ public:
 	 */
 	void StopActiveSound( class UAudioComponent* AudioComponent );
 
+	/**
+	* Stops the active sound
+	*/
+	void StopActiveSound(FActiveSound* ActiveSound);
+
 	/**
 	 * Finds the active sound for the specified audio component
 	 */
@@ -757,6 +762,9 @@ protected:
 	*/
 	FVector GetListenerTransformedDirection(const FVector& Position, float* OutDistance);
 
+	/** Processes the set of pending sounds that need to be stopped */ 
+	void ProcessingPendingActiveSoundStops();
+
 public:
 
 	/** The maximum number of concurrent audible sounds */
@@ -860,7 +868,7 @@ public:
 private:
 
 	TArray<struct FActiveSound*> ActiveSounds;
-	TArray<struct FActiveSound*> PendingSoundsToStop;
+	TSet<struct FActiveSound*> PendingSoundsToStop;
 
 	TMap<UPTRINT, struct FActiveSound*> AudioComponentToActiveSoundMap;
 
