Merging CL 2735396 to Main

//depot/UE4-Orion/Engine/Source/Runtime/Engine/...

to //UE4/Main/Engine/Source/Runtime/Engine/...

diff --git a/Engine/Source/Runtime/Engine/Private/Particles/ParticleModules_Location.cpp b/Engine/Source/Runtime/Engine/Private/Particles/ParticleModules_Location.cpp
--- a/Engine/Source/Runtime/Engine/Private/Particles/ParticleModules_Location.cpp
+++ b/Engine/Source/Runtime/Engine/Private/Particles/ParticleModules_Location.cpp
@@ -1645,7 +1645,7 @@ void UParticleModuleLocationBoneSocket::FinalUpdate(FParticleEmitterInstance* Ow
 		for(int32 SourceIndex = 0; SourceIndex < SourceLocations.Num(); ++SourceIndex)
 		{
 			int32 BoneIndex = InstancePayload->SourceComponent->GetBoneIndex(SourceLocations[SourceIndex].BoneSocketName);
-			if (BoneIndex != INDEX_NONE)
+			if (BoneIndex != INDEX_NONE && InstancePayload->PrevFrameBoneSocketPositions.Num() > SourceIndex)
 			{
 				const FMatrix WorldBoneTM = SourceComponent->GetBoneMatrix(BoneIndex);
 				InstancePayload->PrevFrameBoneSocketPositions[SourceIndex] = WorldBoneTM.GetOrigin();
diff --git a/Engine/Source/Runtime/Engine/Public/ParticleHelper.h b/Engine/Source/Runtime/Engine/Public/ParticleHelper.h
--- a/Engine/Source/Runtime/Engine/Public/ParticleHelper.h
+++ b/Engine/Source/Runtime/Engine/Public/ParticleHelper.h
@@ -827,6 +827,11 @@ public:
 		return ( Find(Item) != INDEX_NONE );
 	}
 
+	int32 Num()
+	{
+		return ArrayMax;
+	}
+
 private:
 	/** The maximum number of elements in the array. This is cannot dynamically change.*/
 	int32 ArrayMax;
