#Fix issue with when animation doesn't load correctly
- caused by merge error

diff --git a/Engine/Source/Runtime/Engine/Private/Animation/AnimSequence.cpp b/Engine/Source/Runtime/Engine/Private/Animation/AnimSequence.cpp
--- a/Engine/Source/Runtime/Engine/Private/Animation/AnimSequence.cpp
+++ b/Engine/Source/Runtime/Engine/Private/Animation/AnimSequence.cpp
@@ -1685,9 +1685,7 @@ bool UAnimSequence::CompressRawAnimData(float MaxPosDiff, float MaxAngleDiff)
 	bool bRemovedKeys = false;
 #if WITH_EDITORONLY_DATA
 
-	const USkeleton* MySkeleton = GetSkeleton();
-
-	if ( MySkeleton && ensureMsgf(RawAnimationData.Num() > 0, TEXT("%s is trying to compress while raw animation is missing"), *GetName()))
+	if ( ensureMsgf (RawAnimationData.Num() > 0, TEXT("%s is trying to compress while raw animation is missing"), *GetName()) )
 	{
 		// This removes trivial keys, and this has to happen before the removing tracks
 		for(int32 TrackIndex=0; TrackIndex<RawAnimationData.Num(); TrackIndex++)
@@ -1695,31 +1693,36 @@ bool UAnimSequence::CompressRawAnimData(float MaxPosDiff, float MaxAngleDiff)
 			bRemovedKeys |= CompressRawAnimSequenceTrack(RawAnimationData[TrackIndex], MaxPosDiff, MaxAngleDiff);
 		}
 
-		bool bCompressScaleKeys = false;
-		// go through remove keys if not needed
-		for(int32 TrackIndex=0; TrackIndex<RawAnimationData.Num(); TrackIndex++)
+		const USkeleton* MySkeleton = GetSkeleton();
+
+		if(MySkeleton)
 		{
-			FRawAnimSequenceTrack const& RawData = RawAnimationData[TrackIndex];
-			if(RawData.ScaleKeys.Num() > 0)
+			bool bCompressScaleKeys = false;
+			// go through remove keys if not needed
+			for(int32 TrackIndex=0; TrackIndex<RawAnimationData.Num(); TrackIndex++)
 			{
-				// if scale key exists, see if we can just empty it
-				if((RawData.ScaleKeys.Num() > 1) || (RawData.ScaleKeys[0].Equals(FVector(1.f)) == false))
+				FRawAnimSequenceTrack const& RawData = RawAnimationData[TrackIndex];
+				if(RawData.ScaleKeys.Num() > 0)
 				{
-					bCompressScaleKeys = true;
-					break;
+					// if scale key exists, see if we can just empty it
+					if((RawData.ScaleKeys.Num() > 1) || (RawData.ScaleKeys[0].Equals(FVector(1.f)) == false))
+					{
+						bCompressScaleKeys = true;
+						break;
+					}
 				}
 			}
-		} 
 
-		// if we don't have scale, we should delete all scale keys
-		// if you have one track that has scale, we still should support scale, so compress scale
-		if(!bCompressScaleKeys)
-		{
-			// then remove all scale keys
-			for(int32 TrackIndex=0; TrackIndex<RawAnimationData.Num(); TrackIndex++)
+			// if we don't have scale, we should delete all scale keys
+			// if you have one track that has scale, we still should support scale, so compress scale
+			if(!bCompressScaleKeys)
 			{
-				FRawAnimSequenceTrack& RawData = RawAnimationData[TrackIndex];
-				RawData.ScaleKeys.Empty();
+				// then remove all scale keys
+				for(int32 TrackIndex=0; TrackIndex<RawAnimationData.Num(); TrackIndex++)
+				{
+					FRawAnimSequenceTrack& RawData = RawAnimationData[TrackIndex];
+					RawData.ScaleKeys.Empty();
+				}
 			}
 		}
 
