Fix compile warnings introduced by Orion -> Main merge

diff --git a/Engine/Source/Runtime/Engine/Private/GameplayStatics.cpp b/Engine/Source/Runtime/Engine/Private/GameplayStatics.cpp
--- a/Engine/Source/Runtime/Engine/Private/GameplayStatics.cpp
+++ b/Engine/Source/Runtime/Engine/Private/GameplayStatics.cpp
@@ -683,16 +683,16 @@ UParticleSystemComponent* UGameplayStatics::SpawnEmitterAtLocation(UWorld* World
 
 UParticleSystemComponent* UGameplayStatics::SpawnEmitterAttached(UParticleSystem* EmitterTemplate, USceneComponent* AttachToComponent, FName AttachPointName, FVector Location, FRotator Rotation, EAttachLocation::Type LocationType, bool bAutoDestroy)
 {
-	UParticleSystemComponent* PSC = NULL;
+	UParticleSystemComponent* PSC = nullptr;
 	if (EmitterTemplate)
 	{
-		if (AttachToComponent == NULL)
+		if (AttachToComponent == nullptr)
 		{
 			UE_LOG(LogScript, Warning, TEXT("UGameplayStatics::SpawnEmitterAttached: NULL AttachComponent specified!"));
 		}
 		else if (AttachToComponent->GetWorld() && AttachToComponent->GetWorld()->GetNetMode() != NM_DedicatedServer)
 		{
-			UParticleSystemComponent* PSC = CreateParticleSystem(EmitterTemplate, AttachToComponent->GetWorld(), AttachToComponent->GetOwner(), bAutoDestroy);
+			PSC = CreateParticleSystem(EmitterTemplate, AttachToComponent->GetWorld(), AttachToComponent->GetOwner(), bAutoDestroy);
 			PSC->AttachTo(AttachToComponent, AttachPointName);
 			if (LocationType == EAttachLocation::KeepWorldPosition)
 			{
@@ -704,10 +704,9 @@ UParticleSystemComponent* UGameplayStatics::SpawnEmitterAttached(UParticleSystem
 			}
 			PSC->SetRelativeScale3D(FVector(1.f));
 			PSC->ActivateSystem(true);
-			return PSC;
 		}
 	}
-	return nullptr;
+	return PSC;
 }
 
 void UGameplayStatics::BreakHitResult(const FHitResult& Hit, bool& bBlockingHit, bool& bInitialOverlap, float& Time, FVector& Location, FVector& ImpactPoint, FVector& Normal, FVector& ImpactNormal, UPhysicalMaterial*& PhysMat, AActor*& HitActor, UPrimitiveComponent*& HitComponent, FName& HitBoneName, int32& HitItem, FVector& TraceStart, FVector& TraceEnd)
