fixed so that when invalid data asset is found, make sure to reset the other data to match and avoid to ensure

diff --git a/Engine/Source/Runtime/Engine/Private/Animation/AnimSequence.cpp b/Engine/Source/Runtime/Engine/Private/Animation/AnimSequence.cpp
--- a/Engine/Source/Runtime/Engine/Private/Animation/AnimSequence.cpp
+++ b/Engine/Source/Runtime/Engine/Private/Animation/AnimSequence.cpp
@@ -367,6 +367,29 @@ void UAnimSequence::PostLoad()
 	{
 		RemoveNaNTracks();
 	}
+
+	// fix up broken data
+	// recently we saw RawAnimationData is empty but has NumFrames and SequenceLength
+	// We're cleaning it up if no Raw data exists, and NumFrames > 0, then it is inconsistent
+	if (RawAnimationData.Num() == 0 && NumFrames != 0)
+	{
+		NumFrames = 0; 
+		SequenceLength = 0.f;
+		
+		UE_LOG(LogAnimation, Warning, TEXT("[%s] has no animation data but has frame number - removing frame number."), *GetName());
+		
+		if (!IsRunningGame())
+		{
+			static FName NAME_LoadErrors("LoadErrors");
+			FMessageLog LoadErrors(NAME_LoadErrors);
+
+			TSharedRef<FTokenizedMessage> Message = LoadErrors.Warning();
+			Message->AddToken(FTextToken::Create(LOCTEXT("InvalidAnimationData1", "The Animation ")));
+			Message->AddToken(FAssetNameToken::Create(GetPathName(), FText::FromString(GetName())));
+			Message->AddToken(FTextToken::Create(LOCTEXT("InvalidAnimationData2", " has no animation data but contains invalid number of frames and sequence length. Resetting...")));
+			LoadErrors.Notify();
+		}
+	}
 #endif // WITH_EDITOR
 
 	Super::PostLoad();
@@ -380,10 +403,18 @@ void UAnimSequence::PostLoad()
 	if ( NumFrames == 0 )
 	{
 		UE_LOG(LogAnimation, Warning, TEXT("No animation data exists for sequence %s (%s)"), *GetName(), (GetOuter() ? *GetOuter()->GetFullName() : *GetFullName()) );
+		
 #if WITH_EDITOR
-		if( GIsEditor )
+		if (!IsRunningGame())
 		{
-			UE_LOG(LogAnimation, Warning, TEXT("Removing bad AnimSequence (%s) from %s."), *GetName(), (GetOuter() ? *GetOuter()->GetFullName() : *GetFullName()) );
+			static FName NAME_LoadErrors("LoadErrors");
+			FMessageLog LoadErrors(NAME_LoadErrors);
+
+			TSharedRef<FTokenizedMessage> Message = LoadErrors.Warning();
+			Message->AddToken(FTextToken::Create(LOCTEXT("EmptyAnimationData1", "The Animation ")));
+			Message->AddToken(FAssetNameToken::Create(GetPathName(), FText::FromString(GetName())));
+			Message->AddToken(FTextToken::Create(LOCTEXT("EmptyAnimationData2", " has no animation data. Recommend to remove.")));
+			LoadErrors.Notify();
 		}
 #endif
 	}
@@ -1692,7 +1723,8 @@ bool UAnimSequence::CompressRawAnimData(float MaxPosDiff, float MaxAngleDiff)
 	bool bRemovedKeys = false;
 #if WITH_EDITORONLY_DATA
 
-	if ( ensureMsgf (RawAnimationData.Num() > 0, TEXT("%s is trying to compress while raw animation is missing"), *GetName()) )
+	// check if track name exists first, if no data to start with that should be been empty
+	if (NumFrames > 0 && ensureMsgf(RawAnimationData.Num() > 0, TEXT("%s is trying to compress while raw animation is missing"), *GetName()))
 	{
 		// This removes trivial keys, and this has to happen before the removing tracks
 		for(int32 TrackIndex=0; TrackIndex<RawAnimationData.Num(); TrackIndex++)
@@ -1736,6 +1768,11 @@ bool UAnimSequence::CompressRawAnimData(float MaxPosDiff, float MaxAngleDiff)
 		CompressedTrackOffsets.Empty();
 		CompressedScaleOffsets.Empty();
 	}
+	else
+	{
+		CompressedTrackOffsets.Empty();
+		CompressedScaleOffsets.Empty();
+	}
 	
 #endif
 	return bRemovedKeys;
