Fix clang compiler crash when compiling android from the command line (or on the build machine)

diff --git a/Engine/Source/Runtime/Engine/Classes/Engine/NetSerialization.h b/Engine/Source/Runtime/Engine/Classes/Engine/NetSerialization.h
--- a/Engine/Source/Runtime/Engine/Classes/Engine/NetSerialization.h
+++ b/Engine/Source/Runtime/Engine/Classes/Engine/NetSerialization.h
@@ -383,6 +383,15 @@ struct FFastArraySerializer
 	static bool FastArrayDeltaSerialize( TArray<Type> &Items, FNetDeltaSerializeInfo& Parms, SerializerType& ArraySerializer );
 };
 
+// Struct used only in FFastArraySerializer::FastArrayDeltaSerialize, however, declaring it within the templated function
+// causes crashes on some clang compilers
+struct FFastArraySerializer_FastArrayDeltaSerialize_FIdxIDPair
+{
+	FFastArraySerializer_FastArrayDeltaSerialize_FIdxIDPair(int32 _idx, int32 _id) : Idx(_idx), ID(_id) { }
+	int32	Idx;
+	int32	ID;
+};
+
 /** The function that implements Fast TArray Replication  */
 template< typename Type, typename SerializerType >
 bool FFastArraySerializer::FastArrayDeltaSerialize( TArray<Type> &Items, FNetDeltaSerializeInfo& Parms, SerializerType& ArraySerializer )
@@ -532,15 +541,7 @@ bool FFastArraySerializer::FastArrayDeltaSerialize( TArray<Type> &Items, FNetDel
 			return false;
 		}
 
-
-		struct FIdxIDPair
-		{
-			FIdxIDPair(int32 _idx, int32 _id) : Idx(_idx), ID(_id) { }
-			int32	Idx;
-			int32	ID;
-		};
-
-		TArray<FIdxIDPair, TInlineAllocator<8> >	ChangedElements;
+		TArray<FFastArraySerializer_FastArrayDeltaSerialize_FIdxIDPair, TInlineAllocator<8> >	ChangedElements;
 		TArray<int32, TInlineAllocator<8> >		DeletedElements;
 
 		int32 DeleteCount = (OldMap ? OldMap->Num() : 0) - Items.Num(); // Note: this is incremented when we add new items below.
@@ -596,7 +597,7 @@ bool FFastArraySerializer::FastArrayDeltaSerialize( TArray<Type> &Items, FNetDel
 					UE_LOG(LogNetFastTArray, Log, TEXT("       Changed! Was: %d. Element ID: %d"), *OldValuePtr, Items[i].ReplicationID);
 
 					// Changed
-					new (ChangedElements)FIdxIDPair(i, Items[i].ReplicationID);
+					ChangedElements.Add(FFastArraySerializer_FastArrayDeltaSerialize_FIdxIDPair(i, Items[i].ReplicationID));
 				}
 			}
 			else
@@ -606,7 +607,7 @@ bool FFastArraySerializer::FastArrayDeltaSerialize( TArray<Type> &Items, FNetDel
 				// The item really should have a valid ReplicationID but in the case of loading from a save game,
 				// items may not have been marked dirty individually. Its ok to just assign them one here.
 				// New
-				new (ChangedElements)FIdxIDPair(i, Items[i].ReplicationID);
+				ChangedElements.Add(FFastArraySerializer_FastArrayDeltaSerialize_FIdxIDPair(i, Items[i].ReplicationID));
 				DeleteCount++; // We added something new, so our initial DeleteCount value must be incremented.
 			}
 		}
