Mergine changes from //depot/UE4 to //UE4/Main

diff --git a/Engine/Source/Editor/Sequencer/Private/SequencerNodeTree.h b/Engine/Source/Editor/Sequencer/Private/SequencerNodeTree.h
--- a/Engine/Source/Editor/Sequencer/Private/SequencerNodeTree.h
+++ b/Engine/Source/Editor/Sequencer/Private/SequencerNodeTree.h
@@ -94,7 +94,7 @@ private:
 	 * @param Track	The type to get sections from
 	 * @param SectionAreaNode	The section area which section interfaces belong to
 	 */
-	void MakeSectionInterfaces( UMovieSceneTrack& Track, TSharedRef<class FTrackNode>& SectionAreaNode );
+	void MakeSectionInterfaces( UMovieSceneTrack& Track, TSharedRef<class FSequencerTrackNode>& SectionAreaNode );
 
 	/**
 	 * Creates a new object binding node and any parent binding nodes.
diff --git a/Engine/Source/Editor/Sequencer/Private/SequencerPrivatePCH.h b/Engine/Source/Editor/Sequencer/Private/SequencerPrivatePCH.h
--- a/Engine/Source/Editor/Sequencer/Private/SequencerPrivatePCH.h
+++ b/Engine/Source/Editor/Sequencer/Private/SequencerPrivatePCH.h
@@ -20,6 +20,7 @@
 #include "SequencerSectionCategoryNode.h"
 #include "SequencerSectionKeyAreaNode.h"
 #include "SequencerTrackNode.h"
+#include "Sequencer.h"
 #include "ISectionLayoutBuilder.h"
 #include "ITimeSlider.h"
 #include "SSequencer.h"
diff --git a/Engine/Source/Editor/Sequencer/Private/SequencerSectionLayoutBuilder.cpp b/Engine/Source/Editor/Sequencer/Private/SequencerSectionLayoutBuilder.cpp
--- a/Engine/Source/Editor/Sequencer/Private/SequencerSectionLayoutBuilder.cpp
+++ b/Engine/Source/Editor/Sequencer/Private/SequencerSectionLayoutBuilder.cpp
@@ -4,7 +4,7 @@
 #include "SequencerSectionLayoutBuilder.h"
 
 
-FSequencerSectionLayoutBuilder::FSequencerSectionLayoutBuilder( TSharedRef<FTrackNode> InRootNode )
+FSequencerSectionLayoutBuilder::FSequencerSectionLayoutBuilder( TSharedRef<FSequencerTrackNode> InRootNode )
 	: RootNode( InRootNode )
 	, CurrentNode( InRootNode )
 {}
diff --git a/Engine/Source/Editor/Sequencer/Private/SequencerSectionLayoutBuilder.h b/Engine/Source/Editor/Sequencer/Private/SequencerSectionLayoutBuilder.h
--- a/Engine/Source/Editor/Sequencer/Private/SequencerSectionLayoutBuilder.h
+++ b/Engine/Source/Editor/Sequencer/Private/SequencerSectionLayoutBuilder.h
@@ -3,14 +3,14 @@
 #pragma once
 
 class ISequencerSection;
-class FTrackNode;
+class FSequencerTrackNode;
 
 
 class FSequencerSectionLayoutBuilder
 	: public ISectionLayoutBuilder
 {
 public:
-	FSequencerSectionLayoutBuilder( TSharedRef<FTrackNode> InRootNode );
+	FSequencerSectionLayoutBuilder( TSharedRef<FSequencerTrackNode> InRootNode );
 
 public:
 
@@ -24,7 +24,7 @@ public:
 private:
 
 	/** Root node of the tree */
-	TSharedRef<FTrackNode> RootNode;
+	TSharedRef<FSequencerTrackNode> RootNode;
 
 	/** The current node that other nodes are added to */
 	TSharedRef<FSequencerDisplayNode> CurrentNode;
diff --git a/Engine/Source/Editor/Sequencer/Private/SequencerSettings.cpp b/Engine/Source/Editor/Sequencer/Private/SequencerSettings.cpp
--- a/Engine/Source/Editor/Sequencer/Private/SequencerSettings.cpp
+++ b/Engine/Source/Editor/Sequencer/Private/SequencerSettings.cpp
@@ -12,7 +12,6 @@ USequencerSettings::USequencerSettings( const FObjectInitializer& ObjectInitiali
 	SpawnPosition = SSP_Origin;
 	bShowFrameNumbers = true;
 	bShowRangeSlider = false;
-	bLockInOutToStartEndRange = false;
 	bIsSnapEnabled = true;
 	TimeSnapInterval = .05f;
 	bSnapKeyTimesToInterval = true;
@@ -27,6 +26,7 @@ USequencerSettings::USequencerSettings( const FObjectInitializer& ObjectInitiali
 	bAutoScrollEnabled = false;
 	bShowCurveEditor = false;
 	bShowCurveEditorCurveToolTips = true;
+	bLooping = false;
 }
 
 bool USequencerSettings::GetAutoKeyEnabled() const
@@ -127,20 +127,6 @@ void USequencerSettings::SetShowRangeSlider(bool InbShowRangeSlider)
 	}
 }
 
-bool USequencerSettings::GetLockInOutToStartEndRange() const
-{
-	return bLockInOutToStartEndRange;
-}
-
-void USequencerSettings::SetLockInOutToStartEndRange(bool InbLockInOutToStartEndRange)
-{
-	if ( bLockInOutToStartEndRange != InbLockInOutToStartEndRange )
-	{
-		bLockInOutToStartEndRange = InbLockInOutToStartEndRange;
-		SaveConfig();
-	}
-}
-
 bool USequencerSettings::GetIsSnapEnabled() const
 {
 	return bIsSnapEnabled;
@@ -325,6 +311,20 @@ void USequencerSettings::SetShowCurveEditor(bool InbShowCurveEditor)
 	}
 }
 
+bool USequencerSettings::IsLooping() const
+{
+	return bLooping;
+}
+
+void USequencerSettings::SetLooping(bool bInLooping)
+{
+	if (bLooping != bInLooping)
+	{
+		bLooping = bInLooping;
+		SaveConfig();
+	}
+}
+
 bool USequencerSettings::GetShowCurveEditorCurveToolTips() const
 {
 	return bShowCurveEditorCurveToolTips;
@@ -358,5 +358,4 @@ ULevelEditorSequencerSettings::ULevelEditorSequencerSettings( const FObjectIniti
 	bKeyInterpPropertiesOnly = true;
 	TimeSnapInterval = 0.033334f;
 	bShowRangeSlider = true;
-	bLockInOutToStartEndRange = true;
 }
\ No newline at end of file
diff --git a/Engine/Source/Editor/Sequencer/Private/SequencerSettings.h b/Engine/Source/Editor/Sequencer/Private/SequencerSettings.h
--- a/Engine/Source/Editor/Sequencer/Private/SequencerSettings.h
+++ b/Engine/Source/Editor/Sequencer/Private/SequencerSettings.h
@@ -89,11 +89,6 @@ public:
 	/** Sets whether or not to show frame numbers. */
 	void SetShowRangeSlider(bool InbShowRangeSlider);
 
-	/** Gets whether or not the in and out are bounded by start and end when scrolling. */
-	bool GetLockInOutToStartEndRange() const;
-	/** Sets whether or not the in and out are bounded by start and end when scrolling. */
-	void SetLockInOutToStartEndRange(bool InbLockInOutToStartEndRange);
-
 	/** Gets whether or not snapping is enabled. */
 	bool GetIsSnapEnabled() const;
 	/** Sets whether or not snapping is enabled. */
@@ -163,6 +158,11 @@ public:
 	bool GetShowCurveEditorCurveToolTips() const;
 	/** Sets whether or not to show curve tool tips in the curve editor. */
 	void SetShowCurveEditorCurveToolTips(bool InbShowCurveEditorCurveToolTips);
+	
+	/** Gets the loop mode. */
+	bool IsLooping() const;
+	/** Sets the loop mode. */
+	void SetLooping(bool bInLooping);
 
 	/** Snaps a time value in seconds to the currently selected interval. */
 	float SnapTimeToInterval(float InTimeValue) const;
@@ -192,9 +192,6 @@ protected:
 	UPROPERTY( config )
 	bool bShowRangeSlider;
 
-	UPROPERTY( config )
-	bool bLockInOutToStartEndRange;
-
 	UPROPERTY( config )
 	bool bIsSnapEnabled;
 
@@ -237,6 +234,9 @@ protected:
 	UPROPERTY( config )
 	bool bShowCurveEditorCurveToolTips;
 
+	UPROPERTY( config )
+	bool bLooping;
+
 	FOnShowCurveEditorChanged OnShowCurveEditorChanged;
 };
 
diff --git a/Engine/Source/Editor/Sequencer/Private/TimeSliderController.cpp b/Engine/Source/Editor/Sequencer/Private/TimeSliderController.cpp
--- a/Engine/Source/Editor/Sequencer/Private/TimeSliderController.cpp
+++ b/Engine/Source/Editor/Sequencer/Private/TimeSliderController.cpp
@@ -391,8 +391,8 @@ int32 FSequencerTimeSliderController::DrawPlaybackRange(const FGeometry& Allotte
 
 	TRange<float> PlaybackRange = TimeSliderArgs.PlaybackRange.Get();
 
-	float PlaybackRangeL = RangeToScreen.InputToLocalX(PlaybackRange.GetLowerBoundValue());
-	float PlaybackRangeR = RangeToScreen.InputToLocalX(PlaybackRange.GetUpperBoundValue());
+	float PlaybackRangeL = RangeToScreen.InputToLocalX(PlaybackRange.GetLowerBoundValue()) - 1;
+	float PlaybackRangeR = RangeToScreen.InputToLocalX(PlaybackRange.GetUpperBoundValue()) + 1;
 
 	FSlateDrawElement::MakeBox(
 		OutDrawElements,
@@ -517,7 +517,7 @@ FReply FSequencerTimeSliderController::OnMouseButtonUp( TSharedRef<SWidget> Widg
 					NewValue = LastRange[1];
 				}
 
-				TimeSliderArgs.OnViewRangeChanged.ExecuteIfBound(TRange<float>(DownValue, NewValue), EViewRangeInterpolation::Immediate, false);
+				TimeSliderArgs.OnViewRangeChanged.ExecuteIfBound(TRange<float>(DownValue, NewValue), EViewRangeInterpolation::Immediate);
 						
 				if( !TimeSliderArgs.ViewRange.IsBound() )
 				{	
@@ -576,8 +576,7 @@ FReply FSequencerTimeSliderController::OnMouseMove( TSharedRef<SWidget> WidgetOw
 				float NewViewOutputMin = LocalViewRangeMin - InputDelta.X;
 				float NewViewOutputMax = LocalViewRangeMax - InputDelta.X;
 
-				const bool bMaintainRange = true;
-				ClampViewRange(NewViewOutputMin, NewViewOutputMax, bMaintainRange);
+				ClampViewRange(NewViewOutputMin, NewViewOutputMax);
 				SetViewRange(NewViewOutputMin, NewViewOutputMax, EViewRangeInterpolation::Immediate);
 			}
 		}
@@ -793,53 +792,26 @@ float FSequencerTimeSliderController::FrameToTime(int32 Frame) const
 	return SequencerHelpers::FrameToTime(Frame, FrameRate);
 }
 
-void FSequencerTimeSliderController::ClampViewRange(float& NewRangeMin, float& NewRangeMax, bool bMaintainRange)
+void FSequencerTimeSliderController::ClampViewRange(float& NewRangeMin, float& NewRangeMax)
 {
-	// If locked, clamp the new range to clamp range
-	if (TimeSliderArgs.Settings->GetShowRangeSlider() && TimeSliderArgs.Settings->GetLockInOutToStartEndRange())
+	bool bNeedsClampSet = false;
+	float NewClampRangeMin = TimeSliderArgs.ClampRange.Get().GetLowerBoundValue();
+	if ( NewRangeMin < TimeSliderArgs.ClampRange.Get().GetLowerBoundValue() )
 	{
-		if (bMaintainRange)
-		{
-			if ( NewRangeMin < TimeSliderArgs.ClampRange.Get().GetLowerBoundValue() )
-			{
-				NewRangeMin = TimeSliderArgs.ClampRange.Get().GetLowerBoundValue();
-				NewRangeMax = NewRangeMin + (TimeSliderArgs.ViewRange.Get().GetUpperBoundValue() - TimeSliderArgs.ViewRange.Get().GetLowerBoundValue());
-			}
-	
-			else if ( NewRangeMax > TimeSliderArgs.ClampRange.Get().GetUpperBoundValue() )
-			{
-				NewRangeMax = TimeSliderArgs.ClampRange.Get().GetUpperBoundValue();
-				NewRangeMin = NewRangeMax - (TimeSliderArgs.ViewRange.Get().GetUpperBoundValue() - TimeSliderArgs.ViewRange.Get().GetLowerBoundValue());
-			}
-		}
-		else
-		{
-			NewRangeMin = FMath::Clamp(NewRangeMin, TimeSliderArgs.ClampRange.Get().GetLowerBoundValue(), TimeSliderArgs.ClampRange.Get().GetUpperBoundValue());
-			NewRangeMax = FMath::Clamp(NewRangeMax, TimeSliderArgs.ClampRange.Get().GetLowerBoundValue(), TimeSliderArgs.ClampRange.Get().GetUpperBoundValue());
-		}
+		NewClampRangeMin = NewRangeMin;
+		bNeedsClampSet = true;
 	}
-	// Otherwise, grow the clamp range to the new range
-	else
+
+	float NewClampRangeMax = TimeSliderArgs.ClampRange.Get().GetUpperBoundValue();
+	if ( NewRangeMax > TimeSliderArgs.ClampRange.Get().GetUpperBoundValue() )
 	{
-		bool bNeedsClampSet = false;
-		float NewClampRangeMin = TimeSliderArgs.ClampRange.Get().GetLowerBoundValue();
-		if ( NewRangeMin < TimeSliderArgs.ClampRange.Get().GetLowerBoundValue() )
-		{
-			NewClampRangeMin = NewRangeMin;
-			bNeedsClampSet = true;
-		}
-	
-		float NewClampRangeMax = TimeSliderArgs.ClampRange.Get().GetUpperBoundValue();
-		if ( NewRangeMax > TimeSliderArgs.ClampRange.Get().GetUpperBoundValue() )
-		{
-			NewClampRangeMax = NewRangeMax;
-			bNeedsClampSet = true;
-		}
+		NewClampRangeMax = NewRangeMax;
+		bNeedsClampSet = true;
+	}
 
-		if (bNeedsClampSet)
-		{
-			SetClampRange(NewClampRangeMin, NewClampRangeMax);
-		}
+	if (bNeedsClampSet)
+	{
+		SetClampRange(NewClampRangeMin, NewClampRangeMax);
 	}
 }
 
@@ -847,7 +819,7 @@ void FSequencerTimeSliderController::SetViewRange( float NewRangeMin, float NewR
 {
 	const TRange<float> NewRange(NewRangeMin, NewRangeMax);
 
-	TimeSliderArgs.OnViewRangeChanged.ExecuteIfBound( NewRange, Interpolation, false );
+	TimeSliderArgs.OnViewRangeChanged.ExecuteIfBound( NewRange, Interpolation );
 
 	if( !TimeSliderArgs.ViewRange.IsBound() )
 	{	
@@ -882,8 +854,7 @@ bool FSequencerTimeSliderController::ZoomByDelta( float InDelta, float MousePosi
 
 	if( FMath::Abs( OutputChange ) > 0.01f && NewViewOutputMin < NewViewOutputMax )
 	{
-		const bool bMaintainRange = false;
-		ClampViewRange(NewViewOutputMin, NewViewOutputMax, bMaintainRange);
+		ClampViewRange(NewViewOutputMin, NewViewOutputMax);
 		SetViewRange(NewViewOutputMin, NewViewOutputMax, EViewRangeInterpolation::Animated);
 		return true;
 	}
@@ -904,8 +875,7 @@ void FSequencerTimeSliderController::PanByDelta( float InDelta )
 	float NewViewOutputMin = CurrentMin + InDelta;
 	float NewViewOutputMax = CurrentMax + InDelta;
 
-	const bool bMaintainRange = true;
-	ClampViewRange(NewViewOutputMin, NewViewOutputMax, bMaintainRange);
+	ClampViewRange(NewViewOutputMin, NewViewOutputMax);
 	SetViewRange(NewViewOutputMin, NewViewOutputMax, EViewRangeInterpolation::Animated);
 }
 
diff --git a/Engine/Source/Editor/Sequencer/Private/TimeSliderController.h b/Engine/Source/Editor/Sequencer/Private/TimeSliderController.h
--- a/Engine/Source/Editor/Sequencer/Private/TimeSliderController.h
+++ b/Engine/Source/Editor/Sequencer/Private/TimeSliderController.h
@@ -70,9 +70,8 @@ public:
 	 *
 	 * @param NewRangeMin		The new lower bound of the range
 	 * @param NewRangeMax		The new upper bound of the range
-	 * @param bMaintainRange	Maintain the current range size if 
 	 */	
-	void ClampViewRange(float& NewRangeMin, float& NewRangeMax, bool bMaintainRange);
+	void ClampViewRange(float& NewRangeMin, float& NewRangeMax);
 
 	/**
 	 * Set a new range based on a min, max and an interpolation mode
diff --git a/Engine/Source/Editor/Sequencer/Private/Tools/EditToolDragOperations.cpp b/Engine/Source/Editor/Sequencer/Private/Tools/EditToolDragOperations.cpp
--- a/Engine/Source/Editor/Sequencer/Private/Tools/EditToolDragOperations.cpp
+++ b/Engine/Source/Editor/Sequencer/Private/Tools/EditToolDragOperations.cpp
@@ -69,7 +69,7 @@ TOptional<FSequencerSnapField::FSnapResult> SnapToInterval(const TArray<float>&
 /** How many pixels near the mouse has to be before snapping occurs */
 const float PixelSnapWidth = 10.f;
 
-TRange<float> GetSectionBoundaries(UMovieSceneSection* Section, TSharedPtr<FTrackNode> SequencerNode)
+TRange<float> GetSectionBoundaries(UMovieSceneSection* Section, TSharedPtr<FSequencerTrackNode> SequencerNode)
 {
 	// Find the borders of where you can drag to
 	float LowerBound = -FLT_MAX, UpperBound = FLT_MAX;
diff --git a/Engine/Source/Editor/Sequencer/Private/Tools/SequencerEntityVisitor.cpp b/Engine/Source/Editor/Sequencer/Private/Tools/SequencerEntityVisitor.cpp
--- a/Engine/Source/Editor/Sequencer/Private/Tools/SequencerEntityVisitor.cpp
+++ b/Engine/Source/Editor/Sequencer/Private/Tools/SequencerEntityVisitor.cpp
@@ -62,7 +62,7 @@ void FSequencerEntityWalker::HandleNode(const ISequencerEntityVisitor& Visitor,
 {
 	if (InNode.GetType() == ESequencerNode::Track)
 	{
-		HandleNode(Visitor, InNode, static_cast<FTrackNode&>(InNode).GetSections());
+		HandleNode(Visitor, InNode, static_cast<FSequencerTrackNode&>(InNode).GetSections());
 	}
 
 	if (InNode.IsExpanded())
@@ -107,7 +107,7 @@ void FSequencerEntityWalker::HandleNode(const ISequencerEntityVisitor& Visitor,
 		}
 		else if (InNode.GetType() == ESequencerNode::Track)
 		{
-			TSharedPtr<FSequencerSectionKeyAreaNode> SectionKeyNode = static_cast<FTrackNode&>(InNode).GetTopLevelKeyNode();
+			TSharedPtr<FSequencerSectionKeyAreaNode> SectionKeyNode = static_cast<FSequencerTrackNode&>(InNode).GetTopLevelKeyNode();
 			if (SectionKeyNode.IsValid())
 			{
 				HandleKeyAreaNode(Visitor, *SectionKeyNode, InNode, InSections);
diff --git a/Engine/Source/Editor/Sequencer/Private/VirtualTrackArea.cpp b/Engine/Source/Editor/Sequencer/Private/VirtualTrackArea.cpp
--- a/Engine/Source/Editor/Sequencer/Private/VirtualTrackArea.cpp
+++ b/Engine/Source/Editor/Sequencer/Private/VirtualTrackArea.cpp
@@ -50,14 +50,14 @@ TSharedPtr<FSequencerDisplayNode> FVirtualTrackArea::HitTestNode(float InPhysica
 	return TreeView.HitTestNode(InPhysicalPosition);
 }
 
-TSharedPtr<FTrackNode> GetParentTrackNode(FSequencerDisplayNode& In)
+TSharedPtr<FSequencerTrackNode> GetParentTrackNode(FSequencerDisplayNode& In)
 {
 	FSequencerDisplayNode* Current = &In;
 	while(Current && Current->GetType() != ESequencerNode::Object)
 	{
 		if (Current->GetType() == ESequencerNode::Track)
 		{
-			return StaticCastSharedRef<FTrackNode>(Current->AsShared());
+			return StaticCastSharedRef<FSequencerTrackNode>(Current->AsShared());
 		}
 		Current = Current->GetParent().Get();
 	}
@@ -71,7 +71,7 @@ TOptional<FSectionHandle> FVirtualTrackArea::HitTestSection(FVector2D InPhysical
 
 	if (Node.IsValid())
 	{
-		TSharedPtr<FTrackNode> TrackNode = GetParentTrackNode(*Node);
+		TSharedPtr<FSequencerTrackNode> TrackNode = GetParentTrackNode(*Node);
 
 		if (TrackNode.IsValid())
 		{
@@ -111,7 +111,7 @@ FSequencerSelectedKey FVirtualTrackArea::HitTestKey(FVector2D InPhysicalPosition
 	switch (Node->GetType())
 	{
 		case ESequencerNode::KeyArea: 	KeyAreaNode = StaticCastSharedPtr<FSequencerSectionKeyAreaNode>(Node); break;
-		case ESequencerNode::Track:		KeyAreaNode = StaticCastSharedPtr<FTrackNode>(Node)->GetTopLevelKeyNode(); break;
+		case ESequencerNode::Track:		KeyAreaNode = StaticCastSharedPtr<FSequencerTrackNode>(Node)->GetTopLevelKeyNode(); break;
 	}
 
 	if (KeyAreaNode.IsValid())
@@ -128,7 +128,7 @@ FSequencerSelectedKey FVirtualTrackArea::HitTestKey(FVector2D InPhysicalPosition
 	// Failing that, and the node is collapsed, we check for key groupings
 	else if (!Node->IsExpanded())
 	{
-		TSharedPtr<FTrackNode> TrackNode = GetParentTrackNode(*Node);
+		TSharedPtr<FSequencerTrackNode> TrackNode = GetParentTrackNode(*Node);
 		const TArray< TSharedRef<ISequencerSection> >& Sections = TrackNode->GetSections();
 
 		for ( int32 SectionIndex = 0; SectionIndex < Sections.Num(); ++SectionIndex )
diff --git a/Engine/Source/Editor/Sequencer/Public/MovieSceneTrackEditor.h b/Engine/Source/Editor/Sequencer/Public/MovieSceneTrackEditor.h
--- a/Engine/Source/Editor/Sequencer/Public/MovieSceneTrackEditor.h
+++ b/Engine/Source/Editor/Sequencer/Public/MovieSceneTrackEditor.h
@@ -7,11 +7,14 @@
 #include "ISequencerSection.h"
 #include "ISequencerTrackEditor.h"
 #include "MovieScene.h"
+#include "MovieSceneTrack.h"
 #include "ScopedTransaction.h"
 #include "MovieSceneSequence.h"
 
 
 class ISequencerSection;
+class UMovieSceneTrack;
+
 
 /** Delegate for adding keys for a property
  * float - The time at which to add the key.
@@ -125,23 +128,22 @@ public:
 		bool bWasCreated;
 	};
 
-	FFindOrCreateTrackResult FindOrCreateTrackForObject( const FGuid& ObjectHandle, TSubclassOf<class UMovieSceneTrack> TrackClass, FName UniqueTypeName, bool bCreateTrackIfMissing = true )
+	FFindOrCreateTrackResult FindOrCreateTrackForObject( const FGuid& ObjectHandle, TSubclassOf<UMovieSceneTrack> TrackClass, FName PropertyName = NAME_None, bool bCreateTrackIfMissing = true )
 	{
 		FFindOrCreateTrackResult Result;
 		bool bTrackExisted;
 
-		check( UniqueTypeName != NAME_None );
-
 		UMovieScene* MovieScene = GetSequencer()->GetFocusedMovieSceneSequence()->GetMovieScene();
-		Result.Track = MovieScene->FindTrack( TrackClass, ObjectHandle, UniqueTypeName );
+		Result.Track = MovieScene->FindTrack( TrackClass, ObjectHandle, PropertyName );
 		bTrackExisted = Result.Track != nullptr;
 
-		if( !Result.Track && bCreateTrackIfMissing )
+		if (!Result.Track && bCreateTrackIfMissing)
 		{
-			Result.Track = AddTrack( MovieScene, ObjectHandle, TrackClass, UniqueTypeName );
+			Result.Track = AddTrack(MovieScene, ObjectHandle, TrackClass, PropertyName);
 		}
 
 		Result.bWasCreated = bTrackExisted == false && Result.Track != nullptr;
+
 		return Result;
 	}
 
diff --git a/Engine/Source/Editor/SequencerWidgets/Public/ITimeSlider.h b/Engine/Source/Editor/SequencerWidgets/Public/ITimeSlider.h
--- a/Engine/Source/Editor/SequencerWidgets/Public/ITimeSlider.h
+++ b/Engine/Source/Editor/SequencerWidgets/Public/ITimeSlider.h
@@ -14,7 +14,7 @@ enum class EViewRangeInterpolation
 };
 
 DECLARE_DELEGATE_TwoParams( FOnScrubPositionChanged, float, bool )
-DECLARE_DELEGATE_ThreeParams( FOnViewRangeChanged, TRange<float>, EViewRangeInterpolation, bool )
+DECLARE_DELEGATE_TwoParams( FOnViewRangeChanged, TRange<float>, EViewRangeInterpolation )
 DECLARE_DELEGATE_OneParam( FOnRangeChanged, TRange<float> )
 
 /** Structure used to wrap up a range, and an optional animation target */
diff --git a/Engine/Source/Editor/UMGEditor/Private/Animation/MarginTrackEditor.cpp b/Engine/Source/Editor/UMGEditor/Private/Animation/MarginTrackEditor.cpp
--- a/Engine/Source/Editor/UMGEditor/Private/Animation/MarginTrackEditor.cpp
+++ b/Engine/Source/Editor/UMGEditor/Private/Animation/MarginTrackEditor.cpp
@@ -19,7 +19,8 @@ class FMarginPropertySection
 	: public FPropertySection
 {
 public:
-	FMarginPropertySection( UMovieSceneSection& InSectionObject, FName SectionName )
+
+	FMarginPropertySection( UMovieSceneSection& InSectionObject, const FText& SectionName )
 		: FPropertySection(InSectionObject, SectionName) {}
 
 	virtual void GenerateSectionLayout( class ISectionLayoutBuilder& LayoutBuilder ) const override
@@ -74,7 +75,7 @@ TSharedRef<FPropertySection> FMarginTrackEditor::MakePropertySectionInterface( U
 	check( SupportsType( SectionObject.GetOuter()->GetClass() ) );
 
 	UClass* SectionClass = SectionObject.GetOuter()->GetClass();
-	return MakeShareable( new FMarginPropertySection( SectionObject, Track.GetTrackName() ) );
+	return MakeShareable(new FMarginPropertySection(SectionObject, Track.GetDisplayName()));
 }
 
 
diff --git a/Engine/Source/Editor/UMGEditor/Private/Animation/Sequencer2DTransformTrackEditor.cpp b/Engine/Source/Editor/UMGEditor/Private/Animation/Sequencer2DTransformTrackEditor.cpp
--- a/Engine/Source/Editor/UMGEditor/Private/Animation/Sequencer2DTransformTrackEditor.cpp
+++ b/Engine/Source/Editor/UMGEditor/Private/Animation/Sequencer2DTransformTrackEditor.cpp
@@ -18,7 +18,7 @@ class F2DTransformSection
 {
 public:
 
-	F2DTransformSection( UMovieSceneSection& InSectionObject, FName SectionName )
+	F2DTransformSection( UMovieSceneSection& InSectionObject, const FText& SectionName )
 		: FPropertySection(InSectionObject, SectionName) {}
 
 	virtual void GenerateSectionLayout( class ISectionLayoutBuilder& LayoutBuilder ) const override
@@ -112,7 +112,7 @@ TSharedRef<FPropertySection> F2DTransformTrackEditor::MakePropertySectionInterfa
 	check( SupportsType( SectionObject.GetOuter()->GetClass() ) );
 
 	UClass* SectionClass = SectionObject.GetOuter()->GetClass();
-	return MakeShareable( new F2DTransformSection( SectionObject, Track.GetTrackName() ) );
+	return MakeShareable(new F2DTransformSection(SectionObject, Track.GetDisplayName()));
 }
 
 
diff --git a/Engine/Source/Editor/UnrealEd/Private/Lightmass/LightmassRender.cpp b/Engine/Source/Editor/UnrealEd/Private/Lightmass/LightmassRender.cpp
--- a/Engine/Source/Editor/UnrealEd/Private/Lightmass/LightmassRender.cpp
+++ b/Engine/Source/Editor/UnrealEd/Private/Lightmass/LightmassRender.cpp
@@ -196,11 +196,12 @@ public:
 				{
 					TArray<FShaderType*> ShaderTypes;
 					TArray<FVertexFactoryType*> VFTypes;
-					GetDependentShaderAndVFTypes(GMaxRHIShaderPlatform, ShaderTypes, VFTypes);
+					TArray<const FShaderPipelineType*> ShaderPipelineTypes;
+					GetDependentShaderAndVFTypes(GMaxRHIShaderPlatform, ShaderTypes, ShaderPipelineTypes, VFTypes);
 
 					// Overwrite the shader map Id's dependencies with ones that came from the FMaterial actually being compiled (this)
 					// This is necessary as we change FMaterial attributes like GetShadingModel(), which factor into the ShouldCache functions that determine dependent shader types
-					ResourceId.SetShaderDependencies(ShaderTypes, VFTypes);
+					ResourceId.SetShaderDependencies(ShaderTypes, ShaderPipelineTypes, VFTypes);
 				}
 
 				// Override with a special usage so we won't re-use the shader map used by the material for rendering
diff --git a/Engine/Source/Editor/UnrealEd/Private/MaterialGraph.cpp b/Engine/Source/Editor/UnrealEd/Private/MaterialGraph.cpp
--- a/Engine/Source/Editor/UnrealEd/Private/MaterialGraph.cpp
+++ b/Engine/Source/Editor/UnrealEd/Private/MaterialGraph.cpp
@@ -510,6 +510,8 @@ FText UMaterialGraph::GetCustomDataPinName( uint32 Index ) const
 			return LOCTEXT("Backlit", "Backlit");
 		case MSM_Cloth:
 			return LOCTEXT("Cloth", "Cloth");
+		case MSM_Eye:
+			return LOCTEXT("IrisMask", "Iris Mask");
 		default:
 			return LOCTEXT("CustomData0", "Custom Data 0");
 		}
@@ -520,6 +522,8 @@ FText UMaterialGraph::GetCustomDataPinName( uint32 Index ) const
 		{
 		case MSM_ClearCoat:
 			return LOCTEXT("ClearCoatRoughness", "Clear Coat Roughness");
+		case MSM_Eye:
+			return LOCTEXT("IrisDistance", "Iris Distance");
 		default:
 			return LOCTEXT("CustomData1", "Custom Data 1");
 		}
diff --git a/Engine/Source/Programs/ShaderCompileWorker/Private/ShaderCompileWorker.cpp b/Engine/Source/Programs/ShaderCompileWorker/Private/ShaderCompileWorker.cpp
--- a/Engine/Source/Programs/ShaderCompileWorker/Private/ShaderCompileWorker.cpp
+++ b/Engine/Source/Programs/ShaderCompileWorker/Private/ShaderCompileWorker.cpp
@@ -122,7 +122,8 @@ public:
 
 		while(true)
 		{
-			TArray<FJobResult> JobResults;
+			TArray<FJobResult> SingleJobResults;
+			TArray<FPipelineJobResult> PipelineJobResults;
 
 			// Read & Process Input
 			{
@@ -135,7 +136,7 @@ public:
 				UE_LOG(LogShaders, Log, TEXT("Processing shader"));
 				LastCompileTime = FPlatformTime::Seconds();
 
-				ProcessInputFromArchive(InputFilePtr, JobResults);
+				ProcessInputFromArchive(InputFilePtr, SingleJobResults, PipelineJobResults);
 
 				// Close the input file.
 				delete InputFilePtr;
@@ -144,7 +145,7 @@ public:
 			// Prepare for output
 			FArchive* OutputFilePtr = CreateOutputArchive();
 			check(OutputFilePtr);
-			WriteToOutputArchive(OutputFilePtr, JobResults);
+			WriteToOutputArchive(OutputFilePtr, SingleJobResults, PipelineJobResults);
 
 			// Close the output file.
 			delete OutputFilePtr;
@@ -171,6 +172,12 @@ private:
 		FShaderCompilerOutput CompilerOutput;
 	};
 
+	struct FPipelineJobResult
+	{
+		FString PipelineName;
+		TArray<FJobResult> SingleJobs;
+	};
+
 	const int32 ParentProcessId;
 	const FString WorkingDirectory;
 	const FString InputFilename;
@@ -225,7 +232,7 @@ private:
 		}
 	}
 
-	void ProcessInputFromArchive(FArchive* InputFilePtr, TArray<FJobResult>& OutJobResults)
+	void ProcessInputFromArchive(FArchive* InputFilePtr, TArray<FJobResult>& OutSingleJobResults, TArray<FPipelineJobResult>& OutPipelineJobResults)
 	{
 		FArchive& InputFile = *InputFilePtr;
 		int32 InputVersion;
@@ -267,7 +274,7 @@ private:
 				ProcessCompilationJob(CompilerInput, CompilerOutput, WorkingDirectory);
 
 				// Serialize the job's output.
-			FJobResult& JobResult = *new(OutJobResults) FJobResult;
+				FJobResult& JobResult = *new(OutSingleJobResults) FJobResult;
 				JobResult.CompilerOutput = CompilerOutput;
 			}
 		}
@@ -280,6 +287,78 @@ private:
 
 			int32 NumPipelines = 0;
 			InputFile << NumPipelines;
+
+			for (int32 Index = 0; Index < NumPipelines; ++Index)
+			{
+				FPipelineJobResult& PipelineJob = *new(OutPipelineJobResults) FPipelineJobResult;
+
+				InputFile << PipelineJob.PipelineName;
+
+				int32 NumStages = 0;
+				InputFile << NumStages;
+
+				TArray<FShaderCompilerInput> CompilerInputs;
+				CompilerInputs.AddDefaulted(NumStages);
+
+				for (int32 StageIndex = 0; StageIndex < NumStages; ++StageIndex)
+				{
+					// Deserialize the job's inputs.
+					InputFile << CompilerInputs[StageIndex];
+
+					if (IsValidRef(CompilerInputs[StageIndex].SharedEnvironment))
+					{
+						// Merge the shared environment into the per-shader environment before calling into the compile function
+						CompilerInputs[StageIndex].Environment.Merge(*CompilerInputs[StageIndex].SharedEnvironment);
+					}
+				}
+
+				ProcessShaderPipelineCompilationJob(PipelineJob, CompilerInputs);
+			}
+		}
+	}
+
+	void ProcessShaderPipelineCompilationJob(FPipelineJobResult& PipelineJob, TArray<FShaderCompilerInput>& CompilerInputs)
+	{
+		checkf(CompilerInputs.Num() > 0, TEXT("Exiting due to Pipeline %s having zero jobs!"), *PipelineJob.PipelineName);
+
+		// Process the job.
+		FShaderCompilerOutput FirstCompilerOutput;
+		CompilerInputs[0].bCompilingForShaderPipeline = true;
+		CompilerInputs[0].bIncludeUsedOutputs = false;
+		ProcessCompilationJob(CompilerInputs[0], FirstCompilerOutput, WorkingDirectory);
+
+		// Serialize the job's output.
+		FJobResult& JobResult = *new(PipelineJob.SingleJobs) FJobResult;
+		JobResult.CompilerOutput = FirstCompilerOutput;
+
+		bool bEnableRemovingUnused = true;
+
+		//#todo-rco: Only remove for pure VS & PS stages
+		for (int32 Index = 0; Index < CompilerInputs.Num(); ++Index)
+		{
+			auto Stage = CompilerInputs[Index].Target.Frequency;
+			if (Stage != SF_Vertex && Stage != SF_Pixel)
+			{
+				bEnableRemovingUnused = false;
+				break;
+			}
+		}
+
+		for (int32 Index = 1; Index < CompilerInputs.Num(); ++Index)
+		{
+			if (bEnableRemovingUnused && PipelineJob.SingleJobs.Last().CompilerOutput.bSupportsQueryingUsedAttributes)
+			{
+				CompilerInputs[Index].bIncludeUsedOutputs = true;
+				CompilerInputs[Index].bCompilingForShaderPipeline = true;
+				CompilerInputs[Index].UsedOutputs = PipelineJob.SingleJobs.Last().CompilerOutput.UsedAttributes;
+			}
+
+			FShaderCompilerOutput CompilerOutput;
+			ProcessCompilationJob(CompilerInputs[Index], CompilerOutput, WorkingDirectory);
+
+			// Serialize the job's output.
+			FJobResult& JobResult = *new(PipelineJob.SingleJobs) FJobResult;
+			JobResult.CompilerOutput = CompilerOutput;
 		}
 	}
 
@@ -336,7 +415,7 @@ private:
 		return OutputFilePtr;
 	}
 
-	void WriteToOutputArchive(FArchive* OutputFilePtr, TArray<FJobResult>& JobResults)
+	void WriteToOutputArchive(FArchive* OutputFilePtr, TArray<FJobResult>& SingleJobResults, TArray<FPipelineJobResult>& PipelineJobResults)
 	{
 		FArchive& OutputFile = *OutputFilePtr;
 
@@ -354,12 +433,12 @@ private:
 			int32 SingleJobHeader = ShaderCompileWorkerSingleJobHeader;
 			OutputFile << SingleJobHeader;
 
-			int32 NumBatches = JobResults.Num();
+			int32 NumBatches = SingleJobResults.Num();
 			OutputFile << NumBatches;
 
-			for (int32 ResultIndex = 0; ResultIndex < JobResults.Num(); ResultIndex++)
+			for (int32 ResultIndex = 0; ResultIndex < SingleJobResults.Num(); ResultIndex++)
 			{
-				FJobResult& JobResult = JobResults[ResultIndex];
+				FJobResult& JobResult = SingleJobResults[ResultIndex];
 				OutputFile << JobResult.CompilerOutput;
 			}
 		}
@@ -367,9 +446,21 @@ private:
 		{
 			int32 PipelineJobHeader = ShaderCompileWorkerPipelineJobHeader;
 			OutputFile << PipelineJobHeader;
+			int32 NumBatches = PipelineJobResults.Num();
+			OutputFile << NumBatches;
 
-			int32 NumPipelines = 0;
-			OutputFile << NumPipelines;
+			for (int32 ResultIndex = 0; ResultIndex < PipelineJobResults.Num(); ResultIndex++)
+			{
+				auto& PipelineJob = PipelineJobResults[ResultIndex];
+				OutputFile << PipelineJob.PipelineName;
+				int32 NumStageJobs = PipelineJob.SingleJobs.Num();
+				OutputFile << NumStageJobs;
+				for (int32 Index = 0; Index < NumStageJobs; ++Index)
+				{
+					FJobResult& JobResult = PipelineJob.SingleJobs[Index];
+					OutputFile << JobResult.CompilerOutput;
+				}
+			}
 		}
 	}
 
diff --git a/Engine/Source/Runtime/AIModule/Classes/BehaviorTree/Services/BTService_RunEQS.h b/Engine/Source/Runtime/AIModule/Classes/BehaviorTree/Services/BTService_RunEQS.h
--- a/Engine/Source/Runtime/AIModule/Classes/BehaviorTree/Services/BTService_RunEQS.h
+++ b/Engine/Source/Runtime/AIModule/Classes/BehaviorTree/Services/BTService_RunEQS.h
@@ -14,7 +14,7 @@ struct FBTEQSServiceMemory
 	int32 RequestID;
 };
 
-UCLASS(hidecategories=(Service))
+UCLASS()
 class AIMODULE_API UBTService_RunEQS : public UBTService_BlackboardBase
 {
 	GENERATED_BODY()
diff --git a/Engine/Source/Runtime/AIModule/Classes/EnvironmentQuery/EnvQueryTest.h b/Engine/Source/Runtime/AIModule/Classes/EnvironmentQuery/EnvQueryTest.h
--- a/Engine/Source/Runtime/AIModule/Classes/EnvironmentQuery/EnvQueryTest.h
+++ b/Engine/Source/Runtime/AIModule/Classes/EnvironmentQuery/EnvQueryTest.h
@@ -106,6 +106,16 @@ class AIMODULE_API UEnvQueryTest : public UEnvQueryNode
 	UPROPERTY(EditDefaultsOnly, Category=Score, Meta=(ClampMin="0.001", UIMin="0.001"))
 	FAIDataProviderFloatValue ScoringFactor;
 
+	/** When specified gets used to normalize test's results in such a way that the closer a value is to SweetSpotValue
+	 *	the higher normalized result it will produce. Value farthest from SweetSpotValue will be normalized
+	 *	to 0, and all the other values in between will get normalized linearly with the distance to SweetSpotValue. */
+	UPROPERTY(EditDefaultsOnly, Category = Score, Meta=(EditCondition = "bDefineSweetSpot"))
+	FAIDataProviderFloatValue SweetSpotValue;
+
+	/** When set to true enables usage of SweetSpotValue. It's false by default */
+	UPROPERTY()
+	bool bDefineSweetSpot;
+	
 	/** Validation: item type that can be used with this test */
 	TSubclassOf<UEnvQueryItemType> ValidItemType;
 
diff --git a/Engine/Source/Runtime/AIModule/Classes/Perception/AIPerceptionComponent.h b/Engine/Source/Runtime/AIModule/Classes/Perception/AIPerceptionComponent.h
--- a/Engine/Source/Runtime/AIModule/Classes/Perception/AIPerceptionComponent.h
+++ b/Engine/Source/Runtime/AIModule/Classes/Perception/AIPerceptionComponent.h
@@ -203,7 +203,7 @@ public:
 	FORCEINLINE TActorPerceptionContainer::TIterator GetPerceptualDataIterator() { return TActorPerceptionContainer::TIterator(PerceptualData); }
 	FORCEINLINE TActorPerceptionContainer::TConstIterator GetPerceptualDataConstIterator() const { return TActorPerceptionContainer::TConstIterator(PerceptualData); }
 
-	void GetHostileActors(TArray<AActor*>& OutActors) const;
+	virtual void GetHostileActors(TArray<AActor*>& OutActors) const;
 
 	// @note Will stop on first age 0 stimulus
 	const FActorPerceptionInfo* GetFreshestTrace(const FAISenseID Sense) const;
diff --git a/Engine/Source/Runtime/AIModule/Private/EnvironmentQuery/EnvQueryTest.cpp b/Engine/Source/Runtime/AIModule/Private/EnvironmentQuery/EnvQueryTest.cpp
--- a/Engine/Source/Runtime/AIModule/Private/EnvironmentQuery/EnvQueryTest.cpp
+++ b/Engine/Source/Runtime/AIModule/Private/EnvironmentQuery/EnvQueryTest.cpp
@@ -93,6 +93,13 @@ void UEnvQueryTest::NormalizeItemScores(FEnvQueryInstance& QueryInstance)
 
 	if (MinScore != MaxScore)
 	{
+		if (bDefineSweetSpot)
+		{
+			SweetSpotValue.BindData(QueryOwner, QueryInstance.QueryID);
+		}
+		const float BestValue = bDefineSweetSpot ? SweetSpotValue.GetValue() : MaxScore;
+		const float ValueSpan = FMath::Max(FMath::Abs(BestValue - MinScore), FMath::Abs(BestValue - MaxScore));
+
 		for (int32 ItemIndex = 0; ItemIndex < QueryInstance.ItemDetails.Num(); ItemIndex++, DetailInfo++)
 		{
 			if (QueryInstance.Items[ItemIndex].IsValid() == false)
@@ -106,7 +113,8 @@ void UEnvQueryTest::NormalizeItemScores(FEnvQueryInstance& QueryInstance)
 			if (TestValue != UEnvQueryTypes::SkippedItemValue)
 			{
 				const float ClampedScore = FMath::Clamp(TestValue, MinScore, MaxScore);
-				const float NormalizedScore = (ClampedScore - MinScore) / (MaxScore - MinScore);
+				const float NormalizedScore = (ValueSpan - FMath::Abs(BestValue - ClampedScore)) / ValueSpan;
+
 				// TODO? Add an option to invert the normalized score before applying an equation.
  				const float NormalizedScoreForEquation = /*bMirrorNormalizedScore ? (1.0f - NormalizedScore) :*/ NormalizedScore;
 				switch (ScoringEquation)
diff --git a/Engine/Source/Runtime/Analytics/QoSReporter/Private/QoSReporter.cpp b/Engine/Source/Runtime/Analytics/QoSReporter/Private/QoSReporter.cpp
--- a/Engine/Source/Runtime/Analytics/QoSReporter/Private/QoSReporter.cpp
+++ b/Engine/Source/Runtime/Analytics/QoSReporter/Private/QoSReporter.cpp
@@ -9,6 +9,7 @@
 #if WITH_PERFCOUNTERS
 	#include "PerfCountersModule.h"
 #endif // WITH_PERFCOUNTERS
+#include "PerfCountersHelpers.h"
 
 #ifndef WITH_QOSREPORTER
 	#error "WITH_QOSREPORTER should be defined in Build.cs file"
@@ -16,6 +17,7 @@
 
 bool FQoSReporter::bIsInitialized;
 TSharedPtr<IAnalyticsProvider> FQoSReporter::Analytics;
+FGuid FQoSReporter::InstanceId;
 
 /**
 * External code should bind this delegate if QoS reporting is desired,
@@ -81,6 +83,10 @@ void FQoSReporter::Initialize()
 		FName(*DefaultEngineAnalyticsConfig.Execute(TEXT("ProviderModuleName"), true)),
 		DefaultEngineAnalyticsConfig);
 
+	// add a unique id
+	InstanceId = FGuid::NewGuid();
+	FPlatformMisc::CreateGuid(InstanceId);
+
 	// check if Configs override the heartbeat interval
 	
 	float ConfigHeartbeatInterval = 0.0;
@@ -104,6 +110,11 @@ void FQoSReporter::Shutdown()
 	bIsInitialized = false;
 }
 
+FGuid FQoSReporter::GetQoSReporterInstanceId()
+{
+	return InstanceId;
+}
+
 double FQoSReporter::ModuleInitializationTime = FPlatformTime::Seconds();
 bool FQoSReporter::bStartupEventReported = false;
 
@@ -126,6 +137,7 @@ void FQoSReporter::ReportStartupCompleteEvent()
 
 double FQoSReporter::HeartbeatInterval = 300;
 double FQoSReporter::LastHeartbeatTimestamp = 0;
+extern ENGINE_API float GAverageFPS;
 
 void FQoSReporter::Tick()
 {
@@ -134,12 +146,30 @@ void FQoSReporter::Tick()
 		return;
 	}
 
+	static double PreviousTickTime = FPlatformTime::Seconds();
 	double CurrentTime = FPlatformTime::Seconds();
+
 	if (HeartbeatInterval > 0 && CurrentTime - LastHeartbeatTimestamp > HeartbeatInterval)
 	{
 		SendHeartbeat();
 		LastHeartbeatTimestamp = CurrentTime;
 	}
+
+	// detect too long pauses between ticks
+	const double kTooLongPauseBetweenTicksInSeconds = 0.25f;	// 4 fps	- correct perfcounter name below if changing this value and make sure analytics is on the same page
+	const double DeltaBetweenTicks = CurrentTime - PreviousTickTime;
+	static int TimesHappened = 0;
+	if (DeltaBetweenTicks > kTooLongPauseBetweenTicksInSeconds)
+	{
+		++TimesHappened;
+
+		PerfCountersIncrement(TEXT("HitchesAbove250msec"));
+
+		UE_LOG(LogQoSReporter, Warning, TEXT("QoS reporter could not tick for %f sec (longer than threshold of %f sec) - happened %d time(s), average FPS is %f. Sending heartbeats might have been affected"),
+			DeltaBetweenTicks, kTooLongPauseBetweenTicksInSeconds, TimesHappened, GAverageFPS);
+	}
+
+	PreviousTickTime = CurrentTime;
 }
 
 void FQoSReporter::SendHeartbeat()
@@ -148,6 +178,8 @@ void FQoSReporter::SendHeartbeat()
 
 	TArray<FAnalyticsEventAttribute> ParamArray;
 	ParamArray.Add(FAnalyticsEventAttribute(TEXT("Role"), GetApplicationRole()));
+	ParamArray.Add(FAnalyticsEventAttribute(TEXT("SystemId"), FPlatformMisc::GetOperatingSystemId()));
+	ParamArray.Add(FAnalyticsEventAttribute(TEXT("InstanceId"), InstanceId.ToString()));
 
 	if (IsRunningDedicatedServer())
 	{
@@ -170,14 +202,20 @@ void FQoSReporter::AddServerHeartbeatAttributes(TArray<FAnalyticsEventAttribute>
 	IPerfCounters* PerfCounters = IPerfCountersModule::Get().GetPerformanceCounters();
 	if (PerfCounters)
 	{
-		FString PerfCountersJSON = PerfCounters->GetAllCountersAsJson();
-
 		const TMap<FString, IPerfCounters::FJsonVariant>& PerfCounterMap = PerfCounters->GetAllCounters();
 		int NumAddedCounters = 0;
+		// we want to skip reporting average FPS if there are no human players
+		bool bShouldSkipAverageFPS = (PerfCounters->Get("NumClients", 0) == 0);
 
 		// only add string and number values
 		for (const auto& It : PerfCounterMap)
 		{
+			// do not report average FPS unless there are actual people playing
+			if (bShouldSkipAverageFPS && It.Key == TEXT("AverageFPS"))
+			{
+				continue;
+			}
+
 			const IPerfCounters::FJsonVariant& JsonValue = It.Value;
 			switch (JsonValue.Format)
 			{
@@ -232,7 +270,6 @@ void FQoSReporter::AddServerHeartbeatAttributes(TArray<FAnalyticsEventAttribute>
 
 void FQoSReporter::AddClientHeartbeatAttributes(TArray<FAnalyticsEventAttribute> & OutArray)
 {
-	extern ENGINE_API float GAverageFPS;
 	OutArray.Add(FAnalyticsEventAttribute(TEXT("AverageFPS"), GAverageFPS));
 }
 
diff --git a/Engine/Source/Runtime/Analytics/QoSReporter/Public/QoSReporter.h b/Engine/Source/Runtime/Analytics/QoSReporter/Public/QoSReporter.h
--- a/Engine/Source/Runtime/Analytics/QoSReporter/Public/QoSReporter.h
+++ b/Engine/Source/Runtime/Analytics/QoSReporter/Public/QoSReporter.h
@@ -157,6 +157,11 @@ public:
 	 */
 	static void QOSREPORTER_API Tick();
 
+	/**
+	 * This function is expected to be called periodically to update ongoing tasks.
+	 */
+	static FGuid QOSREPORTER_API GetQoSReporterInstanceId();
+
 private:
 
 	/**
@@ -182,6 +187,9 @@ private:
 	/** Whether the module has been initialized. */
 	static bool bIsInitialized;
 
+	/** Unique identifier for this QoS reporter instance (only changed on module initialization) */
+	static FGuid InstanceId;
+
 	/** Chosen analytics provider. */
 	static TSharedPtr<IAnalyticsProvider> Analytics;
 
diff --git a/Engine/Source/Runtime/AnimGraphRuntime/Private/AnimNodes/AnimNode_ApplyAdditive.cpp b/Engine/Source/Runtime/AnimGraphRuntime/Private/AnimNodes/AnimNode_ApplyAdditive.cpp
--- a/Engine/Source/Runtime/AnimGraphRuntime/Private/AnimNodes/AnimNode_ApplyAdditive.cpp
+++ b/Engine/Source/Runtime/AnimGraphRuntime/Private/AnimNodes/AnimNode_ApplyAdditive.cpp
@@ -22,29 +22,41 @@ void FAnimNode_ApplyAdditive::CacheBones(const FAnimationCacheBonesContext& Cont
 
 void FAnimNode_ApplyAdditive::Update(const FAnimationUpdateContext& Context)
 {
-	EvaluateGraphExposedInputs.Execute(Context);
-
 	Base.Update(Context);
-	const float ActualAlpha = AlphaScaleBias.ApplyTo(Alpha);
-	if (ActualAlpha > ZERO_ANIMWEIGHT_THRESH)
+
+	if (IsLODEnabled(Context.AnimInstance, LODThreshold))
 	{
-		Additive.Update(Context.FractionalWeight(ActualAlpha));
+		// @note: If you derive from this class, and if you have input that you rely on for base
+		// this is not going to work	
+		EvaluateGraphExposedInputs.Execute(Context);
+		const float ActualAlpha = AlphaScaleBias.ApplyTo(Alpha);
+		if (ActualAlpha > ZERO_ANIMWEIGHT_THRESH)
+		{
+			Additive.Update(Context.FractionalWeight(ActualAlpha));
+		}
 	}
 }
 
 void FAnimNode_ApplyAdditive::Evaluate(FPoseContext& Output)
 {
-	//@TODO: Could evaluate Base into Output and save a copy
-	const float ActualAlpha = AlphaScaleBias.ApplyTo(Alpha);
-	if (ActualAlpha > ZERO_ANIMWEIGHT_THRESH)
+	if (IsLODEnabled(Output.AnimInstance, LODThreshold))
 	{
-		FPoseContext AdditiveEvalContext(Output);
+		//@TODO: Could evaluate Base into Output and save a copy
+		const float ActualAlpha = AlphaScaleBias.ApplyTo(Alpha);
+		if (ActualAlpha > ZERO_ANIMWEIGHT_THRESH)
+		{
+			FPoseContext AdditiveEvalContext(Output);
 
-		Base.Evaluate(Output);
-		Additive.Evaluate(AdditiveEvalContext);
+			Base.Evaluate(Output);
+			Additive.Evaluate(AdditiveEvalContext);
 
-		FAnimationRuntime::AccumulateAdditivePose(Output.Pose, AdditiveEvalContext.Pose, Output.Curve, AdditiveEvalContext.Curve, ActualAlpha, AAT_LocalSpaceBase);
-		Output.Pose.NormalizeRotations();
+			FAnimationRuntime::AccumulateAdditivePose(Output.Pose, AdditiveEvalContext.Pose, Output.Curve, AdditiveEvalContext.Curve, ActualAlpha, AAT_LocalSpaceBase);
+			Output.Pose.NormalizeRotations();
+		}
+		else
+		{
+			Base.Evaluate(Output);
+		}
 	}
 	else
 	{
@@ -54,6 +66,7 @@ void FAnimNode_ApplyAdditive::Evaluate(FPoseContext& Output)
 
 FAnimNode_ApplyAdditive::FAnimNode_ApplyAdditive()
 	: Alpha(1.0f)
+	, LODThreshold(INDEX_NONE)
 {
 }
 
diff --git a/Engine/Source/Runtime/AnimGraphRuntime/Private/AnimNodes/AnimNode_RotationOffsetBlendSpace.cpp b/Engine/Source/Runtime/AnimGraphRuntime/Private/AnimNodes/AnimNode_RotationOffsetBlendSpace.cpp
--- a/Engine/Source/Runtime/AnimGraphRuntime/Private/AnimNodes/AnimNode_RotationOffsetBlendSpace.cpp
+++ b/Engine/Source/Runtime/AnimGraphRuntime/Private/AnimNodes/AnimNode_RotationOffsetBlendSpace.cpp
@@ -20,7 +20,11 @@ void FAnimNode_RotationOffsetBlendSpace::CacheBones(const FAnimationCacheBonesCo
 
 void FAnimNode_RotationOffsetBlendSpace::UpdateAssetPlayer(const FAnimationUpdateContext& Context)
 {
-	FAnimNode_BlendSpacePlayer::UpdateAssetPlayer(Context);
+	if (IsLODEnabled(Context.AnimInstance, LODThreshold))
+	{
+		FAnimNode_BlendSpacePlayer::UpdateAssetPlayer(Context);
+	}
+
 	BasePose.Update(Context);
 }
 
@@ -29,15 +33,18 @@ void FAnimNode_RotationOffsetBlendSpace::Evaluate(FPoseContext& Context)
 	// Evaluate base pose
 	BasePose.Evaluate(Context);
 
-	// Evaluate MeshSpaceRotation additive blendspace
-	FPoseContext MeshSpaceRotationAdditivePoseContext(Context);
-	FAnimNode_BlendSpacePlayer::Evaluate(MeshSpaceRotationAdditivePoseContext);
+	if (IsLODEnabled(Context.AnimInstance, LODThreshold))
+	{
+		// Evaluate MeshSpaceRotation additive blendspace
+		FPoseContext MeshSpaceRotationAdditivePoseContext(Context);
+		FAnimNode_BlendSpacePlayer::Evaluate(MeshSpaceRotationAdditivePoseContext);
 
-	// Accumulate poses together
-	FAnimationRuntime::AccumulateMeshSpaceRotationAdditiveToLocalPose(Context.Pose, MeshSpaceRotationAdditivePoseContext.Pose, Context.Curve, MeshSpaceRotationAdditivePoseContext.Curve, 1.f);
+		// Accumulate poses together
+		FAnimationRuntime::AccumulateMeshSpaceRotationAdditiveToLocalPose(Context.Pose, MeshSpaceRotationAdditivePoseContext.Pose, Context.Curve, MeshSpaceRotationAdditivePoseContext.Curve, 1.f);
 
-	// Resulting rotations are not normalized, so normalize here.
-	Context.Pose.NormalizeRotations();
+		// Resulting rotations are not normalized, so normalize here.
+		Context.Pose.NormalizeRotations();
+	}
 }
 
 void FAnimNode_RotationOffsetBlendSpace::GatherDebugData(FNodeDebugData& DebugData)
@@ -51,5 +58,6 @@ void FAnimNode_RotationOffsetBlendSpace::GatherDebugData(FNodeDebugData& DebugDa
 }
 
 FAnimNode_RotationOffsetBlendSpace::FAnimNode_RotationOffsetBlendSpace()
+	: LODThreshold(INDEX_NONE)
 {
 }
diff --git a/Engine/Source/Runtime/AnimGraphRuntime/Private/BoneControllers/AnimNode_AnimDynamics.cpp b/Engine/Source/Runtime/AnimGraphRuntime/Private/BoneControllers/AnimNode_AnimDynamics.cpp
--- a/Engine/Source/Runtime/AnimGraphRuntime/Private/BoneControllers/AnimNode_AnimDynamics.cpp
+++ b/Engine/Source/Runtime/AnimGraphRuntime/Private/BoneControllers/AnimNode_AnimDynamics.cpp
@@ -66,9 +66,9 @@ void FAnimNode_AnimDynamics::Initialize(const FAnimationInitializeContext& Conte
 	TimeDebt = 0.0f;
 }
 
-void FAnimNode_AnimDynamics::Update(const FAnimationUpdateContext& Context)
+void FAnimNode_AnimDynamics::UpdateInternal(const FAnimationUpdateContext& Context)
 {
-	FAnimNode_SkeletalControlBase::Update(Context);
+	FAnimNode_SkeletalControlBase::UpdateInternal(Context);
 
 	NextTimeStep = Context.GetDeltaTime();
 }
diff --git a/Engine/Source/Runtime/AnimGraphRuntime/Private/BoneControllers/AnimNode_LookAt.cpp b/Engine/Source/Runtime/AnimGraphRuntime/Private/BoneControllers/AnimNode_LookAt.cpp
--- a/Engine/Source/Runtime/AnimGraphRuntime/Private/BoneControllers/AnimNode_LookAt.cpp
+++ b/Engine/Source/Runtime/AnimGraphRuntime/Private/BoneControllers/AnimNode_LookAt.cpp
@@ -197,9 +197,9 @@ void FAnimNode_LookAt::InitializeBoneReferences(const FBoneContainer& RequiredBo
 	LookAtBone.Initialize(RequiredBones);
 }
 
-void FAnimNode_LookAt::Update(const FAnimationUpdateContext& Context)
+void FAnimNode_LookAt::UpdateInternal(const FAnimationUpdateContext& Context)
 {
-	FAnimNode_SkeletalControlBase::Update(Context);
+	FAnimNode_SkeletalControlBase::UpdateInternal(Context);
 
 	AccumulatedInterpoolationTime = FMath::Clamp(AccumulatedInterpoolationTime+Context.GetDeltaTime(), 0.f, InterpolationTime);;
 }
diff --git a/Engine/Source/Runtime/AnimGraphRuntime/Private/BoneControllers/AnimNode_SkeletalControlBase.cpp b/Engine/Source/Runtime/AnimGraphRuntime/Private/BoneControllers/AnimNode_SkeletalControlBase.cpp
--- a/Engine/Source/Runtime/AnimGraphRuntime/Private/BoneControllers/AnimNode_SkeletalControlBase.cpp
+++ b/Engine/Source/Runtime/AnimGraphRuntime/Private/BoneControllers/AnimNode_SkeletalControlBase.cpp
@@ -21,10 +21,24 @@ void FAnimNode_SkeletalControlBase::CacheBones(const FAnimationCacheBonesContext
 	ComponentPose.CacheBones(Context);
 }
 
+void FAnimNode_SkeletalControlBase::UpdateInternal(const FAnimationUpdateContext& Context)
+{
+	EvaluateGraphExposedInputs.Execute(Context);
+}
+
 void FAnimNode_SkeletalControlBase::Update(const FAnimationUpdateContext& Context)
 {
 	ComponentPose.Update(Context);
-	EvaluateGraphExposedInputs.Execute(Context);
+
+	if (IsLODEnabled(Context.AnimInstance, LODThreshold))
+	{
+		// Apply the skeletal control if it's valid
+		const float ActualAlpha = AlphaScaleBias.ApplyTo(Alpha);
+		if ((ActualAlpha >= ZERO_ANIMWEIGHT_THRESH) && IsValidToEvaluate(Context.AnimInstance->CurrentSkeleton, Context.AnimInstance->RequiredBones))
+		{
+			UpdateInternal(Context);
+		}
+	}
 }
 
 bool ContainsNaN(const TArray<FBoneTransform> & BoneTransforms)
@@ -45,26 +59,29 @@ void FAnimNode_SkeletalControlBase::EvaluateComponentSpace(FComponentSpacePoseCo
 	// Evaluate the input
 	ComponentPose.EvaluateComponentSpace(Output);
 
-	// Apply the skeletal control if it's valid
-	const float ActualAlpha = AlphaScaleBias.ApplyTo(Alpha);
-	if ((ActualAlpha >= ZERO_ANIMWEIGHT_THRESH) && IsValidToEvaluate(Output.AnimInstance->CurrentSkeleton, Output.AnimInstance->RequiredBones))
+	if (IsLODEnabled(Output.AnimInstance, LODThreshold))
 	{
-		USkeletalMeshComponent* Component = Output.AnimInstance->GetSkelMeshComponent();
+		// Apply the skeletal control if it's valid
+		const float ActualAlpha = AlphaScaleBias.ApplyTo(Alpha);
+		if ((ActualAlpha >= ZERO_ANIMWEIGHT_THRESH) && IsValidToEvaluate(Output.AnimInstance->CurrentSkeleton, Output.AnimInstance->RequiredBones))
+		{
+			USkeletalMeshComponent* Component = Output.AnimInstance->GetSkelMeshComponent();
 
 #if WITH_EDITORONLY_DATA
-		// save current pose before applying skeletal control to compute the exact gizmo location in AnimGraphNode
-		ForwardedPose = Output.Pose;
+			// save current pose before applying skeletal control to compute the exact gizmo location in AnimGraphNode
+			ForwardedPose = Output.Pose;
 #endif // #if WITH_EDITORONLY_DATA
 
-		BoneTransforms.Reset(BoneTransforms.Num());
-		EvaluateBoneTransforms(Component, Output.Pose, BoneTransforms);
+			BoneTransforms.Reset(BoneTransforms.Num());
+			EvaluateBoneTransforms(Component, Output.Pose, BoneTransforms);
 
-		checkSlow(!ContainsNaN(BoneTransforms));
+			checkSlow(!ContainsNaN(BoneTransforms));
 
-		if (BoneTransforms.Num() > 0)
-		{
-			const float BlendWeight = FMath::Clamp<float>(ActualAlpha, 0.f, 1.f);
-			Output.Pose.LocalBlendCSBoneTransforms(BoneTransforms, BlendWeight);
+			if (BoneTransforms.Num() > 0)
+			{
+				const float BlendWeight = FMath::Clamp<float>(ActualAlpha, 0.f, 1.f);
+				Output.Pose.LocalBlendCSBoneTransforms(BoneTransforms, BlendWeight);
+			}
 		}
 	}
 }
diff --git a/Engine/Source/Runtime/AnimGraphRuntime/Private/BoneControllers/AnimNode_SpringBone.cpp b/Engine/Source/Runtime/AnimGraphRuntime/Private/BoneControllers/AnimNode_SpringBone.cpp
--- a/Engine/Source/Runtime/AnimGraphRuntime/Private/BoneControllers/AnimNode_SpringBone.cpp
+++ b/Engine/Source/Runtime/AnimGraphRuntime/Private/BoneControllers/AnimNode_SpringBone.cpp
@@ -38,9 +38,9 @@ void FAnimNode_SpringBone::CacheBones(const FAnimationCacheBonesContext& Context
 	FAnimNode_SkeletalControlBase::CacheBones(Context);
 }
 
-void FAnimNode_SpringBone::Update(const FAnimationUpdateContext& Context)
+void FAnimNode_SpringBone::UpdateInternal(const FAnimationUpdateContext& Context)
 {
-	FAnimNode_SkeletalControlBase::Update(Context);
+	FAnimNode_SkeletalControlBase::UpdateInternal(Context);
 
 	RemainingTime += Context.GetDeltaTime();
 
diff --git a/Engine/Source/Runtime/AnimGraphRuntime/Private/BoneControllers/AnimNode_Trail.cpp b/Engine/Source/Runtime/AnimGraphRuntime/Private/BoneControllers/AnimNode_Trail.cpp
--- a/Engine/Source/Runtime/AnimGraphRuntime/Private/BoneControllers/AnimNode_Trail.cpp
+++ b/Engine/Source/Runtime/AnimGraphRuntime/Private/BoneControllers/AnimNode_Trail.cpp
@@ -22,9 +22,9 @@ FAnimNode_Trail::FAnimNode_Trail()
 	TrailRelaxRichCurve->AddKey(1.f, 5.f);
 }
 
-void FAnimNode_Trail::Update(const FAnimationUpdateContext& Context)
+void FAnimNode_Trail::UpdateInternal(const FAnimationUpdateContext& Context)
 {
-	FAnimNode_SkeletalControlBase::Update(Context);
+	FAnimNode_SkeletalControlBase::UpdateInternal(Context);
 
 	ThisTimstep = Context.GetDeltaTime();
 }
diff --git a/Engine/Source/Runtime/AnimGraphRuntime/Private/BoneControllers/AnimNode_TwoBoneIK.cpp b/Engine/Source/Runtime/AnimGraphRuntime/Private/BoneControllers/AnimNode_TwoBoneIK.cpp
--- a/Engine/Source/Runtime/AnimGraphRuntime/Private/BoneControllers/AnimNode_TwoBoneIK.cpp
+++ b/Engine/Source/Runtime/AnimGraphRuntime/Private/BoneControllers/AnimNode_TwoBoneIK.cpp
@@ -4,6 +4,9 @@
 #include "BoneControllers/AnimNode_TwoBoneIK.h"
 #include "AnimationRuntime.h"
 
+DECLARE_CYCLE_STAT(TEXT("TwoBoneIK Eval"), STAT_TwoBoneIK_Eval, STATGROUP_Anim);
+
+
 /////////////////////////////////////////////////////
 // FAnimNode_TwoBoneIK
 
@@ -33,6 +36,8 @@ void FAnimNode_TwoBoneIK::GatherDebugData(FNodeDebugData& DebugData)
 
 void FAnimNode_TwoBoneIK::EvaluateBoneTransforms(USkeletalMeshComponent* SkelComp, FCSPose<FCompactPose>& MeshBases, TArray<FBoneTransform>& OutBoneTransforms)
 {
+	SCOPE_CYCLE_COUNTER(STAT_TwoBoneIK_Eval);
+
 	check(OutBoneTransforms.Num() == 0);
 
 	const FBoneContainer& BoneContainer = MeshBases.GetPose().GetBoneContainer();
diff --git a/Engine/Source/Runtime/AnimGraphRuntime/Private/BoneControllers/AnimNode_WheelHandler.cpp b/Engine/Source/Runtime/AnimGraphRuntime/Private/BoneControllers/AnimNode_WheelHandler.cpp
--- a/Engine/Source/Runtime/AnimGraphRuntime/Private/BoneControllers/AnimNode_WheelHandler.cpp
+++ b/Engine/Source/Runtime/AnimGraphRuntime/Private/BoneControllers/AnimNode_WheelHandler.cpp
@@ -102,7 +102,7 @@ void FAnimNode_WheelHandler::InitializeBoneReferences(const FBoneContainer& Requ
 	WheelSimulators.Sort([](FWheelSimulator L, FWheelSimulator R) { return L.BoneReference.BoneIndex < R.BoneReference.BoneIndex; });
 }
 
-void FAnimNode_WheelHandler::Update(const FAnimationUpdateContext& Context)
+void FAnimNode_WheelHandler::UpdateInternal(const FAnimationUpdateContext& Context)
 {
 	if(VehicleSimComponent)
 	{
@@ -125,7 +125,7 @@ void FAnimNode_WheelHandler::Update(const FAnimationUpdateContext& Context)
 		}
 	}
 
-	FAnimNode_SkeletalControlBase::Update(Context);
+	FAnimNode_SkeletalControlBase::UpdateInternal(Context);
 }
 
 void FAnimNode_WheelHandler::Initialize(const FAnimationInitializeContext& Context)
diff --git a/Engine/Source/Runtime/AnimGraphRuntime/Public/AnimNodes/AnimNode_ApplyAdditive.h b/Engine/Source/Runtime/AnimGraphRuntime/Public/AnimNodes/AnimNode_ApplyAdditive.h
--- a/Engine/Source/Runtime/AnimGraphRuntime/Public/AnimNodes/AnimNode_ApplyAdditive.h
+++ b/Engine/Source/Runtime/AnimGraphRuntime/Public/AnimNodes/AnimNode_ApplyAdditive.h
@@ -23,6 +23,15 @@ struct ANIMGRAPHRUNTIME_API FAnimNode_ApplyAdditive : public FAnimNode_Base
 	UPROPERTY(EditAnywhere, BlueprintReadWrite, Category=Settings)
 	FInputScaleBias AlphaScaleBias;
 
+	/* 
+	 * Max LOD that this node is allowed to run
+	 * For example if you have LODThreadhold to be 2, it will run until LOD 2 (based on 0 index)
+	 * when the component LOD becomes 3, it will stop update/evaluate
+	 * currently transition would be issue and that has to be re-visited
+	 */
+	UPROPERTY(EditAnywhere, BlueprintReadWrite, Category = Performance, meta=(DisplayName="LOD Threshold"))
+	int32 LODThreshold;
+
 public:	
 	FAnimNode_ApplyAdditive();
 
@@ -33,4 +42,5 @@ public:
 	virtual void Evaluate(FPoseContext& Output) override;
 	virtual void GatherDebugData(FNodeDebugData& DebugData) override;
 	// End of FAnimNode_Base interface
+
 };
diff --git a/Engine/Source/Runtime/AnimGraphRuntime/Public/AnimNodes/AnimNode_RotationOffsetBlendSpace.h b/Engine/Source/Runtime/AnimGraphRuntime/Public/AnimNodes/AnimNode_RotationOffsetBlendSpace.h
--- a/Engine/Source/Runtime/AnimGraphRuntime/Public/AnimNodes/AnimNode_RotationOffsetBlendSpace.h
+++ b/Engine/Source/Runtime/AnimGraphRuntime/Public/AnimNodes/AnimNode_RotationOffsetBlendSpace.h
@@ -14,6 +14,14 @@ struct ANIMGRAPHRUNTIME_API FAnimNode_RotationOffsetBlendSpace : public FAnimNod
 	UPROPERTY(EditAnywhere, BlueprintReadWrite, Category=Links)
 	FPoseLink BasePose;
 
+	/*
+	* Max LOD that this node is allowed to run
+	* For example if you have LODThreadhold to be 2, it will run until LOD 2 (based on 0 index)
+	* when the component LOD becomes 3, it will stop update/evaluate
+	* currently transition would be issue and that has to be re-visited
+	*/
+	UPROPERTY(EditAnywhere, BlueprintReadWrite, Category = Performance, meta = (DisplayName = "LOD Threshold"))
+	int32 LODThreshold;
 public:	
 	FAnimNode_RotationOffsetBlendSpace();
 
diff --git a/Engine/Source/Runtime/AnimGraphRuntime/Public/BoneControllers/AnimNode_AnimDynamics.h b/Engine/Source/Runtime/AnimGraphRuntime/Public/BoneControllers/AnimNode_AnimDynamics.h
--- a/Engine/Source/Runtime/AnimGraphRuntime/Public/BoneControllers/AnimNode_AnimDynamics.h
+++ b/Engine/Source/Runtime/AnimGraphRuntime/Public/BoneControllers/AnimNode_AnimDynamics.h
@@ -256,7 +256,7 @@ struct ANIMGRAPHRUNTIME_API FAnimNode_AnimDynamics : public FAnimNode_SkeletalCo
 
 	// FAnimNode_SkeletalControlBase interface
 	virtual void Initialize(const FAnimationInitializeContext& Context) override;
-	virtual void Update(const FAnimationUpdateContext& Context) override;
+	virtual void UpdateInternal(const FAnimationUpdateContext& Context) override;
 	virtual void EvaluateBoneTransforms(USkeletalMeshComponent* SkelComp, FCSPose<FCompactPose>& MeshBases, TArray<FBoneTransform>& OutBoneTransforms) override;
 	virtual void GatherDebugData(FNodeDebugData& DebugData) override;
 	// End of FAnimNode_SkeletalControlBase interface
diff --git a/Engine/Source/Runtime/AnimGraphRuntime/Public/BoneControllers/AnimNode_LookAt.h b/Engine/Source/Runtime/AnimGraphRuntime/Public/BoneControllers/AnimNode_LookAt.h
--- a/Engine/Source/Runtime/AnimGraphRuntime/Public/BoneControllers/AnimNode_LookAt.h
+++ b/Engine/Source/Runtime/AnimGraphRuntime/Public/BoneControllers/AnimNode_LookAt.h
@@ -100,7 +100,7 @@ struct ANIMGRAPHRUNTIME_API FAnimNode_LookAt : public FAnimNode_SkeletalControlB
 
 	// FAnimNode_Base interface
 	virtual void GatherDebugData(FNodeDebugData& DebugData) override;
-	virtual void Update(const FAnimationUpdateContext& Context) override;
+	virtual void UpdateInternal(const FAnimationUpdateContext& Context) override;
 	// End of FAnimNode_Base interface
 
 	// FAnimNode_SkeletalControlBase interface
diff --git a/Engine/Source/Runtime/AnimGraphRuntime/Public/BoneControllers/AnimNode_SkeletalControlBase.h b/Engine/Source/Runtime/AnimGraphRuntime/Public/BoneControllers/AnimNode_SkeletalControlBase.h
--- a/Engine/Source/Runtime/AnimGraphRuntime/Public/BoneControllers/AnimNode_SkeletalControlBase.h
+++ b/Engine/Source/Runtime/AnimGraphRuntime/Public/BoneControllers/AnimNode_SkeletalControlBase.h
@@ -28,10 +28,20 @@ struct ANIMGRAPHRUNTIME_API FAnimNode_SkeletalControlBase : public FAnimNode_Bas
 	UPROPERTY(EditAnywhere, BlueprintReadWrite, Category=Settings)
 	FInputScaleBias AlphaScaleBias;
 
+	/*
+	* Max LOD that this node is allowed to run
+	* For example if you have LODThreadhold to be 2, it will run until LOD 2 (based on 0 index)
+	* when the component LOD becomes 3, it will stop update/evaluate
+	* currently transition would be issue and that has to be re-visited
+	*/
+	UPROPERTY(EditAnywhere, BlueprintReadWrite, Category = Performance, meta = (DisplayName = "LOD Threshold"))
+	int32 LODThreshold;
+
 public:
 
 	FAnimNode_SkeletalControlBase()
 		: Alpha(1.0f)
+		, LODThreshold(INDEX_NONE)
 	{
 	}
 
@@ -44,13 +54,14 @@ public:
 	// FAnimNode_Base interface
 	virtual void Initialize(const FAnimationInitializeContext& Context) override;
 	virtual void CacheBones(const FAnimationCacheBonesContext& Context)  override;
-	virtual void Update(const FAnimationUpdateContext& Context) override;
-	virtual void EvaluateComponentSpace(FComponentSpacePoseContext& Output) override;
+	virtual void Update(const FAnimationUpdateContext& Context) final;
+	virtual void EvaluateComponentSpace(FComponentSpacePoseContext& Output) final;
 	// End of FAnimNode_Base interface
 	
 protected:
 	// Interface for derived skeletal controls to implement
-
+	// use this function to update for skeletal control base
+	virtual void UpdateInternal(const FAnimationUpdateContext& Context);
 	// Evaluate the new component-space transforms for the affected bones.
 	virtual void EvaluateBoneTransforms(USkeletalMeshComponent* SkelComp, FCSPose<FCompactPose>& MeshBases, TArray<FBoneTransform>& OutBoneTransforms) {}
 	// return true if it is valid to Evaluate
@@ -60,7 +71,6 @@ protected:
 
 	/** Allow base to add info to the node debug output */
 	void AddDebugNodeData(FString& OutDebugData);
-
 private:
 
 	// Resused bone transform array to avoid reallocating in skeletal controls
diff --git a/Engine/Source/Runtime/Core/Public/Misc/Build.h b/Engine/Source/Runtime/Core/Public/Misc/Build.h
--- a/Engine/Source/Runtime/Core/Public/Misc/Build.h
+++ b/Engine/Source/Runtime/Core/Public/Misc/Build.h
@@ -243,7 +243,7 @@
 #define UE_BLUEPRINT_EVENTGRAPH_FASTCALLS 1
 
 /** Enable perf counters on dedicated servers */
-#define USE_SERVER_PERF_COUNTERS (UE_SERVER || UE_EDITOR) && WITH_PERFCOUNTERS
+#define USE_SERVER_PERF_COUNTERS ((UE_SERVER || UE_EDITOR) && WITH_PERFCOUNTERS)
 
 #define USE_CIRCULAR_DEPENDENCY_LOAD_DEFERRING 1
 #define USE_DEFERRED_DEPENDENCY_CHECK_VERIFICATION_TESTS (USE_CIRCULAR_DEPENDENCY_LOAD_DEFERRING && 0)
diff --git a/Engine/Source/Runtime/Engine/Classes/Animation/AnimInstance.h b/Engine/Source/Runtime/Engine/Classes/Animation/AnimInstance.h
--- a/Engine/Source/Runtime/Engine/Classes/Animation/AnimInstance.h
+++ b/Engine/Source/Runtime/Engine/Classes/Animation/AnimInstance.h
@@ -646,6 +646,9 @@ public:
 	/** Get next sectionID for given section ID */
 	int32 Montage_GetNextSectionID(UAnimMontage const * const Montage, int32 const & CurrentSectionID) const;
 
+	/** Returns true if any montage is playing currently. Doesn't mean it's active though, it could be blending out. */
+	bool IsAnyMontagePlaying() const;
+
 	/** Get a current Active Montage in this AnimInstance. 
 		Note that there might be multiple Active at the same time. This will only return the first active one it finds. **/
 	UAnimMontage * GetCurrentActiveMontage();
diff --git a/Engine/Source/Runtime/Engine/Classes/Animation/AnimNodeBase.h b/Engine/Source/Runtime/Engine/Classes/Animation/AnimNodeBase.h
--- a/Engine/Source/Runtime/Engine/Classes/Animation/AnimNodeBase.h
+++ b/Engine/Source/Runtime/Engine/Classes/Animation/AnimNodeBase.h
@@ -517,4 +517,11 @@ struct ENGINE_API FAnimNode_Base
 	// End of interface to implement
 
 	virtual ~FAnimNode_Base() {}
+
+protected:
+	/** return true if eanbled, otherwise, return false. This is utility function that can be used per node level */
+	bool IsLODEnabled(UAnimInstance* AnimInstace, int32 InLODThreshold)
+	{
+		return (InLODThreshold == INDEX_NONE || AnimInstace->GetOwningComponent()->PredictedLODLevel <= InLODThreshold);
+	}
 };
diff --git a/Engine/Source/Runtime/Engine/Classes/Animation/AnimNode_ApplyMeshSpaceAdditive.h b/Engine/Source/Runtime/Engine/Classes/Animation/AnimNode_ApplyMeshSpaceAdditive.h
--- a/Engine/Source/Runtime/Engine/Classes/Animation/AnimNode_ApplyMeshSpaceAdditive.h
+++ b/Engine/Source/Runtime/Engine/Classes/Animation/AnimNode_ApplyMeshSpaceAdditive.h
@@ -23,6 +23,15 @@ struct ENGINE_API FAnimNode_ApplyMeshSpaceAdditive : public FAnimNode_Base
 	UPROPERTY(EditAnywhere, BlueprintReadWrite, Category = Settings)
 	FInputScaleBias AlphaScaleBias;
 
+	/*
+	* Max LOD that this node is allowed to run
+	* For example if you have LODThreadhold to be 2, it will run until LOD 2 (based on 0 index)
+	* when the component LOD becomes 3, it will stop update/evaluate
+	* currently transition would be issue and that has to be re-visited
+	*/
+	UPROPERTY(EditAnywhere, BlueprintReadWrite, Category = Performance, meta = (DisplayName = "LOD Threshold"))
+	int32 LODThreshold;
+
 public:
 	FAnimNode_ApplyMeshSpaceAdditive();
 
diff --git a/Engine/Source/Runtime/Engine/Classes/Animation/AnimStateMachineTypes.h b/Engine/Source/Runtime/Engine/Classes/Animation/AnimStateMachineTypes.h
--- a/Engine/Source/Runtime/Engine/Classes/Animation/AnimStateMachineTypes.h
+++ b/Engine/Source/Runtime/Engine/Classes/Animation/AnimStateMachineTypes.h
@@ -270,6 +270,9 @@ struct FBakedAnimationStateMachine
 	UPROPERTY()
 	TArray<FAnimationTransitionBetweenStates> Transitions;
 
+	// Cached StatID for this state machine
+	STAT(mutable TStatId StatID;)
+
 public:
 	FBakedAnimationStateMachine()
 		: InitialState(INDEX_NONE)
@@ -281,6 +284,18 @@ public:
 	// Find the index of a transition from StateNameFrom to StateNameTo
 	ENGINE_API int32 FindTransitionIndex(const FName& InStateNameFrom, const FName& InStateNameTo) const;
 	ENGINE_API int32 FindTransitionIndex(const int32 InStateIdxFrom, const int32 InStateIdxTo) const;
+
+#if STATS
+	/** Get the StatID for timing this state machine */
+	FORCEINLINE TStatId GetStatID() const
+	{
+		if (!StatID.IsValidStat())
+		{
+			StatID = FDynamicStats::CreateStatId<FStatGroup_STATGROUP_Anim>(MachineName);
+		}
+		return StatID;
+	}
+#endif // STATS
 };
 
 UCLASS()
diff --git a/Engine/Source/Runtime/Engine/Classes/Animation/AnimationSettings.h b/Engine/Source/Runtime/Engine/Classes/Animation/AnimationSettings.h
--- a/Engine/Source/Runtime/Engine/Classes/Animation/AnimationSettings.h
+++ b/Engine/Source/Runtime/Engine/Classes/Animation/AnimationSettings.h
@@ -79,6 +79,9 @@ class ENGINE_API UAnimationSettings : public UDeveloperSettings
 	UPROPERTY(config, EditAnywhere, Category = Compression)
 	bool bTryIntervalKeyRemoval;
 
+	UPROPERTY(config, EditAnywhere, Category = Performance)
+	bool bEnablePerformanceLog;
+
 public:
 	static UAnimationSettings * Get() { return CastChecked<UAnimationSettings>(UAnimationSettings::StaticClass()->GetDefaultObject()); }
 
diff --git a/Engine/Source/Runtime/Engine/Classes/Engine/DemoNetDriver.h b/Engine/Source/Runtime/Engine/Classes/Engine/DemoNetDriver.h
--- a/Engine/Source/Runtime/Engine/Classes/Engine/DemoNetDriver.h
+++ b/Engine/Source/Runtime/Engine/Classes/Engine/DemoNetDriver.h
@@ -100,9 +100,6 @@ private:
 
 	TArray<FNetworkGUID> NonQueuedGUIDsForScrubbing;
 
-	// If a channel is associated with Actor, adds the channel's GUID to the list of GUIDs excluded from queuing bunches during scrubbing.
-	void		AddNonQueuedActorForScrubbing(AActor* Actor);
-
 	// Replay tasks
 	TArray< TSharedPtr< FQueuedReplayTask > >	QueuedReplayTasks;
 	TSharedPtr< FQueuedReplayTask >				ActiveReplayTask;
@@ -191,6 +188,9 @@ public:
 	bool ProcessReplayTasks();
 	bool IsNamedTaskInQueue( const FString& Name );
 
+	/** If a channel is associated with Actor, adds the channel's GUID to the list of GUIDs excluded from queuing bunches during scrubbing. */
+	void AddNonQueuedActorForScrubbing(AActor const* Actor);
+	/** Adds the channel's GUID to the list of GUIDs excluded from queuing bunches during scrubbing. */
 	void AddNonQueuedGUIDForScrubbing(FNetworkGUID InGUID);
 
 	virtual bool IsLevelInitializedForActor( const AActor* InActor, const UNetConnection* InConnection ) const override;
diff --git a/Engine/Source/Runtime/Engine/Classes/Engine/Engine.h b/Engine/Source/Runtime/Engine/Classes/Engine/Engine.h
--- a/Engine/Source/Runtime/Engine/Classes/Engine/Engine.h
+++ b/Engine/Source/Runtime/Engine/Classes/Engine/Engine.h
@@ -532,6 +532,8 @@ class IAnalyticsProvider;
 DECLARE_DELEGATE_OneParam(FBeginStreamingPauseDelegate, FViewport*);
 DECLARE_DELEGATE(FEndStreamingPauseDelegate);
 
+DECLARE_MULTICAST_DELEGATE_OneParam(FEngineHitchDetectedDelegate, float /*HitchDurationInSeconds*/);
+
 
 /**
  * Abstract base class of all Engine classes, responsible for management of systems critical to editor or game systems.
@@ -2008,6 +2010,9 @@ public:
 	*/
 	virtual void DumpFPSChartAnalytics(const FString& InMapName, TArray<struct FAnalyticsEventAttribute>& InParamArray);
 
+	/** Delegate called when FPS charting detects a hitch (it is not triggered if a capture isn't in progress). */
+	FEngineHitchDetectedDelegate OnHitchDetectedDelegate;
+
 private:
 
 	/**
diff --git a/Engine/Source/Runtime/Engine/Classes/Engine/EngineTypes.h b/Engine/Source/Runtime/Engine/Classes/Engine/EngineTypes.h
--- a/Engine/Source/Runtime/Engine/Classes/Engine/EngineTypes.h
+++ b/Engine/Source/Runtime/Engine/Classes/Engine/EngineTypes.h
@@ -223,6 +223,7 @@ enum EMaterialShadingModel
 	MSM_TwoSidedFoliage		UMETA(DisplayName="Two Sided"),
 	MSM_Hair				UMETA(DisplayName="Hair"),
 	MSM_Cloth				UMETA(DisplayName="Cloth"),
+	MSM_Eye					UMETA(DisplayName="Eye"),
 	MSM_MAX,
 };
 
@@ -1984,6 +1985,10 @@ public:
 	UPROPERTY()
 	TArray<float> BaseVisibleDistanceFactorThesholds;
 
+	/** Max Evaluation Rate allowed for interpolation to be enabled. Beyond, interpolation will be turned off. */
+	UPROPERTY()
+	int32 MaxEvalRateForInterpolation;
+
 public:
 
 	/** Default constructor. */
@@ -1998,6 +2003,7 @@ public:
 		, AdditionalTime(0.f)
 		, ThisTickDelta(0.f)
 		, BaseNonRenderedUpdateRate(4)
+		, MaxEvalRateForInterpolation(4)
 	{ 
 		BaseVisibleDistanceFactorThesholds.Add(0.4f);
 		BaseVisibleDistanceFactorThesholds.Add(0.2f);
diff --git a/Engine/Source/Runtime/Engine/Classes/GameFramework/Actor.h b/Engine/Source/Runtime/Engine/Classes/GameFramework/Actor.h
--- a/Engine/Source/Runtime/Engine/Classes/GameFramework/Actor.h
+++ b/Engine/Source/Runtime/Engine/Classes/GameFramework/Actor.h
@@ -218,6 +218,9 @@ protected:
 	 */
 	virtual bool HasNetOwner() const;
 
+	UFUNCTION()
+	virtual void OnRep_Owner();
+
 private:
 	/**
 	 * Describes how much control the remote machine has over the actor.
@@ -229,7 +232,7 @@ private:
 	 * Owner of this Actor, used primarily for replication (bNetUseOwnerRelevancy & bOnlyRelevantToOwner) and visibility (PrimitiveComponent bOwnerNoSee and bOnlyOwnerSee)
 	 * @see SetOwner(), GetOwner()
 	 */
-	UPROPERTY(replicated)
+	UPROPERTY(ReplicatedUsing=OnRep_Owner)
 	AActor* Owner;
 
 public:
diff --git a/Engine/Source/Runtime/Engine/Classes/Materials/Material.h b/Engine/Source/Runtime/Engine/Classes/Materials/Material.h
--- a/Engine/Source/Runtime/Engine/Classes/Materials/Material.h
+++ b/Engine/Source/Runtime/Engine/Classes/Materials/Material.h
@@ -1104,7 +1104,7 @@ public:
 	/**
 	 * Go through every material, flush the specified types and re-initialize the material's shader maps.
 	 */
-	ENGINE_API static void UpdateMaterialShaders(TArray<FShaderType*>& ShaderTypesToFlush, TArray<const FVertexFactoryType*>& VFTypesToFlush, EShaderPlatform ShaderPlatform);
+	ENGINE_API static void UpdateMaterialShaders(TArray<FShaderType*>& ShaderTypesToFlush, TArray<const FShaderPipelineType*>& ShaderPipelineTypesToFlush, TArray<const FVertexFactoryType*>& VFTypesToFlush, EShaderPlatform ShaderPlatform);
 
 	/** 
 	 * Backs up all material shaders to memory through serialization, organized by FMaterialShaderMap. 
diff --git a/Engine/Source/Runtime/Engine/Classes/Materials/MaterialExpressionDepthOfFieldFunction.h b/Engine/Source/Runtime/Engine/Classes/Materials/MaterialExpressionDepthOfFieldFunction.h
--- a/Engine/Source/Runtime/Engine/Classes/Materials/MaterialExpressionDepthOfFieldFunction.h
+++ b/Engine/Source/Runtime/Engine/Classes/Materials/MaterialExpressionDepthOfFieldFunction.h
@@ -5,6 +5,7 @@
 #include "Materials/MaterialExpression.h"
 #include "MaterialExpressionDepthOfFieldFunction.generated.h"
 
+// Note: The index is used to map the enum to different code in the shader
 UENUM()
 enum EDepthOfFieldFunctionValue
 {
@@ -14,6 +15,8 @@ enum EDepthOfFieldFunctionValue
 	TDOF_NearMask,
 	/** 0:in Focus or Near .. 1:Far. */
 	TDOF_FarMask,
+	/** in pixels, only works for CircleDOF, use Abs for the actual radius as the sign of the value indicates near out of focus, positive indicates far out of focus */
+	TDOF_CircleOfConfusionRadius,
 	TDOF_MAX,
 };
 
diff --git a/Engine/Source/Runtime/Engine/Classes/Materials/MaterialExpressionTransform.h b/Engine/Source/Runtime/Engine/Classes/Materials/MaterialExpressionTransform.h
--- a/Engine/Source/Runtime/Engine/Classes/Materials/MaterialExpressionTransform.h
+++ b/Engine/Source/Runtime/Engine/Classes/Materials/MaterialExpressionTransform.h
@@ -8,16 +8,16 @@
 UENUM()
 enum EMaterialVectorCoordTransformSource
 {
-	/** Tangent space */
+	/** Tangent space (relative to the surface) */
 	TRANSFORMSOURCE_Tangent UMETA(DisplayName="Tangent Space"),
 	
-	/** Local space */
+	/** Local space (relative to the rendered object, = object space) */
 	TRANSFORMSOURCE_Local UMETA(DisplayName="Local Space"),
 	
-	/** World space */
+	/** World space, a unit is 1cm */
 	TRANSFORMSOURCE_World UMETA(DisplayName="World Space"),
 	
-	/** View space (differs from camera space in the shadow passes) */
+	/** View space (relative to the camera/eye, = camera space, differs from camera space in the shadow passes) */
 	TRANSFORMSOURCE_View UMETA(DisplayName="View Space"),
 
 	/** Camera space */
@@ -29,16 +29,16 @@ enum EMaterialVectorCoordTransformSource
 UENUM()
 enum EMaterialVectorCoordTransform
 {
-	/** Tangent space */
+	/** Tangent space (relative to the surface) */
 	TRANSFORM_Tangent UMETA(DisplayName="Tangent Space"),
 	
-	/** Local space */
+	/** Local space (relative to the rendered object, = object space) */
 	TRANSFORM_Local UMETA(DisplayName="Local Space"),
 	
-	/** World space */
+	/** World space, a unit is 1cm */
 	TRANSFORM_World UMETA(DisplayName="World Space"),
 	
-	/** View space (differs from camera space in the shadow passes) */
+	/** View space (relative to the camera/eye, = camera space, differs from camera space in the shadow passes) */
 	TRANSFORM_View UMETA(DisplayName="View Space"),
 
 	/** Camera space */
@@ -64,7 +64,6 @@ class UMaterialExpressionTransform : public UMaterialExpression
 	UPROPERTY(EditAnywhere, Category=MaterialExpressionTransform, meta=(DisplayName = "Destination"))
 	TEnumAsByte<enum EMaterialVectorCoordTransform> TransformType;
 
-
 	//~ Begin UMaterialExpression Interface
 	virtual int32 Compile(class FMaterialCompiler* Compiler, int32 OutputIndex, int32 MultiplexIndex) override;
 	virtual void GetCaption(TArray<FString>& OutCaptions) const override;
diff --git a/Engine/Source/Runtime/Engine/Private/ActorReplication.cpp b/Engine/Source/Runtime/Engine/Private/ActorReplication.cpp
--- a/Engine/Source/Runtime/Engine/Private/ActorReplication.cpp
+++ b/Engine/Source/Runtime/Engine/Private/ActorReplication.cpp
@@ -376,4 +376,9 @@ bool AActor::IsNameStableForNetworking() const
 bool AActor::IsSupportedForNetworking() const
 {
 	return true;		// All actors are supported for networking
+}
+
+void AActor::OnRep_Owner()
+{
+
 }
\ No newline at end of file
diff --git a/Engine/Source/Runtime/Engine/Private/Animation/AnimCurveTypes.cpp b/Engine/Source/Runtime/Engine/Private/Animation/AnimCurveTypes.cpp
--- a/Engine/Source/Runtime/Engine/Private/Animation/AnimCurveTypes.cpp
+++ b/Engine/Source/Runtime/Engine/Private/Animation/AnimCurveTypes.cpp
@@ -5,6 +5,9 @@
 #include "Animation/AnimInstance.h"
 #include "AnimationRuntime.h"
 
+DECLARE_CYCLE_STAT(TEXT("BlendedCurve InitFrom"), STAT_BlendedCurve_InitFrom, STATGROUP_Anim);
+
+
 /////////////////////////////////////////////////////
 // FFloatCurve
 
@@ -573,6 +576,8 @@ void FBlendedCurve::Override(const FBlendedCurve& CurveToOverrideFrom)
 
 void FBlendedCurve::InitFrom(const USkeleton* Skeleton)
 {
+	SCOPE_CYCLE_COUNTER(STAT_BlendedCurve_InitFrom);
+
 	if(const FSmartNameMapping* Mapping = Skeleton->SmartNames.GetContainer(USkeleton::AnimCurveMappingName))
 	{
 		Mapping->FillUidArray(UIDList);
@@ -584,6 +589,21 @@ void FBlendedCurve::InitFrom(const USkeleton* Skeleton)
 	bInitialized = true;
 }
 
+void FBlendedCurve::InitFrom(const FBlendedCurve& InCurveToInitFrom)
+{
+	SCOPE_CYCLE_COUNTER(STAT_BlendedCurve_InitFrom);
+
+	// make sure this doesn't happen
+	if (ensure(&InCurveToInitFrom != this))
+	{
+		UIDList.Reset();
+		UIDList.Append(InCurveToInitFrom.UIDList);
+		Elements.Reset();
+		Elements.AddZeroed(UIDList.Num());
+		bInitialized = true;
+	}
+}
+
 FBlendedCurve::FBlendedCurve(const class UAnimInstance* AnimInstance)
 {
 	check(AnimInstance);
diff --git a/Engine/Source/Runtime/Engine/Private/Animation/AnimInstance.cpp b/Engine/Source/Runtime/Engine/Private/Animation/AnimInstance.cpp
--- a/Engine/Source/Runtime/Engine/Private/Animation/AnimInstance.cpp
+++ b/Engine/Source/Runtime/Engine/Private/Animation/AnimInstance.cpp
@@ -52,6 +52,7 @@ DEFINE_STAT(STAT_AnimMontageInstance_Advance);
 DEFINE_STAT(STAT_AnimMontageInstance_TickBranchPoints);
 DEFINE_STAT(STAT_AnimMontageInstance_Advance_Iteration);
 DEFINE_STAT(STAT_UpdateCurves);
+DEFINE_STAT(STAT_LocalBlendCSBoneTransforms);
 
 DECLARE_CYCLE_STAT_EXTERN(TEXT("Anim Init Time"), STAT_AnimInitTime, STATGROUP_Anim, );
 DEFINE_STAT(STAT_AnimInitTime);
@@ -64,6 +65,9 @@ DEFINE_STAT(STAT_UpdateTriMeshVertices);
 
 DEFINE_STAT(STAT_AnimGameThreadTime);
 
+DECLARE_CYCLE_STAT(TEXT("TickAssetPlayerInstances"), STAT_TickAssetPlayerInstances, STATGROUP_Anim);
+DECLARE_CYCLE_STAT(TEXT("TickAssetPlayerInstance"), STAT_TickAssetPlayerInstance, STATGROUP_Anim);
+
 
 #define DO_ANIMSTAT_PROCESSING(StatName) DEFINE_STAT(STAT_ ## StatName)
 #include "AnimMTStats.h"
@@ -506,7 +510,7 @@ void UAnimInstance::UpdateAnimationInternal(float DeltaSeconds)
 
 void UAnimInstance::TickAssetPlayerInstances(float DeltaSeconds)
 {
-	QUICK_SCOPE_CYCLE_COUNTER(STAT_TickAssetPlayerInstances);
+	SCOPE_CYCLE_COUNTER(STAT_TickAssetPlayerInstances);
 
 	// Handle all players inside sync groups
 	TArray<FAnimGroupInstance>& SyncGroups = SyncGroupArrays[GetSyncGroupWriteIndex()];
@@ -535,7 +539,7 @@ void UAnimInstance::TickAssetPlayerInstances(float DeltaSeconds)
 			{
 				FAnimTickRecord& GroupLeader = SyncGroup.ActivePlayers[GroupLeaderIndex];
 				// if it has leader score
-				QUICK_SCOPE_CYCLE_COUNTER(STAT_UAnimInstance_UpdateAnimation_TickAssetPlayerInstance);
+				SCOPE_CYCLE_COUNTER(STAT_TickAssetPlayerInstance);
 				FScopeCycleCounterUObject Scope(GroupLeader.SourceAsset);
 				GroupLeader.SourceAsset->TickAssetPlayerInstance(GroupLeader, this, TickContext);
 
@@ -582,7 +586,7 @@ void UAnimInstance::TickAssetPlayerInstances(float DeltaSeconds)
 				{
 					FAnimTickRecord& AssetPlayer = SyncGroup.ActivePlayers[TickIndex];
 					{
-						QUICK_SCOPE_CYCLE_COUNTER(STAT_UAnimInstance_UpdateAnimation_TickAssetPlayerInstance);
+						SCOPE_CYCLE_COUNTER(STAT_TickAssetPlayerInstance);
 						FScopeCycleCounterUObject Scope(AssetPlayer.SourceAsset);
 						TickContext.RootMotionMovementParams.Clear();
 						AssetPlayer.SourceAsset->TickAssetPlayerInstance(AssetPlayer, this, TickContext);
@@ -607,7 +611,7 @@ void UAnimInstance::TickAssetPlayerInstances(float DeltaSeconds)
 		const bool bOnlyOneAnimationInGroup = true;
 		FAnimAssetTickContext TickContext(DeltaSeconds, RootMotionMode, bOnlyOneAnimationInGroup, ValidMarkers);
 		{
-			QUICK_SCOPE_CYCLE_COUNTER(STAT_UAnimInstance_UpdateAnimation_TickAssetPlayerInstance);
+			SCOPE_CYCLE_COUNTER(STAT_TickAssetPlayerInstance);
 			FScopeCycleCounterUObject Scope(AssetPlayerToTick.SourceAsset);
 			AssetPlayerToTick.SourceAsset->TickAssetPlayerInstance(AssetPlayerToTick, this, TickContext);
 		}
@@ -2711,6 +2715,11 @@ int32 UAnimInstance::Montage_GetNextSectionID(UAnimMontage const* const Montage,
 	return INDEX_NONE;
 }
 
+bool UAnimInstance::IsAnyMontagePlaying() const
+{
+	return (MontageInstances.Num() > 0);
+}
+
 UAnimMontage* UAnimInstance::GetCurrentActiveMontage()
 {
 	// Start from end, as most recent instances are added at the end of the queue.
diff --git a/Engine/Source/Runtime/Engine/Private/Animation/AnimMontage.cpp b/Engine/Source/Runtime/Engine/Private/Animation/AnimMontage.cpp
--- a/Engine/Source/Runtime/Engine/Private/Animation/AnimMontage.cpp
+++ b/Engine/Source/Runtime/Engine/Private/Animation/AnimMontage.cpp
@@ -1611,6 +1611,7 @@ bool FAnimMontageInstance::SimulateAdvance(float DeltaTime, float& InOutPosition
 void FAnimMontageInstance::Advance(float DeltaTime, struct FRootMotionMovementParams * OutRootMotionParams, bool bBlendRootMotion)
 {
 	SCOPE_CYCLE_COUNTER(STAT_AnimMontageInstance_Advance);
+	FScopeCycleCounterUObject MontageScope(Montage);
 
 	if( IsValid() )
 	{
diff --git a/Engine/Source/Runtime/Engine/Private/Animation/AnimNode_ApplyMeshSpaceAdditive.cpp b/Engine/Source/Runtime/Engine/Private/Animation/AnimNode_ApplyMeshSpaceAdditive.cpp
--- a/Engine/Source/Runtime/Engine/Private/Animation/AnimNode_ApplyMeshSpaceAdditive.cpp
+++ b/Engine/Source/Runtime/Engine/Private/Animation/AnimNode_ApplyMeshSpaceAdditive.cpp
@@ -22,29 +22,41 @@ void FAnimNode_ApplyMeshSpaceAdditive::CacheBones(const FAnimationCacheBonesCont
 
 void FAnimNode_ApplyMeshSpaceAdditive::Update(const FAnimationUpdateContext& Context)
 {
-	EvaluateGraphExposedInputs.Execute(Context);
-
 	Base.Update(Context);
-	const float ActualAlpha = AlphaScaleBias.ApplyTo(Alpha);
-	if (ActualAlpha > ZERO_ANIMWEIGHT_THRESH)
+
+	if (IsLODEnabled(Context.AnimInstance, LODThreshold))
 	{
-		Additive.Update(Context.FractionalWeight(ActualAlpha));
+		// @note: If you derive this class, and if you have input that you rely on for base
+		// this is not going to work	
+		EvaluateGraphExposedInputs.Execute(Context);
+		const float ActualAlpha = AlphaScaleBias.ApplyTo(Alpha);
+		if (ActualAlpha > ZERO_ANIMWEIGHT_THRESH)
+		{
+			Additive.Update(Context.FractionalWeight(ActualAlpha));
+		}
 	}
 }
 
 void FAnimNode_ApplyMeshSpaceAdditive::Evaluate(FPoseContext& Output)
 {
-	//@TODO: Could evaluate Base into Output and save a copy
-	const float ActualAlpha = AlphaScaleBias.ApplyTo(Alpha);
-	if (ActualAlpha > ZERO_ANIMWEIGHT_THRESH)
+	if (IsLODEnabled(Output.AnimInstance, LODThreshold))
 	{
-		FPoseContext AdditiveEvalContext(Output);
+		//@TODO: Could evaluate Base into Output and save a copy
+		const float ActualAlpha = AlphaScaleBias.ApplyTo(Alpha);
+		if (ActualAlpha > ZERO_ANIMWEIGHT_THRESH)
+		{
+			FPoseContext AdditiveEvalContext(Output);
 
-		Base.Evaluate(Output);
-		Additive.Evaluate(AdditiveEvalContext);
-		
-		FAnimationRuntime::AccumulateMeshSpaceRotationAdditiveToLocalPose(Output.Pose, AdditiveEvalContext.Pose, Output.Curve, AdditiveEvalContext.Curve, ActualAlpha);
+			Base.Evaluate(Output);
+			Additive.Evaluate(AdditiveEvalContext);
+
+			FAnimationRuntime::AccumulateMeshSpaceRotationAdditiveToLocalPose(Output.Pose, AdditiveEvalContext.Pose, Output.Curve, AdditiveEvalContext.Curve, ActualAlpha);
 
+		}
+		else
+		{
+			Base.Evaluate(Output);
+		}
 	}
 	else
 	{
@@ -54,6 +66,7 @@ void FAnimNode_ApplyMeshSpaceAdditive::Evaluate(FPoseContext& Output)
 
 FAnimNode_ApplyMeshSpaceAdditive::FAnimNode_ApplyMeshSpaceAdditive()
 	: Alpha(1.0f)
+	, LODThreshold(INDEX_NONE)
 {
 }
 
diff --git a/Engine/Source/Runtime/Engine/Private/Animation/AnimNode_StateMachine.cpp b/Engine/Source/Runtime/Engine/Private/Animation/AnimNode_StateMachine.cpp
--- a/Engine/Source/Runtime/Engine/Private/Animation/AnimNode_StateMachine.cpp
+++ b/Engine/Source/Runtime/Engine/Private/Animation/AnimNode_StateMachine.cpp
@@ -316,7 +316,8 @@ const FAnimationTransitionBetweenStates& FAnimNode_StateMachine::GetTransitionIn
 
 void FAnimNode_StateMachine::Update(const FAnimationUpdateContext& Context)
 {
-	if (FBakedAnimationStateMachine* Machine = GetMachineDescription())
+	FBakedAnimationStateMachine* Machine = GetMachineDescription();
+	if (Machine != nullptr)
 	{
 		if (Machine->States.Num() == 0)
 		{
@@ -337,6 +338,10 @@ void FAnimNode_StateMachine::Update(const FAnimationUpdateContext& Context)
 	}
 
 	SCOPE_CYCLE_COUNTER(STAT_AnimStateMachineUpdate);
+#if	STATS
+	// Record name of state machine we are updating
+	FScopeCycleCounter MachineNameCycleCounter(Machine->GetStatID());
+#endif // STATS
 
 	bool bFoundValidTransition = false;
 	int32 TransitionCountThisFrame = 0;
diff --git a/Engine/Source/Runtime/Engine/Private/Animation/AnimSequence.cpp b/Engine/Source/Runtime/Engine/Private/Animation/AnimSequence.cpp
--- a/Engine/Source/Runtime/Engine/Private/Animation/AnimSequence.cpp
+++ b/Engine/Source/Runtime/Engine/Private/Animation/AnimSequence.cpp
@@ -21,6 +21,12 @@
 #define USE_SLERP 0
 #define LOCTEXT_NAMESPACE "AnimSequence"
 
+DECLARE_CYCLE_STAT(TEXT("AnimSeq GetBonePose"), STAT_AnimSeq_GetBonePose, STATGROUP_Anim);
+DECLARE_CYCLE_STAT(TEXT("BlendedCurve InitFrom"), STAT_BlendedCurve_InitFrom, STATGROUP_Anim);
+DECLARE_CYCLE_STAT(TEXT("Build Anim Track Pairs"), STAT_BuildAnimTrackPairs, STATGROUP_Anim);
+DECLARE_CYCLE_STAT(TEXT("Extract Pose From Anim Data"), STAT_ExtractPoseFromAnimData, STATGROUP_Anim);
+
+
 /////////////////////////////////////////////////////
 // FRawAnimSequenceTrackNativeDeprecated
 
@@ -978,7 +984,7 @@ struct FGetBonePoseScratchArea : public TThreadSingleton<FGetBonePoseScratchArea
 
 void UAnimSequence::GetBonePose(FCompactPose& OutPose, FBlendedCurve& OutCurve, const FAnimExtractContext& ExtractionContext) const
 {
-	QUICK_SCOPE_CYCLE_COUNTER(STAT_GETBONEPOSE);
+	SCOPE_CYCLE_COUNTER(STAT_AnimSeq_GetBonePose);
 
 	USkeleton* MySkeleton = GetSkeleton();
 	if (!MySkeleton)
@@ -1084,74 +1090,80 @@ void UAnimSequence::GetBonePose(FCompactPose& OutPose, FBlendedCurve& OutCurve,
 	// this is not guaranteed for AnimSequences though... If Root is not animated, Track will not exist.
 	const bool bFirstTrackIsRootBone = (GetSkeletonIndexFromTrackIndex(0) == 0);
 
-	// Handle root bone separately if it is track 0. so we start w/ Index 1.
-	for (int32 TrackIndex = (bFirstTrackIsRootBone ? 1 : 0); TrackIndex < NumTracks; TrackIndex++)
 	{
-		const int32 SkeletonBoneIndex = GetSkeletonIndexFromTrackIndex(TrackIndex);
-		// not sure it's safe to assume that SkeletonBoneIndex can never be INDEX_NONE
-		if (SkeletonBoneIndex != INDEX_NONE)
+		SCOPE_CYCLE_COUNTER(STAT_BuildAnimTrackPairs);
+
+		// Handle root bone separately if it is track 0. so we start w/ Index 1.
+		for (int32 TrackIndex = (bFirstTrackIsRootBone ? 1 : 0); TrackIndex < NumTracks; TrackIndex++)
 		{
-			const FCompactPoseBoneIndex BoneIndex = RequiredBones.GetCompactPoseIndexFromSkeletonIndex(SkeletonBoneIndex);
-			//Nasty, we break our type safety, code in the lower levels should be adjusted for this
-			const int32 CompactPoseBoneIndex = BoneIndex.GetInt();
-			if (CompactPoseBoneIndex != INDEX_NONE)
+			const int32 SkeletonBoneIndex = GetSkeletonIndexFromTrackIndex(TrackIndex);
+			// not sure it's safe to assume that SkeletonBoneIndex can never be INDEX_NONE
+			if (SkeletonBoneIndex != INDEX_NONE)
 			{
-				RotationScalePairs.Add(BoneTrackPair(CompactPoseBoneIndex, TrackIndex));
-
-				// Skip extracting translation component for EBoneTranslationRetargetingMode::Skeleton.
-				switch (BoneTree[SkeletonBoneIndex].TranslationRetargetingMode)
+				const FCompactPoseBoneIndex BoneIndex = RequiredBones.GetCompactPoseIndexFromSkeletonIndex(SkeletonBoneIndex);
+				//Nasty, we break our type safety, code in the lower levels should be adjusted for this
+				const int32 CompactPoseBoneIndex = BoneIndex.GetInt();
+				if (CompactPoseBoneIndex != INDEX_NONE)
 				{
-				case EBoneTranslationRetargetingMode::Animation:
-					TranslationPairs.Add(BoneTrackPair(CompactPoseBoneIndex, TrackIndex));
-					break;
-				case EBoneTranslationRetargetingMode::AnimationScaled:
-					TranslationPairs.Add(BoneTrackPair(CompactPoseBoneIndex, TrackIndex));
-					AnimScaleRetargetingPairs.Add(BoneTrackPair(CompactPoseBoneIndex, SkeletonBoneIndex));
-					break;
-				case EBoneTranslationRetargetingMode::AnimationRelative:
-					TranslationPairs.Add(BoneTrackPair(CompactPoseBoneIndex, TrackIndex));
-					AnimRelativeRetargetingPairs.Add(BoneTrackPair(CompactPoseBoneIndex, SkeletonBoneIndex));
-					break;
+					RotationScalePairs.Add(BoneTrackPair(CompactPoseBoneIndex, TrackIndex));
+
+					// Skip extracting translation component for EBoneTranslationRetargetingMode::Skeleton.
+					switch (BoneTree[SkeletonBoneIndex].TranslationRetargetingMode)
+					{
+					case EBoneTranslationRetargetingMode::Animation:
+						TranslationPairs.Add(BoneTrackPair(CompactPoseBoneIndex, TrackIndex));
+						break;
+					case EBoneTranslationRetargetingMode::AnimationScaled:
+						TranslationPairs.Add(BoneTrackPair(CompactPoseBoneIndex, TrackIndex));
+						AnimScaleRetargetingPairs.Add(BoneTrackPair(CompactPoseBoneIndex, SkeletonBoneIndex));
+						break;
+					case EBoneTranslationRetargetingMode::AnimationRelative:
+						TranslationPairs.Add(BoneTrackPair(CompactPoseBoneIndex, TrackIndex));
+						AnimRelativeRetargetingPairs.Add(BoneTrackPair(CompactPoseBoneIndex, SkeletonBoneIndex));
+						break;
+					}
 				}
 			}
 		}
 	}
 
-	QUICK_SCOPE_CYCLE_COUNTER(STAT_GETBONEPOSE_FORREALZ);
-	// Handle Root Bone separately
-	if (bFirstTrackIsRootBone)
-	{
-		const int32 TrackIndex = 0;
-		FCompactPoseBoneIndex RootBone(0);
-		FTransform& RootAtom = OutPose[RootBone];
-
-		AnimationFormat_GetBoneAtom(
-			RootAtom,
-			*this,
-			TrackIndex,
-			ExtractionContext.CurrentTime);
-
-		// @laurent - we should look into splitting rotation and translation tracks, so we don't have to process translation twice.
-		RetargetBoneTransform(RootAtom, 0, RootBone, RequiredBones);
-	}
-
-	if (RotationScalePairs.Num() > 0)
 	{
-		// get the remaining bone atoms
-		OutPose.PopulateFromAnimation( //@TODO:@ANIMATION: Nasty hack, be good to not have this function on the pose
-			*this,
-			RotationScalePairs,
-			TranslationPairs,
-			RotationScalePairs,
-			ExtractionContext.CurrentTime);
+		SCOPE_CYCLE_COUNTER(STAT_ExtractPoseFromAnimData);
+		// Handle Root Bone separately
+		if (bFirstTrackIsRootBone)
+		{
+			const int32 TrackIndex = 0;
+			FCompactPoseBoneIndex RootBone(0);
+			FTransform& RootAtom = OutPose[RootBone];
+
+			AnimationFormat_GetBoneAtom(
+				RootAtom,
+				*this,
+				TrackIndex,
+				ExtractionContext.CurrentTime);
+
+			// @laurent - we should look into splitting rotation and translation tracks, so we don't have to process translation twice.
+			RetargetBoneTransform(RootAtom, 0, RootBone, RequiredBones);
+		}
 
-		/*AnimationFormat_GetAnimationPose(
-			*((FTransformArray*)&OutAtoms), 
-			RotationScalePairs,
-			TranslationPairs,
-			RotationScalePairs,
-			*this,
-			ExtractionContext.CurrentTime);*/
+		if (RotationScalePairs.Num() > 0)
+		{
+			// get the remaining bone atoms
+			OutPose.PopulateFromAnimation( //@TODO:@ANIMATION: Nasty hack, be good to not have this function on the pose
+				*this,
+				RotationScalePairs,
+				TranslationPairs,
+				RotationScalePairs,
+				ExtractionContext.CurrentTime);
+
+			/*AnimationFormat_GetAnimationPose(
+				*((FTransformArray*)&OutAtoms), 
+				RotationScalePairs,
+				TranslationPairs,
+				RotationScalePairs,
+				*this,
+				ExtractionContext.CurrentTime);*/
+		}
 	}
 
 	// Once pose has been extracted, snap root bone back to first frame if we are extracting root motion.
diff --git a/Engine/Source/Runtime/Engine/Private/Animation/AnimSequenceBase.cpp b/Engine/Source/Runtime/Engine/Private/Animation/AnimSequenceBase.cpp
--- a/Engine/Source/Runtime/Engine/Private/Animation/AnimSequenceBase.cpp
+++ b/Engine/Source/Runtime/Engine/Private/Animation/AnimSequenceBase.cpp
@@ -10,6 +10,9 @@
 
 DEFINE_LOG_CATEGORY(LogAnimMarkerSync);
 
+DECLARE_CYCLE_STAT(TEXT("AnimSeq EvalCurveData"), STAT_AnimSeq_EvalCurveData, STATGROUP_Anim);
+
+
 /////////////////////////////////////////////////////
 
 UAnimSequenceBase::UAnimSequenceBase(const FObjectInitializer& ObjectInitializer)
@@ -477,6 +480,7 @@ uint8* UAnimSequenceBase::FindArrayProperty(const TCHAR* PropName, UArrayPropert
 /** Add curve data to Instance at the time of CurrentTime **/
 void UAnimSequenceBase::EvaluateCurveData(FBlendedCurve& OutCurve, float CurrentTime ) const
 {
+	SCOPE_CYCLE_COUNTER(STAT_AnimSeq_EvalCurveData);
 	RawCurveData.EvaluateCurveData(OutCurve, CurrentTime);
 }
 
diff --git a/Engine/Source/Runtime/Engine/Private/Animation/AnimationRuntime.cpp b/Engine/Source/Runtime/Engine/Private/Animation/AnimationRuntime.cpp
--- a/Engine/Source/Runtime/Engine/Private/Animation/AnimationRuntime.cpp
+++ b/Engine/Source/Runtime/Engine/Private/Animation/AnimationRuntime.cpp
@@ -19,6 +19,12 @@
 DEFINE_LOG_CATEGORY(LogAnimation);
 DEFINE_LOG_CATEGORY(LogRootMotion);
 
+DECLARE_CYCLE_STAT(TEXT("ConvertPoseToMeshRot"), STAT_ConvertPoseToMeshRot, STATGROUP_Anim);
+DECLARE_CYCLE_STAT(TEXT("ConvertMeshRotPoseToLocalSpace"), STAT_ConvertMeshRotPoseToLocalSpace, STATGROUP_Anim);
+DECLARE_CYCLE_STAT(TEXT("AccumulateMeshSpaceRotAdditiveToLocalPose"), STAT_AccumulateMeshSpaceRotAdditiveToLocalPose, STATGROUP_Anim);
+DECLARE_CYCLE_STAT(TEXT("BlendPosesPerBoneFilter"), STAT_BlendPosesPerBoneFilter, STATGROUP_Anim);
+
+
 /////////////////////////////////////////////////////////
 // Templated Transform Blend Functionality
 
@@ -435,6 +441,8 @@ void FAnimationRuntime::ConvertPoseToAdditive(FCompactPose& TargetPose, const FC
 
 void FAnimationRuntime::ConvertPoseToMeshRotation(FCompactPose& LocalPose)
 {
+	SCOPE_CYCLE_COUNTER(STAT_ConvertPoseToMeshRot);
+
 	// Convert all rotations to mesh space
 	// only the root bone doesn't have a parent. So skip it to save a branch in the iteration.
 	for (FCompactPoseBoneIndex BoneIndex(1); BoneIndex < LocalPose.GetNumBones(); ++BoneIndex)
@@ -448,6 +456,8 @@ void FAnimationRuntime::ConvertPoseToMeshRotation(FCompactPose& LocalPose)
 
 void FAnimationRuntime::ConvertMeshRotationPoseToLocalSpace(FCompactPose& Pose)
 {
+	SCOPE_CYCLE_COUNTER(STAT_ConvertMeshRotPoseToLocalSpace);
+
 	// Convert all rotations to mesh space
 	// only the root bone doesn't have a parent. So skip it to save a branch in the iteration.
 	for (FCompactPoseBoneIndex BoneIndex(Pose.GetNumBones()-1); BoneIndex > 0; --BoneIndex)
@@ -503,6 +513,8 @@ void FAnimationRuntime::AccumulateLocalSpaceAdditivePose(FCompactPose& BasePose,
 
 void FAnimationRuntime::AccumulateMeshSpaceRotationAdditiveToLocalPose(FCompactPose& BasePose, const FCompactPose& MeshSpaceRotationAdditive, FBlendedCurve& BaseCurve, const FBlendedCurve& AdditiveCurve, float Weight)
 {
+	SCOPE_CYCLE_COUNTER(STAT_AccumulateMeshSpaceRotAdditiveToLocalPose);
+
 	if (Weight > ZERO_ANIMWEIGHT_THRESH)
 	{
 		// Convert base pose from local space to mesh space rotation.
@@ -979,6 +991,8 @@ void FAnimationRuntime::BlendPosesPerBoneFilter(
 	bool bMeshSpaceRotationBlending, 
 	ECurveBlendOption::Type CurveBlendOption)
 {
+	SCOPE_CYCLE_COUNTER(STAT_BlendPosesPerBoneFilter);
+
 	ensure(OutPose.GetNumBones() == BasePose.GetNumBones());
 
 	if (BlendPoses.Num() != 0)
diff --git a/Engine/Source/Runtime/Engine/Private/Animation/AnimationSettings.cpp b/Engine/Source/Runtime/Engine/Private/Animation/AnimationSettings.cpp
--- a/Engine/Source/Runtime/Engine/Private/Animation/AnimationSettings.cpp
+++ b/Engine/Source/Runtime/Engine/Private/Animation/AnimationSettings.cpp
@@ -20,7 +20,7 @@ UAnimationSettings::UAnimationSettings(const FObjectInitializer& ObjectInitializ
 	, bTryPerTrackBitwiseCompression(true)
 	, bTryLinearKeyRemovalCompression(true)
 	, bTryIntervalKeyRemoval(true)
-
+	, bEnablePerformanceLog(false)
 {
 	SectionName = TEXT("Animation");
 
diff --git a/Engine/Source/Runtime/Engine/Private/Animation/BlendSpaceBase.cpp b/Engine/Source/Runtime/Engine/Private/Animation/BlendSpaceBase.cpp
--- a/Engine/Source/Runtime/Engine/Private/Animation/BlendSpaceBase.cpp
+++ b/Engine/Source/Runtime/Engine/Private/Animation/BlendSpaceBase.cpp
@@ -10,6 +10,8 @@
 #include "Animation/BlendSpaceBase.h"
 #include "AnimationUtils.h"
 
+DECLARE_CYCLE_STAT(TEXT("BlendSpace GetAnimPose"), STAT_BlendSpace_GetAnimPose, STATGROUP_Anim);
+
 struct FSyncPattern
 {
 	// The markers that make up this patter
@@ -552,7 +554,6 @@ bool UBlendSpaceBase::InterpolateWeightOfSampleData(float DeltaTime, const TArra
 
 		if (OldSample.PerBoneBlendData.Num()!=PerBoneBlend.Num())
 		{
-			OldSample.PerBoneBlendData.Empty(PerBoneBlend.Num());
 			OldSample.PerBoneBlendData.Init(OldSample.TotalWeight, PerBoneBlend.Num());
 		}
 
@@ -623,7 +624,6 @@ bool UBlendSpaceBase::InterpolateWeightOfSampleData(float DeltaTime, const TArra
 
 		if (OldSample.PerBoneBlendData.Num()!=PerBoneBlend.Num())
 		{
-			OldSample.PerBoneBlendData.Empty(PerBoneBlend.Num());
 			OldSample.PerBoneBlendData.Init(OldSample.TotalWeight, PerBoneBlend.Num());
 		}
 
@@ -1168,6 +1168,9 @@ void UBlendSpaceBase::FillupGridElements(const TArray<FVector> & PointList, cons
 
 void UBlendSpaceBase::GetAnimationPose(TArray<FBlendSampleData>& BlendSampleDataCache, /*out*/ FCompactPose& OutPose, /*out*/ FBlendedCurve& OutCurve)
 {
+	SCOPE_CYCLE_COUNTER(STAT_BlendSpace_GetAnimPose);
+	FScopeCycleCounterUObject BlendSpaceScope(this);
+
 	if(BlendSampleDataCache.Num() == 0)
 	{
 		OutPose.ResetToRefPose();
diff --git a/Engine/Source/Runtime/Engine/Private/ChartCreation.cpp b/Engine/Source/Runtime/Engine/Private/ChartCreation.cpp
--- a/Engine/Source/Runtime/Engine/Private/ChartCreation.cpp
+++ b/Engine/Source/Runtime/Engine/Private/ChartCreation.cpp
@@ -319,6 +319,7 @@ void UEngine::TickFPSChart( float DeltaSeconds )
 				if (DeltaSeconds > (LastDeltaSeconds * HitchMultiplierAmount))
 				{
 					// We have a hitch!
+					OnHitchDetectedDelegate.Broadcast(DeltaSeconds);
 					
 					// Track the hitch by bucketing it based on time severity
 					const int32 HitchBucketCount = STAT_FPSChart_LastHitchBucketStat - STAT_FPSChart_FirstHitchStat;
diff --git a/Engine/Source/Runtime/Engine/Private/Components/SkinnedMeshComponent.cpp b/Engine/Source/Runtime/Engine/Private/Components/SkinnedMeshComponent.cpp
--- a/Engine/Source/Runtime/Engine/Private/Components/SkinnedMeshComponent.cpp
+++ b/Engine/Source/Runtime/Engine/Private/Components/SkinnedMeshComponent.cpp
@@ -27,12 +27,12 @@ FAutoConsoleVariableRef CVarSkeletalMeshLODBias(
 	);
 
 static TAutoConsoleVariable<int32> CVarEnableAnimRateOptimization(
-	TEXT("EnableAnimRateOptimization"),
+	TEXT("a.URO.Enable"),
 	1,
 	TEXT("True to anim rate optimization."));
 
 static TAutoConsoleVariable<int32> CVarDrawAnimRateOptimization(
-	TEXT("DrawAnimRateOptimization"),
+	TEXT("a.URO.Draw"),
 	0,
 	TEXT("True to draw color coded boxes for anim rate."));
 
@@ -2005,7 +2005,7 @@ void FAnimUpdateRateParameters::SetTrailMode(float DeltaTime, uint8 UpdateRateSh
 	UpdateRate = FMath::Max(NewUpdateRate, 1);
 	// Make sure EvaluationRate is a multiple of UpdateRate.
 	EvaluationRate = FMath::Max((NewEvaluationRate / UpdateRate) * UpdateRate, 1);
-	bInterpolateSkippedFrames = bNewInterpSkippedFrames || (FAnimUpdateRateManager::CVarForceInterpolation.GetValueOnAnyThread() == 1);
+	bInterpolateSkippedFrames = (bNewInterpSkippedFrames && (EvaluationRate < MaxEvalRateForInterpolation)) || (FAnimUpdateRateManager::CVarForceInterpolation.GetValueOnAnyThread() == 1);
 
 	// Make sure we don't overflow. we don't need very large numbers.
 	const uint32 Counter = (GFrameCounter + UpdateRateShift)% MAX_uint8;
diff --git a/Engine/Source/Runtime/Engine/Private/DemoNetDriver.cpp b/Engine/Source/Runtime/Engine/Private/DemoNetDriver.cpp
--- a/Engine/Source/Runtime/Engine/Private/DemoNetDriver.cpp
+++ b/Engine/Source/Runtime/Engine/Private/DemoNetDriver.cpp
@@ -2100,12 +2100,12 @@ AActor* UDemoNetDriver::GetActorForGUID(FNetworkGUID InGUID) const
 
 }
 
-void UDemoNetDriver::AddNonQueuedActorForScrubbing(AActor* Actor)
+void UDemoNetDriver::AddNonQueuedActorForScrubbing(AActor const* Actor)
 {
-	auto FoundChannel = ServerConnection->ActorChannels.Find(Actor);
+	UActorChannel const* const* const FoundChannel = ServerConnection->ActorChannels.Find(Actor);
 	if (FoundChannel != nullptr && *FoundChannel != nullptr)
 	{
-		FNetworkGUID ActorGUID = (*FoundChannel)->ActorNetGUID;
+		FNetworkGUID const ActorGUID = (*FoundChannel)->ActorNetGUID;
 		NonQueuedGUIDsForScrubbing.Add(ActorGUID);
 	}
 }
