Mergine changes from //depot/UE4 to //UE4/Main

diff --git a/Engine/Source/Runtime/Engine/Private/GlobalShader.cpp b/Engine/Source/Runtime/Engine/Private/GlobalShader.cpp
--- a/Engine/Source/Runtime/Engine/Private/GlobalShader.cpp
+++ b/Engine/Source/Runtime/Engine/Private/GlobalShader.cpp
@@ -24,10 +24,42 @@ FGlobalShaderMapId::FGlobalShaderMapId(EShaderPlatform Platform)
 {
 	TArray<FShaderType*> ShaderTypes;
 
+	TArray<const FShaderPipelineType*> ShaderPipelineTypes;
+	for (TLinkedList<FShaderPipelineType*>::TIterator ShaderPipelineIt(FShaderPipelineType::GetTypeList()); ShaderPipelineIt; ShaderPipelineIt.Next())
+	{
+		const FShaderPipelineType* Pipeline = *ShaderPipelineIt;
+		if (Pipeline->IsGlobalTypePipeline())
+		{
+			// Verify all the ShouldCache are in sync
+			int32 NumShouldCache = 0;
+			const auto& Stages = Pipeline->GetStages();
+			for (int32 Index = 0; Index < Stages.Num(); ++Index)
+			{
+				const FGlobalShaderType* GlobalShaderType = Stages[Index]->GetGlobalShaderType();
+				if (GlobalShaderType->ShouldCache(Platform))
+				{
+					++NumShouldCache;
+				}
+			}
+
+			if (NumShouldCache != 0)
+			{
+				if (NumShouldCache != Stages.Num())
+				{
+					//#todo-rco: Error or Warning?
+					UE_LOG(LogShaders, Error, TEXT("Mismatched ShouldCache() function on the shaders for Global ShaderPipeline %s"), Pipeline->GetName());
+				}
+				else
+				{
+					ShaderPipelineTypes.Add(Pipeline);
+				}
+			}
+		}
+	}
+
 	for (TLinkedList<FShaderType*>::TIterator ShaderTypeIt(FShaderType::GetTypeList()); ShaderTypeIt; ShaderTypeIt.Next())
 	{
 		FGlobalShaderType* GlobalShaderType = ShaderTypeIt->GetGlobalShaderType();
-
 		if (GlobalShaderType && GlobalShaderType->ShouldCache(Platform))
 		{
 			ShaderTypes.Add(GlobalShaderType);
@@ -35,6 +67,7 @@ FGlobalShaderMapId::FGlobalShaderMapId(EShaderPlatform Platform)
 	}
 
 	ShaderTypes.Sort(FCompareShaderTypes());
+	ShaderPipelineTypes.Sort(FCompareShaderPipelineTypes());
 
 	for (int32 TypeIndex = 0; TypeIndex < ShaderTypes.Num(); TypeIndex++)
 	{
@@ -43,6 +76,21 @@ FGlobalShaderMapId::FGlobalShaderMapId(EShaderPlatform Platform)
 		Dependency.SourceHash = ShaderTypes[TypeIndex]->GetSourceHash();
 		ShaderTypeDependencies.Add(Dependency);
 	}
+
+	//#todo-rco: Is this currently necessary? Might be redundant until we only have Pipelines...
+	for (int32 TypeIndex = 0; TypeIndex < ShaderPipelineTypes.Num(); TypeIndex++)
+	{
+		auto& Stages = ShaderPipelineTypes[TypeIndex]->GetStages();
+		for (int32 Index = 0; Index < Stages.Num(); ++Index)
+		{
+			FShaderType* GlobalShaderType = (FShaderType*)Stages[Index];
+			FShaderTypeDependency Dependency;
+			Dependency.ShaderPipelineType = ShaderPipelineTypes[TypeIndex];
+			Dependency.ShaderType = (FShaderType*)Stages[Index];
+			Dependency.SourceHash = Stages[Index]->GetSourceHash();
+			ShaderTypeDependencies.Add(Dependency);
+		}
+	}
 }
 
 void FGlobalShaderMapId::AppendKeyString(FString& KeyString) const
@@ -51,8 +99,15 @@ void FGlobalShaderMapId::AppendKeyString(FString& KeyString) const
 
 	for (int32 ShaderIndex = 0; ShaderIndex < ShaderTypeDependencies.Num(); ShaderIndex++)
 	{
-		KeyString += TEXT("_");
 		const FShaderTypeDependency& ShaderTypeDependency = ShaderTypeDependencies[ShaderIndex];
+
+		KeyString += TEXT("_");
+		if (ShaderTypeDependency.ShaderPipelineType)
+		{
+			KeyString += ShaderTypeDependency.ShaderPipelineType->GetName();
+			KeyString += TEXT("_");
+		}
+
 		KeyString += ShaderTypeDependency.ShaderType->GetName();
 
 		// Add the type's source hash so that we can invalidate cached shaders when .usf changes are made
@@ -83,9 +138,9 @@ void FGlobalShaderMapId::AppendKeyString(FString& KeyString) const
 	}
 }
 
-void FGlobalShaderType::BeginCompileShader(EShaderPlatform Platform, TArray<FShaderCompileJob*>& NewJobs)
+void FGlobalShaderType::BeginCompileShader(EShaderPlatform Platform, TArray<FShaderCommonCompileJob*>& NewJobs)
 {
-	FShaderCompileJob* NewJob = new FShaderCompileJob(GlobalShaderMapId, NULL, this);
+	FShaderCompileJob* NewJob = new FShaderCompileJob(GlobalShaderMapId, nullptr, this);
 	FShaderCompilerEnvironment& ShaderEnvironment = NewJob->Input.Environment;
 
 	UE_LOG(LogShaders, Verbose, TEXT("	%s"), GetName());
@@ -108,12 +163,29 @@ void FGlobalShaderType::BeginCompileShader(EShaderPlatform Platform, TArray<FSha
 		);
 }
 
-FShader* FGlobalShaderType::FinishCompileShader(const FShaderCompileJob& CurrentJob)
+void FGlobalShaderType::BeginCompileShaderPipeline(EShaderPlatform Platform, const FShaderPipelineType* ShaderPipeline, const TArray<FGlobalShaderType*>& ShaderStages, TArray<FShaderCommonCompileJob*>& NewJobs)
+{
+	check(ShaderStages.Num() > 0);
+	check(ShaderPipeline);
+	UE_LOG(LogShaders, Verbose, TEXT("	Pipeline: %s"), ShaderPipeline->GetName());
+
+	// Add all the jobs as individual first, then add the dependencies into a pipeline job
+	auto* NewPipelineJob = new FShaderPipelineCompileJob(GlobalShaderMapId, ShaderPipeline, ShaderStages.Num());
+	for (int32 Index = 0; Index < ShaderStages.Num(); ++Index)
+	{
+		auto* ShaderStage = ShaderStages[Index];
+		ShaderStage->BeginCompileShader(Platform, NewPipelineJob->StageJobs);
+	}
+
+	NewJobs.Add(NewPipelineJob);
+}
+
+FShader* FGlobalShaderType::FinishCompileShader(const FShaderCompileJob& CurrentJob, const FShaderPipelineType* ShaderPipelineType)
 {
 	FShader* Shader = nullptr;
 	if (CurrentJob.bSucceeded)
 	{
-		FShaderType* SpecificType = CurrentJob.ShaderType->LimitShaderResourceToThisType() ? CurrentJob.ShaderType : NULL;
+		FShaderType* SpecificType = CurrentJob.ShaderType->LimitShaderResourceToThisType() ? CurrentJob.ShaderType : nullptr;
 
 		// Reuse an existing resource with the same key or create a new one based on the compile output
 		// This allows FShaders to share compiled bytecode and RHI shader references
@@ -121,12 +193,12 @@ FShader* FGlobalShaderType::FinishCompileShader(const FShaderCompileJob& Current
 		check(Resource);
 
 		// Find a shader with the same key in memory
-		Shader = CurrentJob.ShaderType->FindShaderById(FShaderId(GGlobalShaderMapHash, NULL, CurrentJob.ShaderType, CurrentJob.Input.Target));
+		Shader = CurrentJob.ShaderType->FindShaderById(FShaderId(GGlobalShaderMapHash, ShaderPipelineType, nullptr, CurrentJob.ShaderType, CurrentJob.Input.Target));
 
 		// There was no shader with the same key so create a new one with the compile output, which will bind shader parameters
 		if (!Shader)
 		{
-			Shader = (*ConstructCompiledRef)(CompiledShaderInitializerType(this, CurrentJob.Output, Resource, GGlobalShaderMapHash, NULL));
+			Shader = (*ConstructCompiledRef)(CompiledShaderInitializerType(this, CurrentJob.Output, Resource, GGlobalShaderMapHash, ShaderPipelineType, nullptr));
 			CurrentJob.Output.ParameterMap.VerifyBindingsAreComplete(GetName(), (EShaderFrequency)CurrentJob.Output.Target.Frequency, CurrentJob.VFType);
 		}
 	}
@@ -134,10 +206,10 @@ FShader* FGlobalShaderType::FinishCompileShader(const FShaderCompileJob& Current
 	static auto* CVarShowShaderWarnings = IConsoleManager::Get().FindConsoleVariable(TEXT("r.ShowShaderCompilerWarnings"));
 	if (CVarShowShaderWarnings && CVarShowShaderWarnings->GetInt() && CurrentJob.Output.Errors.Num() > 0)
 	{
-		UE_LOG(LogShaderCompilers, Warning, TEXT("Warnings compiling global shader %s:\n"), CurrentJob.ShaderType->GetName());
+		UE_LOG(LogShaderCompilers, Warning, TEXT("Warnings compiling global shader %s %s %s:\n"), CurrentJob.ShaderType->GetName(), ShaderPipelineType ? TEXT("ShaderPipeline") : TEXT(""), ShaderPipelineType ? ShaderPipelineType->GetName() : TEXT(""));
 		for (int32 ErrorIndex = 0; ErrorIndex < CurrentJob.Output.Errors.Num(); ErrorIndex++)
 		{
-			UE_LOG(LogShaderCompilers, Warning, TEXT("	%s"), *CurrentJob.Output.Errors[ErrorIndex].StrippedErrorMessage);
+			UE_LOG(LogShaderCompilers, Warning, TEXT("	%s"), *CurrentJob.Output.Errors[ErrorIndex].GetErrorString());
 		}
 	}
 
@@ -208,21 +280,66 @@ void VerifyGlobalShaders(EShaderPlatform Platform, bool bLoadedFromCacheFile)
 		UE_LOG(LogShaders, Warning, TEXT("	Empty global shader map, recompiling all global shaders"));
 	}
 	
-	TArray<FShaderCompileJob*> GlobalShaderJobs;
+	bool bErrorOnMissing = bLoadedFromCacheFile;
+	if (FPlatformProperties::RequiresCookedData())
+	{
+		// We require all shaders to exist on cooked platforms because we can't compile them.
+		bErrorOnMissing = true;
+	}
 
-	for(TLinkedList<FShaderType*>::TIterator ShaderTypeIt(FShaderType::GetTypeList());ShaderTypeIt;ShaderTypeIt.Next())
+	TArray<FShaderCommonCompileJob*> GlobalShaderJobs;
+	for (TLinkedList<FShaderPipelineType*>::TIterator ShaderPipelineIt(FShaderPipelineType::GetTypeList());ShaderPipelineIt;ShaderPipelineIt.Next())
 	{
-		FGlobalShaderType* GlobalShaderType = ShaderTypeIt->GetGlobalShaderType();
-		if(GlobalShaderType && GlobalShaderType->ShouldCache(Platform))
+		const FShaderPipelineType* Pipeline = *ShaderPipelineIt;
+		if (Pipeline->IsGlobalTypePipeline())
 		{
-			if(!GlobalShaderMap->HasShader(GlobalShaderType))
+			if (!GlobalShaderMap->GetShaderPipeline(Pipeline))
 			{
-				bool bErrorOnMissing = bLoadedFromCacheFile;
-				if (FPlatformProperties::RequiresCookedData())
+				// First see if all the ShouldCache pass
+				auto& Stages = Pipeline->GetStages();
+				int32 ShaderStagesToCompile = 0;
+				TArray<FGlobalShaderType*> ShaderStages;
+				for (int32 Index = 0; Index < Stages.Num(); ++Index)
 				{
-					// We require all shaders to exist on cooked platforms because we can't compile them.
-					bErrorOnMissing = true;
+					FGlobalShaderType* GlobalShaderType = ((FShaderType*)(Stages[Index]))->GetGlobalShaderType();
+					if (GlobalShaderType->ShouldCache(Platform))
+					{
+						ShaderStages.Add(GlobalShaderType);
+						++ShaderStagesToCompile;
+					}
 				}
+
+				if (ShaderStagesToCompile > 0)
+				{
+					if (bErrorOnMissing)
+					{
+						UE_LOG(LogShaders, Fatal, TEXT("Missing global shader pipeline %s, Please make sure cooking was successful."), Pipeline->GetName());
+					}
+
+					if (ShaderStagesToCompile != ShaderStages.Num())
+					{
+						UE_LOG(LogShaders, Fatal, TEXT("Missing global shader pipeline %s, has mismatched ShouldCache() functions, Please make sure cooking was successful."), Pipeline->GetName());
+					}
+
+					if (!bEmptyMap)
+					{
+						UE_LOG(LogShaders, Warning, TEXT("	%s"), Pipeline->GetName());
+					}
+
+					// Make a pipeline job with all the stages
+					FGlobalShaderType::BeginCompileShaderPipeline(Platform, Pipeline, ShaderStages, GlobalShaderJobs);
+				}
+			}
+		}
+	}
+
+	for(TLinkedList<FShaderType*>::TIterator ShaderTypeIt(FShaderType::GetTypeList());ShaderTypeIt;ShaderTypeIt.Next())
+	{
+		FGlobalShaderType* GlobalShaderType = ShaderTypeIt->GetGlobalShaderType();
+		if (GlobalShaderType && GlobalShaderType->ShouldCache(Platform))
+		{
+			if (!GlobalShaderMap->HasShader(GlobalShaderType))
+			{
 				if (bErrorOnMissing)
 				{
 					UE_LOG(LogShaders, Fatal,TEXT("Missing global shader %s, Please make sure cooking was successful."), GlobalShaderType->GetName());
@@ -444,10 +561,9 @@ TShaderMap<FGlobalShaderType>* GetGlobalShaderMap(EShaderPlatform Platform, bool
 		extern int32 GCreateShadersOnLoad;
 		if (GCreateShadersOnLoad && Platform == GMaxRHIShaderPlatform)
 		{
-			for (TMap<FShaderType*, TRefCountPtr<FShader> >::TConstIterator ShaderIt(GGlobalShaderMap[Platform]->GetShaders()); ShaderIt; ++ShaderIt)
+			for (auto& Pair : GGlobalShaderMap[Platform]->GetShaders())
 			{
-				FShader* Shader = ShaderIt.Value();
-
+				FShader* Shader = Pair.Value;
 				if (Shader)
 				{
 					Shader->BeginInitializeResources();
@@ -468,13 +584,35 @@ bool IsGlobalShaderMapComplete(const TCHAR* TypeNameSubstring)
 
 		if (GlobalShaderMap)
 		{
+			auto ShouldCacheGlobalShader = [&](const FGlobalShaderType* GlobalShaderType)
+				{
+					return GlobalShaderType
+						&& (TypeNameSubstring == nullptr || (FPlatformString::Strstr(GlobalShaderType->GetName(), TypeNameSubstring) != nullptr))
+						&& GlobalShaderType->ShouldCache(Platform);
+				};
+
+			for (TLinkedList<FShaderPipelineType*>::TIterator ShaderPipelineIt(FShaderPipelineType::GetTypeList()); ShaderPipelineIt; ShaderPipelineIt.Next())
+			{
+				const FShaderPipelineType* Pipeline = *ShaderPipelineIt;
+				if (Pipeline->IsGlobalTypePipeline())
+				{
+					auto& Stages = Pipeline->GetStages();
+					// Only need to check Stage 0 as we validated ShouldCache matching on FGlobalShaderMapId()
+					const FGlobalShaderType* GlobalShaderType = Stages[0]->GetGlobalShaderType();
+					if (ShouldCacheGlobalShader(GlobalShaderType))
+					{
+						if (!GlobalShaderMap->GetShaderPipeline(Pipeline))
+						{
+							return false;
+						}
+					}
+				}
+			}
+
 			for (TLinkedList<FShaderType*>::TIterator ShaderTypeIt(FShaderType::GetTypeList()); ShaderTypeIt; ShaderTypeIt.Next())
 			{
 				FGlobalShaderType* GlobalShaderType = ShaderTypeIt->GetGlobalShaderType();
-
-				if (GlobalShaderType
-					&& (TypeNameSubstring == nullptr || (FPlatformString::Strstr(GlobalShaderType->GetName(), TypeNameSubstring) != nullptr))
-					&& GlobalShaderType->ShouldCache(Platform))
+				if (ShouldCacheGlobalShader(GlobalShaderType))
 				{
 					if (!GlobalShaderMap->HasShader(GlobalShaderType))
 					{
@@ -533,12 +671,14 @@ bool RecompileChangedShadersForPlatform( const FString& PlatformName )
 	// figure out which shaders are out of date
 	TArray<FShaderType*> OutdatedShaderTypes;
 	TArray<const FVertexFactoryType*> OutdatedFactoryTypes;
+	TArray<const FShaderPipelineType*> OutdatedShaderPipelineTypes;
 
 	// Pick up new changes to shader files
 	FlushShaderFileCache();
 
 	FShaderType::GetOutdatedTypes(OutdatedShaderTypes, OutdatedFactoryTypes);
-	UE_LOG(LogShaders, Display, TEXT("We found %d out of date shader types, and %d out of date VF types!"), OutdatedShaderTypes.Num(), OutdatedFactoryTypes.Num());
+	FShaderPipelineType::GetOutdatedTypes(OutdatedShaderTypes, OutdatedShaderPipelineTypes, OutdatedFactoryTypes);
+	UE_LOG(LogShaders, Display, TEXT("We found %d out of date shader types, %d outdated pipeline types, and %d out of date VF types!"), OutdatedShaderTypes.Num(), OutdatedShaderPipelineTypes.Num(), OutdatedFactoryTypes.Num());
 
 	for (int32 FormatIndex = 0; FormatIndex < DesiredShaderFormats.Num(); FormatIndex++)
 	{
@@ -547,7 +687,7 @@ bool RecompileChangedShadersForPlatform( const FString& PlatformName )
 
 		// Only compile for the desired platform if requested
 		// Kick off global shader recompiles
-		BeginRecompileGlobalShaders(OutdatedShaderTypes, ShaderPlatform);
+		BeginRecompileGlobalShaders(OutdatedShaderTypes, OutdatedShaderPipelineTypes, ShaderPlatform);
 		
 		// Block on global shaders
 		FinishRecompileGlobalShaders();
@@ -609,6 +749,7 @@ void RecompileShadersForRemote(
 	// figure out which shaders are out of date
 	TArray<FShaderType*> OutdatedShaderTypes;
 	TArray<const FVertexFactoryType*> OutdatedFactoryTypes;
+	TArray<const FShaderPipelineType*> OutdatedShaderPipelineTypes;
 
 	// Pick up new changes to shader files
 	FlushShaderFileCache();
@@ -616,7 +757,8 @@ void RecompileShadersForRemote(
 	if( bCompileChangedShaders )
 	{
 		FShaderType::GetOutdatedTypes( OutdatedShaderTypes, OutdatedFactoryTypes );
-		UE_LOG( LogShaders, Display, TEXT( "We found %d out of date shader types, and %d out of date VF types!" ), OutdatedShaderTypes.Num(), OutdatedFactoryTypes.Num() );
+		FShaderPipelineType::GetOutdatedTypes(OutdatedShaderTypes, OutdatedShaderPipelineTypes, OutdatedFactoryTypes);
+		UE_LOG(LogShaders, Display, TEXT("We found %d out of date shader types, %d outdated pipeline types, and %d out of date VF types!"), OutdatedShaderTypes.Num(), OutdatedShaderPipelineTypes.Num(), OutdatedFactoryTypes.Num());
 	}
 
 	{
@@ -631,7 +773,7 @@ void RecompileShadersForRemote(
 				if( bCompileChangedShaders )
 				{
 					// Kick off global shader recompiles
-					BeginRecompileGlobalShaders( OutdatedShaderTypes, ShaderPlatform );
+					BeginRecompileGlobalShaders(OutdatedShaderTypes, OutdatedShaderPipelineTypes, ShaderPlatform);
 
 					// Block on global shaders
 					FinishRecompileGlobalShaders();
@@ -674,7 +816,7 @@ void RecompileShadersForRemote(
 	}
 }
 
-void BeginRecompileGlobalShaders(const TArray<FShaderType*>& OutdatedShaderTypes, EShaderPlatform ShaderPlatform)
+void BeginRecompileGlobalShaders(const TArray<FShaderType*>& OutdatedShaderTypes, const TArray<const FShaderPipelineType*>& OutdatedShaderPipelineTypes, EShaderPlatform ShaderPlatform)
 {
 	if( !FPlatformProperties::RequiresCookedData() )
 	{
@@ -690,7 +832,16 @@ void BeginRecompileGlobalShaders(const TArray<FShaderType*>& OutdatedShaderTypes
 			{
 				UE_LOG(LogShaders, Log, TEXT("Flushing Global Shader %s"), CurrentGlobalShaderType->GetName());
 				GlobalShaderMap->RemoveShaderType(CurrentGlobalShaderType);
-			
+			}
+		}
+
+		for (int32 PipelineTypeIndex = 0; PipelineTypeIndex < OutdatedShaderPipelineTypes.Num(); ++PipelineTypeIndex)
+		{
+			const FShaderPipelineType* ShaderPipelineType = OutdatedShaderPipelineTypes[PipelineTypeIndex];
+			if (ShaderPipelineType->IsGlobalTypePipeline())
+			{
+				UE_LOG(LogShaders, Log, TEXT("Flushing Global Shader Pipeline %s"), ShaderPipelineType->GetName());
+				GlobalShaderMap->RemoveShaderPipelineType(ShaderPipelineType);
 			}
 		}
 
@@ -710,29 +861,65 @@ void FinishRecompileGlobalShaders()
 	GShaderCompilingManager->ProcessAsyncResults(false, true);
 }
 
-void ProcessCompiledGlobalShaders(const TArray<FShaderCompileJob*>& CompilationResults)
+void ProcessCompiledGlobalShaders(const TArray<FShaderCommonCompileJob*>& CompilationResults)
 {
 	UE_LOG(LogShaders, Warning, TEXT("Compiled %u global shaders"), CompilationResults.Num());
 
 	TArray<EShaderPlatform> ShaderPlatformsProcessed;
 
+	auto ProcessCompiledJob = [&](FShaderCompileJob* SingleJob, const FShaderPipelineType* Pipeline)
+		{
+			FGlobalShaderType* GlobalShaderType = SingleJob->ShaderType->GetGlobalShaderType();
+			check(GlobalShaderType);
+			FShader* Shader = GlobalShaderType->FinishCompileShader(*SingleJob, Pipeline);
+			if (Shader)
+			{
+				// Add the new global shader instance to the global shader map.
+				EShaderPlatform Platform = (EShaderPlatform)SingleJob->Input.Target.Platform;
+				if (!Pipeline || Pipeline->ShouldAddShadersToIndividualMap())
+				{
+					GGlobalShaderMap[Platform]->AddShader(GlobalShaderType, Shader);
+				}
+				ShaderPlatformsProcessed.AddUnique(Platform);
+			}
+			else
+			{
+				UE_LOG(LogShaders, Fatal, TEXT("Failed to compile global shader %s %s %s.  Enable 'r.ShaderDevelopmentMode' in ConsoleVariables.ini for retries."),
+					GlobalShaderType->GetName(),
+					Pipeline ? TEXT("for pipeline") : TEXT(""),
+					Pipeline ? Pipeline->GetName() : TEXT(""));
+			}
+
+			return Shader;
+		};
+
 	for (int32 ResultIndex = 0; ResultIndex < CompilationResults.Num(); ResultIndex++)
 	{
-		const FShaderCompileJob& CurrentJob = *CompilationResults[ResultIndex];
-		FGlobalShaderType* GlobalShaderType = CurrentJob.ShaderType->GetGlobalShaderType();
-		check(GlobalShaderType);
-		FShader* Shader = GlobalShaderType->FinishCompileShader(CurrentJob);
-
-		if (Shader)
+		const FShaderCommonCompileJob& CurrentJob = *CompilationResults[ResultIndex];
+		FShaderCompileJob* SingleJob = nullptr;
+		if ((SingleJob = (FShaderCompileJob*)CurrentJob.GetSingleShaderJob()) != nullptr)
 		{
-			// Add the new global shader instance to the global shader map.
-			EShaderPlatform Platform = (EShaderPlatform)CurrentJob.Input.Target.Platform;
-			GGlobalShaderMap[Platform]->AddShader(GlobalShaderType,Shader);
-			ShaderPlatformsProcessed.AddUnique(Platform);
+			ProcessCompiledJob(SingleJob, nullptr);
 		}
 		else
 		{
-			UE_LOG(LogShaders, Fatal,TEXT("Failed to compile global shader %s.  Enable 'r.ShaderDevelopmentMode' in ConsoleVariables.ini for retries."), GlobalShaderType->GetName());
+			const auto* PipelineJob = CurrentJob.GetShaderPipelineJob();
+			check(PipelineJob);
+			TArray<FShader*> ShaderStages;
+			for (int32 Index = 0; Index < PipelineJob->StageJobs.Num(); ++Index)
+			{
+				SingleJob = PipelineJob->StageJobs[Index]->GetSingleShaderJob();
+				FShader* Shader = ProcessCompiledJob(SingleJob, PipelineJob->ShaderPipeline);
+				ShaderStages.Add(Shader);
+			}
+
+			FShaderPipeline* ShaderPipeline = new FShaderPipeline(PipelineJob->ShaderPipeline, ShaderStages);
+			if (ShaderPipeline)
+			{
+				EShaderPlatform Platform = (EShaderPlatform)PipelineJob->StageJobs[0]->GetSingleShaderJob()->Input.Target.Platform;
+				check(ShaderPipeline && !GGlobalShaderMap[Platform]->GetShaderPipeline(ShaderPipeline->PipelineType));
+				GGlobalShaderMap[Platform]->AddShaderPipeline(PipelineJob->ShaderPipeline, ShaderPipeline);
+			}
 		}
 	}
 
diff --git a/Engine/Source/Runtime/Engine/Private/Materials/Material.cpp b/Engine/Source/Runtime/Engine/Private/Materials/Material.cpp
--- a/Engine/Source/Runtime/Engine/Private/Materials/Material.cpp
+++ b/Engine/Source/Runtime/Engine/Private/Materials/Material.cpp
@@ -3224,14 +3224,14 @@ void UMaterial::CancelOutstandingCompilation()
 }
 #endif
 
-void UMaterial::UpdateMaterialShaders(TArray<FShaderType*>& ShaderTypesToFlush, TArray<const FVertexFactoryType*>& VFTypesToFlush, EShaderPlatform ShaderPlatform)
+void UMaterial::UpdateMaterialShaders(TArray<FShaderType*>& ShaderTypesToFlush, TArray<const FShaderPipelineType*>& ShaderPipelineTypesToFlush, TArray<const FVertexFactoryType*>& VFTypesToFlush, EShaderPlatform ShaderPlatform)
 {
 	// Create a material update context so we can safely update materials.
 	{
 		FMaterialUpdateContext UpdateContext(FMaterialUpdateContext::EOptions::Default, ShaderPlatform);
 
 		// Go through all material shader maps and flush the appropriate shaders
-		FMaterialShaderMap::FlushShaderTypes(ShaderTypesToFlush, VFTypesToFlush);
+		FMaterialShaderMap::FlushShaderTypes(ShaderTypesToFlush, ShaderPipelineTypesToFlush, VFTypesToFlush);
 
 		// There should be no references to the given material shader types at this point
 		// If there still are shaders of the given types, they may be reused when we call CacheResourceShaders instead of compiling new shaders
@@ -4138,10 +4138,10 @@ bool UMaterial::IsPropertyActive(EMaterialProperty InProperty) const
 		Active = ShadingModel == MSM_Subsurface || ShadingModel == MSM_PreintegratedSkin || ShadingModel == MSM_TwoSidedFoliage || ShadingModel == MSM_Cloth;
 		break;
 	case MP_CustomData0:
-		Active = ShadingModel == MSM_ClearCoat || ShadingModel == MSM_Hair || ShadingModel == MSM_Cloth;
+		Active = ShadingModel == MSM_ClearCoat || ShadingModel == MSM_Hair || ShadingModel == MSM_Cloth || ShadingModel == MSM_Eye;
 		break;
 	case MP_CustomData1:
-		Active = ShadingModel == MSM_ClearCoat;
+		Active = ShadingModel == MSM_ClearCoat || ShadingModel == MSM_Eye;
 		break;
 	case MP_TessellationMultiplier:
 	case MP_WorldDisplacement:
diff --git a/Engine/Source/Runtime/Engine/Private/Materials/MaterialShader.cpp b/Engine/Source/Runtime/Engine/Private/Materials/MaterialShader.cpp
--- a/Engine/Source/Runtime/Engine/Private/Materials/MaterialShader.cpp
+++ b/Engine/Source/Runtime/Engine/Private/Materials/MaterialShader.cpp
@@ -165,6 +165,14 @@ uint32 FMaterialShaderMap::NextCompilingId = 1;
  */
 TMap<TRefCountPtr<FMaterialShaderMap>, TArray<FMaterial*> > FMaterialShaderMap::ShaderMapsBeingCompiled;
 
+
+
+static inline bool ShouldCacheMaterialShader(const FMaterialShaderType* ShaderType, EShaderPlatform Platform, const FMaterial* Material)
+{
+	return ShaderType->ShouldCache(Platform, Material) && Material->ShouldCache(Platform, ShaderType, nullptr);
+}
+
+
 /** Converts an EMaterialShadingModel to a string description. */
 FString GetShadingModelString(EMaterialShadingModel ShadingModel)
 {
@@ -179,6 +187,7 @@ FString GetShadingModelString(EMaterialShadingModel ShadingModel)
 		case MSM_SubsurfaceProfile:	ShadingModelName = TEXT("MSM_SubsurfaceProfile"); break;
 		case MSM_TwoSidedFoliage:	ShadingModelName = TEXT("MSM_TwoSidedFoliage"); break;
 		case MSM_Cloth:				ShadingModelName = TEXT("MSM_Cloth"); break;
+		case MSM_Eye:				ShadingModelName = TEXT("MSM_Eye"); break;
 		default: ShadingModelName = TEXT("Unknown"); break;
 	}
 	return ShadingModelName;
@@ -597,6 +606,7 @@ bool FMaterialShaderMapId::operator==(const FMaterialShaderMapId& ReferenceSet)
 		const FShaderTypeDependency& ShaderTypeDependency = ShaderTypeDependencies[ShaderIndex];
 		
 		if (ShaderTypeDependency.ShaderType != ReferenceSet.ShaderTypeDependencies[ShaderIndex].ShaderType
+			|| ShaderTypeDependency.ShaderPipelineType != ReferenceSet.ShaderTypeDependencies[ShaderIndex].ShaderPipelineType
 			|| ShaderTypeDependency.SourceHash != ReferenceSet.ShaderTypeDependencies[ShaderIndex].SourceHash)
 		{
 			return false;
@@ -664,9 +674,13 @@ void FMaterialShaderMapId::AppendKeyString(FString& KeyString) const
 	// Add the inputs for any shaders that are stored inline in the shader map
 	for (int32 ShaderIndex = 0; ShaderIndex < ShaderTypeDependencies.Num(); ShaderIndex++)
 	{
-		KeyString += TEXT("_");
-
 		const FShaderTypeDependency& ShaderTypeDependency = ShaderTypeDependencies[ShaderIndex];
+		KeyString += TEXT("_");
+		if (ShaderTypeDependency.ShaderPipelineType)
+		{
+			KeyString += ShaderTypeDependency.ShaderPipelineType->GetName();
+			KeyString += TEXT("_");
+		}
 		KeyString += ShaderTypeDependency.ShaderType->GetName();
 		KeyString += ShaderTypeDependency.SourceHash.ToString();
 		ShaderTypeDependency.ShaderType->GetSerializationHistory().AppendKeyString(KeyString);
@@ -718,7 +732,7 @@ void FMaterialShaderMapId::AppendKeyString(FString& KeyString) const
 	KeyString += BytesToHex(&BasePropertyOverridesHash.Hash[0], sizeof(BasePropertyOverridesHash.Hash));
 }
 
-void FMaterialShaderMapId::SetShaderDependencies(const TArray<FShaderType*>& ShaderTypes, const TArray<FVertexFactoryType*>& VFTypes)
+void FMaterialShaderMapId::SetShaderDependencies(const TArray<FShaderType*>& ShaderTypes, const TArray<const FShaderPipelineType*>& ShaderPipelineTypes, const TArray<FVertexFactoryType*>& VFTypes)
 {
 	if (!FPlatformProperties::RequiresCookedData())
 	{
@@ -737,6 +751,21 @@ void FMaterialShaderMapId::SetShaderDependencies(const TArray<FShaderType*>& Sha
 			Dependency.VFSourceHash = VFTypes[VFTypeIndex]->GetSourceHash();
 			VertexFactoryTypeDependencies.Add(Dependency);
 		}
+
+		//#todo-rco: Is this currently necessary? Might be redundant until we only have Pipelines...
+		for (int32 TypeIndex = 0; TypeIndex < ShaderPipelineTypes.Num(); TypeIndex++)
+		{
+			auto& Stages = ShaderPipelineTypes[TypeIndex]->GetStages();
+			for (int32 Index = 0; Index < Stages.Num(); ++Index)
+			{
+				FShaderType* GlobalShaderType = (FShaderType*)Stages[Index];
+				FShaderTypeDependency Dependency;
+				Dependency.ShaderPipelineType = ShaderPipelineTypes[TypeIndex];
+				Dependency.ShaderType = (FShaderType*)Stages[Index];
+				Dependency.SourceHash = Stages[Index]->GetSourceHash();
+				ShaderTypeDependencies.Add(Dependency);
+			}
+		}
 	}
 }
 
@@ -766,10 +795,10 @@ void FMaterialShaderType::BeginCompileShader(
 	const FMaterial* Material,
 	FShaderCompilerEnvironment* MaterialEnvironment,
 	EShaderPlatform Platform,
-	TArray<FShaderCompileJob*>& NewJobs
+	TArray<FShaderCommonCompileJob*>& NewJobs
 	)
 {
-	FShaderCompileJob* NewJob = new FShaderCompileJob(ShaderMapId, NULL, this);
+	FShaderCompileJob* NewJob = new FShaderCompileJob(ShaderMapId, nullptr, this);
 
 	NewJob->Input.SharedEnvironment = MaterialEnvironment;
 	FShaderCompilerEnvironment& ShaderEnvironment = NewJob->Input.Environment;
@@ -785,7 +814,7 @@ void FMaterialShaderType::BeginCompileShader(
 	// Compile the shader environment passed in with the shader type's source code.
 	::GlobalBeginCompileShader(
 		Material->GetFriendlyName(),
-		NULL,
+		nullptr,
 		this,
 		GetShaderFilename(),
 		GetFunctionName(),
@@ -795,6 +824,30 @@ void FMaterialShaderType::BeginCompileShader(
 		);
 }
 
+void FMaterialShaderType::BeginCompileShaderPipeline(
+	uint32 ShaderMapId,
+	EShaderPlatform Platform,
+	const FMaterial* Material,
+	FShaderCompilerEnvironment* MaterialEnvironment,
+	const FShaderPipelineType* ShaderPipeline,
+	const TArray<FMaterialShaderType*>& ShaderStages,
+	TArray<FShaderCommonCompileJob*>& NewJobs)
+{
+	check(ShaderStages.Num() > 0);
+	check(ShaderPipeline);
+	UE_LOG(LogShaders, Verbose, TEXT("	Pipeline: %s"), ShaderPipeline->GetName());
+
+	// Add all the jobs as individual first, then add the dependencies into a pipeline job
+	auto* NewPipelineJob = new FShaderPipelineCompileJob(ShaderMapId, ShaderPipeline, ShaderStages.Num());
+	for (int32 Index = 0; Index < ShaderStages.Num(); ++Index)
+	{
+		auto* ShaderStage = ShaderStages[Index];
+		ShaderStage->BeginCompileShader(ShaderMapId, Material, MaterialEnvironment, Platform, NewPipelineJob->StageJobs);
+	}
+
+	NewJobs.Add(NewPipelineJob);
+}
+
 /**
  * Either creates a new instance of this type or returns an equivalent existing shader.
  * @param Material - The material to link the shader with.
@@ -804,6 +857,7 @@ FShader* FMaterialShaderType::FinishCompileShader(
 	const FUniformExpressionSet& UniformExpressionSet,
 	const FSHAHash& MaterialShaderMapHash,
 	const FShaderCompileJob& CurrentJob,
+	const FShaderPipelineType* ShaderPipeline,
 	const FString& InDebugDescription
 	)
 {
@@ -816,12 +870,12 @@ FShader* FMaterialShaderType::FinishCompileShader(
 	FShaderResource* Resource = FShaderResource::FindOrCreateShaderResource(CurrentJob.Output, SpecificType);
 
 	// Find a shader with the same key in memory
-	FShader* Shader = CurrentJob.ShaderType->FindShaderById(FShaderId(MaterialShaderMapHash, CurrentJob.VFType, CurrentJob.ShaderType, CurrentJob.Input.Target));
+	FShader* Shader = CurrentJob.ShaderType->FindShaderById(FShaderId(MaterialShaderMapHash, ShaderPipeline, CurrentJob.VFType, CurrentJob.ShaderType, CurrentJob.Input.Target));
 
 	// There was no shader with the same key so create a new one with the compile output, which will bind shader parameters
 	if (!Shader)
 	{
-		Shader = (*ConstructCompiledRef)(CompiledShaderInitializerType(this, CurrentJob.Output, Resource, UniformExpressionSet, MaterialShaderMapHash, NULL, InDebugDescription));
+		Shader = (*ConstructCompiledRef)(CompiledShaderInitializerType(this, CurrentJob.Output, Resource, UniformExpressionSet, MaterialShaderMapHash, ShaderPipeline, nullptr, InDebugDescription));
 		CurrentJob.Output.ParameterMap.VerifyBindingsAreComplete(GetName(), (EShaderFrequency)CurrentJob.Output.Target.Frequency, CurrentJob.VFType);
 	}
 
@@ -841,7 +895,7 @@ FMaterialShaderMap* FMaterialShaderMap::FindId(const FMaterialShaderMapId& Shade
 }
 
 /** Flushes the given shader types from any loaded FMaterialShaderMap's. */
-void FMaterialShaderMap::FlushShaderTypes(TArray<FShaderType*>& ShaderTypesToFlush, TArray<const FVertexFactoryType*>& VFTypesToFlush)
+void FMaterialShaderMap::FlushShaderTypes(TArray<FShaderType*>& ShaderTypesToFlush, TArray<const FShaderPipelineType*>& ShaderPipelineTypesToFlush, TArray<const FVertexFactoryType*>& VFTypesToFlush)
 {
 	for (int32 ShaderMapIndex = 0; ShaderMapIndex < AllMaterialShaderMaps.Num(); ShaderMapIndex++)
 	{
@@ -855,10 +909,14 @@ void FMaterialShaderMap::FlushShaderTypes(TArray<FShaderType*>& ShaderTypesToFlu
 		{
 			CurrentShaderMap->FlushShadersByVertexFactoryType(VFTypesToFlush[VFTypeIndex]);
 		}
+		for (int32 TypeIndex = 0; TypeIndex < ShaderPipelineTypesToFlush.Num(); TypeIndex++)
+		{
+			CurrentShaderMap->FlushShadersByShaderPipelineType(ShaderPipelineTypesToFlush[TypeIndex]);
+		}
 	}
 }
 
-void FMaterialShaderMap::FixupShaderTypes(EShaderPlatform Platform, const TMap<FShaderType*, FString>& ShaderTypeNames, const TMap<FVertexFactoryType*, FString>& VertexFactoryTypeNames)
+void FMaterialShaderMap::FixupShaderTypes(EShaderPlatform Platform, const TMap<FShaderType*, FString>& ShaderTypeNames, const TMap<const FShaderPipelineType*, FString>& ShaderPipelineTypeNames, const TMap<FVertexFactoryType*, FString>& VertexFactoryTypeNames)
 {
 	TArray<FMaterialShaderMapId> Keys;
 	FMaterialShaderMap::GIdToMaterialShaderMap[Platform].GenerateKeyArray(Keys);
@@ -887,6 +945,9 @@ void FMaterialShaderMap::FixupShaderTypes(EShaderPlatform Platform, const TMap<F
 			Key.VertexFactoryTypeDependencies[VFTypeIndex].VertexFactoryType = FoundVFType;
 		}
 
+		//#todo-rco
+		check(ShaderPipelineTypeNames.Num() == 0);
+
 		FMaterialShaderMap::GIdToMaterialShaderMap[Platform].Add(Key, Values[PairIndex]);
 	}
 }
@@ -992,17 +1053,17 @@ void FMaterialShaderMap::SaveForRemoteRecompile(FArchive& Ar, const TMap<FString
 			if (ShaderMap)
 			{
 				// get all shaders in the shader map
-				TMap<FShaderId,FShader*> Shaders;
-				ShaderMap->GetShaderList(Shaders);
+				TMap<FShaderId, FShader*> ShaderList;
+				ShaderMap->GetShaderList(ShaderList);
 
 				// get the resources from the shaders
-				for (TMap<FShaderId, FShader*>::TIterator ShaderIt(Shaders); ShaderIt; ++ShaderIt)
+				for (auto& KeyValue : ShaderList)
 				{
-					FShader *Shader = ShaderIt.Value();
+					FShader* Shader = KeyValue.Value;
 					FShaderResourceId ShaderId = Shader->GetResourceId();
 
 					// skip this shader if the Id was already on the client (ie, it didn't change)
-					if (ClientResourceIds.Contains(ShaderId) == false )
+					if (ClientResourceIds.Contains(ShaderId) == false)
 					{
 						// lookup the resource by ID
 						FShaderResource* Resource = FShaderResource::FindShaderResourceById(ShaderId);
@@ -1273,7 +1334,7 @@ void FMaterialShaderMap::Compile(
   
 			uint32 NumShaders = 0;
 			uint32 NumVertexFactories = 0;
-			TArray<FShaderCompileJob*> NewJobs;
+			TArray<FShaderCommonCompileJob*> NewJobs;
   
 			// Iterate over all vertex factory types.
 			for(TLinkedList<FVertexFactoryType*>::TIterator VertexFactoryTypeIt(FVertexFactoryType::GetTypeList());VertexFactoryTypeIt;VertexFactoryTypeIt.Next())
@@ -1283,7 +1344,7 @@ void FMaterialShaderMap::Compile(
   
 				if(VertexFactoryType->IsUsedWithMaterials())
 				{
-					FMeshMaterialShaderMap* MeshShaderMap = NULL;
+					FMeshMaterialShaderMap* MeshShaderMap = nullptr;
   
 					// look for existing map for this vertex factory type
 					int32 MeshShaderMapIndex = INDEX_NONE;
@@ -1297,7 +1358,7 @@ void FMaterialShaderMap::Compile(
 						}
 					}
   
-					if (MeshShaderMap == NULL)
+					if (MeshShaderMap == nullptr)
 					{
 						// Create a new mesh material shader map.
 						MeshShaderMapIndex = MeshShaderMaps.Num();
@@ -1325,10 +1386,7 @@ void FMaterialShaderMap::Compile(
 			for(TLinkedList<FShaderType*>::TIterator ShaderTypeIt(FShaderType::GetTypeList());ShaderTypeIt;ShaderTypeIt.Next())
 			{
 				FMaterialShaderType* ShaderType = ShaderTypeIt->GetMaterialShaderType();
-				if (ShaderType && 
-					ShaderType->ShouldCache(InPlatform,Material) && 
-					Material->ShouldCache(InPlatform, ShaderType, NULL)
-					)
+				if (ShaderType &&  ShouldCacheMaterialShader(ShaderType, InPlatform, Material))
 				{
 					// Verify that the shader map Id contains inputs for any shaders that will be put into this shader map
 					check(InShaderMapId.ContainsShaderType(ShaderType));
@@ -1351,6 +1409,39 @@ void FMaterialShaderMap::Compile(
 				}
 			}
   
+			for (TLinkedList<FShaderPipelineType*>::TIterator ShaderPipelineIt(FShaderPipelineType::GetTypeList());ShaderPipelineIt;ShaderPipelineIt.Next())
+			{
+				const FShaderPipelineType* Pipeline = *ShaderPipelineIt;
+				if (Pipeline->IsMaterialTypePipeline())
+				{
+					auto& Stages = Pipeline->GetStages();
+					bool bCompile = true;
+					TArray<FMaterialShaderType*> ShaderStagesToCompile;
+					bool bCompilePipeline = true;
+					for (int32 Index = 0; Index < Stages.Num(); ++Index)
+					{
+						FMaterialShaderType* ShaderType = (FMaterialShaderType*)(Stages[Index]->GetMaterialShaderType());
+						if (ShaderType && ShouldCacheMaterialShader(ShaderType, InPlatform, Material))
+						{
+							// Verify that the shader map Id contains inputs for any shaders that will be put into this shader map
+							check(InShaderMapId.ContainsShaderType(ShaderType));
+							ShaderStagesToCompile.Add(ShaderType);
+						}
+						else
+						{
+							bCompilePipeline = false;
+							break;
+						}
+
+						if (bCompilePipeline)
+						{
+							NumShaders += ShaderStagesToCompile.Num();
+							FMaterialShaderType::BeginCompileShaderPipeline(CompilingId, InPlatform, Material, MaterialEnvironment, Pipeline, ShaderStagesToCompile, NewJobs);
+						}
+					}
+				}
+			}
+
 			if (!CorrespondingMaterials)
 			{
 				UE_LOG(LogShaders, Warning, TEXT("		%u Shaders among %u VertexFactories"), NumShaders, NumVertexFactories);
@@ -1379,7 +1470,59 @@ void FMaterialShaderMap::Compile(
 	}
 }
 
-bool FMaterialShaderMap::ProcessCompilationResults(const TArray<FShaderCompileJob*>& InCompilationResults, int32& InOutJobIndex, float& TimeBudget)
+FShader* FMaterialShaderMap::ProcessCompilationResultsForSingleJob(FShaderCompileJob* SingleJob, const FShaderPipelineType* ShaderPipeline, const FSHAHash& MaterialShaderMapHash)
+{
+	check(SingleJob);
+	const FShaderCompileJob& CurrentJob = *SingleJob;
+	check(CurrentJob.Id == CompilingId);
+
+	FShader* Shader = nullptr;
+	if (CurrentJob.VFType)
+	{
+		FVertexFactoryType* VertexFactoryType = CurrentJob.VFType;
+		check(VertexFactoryType->IsUsedWithMaterials());
+		FMeshMaterialShaderMap* MeshShaderMap = nullptr;
+		int32 MeshShaderMapIndex = INDEX_NONE;
+
+		// look for existing map for this vertex factory type
+		for (int32 ShaderMapIndex = 0; ShaderMapIndex < MeshShaderMaps.Num(); ShaderMapIndex++)
+		{
+			if (MeshShaderMaps[ShaderMapIndex].GetVertexFactoryType() == VertexFactoryType)
+			{
+				MeshShaderMap = &MeshShaderMaps[ShaderMapIndex];
+				MeshShaderMapIndex = ShaderMapIndex;
+				break;
+			}
+		}
+
+		check(MeshShaderMap);
+		FMeshMaterialShaderType* MeshMaterialShaderType = CurrentJob.ShaderType->GetMeshMaterialShaderType();
+		check(MeshMaterialShaderType);
+		Shader = MeshMaterialShaderType->FinishCompileShader(MaterialCompilationOutput.UniformExpressionSet, MaterialShaderMapHash, CurrentJob, ShaderPipeline, FriendlyName);
+		check(Shader);
+		if (!ShaderPipeline)
+		{
+			check(!MeshShaderMap->HasShader(MeshMaterialShaderType));
+			MeshShaderMap->AddShader(MeshMaterialShaderType, Shader);
+		}
+	}
+	else
+	{
+		FMaterialShaderType* MaterialShaderType = CurrentJob.ShaderType->GetMaterialShaderType();
+		check(MaterialShaderType);
+		Shader = MaterialShaderType->FinishCompileShader(MaterialCompilationOutput.UniformExpressionSet, MaterialShaderMapHash, CurrentJob, ShaderPipeline, FriendlyName);
+		check(Shader);
+		if (!ShaderPipeline)
+		{
+			check(!HasShader(MaterialShaderType));
+			AddShader(MaterialShaderType, Shader);
+		}
+	}
+
+	return Shader;
+}
+
+bool FMaterialShaderMap::ProcessCompilationResults(const TArray<FShaderCommonCompileJob*>& InCompilationResults, int32& InOutJobIndex, float& TimeBudget)
 {
 	check(InOutJobIndex < InCompilationResults.Num());
 
@@ -1388,43 +1531,61 @@ bool FMaterialShaderMap::ProcessCompilationResults(const TArray<FShaderCompileJo
 	FSHAHash MaterialShaderMapHash;
 	ShaderMapId.GetMaterialHash(MaterialShaderMapHash);
 
-	do 
+	do
 	{
-		const FShaderCompileJob& CurrentJob = *InCompilationResults[InOutJobIndex];
-		check(CurrentJob.Id == CompilingId);
-
-		if (CurrentJob.VFType)
+		FShaderCompileJob* SingleJob = InCompilationResults[InOutJobIndex]->GetSingleShaderJob();
+		if (SingleJob)
+		{
+			ProcessCompilationResultsForSingleJob(SingleJob, nullptr, MaterialShaderMapHash);
+		}
+		else
 		{
-			FVertexFactoryType* VertexFactoryType = CurrentJob.VFType;
-			check(VertexFactoryType->IsUsedWithMaterials());
-			FMeshMaterialShaderMap* MeshShaderMap = NULL;
-			int32 MeshShaderMapIndex = INDEX_NONE;
+			auto* PipelineJob = InCompilationResults[InOutJobIndex]->GetShaderPipelineJob();
+			check(PipelineJob);
+
+			const FShaderPipelineCompileJob& CurrentJob = *PipelineJob;
+			check(CurrentJob.Id == CompilingId);
+
+			TArray<FShader*> ShaderStages;
+			FVertexFactoryType* VertexFactoryType = CurrentJob.StageJobs[0]->GetSingleShaderJob()->VFType;
+			for (int32 Index = 0; Index < CurrentJob.StageJobs.Num(); ++Index)
+			{
+				SingleJob = CurrentJob.StageJobs[Index]->GetSingleShaderJob();
+				FShader* Shader = ProcessCompilationResultsForSingleJob(SingleJob, PipelineJob->ShaderPipeline, MaterialShaderMapHash);
+				ShaderStages.Add(Shader);
+				check(VertexFactoryType == CurrentJob.StageJobs[Index]->GetSingleShaderJob()->VFType);
+			}
 
-			// look for existing map for this vertex factory type
-			for (int32 ShaderMapIndex = 0; ShaderMapIndex < MeshShaderMaps.Num(); ShaderMapIndex++)
+			FShaderPipeline* ShaderPipeline = new FShaderPipeline(PipelineJob->ShaderPipeline, ShaderStages);
+			if (ShaderPipeline)
 			{
-				if (MeshShaderMaps[ShaderMapIndex].GetVertexFactoryType() == VertexFactoryType)
+				if (VertexFactoryType)
 				{
-					MeshShaderMap = &MeshShaderMaps[ShaderMapIndex];
-					MeshShaderMapIndex = ShaderMapIndex;
-					break;
+					check(VertexFactoryType->IsUsedWithMaterials());
+					FMeshMaterialShaderMap* MeshShaderMap = nullptr;
+					int32 MeshShaderMapIndex = INDEX_NONE;
+
+					// look for existing map for this vertex factory type
+					for (int32 ShaderMapIndex = 0; ShaderMapIndex < MeshShaderMaps.Num(); ShaderMapIndex++)
+					{
+						if (MeshShaderMaps[ShaderMapIndex].GetVertexFactoryType() == VertexFactoryType)
+						{
+							MeshShaderMap = &MeshShaderMaps[ShaderMapIndex];
+							MeshShaderMapIndex = ShaderMapIndex;
+							break;
+						}
+					}
+
+					check(MeshShaderMap);
+					check(!MeshShaderMap->GetShaderPipeline(ShaderPipeline->PipelineType));
+					MeshShaderMap->AddShaderPipeline(PipelineJob->ShaderPipeline, ShaderPipeline);
+				}
+				else
+				{
+					check(!GetShaderPipeline(ShaderPipeline->PipelineType));
+					AddShaderPipeline(PipelineJob->ShaderPipeline, ShaderPipeline);
 				}
 			}
-
-			check(MeshShaderMap);
-			FMeshMaterialShaderType* MeshMaterialShaderType = CurrentJob.ShaderType->GetMeshMaterialShaderType();
-			check(MeshMaterialShaderType);
-			FShader* Shader = MeshMaterialShaderType->FinishCompileShader(MaterialCompilationOutput.UniformExpressionSet, MaterialShaderMapHash, CurrentJob, FriendlyName);
-			check(Shader && !MeshShaderMap->HasShader(MeshMaterialShaderType));
-			MeshShaderMap->AddShader(MeshMaterialShaderType, Shader);
-		}
-		else
-		{
-			FMaterialShaderType* MaterialShaderType = CurrentJob.ShaderType->GetMaterialShaderType();
-			check(MaterialShaderType);
-			FShader* Shader = MaterialShaderType->FinishCompileShader(MaterialCompilationOutput.UniformExpressionSet, MaterialShaderMapHash, CurrentJob, FriendlyName);
-			check(Shader && !HasShader(MaterialShaderType));
-			AddShader(MaterialShaderType, Shader);
 		}
 
 		InOutJobIndex++;
@@ -1439,7 +1600,7 @@ bool FMaterialShaderMap::ProcessCompilationResults(const TArray<FShaderCompileJo
 	{
 		for (int32 ShaderMapIndex = MeshShaderMaps.Num() - 1; ShaderMapIndex >= 0; ShaderMapIndex--)
 		{
-			if (MeshShaderMaps[ShaderMapIndex].GetNumShaders() == 0)
+			if (MeshShaderMaps[ShaderMapIndex].GetNumShaders() == 0 && MeshShaderMaps[ShaderMapIndex].GetNumShaderPipelines() == 0)
 			{
 				// If the mesh material shader map is complete and empty, discard it.
 				MeshShaderMaps.RemoveAt(ShaderMapIndex);
@@ -1482,8 +1643,6 @@ bool FMaterialShaderMap::TryToAddToExistingCompilationTask(FMaterial* Material)
 
 bool FMaterialShaderMap::IsComplete(const FMaterial* Material, bool bSilent)
 {
-	bool bIsComplete = true;
-
 	// Make sure we are operating on a referenced shader map or the below Find will cause this shader map to be deleted,
 	// Since it creates a temporary ref counted pointer.
 	check(NumRefs > 0);
@@ -1504,39 +1663,87 @@ bool FMaterialShaderMap::IsComplete(const FMaterial* Material, bool bSilent)
 		{
 			// Find the shaders for this vertex factory type.
 			const FMeshMaterialShaderMap* MeshShaderMap = GetMeshShaderMap(VertexFactoryType);
-			if(!FMeshMaterialShaderMap::IsComplete(MeshShaderMap,Platform,Material,VertexFactoryType,bSilent))
+			if (!FMeshMaterialShaderMap::IsComplete(MeshShaderMap,Platform,Material,VertexFactoryType,bSilent))
 			{
 				if (!MeshShaderMap && !bSilent)
 				{
 					UE_LOG(LogShaders, Warning, TEXT("Incomplete material %s, missing Vertex Factory %s."), *Material->GetFriendlyName(), VertexFactoryType->GetName());
 				}
-				bIsComplete = false;
-				break;
+				return false;
 			}
 		}
 	}
 
+	auto IsMaterialShaderComplete = [&](const FMaterialShaderType* ShaderType, const FShaderPipelineType* Pipeline) -> bool
+		{
+			if (ShouldCacheMaterialShader(ShaderType, Platform, Material) &&
+				((Pipeline && !GetShaderPipeline(Pipeline)) || (!Pipeline && !HasShader((FShaderType*)ShaderType))))
+			{
+				if (!bSilent)
+				{
+					if (Pipeline)
+					{
+						UE_LOG(LogShaders, Warning, TEXT("Incomplete material %s, missing pipeline %s."), *Material->GetFriendlyName(), Pipeline->GetName());
+					}
+					else
+					{
+						UE_LOG(LogShaders, Warning, TEXT("Incomplete material %s, missing FMaterialShader %s."), *Material->GetFriendlyName(), ShaderType->GetName());
+					}
+				}
+				return false;
+			}
+
+			return true;
+		};
+
 	// Iterate over all material shader types.
 	for(TLinkedList<FShaderType*>::TIterator ShaderTypeIt(FShaderType::GetTypeList());ShaderTypeIt;ShaderTypeIt.Next())
 	{
 		// Find this shader type in the material's shader map.
 		FMaterialShaderType* ShaderType = ShaderTypeIt->GetMaterialShaderType();
-		if (ShaderType && 
-			!HasShader(ShaderType) &&
-			ShaderType->ShouldCache(Platform,Material) && 
-			Material->ShouldCache(Platform, ShaderType, NULL)
-			)
+		if (ShaderType && !IsMaterialShaderComplete(ShaderType, nullptr))
 		{
-			if (!bSilent)
+			return false;
+		}
+	}
+
+	// Iterate over all pipeline types
+	for (TLinkedList<FShaderPipelineType*>::TIterator ShaderPipelineIt(FShaderPipelineType::GetTypeList());ShaderPipelineIt;ShaderPipelineIt.Next())
+	{
+		const FShaderPipelineType* Pipeline = *ShaderPipelineIt;
+		if (Pipeline->IsMaterialTypePipeline())
+		{
+			auto& Stages = Pipeline->GetStages();
+
+			// Verify all the ShouldCache are in sync
+			int32 NumShouldCache = 0;
+			for (int32 Index = 0; Index < Stages.Num(); ++Index)
 			{
-				UE_LOG(LogShaders, Warning, TEXT("Incomplete material %s, missing FMaterialShader %s."), *Material->GetFriendlyName(), ShaderType->GetName());
+				auto* ShaderType = Stages[Index]->GetMaterialShaderType();
+				if (ShouldCacheMaterialShader(ShaderType, Platform, Material))
+				{
+					++NumShouldCache;
+				}
+			}
+
+			if (NumShouldCache > 0 && NumShouldCache != Stages.Num())
+			{
+				//#todo-rco: Error or Warning?
+				UE_LOG(LogShaders, Error, TEXT("Mismatched ShouldCache() function on the shaders for Material ShaderPipeline %s"), Pipeline->GetName());
+			}
+
+			for (int32 Index = 0; Index < Stages.Num(); ++Index)
+			{
+				auto* ShaderType = Stages[Index]->GetMaterialShaderType();
+				if (!IsMaterialShaderComplete(ShaderType, Pipeline))
+				{
+					return false;
+				}
 			}
-			bIsComplete = false;
-			break;
 		}
 	}
 
-	return bIsComplete;
+	return true;
 }
 
 void FMaterialShaderMap::LoadMissingShadersFromMemory(const FMaterial* Material)
@@ -1563,12 +1770,9 @@ void FMaterialShaderMap::LoadMissingShadersFromMemory(const FMaterial* Material)
 
 		const bool bShaderAlreadyExists = HasShader(ShaderType);
 
-		if (ShaderType && 
-			ShaderType->ShouldCache(Platform, Material) && 
-			Material->ShouldCache(Platform, ShaderType, NULL) &&
-			!bShaderAlreadyExists)
+		if (ShaderType && ShouldCacheMaterialShader(ShaderType, Platform, Material) && !bShaderAlreadyExists)
 		{
-			FShaderId ShaderId(MaterialShaderMapHash, NULL, ShaderType, FShaderTarget(ShaderType->GetFrequency(), Platform));
+			FShaderId ShaderId(MaterialShaderMapHash, nullptr, nullptr, ShaderType, FShaderTarget(ShaderType->GetFrequency(), Platform));
 			FShader* FoundShader = ShaderType->FindShaderById(ShaderId);	
 
 			if (FoundShader)
@@ -1596,7 +1800,7 @@ void FMaterialShaderMap::LoadMissingShadersFromMemory(const FMaterial* Material)
 	}
 }
 
-void FMaterialShaderMap::GetShaderList(TMap<FShaderId,FShader*>& OutShaders) const
+void FMaterialShaderMap::GetShaderList(TMap<FShaderId, FShader*>& OutShaders) const
 {
 	TShaderMap<FMaterialShaderType>::GetShaderList(OutShaders);
 	for(int32 Index = 0;Index < MeshShaderMaps.Num();Index++)
@@ -1612,10 +1816,9 @@ void FMaterialShaderMap::Register(EShaderPlatform InShaderPlatform)
 {
 	if (GCreateShadersOnLoad && Platform == InShaderPlatform)
 	{
-		for (TMap<FShaderType*, TRefCountPtr<FShader> >::TConstIterator ShaderIt(GetShaders()); ShaderIt; ++ShaderIt)
+		for (auto KeyValue : GetShaders())
 		{
-			FShader* Shader = ShaderIt.Value();
-
+			FShader* Shader = KeyValue.Value;
 			if (Shader)
 			{
 				Shader->BeginInitializeResources();
@@ -1624,10 +1827,9 @@ void FMaterialShaderMap::Register(EShaderPlatform InShaderPlatform)
 
 		for (int32 Index = 0; Index < MeshShaderMaps.Num(); Index++)
 		{
-			for (TMap<FShaderType*, TRefCountPtr<FShader> >::TConstIterator ShaderIt(MeshShaderMaps[Index].GetShaders()); ShaderIt; ++ShaderIt)
+			for (TMap<FShaderType*, TRefCountPtr<FShader> >::TConstIterator It(MeshShaderMaps[Index].GetShaders()); It; ++It)
 			{
-				FShader* Shader = ShaderIt.Value();
-
+				FShader* Shader = It.Value();
 				if (Shader)
 				{
 					Shader->BeginInitializeResources();
@@ -1710,6 +1912,21 @@ void FMaterialShaderMap::FlushShadersByShaderType(FShaderType* ShaderType)
 	}
 }
 
+void FMaterialShaderMap::FlushShadersByShaderPipelineType(const FShaderPipelineType* ShaderPipelineType)
+{
+	// flush from all the vertex factory shader maps
+	for (int32 Index = 0;Index < MeshShaderMaps.Num();Index++)
+	{
+		MeshShaderMaps[Index].FlushShadersByShaderPipelineType(ShaderPipelineType);
+	}
+
+	if (ShaderPipelineType->IsMaterialTypePipeline())
+	{
+		RemoveShaderPipelineType(ShaderPipelineType);
+	}
+}
+
+
 /**
  * Removes all entries in the cache with exceptions based on a vertex factory type
  * @param ShaderType - The shader type to flush
@@ -1917,16 +2134,26 @@ const FMeshMaterialShaderMap* FMaterialShaderMap::GetMeshShaderMap(FVertexFactor
 	return MeshShaderMap;
 }
 
-uint32 FMaterialShaderMap::GetNumInstructionsForShader(const FShaderType* ShaderType) const
+uint32 FMaterialShaderMap::GetMaxNumInstructionsForShader(const FShaderType* ShaderType) const
 {
 	uint32 MaxNumInstructions = 0;
 	auto* FoundShader = Shaders.Find(ShaderType);
 	if (FoundShader && *FoundShader)
 	{
-		return (*FoundShader)->GetNumInstructions();
+		MaxNumInstructions = FMath::Max(MaxNumInstructions, (*FoundShader)->GetNumInstructions());
 	}
 
-	return 0;
+	for (auto& Pair : ShaderPipelines)
+	{
+		FShaderPipeline* Pipeline = Pair.Value;
+		auto* Shader = Pipeline->GetShader(ShaderType->GetFrequency());
+		if (Shader)
+		{
+			MaxNumInstructions = FMath::Max(MaxNumInstructions, Shader->GetNumInstructions());
+		}
+	}
+
+	return MaxNumInstructions;
 }
 
 void FMaterialShaderMap::InitOrderedMeshShaderMaps()
@@ -1947,11 +2174,13 @@ void FMaterialShaderMap::InitOrderedMeshShaderMaps()
  * 
  * @param	Platform	Platform to dump stats for.
  */
-void DumpMaterialStats( EShaderPlatform Platform )
+void DumpMaterialStats(EShaderPlatform Platform)
 {
 #if ALLOW_DEBUG_FILES
 	FDiagnosticTableViewer MaterialViewer(*FDiagnosticTableViewer::GetUniqueTemporaryFilePath(TEXT("MaterialStats")));
 
+	//#todo-rco: Pipelines
+
 	// Mapping from friendly material name to shaders associated with it.
 	TMultiMap<FString,FShader*> MaterialToShaderMap;
 	// Set of material names.
@@ -1961,18 +2190,18 @@ void DumpMaterialStats( EShaderPlatform Platform )
 	for (int32 ShaderMapIndex = 0; ShaderMapIndex < FMaterialShaderMap::AllMaterialShaderMaps.Num(); ShaderMapIndex++)
 	{
 		FMaterialShaderMap* MaterialShaderMap = FMaterialShaderMap::AllMaterialShaderMaps[ShaderMapIndex];
-		TMap<FShaderId,FShader*> Shaders;
-		MaterialShaderMap->GetShaderList( Shaders );
+		TMap<FShaderId, FShader*> Shaders;
+		MaterialShaderMap->GetShaderList(Shaders);
 
 		// Add friendly name to list of materials.
 		FString FriendlyName = MaterialShaderMap->GetFriendlyName();
-		MaterialNames.Add( FriendlyName );
+		MaterialNames.Add(FriendlyName);
 
 		// Add shaders to mapping per friendly name as there might be multiple
-		for( TMap<FShaderId,FShader*>::TConstIterator ShaderMapIt(Shaders); ShaderMapIt; ++ShaderMapIt )
+		for (auto& KeyValue : Shaders)
 		{
-			FShader* Shader = ShaderMapIt.Value();
-			MaterialToShaderMap.AddUnique(FriendlyName,Shader);
+			FShader* Shader = KeyValue.Value;
+			MaterialToShaderMap.AddUnique(FriendlyName, Shader);
 		}
 	}
 
diff --git a/Engine/Source/Runtime/Engine/Private/Materials/MaterialShared.cpp b/Engine/Source/Runtime/Engine/Private/Materials/MaterialShared.cpp
--- a/Engine/Source/Runtime/Engine/Private/Materials/MaterialShared.cpp
+++ b/Engine/Source/Runtime/Engine/Private/Materials/MaterialShared.cpp
@@ -262,14 +262,15 @@ void FMaterial::GetShaderMapId(EShaderPlatform Platform, FMaterialShaderMapId& O
 { 
 	TArray<FShaderType*> ShaderTypes;
 	TArray<FVertexFactoryType*> VFTypes;
+	TArray<const FShaderPipelineType*> ShaderPipelineTypes;
 
-	GetDependentShaderAndVFTypes(Platform, ShaderTypes, VFTypes);
+	GetDependentShaderAndVFTypes(Platform, ShaderTypes, ShaderPipelineTypes, VFTypes);
 
 	OutId.Usage = GetShaderMapUsage();
 	OutId.BaseMaterialId = GetMaterialId();
 	OutId.QualityLevel = GetQualityLevelForShaderMapId();
 	OutId.FeatureLevel = GetFeatureLevel();
-	OutId.SetShaderDependencies(ShaderTypes, VFTypes);
+	OutId.SetShaderDependencies(ShaderTypes, ShaderPipelineTypes, VFTypes);
 	GetReferencedTexturesHash(Platform, OutId.TextureReferencesHash);
 }
 
@@ -597,7 +598,7 @@ void FMaterial::SerializeInlineShaderMap(FArchive& Ar)
 		{
 			FinishCompilation();
 
-			bool bValid = GameThreadShaderMap != NULL && GameThreadShaderMap->CompiledSuccessfully();
+			bool bValid = GameThreadShaderMap != nullptr && GameThreadShaderMap->CompiledSuccessfully();
 			
 			Ar << bValid;
 
@@ -655,13 +656,13 @@ void FMaterial::ReleaseShaderMap()
 {
 	if (GameThreadShaderMap)
 	{
-		GameThreadShaderMap = NULL;
+		GameThreadShaderMap = nullptr;
 		
 		ENQUEUE_UNIQUE_RENDER_COMMAND_ONEPARAMETER(
 			ReleaseShaderMap,
 			FMaterial*,Material,this,
 		{
-			Material->SetRenderingThreadShaderMap(NULL);
+			Material->SetRenderingThreadShaderMap(nullptr);
 		});
 	}
 }
@@ -672,7 +673,7 @@ bool IsTranslucentBlendMode(EBlendMode BlendMode)
 }
 
 int32 FMaterialResource::GetMaterialDomain() const { return Material->MaterialDomain; }
-bool FMaterialResource::IsTangentSpaceNormal() const { return Material->bTangentSpaceNormal || (Material->Normal.Expression == NULL && !Material->bUseMaterialAttributes); }
+bool FMaterialResource::IsTangentSpaceNormal() const { return Material->bTangentSpaceNormal || (Material->Normal.Expression == nullptr && !Material->bUseMaterialAttributes); }
 bool FMaterialResource::ShouldInjectEmissiveIntoLPV() const { return Material->bUseEmissiveForDynamicAreaLighting; }
 bool FMaterialResource::ShouldBlockGI() const { return Material->bBlockGI; }
 bool FMaterialResource::ShouldGenerateSphericalParticleNormals() const { return Material->bGenerateSphericalParticleNormals; }
@@ -685,9 +686,9 @@ bool FMaterialResource::IsLightFunction() const { return Material->MaterialDomai
 bool FMaterialResource::IsUsedWithEditorCompositing() const { return Material->bUsedWithEditorCompositing; }
 bool FMaterialResource::IsUsedWithDeferredDecal() const { return Material->MaterialDomain == MD_DeferredDecal; }
 bool FMaterialResource::IsSpecialEngineMaterial() const { return Material->bUsedAsSpecialEngineMaterial; }
-bool FMaterialResource::HasVertexPositionOffsetConnected() const { return !Material->bUseMaterialAttributes && Material->WorldPositionOffset.Expression != NULL; }
-bool FMaterialResource::HasPixelDepthOffsetConnected() const { return !Material->bUseMaterialAttributes && Material->PixelDepthOffset.Expression != NULL; }
-bool FMaterialResource::HasMaterialAttributesConnected() const { return Material->bUseMaterialAttributes && Material->MaterialAttributes.Expression != NULL; }
+bool FMaterialResource::HasVertexPositionOffsetConnected() const { return !Material->bUseMaterialAttributes && Material->WorldPositionOffset.Expression != nullptr; }
+bool FMaterialResource::HasPixelDepthOffsetConnected() const { return !Material->bUseMaterialAttributes && Material->PixelDepthOffset.Expression != nullptr; }
+bool FMaterialResource::HasMaterialAttributesConnected() const { return Material->bUseMaterialAttributes && Material->MaterialAttributes.Expression != nullptr; }
 FString FMaterialResource::GetBaseMaterialPathName() const { return Material->GetPathName(); }
 
 bool FMaterialResource::IsUsedWithSkeletalMesh() const
@@ -911,7 +912,7 @@ void FMaterialResource::GetRepresentativeInstructionCounts(TArray<FString> &Desc
 			for (int32 InstructionIndex = 0; InstructionIndex < ShaderTypeNames.Num(); InstructionIndex++)
 			{
 				FShaderType* ShaderType = FindShaderTypeByName(*ShaderTypeNames[InstructionIndex]);
-				int32 NumInstructions = MaterialShaderMap->GetNumInstructionsForShader(ShaderType);
+				int32 NumInstructions = MaterialShaderMap->GetMaxNumInstructionsForShader(ShaderType);
 				if (NumInstructions > 0)
 				{
 					//if the shader was found, add it to the output arrays
@@ -933,7 +934,7 @@ void FMaterialResource::GetRepresentativeInstructionCounts(TArray<FString> &Desc
 					FShaderType* ShaderType = FindShaderTypeByName(*ShaderTypeNames[InstructionIndex]);
 					if (ShaderType)
 					{
-						int32 NumInstructions = MaterialShaderMap->GetNumInstructionsForShader(ShaderType);
+						int32 NumInstructions = MaterialShaderMap->GetMaxNumInstructionsForShader(ShaderType);
 						if (NumInstructions > 0)
 						{
 							//if the shader was found, add it to the output arrays
@@ -1034,7 +1035,7 @@ SIZE_T FMaterialResource::GetResourceSizeInclusive()
 {
 	SIZE_T ResourceSize = 0;
 	TSet<const FMaterialShaderMap*> UniqueShaderMaps;
-	TMap<FShaderId,FShader*> UniqueShaders;
+	TMap<FShaderId, FShader*> UniqueShaders;
 	TSet<FShaderResourceId> UniqueShaderResourceIds;
 
 	ResourceSize += sizeof(FMaterialResource);
@@ -1050,14 +1051,14 @@ SIZE_T FMaterialResource::GetResourceSizeInclusive()
 		}
 	}
 
-	for (TMap<FShaderId,FShader*>::TConstIterator It(UniqueShaders); It; ++It)
+	for (auto& KeyValue : UniqueShaders)
 	{
-		FShader* Shader = It.Value();
+		auto* Shader = KeyValue.Value;
 		if (Shader)
 		{
 			FShaderResourceId ResourceId = Shader->GetResourceId();
 			bool bCountedResource = false;
-			UniqueShaderResourceIds.Add(ResourceId,&bCountedResource);
+			UniqueShaderResourceIds.Add(ResourceId, &bCountedResource);
 			if (!bCountedResource)
 			{
 				ResourceSize += Shader->GetResourceSizeBytes();
@@ -1066,6 +1067,8 @@ SIZE_T FMaterialResource::GetResourceSizeInclusive()
 		}
 	}
 
+	//#todo-rco: Pipelines
+
 	return ResourceSize;
 }
 
@@ -1244,7 +1247,8 @@ void FMaterial::SetupMaterialEnvironment(
 		case MSM_TwoSidedFoliage:	OutEnvironment.SetDefine(TEXT("MATERIAL_SHADINGMODEL_TWOSIDED_FOLIAGE"),	TEXT("1")); break;
 		case MSM_Hair:				OutEnvironment.SetDefine(TEXT("MATERIAL_SHADINGMODEL_HAIR"),				TEXT("1")); break;
 		case MSM_Cloth:				OutEnvironment.SetDefine(TEXT("MATERIAL_SHADINGMODEL_CLOTH"),				TEXT("1")); break;
-		default: 
+		case MSM_Eye:				OutEnvironment.SetDefine(TEXT("MATERIAL_SHADINGMODEL_EYE"),					TEXT("1")); break;
+		default:
 			UE_LOG(LogMaterial, Warning, TEXT("Unknown material shading model: %u  Setting to MSM_DefaultLit"),(int32)GetShadingModel());
 			OutEnvironment.SetDefine(TEXT("MATERIAL_SHADINGMODEL_DEFAULT_LIT"),TEXT("1"));
 	};
@@ -1296,7 +1300,7 @@ bool FMaterial::CacheShaders(const FMaterialShaderMapId& ShaderMapId, EShaderPla
 	// If we loaded this material with inline shaders, use what was loaded (GameThreadShaderMap) instead of looking in the DDC
 	if (bContainsInlineShaders)
 	{
-		FMaterialShaderMap* ExistingShaderMap = NULL;
+		FMaterialShaderMap* ExistingShaderMap = nullptr;
 		
 		if (GameThreadShaderMap)
 		{
@@ -1332,11 +1336,10 @@ bool FMaterial::CacheShaders(const FMaterialShaderMapId& ShaderMapId, EShaderPla
 	{
 		OutstandingCompileShaderMapIds.Add(GameThreadShaderMap->GetCompilingId());
 		// Reset the shader map so the default material will be used until the compile finishes.
-		GameThreadShaderMap = NULL;
+		GameThreadShaderMap = nullptr;
 		bSucceeded = true;
 	}
-	else
-	if (!GameThreadShaderMap || !GameThreadShaderMap->IsComplete(this, true))
+	else if (!GameThreadShaderMap || !GameThreadShaderMap->IsComplete(this, true))
 	{
 		if (bContainsInlineShaders || FPlatformProperties::RequiresCookedData())
 		{
@@ -1351,7 +1354,7 @@ bool FMaterial::CacheShaders(const FMaterialShaderMapId& ShaderMapId, EShaderPla
 			}
 
 			// Reset the shader map so the default material will be used.
-			GameThreadShaderMap = NULL;
+			GameThreadShaderMap = nullptr;
 		}
 		else
 		{
@@ -1373,7 +1376,7 @@ bool FMaterial::CacheShaders(const FMaterialShaderMapId& ShaderMapId, EShaderPla
 			if (!bSucceeded)
 			{
 				// If it failed to compile the material, reset the shader map so the material isn't used.
-				GameThreadShaderMap = NULL;
+				GameThreadShaderMap = nullptr;
 
 				if (IsDefaultMaterial())
 				{
@@ -1444,13 +1447,13 @@ bool FMaterial::BeginCompileShaderMap(
 		if (bSynchronousCompile)
 		{
 			// If this is a synchronous compile, assign the compile result to the output
-			OutShaderMap = NewShaderMap->CompiledSuccessfully() ? NewShaderMap : NULL;
+			OutShaderMap = NewShaderMap->CompiledSuccessfully() ? NewShaderMap : nullptr;
 		}
 		else
 		{
 			OutstandingCompileShaderMapIds.Add( NewShaderMap->GetCompilingId() );
 			// Async compile, use NULL so that rendering will fall back to the default material.
-			OutShaderMap = NULL;
+			OutShaderMap = nullptr;
 		}
 	}
 
@@ -1494,7 +1497,7 @@ const FMaterial* FColoredMaterialRenderProxy::GetMaterial(ERHIFeatureLevel::Type
 FShader* FMaterial::GetShader(FMeshMaterialShaderType* ShaderType, FVertexFactoryType* VertexFactoryType) const
 {
 	const FMeshMaterialShaderMap* MeshShaderMap = RenderingThreadShaderMap->GetMeshShaderMap(VertexFactoryType);
-	FShader* Shader = MeshShaderMap ? MeshShaderMap->GetShader(ShaderType) : NULL;
+	FShader* Shader = MeshShaderMap ? MeshShaderMap->GetShader(ShaderType) : nullptr;
 	if (!Shader)
 	{
 		// Get the ShouldCache results that determine whether the shader should be compiled
@@ -1510,16 +1513,12 @@ FShader* FMaterial::GetShader(FMeshMaterialShaderType* ShaderType, FVertexFactor
 		// This is usually the result of an incorrect ShouldCache function.
 		UE_LOG(LogMaterial, Fatal,
 			TEXT("Couldn't find Shader %s for Material Resource %s!\n")
-			TEXT("		With VF=%s, Platform=%s \n")
+			TEXT("		With VF=%s, Platform=%s\n")
 			TEXT("		ShouldCache: Mat=%u, VF=%u, Shader=%u \n")
 			TEXT("		MaterialUsageDesc: %s"),
-			ShaderType->GetName(), 
-			*GetFriendlyName(),
-			VertexFactoryType->GetName(),
-			*LegacyShaderPlatformToShaderFormat(ShaderPlatform).ToString(),
-			bMaterialShouldCache,
-			bVFShouldCache,
-			bShaderShouldCache,
+			ShaderType->GetName(),  *GetFriendlyName(),
+			VertexFactoryType->GetName(), *LegacyShaderPlatformToShaderFormat(ShaderPlatform).ToString(),
+			bMaterialShouldCache, bVFShouldCache, bShaderShouldCache,
 			*MaterialUsage
 			);
 	}
@@ -1607,7 +1606,7 @@ void FMaterialRenderProxy::CacheUniformExpressions()
 
 			if (!bIsFallbackMaterial)
 			{
-				FMaterialRenderContext MaterialRenderContext(this, *Material, NULL);
+				FMaterialRenderContext MaterialRenderContext(this, *Material, nullptr);
 				MaterialRenderContext.bShowSelection = GIsEditor;
 				EvaluateUniformExpressions(UniformExpressionCache[(int32)FeatureLevel], MaterialRenderContext);
 			}
@@ -1645,7 +1644,7 @@ void FMaterialRenderProxy::InvalidateUniformExpressionCache()
 	{
 		UniformExpressionCache[i].bUpToDate = false;
 		UniformExpressionCache[i].UniformBuffer.SafeRelease();
-		UniformExpressionCache[i].CachedUniformExpressionShaderMap = NULL;
+		UniformExpressionCache[i].CachedUniformExpressionShaderMap = nullptr;
 	}
 }
 
@@ -1808,7 +1807,7 @@ FString FMaterialResource::GetMaterialUsageDescription() const
 	return BaseDescription;
 }
 
-void FMaterial::GetDependentShaderAndVFTypes(EShaderPlatform Platform, TArray<FShaderType*>& OutShaderTypes, TArray<FVertexFactoryType*>& OutVFTypes) const
+void FMaterial::GetDependentShaderAndVFTypes(EShaderPlatform Platform, TArray<FShaderType*>& OutShaderTypes, TArray<const FShaderPipelineType*>& OutShaderPipelineTypes, TArray<FVertexFactoryType*>& OutVFTypes) const
 {
 	// Iterate over all vertex factory types.
 	for (TLinkedList<FVertexFactoryType*>::TIterator VertexFactoryTypeIt(FVertexFactoryType::GetTypeList()); VertexFactoryTypeIt; VertexFactoryTypeIt.Next())
@@ -1826,7 +1825,6 @@ void FMaterial::GetDependentShaderAndVFTypes(EShaderPlatform Platform, TArray<FS
 				FMeshMaterialShaderType* ShaderType = ShaderTypeIt->GetMeshMaterialShaderType();
 
 				if (ShaderType && 
-					VertexFactoryType && 
 					ShaderType->ShouldCache(Platform, this, VertexFactoryType) && 
 					ShouldCache(Platform, ShaderType, VertexFactoryType) &&
 					VertexFactoryType->ShouldCache(Platform, this, ShaderType)
@@ -1837,6 +1835,27 @@ void FMaterial::GetDependentShaderAndVFTypes(EShaderPlatform Platform, TArray<FS
 				}
 			}
 
+			for (TLinkedList<FShaderPipelineType*>::TIterator PipelineTypeIt(FShaderPipelineType::GetTypeList()); PipelineTypeIt; PipelineTypeIt.Next())
+			{
+				auto* PipelineType = *PipelineTypeIt;
+				if (PipelineType->IsMeshMaterialTypePipeline())
+				{
+					for (const FShaderType* Type : PipelineType->GetStages())
+					{
+						const FMeshMaterialShaderType* ShaderType = Type->GetMeshMaterialShaderType();
+						if (ShaderType->ShouldCache(Platform, this, VertexFactoryType) &&
+							ShouldCache(Platform, ShaderType, VertexFactoryType) &&
+							VertexFactoryType->ShouldCache(Platform, this, ShaderType)
+							)
+						{
+							bAddedTypeFromThisVF = true;
+							OutShaderPipelineTypes.AddUnique(PipelineType);
+							OutShaderTypes.AddUnique((FShaderType*)Type);
+						}
+					}
+				}
+			}
+
 			if (bAddedTypeFromThisVF)
 			{
 				OutVFTypes.Add(VertexFactoryType);
@@ -1851,16 +1870,37 @@ void FMaterial::GetDependentShaderAndVFTypes(EShaderPlatform Platform, TArray<FS
 
 		if (ShaderType && 
 			ShaderType->ShouldCache(Platform, this) && 
-			ShouldCache(Platform, ShaderType, NULL)
+			ShouldCache(Platform, ShaderType, nullptr)
 			)
 		{
 			OutShaderTypes.Add(ShaderType);
 		}
 	}
 
+	for (TLinkedList<FShaderPipelineType*>::TIterator PipelineTypeIt(FShaderPipelineType::GetTypeList()); PipelineTypeIt; PipelineTypeIt.Next())
+	{
+		auto* PipelineType = *PipelineTypeIt;
+		if (PipelineType->IsMaterialTypePipeline())
+		{
+			for (const FShaderType* Type : PipelineType->GetStages())
+			{
+				const FMaterialShaderType* ShaderType = Type->GetMaterialShaderType();
+				if (ShaderType &&
+					ShaderType->ShouldCache(Platform, this) &&
+					ShouldCache(Platform, ShaderType, nullptr)
+					)
+				{
+					OutShaderPipelineTypes.AddUnique(PipelineType);
+					OutShaderTypes.AddUnique((FShaderType*)Type);
+				}
+			}
+		}
+	}
+
 	// Sort by name so that we get deterministic keys
 	OutShaderTypes.Sort(FCompareShaderTypes());
 	OutVFTypes.Sort(FCompareVertexFactoryTypes());
+	OutShaderPipelineTypes.Sort(FCompareShaderPipelineNameTypes());
 }
 
 void FMaterial::GetReferencedTexturesHash(EShaderPlatform Platform, FSHAHash& OutHash) const
@@ -2441,4 +2481,4 @@ bool FMaterialShaderMapId::ContainsVertexFactoryType(const FVertexFactoryType* V
 	}
 
 	return false;
-}
\ No newline at end of file
+}
diff --git a/Engine/Source/Runtime/Engine/Private/Materials/MeshMaterialShader.cpp b/Engine/Source/Runtime/Engine/Private/Materials/MeshMaterialShader.cpp
--- a/Engine/Source/Runtime/Engine/Private/Materials/MeshMaterialShader.cpp
+++ b/Engine/Source/Runtime/Engine/Private/Materials/MeshMaterialShader.cpp
@@ -8,11 +8,20 @@
 #include "ShaderCompiler.h"
 #include "MeshMaterialShaderType.h"
 
+
+static inline bool ShouldCacheMeshShader(const FMeshMaterialShaderType* ShaderType, EShaderPlatform Platform, const FMaterial* Material, FVertexFactoryType* InVertexFactoryType)
+{
+	return ShaderType->ShouldCache(Platform, Material, InVertexFactoryType) &&
+		Material->ShouldCache(Platform, ShaderType, InVertexFactoryType) &&
+		InVertexFactoryType->ShouldCache(Platform, Material, ShaderType);
+}
+
 /**
 * Finds a FMeshMaterialShaderType by name.
 */
 FMeshMaterialShaderType* FMeshMaterialShaderType::GetTypeByName(const FString& TypeName)
 {
+	//#todo-rco: Unused?
 	for(TLinkedList<FShaderType*>::TIterator It(FShaderType::GetTypeList()); It; It.Next())
 	{
 		FString CurrentTypeName = FString(It->GetName());
@@ -37,7 +46,8 @@ void FMeshMaterialShaderType::BeginCompileShader(
 	const FMaterial* Material,
 	FShaderCompilerEnvironment* MaterialEnvironment,
 	FVertexFactoryType* VertexFactoryType,
-	TArray<FShaderCompileJob*>& NewJobs
+	const FShaderPipelineType* ShaderPipeline,
+	TArray<FShaderCommonCompileJob*>& NewJobs
 	)
 {
 	FShaderCompileJob* NewJob = new FShaderCompileJob(ShaderMapId, VertexFactoryType, this);
@@ -73,6 +83,31 @@ void FMeshMaterialShaderType::BeginCompileShader(
 		);
 }
 
+void FMeshMaterialShaderType::BeginCompileShaderPipeline(
+	uint32 ShaderMapId,
+	EShaderPlatform Platform,
+	const FMaterial* Material,
+	FShaderCompilerEnvironment* MaterialEnvironment,
+	FVertexFactoryType* VertexFactoryType,
+	const FShaderPipelineType* ShaderPipeline,
+	const TArray<FMeshMaterialShaderType*>& ShaderStages,
+	TArray<FShaderCommonCompileJob*>& NewJobs)
+{
+	check(ShaderStages.Num() > 0);
+	check(ShaderPipeline);
+	UE_LOG(LogShaders, Verbose, TEXT("	Pipeline: %s"), ShaderPipeline->GetName());
+
+	// Add all the jobs as individual first, then add the dependencies into a pipeline job
+	auto* NewPipelineJob = new FShaderPipelineCompileJob(ShaderMapId, ShaderPipeline, ShaderStages.Num());
+	for (int32 Index = 0; Index < ShaderStages.Num(); ++Index)
+	{
+		auto* ShaderStage = ShaderStages[Index];
+		ShaderStage->BeginCompileShader(ShaderMapId, Platform, Material, MaterialEnvironment, VertexFactoryType, ShaderPipeline, NewPipelineJob->StageJobs);
+	}
+
+	NewJobs.Add(NewPipelineJob);
+}
+
 /**
  * Either creates a new instance of this type or returns an equivalent existing shader.
  * @param Material - The material to link the shader with.
@@ -82,6 +117,7 @@ FShader* FMeshMaterialShaderType::FinishCompileShader(
 	const FUniformExpressionSet& UniformExpressionSet, 
 	const FSHAHash& MaterialShaderMapHash,
 	const FShaderCompileJob& CurrentJob,
+	const FShaderPipelineType* ShaderPipeline,
 	const FString& InDebugDescription)
 {
 	check(CurrentJob.bSucceeded);
@@ -94,12 +130,12 @@ FShader* FMeshMaterialShaderType::FinishCompileShader(
 	FShaderResource* Resource = FShaderResource::FindOrCreateShaderResource(CurrentJob.Output, SpecificType);
 
 	// Find a shader with the same key in memory
-	FShader* Shader = CurrentJob.ShaderType->FindShaderById(FShaderId(MaterialShaderMapHash, CurrentJob.VFType, CurrentJob.ShaderType, CurrentJob.Input.Target));
+	FShader* Shader = CurrentJob.ShaderType->FindShaderById(FShaderId(MaterialShaderMapHash, ShaderPipeline, CurrentJob.VFType, CurrentJob.ShaderType, CurrentJob.Input.Target));
 
 	// There was no shader with the same key so create a new one with the compile output, which will bind shader parameters
 	if (!Shader)
 	{
-		Shader = (*ConstructCompiledRef)(CompiledShaderInitializerType(this, CurrentJob.Output, Resource, UniformExpressionSet, MaterialShaderMapHash, InDebugDescription, CurrentJob.VFType));
+		Shader = (*ConstructCompiledRef)(CompiledShaderInitializerType(this, CurrentJob.Output, Resource, UniformExpressionSet, MaterialShaderMapHash, InDebugDescription, ShaderPipeline, CurrentJob.VFType));
 		CurrentJob.Output.ParameterMap.VerifyBindingsAreComplete(GetName(), (EShaderFrequency)CurrentJob.Output.Target.Frequency, CurrentJob.VFType);
 	}
 
@@ -118,21 +154,70 @@ uint32 FMeshMaterialShaderMap::BeginCompile(
 	const FMaterial* Material,
 	FShaderCompilerEnvironment* MaterialEnvironment,
 	EShaderPlatform Platform,
-	TArray<FShaderCompileJob*>& NewJobs
+	TArray<FShaderCommonCompileJob*>& NewJobs
 	)
 {
+	if (!VertexFactoryType)
+	{
+		return 0;
+	}
+
 	uint32 NumShadersPerVF = 0;
 
+	for (TLinkedList<FShaderPipelineType*>::TIterator ShaderPipelineIt(FShaderPipelineType::GetTypeList());ShaderPipelineIt;ShaderPipelineIt.Next())
+	{
+		const FShaderPipelineType* Pipeline = *ShaderPipelineIt;
+		if (Pipeline->IsMeshMaterialTypePipeline())
+		{
+			auto& Stages = Pipeline->GetStages();
+
+			// Verify all the ShouldCache are in sync
+			int32 NumShouldCache = 0;
+			for (int32 Index = 0; Index < Stages.Num(); ++Index)
+			{
+				auto* ShaderType = Stages[Index]->GetMeshMaterialShaderType();
+				if (ShouldCacheMeshShader(ShaderType, Platform, Material, VertexFactoryType))
+				{
+					++NumShouldCache;
+				}
+			}
+
+			if (NumShouldCache > 0)
+			{
+				NumShadersPerVF += Stages.Num();
+				if (NumShouldCache != Stages.Num())
+				{
+					//#todo-rco: Error or Warning?
+					UE_LOG(LogShaders, Error, TEXT("Mismatched ShouldCache() function on the shaders for MeshMaterial ShaderPipeline %s"), Pipeline->GetName());
+				}
+
+				TArray<FMeshMaterialShaderType*> ShaderStagesToCompile;
+				ShaderStagesToCompile.Reserve(SF_NumFrequencies);
+				for (int32 Index = 0; Index < Stages.Num(); ++Index)
+				{
+					FMeshMaterialShaderType* ShaderType = (FMeshMaterialShaderType*)(Stages[Index]->GetMeshMaterialShaderType());
+
+					// No need to ShouldCache again as we tested each stage beforehand
+
+					// Verify that the shader map Id contains inputs for any shaders that will be put into this shader map
+					check(InShaderMapId.ContainsVertexFactoryType(VertexFactoryType));
+					check(InShaderMapId.ContainsShaderType(ShaderType));
+					ShaderStagesToCompile.Add(ShaderType);
+				}
+
+				//#todo-rco: Check if the MeshShaderMap already contains this pipeline
+
+				// Make a pipeline job with all the stages
+				FMeshMaterialShaderType::BeginCompileShaderPipeline(ShaderMapId, Platform, Material, MaterialEnvironment, VertexFactoryType, Pipeline, ShaderStagesToCompile, NewJobs);
+			}
+		}
+	}
+
 	// Iterate over all mesh material shader types.
 	for(TLinkedList<FShaderType*>::TIterator ShaderTypeIt(FShaderType::GetTypeList());ShaderTypeIt;ShaderTypeIt.Next())
 	{
 		FMeshMaterialShaderType* ShaderType = ShaderTypeIt->GetMeshMaterialShaderType();
-		if (ShaderType && 
-			VertexFactoryType && 
-			ShaderType->ShouldCache(Platform, Material, VertexFactoryType) && 
-			Material->ShouldCache(Platform, ShaderType, VertexFactoryType) &&
-			VertexFactoryType->ShouldCache(Platform, Material, ShaderType)
-			)
+		if (ShaderType && ShouldCacheMeshShader(ShaderType, Platform, Material, VertexFactoryType))
 		{
 			// Verify that the shader map Id contains inputs for any shaders that will be put into this shader map
 			check(InShaderMapId.ContainsVertexFactoryType(VertexFactoryType));
@@ -149,6 +234,7 @@ uint32 FMeshMaterialShaderMap::BeginCompile(
 					Material,
 					MaterialEnvironment,
 					VertexFactoryType,
+					nullptr,
 					NewJobs
 					);
 			}
@@ -170,6 +256,7 @@ uint32 FMeshMaterialShaderMap::BeginCompile(
  */
 void FMeshMaterialShaderMap::FinishCompile(uint32 ShaderMapId, const FUniformExpressionSet& UniformExpressionSet, const FSHAHash& MaterialShaderMapHash, const TArray<FShaderCompileJob*>& CompilationResults, const FString& InDebugDescription)
 {
+	//#todo-rco: Unused?
 	// Find the matching FMeshMaterialShaderType for each compile job
 	for (int32 JobIndex = 0; JobIndex < CompilationResults.Num(); JobIndex++)
 	{
@@ -181,7 +268,8 @@ void FMeshMaterialShaderMap::FinishCompile(uint32 ShaderMapId, const FUniformExp
 				FMeshMaterialShaderType* MeshMaterialShaderType = ShaderTypeIt->GetMeshMaterialShaderType();
 				if (*ShaderTypeIt == CurrentJob.ShaderType && MeshMaterialShaderType != NULL)
 				{
-					FShader* Shader = MeshMaterialShaderType->FinishCompileShader(UniformExpressionSet, MaterialShaderMapHash, CurrentJob, InDebugDescription);
+					//#todo-rco: Pipeline?
+					FShader* Shader = MeshMaterialShaderType->FinishCompileShader(UniformExpressionSet, MaterialShaderMapHash, CurrentJob, nullptr, InDebugDescription);
 					check(Shader);
 					AddShader(MeshMaterialShaderType,Shader);
 				}
@@ -198,29 +286,81 @@ bool FMeshMaterialShaderMap::IsComplete(
 	bool bSilent
 	)
 {
-	bool bIsComplete = true;
+	auto IsMeshShaderComplete = [&](const FMeshMaterialShaderType* ShaderType, const FShaderPipelineType* Pipeline) -> bool
+		{
+			// If we should cache this shader then the map is empty IF
+			//		The shadermap is empty
+			//		OR If it doesn't have a pipeline and needs one
+			//		OR it's not in the shadermap
+			//#todo-rco: Remove cast
+			if (ShouldCacheMeshShader(ShaderType, Platform, Material, InVertexFactoryType) &&
+				(!MeshShaderMap || (Pipeline && !MeshShaderMap->GetShaderPipeline(Pipeline)) || (!Pipeline && !MeshShaderMap->HasShader((FShaderType*)ShaderType))))
+			{
+				if (!bSilent)
+				{
+					if (Pipeline)
+					{
+						UE_LOG(LogShaders, Warning, TEXT("Incomplete material %s, missing pipeline %s from %s."), *Material->GetFriendlyName(), Pipeline->GetName(), InVertexFactoryType->GetName());
+					}
+					else
+					{
+						UE_LOG(LogShaders, Warning, TEXT("Incomplete material %s, missing %s from %s."), *Material->GetFriendlyName(), ShaderType->GetName(), InVertexFactoryType->GetName());
+					}
+				}
+				return false;
+			}
 
-	// Iterate over all mesh material shader types.
-	for(TLinkedList<FShaderType*>::TIterator ShaderTypeIt(FShaderType::GetTypeList());ShaderTypeIt;ShaderTypeIt.Next())
+			return true;
+		};
+
+	// Iterate over all pipeline types
+	for (TLinkedList<FShaderPipelineType*>::TIterator ShaderPipelineIt(FShaderPipelineType::GetTypeList());ShaderPipelineIt;ShaderPipelineIt.Next())
 	{
-		FMeshMaterialShaderType* ShaderType = ShaderTypeIt->GetMeshMaterialShaderType();
-		if (ShaderType && 
-			(!MeshShaderMap || !MeshShaderMap->HasShader(ShaderType)) &&
-			ShaderType->ShouldCache(Platform, Material, InVertexFactoryType) && 
-			Material->ShouldCache(Platform, ShaderType, InVertexFactoryType) &&
-			InVertexFactoryType->ShouldCache(Platform, Material, ShaderType)
-			)
+		const FShaderPipelineType* ShaderPipelineType = *ShaderPipelineIt;
+		if (ShaderPipelineType->IsMeshMaterialTypePipeline())
 		{
-			if (!bSilent)
+			auto& Stages = ShaderPipelineType->GetStages();
+
+			// Verify all the ShouldCache are in sync
+			int32 NumShouldCache = 0;
+			for (int32 Index = 0; Index < Stages.Num(); ++Index)
+			{
+				auto* ShaderType = Stages[Index]->GetMeshMaterialShaderType();
+				if (ShouldCacheMeshShader(ShaderType, Platform, Material, InVertexFactoryType))
+				{
+					++NumShouldCache;
+				}
+			}
+
+			if (NumShouldCache > 0 && NumShouldCache != Stages.Num())
 			{
-				UE_LOG(LogShaders, Warning, TEXT("Incomplete material %s, missing %s from %s."), *Material->GetFriendlyName(), ShaderType->GetName(), InVertexFactoryType->GetName());
+				//#todo-rco: Error or Warning?
+				UE_LOG(LogShaders, Error, TEXT("Mismatched ShouldCache() function on the shaders for MeshMaterial ShaderPipeline %s"), ShaderPipelineType->GetName());
 			}
-			bIsComplete = false;
-			break;
+
+			// Now check the completeness of the shader map
+			for (int32 Index = 0; Index < Stages.Num(); ++Index)
+			{
+				auto* ShaderType = Stages[Index]->GetMeshMaterialShaderType();
+				if (ShaderType && !IsMeshShaderComplete(ShaderType, ShaderPipelineType))
+				{
+					return false;
+				}
+			}
+		}
+	}
+
+	// Iterate over all mesh material shader types.
+	for (TLinkedList<FShaderType*>::TIterator ShaderTypeIt(FShaderType::GetTypeList());ShaderTypeIt;ShaderTypeIt.Next())
+	{
+		FMeshMaterialShaderType* ShaderType = ShaderTypeIt->GetMeshMaterialShaderType();
+		if (ShaderType && !IsMeshShaderComplete(ShaderType, nullptr))
+		{
+			return false;
 		}
 	}
 
-	return bIsComplete;
+	return true;
 }
 
 void FMeshMaterialShaderMap::LoadMissingShadersFromMemory(
@@ -228,18 +368,16 @@ void FMeshMaterialShaderMap::LoadMissingShadersFromMemory(
 	const FMaterial* Material, 
 	EShaderPlatform Platform)
 {
+	//#todo-rco: Pipelines
 	for(TLinkedList<FShaderType*>::TIterator ShaderTypeIt(FShaderType::GetTypeList());ShaderTypeIt;ShaderTypeIt.Next())
 	{
 		FMeshMaterialShaderType* ShaderType = ShaderTypeIt->GetMeshMaterialShaderType();
+		//#todo-rco: Pass FNullShaderPipeline instead of nullptr
 		const bool bShaderAlreadyExists = HasShader(ShaderType);
 
-		if (ShaderType && 
-			ShaderType->ShouldCache(Platform, Material, VertexFactoryType) && 
-			Material->ShouldCache(Platform, ShaderType, VertexFactoryType) &&
-			VertexFactoryType->ShouldCache(Platform, Material, ShaderType) &&
-			!bShaderAlreadyExists) 
+		if (ShaderType && ShouldCacheMeshShader(ShaderType, Platform, Material, VertexFactoryType) && !bShaderAlreadyExists) 
 		{
-			const FShaderId ShaderId(MaterialShaderMapHash, VertexFactoryType, ShaderType, FShaderTarget(ShaderType->GetFrequency(), Platform));
+			const FShaderId ShaderId(MaterialShaderMapHash, nullptr, VertexFactoryType, ShaderType, FShaderTarget(ShaderType->GetFrequency(), Platform));
 			FShader* FoundShader = ShaderType->FindShaderById(ShaderId);	
 
 			if (FoundShader)
@@ -258,7 +396,14 @@ void FMeshMaterialShaderMap::FlushShadersByShaderType(FShaderType* ShaderType)
 {
 	if (ShaderType->GetMeshMaterialShaderType())
 	{
-		RemoveShaderType(ShaderType->GetMeshMaterialShaderType());
+		RemoveShaderType(ShaderType);
 	}
 }
 
+void FMeshMaterialShaderMap::FlushShadersByShaderPipelineType(const FShaderPipelineType* ShaderPipelineType)
+{
+	if (ShaderPipelineType->IsMeshMaterialTypePipeline())
+	{
+		RemoveShaderPipelineType(ShaderPipelineType);
+	}
+}
diff --git a/Engine/Source/Runtime/Engine/Private/NetworkDriver.cpp b/Engine/Source/Runtime/Engine/Private/NetworkDriver.cpp
--- a/Engine/Source/Runtime/Engine/Private/NetworkDriver.cpp
+++ b/Engine/Source/Runtime/Engine/Private/NetworkDriver.cpp
@@ -203,13 +203,51 @@
 	// Reset queued bunch amortization timer
 	ProcessQueuedBunchesCurrentFrameMilliseconds = 0.0f;
 
-#if STATS
+#if USE_SERVER_PERF_COUNTERS || STATS
 	const double CurrentRealtimeSeconds = FPlatformTime::Seconds();
 
-	// Update network stats (only main game net driver for now)
+	// Update network stats (only main game net driver for now) if stats or perf counters are used
 	if (NetDriverName == NAME_GameNetDriver && 
 		 CurrentRealtimeSeconds - StatUpdateTime > StatPeriod && ( ClientConnections.Num() > 0 || ServerConnection != NULL ) )
 	{
+		int32 ClientInBytesMax = 0;
+		int32 ClientInBytesMin = 0;
+		int32 ClientInBytesAvg = 0;
+		int32 ClientOutBytesMax = 0;
+		int32 ClientOutBytesMin = 0;
+		int32 ClientOutBytesAvg = 0;
+		int NumClients = 0;
+
+		// these need to be updated even if we are not collecting stats, since they get reported to analytics/QoS
+		for (UNetConnection * Client : ClientConnections)
+		{
+			if (Client)
+			{
+				ClientInBytesMax = FMath::Max(ClientInBytesMax, Client->InBytesPerSecond);
+				if (ClientInBytesMin == 0 || Client->InBytesPerSecond < ClientInBytesMin)
+				{
+					ClientInBytesMin = Client->InBytesPerSecond;
+				}
+				ClientInBytesAvg += Client->InBytesPerSecond;
+
+				ClientOutBytesMax = FMath::Max(ClientOutBytesMax, Client->OutBytesPerSecond);
+				if (ClientOutBytesMin == 0 || Client->OutBytesPerSecond < ClientOutBytesMin)
+				{
+					ClientOutBytesMin = Client->OutBytesPerSecond;
+				}
+				ClientOutBytesAvg += Client->OutBytesPerSecond;
+
+				++NumClients;
+			}
+		}
+
+		if (NumClients > 1)
+		{
+			ClientInBytesAvg /= NumClients;
+			ClientOutBytesAvg /= NumClients;
+		}
+
+#if STATS
 		int32 Ping = 0;
 		int32 NumOpenChannels = 0;
 		int32 NumActorChannels = 0;
@@ -220,13 +258,6 @@
 		int32 UnAckCount = 0;
 		int32 PendingCount = 0;
 		int32 NetSaturated = 0;
-		int32 ClientInBytesMax = 0;
-		int32 ClientInBytesMin = 0;
-		int32 ClientInBytesAvg = 0;
-		int32 ClientOutBytesMax = 0;
-		int32 ClientOutBytesMin = 0;
-		int32 ClientOutBytesAvg = 0;
-		int NumClients = 0;
 
 		if (FThreadStats::IsCollectingData())
 		{
@@ -302,42 +333,56 @@
 
 				NetSaturated = Connection->IsNetReady(false) ? 0 : 1;
 			}
-
-			for (UNetConnection * Client : ClientConnections)
-			{
-				if (Client)
-				{
-					ClientInBytesMax = FMath::Max(ClientInBytesMax, Client->InBytesPerSecond);
-					if (ClientInBytesMin == 0 || Client->InBytesPerSecond < ClientInBytesMin)
-					{
-						ClientInBytesMin = Client->InBytesPerSecond;
-					}
-					ClientInBytesAvg += Client->InBytesPerSecond;
-
-					ClientOutBytesMax = FMath::Max(ClientOutBytesMax, Client->OutBytesPerSecond);
-					if (ClientOutBytesMin == 0 || Client->OutBytesPerSecond < ClientOutBytesMin)
-					{
-						ClientOutBytesMin = Client->OutBytesPerSecond;
-					}
-					ClientOutBytesAvg += Client->OutBytesPerSecond;
-
-					++NumClients;
-				}
-			}
-
-			if (NumClients > 1)
-			{
-				ClientInBytesAvg /= NumClients;
-				ClientOutBytesAvg /= NumClients;
-			}
 		}
 
+		// Copy the net status values over
+		SET_DWORD_STAT(STAT_Ping, Ping);
+		SET_DWORD_STAT(STAT_Channels, NumOpenChannels);
+
+		SET_DWORD_STAT(STAT_OutLoss, OutPacketsLost);
+		SET_DWORD_STAT(STAT_InLoss, InPacketsLost);
+		SET_DWORD_STAT(STAT_InRate, InBytes);
+		SET_DWORD_STAT(STAT_OutRate, OutBytes);
+		SET_DWORD_STAT(STAT_InRateClientMax, ClientInBytesMax);
+		SET_DWORD_STAT(STAT_InRateClientMin, ClientInBytesMin);
+		SET_DWORD_STAT(STAT_InRateClientAvg, ClientInBytesAvg);
+		SET_DWORD_STAT(STAT_OutRateClientMax, ClientOutBytesMax);
+		SET_DWORD_STAT(STAT_OutRateClientMin, ClientOutBytesMin);
+		SET_DWORD_STAT(STAT_OutRateClientAvg, ClientOutBytesAvg);
+
+		SET_DWORD_STAT(STAT_NetNumClients, NumClients);
+		SET_DWORD_STAT(STAT_InPackets, InPackets);
+		SET_DWORD_STAT(STAT_OutPackets, OutPackets);
+		SET_DWORD_STAT(STAT_InBunches, InBunches);
+		SET_DWORD_STAT(STAT_OutBunches, OutBunches);
+
+		SET_DWORD_STAT(STAT_NetGUIDInRate, NetGUIDInBytes);
+		SET_DWORD_STAT(STAT_NetGUIDOutRate, NetGUIDOutBytes);
+
+		SET_DWORD_STAT(STAT_VoicePacketsSent, VoicePacketsSent);
+		SET_DWORD_STAT(STAT_VoicePacketsRecv, VoicePacketsRecv);
+		SET_DWORD_STAT(STAT_VoiceBytesSent, VoiceBytesSent);
+		SET_DWORD_STAT(STAT_VoiceBytesRecv, VoiceBytesRecv);
+
+		SET_DWORD_STAT(STAT_PercentInVoice, VoiceInPercent);
+		SET_DWORD_STAT(STAT_PercentOutVoice, VoiceOutPercent);
+
+		SET_DWORD_STAT(STAT_NumActorChannels, NumActorChannels);
+		SET_DWORD_STAT(STAT_NumDormantActors, NumDormantActors);
+		SET_DWORD_STAT(STAT_NumActors, NumActors);
+		SET_DWORD_STAT(STAT_NumActorChannelsReadyDormant, NumActorChannelsReadyDormant);
+		SET_DWORD_STAT(STAT_NumNetGUIDsAckd, AckCount);
+		SET_DWORD_STAT(STAT_NumNetGUIDsPending, UnAckCount);
+		SET_DWORD_STAT(STAT_NumNetGUIDsUnAckd, PendingCount);
+		SET_DWORD_STAT(STAT_NetSaturated, NetSaturated);
+#endif // STATS
+
 #if USE_SERVER_PERF_COUNTERS
 		IPerfCounters* PerfCounters = IPerfCountersModule::Get().GetPerformanceCounters();
 		if (PerfCounters)
 		{
 			// Update total connections
-			PerfCounters->Set(TEXT("NumConnections"), ClientConnections.Num(), IPerfCounters::Flags::Transient);
+			PerfCounters->Set(TEXT("NumConnections"), ClientConnections.Num());
 
 			const int kNumBuckets = 8;	// evenly spaced with increment of 30 ms; last bucket collects all off-scale pings as well
 			if (ClientConnections.Num() > 0)
@@ -380,7 +425,7 @@
 					}
 				}
 
-				PerfCounters->Set(TEXT("AvgPing"), AvgPing / PingCount);
+				PerfCounters->Set(TEXT("AvgPing"), AvgPing / PingCount, IPerfCounters::Flags::Transient);
 				float CurrentMaxPing = PerfCounters->Get(TEXT("MaxPing"), MaxPing);
 				PerfCounters->Set(TEXT("MaxPing"), FMath::Max(MaxPing, CurrentMaxPing), IPerfCounters::Flags::Transient);
 				float CurrentMinPing = PerfCounters->Get(TEXT("MinPing"), MinPing);
@@ -409,7 +454,9 @@
 			}
 			else
 			{
-				PerfCounters->Set(TEXT("AvgPing"), 0.0f);
+				PerfCounters->Set(TEXT("NumConnections"), 0);
+
+				PerfCounters->Set(TEXT("AvgPing"), 0.0f, IPerfCounters::Flags::Transient);
 				PerfCounters->Set(TEXT("MaxPing"), -FLT_MAX, IPerfCounters::Flags::Transient);
 				PerfCounters->Set(TEXT("MinPing"), FLT_MAX, IPerfCounters::Flags::Transient);
 
@@ -419,50 +466,19 @@
 					PerfCounters->Set(FString::Printf(TEXT("PingBucket%d"), BucketIdx), 0.0f, IPerfCounters::Flags::Transient);
 				}
 			}
+
+			// set the per connection stats (these are calculated earlier).
+			// Note that NumClients may be != NumConnections. Also, if NumClients is 0, the rest of counters should be 0 as well
+			PerfCounters->Set(TEXT("NumClients"), NumClients);
+			PerfCounters->Set(TEXT("InRateClientMax"), ClientInBytesMax);
+			PerfCounters->Set(TEXT("InRateClientMin"), ClientInBytesMin);
+			PerfCounters->Set(TEXT("InRateClientAvg"), ClientInBytesAvg);
+			PerfCounters->Set(TEXT("OutRateClientMax"), ClientOutBytesMax);
+			PerfCounters->Set(TEXT("OutRateClientMin"), ClientOutBytesMin);
+			PerfCounters->Set(TEXT("OutRateClientAvg"), ClientOutBytesAvg);
 		}
 #endif // USE_SERVER_PERF_COUNTERS
 
-		// Copy the net status values over
-		SET_DWORD_STAT(STAT_Ping, Ping);
-		SET_DWORD_STAT(STAT_Channels, NumOpenChannels);
-
-		SET_DWORD_STAT(STAT_OutLoss, OutPacketsLost);
-		SET_DWORD_STAT(STAT_InLoss, InPacketsLost);
-		SET_DWORD_STAT(STAT_InRate, InBytes);
-		SET_DWORD_STAT(STAT_OutRate, OutBytes);
-		SET_DWORD_STAT(STAT_InRateClientMax, ClientInBytesMax);
-		SET_DWORD_STAT(STAT_InRateClientMin, ClientInBytesMin);
-		SET_DWORD_STAT(STAT_InRateClientAvg, ClientInBytesAvg);
-		SET_DWORD_STAT(STAT_OutRateClientMax, ClientOutBytesMax);
-		SET_DWORD_STAT(STAT_OutRateClientMin, ClientOutBytesMin);
-		SET_DWORD_STAT(STAT_OutRateClientAvg, ClientOutBytesAvg);
-
-		SET_DWORD_STAT(STAT_NetNumClients, NumClients);
-		SET_DWORD_STAT(STAT_InPackets, InPackets);
-		SET_DWORD_STAT(STAT_OutPackets, OutPackets);
-		SET_DWORD_STAT(STAT_InBunches, InBunches);
-		SET_DWORD_STAT(STAT_OutBunches, OutBunches);
-
-		SET_DWORD_STAT(STAT_NetGUIDInRate, NetGUIDInBytes);
-		SET_DWORD_STAT(STAT_NetGUIDOutRate, NetGUIDOutBytes);
-
-		SET_DWORD_STAT(STAT_VoicePacketsSent, VoicePacketsSent);
-		SET_DWORD_STAT(STAT_VoicePacketsRecv, VoicePacketsRecv);
-		SET_DWORD_STAT(STAT_VoiceBytesSent, VoiceBytesSent);
-		SET_DWORD_STAT(STAT_VoiceBytesRecv, VoiceBytesRecv);
-
-		SET_DWORD_STAT(STAT_PercentInVoice, VoiceInPercent);
-		SET_DWORD_STAT(STAT_PercentOutVoice, VoiceOutPercent);
-
-		SET_DWORD_STAT(STAT_NumActorChannels, NumActorChannels);
-		SET_DWORD_STAT(STAT_NumDormantActors, NumDormantActors);
-		SET_DWORD_STAT(STAT_NumActors, NumActors);
-		SET_DWORD_STAT(STAT_NumActorChannelsReadyDormant, NumActorChannelsReadyDormant);
-		SET_DWORD_STAT(STAT_NumNetGUIDsAckd, AckCount);
-		SET_DWORD_STAT(STAT_NumNetGUIDsPending, UnAckCount);
-		SET_DWORD_STAT(STAT_NumNetGUIDsUnAckd, PendingCount);
-		SET_DWORD_STAT(STAT_NetSaturated, NetSaturated);
-		
 		// Reset everything
 		InBytes = 0;
 		OutBytes = 0;
@@ -482,7 +498,7 @@
 		VoiceOutPercent = 0;
 		StatUpdateTime = CurrentRealtimeSeconds;
 	}
-#endif // STATS
+#endif // (USE_SERVER_PERF_COUNTERS) || STATS
 
 	// Poll all sockets.
 	if( ServerConnection )
diff --git a/Engine/Source/Runtime/Engine/Private/Particles/ParticleEmitterInstances.cpp b/Engine/Source/Runtime/Engine/Private/Particles/ParticleEmitterInstances.cpp
--- a/Engine/Source/Runtime/Engine/Private/Particles/ParticleEmitterInstances.cpp
+++ b/Engine/Source/Runtime/Engine/Private/Particles/ParticleEmitterInstances.cpp
@@ -2623,9 +2623,9 @@ bool FParticleEmitterInstance::FillReplayData( FDynamicEmitterReplayDataBase& Ou
 		NewReplayData->SubImages_Horizontal = LODLevel->RequiredModule->SubImages_Horizontal;
 		NewReplayData->SubImages_Vertical = LODLevel->RequiredModule->SubImages_Vertical;
 
-		NewReplayData->bOverrideSystemMacroUV = LODLevel->RequiredModule->bOverrideSystemMacroUV;
-		NewReplayData->MacroUVRadius = LODLevel->RequiredModule->MacroUVRadius;
-		NewReplayData->MacroUVPosition = LODLevel->RequiredModule->MacroUVPosition;
+		NewReplayData->MacroUVOverride.bOverride = LODLevel->RequiredModule->bOverrideSystemMacroUV;
+		NewReplayData->MacroUVOverride.Radius = LODLevel->RequiredModule->MacroUVRadius;
+		NewReplayData->MacroUVOverride.Position = LODLevel->RequiredModule->MacroUVPosition;
         
 		NewReplayData->bLockAxis = false;
 		if (Module_AxisLock && (Module_AxisLock->bEnabled == true))
diff --git a/Engine/Source/Runtime/Engine/Private/Particles/ParticleGpuSimulation.cpp b/Engine/Source/Runtime/Engine/Private/Particles/ParticleGpuSimulation.cpp
--- a/Engine/Source/Runtime/Engine/Private/Particles/ParticleGpuSimulation.cpp
+++ b/Engine/Source/Runtime/Engine/Private/Particles/ParticleGpuSimulation.cpp
@@ -466,6 +466,7 @@ BEGIN_UNIFORM_BUFFER_STRUCT( FGPUSpriteEmitterDynamicUniformParameters, )
 	DECLARE_UNIFORM_BUFFER_STRUCT_MEMBER( FVector4, AxisLockRight )
 	DECLARE_UNIFORM_BUFFER_STRUCT_MEMBER( FVector4, AxisLockUp )
 	DECLARE_UNIFORM_BUFFER_STRUCT_MEMBER( FVector4, DynamicColor)
+	DECLARE_UNIFORM_BUFFER_STRUCT_MEMBER( FVector4, MacroUVParameters )
 END_UNIFORM_BUFFER_STRUCT( FGPUSpriteEmitterDynamicUniformParameters )
 
 IMPLEMENT_UNIFORM_BUFFER_STRUCT(FGPUSpriteEmitterDynamicUniformParameters,TEXT("EmitterDynamicUniforms"));
@@ -473,9 +474,9 @@ IMPLEMENT_UNIFORM_BUFFER_STRUCT(FGPUSpriteEmitterDynamicUniformParameters,TEXT("
 typedef TUniformBufferRef<FGPUSpriteEmitterDynamicUniformParameters> FGPUSpriteEmitterDynamicUniformBufferRef;
 
 /**
- * Shader parameters for the particle vertex factory.
+ * Vertex shader parameters for the particle vertex factory.
  */
-class FGPUSpriteVertexFactoryShaderParameters : public FVertexFactoryShaderParameters
+class FGPUSpriteVertexFactoryShaderParametersVS : public FVertexFactoryShaderParameters
 {
 public:
 	virtual void Bind( const FShaderParameterMap& ParameterMap ) override
@@ -530,6 +531,22 @@ private:
 	FShaderResourceParameter CurveTextureSampler;
 };
 
+/**
+ * Pixel shader parameters for the particle vertex factory.
+ */
+class FGPUSpriteVertexFactoryShaderParametersPS : public FVertexFactoryShaderParameters
+{
+public:
+	virtual void Bind( const FShaderParameterMap& ParameterMap ) override {}
+
+	virtual void Serialize(FArchive& Ar) override {}
+
+	virtual void SetMesh(FRHICommandList& RHICmdList, FShader* Shader,const FVertexFactory* VertexFactory,const FSceneView& View,const FMeshBatchElement& BatchElement,uint32 DataFlags) const override;
+	virtual uint32 GetSize() const override { return sizeof(*this); }
+
+private:
+};
+
 /**
  * GPU Sprite vertex factory vertex declaration.
  */
@@ -643,14 +660,22 @@ public:
 	 */
 	static FVertexFactoryShaderParameters* ConstructShaderParameters(EShaderFrequency ShaderFrequency)
 	{
-		return ShaderFrequency == SF_Vertex ? new FGPUSpriteVertexFactoryShaderParameters() : NULL;
+		if (ShaderFrequency == SF_Vertex)
+		{
+			return new FGPUSpriteVertexFactoryShaderParametersVS();
+		}
+		else if (ShaderFrequency == SF_Pixel)
+		{
+			return new FGPUSpriteVertexFactoryShaderParametersPS();
+		}
+		return NULL;
 	}
 };
 
 /**
  * Set vertex factory shader parameters.
  */
-void FGPUSpriteVertexFactoryShaderParameters::SetMesh(FRHICommandList& RHICmdList, FShader* Shader, const FVertexFactory* VertexFactory, const FSceneView& View, const FMeshBatchElement& BatchElement, uint32 DataFlags) const
+void FGPUSpriteVertexFactoryShaderParametersVS::SetMesh(FRHICommandList& RHICmdList, FShader* Shader, const FVertexFactory* VertexFactory, const FSceneView& View, const FMeshBatchElement& BatchElement, uint32 DataFlags) const
 {
 	FGPUSpriteVertexFactory* GPUVF = (FGPUSpriteVertexFactory*)VertexFactory;
 	FVertexShaderRHIParamRef VertexShader = Shader->GetVertexShader();
@@ -669,6 +694,13 @@ void FGPUSpriteVertexFactoryShaderParameters::SetMesh(FRHICommandList& RHICmdLis
 	SetTextureParameter(RHICmdList, VertexShader, CurveTexture, CurveTextureSampler, SamplerStateLinear, GParticleCurveTexture.GetCurveTexture() );
 }
 
+void FGPUSpriteVertexFactoryShaderParametersPS::SetMesh(FRHICommandList& RHICmdList, FShader* Shader, const FVertexFactory* VertexFactory, const FSceneView& View, const FMeshBatchElement& BatchElement, uint32 DataFlags) const
+{
+	FGPUSpriteVertexFactory* GPUVF = (FGPUSpriteVertexFactory*)VertexFactory;
+	FPixelShaderRHIParamRef PixelShader = Shader->GetPixelShader();
+	SetUniformBufferParameter(RHICmdList, PixelShader, Shader->GetUniformBufferParameter<FGPUSpriteEmitterDynamicUniformParameters>(), GPUVF->EmitterDynamicUniformBuffer );
+}
+
 IMPLEMENT_VERTEX_FACTORY_TYPE(FGPUSpriteVertexFactory,"ParticleGPUSpriteVertexFactory",true,false,true,false,false);
 
 /*-----------------------------------------------------------------------------
@@ -2626,6 +2658,8 @@ struct FGPUSpriteDynamicEmitterData : FDynamicEmitterDataBase
 	/** Tile vector field in z axis? */
 	uint32 bLocalVectorFieldTileZ : 1;
 
+	/** Current MacroUV override settings */
+	FMacroUVOverride MacroUVOverride;
 
 	/** Constructor. */
 	explicit FGPUSpriteDynamicEmitterData( const UParticleModuleRequired* InRequiredModule )
@@ -2736,11 +2770,17 @@ struct FGPUSpriteDynamicEmitterData : FDynamicEmitterDataBase
 			if (Simulation->SimulationIndex != INDEX_NONE
 				&& Simulation->VertexBuffer.ParticleCount > 0)
 			{
+				FGPUSpriteEmitterDynamicUniformParameters PerViewDynamicParameters = EmitterDynamicParameters;
+				FVector2D ObjectNDCPosition;
+				FVector2D ObjectMacroUVScales;
+				Proxy->GetObjectPositionAndScale(*View,ObjectNDCPosition, ObjectMacroUVScales);
+				PerViewDynamicParameters.MacroUVParameters = FVector4(ObjectNDCPosition.X, ObjectNDCPosition.Y, ObjectMacroUVScales.X, ObjectMacroUVScales.Y); 
+
 				FGPUSpriteEmitterDynamicUniformBufferRef LocalDynamicUniformBuffer;
-				// Create view agnostic render data.  Do here rather than in CreateRenderThreadResources because in some cases Render can be called before CreateRenderThreadResources
+				// Do here rather than in CreateRenderThreadResources because in some cases Render can be called before CreateRenderThreadResources
 				{
 					// Create per-emitter uniform buffer for dynamic parameters
-					LocalDynamicUniformBuffer = FGPUSpriteEmitterDynamicUniformBufferRef::CreateUniformBufferImmediate(EmitterDynamicParameters, UniformBuffer_SingleFrame);
+					LocalDynamicUniformBuffer = FGPUSpriteEmitterDynamicUniformBufferRef::CreateUniformBufferImmediate(PerViewDynamicParameters, UniformBuffer_SingleFrame);
 				}
 
 				if (bUseLocalSpace == false)
@@ -2850,6 +2890,10 @@ struct FGPUSpriteDynamicEmitterData : FDynamicEmitterDataBase
 		static FDynamicEmitterReplayDataBase DummyData;
 		return DummyData;
 	}
+
+	/** Returns the current macro uv override. */
+	virtual const FMacroUVOverride& GetMacroUVOverride() const { return MacroUVOverride; }
+
 };
 
 /*-----------------------------------------------------------------------------
@@ -3079,6 +3123,10 @@ public:
 		}
 		DynamicData->EmitterDynamicParameters.DynamicColor = ColorOverLife * ColorScaleOverLife;
 
+		DynamicData->MacroUVOverride.bOverride = LODLevel->RequiredModule->bOverrideSystemMacroUV;
+		DynamicData->MacroUVOverride.Radius = LODLevel->RequiredModule->MacroUVRadius;
+		DynamicData->MacroUVOverride.Position = LODLevel->RequiredModule->MacroUVPosition;
+
 		const bool bSimulateGPUParticles = 
 			FXConsoleVariables::bFreezeGPUSimulation == false &&
 			FXConsoleVariables::bFreezeParticleSimulation == false &&
