Mergine changes from //depot/UE4 to //UE4/Main

diff --git a/Engine/Source/Runtime/Engine/Private/Particles/ParticleSystemRender.cpp b/Engine/Source/Runtime/Engine/Private/Particles/ParticleSystemRender.cpp
--- a/Engine/Source/Runtime/Engine/Private/Particles/ParticleSystemRender.cpp
+++ b/Engine/Source/Runtime/Engine/Private/Particles/ParticleSystemRender.cpp
@@ -6628,11 +6628,11 @@ void FParticleSystemSceneProxy::GetObjectPositionAndScale(const FSceneView& View
 	FVector MacroUVPosition = DynamicData->SystemPositionForMacroUVs;
    
 	uint32 Index = DynamicData->EmitterIndex;
-	const FDynamicEmitterReplayDataBase& EmitterData = DynamicData->DynamicEmitterDataArray[Index]->GetSource();
-	if (EmitterData.bOverrideSystemMacroUV)
+	const FMacroUVOverride& MacroUVOverride = DynamicData->DynamicEmitterDataArray[Index]->GetMacroUVOverride();
+	if (MacroUVOverride.bOverride)
 	{
-		MacroUVRadius = EmitterData.MacroUVRadius;
-		MacroUVPosition = GetLocalToWorld().TransformVector(EmitterData.MacroUVPosition);
+		MacroUVRadius = MacroUVOverride.Radius;
+		MacroUVPosition = GetLocalToWorld().TransformVector(MacroUVOverride.Position);
 
 		if (MacroUVPosition.ContainsNaN())
 		{
diff --git a/Engine/Source/Runtime/Engine/Private/RendererSupport.cpp b/Engine/Source/Runtime/Engine/Private/RendererSupport.cpp
--- a/Engine/Source/Runtime/Engine/Private/RendererSupport.cpp
+++ b/Engine/Source/Runtime/Engine/Private/RendererSupport.cpp
@@ -16,11 +16,12 @@
 #include "ComponentReregisterContext.h"
 
 /** Clears and optionally backs up all references to renderer module classes in other modules, particularly engine. */
-void ClearReferencesToRendererModuleClasses(
+static void ClearReferencesToRendererModuleClasses(
 	TMap<UWorld*, bool>& WorldsToUpdate, 
 	TMap<FMaterialShaderMap*, TScopedPointer<TArray<uint8> > >& ShaderMapToSerializedShaderData,
 	FGlobalShaderBackupData& GlobalShaderBackup,
 	TMap<FShaderType*, FString>& ShaderTypeNames,
+	TMap<const FShaderPipelineType*, FString>& ShaderPipelineTypeNames,
 	TMap<FVertexFactoryType*, FString>& VertexFactoryTypeNames)
 {
 	// Destroy all renderer scenes
@@ -61,6 +62,12 @@ void ClearReferencesToRendererModuleClasses(
 		ShaderTypeNames.Add(ShaderType, ShaderType->GetName());
 	}
 
+	for (TLinkedList<FShaderPipelineType*>::TIterator It(FShaderPipelineType::GetTypeList()); It; It.Next())
+	{
+		const FShaderPipelineType* ShaderPipelineType = *It;
+		ShaderPipelineTypeNames.Add(ShaderPipelineType, ShaderPipelineType->GetName());
+	}
+
 	for(TLinkedList<FVertexFactoryType*>::TIterator It(FVertexFactoryType::GetTypeList()); It; It.Next())
 	{
 		FVertexFactoryType* VertexFactoryType = *It;
@@ -120,6 +127,7 @@ void RestoreReferencesToRendererModuleClasses(
 	const TMap<FMaterialShaderMap*, TScopedPointer<TArray<uint8> > >& ShaderMapToSerializedShaderData,
 	const FGlobalShaderBackupData& GlobalShaderBackup,
 	const TMap<FShaderType*, FString>& ShaderTypeNames,
+	const TMap<const FShaderPipelineType*, FString>& ShaderPipelineTypeNames,
 	const TMap<FVertexFactoryType*, FString>& VertexFactoryTypeNames)
 {
 	FlushShaderFileCache();
@@ -156,12 +164,13 @@ void RestoreReferencesToRendererModuleClasses(
 		{
 			EShaderPlatform ShaderPlatform = GetFeatureLevelShaderPlatform((ERHIFeatureLevel::Type)i);
 			check(ShaderPlatform < EShaderPlatform::SP_NumPlatforms);
-			FMaterialShaderMap::FixupShaderTypes(ShaderPlatform, ShaderTypeNames, VertexFactoryTypeNames);
+			FMaterialShaderMap::FixupShaderTypes(ShaderPlatform, ShaderTypeNames, ShaderPipelineTypeNames, VertexFactoryTypeNames);
 		}
 	}
 
 	TArray<FShaderType*> OutdatedShaderTypes;
 	TArray<const FVertexFactoryType*> OutdatedFactoryTypes;
+	TArray<const FShaderPipelineType*> OutdatedShaderPipelineTypes;
 	FShaderType::GetOutdatedTypes(OutdatedShaderTypes, OutdatedFactoryTypes);
 
 	// Recompile any missing shaders
@@ -169,8 +178,8 @@ void RestoreReferencesToRendererModuleClasses(
 	{
 		auto ShaderPlatform = GetFeatureLevelShaderPlatform(FeatureLevel);
 		check(ShaderPlatform < EShaderPlatform::SP_NumPlatforms);
-		BeginRecompileGlobalShaders(OutdatedShaderTypes, ShaderPlatform);
-		UMaterial::UpdateMaterialShaders(OutdatedShaderTypes, OutdatedFactoryTypes, ShaderPlatform);
+		BeginRecompileGlobalShaders(OutdatedShaderTypes, OutdatedShaderPipelineTypes, ShaderPlatform);
+		UMaterial::UpdateMaterialShaders(OutdatedShaderTypes, OutdatedShaderPipelineTypes, OutdatedFactoryTypes, ShaderPlatform);
 	});
 
 	// Block on global shader jobs
@@ -201,9 +210,10 @@ void RecompileRenderer(const TArray<FString>& Args)
 		TMap<FMaterialShaderMap*, TScopedPointer<TArray<uint8> > > ShaderMapToSerializedShaderData;
 		FGlobalShaderBackupData GlobalShaderBackup;
 		TMap<FShaderType*, FString> ShaderTypeNames;
+		TMap<const FShaderPipelineType*, FString> ShaderPipelineTypeNames;
 		TMap<FVertexFactoryType*, FString> VertexFactoryTypeNames;
 
-		ClearReferencesToRendererModuleClasses(WorldsToUpdate, ShaderMapToSerializedShaderData, GlobalShaderBackup, ShaderTypeNames, VertexFactoryTypeNames);
+		ClearReferencesToRendererModuleClasses(WorldsToUpdate, ShaderMapToSerializedShaderData, GlobalShaderBackup, ShaderTypeNames, ShaderPipelineTypeNames, VertexFactoryTypeNames);
 
 		EndShutdownTime = FPlatformTime::Seconds();
 		UE_LOG(LogShaders, Warning, TEXT("Shutdown complete %.1fs"),(float)(EndShutdownTime - StartTime));
@@ -213,7 +223,7 @@ void RecompileRenderer(const TArray<FString>& Args)
 		EndRecompileTime = FPlatformTime::Seconds();
 		UE_LOG(LogShaders, Warning, TEXT("Recompile complete %.1fs"),(float)(EndRecompileTime - EndShutdownTime));
 
-		RestoreReferencesToRendererModuleClasses(WorldsToUpdate, ShaderMapToSerializedShaderData, GlobalShaderBackup, ShaderTypeNames, VertexFactoryTypeNames);
+		RestoreReferencesToRendererModuleClasses(WorldsToUpdate, ShaderMapToSerializedShaderData, GlobalShaderBackup, ShaderTypeNames, ShaderPipelineTypeNames, VertexFactoryTypeNames);
 	}
 
 #if WITH_EDITOR
diff --git a/Engine/Source/Runtime/Engine/Private/ShaderCompiler/ShaderCompiler.cpp b/Engine/Source/Runtime/Engine/Private/ShaderCompiler/ShaderCompiler.cpp
--- a/Engine/Source/Runtime/Engine/Private/ShaderCompiler/ShaderCompiler.cpp
+++ b/Engine/Source/Runtime/Engine/Private/ShaderCompiler/ShaderCompiler.cpp
@@ -81,6 +81,8 @@ static TAutoConsoleVariable<FString> CVarD3DCompilerPath(TEXT("r.D3DCompilerPath
 	ECVF_Cheat);
 #endif
 
+extern bool CompileShaderPipeline(const IShaderFormat* Compiler, FName Format, FShaderPipelineCompileJob* PipelineJob, const FString& Dir);
+
 namespace XGEConsoleVariables
 {
 	int32 Enabled = 0;
@@ -175,8 +177,39 @@ static inline void GetFormatVersionMap(TMap<FString, uint16>& OutFormatVersionMa
 	}
 }
 
+static int32 GetNumTotalJobs(const TArray<FShaderCommonCompileJob*>& Jobs)
+{
+	int32 NumJobs = 0;
+	for (int32 Index = 0; Index < Jobs.Num(); ++Index)
+	{
+		auto* PipelineJob = Jobs[Index]->GetShaderPipelineJob();
+		NumJobs += PipelineJob ? PipelineJob->StageJobs.Num() : 1;
+	}
+
+	return NumJobs;
+}
+
+static void SplitJobsByType(const TArray<FShaderCommonCompileJob*>& QueuedJobs, TArray<FShaderCompileJob*>& OutQueuedSingleJobs, TArray<FShaderPipelineCompileJob*>& OutQueuedPipelineJobs)
+{
+	for (int32 Index = 0; Index < QueuedJobs.Num(); ++Index)
+	{
+		auto* CommonJob = QueuedJobs[Index];
+		auto* PipelineJob = CommonJob->GetShaderPipelineJob();
+		if (PipelineJob)
+		{
+			OutQueuedPipelineJobs.Add(PipelineJob);
+		}
+		else
+		{
+			auto* SingleJob = CommonJob->GetSingleShaderJob();
+			check(SingleJob);
+			OutQueuedSingleJobs.Add(SingleJob);
+		}
+	}
+}
+
 // Serialize Queued Job information
-static bool DoWriteTasks(const TArray<FShaderCompileJob*>& QueuedJobs, FArchive& TransferFile)
+static bool DoWriteTasks(const TArray<FShaderCommonCompileJob*>& QueuedJobs, FArchive& TransferFile)
 {
 	int32 InputVersion = ShaderCompileWorkerInputVersion;
 	TransferFile << InputVersion;
@@ -186,18 +219,22 @@ static bool DoWriteTasks(const TArray<FShaderCompileJob*>& QueuedJobs, FArchive&
 
 	TransferFile << FormatVersionMap;
 
+	TArray<FShaderCompileJob*> QueuedSingleJobs;
+	TArray<FShaderPipelineCompileJob*> QueuedPipelineJobs;
+	SplitJobsByType(QueuedJobs, QueuedSingleJobs, QueuedPipelineJobs);
+
 	// Write individual shader jobs
 	{
 		int32 SingleJobHeader = ShaderCompileWorkerSingleJobHeader;
 		TransferFile << SingleJobHeader;
 
-		int32 NumBatches = QueuedJobs.Num();
+		int32 NumBatches = QueuedSingleJobs.Num();
 		TransferFile << NumBatches;
 
 		// Serialize all the batched jobs
-		for (int32 JobIndex = 0; JobIndex < QueuedJobs.Num(); JobIndex++)
+		for (int32 JobIndex = 0; JobIndex < QueuedSingleJobs.Num(); JobIndex++)
 		{
-			TransferFile << QueuedJobs[JobIndex]->Input;
+			TransferFile << QueuedSingleJobs[JobIndex]->Input;
 		}
 	}
 
@@ -206,15 +243,27 @@ static bool DoWriteTasks(const TArray<FShaderCompileJob*>& QueuedJobs, FArchive&
 		int32 PipelineJobHeader = ShaderCompileWorkerPipelineJobHeader;
 		TransferFile << PipelineJobHeader;
 
-		int32 NumBatches = 0;
+		int32 NumBatches = QueuedPipelineJobs.Num();
 		TransferFile << NumBatches;
+		for (int32 JobIndex = 0; JobIndex < QueuedPipelineJobs.Num(); JobIndex++)
+		{
+			auto* PipelineJob = QueuedPipelineJobs[JobIndex];
+			FString PipelineName = PipelineJob->ShaderPipeline->GetName();
+			TransferFile << PipelineName;
+			int32 NumStageJobs = PipelineJob->StageJobs.Num();
+			TransferFile << NumStageJobs;
+			for (int32 Index = 0; Index < NumStageJobs; Index++)
+			{
+				TransferFile << PipelineJob->StageJobs[Index]->GetSingleShaderJob()->Input;
+			}
+		}
 	}
 
 	return TransferFile.Close();
 }
 
 // Process results from Worker Process
-static void DoReadTaskResults(const TArray<FShaderCompileJob*>& QueuedJobs, FArchive& OutputFile)
+static void DoReadTaskResults(const TArray<FShaderCommonCompileJob*>& QueuedJobs, FArchive& OutputFile)
 {
 	int32 OutputVersion = ShaderCompileWorkerOutputVersion;
 	OutputFile << OutputVersion;
@@ -279,6 +328,10 @@ static void DoReadTaskResults(const TArray<FShaderCompileJob*>& QueuedJobs, FArc
 		UE_LOG(LogShaderCompilers, Fatal, TEXT("ShaderCompileWorker crashed! \n %s \n %s"), ExceptionInfo.GetData(), Callstack.GetData());
 	}
 
+	TArray<FShaderCompileJob*> QueuedSingleJobs;
+	TArray<FShaderPipelineCompileJob*> QueuedPipelineJobs;
+	SplitJobsByType(QueuedJobs, QueuedSingleJobs, QueuedPipelineJobs);
+
 	auto ReadSingleJob = [&](FShaderCompileJob* CurrentJob)
 		{
 			check(!CurrentJob->bFinalized);
@@ -305,10 +358,10 @@ static void DoReadTaskResults(const TArray<FShaderCompileJob*>& QueuedJobs, FArc
 
 		int32 NumJobs;
 		OutputFile << NumJobs;
-		checkf(NumJobs == QueuedJobs.Num(), TEXT("Worker returned %u single jobs, %u expected"), NumJobs, QueuedJobs.Num());
+		checkf(NumJobs == QueuedSingleJobs.Num(), TEXT("Worker returned %u single jobs, %u expected"), NumJobs, QueuedSingleJobs.Num());
 		for (int32 JobIndex = 0; JobIndex < NumJobs; JobIndex++)
 		{
-			auto* CurrentJob = QueuedJobs[JobIndex];
+			auto* CurrentJob = QueuedSingleJobs[JobIndex];
 			ReadSingleJob(CurrentJob);
 		}
 	}
@@ -325,6 +378,35 @@ static void DoReadTaskResults(const TArray<FShaderCompileJob*>& QueuedJobs, FArc
 
 		int32 NumJobs;
 		OutputFile << NumJobs;
+		checkf(NumJobs == QueuedPipelineJobs.Num(), TEXT("Worker returned %u pipeline jobs, %u expected"), NumJobs, QueuedPipelineJobs.Num());
+		for (int32 JobIndex = 0; JobIndex < NumJobs; JobIndex++)
+		{
+			FShaderPipelineCompileJob* CurrentJob = QueuedPipelineJobs[JobIndex];
+
+			FString PipelineName;
+			OutputFile << PipelineName;
+			checkf(PipelineName == CurrentJob->ShaderPipeline->GetName(), TEXT("Worker returned Pipeline %s, expected %s!"), *PipelineName, CurrentJob->ShaderPipeline->GetName());
+
+			check(!CurrentJob->bFinalized);
+			CurrentJob->bFinalized = true;
+			CurrentJob->bFailedRemovingUnused = false;
+
+			int32 NumStageJobs = -1;
+			OutputFile << NumStageJobs;
+
+			if (NumStageJobs != CurrentJob->StageJobs.Num())
+			{
+				checkf(NumJobs == QueuedPipelineJobs.Num(), TEXT("Worker returned %u stage pipeline jobs, %u expected"), NumStageJobs, CurrentJob->StageJobs.Num());
+			}
+
+			CurrentJob->bSucceeded = true;
+			for (int32 Index = 0; Index < NumStageJobs; Index++)
+			{
+				auto* SingleJob = CurrentJob->StageJobs[Index]->GetSingleShaderJob();
+				ReadSingleJob(SingleJob);
+				CurrentJob->bFailedRemovingUnused = CurrentJob->bFailedRemovingUnused | SingleJob->Output.bFailedRemovingUnused;
+			}
+		}
 	}
 }
 
@@ -347,7 +429,7 @@ struct FShaderCompileWorkerInfo
 	double StartTime;
 
 	/** Jobs that this worker is responsible for compiling. */
-	TArray<FShaderCompileJob*> QueuedJobs;
+	TArray<FShaderCommonCompileJob*> QueuedJobs;
 
 	FShaderCompileWorkerInfo() :
 		bIssuedTasksToWorker(false),		
@@ -534,8 +616,24 @@ int32 FShaderCompileThreadRunnable::PullTasksFromQueue()
 
 						for (int32 JobIndex = 0; JobIndex < CurrentWorkerInfo.QueuedJobs.Num(); JobIndex++)
 						{
-							const FShaderCompileJob& Job = *CurrentWorkerInfo.QueuedJobs[JobIndex];
-							JobNames += FString(Job.ShaderType->GetName()) + TEXT(" Instructions = ") + FString::FromInt(Job.Output.NumInstructions) + TEXT(", ");
+							const FShaderCommonCompileJob& Job = *CurrentWorkerInfo.QueuedJobs[JobIndex];
+							if (auto* SingleJob = Job.GetSingleShaderJob())
+							{
+								JobNames += FString(SingleJob->ShaderType->GetName()) + TEXT(" Instructions = ") + FString::FromInt(SingleJob->Output.NumInstructions);
+							}
+							else
+							{
+								auto* PipelineJob = Job.GetShaderPipelineJob();
+								JobNames += FString(PipelineJob->ShaderPipeline->GetName());
+								if (PipelineJob->bFailedRemovingUnused)
+								{
+									 JobNames += FString(TEXT("(failed to optimize)"));
+								}
+							}
+							if (JobIndex < CurrentWorkerInfo.QueuedJobs.Num() - 1)
+							{
+								JobNames += TEXT(", ");
+							}
 						}
 
 						UE_LOG(LogShaders, Display, TEXT("Finished batch of %u jobs in %.3fs, %s"), CurrentWorkerInfo.QueuedJobs.Num(), ElapsedTime, *JobNames);
@@ -754,13 +852,16 @@ void FShaderCompileThreadRunnable::CompileDirectlyThroughDll()
 		{
 			for (int32 JobIndex = 0; JobIndex < CurrentWorkerInfo.QueuedJobs.Num(); JobIndex++)
 			{
-				FShaderCompileJob& CurrentJob = *CurrentWorkerInfo.QueuedJobs[JobIndex];
+				FShaderCommonCompileJob& CurrentJob = *CurrentWorkerInfo.QueuedJobs[JobIndex];
 
 				check(!CurrentJob.bFinalized);
 				CurrentJob.bFinalized = true;
 
 				static ITargetPlatformManagerModule& TPM = GetTargetPlatformManagerRef();
-				const FName Format = LegacyShaderPlatformToShaderFormat(EShaderPlatform(CurrentJob.Input.Target.Platform));
+				auto* SingleJob = CurrentJob.GetSingleShaderJob();
+				if (SingleJob)
+				{
+					const FName Format = LegacyShaderPlatformToShaderFormat(EShaderPlatform(SingleJob->Input.Target.Platform));
 				const IShaderFormat* Compiler = TPM.FindShaderFormat(Format);
 
 				if (!Compiler)
@@ -769,23 +870,55 @@ void FShaderCompileThreadRunnable::CompileDirectlyThroughDll()
 				}
 				CA_ASSUME(Compiler != NULL);
 
-				if (IsValidRef(CurrentJob.Input.SharedEnvironment))
+					if (IsValidRef(SingleJob->Input.SharedEnvironment))
 				{
 					// Merge the shared environment into the per-shader environment before calling into the compile function
 					// Normally this happens in the worker
-					CurrentJob.Input.Environment.Merge(*CurrentJob.Input.SharedEnvironment);
+						SingleJob->Input.Environment.Merge(*SingleJob->Input.SharedEnvironment);
 				}
 
 				// Compile the shader directly through the platform dll (directly from the shader dir as the working directory)
-				Compiler->CompileShader(Format, CurrentJob.Input, CurrentJob.Output, FString(FPlatformProcess::ShaderDir()));
+					Compiler->CompileShader(Format, SingleJob->Input, SingleJob->Output, FString(FPlatformProcess::ShaderDir()));
 
-				CurrentJob.bSucceeded = CurrentJob.Output.bSucceeded;
+					SingleJob->bSucceeded = SingleJob->Output.bSucceeded;
 
-				if (CurrentJob.Output.bSucceeded)
+					if (SingleJob->Output.bSucceeded)
 				{
 					// Generate a hash of the output and cache it
 					// The shader processing this output will use it to search for existing FShaderResources
-					CurrentJob.Output.GenerateOutputHash();
+						SingleJob->Output.GenerateOutputHash();
+					}
+				}
+				else
+				{
+					auto* PipelineJob = CurrentJob.GetShaderPipelineJob();
+					check(PipelineJob);
+
+					EShaderPlatform Platform = (EShaderPlatform)PipelineJob->StageJobs[0]->GetSingleShaderJob()->Input.Target.Platform;
+					const FName Format = LegacyShaderPlatformToShaderFormat(Platform);
+					const IShaderFormat* Compiler = TPM.FindShaderFormat(Format);
+
+					if (!Compiler)
+					{
+						UE_LOG(LogShaderCompilers, Fatal, TEXT("Can't compile shaders for format %s, couldn't load compiler dll"), *Format.ToString());
+					}
+					CA_ASSUME(Compiler != NULL);
+
+					// Verify same platform on all stages
+					for (int32 Index = 1; Index < PipelineJob->StageJobs.Num(); ++Index)
+					{
+						auto* SingleStage = PipelineJob->StageJobs[Index]->GetSingleShaderJob();
+						if (!SingleStage)
+						{
+							UE_LOG(LogShaderCompilers, Fatal, TEXT("Can't nest Shader Pipelines inside Shader Pipeline '%s'!"), PipelineJob->ShaderPipeline->GetName());
+						}
+						if (Platform != SingleStage->Input.Target.Platform)
+						{
+							UE_LOG(LogShaderCompilers, Fatal, TEXT("Mismatched Target Platform %s while compiling Shader Pipeline '%s'."), *Format.GetPlainNameString(), PipelineJob->ShaderPipeline->GetName());
+						}
+					}
+
+					CompileShaderPipeline(Compiler, Format, PipelineJob, FString(FPlatformProcess::ShaderDir()));
 				}
 			}
 
@@ -959,7 +1092,7 @@ FShaderCompilingManager::FShaderCompilingManager() :
 	Thread->StartThread();
 }
 
-void FShaderCompilingManager::AddJobs(TArray<FShaderCompileJob*>& NewJobs, bool bApplyCompletedShaderMapForRendering, bool bOptimizeForLowLatency)
+void FShaderCompilingManager::AddJobs(TArray<FShaderCommonCompileJob*>& NewJobs, bool bApplyCompletedShaderMapForRendering, bool bOptimizeForLowLatency)
 {
 	check(!FPlatformProperties::RequiresCookedData());
 
@@ -1001,9 +1134,17 @@ void FShaderCompilingManager::AddJobs(TArray<FShaderCompileJob*>& NewJobs, bool
 		NewJobs[JobIndex]->bOptimizeForLowLatency = bOptimizeForLowLatency;
 		FShaderMapCompileResults& ShaderMapInfo = ShaderMapJobs.FindOrAdd(NewJobs[JobIndex]->Id);
 		ShaderMapInfo.bApplyCompletedShaderMapForRendering = bApplyCompletedShaderMapForRendering;
+		auto* PipelineJob = NewJobs[JobIndex]->GetShaderPipelineJob();
+		if (PipelineJob)
+		{
+			ShaderMapInfo.NumJobsQueued += PipelineJob->StageJobs.Num();
+		}
+		else
+		{
 		ShaderMapInfo.NumJobsQueued++;
 	}
 }
+}
 
 /** Launches the worker, returns the launched process handle. */
 FProcHandle FShaderCompilingManager::LaunchWorker(const FString& WorkingDirectory, uint32 InProcessId, uint32 ThreadId, const FString& WorkerInputFile, const FString& WorkerOutputFile)
@@ -1072,7 +1213,8 @@ void FShaderCompilingManager::BlockOnShaderMapCompletion(const TArray<int32>& Sh
 					{
 						const FShaderMapCompileResults& Results = *ResultsPtr;
 
-						if (Results.FinishedJobs.Num() == Results.NumJobsQueued)
+						int32 FinishedJobs = GetNumTotalJobs(Results.FinishedJobs);
+						if (FinishedJobs == Results.NumJobsQueued)
 						{
 							CompiledShaderMaps.Add(ShaderMapIdsToFinishCompiling[ShaderMapIndex], FShaderMapFinalizeResults(Results));
 							ShaderMapJobs.Remove(ShaderMapIdsToFinishCompiling[ShaderMapIndex]);
@@ -1111,7 +1253,7 @@ void FShaderCompilingManager::BlockOnShaderMapCompletion(const TArray<int32>& Sh
 			if (ResultsPtr)
 			{
 				const FShaderMapCompileResults& Results = *ResultsPtr;
-				check(Results.FinishedJobs.Num() == Results.NumJobsQueued);
+				check(GetNumTotalJobs(Results.FinishedJobs) == Results.NumJobsQueued);
 
 				CompiledShaderMaps.Add(ShaderMapIdsToFinishCompiling[ShaderMapIndex], FShaderMapFinalizeResults(Results));
 				ShaderMapJobs.Remove(ShaderMapIdsToFinishCompiling[ShaderMapIndex]);
@@ -1138,7 +1280,7 @@ void FShaderCompilingManager::BlockOnAllShaderMapCompletion(TMap<int32, FShaderM
 				{
 					const FShaderMapCompileResults& Results = It.Value();
 
-					if (Results.FinishedJobs.Num() == Results.NumJobsQueued)
+					if (GetNumTotalJobs(Results.FinishedJobs) == Results.NumJobsQueued)
 					{
 						CompiledShaderMaps.Add(It.Key(), FShaderMapFinalizeResults(Results));
 						It.RemoveCurrent();
@@ -1173,7 +1315,7 @@ void FShaderCompilingManager::BlockOnAllShaderMapCompletion(TMap<int32, FShaderM
 		{
 			const FShaderMapCompileResults& Results = It.Value();
 
-			check(Results.FinishedJobs.Num() == Results.NumJobsQueued);
+			check(GetNumTotalJobs(Results.FinishedJobs) == Results.NumJobsQueued);
 			CompiledShaderMaps.Add(It.Key(), FShaderMapFinalizeResults(Results));
 			It.RemoveCurrent();
 		}
@@ -1212,7 +1354,7 @@ void FShaderCompilingManager::ProcessCompiledShaderMaps(
 		{
 			TArray<FString> Errors;
 			FShaderMapFinalizeResults& CompileResults = ProcessIt.Value();
-			const TArray<FShaderCompileJob*>& ResultArray = CompileResults.FinishedJobs;
+			const TArray<FShaderCommonCompileJob*>& ResultArray = CompileResults.FinishedJobs;
 
 			// Make a copy of the array as this entry of FMaterialShaderMap::ShaderMapsBeingCompiled will be removed below
 			TArray<FMaterial*> MaterialsArray = *Materials;
@@ -1229,17 +1371,29 @@ void FShaderCompilingManager::ProcessCompiledShaderMaps(
 					{
 						for (int32 ErrorIndex = 0; ErrorIndex < SingleJob->Output.Errors.Num(); ErrorIndex++)
 						{
-							Errors.AddUnique(SingleJob->Output.Errors[ErrorIndex].StrippedErrorMessage);
+							Errors.AddUnique(SingleJob->Output.Errors[ErrorIndex].GetErrorString());
 						}
 					}
 				};
 
 			for (int32 JobIndex = 0; JobIndex < ResultArray.Num(); JobIndex++)
 			{
-				FShaderCompileJob& CurrentJob = *ResultArray[JobIndex];
+				FShaderCommonCompileJob& CurrentJob = *ResultArray[JobIndex];
 				bSuccess = bSuccess && CurrentJob.bSucceeded;
 
-				CheckSingleJob(&CurrentJob);
+				auto* SingleJob = CurrentJob.GetSingleShaderJob();
+				if (SingleJob)
+				{
+					CheckSingleJob(SingleJob);
+				}
+				else
+				{
+					auto* PipelineJob = CurrentJob.GetShaderPipelineJob();
+					for (int32 Index = 0; Index < PipelineJob->StageJobs.Num(); ++Index)
+					{
+						CheckSingleJob(PipelineJob->StageJobs[Index]->GetSingleShaderJob());
+					}
+				}
 			}
 
 			bool bShaderMapComplete = true;
@@ -1336,7 +1490,7 @@ void FShaderCompilingManager::ProcessCompiledShaderMaps(
 
 			if (GlobalShaderResults)
 			{
-				const TArray<FShaderCompileJob*>& CompilationResults = GlobalShaderResults->FinishedJobs;
+				const TArray<FShaderCommonCompileJob*>& CompilationResults = GlobalShaderResults->FinishedJobs;
 
 				ProcessCompiledGlobalShaders(CompilationResults);
 
@@ -1433,17 +1587,12 @@ bool FShaderCompilingManager::HandlePotentialRetryOnError(TMap<int32, FShaderMap
 				|| bSpecialEngineMaterial 
 				|| It.Key() == GlobalShaderMapId)
 			{
-				const TArray<FShaderCompileJob*>& CompleteJobs = Results.FinishedJobs;
-				TArray<const FShaderCompileJob*> ErrorJobs;
+				const TArray<FShaderCommonCompileJob*>& CompleteJobs = Results.FinishedJobs;
+				TArray<const FShaderCommonCompileJob*> ErrorJobs;
 				TArray<FString> UniqueErrors;
 				TArray<EShaderPlatform> ErrorPlatforms;
 
-				// Gather unique errors
-				for (int32 JobIndex = 0; JobIndex < CompleteJobs.Num(); JobIndex++)
-				{
-					const FShaderCompileJob& CurrentJob = *CompleteJobs[JobIndex];
-
-					if (!CurrentJob.bSucceeded)
+				auto AddErrorsForFailedJob = [&](const FShaderCompileJob& CurrentJob)
 					{
 						ErrorPlatforms.AddUnique((EShaderPlatform)CurrentJob.Input.Target.Platform);
 						
@@ -1452,12 +1601,34 @@ bool FShaderCompilingManager::HandlePotentialRetryOnError(TMap<int32, FShaderMap
 							const FShaderCompilerError& CurrentError = CurrentJob.Output.Errors[ErrorIndex];
 
 							// Include warnings if LogShaders is unsuppressed, otherwise only include errors
-							if (UE_LOG_ACTIVE(LogShaders, Log) || CurrentError.StrippedErrorMessage.Contains(TEXT("error")) )
+							if (UE_LOG_ACTIVE(LogShaders, Log) || CurrentError.StrippedErrorMessage.Contains(TEXT("error")))
 							{
 								UniqueErrors.AddUnique(CurrentJob.Output.Errors[ErrorIndex].GetErrorString());
 								ErrorJobs.AddUnique(&CurrentJob);
 							}
 						}
+					};
+
+				// Gather unique errors
+				for (int32 JobIndex = 0; JobIndex < CompleteJobs.Num(); JobIndex++)
+				{
+					const FShaderCommonCompileJob& CurrentJob = *CompleteJobs[JobIndex];
+					if (!CurrentJob.bSucceeded)
+					{
+						const auto* SingleJob = CurrentJob.GetSingleShaderJob();
+						if (SingleJob)
+						{
+							AddErrorsForFailedJob(*SingleJob);
+						}
+						else
+						{
+							const auto* PipelineJob = CurrentJob.GetShaderPipelineJob();
+							check(PipelineJob);
+							for (const FShaderCommonCompileJob* CommonJob : PipelineJob->StageJobs)
+							{
+								AddErrorsForFailedJob(*CommonJob->GetSingleShaderJob());
+							}
+						}
 					}
 				}
 
@@ -1480,11 +1651,7 @@ bool FShaderCompilingManager::HandlePotentialRetryOnError(TMap<int32, FShaderMap
 				UE_LOG(LogShaderCompilers, Warning, TEXT("%s"), *ErrorString);
 				ErrorString += TEXT("\n");
 
-				for (int32 JobIndex = 0; JobIndex < CompleteJobs.Num(); JobIndex++)
-				{
-					const FShaderCompileJob& CurrentJob = *CompleteJobs[JobIndex];
-
-					if (!CurrentJob.bSucceeded)
+				auto ProcessErrors = [&](const FShaderCompileJob& CurrentJob)
 					{
 						for (int32 ErrorIndex = 0; ErrorIndex < CurrentJob.Output.Errors.Num(); ErrorIndex++)
 						{
@@ -1504,7 +1671,7 @@ bool FShaderCompilingManager::HandlePotentialRetryOnError(TMap<int32, FShaderMap
 									// So we don't need the actual filled out version to find the line of a code bug.
 									CurrentError.ErrorFile = TEXT("MaterialTemplate.usf");
 								}
-								else if (CurrentError.ErrorFile.Contains(TEXT("memory")) )
+								else if (CurrentError.ErrorFile.Contains(TEXT("memory")))
 								{
 									// Files passed to the shader compiler through memory will be named memory
 									// Only the shader's main file is passed through memory without a filename
@@ -1524,7 +1691,7 @@ bool FShaderCompilingManager::HandlePotentialRetryOnError(TMap<int32, FShaderMap
 
 								FString UniqueErrorPrefix;
 
-								if ( CurrentJob.ShaderType )
+								if (CurrentJob.ShaderType)
 								{
 									// Construct a path that will enable VS.NET to find the shader file, relative to the solution
 									const FString SolutionPath = FPaths::RootDir();
@@ -1564,6 +1731,27 @@ bool FShaderCompilingManager::HandlePotentialRetryOnError(TMap<int32, FShaderMap
 								ErrorString += UniqueErrorString;
 							}
 						}
+					};
+
+				for (int32 JobIndex = 0; JobIndex < CompleteJobs.Num(); JobIndex++)
+				{
+					const FShaderCommonCompileJob& CurrentJob = *CompleteJobs[JobIndex];
+					if (!CurrentJob.bSucceeded)
+					{
+						const auto* SingleJob = CurrentJob.GetSingleShaderJob();
+						if (SingleJob)
+						{
+							ProcessErrors(*SingleJob);
+						}
+						else
+						{
+							const auto* PipelineJob = CurrentJob.GetShaderPipelineJob();
+							check(PipelineJob);
+							for (const FShaderCommonCompileJob* CommonJob : PipelineJob->StageJobs)
+							{
+								ProcessErrors(*CommonJob->GetSingleShaderJob());
+							}
+						}
 					}
 				}
 
@@ -1617,10 +1805,23 @@ bool FShaderCompilingManager::HandlePotentialRetryOnError(TMap<int32, FShaderMap
 				// Reset outputs
 				for (int32 JobIndex = 0; JobIndex < Results.FinishedJobs.Num(); JobIndex++)
 				{
-					FShaderCompileJob& CurrentJob = *Results.FinishedJobs[JobIndex];
+					FShaderCommonCompileJob& CurrentJob = *Results.FinishedJobs[JobIndex];
+					auto* SingleJob = CurrentJob.GetSingleShaderJob();
 
 					// NOTE: Changes to MaterialTemplate.usf before retrying won't work, because the entry for Material.usf in CurrentJob.Environment.IncludeFileNameToContentsMap isn't reset
-					CurrentJob.Output = FShaderCompilerOutput();
+					if (SingleJob)
+					{
+						SingleJob->Output = FShaderCompilerOutput();
+					}
+					else
+					{
+						auto* PipelineJob = CurrentJob.GetShaderPipelineJob();
+						for (FShaderCommonCompileJob* CommonJob : PipelineJob->StageJobs)
+						{
+							CommonJob->GetSingleShaderJob()->Output = FShaderCompilerOutput();
+							CommonJob->bFinalized = false;
+						}
+					}
 					CurrentJob.bFinalized = false;
 				}
 
@@ -1667,12 +1868,21 @@ void FShaderCompilingManager::CancelCompilation(const TCHAR* MaterialName, const
 			int32 JobIndex = CompileQueue.Num();
 			while ( --JobIndex >= 0 )
 			{
-				if (FShaderCompileJob* Job = CompileQueue[JobIndex])
+				if (FShaderCommonCompileJob* Job = CompileQueue[JobIndex])
 				{
 					if (Job->Id == MapIdx)
 					{
+						FShaderPipelineCompileJob* PipelineJob = Job->GetShaderPipelineJob();
+						if (PipelineJob)
+						{
+							TotalNumJobsRemoved += PipelineJob->StageJobs.Num();
+							NumJobsRemoved += PipelineJob->StageJobs.Num();
+						}
+						else
+						{
 						++TotalNumJobsRemoved;
 						++NumJobsRemoved;
+						}
 						CompileQueue.RemoveAt(JobIndex, 1, false);
 					}
 				}
@@ -1793,7 +2003,7 @@ void FShaderCompilingManager::ProcessAsyncResults(bool bLimitExecutionTime, bool
 				{
 					const FShaderMapCompileResults& Results = It.Value();
 
-					if (Results.FinishedJobs.Num() == Results.NumJobsQueued)
+					if (GetNumTotalJobs(Results.FinishedJobs) == Results.NumJobsQueued)
 					{
 						PendingFinalizeShaderMaps.Add(It.Key(), FShaderMapFinalizeResults(Results));
 						ShaderMapsToRemove.Add(It.Key());
@@ -1865,7 +2075,7 @@ void GlobalBeginCompileShader(
 	const TCHAR* FunctionName,
 	FShaderTarget Target,
 	FShaderCompileJob* NewJob,
-	TArray<FShaderCompileJob*>& NewJobs,
+	TArray<FShaderCommonCompileJob*>& NewJobs,
 	bool bAllowDevelopmentShaderCompile
 	)
 {
@@ -2043,7 +2253,7 @@ void GlobalBeginCompileShader(
 		static IConsoleVariable* CVar = IConsoleManager::Get().FindConsoleVariable(TEXT("r.DBuffer"));
 		Input.Environment.SetDefine(TEXT("USE_DBUFFER"), CVar ? CVar->GetInt() : 0);
 	}
-
+	
 	NewJobs.Add(NewJob);
 }
 
@@ -2225,20 +2435,22 @@ bool RecompileShaders(const TCHAR* Cmd, FOutputDevice& Ar)
 		{
 			TArray<FShaderType*> OutdatedShaderTypes;
 			TArray<const FVertexFactoryType*> OutdatedFactoryTypes;
+			TArray<const FShaderPipelineType*> OutdatedShaderPipelineTypes;
 			{
 				FRecompileShadersTimer SearchTimer(TEXT("Searching for changed files"));
 				FShaderType::GetOutdatedTypes(OutdatedShaderTypes, OutdatedFactoryTypes);
+				FShaderPipelineType::GetOutdatedTypes(OutdatedShaderTypes, OutdatedShaderPipelineTypes, OutdatedFactoryTypes);
 			}
 
-			if (OutdatedShaderTypes.Num() > 0 || OutdatedFactoryTypes.Num() > 0)
+			if (OutdatedShaderPipelineTypes.Num() > 0 || OutdatedShaderTypes.Num() > 0 || OutdatedFactoryTypes.Num() > 0)
 			{
 				FRecompileShadersTimer TestTimer(TEXT("RecompileShaders Changed"));
 
 				// Kick off global shader recompiles
 				UMaterialInterface::IterateOverActiveFeatureLevels([&](ERHIFeatureLevel::Type InFeatureLevel) {
 					auto ShaderPlatform = GShaderPlatformForFeatureLevel[InFeatureLevel];
-					BeginRecompileGlobalShaders(OutdatedShaderTypes, ShaderPlatform);
-					UMaterial::UpdateMaterialShaders(OutdatedShaderTypes, OutdatedFactoryTypes, ShaderPlatform);
+					BeginRecompileGlobalShaders(OutdatedShaderTypes, OutdatedShaderPipelineTypes, ShaderPlatform);
+					UMaterial::UpdateMaterialShaders(OutdatedShaderTypes, OutdatedShaderPipelineTypes, OutdatedFactoryTypes, ShaderPlatform);
 				});
 
 				GWarn->StatusUpdate(0, 1, NSLOCTEXT("ShaderCompilingManager", "CompilingGlobalShaderStatus", "Compiling global shaders..."));
@@ -2306,7 +2518,8 @@ bool RecompileShaders(const TCHAR* Cmd, FOutputDevice& Ar)
 		else
 		{
 			TArray<FShaderType*> ShaderTypes = FShaderType::GetShaderTypesByFilename(*FlagStr);
-			if (ShaderTypes.Num() > 0)
+			TArray<const FShaderPipelineType*> ShaderPipelineTypes = FShaderPipelineType::GetShaderPipelineTypesByFilename(*FlagStr);
+			if (ShaderTypes.Num() > 0 || ShaderPipelineTypes.Num() > 0)
 			{
 				FRecompileShadersTimer TestTimer(TEXT("RecompileShaders SingleShader"));
 				
@@ -2314,8 +2527,8 @@ bool RecompileShaders(const TCHAR* Cmd, FOutputDevice& Ar)
 
 				UMaterialInterface::IterateOverActiveFeatureLevels([&](ERHIFeatureLevel::Type InFeatureLevel) {
 					auto ShaderPlatform = GShaderPlatformForFeatureLevel[InFeatureLevel];
-					BeginRecompileGlobalShaders(ShaderTypes, ShaderPlatform);
-					UMaterial::UpdateMaterialShaders(ShaderTypes, FactoryTypes, ShaderPlatform);
+					BeginRecompileGlobalShaders(ShaderTypes, ShaderPipelineTypes, ShaderPlatform);
+					UMaterial::UpdateMaterialShaders(ShaderTypes, ShaderPipelineTypes, FactoryTypes, ShaderPlatform);
 					FinishRecompileGlobalShaders();
 				});
 			}
@@ -2517,7 +2730,7 @@ void FShaderCompileXGEThreadRunnable::PostCompletedJobsForBatch(FShaderBatch* Ba
 {
 	// Enter the critical section so we can access the input and output queues
 	FScopeLock Lock(&Manager->CompileQueueSection);
-	for (FShaderCompileJob* Job : Batch->GetJobs())
+	for (FShaderCommonCompileJob* Job : Batch->GetJobs())
 	{
 		FShaderMapCompileResults& ShaderMapResults = Manager->ShaderMapJobs.FindChecked(Job->Id);
 		ShaderMapResults.FinishedJobs.Add(Job);
@@ -2528,7 +2741,7 @@ void FShaderCompileXGEThreadRunnable::PostCompletedJobsForBatch(FShaderBatch* Ba
 	FPlatformAtomics::InterlockedAdd(&Manager->NumOutstandingJobs, -Batch->NumJobs());
 }
 
-void FShaderCompileXGEThreadRunnable::FShaderBatch::AddJob(FShaderCompileJob* Job)
+void FShaderCompileXGEThreadRunnable::FShaderBatch::AddJob(FShaderCommonCompileJob* Job)
 {
 	// We can only add jobs to a batch which hasn't been written out yet.
 	if (bTransferFileWritten)
@@ -2829,7 +3042,7 @@ int32 FShaderCompileXGEThreadRunnable::CompilingLoop()
 	}
 
 	// Try to prepare more shader jobs (even if a build is in flight).
-	TArray<FShaderCompileJob*> JobQueue;
+	TArray<FShaderCommonCompileJob*> JobQueue;
 	{
 		// Enter the critical section so we can access the input and output queues
 		FScopeLock Lock(&Manager->CompileQueueSection);
diff --git a/Engine/Source/Runtime/Engine/Private/ShaderCompiler/ShaderPipelineCompiler.cpp b/Engine/Source/Runtime/Engine/Private/ShaderCompiler/ShaderPipelineCompiler.cpp
new file mode 100644
--- /dev/null
+++ b/Engine/Source/Runtime/Engine/Private/ShaderCompiler/ShaderPipelineCompiler.cpp
@@ -0,0 +1,106 @@
+// Copyright 1998-2015 Epic Games, Inc. All Rights Reserved.
+
+/*=============================================================================
+	ShaderPipelineCompiler.cpp: Platform independent shader pipeline compilations.
+=============================================================================*/
+
+#include "EnginePrivate.h"
+#include "EditorSupportDelegates.h"
+#include "ExceptionHandling.h"
+#include "GlobalShader.h"
+#include "TargetPlatform.h"
+#include "DerivedDataCacheInterface.h"
+#include "EngineModule.h"
+#include "ShaderCompiler.h"
+#include "RendererInterface.h"
+
+bool CompileShaderPipeline(const IShaderFormat* Compiler, FName Format, FShaderPipelineCompileJob* PipelineJob, const FString& Dir)
+{
+	check(PipelineJob && PipelineJob->StageJobs.Num() > 0);
+
+	FShaderCompileJob* CurrentJob = PipelineJob->StageJobs[0]->GetSingleShaderJob();
+
+	CurrentJob->Input.bCompilingForShaderPipeline = true;
+
+	// First job doesn't have to trim outputs
+	CurrentJob->Input.bIncludeUsedOutputs = false;
+
+	if (IsValidRef(CurrentJob->Input.SharedEnvironment))
+	{
+		// Merge the shared environment into the per-shader environment before calling into the compile function
+		// Normally this happens in the worker
+		CurrentJob->Input.Environment.Merge(*CurrentJob->Input.SharedEnvironment);
+	}
+
+	// Compile the shader directly through the platform dll (directly from the shader dir as the working directory)
+	Compiler->CompileShader(Format, CurrentJob->Input, CurrentJob->Output, Dir);
+
+	CurrentJob->bSucceeded = CurrentJob->Output.bSucceeded;
+
+	if (!CurrentJob->Output.bSucceeded)
+	{
+		// Can't carry on compiling the pipeline
+		return false;
+	}
+
+	// Generate a hash of the output and cache it
+	// The shader processing this output will use it to search for existing FShaderResources
+	CurrentJob->Output.GenerateOutputHash();
+
+	// This tells the shader compiler we do want to remove unused outputs
+	bool bEnableRemovingUnused = true;
+	// This tells us the hlsl parser failed at removing unused outputs
+	bool bJobFailedRemovingUnused = false;
+
+	//#todo-rco: Currently only removes for pure VS & PS stages
+	for (int32 Index = 0; Index < PipelineJob->StageJobs.Num(); ++Index)
+	{
+		auto Stage = PipelineJob->StageJobs[Index]->GetSingleShaderJob()->Input.Target.Frequency;
+		if (Stage != SF_Vertex && Stage != SF_Pixel)
+		{
+			bEnableRemovingUnused = false;
+			break;
+		}
+	}
+
+	for (int32 Index = 1; Index < PipelineJob->StageJobs.Num(); ++Index)
+	{
+		auto* PreviousJob = CurrentJob;
+		CurrentJob = PipelineJob->StageJobs[Index]->GetSingleShaderJob();
+		bEnableRemovingUnused = bEnableRemovingUnused && PreviousJob->Output.bSupportsQueryingUsedAttributes;
+		if (bEnableRemovingUnused)
+		{
+			CurrentJob->Input.bIncludeUsedOutputs = true;
+			CurrentJob->Input.bCompilingForShaderPipeline = true;
+			CurrentJob->Input.UsedOutputs = PreviousJob->Output.UsedAttributes;
+		}
+
+		if (IsValidRef(CurrentJob->Input.SharedEnvironment))
+		{
+			// Merge the shared environment into the per-shader environment before calling into the compile function
+			// Normally this happens in the worker
+			CurrentJob->Input.Environment.Merge(*CurrentJob->Input.SharedEnvironment);
+		}
+
+		// Compile the shader directly through the platform dll (directly from the shader dir as the working directory)
+		Compiler->CompileShader(Format, CurrentJob->Input, CurrentJob->Output, Dir);
+
+		CurrentJob->bSucceeded = CurrentJob->Output.bSucceeded;
+
+		if (!CurrentJob->Output.bSucceeded)
+		{
+			// Can't carry on compiling the pipeline
+			return false;
+		}
+
+		bJobFailedRemovingUnused = CurrentJob->Output.bFailedRemovingUnused || bJobFailedRemovingUnused;
+
+		// Generate a hash of the output and cache it
+		// The shader processing this output will use it to search for existing FShaderResources
+		CurrentJob->Output.GenerateOutputHash();
+	}
+
+	PipelineJob->bSucceeded = true;
+	PipelineJob->bFailedRemovingUnused = bJobFailedRemovingUnused;
+	return true;
+}
diff --git a/Engine/Source/Runtime/Engine/Private/ShaderDerivedDataVersion.h b/Engine/Source/Runtime/Engine/Private/ShaderDerivedDataVersion.h
--- a/Engine/Source/Runtime/Engine/Private/ShaderDerivedDataVersion.h
+++ b/Engine/Source/Runtime/Engine/Private/ShaderDerivedDataVersion.h
@@ -9,5 +9,5 @@
 // In case of merge conflicts with DDC versions, you MUST generate a new GUID and set this new
 // guid as version
 
-#define GLOBALSHADERMAP_DERIVEDDATA_VER			TEXT("653D597996674C67855E3A31377FCDB5")
-#define MATERIALSHADERMAP_DERIVEDDATA_VER		TEXT("3F0644619F9C48FF91E0D494C95EB038")
+#define GLOBALSHADERMAP_DERIVEDDATA_VER			TEXT("753D597996674C67855E3A31377FCDc6")
+#define MATERIALSHADERMAP_DERIVEDDATA_VER		TEXT("4F0644619F9C48FF91E0D494C95EB049")
diff --git a/Engine/Source/Runtime/Engine/Public/Animation/AnimCurveTypes.h b/Engine/Source/Runtime/Engine/Public/Animation/AnimCurveTypes.h
--- a/Engine/Source/Runtime/Engine/Public/Animation/AnimCurveTypes.h
+++ b/Engine/Source/Runtime/Engine/Public/Animation/AnimCurveTypes.h
@@ -267,18 +267,7 @@ struct ENGINE_API FBlendedCurve
 
 	/** Initialize Curve Data from following data */
 	void InitFrom(const class USkeleton* Skeleton);
-	void InitFrom(const FBlendedCurve& InCurveToInitFrom)
-	{
-		// make sure this doesn't happen
-		if (ensure(&InCurveToInitFrom != this))
-		{
-			UIDList.Reset();
-			UIDList.Append(InCurveToInitFrom.UIDList);
-			Elements.Reset();
-			Elements.AddZeroed(UIDList.Num());
-			bInitialized = true;
-		}
-	}
+	void InitFrom(const FBlendedCurve& InCurveToInitFrom);
 
 	/** Set value of InUID to InValue */
 	void Set(USkeleton::AnimCurveUID InUid, float InValue, int32 InFlags);
diff --git a/Engine/Source/Runtime/Engine/Public/Animation/AnimStats.h b/Engine/Source/Runtime/Engine/Public/Animation/AnimStats.h
--- a/Engine/Source/Runtime/Engine/Public/Animation/AnimStats.h
+++ b/Engine/Source/Runtime/Engine/Public/Animation/AnimStats.h
@@ -39,6 +39,7 @@ DECLARE_CYCLE_STAT_EXTERN(TEXT("AnimMontageInstance_Advance"), STAT_AnimMontageI
 DECLARE_CYCLE_STAT_EXTERN(TEXT("AnimMontageInstance_TickBranchPoints"), STAT_AnimMontageInstance_TickBranchPoints, STATGROUP_Anim, );
 DECLARE_CYCLE_STAT_EXTERN(TEXT("AnimMontageInstance_Advance_Iteration"), STAT_AnimMontageInstance_Advance_Iteration, STATGROUP_Anim, );
 DECLARE_CYCLE_STAT_EXTERN(TEXT("UpdateCurves"), STAT_UpdateCurves, STATGROUP_Anim, );
+DECLARE_CYCLE_STAT_EXTERN(TEXT("LocalBlendCSBoneTransforms"), STAT_LocalBlendCSBoneTransforms, STATGROUP_Anim, ENGINE_API);
 
 #define DO_ANIMSTAT_PROCESSING(StatName) DECLARE_CYCLE_STAT_EXTERN(TEXT(#StatName), STAT_ ## StatName, STATGROUP_Anim, ENGINE_API)
 #include "AnimMTStats.h"
diff --git a/Engine/Source/Runtime/Engine/Public/BonePose.h b/Engine/Source/Runtime/Engine/Public/BonePose.h
--- a/Engine/Source/Runtime/Engine/Public/BonePose.h
+++ b/Engine/Source/Runtime/Engine/Public/BonePose.h
@@ -7,7 +7,7 @@
 #include "BoneIndices.h"
 #include "CustomBoneIndexArray.h"
 #include "AnimEncoding.h"
-
+#include "Animation/AnimStats.h"
 
 struct FCompactPose;
 
@@ -552,6 +552,8 @@ void FCSPose<PoseType>::SafeSetCSBoneTransforms(const TArray<struct FBoneTransfo
 template<class PoseType>
 void FCSPose<PoseType>::LocalBlendCSBoneTransforms(const TArray<struct FBoneTransform>& BoneTransforms, float Alpha)
 {
+	SCOPE_CYCLE_COUNTER(STAT_LocalBlendCSBoneTransforms);
+
 	// if Alpha is small enough, skip
 	if (Alpha < ZERO_ANIMWEIGHT_THRESH)
 	{
diff --git a/Engine/Source/Runtime/Engine/Public/GlobalShader.h b/Engine/Source/Runtime/Engine/Public/GlobalShader.h
--- a/Engine/Source/Runtime/Engine/Public/GlobalShader.h
+++ b/Engine/Source/Runtime/Engine/Public/GlobalShader.h
@@ -9,6 +9,7 @@
 #include "Shader.h"
 #include "SceneManagement.h"
 
+class FShaderCommonCompileJob;
 class FShaderCompileJob;
 
 /** Used to identify the global shader map in compile queues. */
@@ -26,7 +27,6 @@ public:
 	void AppendKeyString(FString& KeyString) const;
 
 private:
-
 	/** Shader types that this shader map is dependent on and their stored state. */
 	TArray<FShaderTypeDependency> ShaderTypeDependencies;
 };
@@ -63,15 +63,16 @@ public:
 
 	/**
 	 * Enqueues compilation of a shader of this type.  
-	 * @param Compiler - The shader compiler to use.
-	 * @param OutErrors - Upon compilation failure, OutErrors contains a list of the errors which occured.
-	 * @param bDebugDump - Dump out the preprocessed and disassembled shader for debugging.
-	 * @param ShaderSubDir - Sub directory for dumping out preprocessor output.
 	 */
-	ENGINE_API void BeginCompileShader(EShaderPlatform Platform, TArray<FShaderCompileJob*>& NewJobs);
+	ENGINE_API void BeginCompileShader(EShaderPlatform Platform, TArray<FShaderCommonCompileJob*>& NewJobs);
+
+	/**
+	 * Enqueues compilation of a shader pipeline of this type.  
+	 */
+	ENGINE_API static void BeginCompileShaderPipeline(EShaderPlatform Platform, const FShaderPipelineType* ShaderPipeline, const TArray<FGlobalShaderType*>& ShaderStages, TArray<FShaderCommonCompileJob*>& NewJobs);
 
 	/** Either returns an equivalent existing shader of this type, or constructs a new instance. */
-	FShader* FinishCompileShader(const FShaderCompileJob& CompileJob);
+	FShader* FinishCompileShader(const FShaderCompileJob& CompileJob, const FShaderPipelineType* ShaderPipelineType);
 
 	/**
 	 * Checks if the shader type should be cached for a particular platform.
@@ -215,13 +216,13 @@ extern ENGINE_API void RecompileGlobalShaders();
  * Begins recompiling the specified global shader types, and flushes their bound shader states.
  * FinishRecompileGlobalShaders must be called after this and before using the global shaders for anything.
  */
-extern ENGINE_API void BeginRecompileGlobalShaders(const TArray<FShaderType*>& OutdatedShaderTypes, EShaderPlatform ShaderPlatform);
+extern ENGINE_API void BeginRecompileGlobalShaders(const TArray<FShaderType*>& OutdatedShaderTypes, const TArray<const FShaderPipelineType*>& OutdatedShaderPipelineTypes, EShaderPlatform ShaderPlatform);
 
 /** Finishes recompiling global shaders.  Must be called after BeginRecompileGlobalShaders. */
 extern ENGINE_API void FinishRecompileGlobalShaders();
 
 /** Called by the shader compiler to process completed global shader jobs. */
-extern void ProcessCompiledGlobalShaders(const TArray<FShaderCompileJob*>& CompilationResults);
+extern void ProcessCompiledGlobalShaders(const TArray<FShaderCommonCompileJob*>& CompilationResults);
 
 extern ENGINE_API FString GetGlobalShaderMapDDCKey();
 
diff --git a/Engine/Source/Runtime/Engine/Public/MaterialShaderType.h b/Engine/Source/Runtime/Engine/Public/MaterialShaderType.h
--- a/Engine/Source/Runtime/Engine/Public/MaterialShaderType.h
+++ b/Engine/Source/Runtime/Engine/Public/MaterialShaderType.h
@@ -58,10 +58,11 @@ public:
 			FShaderResource* InResource,
 			const FUniformExpressionSet& InUniformExpressionSet,
 			const FSHAHash& InMaterialShaderMapHash,
+			const FShaderPipelineType* InShaderPipeline,
 			FVertexFactoryType* InVertexFactoryType,
 			const FString& InDebugDescription
 			)
-		: FGlobalShaderType::CompiledShaderInitializerType(InType,CompilerOutput,InResource,InMaterialShaderMapHash,InVertexFactoryType)
+		: FGlobalShaderType::CompiledShaderInitializerType(InType,CompilerOutput,InResource,InMaterialShaderMapHash,InShaderPipeline,InVertexFactoryType)
 		, UniformExpressionSet(InUniformExpressionSet)
 		, DebugDescription(InDebugDescription)
 		{}
@@ -96,9 +97,18 @@ public:
 		const FMaterial* Material,
 		FShaderCompilerEnvironment* MaterialEnvironment,
 		EShaderPlatform Platform,
-		TArray<FShaderCompileJob*>& NewJobs
+		TArray<FShaderCommonCompileJob*>& NewJobs
 		);
 
+	static void BeginCompileShaderPipeline(
+		uint32 ShaderMapId,
+		EShaderPlatform Platform,
+		const FMaterial* Material,
+		FShaderCompilerEnvironment* MaterialEnvironment,
+		const FShaderPipelineType* ShaderPipeline,
+		const TArray<FMaterialShaderType*>& ShaderStages,
+		TArray<FShaderCommonCompileJob*>& NewJobs);
+
 	/**
 	 * Either creates a new instance of this type or returns an equivalent existing shader.
 	 * @param Material - The material to link the shader with.
@@ -108,6 +118,7 @@ public:
 		const FUniformExpressionSet& UniformExpressionSet,
 		const FSHAHash& MaterialShaderMapHash,
 		const FShaderCompileJob& CurrentJob,
+		const FShaderPipelineType* ShaderPipeline,
 		const FString& InDebugDescription
 		);
 
diff --git a/Engine/Source/Runtime/Engine/Public/MaterialShared.h b/Engine/Source/Runtime/Engine/Public/MaterialShared.h
--- a/Engine/Source/Runtime/Engine/Public/MaterialShared.h
+++ b/Engine/Source/Runtime/Engine/Public/MaterialShared.h
@@ -20,6 +20,7 @@ class FMaterial;
 class FMaterialRenderProxy;
 class FMeshMaterialShaderMap;
 class FMeshMaterialShaderType;
+class FShaderCommonCompileJob;
 class FShaderCompileJob;
 class FShaderType;
 class FShaderTypeDependency;
@@ -52,7 +53,9 @@ extern void GetMaterialQualityLevelName(EMaterialQualityLevel::Type InMaterialQu
 
 inline bool IsSubsurfaceShadingModel(EMaterialShadingModel ShadingModel)
 {
-	return ShadingModel == MSM_Subsurface || ShadingModel == MSM_PreintegratedSkin || ShadingModel == MSM_SubsurfaceProfile || ShadingModel == MSM_TwoSidedFoliage || ShadingModel == MSM_Cloth;
+	return ShadingModel == MSM_Subsurface || ShadingModel == MSM_PreintegratedSkin || 
+		ShadingModel == MSM_SubsurfaceProfile || ShadingModel == MSM_TwoSidedFoliage || 
+		ShadingModel == MSM_Cloth;
 }
 
 /**
@@ -422,7 +425,7 @@ public:
 	~FMaterialShaderMapId()
 	{ }
 
-	ENGINE_API void SetShaderDependencies(const TArray<FShaderType*>& ShaderTypes, const TArray<FVertexFactoryType*>& VFTypes);
+	ENGINE_API void SetShaderDependencies(const TArray<FShaderType*>& ShaderTypes, const TArray<const FShaderPipelineType*>& ShaderPipelineTypes, const TArray<FVertexFactoryType*>& VFTypes);
 
 	void Serialize(FArchive& Ar);
 
@@ -489,7 +492,7 @@ public:
 		const FMaterial* Material,
 		FShaderCompilerEnvironment* MaterialEnvironment,
 		EShaderPlatform Platform,
-		TArray<FShaderCompileJob*>& NewJobs
+		TArray<FShaderCommonCompileJob*>& NewJobs
 		);
 
 	/**
@@ -525,6 +528,12 @@ public:
 	 */
 	void FlushShadersByShaderType(FShaderType* ShaderType);
 
+		/**
+	 * Removes all entries in the cache with exceptions based on a shader type
+	 * @param ShaderType - The shader type to flush
+	 */
+	void FlushShadersByShaderPipelineType(const FShaderPipelineType* ShaderPipelineType);
+
 	// Accessors.
 	FVertexFactoryType* GetVertexFactoryType() const { return VertexFactoryType; }
 
@@ -549,9 +558,12 @@ public:
 	static FMaterialShaderMap* FindId(const FMaterialShaderMapId& ShaderMapId, EShaderPlatform Platform);
 
 	/** Flushes the given shader types from any loaded FMaterialShaderMap's. */
-	static void FlushShaderTypes(TArray<FShaderType*>& ShaderTypesToFlush, TArray<const FVertexFactoryType*>& VFTypesToFlush);
+	static void FlushShaderTypes(TArray<FShaderType*>& ShaderTypesToFlush, TArray<const FShaderPipelineType*>& ShaderPipelineTypesToFlush, TArray<const FVertexFactoryType*>& VFTypesToFlush);
 
-	static void FixupShaderTypes(EShaderPlatform Platform, const TMap<FShaderType*, FString>& ShaderTypeNames, const TMap<FVertexFactoryType*, FString>& VertexFactoryTypeNames);
+	static void FixupShaderTypes(EShaderPlatform Platform, 
+		const TMap<FShaderType*, FString>& ShaderTypeNames,
+		const TMap<const FShaderPipelineType*, FString>& ShaderPipelineTypeNames,
+		const TMap<FVertexFactoryType*, FString>& VertexFactoryTypeNames);
 
 	/** 
 	 * Attempts to load the shader map for the given material from the Derived Data Cache.
@@ -581,7 +593,7 @@ public:
 		);
 
 	/** Sorts the incoming compiled jobs into the appropriate mesh shader maps, and finalizes this shader map so that it can be used for rendering. */
-	bool ProcessCompilationResults(const TArray<FShaderCompileJob*>& InCompilationResults, int32& ResultIndex, float& TimeBudget);
+	bool ProcessCompilationResults(const TArray<FShaderCommonCompileJob*>& InCompilationResults, int32& ResultIndex, float& TimeBudget);
 
 	/**
 	 * Checks whether the material shader map is missing any shader types necessary for the given material.
@@ -602,7 +614,7 @@ public:
 	bool TryToAddToExistingCompilationTask(FMaterial* Material);
 
 	/** Builds a list of the shaders in a shader map. */
-	ENGINE_API void GetShaderList(TMap<FShaderId,FShader*>& OutShaders) const;
+	ENGINE_API void GetShaderList(TMap<FShaderId, FShader*>& OutShaders) const;
 
 	/** Registers a material shader map in the global map so it can be used by materials. */
 	void Register(EShaderPlatform InShaderPlatform);
@@ -624,6 +636,12 @@ public:
 	 */
 	void FlushShadersByShaderType(FShaderType* ShaderType);
 
+	/**
+	 * Removes all entries in the cache with exceptions based on a shader pipeline type
+	 * @param ShaderPipelineType - The shader pipeline type to flush
+	 */
+	void FlushShadersByShaderPipelineType(const FShaderPipelineType* ShaderPipelineType);
+
 	/**
 	 * Removes all entries in the cache with exceptions based on a vertex factory type
 	 * @param ShaderType - The shader type to flush
@@ -682,7 +700,7 @@ public:
 	bool UsesEyeAdaptation() const { return MaterialCompilationOutput.bUsesEyeAdaptation; }
 	bool ModifiesMeshPosition() const { return MaterialCompilationOutput.bModifiesMeshPosition; }
 
-	uint32 GetNumInstructionsForShader(const FShaderType* ShaderType) const;
+	uint32 GetMaxNumInstructionsForShader(const FShaderType* ShaderType) const;
 
 	bool IsValidForRendering() const
 	{
@@ -760,6 +778,8 @@ private:
 	/** Debug information about how the material shader map was compiled. */
 	FString DebugDescription;
 
+	FShader* ProcessCompilationResultsForSingleJob(class FShaderCompileJob* SingleJob, const FShaderPipelineType* ShaderPipeline, const FSHAHash& MaterialShaderMapHash);
+
 	/** Initializes OrderedMeshShaderMaps from the contents of MeshShaderMaps. */
 	void InitOrderedMeshShaderMaps();
 
@@ -1181,7 +1201,7 @@ protected:
 	virtual FGuid GetMaterialId() const = 0;
 	
 	/** Produces arrays of any shader and vertex factory type that this material is dependent on. */
-	ENGINE_API void GetDependentShaderAndVFTypes(EShaderPlatform Platform, TArray<FShaderType*>& OutShaderTypes, TArray<FVertexFactoryType*>& OutVFTypes) const;
+	ENGINE_API void GetDependentShaderAndVFTypes(EShaderPlatform Platform, TArray<FShaderType*>& OutShaderTypes, TArray<const FShaderPipelineType*>& OutShaderPipelineTypes, TArray<FVertexFactoryType*>& OutVFTypes) const;
 
 private:
 
diff --git a/Engine/Source/Runtime/Engine/Public/MeshMaterialShaderType.h b/Engine/Source/Runtime/Engine/Public/MeshMaterialShaderType.h
--- a/Engine/Source/Runtime/Engine/Public/MeshMaterialShaderType.h
+++ b/Engine/Source/Runtime/Engine/Public/MeshMaterialShaderType.h
@@ -31,9 +31,10 @@ public:
 			const FUniformExpressionSet& InUniformExpressionSet,
 			const FSHAHash& InMaterialShaderMapHash,
 			const FString& InDebugDescription,
+			const FShaderPipelineType* InShaderPipeline,
 			FVertexFactoryType* InVertexFactoryType
 			):
-			FMaterialShaderType::CompiledShaderInitializerType(InType,CompilerOutput,InResource,InUniformExpressionSet,InMaterialShaderMapHash,InVertexFactoryType,InDebugDescription),
+			FMaterialShaderType::CompiledShaderInitializerType(InType,CompilerOutput,InResource,InUniformExpressionSet,InMaterialShaderMapHash,InShaderPipeline,InVertexFactoryType,InDebugDescription),
 			VertexFactoryType(InVertexFactoryType)
 		{}
 	};
@@ -70,7 +71,19 @@ public:
 		const FMaterial* Material,
 		FShaderCompilerEnvironment* MaterialEnvironment,
 		FVertexFactoryType* VertexFactoryType,
-		TArray<FShaderCompileJob*>& NewJobs
+		const FShaderPipelineType* ShaderPipeline,
+		TArray<FShaderCommonCompileJob*>& NewJobs
+		);
+
+	static void BeginCompileShaderPipeline(
+		uint32 ShaderMapId,
+		EShaderPlatform Platform,
+		const FMaterial* Material,
+		FShaderCompilerEnvironment* MaterialEnvironment,
+		FVertexFactoryType* VertexFactoryType,
+		const FShaderPipelineType* ShaderPipeline,
+		const TArray<FMeshMaterialShaderType*>& ShaderStages,
+		TArray<FShaderCommonCompileJob*>& NewJobs
 		);
 
 	/**
@@ -82,6 +95,7 @@ public:
 		const FUniformExpressionSet& UniformExpressionSet, 
 		const FSHAHash& MaterialShaderMapHash,
 		const FShaderCompileJob& CurrentJob,
+		const FShaderPipelineType* ShaderPipeline,
 		const FString& InDebugDescription
 		);
 
diff --git a/Engine/Source/Runtime/Engine/Public/ParticleHelper.h b/Engine/Source/Runtime/Engine/Public/ParticleHelper.h
--- a/Engine/Source/Runtime/Engine/Public/ParticleHelper.h
+++ b/Engine/Source/Runtime/Engine/Public/ParticleHelper.h
@@ -1243,6 +1243,23 @@ struct FParticleDataContainer
 	void Free();
 };
 
+struct FMacroUVOverride
+{
+	FMacroUVOverride() : bOverride(false), Radius(0.f), Position(0.f,0.f,0.f) {}
+
+	bool	bOverride;
+	float   Radius;
+	FVector Position;
+
+	friend FORCEINLINE FArchive& operator<<(FArchive& Ar, FMacroUVOverride& O)
+	{
+		Ar << O.bOverride;
+		Ar << O.Radius;
+		Ar << O.Position;
+		return Ar;
+	}
+};
+
 /** Source data base class for all emitter types */
 struct FDynamicEmitterReplayDataBase
 {
@@ -1261,10 +1278,7 @@ struct FDynamicEmitterReplayDataBase
 	int32 SortMode;
 
 	/** MacroUV (override) data **/
-	bool   bOverrideSystemMacroUV;
-	float   MacroUVRadius;
-	FVector MacroUVPosition;
-
+	FMacroUVOverride MacroUVOverride;
 
 	/** Constructor */
 	FDynamicEmitterReplayDataBase()
@@ -1272,10 +1286,7 @@ struct FDynamicEmitterReplayDataBase
 		  ActiveParticleCount( 0 ),
 		  ParticleStride( 0 ),
 		  Scale( FVector( 1.0f ) ),
-		  SortMode(0),	// Default to PSORTMODE_None
-		  bOverrideSystemMacroUV(0),
-		  MacroUVRadius(0.f),
-		  MacroUVPosition(0.f,0.f,0.f)
+		  SortMode(0)	// Default to PSORTMODE_None		  
 	{
 	}
 
@@ -1333,9 +1344,7 @@ struct FDynamicEmitterReplayDataBase
 
 		Ar << Scale;
 		Ar << SortMode;
-		Ar << bOverrideSystemMacroUV;
-		Ar << MacroUVRadius;
-		Ar << MacroUVPosition;
+		Ar << MacroUVOverride;
 	}
 
 };
@@ -1406,6 +1415,9 @@ struct FDynamicEmitterDataBase
 	/** Returns the source data for this particle system */
 	virtual const FDynamicEmitterReplayDataBase& GetSource() const = 0;
 
+	/** Returns the current macro uv override. Specialized by FGPUSpriteDynamicEmitterData  */
+	virtual const FMacroUVOverride& GetMacroUVOverride() const { return GetSource().MacroUVOverride; }
+
 	/** true if this emitter is currently selected */
 	uint32	bSelected:1;
 	/** true if this emitter has valid rendering data */
diff --git a/Engine/Source/Runtime/Engine/Public/SceneView.h b/Engine/Source/Runtime/Engine/Public/SceneView.h
--- a/Engine/Source/Runtime/Engine/Public/SceneView.h
+++ b/Engine/Source/Runtime/Engine/Public/SceneView.h
@@ -458,6 +458,7 @@ BEGIN_UNIFORM_BUFFER_STRUCT_WITH_CONSTRUCTOR(FViewUniformShaderParameters,ENGINE
 	DECLARE_UNIFORM_BUFFER_STRUCT_MEMBER(FVector2D,RenderTargetSize)
 	DECLARE_UNIFORM_BUFFER_STRUCT_MEMBER(float,SkyLightParameters)
 	DECLARE_UNIFORM_BUFFER_STRUCT_MEMBER(FVector4,SceneTextureMinMax)
+	DECLARE_UNIFORM_BUFFER_STRUCT_MEMBER(FVector4,CircleDOFParams)
 	DECLARE_UNIFORM_BUFFER_STRUCT_MEMBER(FLinearColor,SkyLightColor)
 	DECLARE_UNIFORM_BUFFER_STRUCT_MEMBER_ARRAY(FVector4,SkyIrradianceEnvironmentMap,[7])
 	DECLARE_UNIFORM_BUFFER_STRUCT_MEMBER(float, MobilePreviewMode)
diff --git a/Engine/Source/Runtime/Engine/Public/ShaderCompiler.h b/Engine/Source/Runtime/Engine/Public/ShaderCompiler.h
--- a/Engine/Source/Runtime/Engine/Public/ShaderCompiler.h
+++ b/Engine/Source/Runtime/Engine/Public/ShaderCompiler.h
@@ -10,6 +10,7 @@
 
 DECLARE_LOG_CATEGORY_EXTERN(LogShaderCompilers, Log, All);
 
+class FShaderCompileJob;
 class FShaderPipelineCompileJob;
 
 /** Stores all of the common information used to compile a shader or pipeline. */
@@ -34,6 +35,8 @@ public:
 
 	virtual FShaderCompileJob* GetSingleShaderJob() { return nullptr; }
 	virtual const FShaderCompileJob* GetSingleShaderJob() const { return nullptr; }
+	virtual FShaderPipelineCompileJob* GetShaderPipelineJob() { return nullptr; }
+	virtual const FShaderPipelineCompileJob* GetShaderPipelineJob() const { return nullptr; }
 };
 
 
@@ -56,13 +59,40 @@ public:
 	{
 	}
 
-	~FShaderCompileJob()
+	virtual FShaderCompileJob* GetSingleShaderJob() override { return this; }
+	virtual const FShaderCompileJob* GetSingleShaderJob() const override { return this; }
+};
+
+class FShaderPipelineCompileJob : public FShaderCommonCompileJob
+{
+public:
+	TArray<FShaderCommonCompileJob*> StageJobs;
+	bool bFailedRemovingUnused;
+
+	/** Shader pipeline that this shader belongs to, may (currently) be NULL */
+	const FShaderPipelineType* ShaderPipeline;
+
+	FShaderPipelineCompileJob(uint32 InId, const FShaderPipelineType* InShaderPipeline, int32 NumStages) :
+		FShaderCommonCompileJob(InId),
+		bFailedRemovingUnused(false),
+		ShaderPipeline(InShaderPipeline)
 	{
+		check(InShaderPipeline && InShaderPipeline->GetName());
+		check(NumStages > 0);
+		StageJobs.Empty(NumStages);
 	}
 
-	virtual FShaderCompileJob* GetSingleShaderJob() override { return this; }
-	virtual const FShaderCompileJob* GetSingleShaderJob() const override { return this; }
+	~FShaderPipelineCompileJob()
+	{
+		for (int32 Index = 0; Index < StageJobs.Num(); ++Index)
+		{
+			delete StageJobs[Index];
+		}
+		StageJobs.Reset();
+	}
 
+	virtual FShaderPipelineCompileJob* GetShaderPipelineJob() override { return this; }
+	virtual const FShaderPipelineCompileJob* GetShaderPipelineJob() const override { return this; }
 };
 
 class FShaderCompileThreadRunnableBase : public FRunnable
@@ -165,7 +195,7 @@ private:
 	 */
 	class FShaderBatch
 	{
-		TArray<FShaderCompileJob*> Jobs;
+		TArray<FShaderCommonCompileJob*> Jobs;
 		bool bTransferFileWritten;
 
 	public:
@@ -200,12 +230,12 @@ private:
 		{
 			return Jobs.Num();
 		}
-		inline const TArray<FShaderCompileJob*>& GetJobs() const
+		inline const TArray<FShaderCommonCompileJob*>& GetJobs() const
 		{
 			return Jobs;
 		}
 
-		void AddJob(FShaderCompileJob* Job);
+		void AddJob(FShaderCommonCompileJob* Job);
 		
 		void WriteTransferFile();
 	};
@@ -251,7 +281,7 @@ struct FShaderMapCompileResults
 	int32 NumJobsQueued;
 	bool bAllJobsSucceeded;
 	bool bApplyCompletedShaderMapForRendering;
-	TArray<FShaderCompileJob*> FinishedJobs;
+	TArray<FShaderCommonCompileJob*> FinishedJobs;
 };
 
 /** Results for a single compiled and finalized shader map. */
@@ -283,7 +313,7 @@ private:
 	/** Tracks whether we are compiling while the game is running.  If true, we need to throttle down shader compiling CPU usage to avoid starving the runtime threads. */
 	bool bCompilingDuringGame;
 	/** Queue of tasks that haven't been assigned to a worker yet. */
-	TArray<FShaderCompileJob*> CompileQueue;
+	TArray<FShaderCommonCompileJob*> CompileQueue;
 	/** Map from shader map Id to the compile results for that map, used to gather compiled results. */
 	TMap<int32, FShaderMapCompileResults> ShaderMapJobs;
 	/** Number of jobs currently being compiled.  This includes CompileQueue and any jobs that have been assigned to workers but aren't complete yet. */
@@ -401,7 +431,7 @@ public:
 	 * Adds shader jobs to be asynchronously compiled. 
 	 * FinishCompilation or ProcessAsyncResults must be used to get the results.
 	 */
-	ENGINE_API void AddJobs(TArray<FShaderCompileJob*>& NewJobs, bool bApplyCompletedShaderMapForRendering, bool bOptimizeForLowLatency);
+	ENGINE_API void AddJobs(TArray<FShaderCommonCompileJob*>& NewJobs, bool bApplyCompletedShaderMapForRendering, bool bOptimizeForLowLatency);
 
 	/**
 	* Removes all outstanding compile jobs for the passed shader maps.
@@ -457,7 +487,7 @@ extern void GlobalBeginCompileShader(
 	const TCHAR* FunctionName,
 	FShaderTarget Target,
 	FShaderCompileJob* NewJob,
-	TArray<FShaderCompileJob*>& NewJobs,
+	TArray<FShaderCommonCompileJob*>& NewJobs,
 	bool bAllowDevelopmentShaderCompile = true
 	);
 
@@ -465,4 +495,4 @@ extern void GlobalBeginCompileShader(
 extern bool RecompileShaders(const TCHAR* Cmd, FOutputDevice& Ar);
 
 /** Returns whether all global shader types containing the substring are complete and ready for rendering. if type name is null, check everything */
-extern ENGINE_API bool IsGlobalShaderMapComplete(const TCHAR* TypeNameSubstring = nullptr);
\ No newline at end of file
+extern ENGINE_API bool IsGlobalShaderMapComplete(const TCHAR* TypeNameSubstring = nullptr);
diff --git a/Engine/Source/Runtime/LevelSequence/Private/LevelSequence.cpp b/Engine/Source/Runtime/LevelSequence/Private/LevelSequence.cpp
--- a/Engine/Source/Runtime/LevelSequence/Private/LevelSequence.cpp
+++ b/Engine/Source/Runtime/LevelSequence/Private/LevelSequence.cpp
@@ -4,6 +4,7 @@
 #include "LevelSequence.h"
 #include "LevelSequenceObject.h"
 #include "MovieScene.h"
+#include "MovieSceneCommonHelpers.h"
 #include "Engine/Blueprint.h"
 
 
@@ -67,9 +68,33 @@ void ULevelSequence::BindToContext(UObject* NewContext)
 
 void ULevelSequence::BindPossessableObject(const FGuid& ObjectId, UObject& PossessedObject)
 {
-	if (UObject* Context = ResolutionContext.Get())
+	
+	UObject* Context = ResolutionContext.Get();
+
+	// Always use the parent as a context when there is a parent spawnable
+	UObject* ParentObject = GetParentObject(&PossessedObject);
+	FGuid ParentId = ParentObject ? FindObjectId(*ParentObject) : FGuid();
+
+	FMovieSceneSpawnable* ParentSpawnable = ParentId.IsValid() ? MovieScene->FindSpawnable(ParentId) : nullptr;
+	if (ParentSpawnable)
+	{
+		Context = ParentObject;
+	}
+
+	if (Context)
 	{
 		ObjectReferences.CreateBinding(ObjectId, &PossessedObject, Context);
+
+		MovieSceneHelpers::SetRuntimeObjectMobility(&PossessedObject);
+
+		// Ensure that we report possessables to their parent spawnable, if necessary
+		FMovieScenePossessable* ChildPossessable = MovieScene->FindPossessable(ObjectId);
+
+		if (ParentSpawnable && ensureMsgf(ChildPossessable, TEXT("Object being possessed without a corresponding possessable")))
+		{
+			ParentSpawnable->AddChildPossessable(ObjectId);
+			ChildPossessable->SetParentSpawnable(ParentId);
+		}
 	}
 }
 
@@ -97,6 +122,13 @@ UObject* ULevelSequence::FindObject(const FGuid& ObjectId) const
 
 	UObject* Context = ResolutionContext.Get();
 
+	// Always use the parent as a context when possible
+	FMovieScenePossessable* Possessable = MovieScene->FindPossessable(ObjectId);
+	if (Possessable && Possessable->GetParentSpawnable().IsValid())
+	{
+		Context = FindObject(Possessable->GetParentSpawnable());
+	}
+
 	// Attempt to resolve the object binding through the remapped bindings
 	UObject* FoundObject = Context ? ObjectReferences.ResolveBinding(ObjectId, Context) : nullptr;
 
@@ -187,11 +219,6 @@ UObject* ULevelSequence::SpawnObject(const FGuid& ObjectId)
 		return nullptr;
 	}
 
-	AActor* ActorCDO = CastChecked<AActor>(SpawnableClass->ClassDefaultObject);
-
-	const FVector SpawnLocation = ActorCDO->GetRootComponent()->RelativeLocation;
-	const FRotator SpawnRotation = ActorCDO->GetRootComponent()->RelativeRotation;
-
 	// @todo sequencer: We should probably spawn these in a specific sub-level!
 	// World->CurrentLevel = ???;
 
@@ -214,8 +241,13 @@ UObject* ULevelSequence::SpawnObject(const FGuid& ObjectId)
 	}
 
 	FTransform SpawnTransform;
-	SpawnTransform.SetTranslation(SpawnLocation);
-	SpawnTransform.SetRotation(SpawnRotation.Quaternion());
+
+	AActor* ActorCDO = CastChecked<AActor>(SpawnableClass->ClassDefaultObject);
+	if (USceneComponent* RootComponent = ActorCDO->GetRootComponent())
+	{
+		SpawnTransform.SetTranslation(RootComponent->RelativeLocation);
+		SpawnTransform.SetRotation(RootComponent->RelativeRotation.Quaternion());
+	}
 
 	//@todo: This is in place of SpawnActorAbsolute which doesn't exist yet in the Orion branch
 	FTransform NewTransform = SpawnTransform;
@@ -236,8 +268,6 @@ UObject* ULevelSequence::SpawnObject(const FGuid& ObjectId)
 		}
 	}
 
-
-
 	UObject* SpawnedObject = GWorld->SpawnActor(SpawnableClass, &NewTransform, SpawnInfo);
 	if (!SpawnedObject)
 	{
@@ -247,6 +277,8 @@ UObject* ULevelSequence::SpawnObject(const FGuid& ObjectId)
 	SpawnedObjects.Add(ObjectId, SpawnedObject);
 	CachedObjectBindings.Add(ObjectId, SpawnedObject);
 
+	MovieSceneHelpers::SetRuntimeObjectMobility(SpawnedObject);
+
 	return SpawnedObject;
 }
 
diff --git a/Engine/Source/Runtime/LevelSequence/Private/LevelSequenceObjectReference.cpp b/Engine/Source/Runtime/LevelSequence/Private/LevelSequenceObjectReference.cpp
--- a/Engine/Source/Runtime/LevelSequence/Private/LevelSequenceObjectReference.cpp
+++ b/Engine/Source/Runtime/LevelSequence/Private/LevelSequenceObjectReference.cpp
@@ -7,29 +7,43 @@ FLevelSequenceObjectReference::FLevelSequenceObjectReference(UObject* InObject,
 {
 	if (InObject->IsA<AActor>() || InObject->IsA<UActorComponent>())
 	{
-		ObjectId = FLazyObjectPtr(InObject).GetUniqueID();
-		ObjectPath = InObject->GetPathName(InContext);
+		if (InContext->IsA<AActor>())
+		{
+			// When the context is an actor, we only use path name lookup within the actor since it's assumed the actor will be a spawnable
+			// As such, a persistent identifier for a dynamically spawned actor is not tenable
+			ObjectPath = InObject->GetPathName(InContext);
+			ObjectId = FUniqueObjectGuid();
+		}
+		else
+		{
+			ObjectId = FLazyObjectPtr(InObject).GetUniqueID();
+			ObjectPath = InObject->GetPathName(InContext);
+		}
 	}
 }
 
 UObject* FLevelSequenceObjectReference::Resolve(UObject* InContext) const
 {
-#if WITH_EDITOR
-	int32 PIEInstanceID = InContext->GetOutermost()->PIEInstanceID;
-	FUniqueObjectGuid FixedUpId = PIEInstanceID == -1 ? ObjectId : ObjectId.FixupForPIE(PIEInstanceID);
-#else
-	FUniqueObjectGuid FixedUpId = ObjectId;
-#endif
-	
-
-	FLazyObjectPtr LazyPtr;
-	LazyPtr = FixedUpId;
-
-	if (UObject* FoundObject = LazyPtr.Get())
+	if (ObjectId.IsValid())
 	{
-		return FoundObject;
+	#if WITH_EDITOR
+		int32 PIEInstanceID = InContext->GetOutermost()->PIEInstanceID;
+		FUniqueObjectGuid FixedUpId = PIEInstanceID == -1 ? ObjectId : ObjectId.FixupForPIE(PIEInstanceID);
+	#else
+		FUniqueObjectGuid FixedUpId = ObjectId;
+	#endif
+		
+
+		FLazyObjectPtr LazyPtr;
+		LazyPtr = FixedUpId;
+
+		if (UObject* FoundObject = LazyPtr.Get())
+		{
+			return FoundObject;
+		}
 	}
-	else if (!ObjectPath.IsEmpty())
+
+	if (!ObjectPath.IsEmpty())
 	{
 		return FindObject<UObject>(InContext, *ObjectPath, false);
 	}
diff --git a/Engine/Source/Runtime/LevelSequence/Private/LevelSequencePlayer.cpp b/Engine/Source/Runtime/LevelSequence/Private/LevelSequencePlayer.cpp
--- a/Engine/Source/Runtime/LevelSequence/Private/LevelSequencePlayer.cpp
+++ b/Engine/Source/Runtime/LevelSequence/Private/LevelSequencePlayer.cpp
@@ -68,7 +68,6 @@ ULevelSequencePlayer::ULevelSequencePlayer(const FObjectInitializer& ObjectIniti
 	, bIsPlaying(false)
 	, TimeCursorPosition(0.0f)
 	, CurrentNumLoops(0)
-	, bAutoPlayNextFrame(false)
 { }
 
 
@@ -130,10 +129,7 @@ void ULevelSequencePlayer::Play()
 
 	bIsPlaying = true;
 
-	if (RootMovieSceneInstance.IsValid())
-	{
-		RootMovieSceneInstance->Update(TimeCursorPosition + SequenceStartOffset, TimeCursorPosition + SequenceStartOffset, *this);
-	}
+	UpdateMovieSceneInstance(TimeCursorPosition, TimeCursorPosition);
 }
 
 void ULevelSequencePlayer::PlayLooping(int32 NumLoops)
@@ -154,10 +150,7 @@ void ULevelSequencePlayer::SetPlaybackPosition(float NewPlaybackPosition)
 	TimeCursorPosition = NewPlaybackPosition;
 	OnCursorPositionChanged();
 
-	if (RootMovieSceneInstance.IsValid())
-	{
-		RootMovieSceneInstance->Update(TimeCursorPosition + SequenceStartOffset, LastTimePosition + SequenceStartOffset, *this);
-	}
+	UpdateMovieSceneInstance(TimeCursorPosition, LastTimePosition);
 }
 
 float ULevelSequencePlayer::GetLength() const
@@ -226,7 +219,6 @@ void ULevelSequencePlayer::Initialize(ULevelSequence* InLevelSequence, UWorld* I
 /* IMovieScenePlayer interface
  *****************************************************************************/
 
-
 void ULevelSequencePlayer::GetRuntimeObjects(TSharedRef<FMovieSceneSequenceInstance> MovieSceneInstance, const FGuid& ObjectId, TArray<UObject*>& OutObjects) const
 {
 	if (UObject* FoundObject = LevelSequence->FindObject(ObjectId))
@@ -236,8 +228,43 @@ void ULevelSequencePlayer::GetRuntimeObjects(TSharedRef<FMovieSceneSequenceInsta
 }
 
 
-void ULevelSequencePlayer::UpdateCameraCut(UObject* ObjectToViewThrough, bool bNewCameraCut) const
+void ULevelSequencePlayer::UpdateCameraCut(UObject* CameraObject, UObject* UnlockIfCameraObject) const
 {
+	// skip missing player controller
+	APlayerController* PC = World->GetGameInstance()->GetFirstLocalPlayerController();
+
+	if (PC == nullptr)
+	{
+		return;
+	}
+
+	// skip same view target
+	AActor* ViewTarget = PC->GetViewTarget();
+
+	if (CameraObject == ViewTarget)
+	{
+		return;
+	}
+
+	// skip unlocking if the current view target differs
+	ACameraActor* UnlockIfCameraActor = Cast<ACameraActor>(UnlockIfCameraObject);
+
+	if ((CameraObject == nullptr) && (UnlockIfCameraActor != ViewTarget))
+	{
+		return;
+	}
+
+	// override the player controller's view target
+	ACameraActor* CameraActor = Cast<ACameraActor>(CameraObject);
+
+	FViewTargetTransitionParams TransitionParams;
+	PC->SetViewTarget(CameraActor, TransitionParams);
+
+	if (PC->PlayerCameraManager)
+	{
+		PC->PlayerCameraManager->bClientSimulatingViewTarget = (CameraActor != nullptr);
+		PC->PlayerCameraManager->bGameCameraCutThisFrame = true;
+	}
 }
 
 void ULevelSequencePlayer::SetViewportSettings(const TMap<FViewportClient*, EMovieSceneViewportParams>& ViewportParamsMap)
@@ -268,13 +295,6 @@ TSharedRef<FMovieSceneSequenceInstance> ULevelSequencePlayer::GetRootMovieSceneS
 
 void ULevelSequencePlayer::Update(const float DeltaSeconds)
 {
-	if (bAutoPlayNextFrame)
-	{
-		Play();
-		bAutoPlayNextFrame = false;
-		return;
-	}
-
 	float LastTimePosition = TimeCursorPosition;
 
 	if (bIsPlaying)
@@ -283,13 +303,16 @@ void ULevelSequencePlayer::Update(const float DeltaSeconds)
 		OnCursorPositionChanged();
 	}
 
-	if(RootMovieSceneInstance.IsValid())
-	{
-		RootMovieSceneInstance->Update(TimeCursorPosition + SequenceStartOffset, LastTimePosition + SequenceStartOffset, *this);
-	}
+	UpdateMovieSceneInstance(TimeCursorPosition, LastTimePosition);
 }
 
-void ULevelSequencePlayer::AutoPlayNextFrame()
+void ULevelSequencePlayer::UpdateMovieSceneInstance(float CurrentPosition, float PreviousPosition)
 {
-	bAutoPlayNextFrame = true;
+	if(RootMovieSceneInstance.IsValid())
+	{
+		RootMovieSceneInstance->Update(CurrentPosition + SequenceStartOffset, PreviousPosition + SequenceStartOffset, *this);
+#if WITH_EDITOR
+		OnLevelSequencePlayerUpdate.Broadcast(*this, CurrentPosition, PreviousPosition);
+#endif
+	}
 }
\ No newline at end of file
diff --git a/Engine/Source/Runtime/LevelSequence/Public/LevelSequencePlayer.h b/Engine/Source/Runtime/LevelSequence/Public/LevelSequencePlayer.h
--- a/Engine/Source/Runtime/LevelSequence/Public/LevelSequencePlayer.h
+++ b/Engine/Source/Runtime/LevelSequence/Public/LevelSequencePlayer.h
@@ -34,7 +34,6 @@ struct FLevelSequencePlaybackSettings
 	float PlayRate;
 };
 
-
 /**
  * ULevelSequencePlayer is used to actually "play" an level sequence asset at runtime.
  *
@@ -126,7 +125,7 @@ protected:
 
 	// IMovieScenePlayer interface
 	virtual void GetRuntimeObjects(TSharedRef<FMovieSceneSequenceInstance> MovieSceneInstance, const FGuid& ObjectHandle, TArray<UObject*>& OutObjects) const override;
-	virtual void UpdateCameraCut(UObject* ObjectToViewThrough, bool bNewCameraCut) const override;
+	virtual void UpdateCameraCut(UObject* CameraObject, UObject* UnlockIfCameraObject) const override;
 	virtual void SetViewportSettings(const TMap<FViewportClient*, EMovieSceneViewportParams>& ViewportParamsMap) override;
 	virtual void GetViewportSettings(TMap<FViewportClient*, EMovieSceneViewportParams>& ViewportParamsMap) const override;
 	virtual EMovieScenePlayerStatus::Type GetPlaybackStatus() const override;
@@ -143,6 +142,9 @@ public:
 
 private:
 
+	/** Update the movie scene instance from the specified previous position, to the specified time position */
+	void UpdateMovieSceneInstance(float CurrentPosition, float PreviousPosition);
+
 	/** Called when the cursor position has changed to implement looping */
 	void OnCursorPositionChanged();
 
@@ -178,6 +180,17 @@ private:
 	/** The world this player will spawn actors in, if needed */
 	TWeakObjectPtr<UWorld> World;
 
-	/** If true, the sequence will begin playing when it is next updated */
-	bool bAutoPlayNextFrame;
+#if WITH_EDITOR
+public:
+
+	/** An event that is broadcast each time this level sequence player is updated */
+	DECLARE_EVENT_ThreeParams( ULevelSequencePlayer, FOnLevelSequencePlayerUpdated, const ULevelSequencePlayer&, float /*current time*/, float /*previous time*/ );
+	FOnLevelSequencePlayerUpdated& OnSequenceUpdated() const { return OnLevelSequencePlayerUpdate; }
+
+private:
+
+	/** The event that will be broadcast every time the sequence is updated */
+	mutable FOnLevelSequencePlayerUpdated OnLevelSequencePlayerUpdate;
+
+#endif
 };
diff --git a/Engine/Source/Runtime/MovieScene/Private/MovieScene.cpp b/Engine/Source/Runtime/MovieScene/Private/MovieScene.cpp
--- a/Engine/Source/Runtime/MovieScene/Private/MovieScene.cpp
+++ b/Engine/Source/Runtime/MovieScene/Private/MovieScene.cpp
@@ -122,6 +122,16 @@ bool UMovieScene::RemovePossessable( const FGuid& PossessableGuid )
 		{	
 			Modify();
 
+			// Remove the parent-child link for a parent spawnable/child possessable if necessary
+			if (CurPossesable.GetParentSpawnable().IsValid())
+			{
+				FMovieSceneSpawnable* ParentSpawnable = FindSpawnable(CurPossesable.GetParentSpawnable());
+				if (ParentSpawnable)
+				{
+					ParentSpawnable->RemoveChildPossessable(PossessableGuid);
+				}
+			}
+
 			// Found it!
 			Possessables.RemoveAt( PossesableIter.GetIndex() );
 
@@ -187,17 +197,48 @@ FMovieScenePossessable& UMovieScene::GetPossessable( const int32 Index )
 }
 
 
+FText UMovieScene::GetObjectDisplayName(const FGuid& ObjectId)
+{
+#if WITH_EDITORONLY_DATA
+	FText& Result = ObjectsToDisplayNames.FindOrAdd(ObjectId.ToString());
+
+	if (!Result.IsEmpty())
+	{
+		return Result;
+	}
+
+	FMovieSceneSpawnable* Spawnable = FindSpawnable(ObjectId);
+
+	if (Spawnable != nullptr)
+	{
+		Result = FText::FromString(Spawnable->GetName());
+		return Result;
+	}
+
+	FMovieScenePossessable* Possessable = FindPossessable(ObjectId);
+
+	if (Possessable != nullptr)
+	{
+		Result = FText::FromString(Possessable->GetName());
+		return Result;
+	}
+#endif
+	return FText::GetEmpty();
+}
+
+
 TRange<float> UMovieScene::GetPlaybackRange() const
 {
 	check(PlaybackRange.HasLowerBound() && PlaybackRange.HasUpperBound());
 	return PlaybackRange;
 }
 
+
 void UMovieScene::SetPlaybackRange(float Start, float End)
 {
 	if (ensure(End >= Start))
 	{
-		PlaybackRange = TRange<float>(Start, End);
+		PlaybackRange = TRange<float>(Start, TRangeBound<float>::Inclusive(End));
 
 #if WITH_EDITORONLY_DATA
 		if (EditorData.WorkingRange.IsEmpty())
@@ -235,29 +276,27 @@ TArray<UMovieSceneSection*> UMovieScene::GetAllSections() const
 }
 
 
-UMovieSceneTrack* UMovieScene::FindTrack( TSubclassOf<UMovieSceneTrack> TrackClass, const FGuid& ObjectGuid, FName UniqueTrackName ) const
+UMovieSceneTrack* UMovieScene::FindTrack(TSubclassOf<UMovieSceneTrack> TrackClass, const FGuid& ObjectGuid, const FName& TrackName) const
 {
-	UMovieSceneTrack* FoundTrack = nullptr;
-	
-	check(UniqueTrackName != NAME_None);
 	check(ObjectGuid.IsValid());
 	
 	for (const auto& Binding : ObjectBindings)
 	{
-		if (Binding.GetObjectGuid() == ObjectGuid) 
+		if (Binding.GetObjectGuid() != ObjectGuid) 
+		{
+			continue;
+		}
+
+		for (const auto& Track : Binding.GetTracks())
 		{
-			for (const auto& Track : Binding.GetTracks())
+			if ((Track->GetClass() == TrackClass) && (Track->GetTrackName() == TrackName))
 			{
-				if ((Track->GetClass() == TrackClass) && (Track->GetTrackName() == UniqueTrackName))
-				{
-					FoundTrack = Track;
-					break;
-				}
+				return Track;
 			}
 		}
 	}
 	
-	return FoundTrack;
+	return nullptr;
 }
 
 
@@ -383,6 +422,7 @@ bool UMovieScene::IsAMasterTrack(const UMovieSceneTrack& Track) const
 	return false;
 }
 
+
 void UMovieScene::UpgradeTimeRanges()
 {
 	// Legacy upgrade for playback ranges:
@@ -395,7 +435,7 @@ void UMovieScene::UpgradeTimeRanges()
 	if (InTime_DEPRECATED != FLT_MAX && OutTime_DEPRECATED != -FLT_MAX)
 	{
 		// Finite range already defined in old data
-		PlaybackRange = TRange<float>(InTime_DEPRECATED, OutTime_DEPRECATED);
+		PlaybackRange = TRange<float>(InTime_DEPRECATED, TRangeBound<float>::Inclusive(OutTime_DEPRECATED));
 	}
 	else if (PlaybackRange.IsEmpty())
 	{
@@ -421,7 +461,12 @@ void UMovieScene::UpgradeTimeRanges()
 			}
 		}
 
-		PlaybackRange = TRange<float>(0.f, MaxBound);
+		PlaybackRange = TRange<float>(0.f, TRangeBound<float>::Inclusive(MaxBound));
+	}
+	else if (PlaybackRange.GetUpperBound().IsExclusive())
+	{
+		// playback ranges are now always inclusive
+		PlaybackRange = TRange<float>(PlaybackRange.GetLowerBound(), TRangeBound<float>::Inclusive(PlaybackRange.GetUpperBoundValue()));
 	}
 
 	// PlaybackRange must always be defined to a finite range
diff --git a/Engine/Source/Runtime/MovieScene/Private/MovieSceneCommonHelpers.cpp b/Engine/Source/Runtime/MovieScene/Private/MovieSceneCommonHelpers.cpp
--- a/Engine/Source/Runtime/MovieScene/Private/MovieSceneCommonHelpers.cpp
+++ b/Engine/Source/Runtime/MovieScene/Private/MovieSceneCommonHelpers.cpp
@@ -10,18 +10,17 @@ TArray<UMovieSceneSection*> MovieSceneHelpers::GetTraversedSections( const TArra
 	TArray<UMovieSceneSection*> TraversedSections;
 
 	bool bPlayingBackwards = CurrentTime - PreviousTime < 0.0f;
-	
 	float MaxTime = bPlayingBackwards ? PreviousTime : CurrentTime;
 	float MinTime = bPlayingBackwards ? CurrentTime : PreviousTime;
 
-	TRange<float> SliderRange(MinTime, MaxTime);
+	TRange<float> TraversedRange(MinTime, TRangeBound<float>::Inclusive(MaxTime));
 
-	for( int32 SectionIndex = 0; SectionIndex < Sections.Num(); ++SectionIndex )
+	for (int32 SectionIndex = 0; SectionIndex < Sections.Num(); ++SectionIndex)
 	{
 		UMovieSceneSection* Section = Sections[SectionIndex];
-		if( Section->GetStartTime() == CurrentTime || SliderRange.Overlaps( TRange<float>( Section->GetRange() ) ) )
+		if ((Section->GetStartTime() == CurrentTime) || TraversedRange.Overlaps(TRange<float>(Section->GetRange())))
 		{
-			TraversedSections.Add( Section );
+			TraversedSections.Add(Section);
 		}
 	}
 
@@ -35,7 +34,7 @@ UMovieSceneSection* MovieSceneHelpers::FindSectionAtTime( const TArray<UMovieSce
 	{
 		UMovieSceneSection* Section = Sections[SectionIndex];
 
-		//@todo Sequencer - There can be multiple sections overlapping in time. Returning instantly does not account for that.
+		//@todo sequencer: There can be multiple sections overlapping in time. Returning instantly does not account for that.
 		if( Section->IsTimeWithinSection( Time ) && Section->IsActive() )
 		{
 			return Section;
@@ -61,7 +60,7 @@ UMovieSceneSection* MovieSceneHelpers::FindNearestSectionAtTime( const TArray<UM
 
 		if (Section->IsActive())
 		{
-			//@todo Sequencer - There can be multiple sections overlapping in time. Returning instantly does not account for that.
+			//@todo sequencer: There can be multiple sections overlapping in time. Returning instantly does not account for that.
 			if( Section->IsTimeWithinSection( Time ) )
 			{
 				return Section;
diff --git a/Engine/Source/Runtime/MovieScene/Private/MovieSceneSequenceInstance.cpp b/Engine/Source/Runtime/MovieScene/Private/MovieSceneSequenceInstance.cpp
--- a/Engine/Source/Runtime/MovieScene/Private/MovieSceneSequenceInstance.cpp
+++ b/Engine/Source/Runtime/MovieScene/Private/MovieSceneSequenceInstance.cpp
@@ -75,20 +75,29 @@ void FMovieSceneSequenceInstance::RestoreState(class IMovieScenePlayer& Player)
 
 void FMovieSceneSequenceInstance::Update( float Position, float LastPosition, class IMovieScenePlayer& Player )
 {
-	// Update  shot track
+	UpdateInternal(Position, LastPosition, Player, MSUP_PreUpdate);
+	UpdateInternal(Position, LastPosition, Player, MSUP_Update);
+	UpdateInternal(Position, LastPosition, Player, MSUP_PostUpdate);
+}
+
+void FMovieSceneSequenceInstance::UpdateInternal( float Position, float LastPosition, class IMovieScenePlayer& Player, EMovieSceneUpdatePass UpdatePass )
+{
+	// Refresh time range so that spawnables can be created if they fall within the playback range, or destroyed if not
+	UMovieScene* MovieScene = MovieSceneSequence->GetMovieScene();
+	TimeRange = MovieScene->GetPlaybackRange();
+
 	TArray<UObject*> NoObjects;
-	if( ShotTrackInstance.IsValid() )
-	{
-		ShotTrackInstance->Update( Position, LastPosition, NoObjects, Player, *this );
-	}
 
-	// Update each master track
+	// update each master track
 	for( FMovieSceneInstanceMap::TIterator It( MasterTrackInstances ); It; ++It )
 	{
-		It.Value()->Update( Position, LastPosition, NoObjects, Player, *this );
+		if (It.Value()->HasUpdatePasses() & UpdatePass)
+		{
+			It.Value()->Update( Position, LastPosition, NoObjects, Player, *this, UpdatePass);
+		}
 	}
 
-	// Update tracks bound to objects
+	// update tracks bound to objects
 	TMap<FGuid, FMovieSceneObjectBindingInstance>::TIterator ObjectIt = ObjectBindingInstances.CreateIterator();
 	for(; ObjectIt; ++ObjectIt )
 	{
@@ -96,7 +105,19 @@ void FMovieSceneSequenceInstance::Update( float Position, float LastPosition, cl
 		
 		for( FMovieSceneInstanceMap::TIterator It = ObjectBindingInstance.TrackInstances.CreateIterator(); It; ++It )
 		{
-			It.Value()->Update( Position, LastPosition, ObjectBindingInstance.RuntimeObjects, Player, *this );
+			if (It.Value()->HasUpdatePasses() & UpdatePass)
+			{
+				It.Value()->Update( Position, LastPosition, ObjectBindingInstance.RuntimeObjects, Player, *this, UpdatePass );
+			}
+		}
+	}
+
+	// update shot track last to make sure spawnable cameras are there, and to override sub-shots
+	if (ShotTrackInstance.IsValid())
+	{
+		if (ShotTrackInstance->HasUpdatePasses() & UpdatePass)
+		{
+			ShotTrackInstance->Update(Position, LastPosition, NoObjects, Player, *this, UpdatePass );
 		}
 	}
 }
@@ -109,15 +130,13 @@ void FMovieSceneSequenceInstance::RefreshInstance( IMovieScenePlayer& Player )
 
 	UMovieSceneTrack* ShotTrack = MovieScene->GetShotTrack();
 
-	TSharedRef<FMovieSceneSequenceInstance> ThisInstance = AsShared();
-
-	FMovieSceneInstanceMap ShotTrackInstanceMap;
-	// Only if root movie scene. Any sub-movie scene that has a shot track is ignored
-	if( ShotTrack && Player.GetRootMovieSceneSequenceInstance() == ThisInstance )
+	if (ShotTrack != nullptr)
 	{
-		if( ShotTrackInstance.IsValid() )
+		FMovieSceneInstanceMap ShotTrackInstanceMap;
+
+		if (ShotTrackInstance.IsValid())
 		{
-			ShotTrackInstanceMap.Add(  ShotTrack, ShotTrackInstance );
+			ShotTrackInstanceMap.Add(ShotTrack, ShotTrackInstance);
 		}
 
 		TArray<UObject*> Objects;
@@ -125,11 +144,11 @@ void FMovieSceneSequenceInstance::RefreshInstance( IMovieScenePlayer& Player )
 		Tracks.Add(ShotTrack);
 		RefreshInstanceMap(Tracks, Objects, ShotTrackInstanceMap, Player);
 
-		ShotTrackInstance = ShotTrackInstanceMap.FindRef( ShotTrack );
+		ShotTrackInstance = ShotTrackInstanceMap.FindRef(ShotTrack);
 	}
-	else if( !ShotTrack && ShotTrackInstance.IsValid() )
+	else if(ShotTrackInstance.IsValid())
 	{
-		ShotTrackInstance->ClearInstance( Player, *this );
+		ShotTrackInstance->ClearInstance(Player, *this);
 		ShotTrackInstance.Reset();
 	}
 
@@ -179,7 +198,7 @@ struct FTrackInstanceEvalSorter
 {
 	bool operator()( const TSharedPtr<IMovieSceneTrackInstance> A, const TSharedPtr<IMovieSceneTrackInstance> B ) const
 	{
-		return A->EvalOrder() > B->EvalOrder();
+		return A->EvalOrder() < B->EvalOrder();
 	}
 };
 
@@ -234,31 +253,85 @@ void FMovieSceneSequenceInstance::RefreshInstanceMap( const TArray<UMovieSceneTr
 	TrackInstances.ValueSort(FTrackInstanceEvalSorter());
 }
 
-void FMovieSceneSequenceInstance::SpawnObject(const FGuid& ObjectId)
+void FMovieSceneSequenceInstance::UpdateObjectBinding(const FGuid& ObjectId, IMovieScenePlayer& Player)
+{
+	UMovieSceneSequence* Sequence = MovieSceneSequence.Get();
+	auto* BindingInstance = ObjectBindingInstances.Find(ObjectId);
+
+	if (!BindingInstance || !Sequence)
+	{
+		return;
+	}
+
+	// Update the runtime objects
+	BindingInstance->RuntimeObjects.Reset();
+	Player.GetRuntimeObjects(SharedThis(this), BindingInstance->ObjectGuid, BindingInstance->RuntimeObjects);
+
+	const FMovieSceneBinding* ObjectBinding = Sequence->GetMovieScene()->GetBindings().FindByPredicate([&](const FMovieSceneBinding& In){
+		return In.GetObjectGuid() == ObjectId;
+	});
+
+	// Refresh the instance map, if we found the binding itself
+	if (ObjectBinding)
+	{
+		RefreshInstanceMap(ObjectBinding->GetTracks(), BindingInstance->RuntimeObjects, BindingInstance->TrackInstances, Player);
+	}
+}
+
+void FMovieSceneSequenceInstance::SpawnObject(const FGuid& ObjectId, IMovieScenePlayer& Player)
 {
 	UMovieSceneSequence* Sequence = GetSequence();
 	auto* BindingInstance = ObjectBindingInstances.Find(ObjectId);
 
-	if (Sequence && BindingInstance)
+	if (!Sequence || !BindingInstance)
 	{
-		BindingInstance->RuntimeObjects.Reset();
+		return;
+	}
 
-		// Spawn the object, and add it to the instance's runtime objects array
-		if (UObject* SpawnedObject = Sequence->SpawnObject(ObjectId))
+	// Reset the runtime objects for the binding instance
+	BindingInstance->RuntimeObjects.Reset();
+
+	// Attempt to spawn the object
+	UObject* SpawnedObject = Sequence->SpawnObject(ObjectId);
+	if (!SpawnedObject)
+	{
+		return;
+	}
+
+	// Add it to the instance's runtime objects array, and update any child possessable binding instances
+	BindingInstance->RuntimeObjects.Emplace(SpawnedObject);
+
+	FMovieSceneSpawnable* Spawnable = Sequence->GetMovieScene()->FindSpawnable(ObjectId);
+	if (Spawnable)
+	{
+		for (const FGuid& Child : Spawnable->GetChildPossessables())
 		{
-			BindingInstance->RuntimeObjects.Emplace(SpawnedObject);
+			UpdateObjectBinding(Child, Player);
 		}
 	}
 }
 
-void FMovieSceneSequenceInstance::DestroySpawnedObject(const FGuid& ObjectId)
+void FMovieSceneSequenceInstance::DestroySpawnedObject(const FGuid& ObjectId, IMovieScenePlayer& Player)
 {
 	UMovieSceneSequence* Sequence = GetSequence();
 	auto* BindingInstance = ObjectBindingInstances.Find(ObjectId);
 	
-	if (Sequence && BindingInstance)
+	if (!Sequence || !BindingInstance)
+	{
+		return;
+	}
+
+	// Destroy the object
+	Sequence->DestroySpawnedObject(ObjectId);
+	BindingInstance->RuntimeObjects.Reset();
+
+	// Update any child possessable object bindings
+	FMovieSceneSpawnable* Spawnable = Sequence->GetMovieScene()->FindSpawnable(ObjectId);
+	if (Spawnable)
 	{
-		Sequence->DestroySpawnedObject(ObjectId);
-		BindingInstance->RuntimeObjects.Reset();
+		for (const FGuid& Child : Spawnable->GetChildPossessables())
+		{
+			UpdateObjectBinding(Child, Player);
+		}
 	}
 }
\ No newline at end of file
diff --git a/Engine/Source/Runtime/MovieScene/Public/IMovieScenePlayer.h b/Engine/Source/Runtime/MovieScene/Public/IMovieScenePlayer.h
--- a/Engine/Source/Runtime/MovieScene/Public/IMovieScenePlayer.h
+++ b/Engine/Source/Runtime/MovieScene/Public/IMovieScenePlayer.h
@@ -61,10 +61,10 @@ public:
 	/**
 	 * Updates the perspective viewports with the actor to view through
 	 *
-	 * @param ObjectToViewThrough	The object, probably a camera, that the viewports should lock to
-	 * @param bNewCameraCut			If true this is a new camera cut
+	 * @param CameraObject The object, probably a camera, that the viewports should lock to
+	 * @param UnlockIfCameraObject If this is not nullptr, release actor lock only if currently locked to this object.
 	 */
-	virtual void UpdateCameraCut(UObject* ObjectToViewThrough, bool bNewCameraCut) const = 0;
+	virtual void UpdateCameraCut(UObject* CameraObject, UObject* UnlockIfCameraObject = nullptr) const = 0;
 
 	/*
 	 * Set the perspective viewport settings
diff --git a/Engine/Source/Runtime/MovieScene/Public/IMovieSceneTrackInstance.h b/Engine/Source/Runtime/MovieScene/Public/IMovieSceneTrackInstance.h
--- a/Engine/Source/Runtime/MovieScene/Public/IMovieSceneTrackInstance.h
+++ b/Engine/Source/Runtime/MovieScene/Public/IMovieSceneTrackInstance.h
@@ -6,6 +6,13 @@
 class IMovieScenePlayer;
 class FMovieSceneSequenceInstance;
 
+enum EMovieSceneUpdatePass
+{
+	MSUP_PreUpdate = 0x00000001,
+	MSUP_Update = 0x00000002,
+	MSUP_PostUpdate = 0x00000004
+};
+
 /**
  * A track instance holds the live objects for a track.  
  */
@@ -35,9 +42,15 @@ public:
 	 * @param LastPosition		The previous playback position
 	 * @param RuntimeObjects	Runtime objects bound to this instance (if any)
 	 * @param Player			The playback interface.  Contains state and some other functionality for runtime playback
+	 * @param UpdatePass        Which update pass
 	 */
-	virtual void Update(float Position, float LastPosition, const TArray<UObject*>& RuntimeObjects, IMovieScenePlayer& Player, FMovieSceneSequenceInstance& SequenceInstance) = 0;
+	virtual void Update(float Position, float LastPosition, const TArray<UObject*>& RuntimeObjects, IMovieScenePlayer& Player, FMovieSceneSequenceInstance& SequenceInstance, EMovieSceneUpdatePass UpdatePass) = 0;
 	
+	/*
+	 * Which update passes does this track instance evaluate in?
+	 */
+	virtual EMovieSceneUpdatePass HasUpdatePasses() { return MSUP_Update; }
+
 	/**
 	 * Refreshes the current instance
 	 */
diff --git a/Engine/Source/Runtime/MovieScene/Public/MovieScene.h b/Engine/Source/Runtime/MovieScene/Public/MovieScene.h
--- a/Engine/Source/Runtime/MovieScene/Public/MovieScene.h
+++ b/Engine/Source/Runtime/MovieScene/Public/MovieScene.h
@@ -179,25 +179,25 @@ public:
 	 *
 	 * @param TrackClass The class of the track to find.
 	 * @param ObjectGuid The runtime object guid that the track is bound to.
-	 * @param UniqueTypeName The unique name of the track to differentiate the one we are searching for from other tracks of the same class.
+	 * @param TrackName The name of the track to differentiate the one we are searching for from other tracks of the same class.
 	 * @return The found track or nullptr if one does not exist.
 	 * @see AddTrack, RemoveTrack
 	 */
-	UMovieSceneTrack* FindTrack(TSubclassOf<UMovieSceneTrack> TrackClass, const FGuid& ObjectGuid, FName UniqueTypeName) const;
+	UMovieSceneTrack* FindTrack(TSubclassOf<UMovieSceneTrack> TrackClass, const FGuid& ObjectGuid, const FName& TrackName) const;
 	
 	/**
 	 * Finds a track.
 	 *
 	 * @param TrackClass The class of the track to find.
 	 * @param ObjectGuid The runtime object guid that the track is bound to.
-	 * @param UniqueTypeName The unique name of the track to differentiate the one we are searching for from other tracks of the same class.
+	 * @param TrackName The name of the track to differentiate the one we are searching for from other tracks of the same class (optional).
 	 * @return The found track or nullptr if one does not exist.
 	 * @see AddTrack, RemoveTrack
 	 */
 	template<typename TrackClass>
-	TrackClass* FindTrack(const FGuid& ObjectGuid, FName UniqueTypeName) const
+	TrackClass* FindTrack(const FGuid& ObjectGuid, const FName& TrackName = NAME_None) const
 	{
-		return Cast<TrackClass>(FindTrack(TrackClass::StaticClass(), ObjectGuid, UniqueTypeName));
+		return Cast<TrackClass>(FindTrack(TrackClass::StaticClass(), ObjectGuid, TrackName));
 	}
 
 	/**
@@ -309,6 +309,13 @@ public:
 
 public:
 
+	/**
+	 * Returns all sections and their associated binding data.
+	 *
+	 * @return A list of sections with object bindings and names.
+	 */
+	TArray<UMovieSceneSection*> GetAllSections() const;
+
 	/**
 	 * @return All object bindings.
 	 */
@@ -317,15 +324,36 @@ public:
 		return ObjectBindings;
 	}
 
+	/**
+	 * Get the display name of the object with the specified identifier.
+	 *
+	 * @param ObjectId The object identifier.
+	 * @result The object's display name.
+	 */
+	FText GetObjectDisplayName(const FGuid& ObjectId);
+
+	/**
+	 * @return The playback time range of this movie scene, relative to its 0-time offset.
+	 */
+	TRange<float> GetPlaybackRange() const;
+
 	/*
 	* Replace an existing binding with another 
 	*/
 	void ReplaceBinding(const FGuid& OldGuid, const FGuid& NewGuid, const FString& Name);
 
 	/**
-	 * @return The playback time range of this movie scene, relative to its 0-time offset.
+	 * Get the display name of the object with the specified identifier.
+	 *
+	 * @param ObjectId The object identifier.
+	 * @result The object's display name.
 	 */
-	TRange<float> GetPlaybackRange() const;
+	void SetObjectDisplayName(const FGuid& ObjectId, const FText& DisplayName)
+	{
+#if WITH_EDITORONLY_DATA
+		ObjectsToDisplayNames.Add(ObjectId.ToString(), DisplayName);
+#endif
+	}
 
 	/**
 	 * Set the start and end playback positions (playback range) for this movie scene
@@ -335,13 +363,6 @@ public:
 	 */
 	void SetPlaybackRange(float Start, float End);
 
-	/**
-	 * Returns all sections and their associated binding data.
-	 *
-	 * @return A list of sections with object bindings and names.
-	 */
-	TArray<UMovieSceneSection*> GetAllSections() const;
-
 public:
 	
 #if WITH_EDITORONLY_DATA
@@ -402,6 +423,10 @@ private:
 	UMovieSceneTrack* ShotTrack;
 
 #if WITH_EDITORONLY_DATA
+	/** Maps object GUIDs to user defined display names. */
+	UPROPERTY()
+	TMap<FString, FText> ObjectsToDisplayNames;
+
 	/** Editor only data that needs to be saved between sessions for editing but has no runtime purpose */
 	UPROPERTY()
 	FMovieSceneEditorData EditorData;
@@ -415,10 +440,13 @@ private:
 	
 	UPROPERTY()
 	float InTime_DEPRECATED;
+
 	UPROPERTY()
 	float OutTime_DEPRECATED;
+
 	UPROPERTY()
 	float StartTime_DEPRECATED;
+
 	UPROPERTY()
 	float EndTime_DEPRECATED;
 };
diff --git a/Engine/Source/Runtime/MovieScene/Public/MovieScenePossessable.h b/Engine/Source/Runtime/MovieScene/Public/MovieScenePossessable.h
--- a/Engine/Source/Runtime/MovieScene/Public/MovieScenePossessable.h
+++ b/Engine/Source/Runtime/MovieScene/Public/MovieScenePossessable.h
@@ -86,6 +86,26 @@ public:
 		return PossessedObjectClass;
 	}
 
+	/**
+	 * Get the guid of this possessable's parent spawnable
+	 *
+	 * @return The guid.
+	 */
+	const FGuid& GetParentSpawnable() const
+	{
+		return ParentSpawnableGuid;
+	}
+
+	/**
+	 * Set the guid of this possessable's parent spawnable
+	 *
+	 * @param InParentSpawnableGuid The guid of this possessable's parent spawnable.
+	 */
+	void SetParentSpawnable(const FGuid& InParentSpawnableGuid)
+	{
+		ParentSpawnableGuid = InParentSpawnableGuid;
+	}
+
 private:
 
 	/** Unique identifier of the possessable object. */
@@ -106,4 +126,8 @@ private:
 	// possessable object available.
 	UPROPERTY()
 	UClass* PossessedObjectClass;
+
+	/** GUID relating to this possessable's parent spawnable, if applicable. */
+	UPROPERTY()
+	FGuid ParentSpawnableGuid;
 };
diff --git a/Engine/Source/Runtime/MovieScene/Public/MovieSceneSection.h b/Engine/Source/Runtime/MovieScene/Public/MovieSceneSection.h
--- a/Engine/Source/Runtime/MovieScene/Public/MovieSceneSection.h
+++ b/Engine/Source/Runtime/MovieScene/Public/MovieSceneSection.h
@@ -72,9 +72,9 @@ public:
 		// Use the single value constructor for zero sized ranges because it creates a range that is inclusive on both upper and lower
 		// bounds which isn't considered "empty".  Use the standard constructor for non-zero sized ranges so that they work well when
 		// calculating overlap with other non-zero sized ranges.
-		return StartTime == EndTime 
-			? TRange<float>(StartTime) 
-			: TRange<float>(StartTime, EndTime);
+		return (StartTime == EndTime)
+			? TRange<float>(StartTime)
+			: TRange<float>(StartTime, TRangeBound<float>::Inclusive(EndTime));
 	}
 	
 	/**
