Fix shadow variable warning about bParallelCloth.

diff --git a/Engine/Source/Runtime/Engine/Private/PhysicsEngine/ClothManager.cpp b/Engine/Source/Runtime/Engine/Private/PhysicsEngine/ClothManager.cpp
--- a/Engine/Source/Runtime/Engine/Private/PhysicsEngine/ClothManager.cpp
+++ b/Engine/Source/Runtime/Engine/Private/PhysicsEngine/ClothManager.cpp
@@ -19,7 +19,8 @@ void FClothManagerData::PrepareCloth(float DeltaTime, FTickFunction& TickFunctio
 	QUICK_SCOPE_CYCLE_COUNTER(STAT_FClothManager_PrepareCloth);
 	if (SkeletalMeshComponents.Num())
 	{
-		const bool bParallelCloth = !!CVarParallelCloth.GetValueOnAnyThread();
+		// TODO: OriC, I took the earlier version but you should verify which is correct. You probably want the above check at this location :)
+		//const bool bParallelCloth = !!CVarParallelCloth.GetValueOnAnyThread();
 		IsPreparingClothInParallel.AtomicSet(bParallelCloth);
 		for (USkeletalMeshComponent* SkeletalMeshComponent : SkeletalMeshComponents)
 		{
