#ANIM: had this issue where RawAnimationData is empty, but other data exists - i.e. NumFrames
 - it turns out morphtarget animations or any animation without real key can cause this, so I'm reverting this change, but adding new check to prevent previous crash when trying to attempt compression with empty RawAnimationData now of [number of tracks] instead.

- Back out revision 6 from //UE4/Main/Engine/Source/Runtime/Engine/Private/Animation/AnimSequence.cpp
- Instead of RawAnimation check, check Track numbers. It should have Track Names if it had any tracks. And it is available only with editor, which is this case.

diff --git a/Engine/Source/Runtime/Engine/Private/Animation/AnimSequence.cpp b/Engine/Source/Runtime/Engine/Private/Animation/AnimSequence.cpp
--- a/Engine/Source/Runtime/Engine/Private/Animation/AnimSequence.cpp
+++ b/Engine/Source/Runtime/Engine/Private/Animation/AnimSequence.cpp
@@ -373,29 +373,6 @@ void UAnimSequence::PostLoad()
 	{
 		RemoveNaNTracks();
 	}
-
-	// fix up broken data
-	// recently we saw RawAnimationData is empty but has NumFrames and SequenceLength
-	// We're cleaning it up if no Raw data exists, and NumFrames > 0, then it is inconsistent
-	if (RawAnimationData.Num() == 0 && NumFrames != 0)
-	{
-		NumFrames = 0; 
-		SequenceLength = 0.f;
-		
-		UE_LOG(LogAnimation, Warning, TEXT("[%s] has no animation data but has frame number - removing frame number."), *GetName());
-		
-		if (!IsRunningGame())
-		{
-			static FName NAME_LoadErrors("LoadErrors");
-			FMessageLog LoadErrors(NAME_LoadErrors);
-
-			TSharedRef<FTokenizedMessage> Message = LoadErrors.Warning();
-			Message->AddToken(FTextToken::Create(LOCTEXT("InvalidAnimationData1", "The Animation ")));
-			Message->AddToken(FAssetNameToken::Create(GetPathName(), FText::FromString(GetName())));
-			Message->AddToken(FTextToken::Create(LOCTEXT("InvalidAnimationData2", " has no animation data but contains invalid number of frames and sequence length. Resetting...")));
-			LoadErrors.Notify();
-		}
-	}
 #endif // WITH_EDITOR
 
 	Super::PostLoad();
@@ -409,7 +386,6 @@ void UAnimSequence::PostLoad()
 	if ( NumFrames == 0 )
 	{
 		UE_LOG(LogAnimation, Warning, TEXT("No animation data exists for sequence %s (%s)"), *GetName(), (GetOuter() ? *GetOuter()->GetFullName() : *GetFullName()) );
-		
 #if WITH_EDITOR
 		if (!IsRunningGame())
 		{
@@ -1735,8 +1711,7 @@ bool UAnimSequence::CompressRawAnimData(float MaxPosDiff, float MaxAngleDiff)
 	bool bRemovedKeys = false;
 #if WITH_EDITORONLY_DATA
 
-	// check if track name exists first, if no data to start with that should be been empty
-	if (NumFrames > 0 && ensureMsgf(RawAnimationData.Num() > 0, TEXT("%s is trying to compress while raw animation is missing"), *GetName()))
+	if (AnimationTrackNames.Num() > 0 && ensureMsgf(RawAnimationData.Num() > 0, TEXT("%s is trying to compress while raw animation is missing"), *GetName()))
 	{
 		// This removes trivial keys, and this has to happen before the removing tracks
 		for(int32 TrackIndex=0; TrackIndex<RawAnimationData.Num(); TrackIndex++)
