Copying //UE4/Dev-Framework to Dev-Main (//UE4/Dev-Main)

diff --git a/Engine/Source/Runtime/Core/Public/UObject/ObjectVersion.h b/Engine/Source/Runtime/Core/Public/UObject/ObjectVersion.h
--- a/Engine/Source/Runtime/Core/Public/UObject/ObjectVersion.h
+++ b/Engine/Source/Runtime/Core/Public/UObject/ObjectVersion.h
@@ -578,6 +578,8 @@ enum EUnrealEngineObjectUE4Version
 	VER_UE4_USERWIDGET_DEFAULT_FOCUSABLE_FALSE,
 	// Custom event nodes implicitly set 'const' on array and non-array pass-by-reference input params
 	VER_UE4_BLUEPRINT_CUSTOM_EVENT_CONST_INPUT,
+	// Renamed HighFrequencyGain to LowPassFilterFrequency
+	VER_UE4_USE_LOW_PASS_FILTER_FREQ,
 	// -----<new versions can be added before this line>-------------------------------------------------
 	// - this needs to be the last line (see note below)
 	VER_UE4_AUTOMATIC_VERSION_PLUS_ONE,
diff --git a/Engine/Source/Runtime/CoreUObject/Private/Misc/UObjectToken.cpp b/Engine/Source/Runtime/CoreUObject/Private/Misc/UObjectToken.cpp
--- a/Engine/Source/Runtime/CoreUObject/Private/Misc/UObjectToken.cpp
+++ b/Engine/Source/Runtime/CoreUObject/Private/Misc/UObjectToken.cpp
@@ -7,7 +7,7 @@ FOnMessageTokenActivated FUObjectToken::DefaultMessageTokenActivated;
 
 FUObjectToken::FOnGetDisplayName FUObjectToken::DefaultGetObjectDisplayName;
 
-FUObjectToken::FUObjectToken( UObject* InObject,  const FText& InLabelOverride )
+FUObjectToken::FUObjectToken( const UObject* InObject,  const FText& InLabelOverride )
 	: ObjectBeingReferenced(InObject)
 {
 	if ( !InLabelOverride.IsEmpty() )
diff --git a/Engine/Source/Runtime/CoreUObject/Private/Serialization/BulkData.cpp b/Engine/Source/Runtime/CoreUObject/Private/Serialization/BulkData.cpp
--- a/Engine/Source/Runtime/CoreUObject/Private/Serialization/BulkData.cpp
+++ b/Engine/Source/Runtime/CoreUObject/Private/Serialization/BulkData.cpp
@@ -900,7 +900,7 @@ void FUntypedBulkData::Serialize( FArchive& Ar, UObject* Owner, int32 Idx )
 		{
 			// check if we save the package compressed
 			UPackage* Pkg = Owner ? dynamic_cast<UPackage*>(Owner->GetOutermost()) : nullptr;
-			if (Pkg && !!(Pkg->PackageFlags & PKG_StoreCompressed) )
+			if (Pkg && Pkg->HasAnyPackageFlags(PKG_StoreCompressed) )
 			{
 				ECompressionFlags BaseCompressionMethod = COMPRESS_Default;
 				if (Ar.IsCooking())
diff --git a/Engine/Source/Runtime/CoreUObject/Private/UObject/LazyObjectPtr.cpp b/Engine/Source/Runtime/CoreUObject/Private/UObject/LazyObjectPtr.cpp
--- a/Engine/Source/Runtime/CoreUObject/Private/UObject/LazyObjectPtr.cpp
+++ b/Engine/Source/Runtime/CoreUObject/Private/UObject/LazyObjectPtr.cpp
@@ -122,7 +122,7 @@ void FLazyObjectPtr::PossiblySerializeObjectGuid(UObject *Object, FArchive& Ar)
 
 			// Don't try and resolve GUIDs when loading a package for diff'ing
 			const UPackage* Package = Object->GetOutermost();
-			bool bLoadedForDiff = (Package && (Package->PackageFlags & PKG_ForDiffing));
+			const bool bLoadedForDiff = (Package && Package->HasAnyPackageFlags(PKG_ForDiffing));
 			if (!bLoadedForDiff && (!(Ar.GetPortFlags() & PPF_Duplicate) || (Ar.GetPortFlags() & PPF_DuplicateForPIE)))
 			{
 				check(!Guid.IsDefault());
diff --git a/Engine/Source/Runtime/CoreUObject/Private/UObject/Linker.cpp b/Engine/Source/Runtime/CoreUObject/Private/UObject/Linker.cpp
--- a/Engine/Source/Runtime/CoreUObject/Private/UObject/Linker.cpp
+++ b/Engine/Source/Runtime/CoreUObject/Private/UObject/Linker.cpp
@@ -477,7 +477,7 @@ FLinkerLoad* GetPackageLinker
 		if ( (GIsEditor && !DoesNativePackageExist) || (!GIsEditor && !DoesLocalizedPackageExist && !DoesNativePackageExist) )
 		{
 			// In memory-only packages have no linker and this is ok.
-			if (!(LoadFlags & LOAD_AllowDll) && !(InOuter->PackageFlags & PKG_InMemoryOnly) && !FLinkerLoad::KnownMissingPackages.Contains(InOuter->GetFName()))
+			if (!(LoadFlags & LOAD_AllowDll) && !InOuter->HasAnyPackageFlags(PKG_InMemoryOnly) && !FLinkerLoad::KnownMissingPackages.Contains(InOuter->GetFName()))
 			{
 				FUObjectThreadContext& ThreadContext = FUObjectThreadContext::Get();
 				FFormatNamedArguments Arguments;
@@ -508,7 +508,7 @@ FLinkerLoad* GetPackageLinker
 			else
 			{
 				// In memory-only packages have no linker and this is ok.
-				if (!(LoadFlags & LOAD_AllowDll) && !(InOuter->PackageFlags & PKG_InMemoryOnly))
+				if (!(LoadFlags & LOAD_AllowDll) && !InOuter->HasAnyPackageFlags(PKG_InMemoryOnly))
 				{
 					FUObjectThreadContext& ThreadContext = FUObjectThreadContext::Get();
 					FFormatNamedArguments Arguments;
@@ -539,7 +539,7 @@ FLinkerLoad* GetPackageLinker
 		UPackage* ExistingPackage = FindObject<UPackage>(nullptr, *PackageName);
 		if (ExistingPackage)
 		{
-			if (!ExistingPackage->GetOuter() && (ExistingPackage->PackageFlags & PKG_InMemoryOnly))
+			if (!ExistingPackage->GetOuter() && ExistingPackage->HasAnyPackageFlags(PKG_InMemoryOnly))
 			{
 				// This is a memory-only in package and so it has no linker and this is ok.
 				return nullptr;
@@ -556,11 +556,11 @@ FLinkerLoad* GetPackageLinker
 		{
 			if (!FLinkerLoad::KnownMissingPackages.Contains(InLongPackageName))
 			{
-				FFormatNamedArguments Arguments;
-				Arguments.Add(TEXT("Filename"), FText::FromString(InLongPackageName));
+			FFormatNamedArguments Arguments;
+			Arguments.Add(TEXT("Filename"), FText::FromString(InLongPackageName));
 
-				// try to recover from this instead of throwing, it seems recoverable just by doing this
-				LogGetPackageLinkerError(Result, InLongPackageName, FText::Format(LOCTEXT("FileNotFound", "Can't find file '{Filename}'"), Arguments), LOCTEXT("FileNotFoundShort", "Can't find file"), InOuter, LoadFlags);
+			// try to recover from this instead of throwing, it seems recoverable just by doing this
+			LogGetPackageLinkerError(Result, InLongPackageName, FText::Format(LOCTEXT("FileNotFound", "Can't find file '{Filename}'"), Arguments), LOCTEXT("FileNotFoundShort", "Can't find file"), InOuter, LoadFlags);
 			}
 			return nullptr;
 		}
@@ -588,7 +588,7 @@ FLinkerLoad* GetPackageLinker
 		UPackage* FilenamePkg = (ExistingPackage ? ExistingPackage : CreatePackage(nullptr, *PackageName));
 		if (FilenamePkg != ExistingPackage && (LoadFlags & LOAD_PackageForPIE))
 		{
-			FilenamePkg->PackageFlags |= PKG_PlayInEditor;
+			FilenamePkg->SetPackageFlags(PKG_PlayInEditor);
 		}
 
 		// If no package specified, use package from file.
diff --git a/Engine/Source/Runtime/CoreUObject/Private/UObject/LinkerLoad.cpp b/Engine/Source/Runtime/CoreUObject/Private/UObject/LinkerLoad.cpp
--- a/Engine/Source/Runtime/CoreUObject/Private/UObject/LinkerLoad.cpp
+++ b/Engine/Source/Runtime/CoreUObject/Private/UObject/LinkerLoad.cpp
@@ -1138,10 +1138,14 @@ FLinkerLoad::ELinkerStatus FLinkerLoad::SerializePackageFileSummary()
 		if( LinkerRootPackage )
 		{
 			// Preserve PIE package flag
-			uint32 PIEFlag = (LinkerRootPackage->PackageFlags & PKG_PlayInEditor);
+			uint32 NewPackageFlags = Summary.PackageFlags;
+			if (LinkerRootPackage->HasAnyPackageFlags(PKG_PlayInEditor))
+			{
+				NewPackageFlags |= PKG_PlayInEditor;
+			}
 			
 			// Propagate package flags
-			LinkerRootPackage->PackageFlags = (Summary.PackageFlags | PIEFlag);
+			LinkerRootPackage->SetPackageFlagsTo(NewPackageFlags);
 
 			// Propagate package folder name
 			LinkerRootPackage->SetFolderName(*Summary.FolderName);
@@ -2450,7 +2454,7 @@ bool FLinkerLoad::VerifyImportInner(const int32 ImportIndex, FString& WarningSuf
 
 		// if we couldn't create the package or it is 
 		// to be linked to any other package's ImportMaps
-		if ( !TmpPkg || (TmpPkg->PackageFlags&PKG_Compiling) != 0 )
+		if ( !TmpPkg || TmpPkg->HasAnyPackageFlags(PKG_Compiling) )
 		{
 			return false;
 		}
@@ -2493,7 +2497,7 @@ bool FLinkerLoad::VerifyImportInner(const int32 ImportIndex, FString& WarningSuf
 			}
 
 			auto* Package = dynamic_cast<UPackage*>(Top->XObject);
-			if (Package && (Package->PackageFlags & PKG_InMemoryOnly))
+			if (Package && Package->HasAnyPackageFlags(PKG_InMemoryOnly))
 			{
 				// This is an import to a memory-only package, just search for it in the package.
 				TmpPkg = Package;
@@ -2669,7 +2673,7 @@ bool FLinkerLoad::VerifyImportInner(const int32 ImportIndex, FString& WarningSuf
 	}
 
 	bool bCameFromMemoryOnlyPackage = false;
-	if (!Pkg && TmpPkg && (TmpPkg->PackageFlags & PKG_InMemoryOnly))
+	if (!Pkg && TmpPkg && TmpPkg->HasAnyPackageFlags(PKG_InMemoryOnly))
 	{
 		Pkg = TmpPkg; // this is a package that exists in memory only, so that is the package to search regardless of FindIfFail
 		bCameFromMemoryOnlyPackage = true;
@@ -3095,7 +3099,7 @@ void FLinkerLoad::Preload( UObject* Object )
 	check(Object);
 
 #if USE_CIRCULAR_DEPENDENCY_LOAD_DEFERRING
-	bool const bIsNonNativeObject = !(Object->GetOutermost()->PackageFlags & PKG_CompiledIn);
+	bool const bIsNonNativeObject = !Object->GetOutermost()->HasAnyPackageFlags(PKG_CompiledIn);
 	// we can determine that this is a blueprint class/struct by checking if it 
 	// is a class/struct object AND if it is not native (blueprint 
 	// structs/classes are the only asset package structs/classes we have)
@@ -3757,7 +3761,7 @@ UObject* FLinkerLoad::CreateExport( int32 Index )
 		// if we are loading objects just to verify an object reference during script compilation,
 		if (!GVerifyObjectReferencesOnly
 		||	(ObjectLoadFlags&RF_ClassDefaultObject) != 0					// only load this object if it's a class default object
-		||	(LinkerRoot->PackageFlags&PKG_ContainsScript) != 0		// or we're loading an existing package and it's a script package
+		||	LinkerRoot->HasAnyPackageFlags(PKG_ContainsScript)		// or we're loading an existing package and it's a script package
 		||	ThisParent->IsTemplate(RF_ClassDefaultObject)			// or if its a subobject template in a CDO
 		||	LoadClass->IsChildOf(UField::StaticClass())				// or if it is a UField
 		||	LoadClass->IsChildOf(UObjectRedirector::StaticClass()))	// or if its a redirector to another object
@@ -3937,7 +3941,7 @@ bool FLinkerLoad::IsImportNative(const int32 Index) const
 		else if (UPackage* ExistingPackage = FindObject<UPackage>(/*Outer =*/nullptr, *Import.ObjectName.ToString()))
 		{
 			// @TODO: what if the package's outer isn't null... what does that mean?
-			bIsImportNative = !ExistingPackage->GetOuter() && (ExistingPackage->PackageFlags & PKG_CompiledIn);
+			bIsImportNative = !ExistingPackage->GetOuter() && ExistingPackage->HasAnyPackageFlags(PKG_CompiledIn);
 		}
 	}
 
diff --git a/Engine/Source/Runtime/CoreUObject/Private/UObject/LinkerSave.cpp b/Engine/Source/Runtime/CoreUObject/Private/UObject/LinkerSave.cpp
--- a/Engine/Source/Runtime/CoreUObject/Private/UObject/LinkerSave.cpp
+++ b/Engine/Source/Runtime/CoreUObject/Private/UObject/LinkerSave.cpp
@@ -29,7 +29,7 @@ FLinkerSave::FLinkerSave(UPackage* InParent, const TCHAR* InFilename, bool bForc
 		Summary.SetFileVersions( GPackageFileUE4Version, GPackageFileLicenseeUE4Version, bInSaveUnversioned );
 		Summary.SavedByEngineVersion = FEngineVersion::Current();
 		Summary.CompatibleWithEngineVersion = FEngineVersion::CompatibleWith();
-		Summary.PackageFlags = Package ? (Package->PackageFlags & ~PKG_NewlyCreated) : 0;
+		Summary.PackageFlags = Package ? (Package->GetPackageFlags() & ~PKG_NewlyCreated) : 0;
 
 		if (Package)
 		{
@@ -66,7 +66,7 @@ FLinkerSave::FLinkerSave(UPackage* InParent, FArchive *InSaver, bool bForceByteS
 		Summary.SetFileVersions(GPackageFileUE4Version, GPackageFileLicenseeUE4Version, bInSaveUnversioned);
 		Summary.SavedByEngineVersion = FEngineVersion::Current();
 		Summary.CompatibleWithEngineVersion = FEngineVersion::CompatibleWith();
-		Summary.PackageFlags = Package ? (Package->PackageFlags & ~PKG_NewlyCreated) : 0;
+		Summary.PackageFlags = Package ? (Package->GetPackageFlags() & ~PKG_NewlyCreated) : 0;
 
 		if (Package)
 		{
@@ -98,7 +98,7 @@ FLinkerSave::FLinkerSave(UPackage* InParent, bool bForceByteSwapping, bool bInSa
 		Summary.SetFileVersions( GPackageFileUE4Version, GPackageFileLicenseeUE4Version, bInSaveUnversioned );
 		Summary.SavedByEngineVersion = FEngineVersion::Current();
 		Summary.CompatibleWithEngineVersion = FEngineVersion::CompatibleWith();
-		Summary.PackageFlags = Package ? (Package->PackageFlags & ~PKG_NewlyCreated) : 0;
+		Summary.PackageFlags = Package ? (Package->GetPackageFlags() & ~PKG_NewlyCreated) : 0;
 
 		if (Package)
 		{
diff --git a/Engine/Source/Runtime/CoreUObject/Private/UObject/Obj.cpp b/Engine/Source/Runtime/CoreUObject/Private/UObject/Obj.cpp
--- a/Engine/Source/Runtime/CoreUObject/Private/UObject/Obj.cpp
+++ b/Engine/Source/Runtime/CoreUObject/Private/UObject/Obj.cpp
@@ -832,7 +832,7 @@ bool UObject::Modify( bool bAlwaysMarkDirty/*=true*/ )
 	{
 		// Do not consider PIE world objects or script packages, as they should never end up in the
 		// transaction buffer and we don't want to mark them dirty here either.
-		if ((GetOutermost()->PackageFlags & (PKG_PlayInEditor | PKG_ContainsScript | PKG_CompiledIn)) == 0)
+		if (GetOutermost()->HasAnyPackageFlags(PKG_PlayInEditor | PKG_ContainsScript | PKG_CompiledIn) == false)
 		{
 			// Attempt to mark the package dirty and save a copy of the object to the transaction
 			// buffer. The save will fail if there isn't a valid transactor, the object isn't
@@ -3431,7 +3431,7 @@ bool StaticExec( UWorld* InWorld, const TCHAR* Cmd, FOutputDevice& Ar )
 								// Recurse if we're in the same package.
 								if( RootObject->GetOutermost() == Object->GetOutermost() 
 								// Or if package doesn't contain script.
-								||	!(Object->GetOutermost()->PackageFlags & PKG_ContainsScript) )
+								||	!Object->GetOutermost()->HasAnyPackageFlags(PKG_ContainsScript) )
 								{
 									// Serialize object. We don't want to use the << operator here as it would call 
 									// this function again instead of serializing members.
@@ -3478,7 +3478,7 @@ bool StaticExec( UWorld* InWorld, const TCHAR* Cmd, FOutputDevice& Ar )
 					// Only list certain class if specified.
 					if( (!ListClass || ObjectReference->GetClass() == ListClass)
 					// Only list non-script objects if specified.
-					&&	(!bShouldOnlyListContent || !(ObjectReference->GetOutermost()->PackageFlags & PKG_ContainsScript))
+					&&	(!bShouldOnlyListContent || !ObjectReference->GetOutermost()->HasAnyPackageFlags(PKG_ContainsScript))
 					// Exclude the transient package.
 					&&	ObjectReference->GetOutermost() != GetTransientPackage() )
 					{
diff --git a/Engine/Source/Runtime/CoreUObject/Private/UObject/Package.cpp b/Engine/Source/Runtime/CoreUObject/Private/UObject/Package.cpp
--- a/Engine/Source/Runtime/CoreUObject/Private/UObject/Package.cpp
+++ b/Engine/Source/Runtime/CoreUObject/Private/UObject/Package.cpp
@@ -41,7 +41,7 @@ void UPackage::PostInitProperties()
 	bIsCookedForEditor = false;
 	// Mark this package as editor-only by default. As soon as something in it is accessed through a non editor-only
 	// property the flag will be removed.
-	bLoadedByEditorPropertiesOnly = !HasAnyFlags(RF_ClassDefaultObject) && !(PackageFlags & PKG_CompiledIn);
+	bLoadedByEditorPropertiesOnly = !HasAnyFlags(RF_ClassDefaultObject) && !HasAnyPackageFlags(PKG_CompiledIn);
 #endif
 }
 
@@ -55,7 +55,7 @@ void UPackage::SetDirtyFlag( bool bIsDirty )
 	{
 		if ( GUndo != NULL
 		// PIE world objects should never end up in the transaction buffer as we cannot undo during gameplay.
-		&& !(GetOutermost()->PackageFlags & (PKG_PlayInEditor|PKG_ContainsScript)) )
+		&& !GetOutermost()->HasAnyPackageFlags(PKG_PlayInEditor|PKG_ContainsScript) )
 		{
 			// make sure we're marked as transactional
 			SetFlags(RF_Transactional);
@@ -68,8 +68,8 @@ void UPackage::SetDirtyFlag( bool bIsDirty )
 		bDirty = bIsDirty;
 
 		if( GIsEditor									// Only fire the callback in editor mode
-			&& !(PackageFlags & PKG_ContainsScript)		// Skip script packages
-			&& !(PackageFlags & PKG_PlayInEditor)		// Skip packages for PIE
+			&& !HasAnyPackageFlags(PKG_ContainsScript)	// Skip script packages
+			&& !HasAnyPackageFlags(PKG_PlayInEditor)	// Skip packages for PIE
 			&& GetTransientPackage() != this )			// Skip the transient package
 		{
 			// Package is changing dirty state, let the editor know so we may prompt for source control checkout
@@ -216,6 +216,16 @@ void UPackage::BeginDestroy()
 	Super::BeginDestroy();
 }
 
+// UE-21181 - Tracking where the loaded editor level's package gets flagged as a PIE object
+#if WITH_EDITOR
+UPackage* UPackage::EditorPackage = nullptr;
+void UPackage::SetPackageFlagsTo( uint32 NewFlags )
+{
+	PackageFlagsPrivate = NewFlags;
+	ensure(((NewFlags & PKG_PlayInEditor) == 0) || (this != EditorPackage));
+}
+#endif
+
 #if WITH_EDITORONLY_DATA
 void FixupPackageEditorOnlyFlag(FName PackageThatGotEditorOnlyFlagCleared, bool bRecursive);
 
diff --git a/Engine/Source/Runtime/CoreUObject/Private/UObject/SavePackage.cpp b/Engine/Source/Runtime/CoreUObject/Private/UObject/SavePackage.cpp
--- a/Engine/Source/Runtime/CoreUObject/Private/UObject/SavePackage.cpp
+++ b/Engine/Source/Runtime/CoreUObject/Private/UObject/SavePackage.cpp
@@ -372,7 +372,7 @@ static bool IsEditorOnlyStruct(UStruct* InStruct)
 	{
 		UPackage* StructPackage = Struct->GetOutermost();
 		check(StructPackage);
-		if (StructPackage->PackageFlags & PKG_EditorOnly)
+		if (StructPackage->HasAnyPackageFlags(PKG_EditorOnly))
 		{
 			return true;
 		}
@@ -424,7 +424,7 @@ static bool IsEditorOnlyObject(UObject* InObject)
 	{
 		Package = InObject->GetOutermost();
 	}
-	if (Package && (Package->PackageFlags & PKG_EditorOnly))
+	if (Package && Package->HasAnyPackageFlags(PKG_EditorOnly))
 	{
 		bResult = true;
 	}
@@ -3260,7 +3260,7 @@ ESavePackageResult UPackage::Save(UPackage* InOuter, UObject* Base, EObjectFlags
 				UE_CLOG(!(SaveFlags & SAVE_NoError), LogSavePackage, Display, TEXT("Package loaded by editor-only properties: %s. Package will not be saved."), *InOuter->GetName());
 				return ESavePackageResult::ReferencedOnlyByEditorOnlyData;
 			}
-			else if (InOuter->PackageFlags & PKG_EditorOnly)
+			else if (InOuter->HasAnyPackageFlags(PKG_EditorOnly))
 			{
 				UE_CLOG(!(SaveFlags & SAVE_NoError), LogSavePackage, Display, TEXT("Package marked as editor-only: %s. Package will not be saved."), *InOuter->GetName());
 				return ESavePackageResult::ReferencedOnlyByEditorOnlyData;
@@ -3388,7 +3388,7 @@ ESavePackageResult UPackage::Save(UPackage* InOuter, UObject* Base, EObjectFlags
 			return ESavePackageResult::Error;
 		}
 
-		bool FilterEditorOnly = (InOuter->PackageFlags & PKG_FilterEditorOnly) != 0;
+		const bool FilterEditorOnly = InOuter->HasAnyPackageFlags(PKG_FilterEditorOnly);
 		// store a list of additional packages to cook when this is cooked (ie streaming levels)
 		TArray<FString> AdditionalPackagesToCook;
 
@@ -3513,7 +3513,7 @@ ESavePackageResult UPackage::Save(UPackage* InOuter, UObject* Base, EObjectFlags
 				TSet<UObject*> DependenciesReferencedByNonRedirectors;
 		
 				/** If true, we are going to compress the package to memory to save a little time */
-				bool bCompressFromMemory = !!(InOuter->PackageFlags & PKG_StoreCompressed);
+				const bool bCompressFromMemory = InOuter->HasAnyPackageFlags(PKG_StoreCompressed);
 
 				/** If true, we are going to save to disk async to save time. */
 				bool bSaveAsync = !!(SaveFlags & SAVE_Async);
@@ -4278,8 +4278,8 @@ ESavePackageResult UPackage::Save(UPackage* InOuter, UObject* Base, EObjectFlags
 						UPackage* Package = dynamic_cast<UPackage*>(Object);
 						if (Package != NULL)
 						{
-							Linker->ExportMap[i].PackageFlags = Package->PackageFlags;
-							if (!(Package->PackageFlags & PKG_ServerSideOnly))
+							Linker->ExportMap[i].PackageFlags = Package->GetPackageFlags();
+							if (!Package->HasAnyPackageFlags(PKG_ServerSideOnly))
 							{
 								Linker->ExportMap[i].PackageGuid = Package->GetGuid();
 							}
@@ -4704,7 +4704,7 @@ ESavePackageResult UPackage::Save(UPackage* InOuter, UObject* Base, EObjectFlags
 				Linker->LinkerRoot->ThisRequiresLocalizationGather(Linker->RequiresLocalizationGather());
 				
 				// Update package flags from package, in case serialization has modified package flags.
-				Linker->Summary.PackageFlags  = Linker->LinkerRoot->PackageFlags & ~PKG_NewlyCreated;
+				Linker->Summary.PackageFlags = Linker->LinkerRoot->GetPackageFlags() & ~PKG_NewlyCreated;
 
 				Linker->Seek(0);
 				*Linker << Linker->Summary;
@@ -4763,7 +4763,7 @@ ESavePackageResult UPackage::Save(UPackage* InOuter, UObject* Base, EObjectFlags
 						UE_LOG_COOK_TIME(TEXT("CompressFromMemory"));
 					}
 					// Compress the temporarily file to destination.
-					else if( InOuter->PackageFlags & PKG_StoreCompressed )
+					else if( InOuter->HasAnyPackageFlags(PKG_StoreCompressed) )
 					{
 						UE_LOG(LogSavePackage, Log,  TEXT("Compressing '%s' to '%s'"), *TempFilename, *NewPath );
 						FFileCompressionHelper CompressionHelper;
@@ -4772,7 +4772,7 @@ ESavePackageResult UPackage::Save(UPackage* InOuter, UObject* Base, EObjectFlags
 						UE_LOG_COOK_TIME(TEXT("CompressFromTempFile"));
 					}
 					// Fully compress package in one "block".
-					else if( InOuter->PackageFlags & PKG_StoreFullyCompressed )
+					else if( InOuter->HasAnyPackageFlags(PKG_StoreFullyCompressed) )
 					{
 						UE_LOG(LogSavePackage, Log,  TEXT("Full-package compressing '%s' to '%s'"), *TempFilename, *NewPath );
 						FFileCompressionHelper CompressionHelper;
@@ -4806,7 +4806,7 @@ ESavePackageResult UPackage::Save(UPackage* InOuter, UObject* Base, EObjectFlags
 					}
 
 					// Make sure to clean up any stale .uncompressed_size files.
-					if( !(InOuter->PackageFlags & PKG_StoreFullyCompressed) )
+					if( !InOuter->HasAnyPackageFlags(PKG_StoreFullyCompressed) )
 					{
 						IFileManager::Get().Delete( *(NewPath + TEXT(".uncompressed_size")) );
 					}
@@ -4907,7 +4907,7 @@ ESavePackageResult UPackage::Save(UPackage* InOuter, UObject* Base, EObjectFlags
 		if( Success == true )
 		{
 			// Package has been save, so unmark NewlyCreated flag.
-			InOuter->PackageFlags &= ~PKG_NewlyCreated;
+			InOuter->ClearPackageFlags(PKG_NewlyCreated);
 
 			// send a message that the package was saved
 			UPackage::PackageSavedEvent.Broadcast(Filename, InOuter);
diff --git a/Engine/Source/Runtime/CoreUObject/Private/UObject/UObjectBaseUtility.cpp b/Engine/Source/Runtime/CoreUObject/Private/UObject/UObjectBaseUtility.cpp
--- a/Engine/Source/Runtime/CoreUObject/Private/UObject/UObjectBaseUtility.cpp
+++ b/Engine/Source/Runtime/CoreUObject/Private/UObject/UObjectBaseUtility.cpp
@@ -254,7 +254,7 @@ bool UObjectBaseUtility::IsInA( const UClass* SomeBaseClass ) const
 */
 bool UObjectBaseUtility::RootPackageHasAnyFlags( uint32 CheckFlagMask ) const
 {
-	return (GetOutermost()->PackageFlags&CheckFlagMask) != 0;
+	return GetOutermost()->HasAnyPackageFlags(CheckFlagMask);
 }
 
 /***********************/
diff --git a/Engine/Source/Runtime/CoreUObject/Private/UObject/UObjectGlobals.cpp b/Engine/Source/Runtime/CoreUObject/Private/UObject/UObjectGlobals.cpp
--- a/Engine/Source/Runtime/CoreUObject/Private/UObject/UObjectGlobals.cpp
+++ b/Engine/Source/Runtime/CoreUObject/Private/UObject/UObjectGlobals.cpp
@@ -290,7 +290,7 @@ void GlobalSetProperty( const TCHAR* Value, UClass* Class, UProperty* Property,
 			if( Object->IsA(Class) && !Object->IsPendingKill() )
 			{
 				// If we're in a PIE session then only allow set commands to affect PlayInEditor objects.
-				if( !GIsPlayInEditorWorld || ( Object->GetOutermost()->PackageFlags & PKG_PlayInEditor ) != 0 )
+				if( !GIsPlayInEditorWorld || Object->GetOutermost()->HasAnyPackageFlags(PKG_PlayInEditor)  )
 				{
 #if WITH_EDITOR
 					if( !Object->HasAnyFlags(RF_ClassDefaultObject) && bNotifyObjectOfChange )
@@ -1015,7 +1015,7 @@ UPackage* LoadPackageInternal(UPackage* InOuter, const TCHAR* InLongPackageName,
 		// If we are loading a package for diff'ing, set the package flag
 		if(LoadFlags & LOAD_ForDiff)
 		{
-			Result->PackageFlags |= PKG_ForDiffing;
+			Result->SetPackageFlags(PKG_ForDiffing);
 		}
 
 		// Save the filename we load from
@@ -1798,7 +1798,7 @@ bool SaveToTransactionBuffer(UObject* Object, bool bMarkDirty)
 	// Neither PIE world objects nor script packages should end up in the transaction buffer. Additionally, in order
 	// to save a copy of the object, we must have a transactor and the object must be transactional.
 	const bool IsTransactional = Object->HasAnyFlags(RF_Transactional);
-	const bool IsNotPIEOrContainsScriptObject = ( ( Object->GetOutermost()->PackageFlags & ( PKG_PlayInEditor | PKG_ContainsScript ) ) == 0 );
+	const bool IsNotPIEOrContainsScriptObject = (Object->GetOutermost()->HasAnyPackageFlags( PKG_PlayInEditor | PKG_ContainsScript) == false);
 
 	if ( GUndo && IsTransactional && IsNotPIEOrContainsScriptObject )
 	{
diff --git a/Engine/Source/Runtime/CoreUObject/Public/Misc/PackageName.h b/Engine/Source/Runtime/CoreUObject/Public/Misc/PackageName.h
--- a/Engine/Source/Runtime/CoreUObject/Public/Misc/PackageName.h
+++ b/Engine/Source/Runtime/CoreUObject/Public/Misc/PackageName.h
@@ -20,7 +20,7 @@ public:
 	static FString ConvertToLongScriptPackageName(const TCHAR* InShortName);
 
 	/**
-	 * Registeres all short package names found in ini files.
+	 * Registers all short package names found in ini files.
 	 */
 	static void RegisterShortPackageNamesForUObjectModules();
 
@@ -187,7 +187,7 @@ public:
 	static FName GetPackageMountPoint(const FString& InPackagePath);
 
 	/**
-	 * Checks if the package exisits on disk.
+	 * Checks if the package exists on disk.
 	 * 
 	 * @param LongPackageName Package name.
 	 * @param OutFilename Package filename on disk.
@@ -215,7 +215,7 @@ public:
 	static bool TryConvertShortPackagePathToLongInObjectPath(const FString& ObjectPath, FString& ConvertedObjectPath);
 
 	/**
-	 * Gets normlized object path i.e. with long package format.
+	 * Gets normalized object path i.e. with long package format.
 	 *
 	 * @param ObjectPath Path to the object.
 	 *
diff --git a/Engine/Source/Runtime/CoreUObject/Public/Misc/UObjectToken.h b/Engine/Source/Runtime/CoreUObject/Public/Misc/UObjectToken.h
--- a/Engine/Source/Runtime/CoreUObject/Public/Misc/UObjectToken.h
+++ b/Engine/Source/Runtime/CoreUObject/Public/Misc/UObjectToken.h
@@ -13,7 +13,7 @@ class FUObjectToken : public IMessageToken
 {
 public:
 	/** Factory method, tokens can only be constructed as shared refs */
-	COREUOBJECT_API static TSharedRef<FUObjectToken> Create( UObject* InObject, const FText& InLabelOverride = FText() )
+	COREUOBJECT_API static TSharedRef<FUObjectToken> Create( const UObject* InObject, const FText& InLabelOverride = FText() )
 	{
 		return MakeShareable(new FUObjectToken(InObject, InLabelOverride));
 	}
@@ -40,7 +40,7 @@ public:
 	}
 
 	/** Get the delegate for displaying the object name */
-	DECLARE_DELEGATE_RetVal_TwoParams(FText, FOnGetDisplayName, UObject*, bool);
+	DECLARE_DELEGATE_RetVal_TwoParams(FText, FOnGetDisplayName, const UObject*, bool);
 	COREUOBJECT_API static FOnGetDisplayName& DefaultOnGetObjectDisplayName()
 	{
 		return DefaultGetObjectDisplayName;
@@ -48,7 +48,7 @@ public:
 
 private:
 	/** Private constructor */
-	FUObjectToken( UObject* InObject,  const FText& InLabelOverride );
+	FUObjectToken( const UObject* InObject,  const FText& InLabelOverride );
 
 	/** An object being referenced by this token, if any */
 	FWeakObjectPtr ObjectBeingReferenced;
diff --git a/Engine/Source/Runtime/CoreUObject/Public/UObject/CoreObject.h b/Engine/Source/Runtime/CoreUObject/Public/UObject/CoreObject.h
--- a/Engine/Source/Runtime/CoreUObject/Public/UObject/CoreObject.h
+++ b/Engine/Source/Runtime/CoreUObject/Public/UObject/CoreObject.h
@@ -41,11 +41,33 @@ enum class ESavePackageResult
 /**
  * A package.
  */
+PRAGMA_DISABLE_DEPRECATION_WARNINGS // Required for auto-generated functions referencing PackageFlags
 class COREUOBJECT_API UPackage : public UObject
 {
-	DECLARE_CASTED_CLASS_INTRINSIC(UPackage,UObject,0,CoreUObject, CASTCLASS_UPackage)
+	// Have to unwind this macro to support the reference variable, can go back to commented declaration when removing the deprecated variable
+	// DECLARE_CASTED_CLASS_INTRINSIC(UPackage,UObject,0,CoreUObject, CASTCLASS_UPackage)
+
+#if WITH_HOT_RELOAD_CTORS
+	DECLARE_CASTED_CLASS_INTRINSIC_NO_CTOR_NO_VTABLE_CTOR( UPackage, UObject, 0, CoreUObject, CASTCLASS_UPackage, NO_API )
+	/** DO NOT USE. This constructor is for internal usage only for hot-reload purposes. */
+	UPackage(FVTableHelper& Helper)
+		: Super(Helper)
+		, PackageFlagsPrivate(PackageFlags) 
+	{
+	};
+#else
+	DECLARE_CASTED_CLASS_INTRINSIC_NO_CTOR(UPackage,UObject,0,CoreUObject, CASTCLASS_UPackage, NO_API)
+#endif
+
 
 public:
+
+	UPackage(const FObjectInitializer& ObjectInitializer = FObjectInitializer::Get())
+		: Super(ObjectInitializer)
+		, PackageFlagsPrivate(PackageFlags)
+	{
+	}
+
 	/** delegate type for package dirty state events.  ( Params: UPackage* ModifiedPackage ) */
 	DECLARE_MULTICAST_DELEGATE_OneParam(FOnPackageDirtyStateChanged, class UPackage*);
 	/** delegate type for package saved events ( Params: const FString& PackageFileName, UObject* Outer ) */
@@ -54,7 +76,7 @@ public:
 	DECLARE_MULTICAST_DELEGATE_TwoParams(FOnPackageMarkedDirty, class UPackage*, bool);
 
 	/** Delegate to notify subscribers when a package has been saved. This is triggered when the package saving
-	 *  has completed and was successfull. */
+	 *  has completed and was successful. */
 	static FOnPackageSaved PackageSavedEvent;
 	/** Delegate to notify subscribers when the dirty state of a package is changed.
 	 *  Allows the editor to register the modified package as one that should be prompted for source control checkout. 
@@ -103,7 +125,14 @@ public:
 	bool IsLoadedByEditorPropertiesOnly() const { return bLoadedByEditorPropertiesOnly; }
 #endif
 
-	/** Package flags, serialized.*/
+private:
+	/** New private Package Flags. When deprecated PackageFlags removed, make this a member instead of a reference */
+	uint32&	PackageFlagsPrivate;
+
+public:
+
+	/** Old deprecated public Package flags, serialized. When deprecation removed make PackageFlagsPrivate no longer a reference */
+	DEPRECATED(4.10, "Direct access of PackageFlags is deprecated, use Get/Set/Clear/HasAny/HasAll PackageFlags instead.")
 	uint32	PackageFlags;
 
 	/** The name of the file that this package was loaded from */
@@ -262,7 +291,7 @@ public:
 	 */
 	void ThisContainsMap() 
 	{
-		PackageFlags |= PKG_ContainsMap;
+		SetPackageFlags(PKG_ContainsMap);
 	}
 
 	/**
@@ -272,7 +301,7 @@ public:
 	 */
 	bool ContainsMap() const
 	{
-		return (PackageFlags & PKG_ContainsMap) ? true : false;
+		return HasAnyPackageFlags(PKG_ContainsMap);
 	}
 
 	/**
@@ -282,11 +311,11 @@ public:
 	{
 		if(Value)
 		{
-			PackageFlags |= PKG_RequiresLocalizationGather;
+			SetPackageFlags(PKG_RequiresLocalizationGather);
 		}
 		else
 		{
-			PackageFlags &= ~PKG_RequiresLocalizationGather;
+			ClearPackageFlags(PKG_RequiresLocalizationGather);
 		}
 	}
 
@@ -297,7 +326,77 @@ public:
 	 */
 	bool RequiresLocalizationGather() const
 	{
-		return (PackageFlags & PKG_RequiresLocalizationGather) ? true : false;
+		return HasAnyPackageFlags(PKG_RequiresLocalizationGather);
+	}
+
+// UE-21181 - trying to track when a flag gets set on a package due to PIE
+#if WITH_EDITOR
+	static UPackage* EditorPackage;
+	void SetPackageFlagsTo( uint32 NewFlags );
+#else
+	/**
+	 * Sets all package flags to the specified values.
+	 *
+	 * @param	NewFlags		New value for package flags
+	 */
+	FORCEINLINE void SetPackageFlagsTo( uint32 NewFlags )
+	{
+		PackageFlagsPrivate = NewFlags;
+	}
+#endif
+
+	/**
+	 * Set the specified flags to true. Does not affect any other flags.
+	 *
+	 * @param	NewFlags		Package flags to enable
+	 */
+	FORCEINLINE void SetPackageFlags( uint32 NewFlags )
+	{
+		SetPackageFlagsTo(PackageFlagsPrivate | NewFlags);
+	}
+
+	/**
+	 * Set the specified flags to false. Does not affect any other flags.
+	 *
+	 * @param	NewFlags		Package flags to disable
+	 */
+	FORCEINLINE void ClearPackageFlags( uint32 NewFlags )
+	{
+		SetPackageFlagsTo(PackageFlagsPrivate & ~NewFlags);
+	}
+
+	/**
+	 * Used to safely check whether the passed in flag is set.
+	 *
+	 * @param	FlagsToCheck		Package flags to check for
+	 *
+	 * @return	true if the passed in flag is set, false otherwise
+	 *			(including no flag passed in, unless the FlagsToCheck is CLASS_AllFlags)
+	 */
+	FORCEINLINE bool HasAnyPackageFlags( uint32 FlagsToCheck ) const
+	{
+		return (PackageFlagsPrivate & FlagsToCheck) != 0;
+	}
+
+	/**
+	 * Used to safely check whether all of the passed in flags are set.
+	 *
+	 * @param FlagsToCheck	Package flags to check for
+	 * @return true if all of the passed in flags are set (including no flags passed in), false otherwise
+	 */
+	FORCEINLINE bool HasAllPackagesFlags( uint32 FlagsToCheck ) const
+	{
+		return ((PackageFlagsPrivate & FlagsToCheck) == FlagsToCheck);
+	}
+
+	/**
+	 * Gets the package flags.
+	 *
+	 * @return	The package flags.
+	 */
+	FORCEINLINE uint32 GetPackageFlags() const
+	{
+		return PackageFlagsPrivate;
 	}
 
 	typedef TFunction<void (const UObject* const, TArray<FGatherableTextData>&)> FLocalizationDataGatheringCallback;
@@ -472,6 +571,7 @@ public:
 	static EObjectMark GetObjectMarksForTargetPlatform( const class ITargetPlatform* TargetPlatform, const bool bIsCooking );
 
 };
+PRAGMA_ENABLE_DEPRECATION_WARNINGS
 
 /*-----------------------------------------------------------------------------
 	UMetaData.
diff --git a/Engine/Source/Runtime/Engine/Classes/AI/Navigation/NavFilters/NavigationQueryFilter.h b/Engine/Source/Runtime/Engine/Classes/AI/Navigation/NavFilters/NavigationQueryFilter.h
--- a/Engine/Source/Runtime/Engine/Classes/AI/Navigation/NavFilters/NavigationQueryFilter.h
+++ b/Engine/Source/Runtime/Engine/Classes/AI/Navigation/NavFilters/NavigationQueryFilter.h
@@ -16,7 +16,7 @@ struct ENGINE_API FNavigationFilterArea
 	TSubclassOf<UNavArea> AreaClass;
 
 	/** override for travel cost */
-	UPROPERTY(EditAnywhere, Category=Area, meta=(EditCondition="bOverrideTravelCost",ClampMin=1))
+	UPROPERTY(EditAnywhere, Category=Area, meta=(EditCondition="bOverrideTravelCost",ClampMin=0.001))
 	float TravelCostOverride;
 
 	/** override for entering cost */
diff --git a/Engine/Source/Runtime/Engine/Classes/AI/Navigation/NavigationSystem.h b/Engine/Source/Runtime/Engine/Classes/AI/Navigation/NavigationSystem.h
--- a/Engine/Source/Runtime/Engine/Classes/AI/Navigation/NavigationSystem.h
+++ b/Engine/Source/Runtime/Engine/Classes/AI/Navigation/NavigationSystem.h
@@ -733,6 +733,8 @@ public:
 	/** */
 	virtual void OnWorldInitDone(FNavigationSystemRunMode Mode);
 
+	FORCEINLINE bool IsInitialized() const { return bWorldInitDone; }
+
 	/** adds BSP collisions of currently streamed in levels to octree */
 	void InitializeLevelCollisions();
 
@@ -836,9 +838,10 @@ protected:
 	/** set of locking flags applied on startup of navigation system */
 	uint8 InitialNavBuildingLockFlags;
 
-	uint8 bInitialSetupHasBeenPerformed:1;
-	uint8 bInitialLevelsAdded:1;
-	uint8 bAsyncBuildPaused:1;
+	uint8 bInitialSetupHasBeenPerformed : 1;
+	uint8 bInitialLevelsAdded : 1;
+	uint8 bWorldInitDone : 1;
+	uint8 bAsyncBuildPaused : 1;
 
 	/** cached navigable world bounding box*/
 	mutable FBox NavigableWorldBounds;
diff --git a/Engine/Source/Runtime/Engine/Classes/AI/Navigation/NavigationTypes.h b/Engine/Source/Runtime/Engine/Classes/AI/Navigation/NavigationTypes.h
--- a/Engine/Source/Runtime/Engine/Classes/AI/Navigation/NavigationTypes.h
+++ b/Engine/Source/Runtime/Engine/Classes/AI/Navigation/NavigationTypes.h
@@ -472,16 +472,16 @@ struct ENGINE_API FNavAgentProperties : public FMovementProperties
 	FORCEINLINE bool IsValid() const { return AgentRadius >= 0.f && AgentHeight >= 0.f; }
 	FORCEINLINE bool HasStepHeightOverride() const { return AgentStepHeight >= 0.0f; }
 
+	bool IsNavDataMatching(const FNavAgentProperties& Other) const;
+
 	FORCEINLINE bool IsEquivalent(const FNavAgentProperties& Other, float Precision = 5.f) const
 	{
-		return FGenericPlatformMath::Abs(AgentRadius - Other.AgentRadius) < Precision &&
-			FGenericPlatformMath::Abs(AgentHeight - Other.AgentHeight) < Precision &&
-			FGenericPlatformMath::Abs(AgentStepHeight - Other.AgentStepHeight) < Precision &&
-			IsNavDataMatching(Other);
+		return FGenericPlatformMath::Abs(AgentRadius - Other.AgentRadius) < Precision
+			&& FGenericPlatformMath::Abs(AgentHeight - Other.AgentHeight) < Precision
+			&& (HasStepHeightOverride() == false || FGenericPlatformMath::Abs(AgentStepHeight - Other.AgentStepHeight) < Precision)
+			&& IsNavDataMatching(Other);
 	}
-
-	bool IsNavDataMatching(const FNavAgentProperties& Other) const;
-
+	
 	bool operator==(const FNavAgentProperties& Other) const
 	{
 		return IsEquivalent(Other);
diff --git a/Engine/Source/Runtime/Engine/Classes/Animation/AnimBlueprintGeneratedClass.h b/Engine/Source/Runtime/Engine/Classes/Animation/AnimBlueprintGeneratedClass.h
--- a/Engine/Source/Runtime/Engine/Classes/Animation/AnimBlueprintGeneratedClass.h
+++ b/Engine/Source/Runtime/Engine/Classes/Animation/AnimBlueprintGeneratedClass.h
@@ -144,6 +144,7 @@ public:
 
 	void ResetNodeVisitSites();
 	void RecordNodeVisit(int32 TargetNodeIndex, int32 SourceNodeIndex, float BlendWeight);
+	void RecordNodeVisitArray(const TArray<FNodeVisit>& Nodes);
 #endif
 };
 
@@ -296,14 +297,14 @@ public:
 };
 
 template<typename NodeType>
-NodeType* GetNodeFromPropertyIndex(UAnimInstance* AnimInstance, const UAnimBlueprintGeneratedClass* AnimBlueprintClass, int32 PropertyIndex)
+NodeType* GetNodeFromPropertyIndex(UObject* AnimInstanceObject, const UAnimBlueprintGeneratedClass* AnimBlueprintClass, int32 PropertyIndex)
 {
 	if (PropertyIndex != INDEX_NONE)
 	{
 		UStructProperty* NodeProperty = AnimBlueprintClass->AnimNodeProperties[AnimBlueprintClass->AnimNodeProperties.Num() - 1 - PropertyIndex]; //@TODO: Crazysauce
 		check(NodeProperty->Struct == NodeType::StaticStruct());
-		return NodeProperty->ContainerPtrToValuePtr<NodeType>(AnimInstance);
+		return NodeProperty->ContainerPtrToValuePtr<NodeType>(AnimInstanceObject);
 	}
 
 	return NULL;
-}
\ No newline at end of file
+}
diff --git a/Engine/Source/Runtime/Engine/Classes/Animation/AnimComposite.h b/Engine/Source/Runtime/Engine/Classes/Animation/AnimComposite.h
--- a/Engine/Source/Runtime/Engine/Classes/Animation/AnimComposite.h
+++ b/Engine/Source/Runtime/Engine/Classes/Animation/AnimComposite.h
@@ -26,7 +26,7 @@ public:
 #endif // WITH_EDITORONLY_DATA
 
 	//~ Begin UAnimSequenceBase Interface
-	ENGINE_API virtual void OnAssetPlayerTickedInternal(FAnimAssetTickContext &Context, const float PreviousTime, const float MoveDelta, const FAnimTickRecord &Instance, class UAnimInstance* InstanceOwner) const override;
+	ENGINE_API virtual void HandleAssetPlayerTickedInternal(FAnimAssetTickContext &Context, const float PreviousTime, const float MoveDelta, const FAnimTickRecord &Instance, struct FAnimNotifyQueue& NotifyQueue) const override;
 	virtual void GetAnimationPose(FCompactPose& OutPose, FBlendedCurve& OutCurve, const FAnimExtractContext& ExtractionContext) const override;	
 	//~ End UAnimSequenceBase Interface
 	//~ Begin UAnimSequence Interface
diff --git a/Engine/Source/Runtime/Engine/Classes/Animation/AnimInstance.h b/Engine/Source/Runtime/Engine/Classes/Animation/AnimInstance.h
--- a/Engine/Source/Runtime/Engine/Classes/Animation/AnimInstance.h
+++ b/Engine/Source/Runtime/Engine/Classes/Animation/AnimInstance.h
@@ -8,6 +8,7 @@
 #include "AnimStateMachineTypes.h"
 #include "BonePose.h"
 #include "Animation/AnimTypes.h"
+#include "Animation/AnimNotifyQueue.h"
 #include "AnimInstance.generated.h"
 
 struct FAnimMontageInstance;
@@ -24,6 +25,7 @@ struct FTransform;
 class FDebugDisplayInfo;
 struct FAnimNode_AssetPlayerBase;
 struct FAnimNode_Base;
+struct FAnimInstanceProxy;
 
 DECLARE_DELEGATE_OneParam(FOnMontageStarted, UAnimMontage*)
 DECLARE_DELEGATE_TwoParams(FOnMontageEnded, UAnimMontage*, bool /*bInterrupted*/)
@@ -320,11 +322,13 @@ struct FMontageActiveSlotTracker
 
 struct FMontageEvaluationState
 {
-	FMontageEvaluationState(UAnimMontage* InMontage, float InWeight, float InDesiredWeight, float InPosition) 
+	FMontageEvaluationState(UAnimMontage* InMontage, float InWeight, float InDesiredWeight, float InPosition, bool bInIsPlaying, bool bInIsActive) 
 		: Montage(InMontage)
 		, MontageWeight(InWeight)
 		, DesiredWeight(InDesiredWeight)
-		, MontagePosition(InPosition) 
+		, MontagePosition(InPosition)
+		, bIsPlaying(bInIsPlaying)
+		, bIsActive(bInIsActive)
 	{}
 
 	// The montage to evaluate
@@ -338,6 +342,12 @@ struct FMontageEvaluationState
 
 	// The position to evaluate this montage at
 	float MontagePosition;
+
+	// Whether this montage is playing
+	bool bIsPlaying;
+
+	// Whether this montage is valid and not stopped
+	bool bIsActive;
 };
 
 struct FGraphTraversalCounter
@@ -402,6 +412,13 @@ class ENGINE_API UAnimInstance : public UObject
 {
 	GENERATED_UCLASS_BODY()
 
+	typedef FAnimInstanceProxy ProxyType;
+
+	// Disable compiler-generated deprecation warnings by implementing our own destructor
+	PRAGMA_DISABLE_DEPRECATION_WARNINGS
+	~UAnimInstance() {}
+	PRAGMA_ENABLE_DEPRECATION_WARNINGS
+
 	/** DeltaTime **/
 	UPROPERTY()
 	float DeltaTime_DEPRECATED;
@@ -411,13 +428,15 @@ class ENGINE_API UAnimInstance : public UObject
 	USkeleton* CurrentSkeleton;
 
 	// The list of animation assets which are going to be evaluated this frame and need to be ticked (ungrouped)
+	DEPRECATED(4.11, "This cannot be accessed directly as it is potentially in use on worker threads")
 	TArray<FAnimTickRecord> UngroupedActivePlayerArrays[2];
 
 	// The set of tick groups for this anim instance
+	DEPRECATED(4.11, "This cannot be accessed directly as it is potentially in use on worker threads")
 	TArray<FAnimGroupInstance> SyncGroupArrays[2];
 
 	/** Array indicating active vertex anims (by reference) generated by anim instance. */
-	UPROPERTY(transient)
+	DEPRECATED(4.11, "This cannot be accessed directly as it is potentially in use on worker threads")
 	TArray<struct FActiveVertexAnim> VertexAnims;
 
 	// Sets where this blueprint pulls Root Motion from
@@ -446,16 +465,21 @@ class ENGINE_API UAnimInstance : public UObject
 public:
 
 	// @todo document
+	DEPRECATED(4.11, "This cannot be called directly as it relies on data potentially in use on worker threads")
 	void MakeSequenceTickRecord(FAnimTickRecord& TickRecord, UAnimSequenceBase* Sequence, bool bLooping, float PlayRate, float FinalBlendWeight, float& CurrentTime, FMarkerTickRecord& MarkerTickRecord) const;
+	DEPRECATED(4.11, "This cannot be called directly as it relies on data potentially in use on worker threads")
 	void MakeBlendSpaceTickRecord(FAnimTickRecord& TickRecord, UBlendSpaceBase* BlendSpace, const FVector& BlendInput, TArray<FBlendSampleData>& BlendSampleDataCache, FBlendFilter& BlendFilter, bool bLooping, float PlayRate, float FinalBlendWeight, float& CurrentTime, FMarkerTickRecord& MarkerTickRecord) const;
 	void MakeMontageTickRecord(FAnimTickRecord& TickRecord, class UAnimMontage* Montage, float CurrentPosition, float PreviousPosition, float MoveDelta, float Weight, TArray<FPassedMarker>& MarkersPassedThisTick, FMarkerTickRecord& MarkerTickRecord);
 
+	DEPRECATED(4.11, "This cannot be called directly as it relies on data potentially in use on worker threads")
 	void SequenceAdvanceImmediate(UAnimSequenceBase* Sequence, bool bLooping, float PlayRate, float DeltaSeconds, /*inout*/ float& CurrentTime, FMarkerTickRecord& MarkerTickRecord);
 
 	// @todo document
+	DEPRECATED(4.11, "This cannot be called directly as it relies on data potentially in use on worker threads")
 	void BlendSpaceAdvanceImmediate(UBlendSpaceBase* BlendSpace, const FVector& BlendInput, TArray<FBlendSampleData> & BlendSampleDataCache, FBlendFilter & BlendFilter, bool bLooping, float PlayRate, float DeltaSeconds, /*inout*/ float& CurrentTime, FMarkerTickRecord& MarkerTickRecord);
 
 	// Creates an uninitialized tick record in the list for the correct group or the ungrouped array.  If the group is valid, OutSyncGroupPtr will point to the group.
+	DEPRECATED(4.11, "This cannot be called directly as it relies on data potentially in use on worker threads")
 	FAnimTickRecord& CreateUninitializedTickRecord(int32 GroupIndex, FAnimGroupInstance*& OutSyncGroupPtr);
 
 	/**
@@ -467,18 +491,26 @@ public:
 	 * @param out_SourceWeight : The Source weight for this node. 
 	 * @param out_TotalNodeWeight : Total weight of this node
 	 */
+	DEPRECATED(4.11, "This cannot be called directly as it relies on data potentially in use on worker threads")
 	void GetSlotWeight(FName const& SlotNodeName, float& out_SlotNodeWeight, float& out_SourceWeight, float& out_TotalNodeWeight) const;
+
+	DEPRECATED(4.11, "This cannot be called directly as it relies on data potentially in use on worker threads")
 	void SlotEvaluatePose(FName SlotNodeName, const FCompactPose& SourcePose, const FBlendedCurve& SourceCurve, float InSourceWeight, FCompactPose& BlendedPose, FBlendedCurve& BlendedCurve, float InBlendWeight, float InTotalNodeWeight);
 
 	// slot node run-time functions
+	DEPRECATED(4.11, "This cannot be called directly as it relies on data potentially in use on worker threads")
 	void ReinitializeSlotNodes();
+	DEPRECATED(4.11, "This cannot be called directly as it relies on data potentially in use on worker threads")
 	void RegisterSlotNodeWithAnimInstance(FName SlotNodeName);
+	DEPRECATED(4.11, "This cannot be called directly as it relies on data potentially in use on worker threads")
 	void UpdateSlotNodeWeight(FName SlotNodeName, float Weight);
 	// if it doesn't tick, it will keep old weight, so we'll have to clear it in the beginning of tick
+	DEPRECATED(4.11, "This cannot be called directly as it relies on data potentially in use on worker threads")
 	void ClearSlotNodeWeights();
 	bool IsSlotNodeRelevantForNotifies(FName SlotNodeName) const;
 
 	// Allow slot nodes to store off their root motion weight during ticking
+	DEPRECATED(4.11, "This cannot be called directly as it relies on data potentially in use on worker threads")
 	void UpdateSlotRootMotionWeight(FName SlotNodeName, float Weight);
 	// Get the root motion weight for the montage slot
 	float GetSlotRootMotionWeight(FName SlotNodeName) const;
@@ -655,7 +687,7 @@ public:
 
 	/** Get Currently active montage instance.
 		Note that there might be multiple Active at the same time. This will only return the first active one it finds. **/
-	FAnimMontageInstance * GetActiveMontageInstance();
+	FAnimMontageInstance * GetActiveMontageInstance() const;
 
 	/** Get Active FAnimMontageInstance for given Montage asset. Will return NULL if Montage is not currently Active. */
 	FAnimMontageInstance * GetActiveInstanceForMontage(UAnimMontage const & Montage) const;
@@ -666,6 +698,7 @@ public:
 	TArray<struct FAnimMontageInstance*> MontageInstances;
 
 	// Cached data for montage evaluation, save us having to access MontageInstances from slot nodes as that isn't thread safe
+	DEPRECATED(4.11, "Please use FAnimInstanceProxy::GetMontageEvaluationData()")
 	TArray<FMontageEvaluationState> MontageEvaluationData;
 
 	virtual void OnMontageInstanceStopped(FAnimMontageInstance & StoppedMontageInstance);
@@ -798,23 +831,29 @@ public:
 	float GetRelevantAnimTimeFraction(int32 MachineIndex, int32 StateIndex);
 
 	/** Gets an unchecked (can return nullptr) node given an index into the node property array */
+	DEPRECATED(4.11, "Please use FAnimInstanceProxy::GetNodeFromIndexUntyped")
 	FAnimNode_Base* GetNodeFromIndexUntyped(int32 NodeIdx, UScriptStruct* RequiredStructType);
 
 	/** Gets a checked node given an index into the node property array */
+	DEPRECATED(4.11, "Please use FAnimInstanceProxy::GetCheckedNodeFromIndexUntyped")
 	FAnimNode_Base* GetCheckedNodeFromIndexUntyped(int32 NodeIdx, UScriptStruct* RequiredStructType);
 
 	/** Gets a checked node given an index into the node property array */
 	template<class NodeType>
+	DEPRECATED(4.11, "Please use FAnimInstanceProxy::GetCheckedNodeFromIndex")
 	NodeType* GetCheckedNodeFromIndex(int32 NodeIdx)
 	{
-		return (NodeType*)GetNodeFromIndexUntypedChecked(NodeIdx, NodeType::StaticStruct());
+		return (NodeType*)GetCheckedNodeFromIndexUntyped(NodeIdx, NodeType::StaticStruct());
 	}
 
 	/** Gets an unchecked (can return nullptr) node given an index into the node property array */
 	template<class NodeType>
+	DEPRECATED(4.11, "Please use FAnimInstanceProxy::GetNodeFromIndex")
 	NodeType* GetNodeFromIndex(int32 NodeIdx)
 	{
+		PRAGMA_DISABLE_DEPRECATION_WARNINGS
 		return (NodeType*)GetNodeFromIndexUntyped(NodeIdx, NodeType::StaticStruct());
+		PRAGMA_ENABLE_DEPRECATION_WARNINGS
 	}
 
 	/** Gets the runtime instance of the specified state machine by Name */
@@ -825,7 +864,6 @@ public:
 
 	/** Returns the baked sync group index from the compile step */
 	int32 GetSyncGroupIndexFromName(FName SyncGroupName) const;
-
 protected:
 
 	/** Gets the index of the state machine matching MachineName */
@@ -928,11 +966,14 @@ public:
 
 public:
 	// Root node of animation graph
+	DEPRECATED(4.11, "RootNode access has been moved to FAnimInstanceProxy")
 	struct FAnimNode_Base* RootNode;
 
 public:
 	//~ Begin UObject Interface
 	virtual void Serialize(FArchive& Ar) override;
+	virtual void BeginDestroy() override;
+	virtual void PostInitProperties() override;
 	static void AddReferencedObjects(UObject* InThis, FReferenceCollector& Collector);
 	//~ End UObject Interface
 
@@ -954,7 +995,15 @@ public:
 	/** Called after updates are completed, dispatches notifies etc. */
 	void PostUpdateAnimation();
 
+	DEPRECATED(4.11, "This function should no longer be used. Use ParallelEvaluateAnimation")
 	void EvaluateAnimation(struct FPoseContext& Output);
+
+	/** Check whether evaluation can be performed on the supplied skeletal mesh. Can be called from worker threads. */
+	bool ParallelCanEvaluate(const USkeletalMesh* InSkeletalMesh) const;
+
+	/** Perform evaluation. Can be called from worker threads. */
+	void ParallelEvaluateAnimation(bool bForceRefPose, const USkeletalMesh* InSkeletalMesh, TArray<FTransform>& OutLocalAtoms, TArray<FActiveVertexAnim>& OutVertexAnims, FBlendedCurve& OutCurve);
+
 	void PostEvaluateAnimation();
 	void UninitializeAnimation();
 
@@ -966,10 +1015,12 @@ public:
 	virtual void NativeUpdateAnimation(float DeltaSeconds);
 	// Native update override point. Can be called from a worker thread. This is a good place to do any
 	// heavy lifting (as opposed to NativeUpdateAnimation_GameThread()).
+	// This function should not be used. Worker thread updates should be performed in the FAnimInstanceProxy attached to this instance.
 	virtual void NativeUpdateAnimation_WorkerThread(float DeltaSeconds);
 	// Native evaluate override point.
 	// @return true if this function is implemented, false otherwise.
 	// Note: the node graph will not be evaluated if this function returns true
+	DEPRECATED(4.11, "Please use FAnimInstanceProxy::Evaluate")
 	virtual bool NativeEvaluateAnimation(FPoseContext& Output);
 	// Native Post Evaluate override point
 	virtual void NativePostEvaluateAnimation();
@@ -999,25 +1050,33 @@ public:
 	void DisplayDebug(UCanvas* Canvas, const FDebugDisplayInfo& DebugDisplay, float& YL, float& YPos);
 public:
 
+	/** Access the required bones array */
+	FBoneContainer& GetRequiredBones();	
+
 	/** Temporary array of bone indices required this frame. Should be subset of Skeleton and Mesh's RequiredBones */
+	DEPRECATED(4.11, "This cannot be accessed directly, use UnimInstance::GetRequiredBones")
 	FBoneContainer RequiredBones;
 
 	/** Animation Notifies that has been triggered in the latest tick **/
+	DEPRECATED(4.11, "Use UAnimInstance::NotifyQueue")
 	TArray<const struct FAnimNotifyEvent *> AnimNotifies;
 
+	/** Animation Notifies that has been triggered in the latest tick **/
+	FAnimNotifyQueue NotifyQueue;
+
 	/** Currently Active AnimNotifyState, stored as a copy of the event as we need to
 		call NotifyEnd on the event after a deletion in the editor. After this the event
 		is removed correctly. */
 	UPROPERTY(transient)
 	TArray<FAnimNotifyEvent> ActiveAnimNotifyState;
 
-	// Some nodes need access to time dilation (stored in world settings) so we store it off and access it here
-	float CurrentTimeDilation;
-
 protected:
-	int32 GetSyncGroupReadIndex() const { return 1 - SyncGroupWriteIndex; }
-	int32 GetSyncGroupWriteIndex() const { return SyncGroupWriteIndex; }
-	void TickSyncGroupWriteIndex() { SyncGroupWriteIndex = GetSyncGroupReadIndex(); }
+	DEPRECATED(4.11, "This cannot be called directly as it relies on data potentially in use on worker threads")
+	int32 GetSyncGroupReadIndex() const;
+	DEPRECATED(4.11, "This cannot be called directly as it relies on data potentially in use on worker threads")
+	int32 GetSyncGroupWriteIndex() const;
+	DEPRECATED(4.11, "This cannot be called directly as it relies on data potentially in use on worker threads")
+	void TickSyncGroupWriteIndex();
 
 private:
 	/** Curve Values that are added to trigger in event**/
@@ -1025,22 +1084,25 @@ private:
 	/** Material parameters that we had been changing and now need to clear */
 	TArray<FName> MaterialParamatersToClear;
 
-	// Current sync group buffer index
-	int32 SyncGroupWriteIndex;
-
 	//This frames marker sync data
 	FMarkerTickContext MarkerTickContext;
 
-	// The last time delta passed into UpdateAnimation()
-	float CurrentDeltaSeconds;
-
 public: 
 	/** Update all internal curves from Blended Curve */
 	void UpdateCurves(const FBlendedCurve& InCurves);
 
+	/** Check whether we have active morph target curves */
+	bool HasMorphTargetCurves() const;
+
+	/** Access active morph target curves */
+	TMap<FName, float>& GetMorphTargetCurves();
+
 	/** Morph Target Curves that will be used for SkeletalMeshComponent **/
+	DEPRECATED(4.11, "This cannot be accessed directly, use UAnimInstance::GetMorphTargetCurves")
 	TMap<FName, float>	MorphTargetCurves;
+
 	/** Material Curves that will be used for SkeletalMeshComponent **/
+	DEPRECATED(4.11, "This cannot be accessed directly, use FAnimInstanceProxy::GetMaterialParameterCurves")
 	TMap<FName, float>	MaterialParameterCurves;
 
 #if WITH_EDITORONLY_DATA
@@ -1052,11 +1114,16 @@ public:
 #endif
 
 public:
+	DEPRECATED(4.11, "This cannot be accessed directly as it is potentially in use on worker threads")
 	FGraphTraversalCounter InitializationCounter;
+	DEPRECATED(4.11, "This cannot be accessed directly as it is potentially in use on worker threads")
 	FGraphTraversalCounter CachedBonesCounter;
+	DEPRECATED(4.11, "This cannot be accessed directly as it is potentially in use on worker threads")
 	FGraphTraversalCounter UpdateCounter;
+	DEPRECATED(4.11, "This cannot be accessed directly as it is potentially in use on worker threads")
 	FGraphTraversalCounter EvaluationCounter;
 	FGraphTraversalCounter DebugDataCounter;
+	DEPRECATED(4.11, "This cannot be accessed directly as it is potentially in use on worker threads")
 	FGraphTraversalCounter SlotNodeInitializationCounter;
 
 private:
@@ -1070,7 +1137,7 @@ public:
 	void RecalcRequiredBones();
 
 	/** When RequiredBones mapping has changed, AnimNodes need to update their bones caches. */
-	UPROPERTY(Transient)
+	DEPRECATED(4.11, "This cannot be accessed directly, use FAnimInstanceProxy::bBoneCachesInvalidated")
 	bool bBoneCachesInvalidated;
 
 	// @todo document
@@ -1079,15 +1146,19 @@ public:
 	virtual UWorld* GetWorld() const override;
 
 	/** Add anim notifier **/
+	DEPRECATED(4.11, "Please use NotifyQueue.AddAnimNotifies")
 	void AddAnimNotifies(const TArray<const FAnimNotifyEvent*>& NewNotifies, const float InstanceWeight);
 
 	/** Should the notifies current filtering mode stop it from triggering */
+	DEPRECATED(4.11, "This cannot be called directly as it relies on data potentially in use on worker threads")
 	bool PassesFiltering(const FAnimNotifyEvent* Notify) const;
 
 	/** Work out whether this notify should be triggered based on its chance of triggering value */
+	DEPRECATED(4.11, "This cannot be called directly as it relies on data potentially in use on worker threads")
 	bool PassesChanceOfTriggering(const FAnimNotifyEvent* Event) const;
 
 	/** Queues an Anim Notify from the shared list on our generated class */
+	DEPRECATED(4.11, "This cannot be called directly as it relies on data potentially in use on worker threads")
 	void AddAnimNotifyFromGeneratedClass(int32 NotifyIndex);
 
 	/** Trigger AnimNotifies **/
@@ -1097,6 +1168,12 @@ public:
 	/** Add curve float data using a curve Uid, the name of the curve will be resolved from the skeleton **/
 	void AddCurveValue(const USkeleton::AnimCurveUID Uid, float Value, int32 CurveTypeFlags);
 
+	/** Pass-thru function to proxy - only call on the game thread */
+	void UpdateMorphTargetCurves(const TMap<FName, float>& InMorphTargetCurves);
+
+	/** Pass-thru function to proxy - only call on the game thread */
+	void UpdateComponentsMaterialParameters(UPrimitiveComponent* Component);
+
 protected:
 	/** 
 	 * Add curve float data, using a curve name. External values should all be added using
@@ -1117,32 +1194,20 @@ public:
 	/** Get current accumulated root motion, removing it from the AnimInstance in the process */
 	FRootMotionMovementParams ConsumeExtractedRootMotion(float Alpha);
 
+	/** Wrapper around UpdateActiveVertexAnims that can use vertex anims internal to the proxy */
+	TArray<FActiveVertexAnim> UpdateActiveVertexAnims(const USkeletalMesh* SkeletalMesh);
+
 private:
 	/** Active Root Motion Montage Instance, if any. */
 	struct FAnimMontageInstance * RootMotionMontageInstance;
 
-	// Root motion extracted from animation since the last time ConsumeExtractedRootMotion was called
-	FRootMotionMovementParams ExtractedRootMotion;
-
-	// Native bindings
 private:
-	/** Bind any native delegates that we have set up */
-	void BindNativeDelegates();
-
-	/** Native transition rules */
-	TArray<FNativeTransitionBinding> NativeTransitionBindings;
-
-	/** Native state entry bindings */
-	TArray<FNativeStateBinding> NativeStateEntryBindings;
-
-	/** Native state exit bindings */
-	TArray<FNativeStateBinding> NativeStateExitBindings;
-
 	// update montage
 	void UpdateMontage(float DeltaSeconds);
 
 protected:
 	/** Update all animation node */
+	DEPRECATED(4.11, "This cannot be called directly as it relies on data potentially in use on worker threads")
 	virtual void UpdateAnimationNode(float DeltaSeconds);
 
 	// Updates the montage data used for evaluation based on the current playing montages
@@ -1152,8 +1217,94 @@ protected:
 	void PreUpdateAnimation(float DeltaSeconds);
 
 	/** Actually does the update work, can be called from a worker thread  */
-	void UpdateAnimationInternal(float DeltaSeconds);
+	void UpdateAnimationInternal_Concurrent(float DeltaSeconds, FAnimInstanceProxy& Proxy);
+
+	/** Override point for derived classes to create their own proxy objects (allows custom allocation) */
+	virtual FAnimInstanceProxy* CreateAnimInstanceProxy();
+
+	/** Override point for derived classes to destroy their own proxy objects (allows custom allocation) */
+	virtual void DestroyAnimInstanceProxy(FAnimInstanceProxy* InProxy);
 
-	/** Tick all active asset player instances */
-	virtual void TickAssetPlayerInstances(float DeltaSeconds);
+	/** Access the proxy but block if a task is currently in progress as it wouldn't be safe to access it */
+	template <typename T = FAnimInstanceProxy>
+	FORCEINLINE T& GetProxyOnGameThread()
+	{
+		check(IsInGameThread());
+		if(GetOuter() && GetOuter()->IsA<USkeletalMeshComponent>())
+		{
+			bool bBlockOnTask = true;
+			bool bPerformPostAnimEvaluation = true;
+			GetSkelMeshComponent()->HandleExistingParallelEvaluationTask(bBlockOnTask, bPerformPostAnimEvaluation);
+		}
+		if(AnimInstanceProxy == nullptr)
+		{
+			AnimInstanceProxy = CreateAnimInstanceProxy();
+		}
+		return *static_cast<T*>(AnimInstanceProxy);
+	}
+
+	/** Access the proxy but block if a task is currently in progress as it wouldn't be safe to access it */
+	template <typename T = FAnimInstanceProxy>
+	FORCEINLINE const T& GetProxyOnGameThread() const
+	{
+		check(IsInGameThread());
+		if(GetOuter() && GetOuter()->IsA<USkeletalMeshComponent>())
+		{
+			bool bBlockOnTask = true;
+			bool bPerformPostAnimEvaluation = true;
+			GetSkelMeshComponent()->HandleExistingParallelEvaluationTask(bBlockOnTask, bPerformPostAnimEvaluation);
+		}
+		if(AnimInstanceProxy == nullptr)
+		{
+			AnimInstanceProxy = const_cast<UAnimInstance*>(this)->CreateAnimInstanceProxy();
+		}
+		return *static_cast<const T*>(AnimInstanceProxy);
+	}
+
+	/** Access the proxy but block if a task is currently in progress (and we are on the game thread) as it wouldn't be safe to access it */
+	template <typename T = FAnimInstanceProxy>
+	FORCEINLINE T& GetProxyOnAnyThread()
+	{
+		if(GetOuter() && GetOuter()->IsA<USkeletalMeshComponent>())
+		{
+			if(IsInGameThread())
+			{
+				bool bBlockOnTask = true;
+				bool bPerformPostAnimEvaluation = true;
+				GetSkelMeshComponent()->HandleExistingParallelEvaluationTask(bBlockOnTask, bPerformPostAnimEvaluation);
+			}
+		}
+		if(AnimInstanceProxy == nullptr)
+		{
+			AnimInstanceProxy = CreateAnimInstanceProxy();
+		}
+		return *static_cast<T*>(AnimInstanceProxy);
+	}
+
+	/** Access the proxy but block if a task is currently in progress (and we are on the game thread) as it wouldn't be safe to access it */
+	template <typename T = FAnimInstanceProxy>
+	FORCEINLINE const T& GetProxyOnAnyThread() const
+	{
+		if(GetOuter() && GetOuter()->IsA<USkeletalMeshComponent>())
+		{
+			if(IsInGameThread())
+			{
+				bool bBlockOnTask = true;
+				bool bPerformPostAnimEvaluation = true;
+				GetSkelMeshComponent()->HandleExistingParallelEvaluationTask(bBlockOnTask, bPerformPostAnimEvaluation);
+			}
+		}
+		if(AnimInstanceProxy == nullptr)
+		{
+			AnimInstanceProxy = const_cast<UAnimInstance*>(this)->CreateAnimInstanceProxy();
+		}
+		return *static_cast<const T*>(AnimInstanceProxy);
+	}
+
+	// TODO: Remove after deprecation (4.11)
+	friend struct FAnimationBaseContext;
+
+protected:
+	/** Proxy object, nothing should access this from an externally-callable API as it is used as a scratch area on worker threads */
+	mutable FAnimInstanceProxy* AnimInstanceProxy;
 };
diff --git a/Engine/Source/Runtime/Engine/Classes/Animation/AnimMontage.h b/Engine/Source/Runtime/Engine/Classes/Animation/AnimMontage.h
--- a/Engine/Source/Runtime/Engine/Classes/Animation/AnimMontage.h
+++ b/Engine/Source/Runtime/Engine/Classes/Animation/AnimMontage.h
@@ -475,7 +475,7 @@ public:
 #endif // WITH_EDITOR
 	virtual void GetMarkerIndicesForTime(float CurrentTime, bool bLooping, const TArray<FName>& ValidMarkerNames, FMarkerPair& OutPrevMarker, FMarkerPair& OutNextMarker) const override;
 	virtual FMarkerSyncAnimPosition GetMarkerSyncPositionfromMarkerIndicies(int32 PrevMarker, int32 NextMarker, float CurrentTime) const override;
-	virtual void TickAssetPlayerInstance(FAnimTickRecord& Instance, class UAnimInstance* InstanceOwner, FAnimAssetTickContext& Context) const override;
+	virtual void TickAssetPlayer(FAnimTickRecord& Instance, struct FAnimNotifyQueue& NotifyQueue, FAnimAssetTickContext& Context) const override;
 	virtual TArray<FName>* GetUniqueMarkerNames() override { return &MarkerData.UniqueMarkerNames; }
 	//~ End AnimSequenceBase Interface
 
diff --git a/Engine/Source/Runtime/Engine/Classes/Animation/AnimNodeBase.h b/Engine/Source/Runtime/Engine/Classes/Animation/AnimNodeBase.h
--- a/Engine/Source/Runtime/Engine/Classes/Animation/AnimNodeBase.h
+++ b/Engine/Source/Runtime/Engine/Classes/Animation/AnimNodeBase.h
@@ -9,36 +9,36 @@
 #include "BonePose.h"
 #include "AnimNodeBase.generated.h"
 
+struct FAnimInstanceProxy;
+
 /** Base class for update/evaluate contexts */
 struct FAnimationBaseContext
 {
 public:
+	DEPRECATED(4.11, "Please use AnimInstanceProxy")
 	UAnimInstance* AnimInstance;
+
+	FAnimInstanceProxy* AnimInstanceProxy;
+
 protected:
-	FAnimationBaseContext(UAnimInstance* InInstance)
-		: AnimInstance(InInstance)
-	{
-	}
+	// DEPRECATED - Please use constructor that uses an FAnimInstanceProxy*
+	ENGINE_API FAnimationBaseContext(UAnimInstance* InAnimInstance);
+
+	ENGINE_API FAnimationBaseContext(FAnimInstanceProxy* InAnimInstanceProxy);
+
+public:
+	// we define a copy constructor here simply to avoid deprecation warnings with clang
+	ENGINE_API FAnimationBaseContext(const FAnimationBaseContext& InContext);
 
 public:
 	// Get the Blueprint Generated Class associated with this context, if there is one.
 	// Note: This can return NULL, so check the result.
-	UAnimBlueprintGeneratedClass* GetAnimBlueprintClass() const
-	{
-		return Cast<UAnimBlueprintGeneratedClass>(AnimInstance->GetClass());
-	}
+	ENGINE_API UAnimBlueprintGeneratedClass* GetAnimBlueprintClass() const;
 
 #if WITH_EDITORONLY_DATA
 	// Get the AnimBlueprint associated with this context, if there is one.
 	// Note: This can return NULL, so check the result.
-	UAnimBlueprint* GetAnimBlueprint() const
-	{
-		if (UAnimBlueprintGeneratedClass* Class = GetAnimBlueprintClass())
-		{
-			return Cast<UAnimBlueprint>(Class->ClassGeneratedBy);
-		}
-		return NULL;
-	}
+	ENGINE_API UAnimBlueprint* GetAnimBlueprint() const;
 #endif //WITH_EDITORONLY_DATA
 };
 
@@ -47,8 +47,14 @@ public:
 struct FAnimationInitializeContext : public FAnimationBaseContext
 {
 public:
-	FAnimationInitializeContext(UAnimInstance* InInstance)
-		: FAnimationBaseContext(InInstance)
+	DEPRECATED(4.11, "Please use constructor that uses an FAnimInstanceProxy*")
+	FAnimationInitializeContext(UAnimInstance* InAnimInstance)
+		: FAnimationBaseContext(InAnimInstance)
+	{
+	}
+
+	FAnimationInitializeContext(FAnimInstanceProxy* InAnimInstanceProxy)
+		: FAnimationBaseContext(InAnimInstanceProxy)
 	{
 	}
 };
@@ -60,8 +66,14 @@ public:
 struct FAnimationCacheBonesContext : public FAnimationBaseContext
 {
 public:
-	FAnimationCacheBonesContext(UAnimInstance* InInstance)
-		: FAnimationBaseContext(InInstance)
+	DEPRECATED(4.11, "Please use constructor that uses an FAnimInstanceProxy*")
+	FAnimationCacheBonesContext(UAnimInstance* InAnimInstance)
+		: FAnimationBaseContext(InAnimInstance)
+	{
+	}
+
+	FAnimationCacheBonesContext(FAnimInstanceProxy* InAnimInstanceProxy)
+		: FAnimationBaseContext(InAnimInstanceProxy)
 	{
 	}
 };
@@ -73,8 +85,16 @@ private:
 	float CurrentWeight;
 	float DeltaTime;
 public:
-	FAnimationUpdateContext(UAnimInstance* Instance, float InDeltaTime)
-		: FAnimationBaseContext(Instance)
+	DEPRECATED(4.11, "Please use constructor that uses an FAnimInstanceProxy*")
+	FAnimationUpdateContext(UAnimInstance* InAnimInstance, float InDeltaTime)
+		: FAnimationBaseContext(InAnimInstance)
+		, CurrentWeight(1.0f)
+		, DeltaTime(InDeltaTime)
+	{
+	}
+
+	FAnimationUpdateContext(FAnimInstanceProxy* InAnimInstanceProxy, float InDeltaTime)
+		: FAnimationBaseContext(InAnimInstanceProxy)
 		, CurrentWeight(1.0f)
 		, DeltaTime(InDeltaTime)
 	{
@@ -82,14 +102,14 @@ public:
 
 	FAnimationUpdateContext FractionalWeight(float Multiplier) const
 	{
-		FAnimationUpdateContext Result(AnimInstance, DeltaTime);
+		FAnimationUpdateContext Result(AnimInstanceProxy, DeltaTime);
 		Result.CurrentWeight = CurrentWeight * Multiplier;
 		return Result;
 	}
 
 	FAnimationUpdateContext FractionalWeightAndTime(float WeightMultiplier, float TimeMultiplier) const
 	{
-		FAnimationUpdateContext Result(AnimInstance, DeltaTime * TimeMultiplier);
+		FAnimationUpdateContext Result(AnimInstanceProxy, DeltaTime * TimeMultiplier);
 		Result.CurrentWeight = CurrentWeight * WeightMultiplier;
 		return Result;
 	}
@@ -110,27 +130,29 @@ public:
 	FBlendedCurve	Curve;
 
 public:
-	// This constructor allocates a new uninitialized pose for the specified anim instance
+	DEPRECATED(4.11, "Please use constructor that uses an FAnimInstanceProxy*")
 	FPoseContext(UAnimInstance* InAnimInstance)
 		: FAnimationBaseContext(InAnimInstance)
 	{
-		Initialize(InAnimInstance);
+		Initialize(AnimInstanceProxy);
 	}
 
-	// This constructor allocates a new uninitialized pose, copying non-pose state from the source context
-	FPoseContext(const FPoseContext& SourceContext)
-		: FAnimationBaseContext(SourceContext.AnimInstance)
+	// This constructor allocates a new uninitialized pose for the specified anim instance
+	FPoseContext(FAnimInstanceProxy* InAnimInstanceProxy)
+		: FAnimationBaseContext(InAnimInstanceProxy)
 	{
-		Initialize(SourceContext.AnimInstance);
+		Initialize(InAnimInstanceProxy);
 	}
 
-	void Initialize(UAnimInstance* InAnimInstance)
+	// This constructor allocates a new uninitialized pose, copying non-pose state from the source context
+	FPoseContext(const FPoseContext& SourceContext)
+		: FAnimationBaseContext(SourceContext.AnimInstanceProxy)
 	{
-		checkSlow(AnimInstance && AnimInstance->RequiredBones.IsValid());
-		Pose.SetBoneContainer(&AnimInstance->RequiredBones);
-		Curve.InitFrom(AnimInstance->CurrentSkeleton);
+		Initialize(SourceContext.AnimInstanceProxy);
 	}
 
+	ENGINE_API void Initialize(FAnimInstanceProxy* InAnimInstanceProxy);
+
 	void ResetToRefPose()
 	{
 		Pose.ResetToRefPose();	
@@ -153,9 +175,9 @@ public:
 
 	FPoseContext& operator=(const FPoseContext& Other)
 	{
-		if (AnimInstance != Other.AnimInstance)
+		if (AnimInstanceProxy != Other.AnimInstanceProxy)
 		{
-			Initialize(AnimInstance);
+			Initialize(AnimInstanceProxy);
 		}
 
 		Pose = Other.Pose;
@@ -173,29 +195,30 @@ public:
 	FBlendedCurve			Curve;
 
 public:
-	// This constructor allocates a new uninitialized pose for the specified anim instance
+	DEPRECATED(4.11, "Please use constructor that uses an FAnimInstanceProxy*")
 	FComponentSpacePoseContext(UAnimInstance* InAnimInstance)
 		: FAnimationBaseContext(InAnimInstance)
+	{
+	}
+
+	// This constructor allocates a new uninitialized pose for the specified anim instance
+	FComponentSpacePoseContext(FAnimInstanceProxy* InAnimInstanceProxy)
+		: FAnimationBaseContext(InAnimInstanceProxy)
 	{
 		// No need to initialize, done through FA2CSPose::AllocateLocalPoses
 	}
 
 	// This constructor allocates a new uninitialized pose, copying non-pose state from the source context
 	FComponentSpacePoseContext(const FComponentSpacePoseContext& SourceContext)
-		: FAnimationBaseContext(SourceContext.AnimInstance)
+		: FAnimationBaseContext(SourceContext.AnimInstanceProxy)
 	{
 		// No need to initialize, done through FA2CSPose::AllocateLocalPoses
 	}
 
-	void ResetToRefPose()
-	{
-		checkSlow( AnimInstance && AnimInstance->RequiredBones.IsValid() );
-		Pose.InitPose(&AnimInstance->RequiredBones);
-		Curve.InitFrom(AnimInstance->CurrentSkeleton);
-	}
+	ENGINE_API void ResetToRefPose();
 
-	bool ContainsNaN() const;
-	bool IsNormalized() const;
+	ENGINE_API bool ContainsNaN() const;
+	ENGINE_API bool IsNormalized() const;
 };
 
 struct ENGINE_API FNodeDebugData
@@ -374,6 +397,14 @@ public:
 	void EvaluateComponentSpace(FComponentSpacePoseContext& Output);
 };
 
+UENUM()
+enum class EPostCopyOperation : uint8
+{
+	None,
+
+	LogicalNegateBool,
+};
+
 USTRUCT()
 struct FExposedValueCopyRecord
 {
@@ -385,7 +416,10 @@ struct FExposedValueCopyRecord
 	{}
 
 	UPROPERTY()
-	UProperty* SourceProperty;
+	FName SourcePropertyName;
+
+	UPROPERTY()
+	FName SourceSubPropertyName;
 
 	UPROPERTY()
 	int32 SourceArrayIndex;
@@ -399,6 +433,9 @@ struct FExposedValueCopyRecord
 	UPROPERTY()
 	int32 Size;
 
+	UPROPERTY()
+	EPostCopyOperation PostCopyOperation;
+
 	// Cached source copy ptr
 	void* Source;
 
@@ -408,10 +445,17 @@ struct FExposedValueCopyRecord
 
 // An exposed value updater
 USTRUCT()
-struct FExposedValueHandler
+struct ENGINE_API FExposedValueHandler
 {
 	GENERATED_USTRUCT_BODY()
 
+	FExposedValueHandler()
+		: BoundFunction(NAME_None)
+		, Function(nullptr)
+		, bInitialized(false)
+	{
+	}
+
 	// The function to call to update associated properties (can be NULL)
 	UPROPERTY()
 	FName BoundFunction;
@@ -420,64 +464,17 @@ struct FExposedValueHandler
 	UPROPERTY()
 	TArray<FExposedValueCopyRecord> CopyRecords;
 
-	// funtion pointer if BoundFunction != NAME_None
+	// function pointer if BoundFunction != NAME_None
 	UFunction* Function;
 
-	void Initialize(FAnimNode_Base* AnimNode, UAnimInstance* AnimInstance) 
-	{
-		if (BoundFunction != NAME_None)
-		{
-			Function = AnimInstance->FindFunction(BoundFunction);
-		}
-		else
-		{
-			Function = NULL;
-		}
-
-		// initialize copy records
-		for(auto& CopyRecord : CopyRecords)
-		{
-			if(UArrayProperty* SourceArrayProperty = Cast<UArrayProperty>(CopyRecord.SourceProperty))
-			{
-				FScriptArrayHelper ArrayHelper(SourceArrayProperty, CopyRecord.SourceProperty->ContainerPtrToValuePtr<uint8>(AnimInstance));
-				check(ArrayHelper.IsValidIndex(CopyRecord.SourceArrayIndex));
-				CopyRecord.Source = ArrayHelper.GetRawPtr(CopyRecord.SourceArrayIndex);
-			}
-			else
-			{
-				CopyRecord.Source = CopyRecord.SourceProperty->ContainerPtrToValuePtr<uint8>(AnimInstance, CopyRecord.SourceArrayIndex);
-			}
-
-			if(UArrayProperty* DestArrayProperty = Cast<UArrayProperty>(CopyRecord.DestProperty))
-			{
-				FScriptArrayHelper ArrayHelper(DestArrayProperty, CopyRecord.DestProperty->ContainerPtrToValuePtr<uint8>(AnimNode));
-				check(ArrayHelper.IsValidIndex(CopyRecord.DestArrayIndex));
-				CopyRecord.Dest = ArrayHelper.GetRawPtr(CopyRecord.DestArrayIndex);
-			}
-			else
-			{
-				CopyRecord.Dest = CopyRecord.DestProperty->ContainerPtrToValuePtr<uint8>(AnimNode, CopyRecord.DestArrayIndex);
-			}
-		}
-	}
+	// Prevent multiple initialization
+	bool bInitialized;
 
-	void Execute(const FAnimationBaseContext& Context) const
-	{
-		if (Function != nullptr)
-		{
-			Context.AnimInstance->ProcessEvent(Function, NULL);
-		}
+	// Bind copy records and cache UFunction if necessary
+	void Initialize(FAnimNode_Base* AnimNode, UObject* AnimInstanceObject);
 
-		for(const auto& CopyRecord : CopyRecords)
-		{
-			// if any of these checks fail then it's likely that Initialize has not been called.
-			// has new anim node type been added that doesnt call the base class Initialize()?
-			checkSlow(CopyRecord.Dest != nullptr);
-			checkSlow(CopyRecord.Source != nullptr);
-			checkSlow(CopyRecord.Size != 0);
-			FMemory::Memcpy(CopyRecord.Dest, CopyRecord.Source, CopyRecord.Size);
-		}
-	}
+	// Execute the function and copy records
+	void Execute(const FAnimationBaseContext& Context) const;
 };
 
 /**
@@ -519,9 +516,6 @@ struct ENGINE_API FAnimNode_Base
 	virtual ~FAnimNode_Base() {}
 
 protected:
-	/** return true if eanbled, otherwise, return false. This is utility function that can be used per node level */
-	bool IsLODEnabled(UAnimInstance* AnimInstace, int32 InLODThreshold)
-	{
-		return (InLODThreshold == INDEX_NONE || AnimInstace->GetOwningComponent()->PredictedLODLevel <= InLODThreshold);
-	}
+	/** return true if enabled, otherwise, return false. This is utility function that can be used per node level */
+	bool IsLODEnabled(FAnimInstanceProxy* AnimInstanceProxy, int32 InLODThreshold);
 };
diff --git a/Engine/Source/Runtime/Engine/Classes/Animation/AnimNode_AssetPlayerBase.cpp b/Engine/Source/Runtime/Engine/Classes/Animation/AnimNode_AssetPlayerBase.cpp
--- a/Engine/Source/Runtime/Engine/Classes/Animation/AnimNode_AssetPlayerBase.cpp
+++ b/Engine/Source/Runtime/Engine/Classes/Animation/AnimNode_AssetPlayerBase.cpp
@@ -2,6 +2,7 @@
 
 #include "EnginePrivate.h"
 #include "AnimNode_AssetPlayerBase.h"
+#include "Animation/AnimInstanceProxy.h"
 
 FAnimNode_AssetPlayerBase::FAnimNode_AssetPlayerBase()
 	: bIgnoreForRelevancyTest(false)
@@ -37,9 +38,9 @@ void FAnimNode_AssetPlayerBase::CreateTickRecordForNode(const FAnimationUpdateCo
 	FAnimGroupInstance* SyncGroup;
 	const int32 GroupIndexToUse = ((GroupRole != EAnimGroupRole::TransitionLeader) || bHasBeenFullWeight) ? GroupIndex : INDEX_NONE;
 
-	FAnimTickRecord& TickRecord = Context.AnimInstance->CreateUninitializedTickRecord(GroupIndexToUse, /*out*/ SyncGroup);
+	FAnimTickRecord& TickRecord = Context.AnimInstanceProxy->CreateUninitializedTickRecord(GroupIndexToUse, /*out*/ SyncGroup);
 
-	Context.AnimInstance->MakeSequenceTickRecord(TickRecord, Sequence, bLooping, PlayRate, FinalBlendWeight, /*inout*/ InternalTimeAccumulator, MarkerTickRecord);
+	Context.AnimInstanceProxy->MakeSequenceTickRecord(TickRecord, Sequence, bLooping, PlayRate, FinalBlendWeight, /*inout*/ InternalTimeAccumulator, MarkerTickRecord);
 
 	// Update the sync group if it exists
 	if (SyncGroup != NULL)
diff --git a/Engine/Source/Runtime/Engine/Classes/Animation/AnimNode_StateMachine.h b/Engine/Source/Runtime/Engine/Classes/Animation/AnimNode_StateMachine.h
--- a/Engine/Source/Runtime/Engine/Classes/Animation/AnimNode_StateMachine.h
+++ b/Engine/Source/Runtime/Engine/Classes/Animation/AnimNode_StateMachine.h
@@ -138,7 +138,7 @@ public:
 
 protected:
 	// The state machine description this is an instance of
-	FBakedAnimationStateMachine* PRIVATE_MachineDescription;
+	const FBakedAnimationStateMachine* PRIVATE_MachineDescription;
 
 	// The current state within the state machine
 	UPROPERTY()
@@ -199,7 +199,7 @@ public:
 
 protected:
 	// Tries to get the instance information for the state machine
-	FBakedAnimationStateMachine* GetMachineDescription();
+	const FBakedAnimationStateMachine* GetMachineDescription() const;
 
 	void SetState(const FAnimationBaseContext& Context, int32 NewStateIndex);
 	void SetStateInternal(int32 NewStateIndex);
@@ -232,5 +232,5 @@ protected:
 	FAnimNode_AssetPlayerBase* GetRelevantAssetPlayerFromState(const FAnimationUpdateContext& Context, const FBakedAnimationState& StateInfo);
 
 public:
-	friend class UAnimInstance;
+	friend struct FAnimInstanceProxy;
 };
diff --git a/Engine/Source/Runtime/Engine/Classes/Animation/AnimSequence.h b/Engine/Source/Runtime/Engine/Classes/Animation/AnimSequence.h
--- a/Engine/Source/Runtime/Engine/Classes/Animation/AnimSequence.h
+++ b/Engine/Source/Runtime/Engine/Classes/Animation/AnimSequence.h
@@ -560,7 +560,7 @@ public:
 	//~ End UAnimationAsset Interface
 
 	//~ Begin UAnimSequenceBase Interface
-	virtual void OnAssetPlayerTickedInternal(FAnimAssetTickContext &Context, const float PreviousTime, const float MoveDelta, const FAnimTickRecord &Instance, class UAnimInstance* InstanceOwner) const override;
+	virtual void HandleAssetPlayerTickedInternal(FAnimAssetTickContext &Context, const float PreviousTime, const float MoveDelta, const FAnimTickRecord &Instance, struct FAnimNotifyQueue& NotifyQueue) const override;
 	virtual bool HasRootMotion() const override { return bEnableRootMotion; }
 	virtual void RefreshCacheData() override;
 	//~ End UAnimSequenceBase Interface
diff --git a/Engine/Source/Runtime/Engine/Classes/Animation/AnimSequenceBase.h b/Engine/Source/Runtime/Engine/Classes/Animation/AnimSequenceBase.h
--- a/Engine/Source/Runtime/Engine/Classes/Animation/AnimSequenceBase.h
+++ b/Engine/Source/Runtime/Engine/Classes/Animation/AnimSequenceBase.h
@@ -121,7 +121,7 @@ class UAnimSequenceBase : public UAnimationAsset
 	ENGINE_API virtual void RefreshCacheData();
 
 	//~ Begin UAnimationAsset Interface
-	virtual void TickAssetPlayerInstance(FAnimTickRecord& Instance, class UAnimInstance* InstanceOwner, FAnimAssetTickContext& Context) const override;
+	virtual void TickAssetPlayer(FAnimTickRecord& Instance, struct FAnimNotifyQueue& NotifyQueue, FAnimAssetTickContext& Context) const override;
 
 	void TickByMarkerAsFollower(FMarkerTickRecord &Instance, FMarkerTickContext &MarkerContext, float& CurrentTime, float& OutPreviousTime, const float MoveDelta, const bool bLooping) const;
 
@@ -143,7 +143,10 @@ class UAnimSequenceBase : public UAnimationAsset
 	*/
 	ENGINE_API virtual void GetAnimationPose(struct FCompactPose& OutPose, FBlendedCurve& OutCurve, const FAnimExtractContext& ExtractionContext) const PURE_VIRTUAL(UAnimSequenceBase::GetAnimationPose, );
 	
-	virtual void OnAssetPlayerTickedInternal(FAnimAssetTickContext &Context, const float PreviousTime, const float MoveDelta, const FAnimTickRecord &Instance, class UAnimInstance* InstanceOwner) const;
+	DEPRECATED(4.11, "This function is deprecated, please use HandleAssetPlayerTickedInternal")
+	virtual void OnAssetPlayerTickedInternal(FAnimAssetTickContext &Context, const float PreviousTime, const float MoveDelta, const FAnimTickRecord &Instance, class UAnimInstance* InAnimInstance) const;
+
+	virtual void HandleAssetPlayerTickedInternal(FAnimAssetTickContext &Context, const float PreviousTime, const float MoveDelta, const FAnimTickRecord &Instance, struct FAnimNotifyQueue& NotifyQueue) const;
 
 	virtual bool HasRootMotion() const { return false; }
 
diff --git a/Engine/Source/Runtime/Engine/Classes/Animation/AnimSingleNodeInstance.h b/Engine/Source/Runtime/Engine/Classes/Animation/AnimSingleNodeInstance.h
--- a/Engine/Source/Runtime/Engine/Classes/Animation/AnimSingleNodeInstance.h
+++ b/Engine/Source/Runtime/Engine/Classes/Animation/AnimSingleNodeInstance.h
@@ -17,19 +17,24 @@ class ENGINE_API UAnimSingleNodeInstance : public UAnimInstance
 {
 	GENERATED_UCLASS_BODY()
 
+	// Disable compiler-generated deprecation warnings by implementing our own destructor
+	PRAGMA_DISABLE_DEPRECATION_WARNINGS
+	~UAnimSingleNodeInstance() {}
+	PRAGMA_ENABLE_DEPRECATION_WARNINGS
+
 	/** Current Asset being played **/
-	UPROPERTY(transient)
+	DEPRECATED(4.11, "Please use FAnimSingleNodeInstanceProxy::CurrentAsset")
 	class UAnimationAsset* CurrentAsset;
 
-	UPROPERTY(transient)
+	DEPRECATED(4.11, "Please use FAnimSingleNodeInstanceProxy::CurrentVertexAnim")
 	class UVertexAnimation* CurrentVertexAnim;
 
 	/** Random cached values to play each asset **/
-	UPROPERTY(transient)
+	DEPRECATED(4.11, "Please use FAnimSingleNodeInstanceProxy::BlendSpaceInput")
 	FVector BlendSpaceInput;
 
 	/** Random cached values to play each asset **/
-	UPROPERTY(transient)
+	DEPRECATED(4.11, "Please use FAnimSingleNodeInstanceProxy::BlendSampleData")
 	TArray<FBlendSampleData> BlendSampleData;
 
 	/** Random cached values to play each asset **/
@@ -37,36 +42,35 @@ class ENGINE_API UAnimSingleNodeInstance : public UAnimInstance
 	FBlendFilter BlendFilter;
 
 	/** Shared parameters for previewing blendspace or animsequence **/
-	UPROPERTY(transient)
+	DEPRECATED(4.11, "Please use FAnimSingleNodeInstanceProxy::CurrentTime")
 	float CurrentTime;
 
 	/** Cache for data needed during marker sync */
+	DEPRECATED(4.11, "Please use FAnimSingleNodeInstanceProxy::GetMarkerTickRecord")
 	FMarkerTickRecord MarkerTickRecord;
 
-	UPROPERTY(transient)
+	DEPRECATED(4.11, "Please use FAnimSingleNodeInstanceProxy::PlayRate")
 	float PlayRate;
 	 
 	UPROPERTY(Transient)
 	FPostEvaluateAnimEvent PostEvaluateAnimEvent;
 
-	UPROPERTY(transient)
+	DEPRECATED(4.11, "Please use FAnimSingleNodeInstanceProxy::bLooping")
 	uint32 bLooping:1;
 
-	UPROPERTY(transient)
+	DEPRECATED(4.11, "Please use FAnimSingleNodeInstanceProxy::bPlaying")
 	uint32 bPlaying:1;
 
-	UPROPERTY(transient)
+	DEPRECATED(4.11, "Please use FAnimSingleNodeInstanceProxy::bReverse")
 	uint32 bReverse:1;
 
 	//~ Begin UAnimInstance Interface
 	virtual void NativeInitializeAnimation() override;
-	virtual bool NativeEvaluateAnimation(FPoseContext& Output) override;
 	virtual void NativePostEvaluateAnimation() override;
 	virtual void OnMontageInstanceStopped(FAnimMontageInstance& StoppedMontageInstance) override;
 
 protected:
 	virtual void Montage_Advance(float DeltaTime) override;
-	void InternalBlendSpaceEvaluatePose(class UBlendSpaceBase* BlendSpace, TArray<FBlendSampleData>& BlendSampleDataCache, FPoseContext& OutContext);
 	//~ End UAnimInstance Interface
 public:
 
@@ -111,16 +115,29 @@ public:
 	/** Updates the blendspace samples list in the case of our asset being a blendspace */
 	void UpdateBlendspaceSamples(FVector InBlendInput);
 
-#if WITH_EDITORONLY_DATA
-	float PreviewPoseCurrentTime;
-#endif
-private:
-	/** Slot node weight transient data */
-	FSlotNodeWeightInfo WeightInfo;
+	/** Check whether we are currently playing */
+	bool IsPlaying() const;
+
+	/** Check whether we are currently playing in reverse */
+	bool IsReverse() const;
+
+	/** Check whether we are currently looping */
+	bool IsLooping() const;
+
+	/** Get the current playback time */
+	float GetCurrentTime() const;
+
+	/** Get the current play rate multiplier */
+	float GetPlayRate() const;
+
+	/** Get the currently playing asset. Can return NULL */
+	UAnimationAsset* GetCurrentAsset();
+
+	/** Get the currently playing vertex animation. Can return NULL */
+	UVertexAnimation* GetCurrentVertexAnimation();
 
 protected:
-	// UAnimInstance interface to override animation node part
-	virtual void UpdateAnimationNode(float DeltaSeconds) override;
+	virtual FAnimInstanceProxy* CreateAnimInstanceProxy() override;
 };
 
 
diff --git a/Engine/Source/Runtime/Engine/Classes/Animation/AnimationAsset.h b/Engine/Source/Runtime/Engine/Classes/Animation/AnimationAsset.h
--- a/Engine/Source/Runtime/Engine/Classes/Animation/AnimationAsset.h
+++ b/Engine/Source/Runtime/Engine/Classes/Animation/AnimationAsset.h
@@ -694,7 +694,17 @@ public:
 	 * @param InstanceOwner	AnimInstance playing this asset
 	 * @param Context		The tick context (leader/follower, delta time, sync point, etc...)
 	 */
-	virtual void TickAssetPlayerInstance(FAnimTickRecord& Instance, class UAnimInstance* InstanceOwner, FAnimAssetTickContext& Context) const {}
+	DEPRECATED(4.11, "This function is deprecated, use TickAssetPlayer")
+	ENGINE_API virtual void TickAssetPlayerInstance(FAnimTickRecord& Instance, class UAnimInstance* AnimInstance, FAnimAssetTickContext& Context) const;
+
+	/** Advances the asset player instance 
+	 * 
+	 * @param Instance		AnimationTickRecord Instance - saves data to evaluate
+	 * @param NotifyQueue	Queue for any notifies we create
+	 * @param Context		The tick context (leader/follower, delta time, sync point, etc...)
+	 */
+	virtual void TickAssetPlayer(FAnimTickRecord& Instance, struct FAnimNotifyQueue& NotifyQueue, FAnimAssetTickContext& Context) const {}
+
 	// this is used in editor only when used for transition getter
 	// this doesn't mean max time. In Sequence, this is SequenceLength,
 	// but for BlendSpace CurrentTime is normalized [0,1], so this is 1
diff --git a/Engine/Source/Runtime/Engine/Classes/Animation/BlendSpaceBase.h b/Engine/Source/Runtime/Engine/Classes/Animation/BlendSpaceBase.h
--- a/Engine/Source/Runtime/Engine/Classes/Animation/BlendSpaceBase.h
+++ b/Engine/Source/Runtime/Engine/Classes/Animation/BlendSpaceBase.h
@@ -264,7 +264,7 @@ public:
 	//~ End UObject Interface
 
 	//~ Begin UAnimationAsset Interface
-	virtual void TickAssetPlayerInstance(FAnimTickRecord& Instance, class UAnimInstance* InstanceOwner, FAnimAssetTickContext& Context) const override;
+	virtual void TickAssetPlayer(FAnimTickRecord& Instance, struct FAnimNotifyQueue& NotifyQueue, FAnimAssetTickContext& Context) const override;
 	// this is used in editor only when used for transition getter
 	// this doesn't mean max time. In Sequence, this is SequenceLength,
 	// but for BlendSpace CurrentTime is normalized [0,1], so this is 1
diff --git a/Engine/Source/Runtime/Engine/Classes/Camera/CameraComponent.h b/Engine/Source/Runtime/Engine/Classes/Camera/CameraComponent.h
--- a/Engine/Source/Runtime/Engine/Classes/Camera/CameraComponent.h
+++ b/Engine/Source/Runtime/Engine/Classes/Camera/CameraComponent.h
@@ -81,15 +81,17 @@ class ENGINE_API UCameraComponent : public USceneComponent
 
 	// UActorComponent interface
 	virtual void OnRegister() override;
-	virtual void OnUnregister() override;
 	virtual void PostLoad() override;
 #if WITH_EDITOR
 	virtual void OnComponentDestroyed() override;
 	virtual void CheckForErrors() override;
+#endif
 	// End of UActorComponent interface
 
 	// UObject interface
+#if WITH_EDITORONLY_DATA
 	virtual void PostEditChangeProperty(FPropertyChangedEvent& PropertyChangedEvent) override;
+	static void AddReferencedObjects(UObject* InThis, FReferenceCollector& Collector);
 #endif
 	virtual void Serialize(FArchive& Ar) override;
 	// End of UObject interface
@@ -108,14 +110,12 @@ class ENGINE_API UCameraComponent : public USceneComponent
 protected:
 #if WITH_EDITORONLY_DATA
 	// The frustum component used to show visually where the camera field of view is
-	UPROPERTY(transient)
 	class UDrawFrustumComponent* DrawFrustum;
 
 	UPROPERTY(transient)
 	class UStaticMesh* CameraMesh;
 	
 	// The camera mesh to show visually where the camera is placed
-	UPROPERTY(transient)
 	class UStaticMeshComponent* ProxyMeshComponent;
 
 public:
diff --git a/Engine/Source/Runtime/Engine/Classes/Components/AudioComponent.h b/Engine/Source/Runtime/Engine/Classes/Components/AudioComponent.h
--- a/Engine/Source/Runtime/Engine/Classes/Components/AudioComponent.h
+++ b/Engine/Source/Runtime/Engine/Classes/Components/AudioComponent.h
@@ -101,6 +101,9 @@ class ENGINE_API UAudioComponent : public USceneComponent
 	UPROPERTY(EditAnywhere, BlueprintReadWrite, Category = Sound)
 	uint32 bIsUISound : 1;
 
+	UPROPERTY(EditAnywhere, BlueprintReadWrite, Category = Sound)
+	uint32 bOverridePriority : 1;
+
 	/** The lower bound to use when randomly determining a pitch multiplier */
 	UPROPERTY(EditAnywhere, BlueprintReadWrite, Category=Modulation)
 	float PitchModulationMin;
@@ -121,9 +124,12 @@ class ENGINE_API UAudioComponent : public USceneComponent
 	UPROPERTY(EditAnywhere, BlueprintReadWrite, Category=Sound)
 	float VolumeMultiplier;
 
-	/** A priority value that scales against final ouput volume and is used to determine a sound's play priority. */
-	UPROPERTY(EditAnywhere, BlueprintReadWrite, Category = Sound)
-	float VolumeWeightedPriorityScale;
+	/** A priority value that is used for sounds that play on this component that scales against final output volume. */
+	UPROPERTY(EditAnywhere, BlueprintReadWrite, Category = Sound, meta = (ClampMin = "0.0", UIMin = "0.0", EditCondition = "bOverridePriority"))
+	float Priority;
+
+	UPROPERTY()
+	float VolumeWeightedPriorityScale_DEPRECATED;
 
 	/** A pitch multiplier to apply to sounds generated by this component */
 	UPROPERTY(EditAnywhere, BlueprintReadWrite, Category=Sound)
diff --git a/Engine/Source/Runtime/Engine/Classes/Components/DirectionalLightComponent.h b/Engine/Source/Runtime/Engine/Classes/Components/DirectionalLightComponent.h
--- a/Engine/Source/Runtime/Engine/Classes/Components/DirectionalLightComponent.h
+++ b/Engine/Source/Runtime/Engine/Classes/Components/DirectionalLightComponent.h
@@ -168,7 +168,7 @@ class UDirectionalLightComponent : public ULightComponent
 	UFUNCTION(BlueprintCallable, Category="Rendering|Lighting")
 	void SetLightShaftOverrideDirection(FVector NewValue);
 
-	//~ Begin ULightComponent Interface.
+	//~ Begin ULightComponent Interface
 	virtual FVector4 GetLightPosition() const override;
 	virtual ELightComponentType GetLightType() const override;
 	virtual FLightmassLightSettings GetLightmassSettings() const override
@@ -183,6 +183,7 @@ class UDirectionalLightComponent : public ULightComponent
 	{
 		return bUsedAsAtmosphereSunLight;
 	}
+	//~ End ULightComponent Interface
 
 	//~ Begin UObject Interface
 #if WITH_EDITOR
diff --git a/Engine/Source/Runtime/Engine/Classes/Components/InputComponent.h b/Engine/Source/Runtime/Engine/Classes/Components/InputComponent.h
--- a/Engine/Source/Runtime/Engine/Classes/Components/InputComponent.h
+++ b/Engine/Source/Runtime/Engine/Classes/Components/InputComponent.h
@@ -669,6 +669,18 @@ public:
 		return AddActionBinding(AB);
 	}
 
+	/**
+	 * Binds a delegate function to an Action defined in the project settings.
+	 * Returned reference is only guaranteed to be valid until another action is bound.
+	 */
+	template<class UserClass>
+	FInputActionBinding& BindAction( const FName ActionName, const EInputEvent KeyEvent, UserClass* Object, typename FInputActionHandlerWithKeySignature::TUObjectMethodDelegate< UserClass >::FMethodPtr Func )
+	{
+		FInputActionBinding AB( ActionName, KeyEvent );
+		AB.ActionDelegate.BindDelegate(Object, Func);
+		return AddActionBinding(AB);
+	}
+
 	/**
 	 * Binds a delegate function an Axis defined in the project settings.
 	 * Returned reference is only guaranteed to be valid until another axis is bound.
diff --git a/Engine/Source/Runtime/Engine/Classes/Components/PrimitiveComponent.h b/Engine/Source/Runtime/Engine/Classes/Components/PrimitiveComponent.h
--- a/Engine/Source/Runtime/Engine/Classes/Components/PrimitiveComponent.h
+++ b/Engine/Source/Runtime/Engine/Classes/Components/PrimitiveComponent.h
@@ -720,6 +720,10 @@ public:
 	UPROPERTY(BlueprintAssignable, Category="Input|Touch Input")
 	FComponentEndTouchOverSignature OnInputTouchLeave;
 
+	/** Scale the bounds of this object, used for frustum culling. Useful for features like WorldPositionOffset. */
+	UFUNCTION(BlueprintCallable, Category = "Rendering")
+	void SetBoundsScale(float NewBoundsScale=1.f);
+
 	/**
 	 * Returns the material used by the element at the specified index
 	 * @param ElementIndex - The element to access the material of.
diff --git a/Engine/Source/Runtime/Engine/Classes/Components/SceneCaptureComponentCube.h b/Engine/Source/Runtime/Engine/Classes/Components/SceneCaptureComponentCube.h
--- a/Engine/Source/Runtime/Engine/Classes/Components/SceneCaptureComponentCube.h
+++ b/Engine/Source/Runtime/Engine/Classes/Components/SceneCaptureComponentCube.h
@@ -31,7 +31,7 @@ public:
 
 	/** Render the scene to the texture */
 	UFUNCTION(BlueprintCallable,Category = "Rendering|SceneCapture")
-	void UpdateContent();
+	ENGINE_API void UpdateContent();
 
 	ENGINE_API static void UpdateDeferredCaptures(FSceneInterface* Scene);
 };
\ No newline at end of file
diff --git a/Engine/Source/Runtime/Engine/Classes/Components/SkinnedMeshComponent.h b/Engine/Source/Runtime/Engine/Classes/Components/SkinnedMeshComponent.h
--- a/Engine/Source/Runtime/Engine/Classes/Components/SkinnedMeshComponent.h
+++ b/Engine/Source/Runtime/Engine/Classes/Components/SkinnedMeshComponent.h
@@ -634,11 +634,6 @@ protected:
 	/** Track whether we still need to flip to recently modified buffer */
 	bool bNeedToFlipSpaceBaseBuffers;
 
-	/**
-	* Combine CurveKeys (that reference morph targets by name) and ActiveAnims (that reference vertex anims by reference) into the ActiveVertexAnims array.
-	*/
-	static TArray<struct FActiveVertexAnim> UpdateActiveVertexAnims(const USkeletalMesh* InSkeletalMesh, const TMap<FName, float>& InMorphCurveAnims, const TArray<FActiveVertexAnim>& InActiveAnims);
-
 	/** 
 	 * Should update transform in Tick
 	 * 
diff --git a/Engine/Source/Runtime/Engine/Classes/Debug/DebugDrawService.h b/Engine/Source/Runtime/Engine/Classes/Debug/DebugDrawService.h
--- a/Engine/Source/Runtime/Engine/Classes/Debug/DebugDrawService.h
+++ b/Engine/Source/Runtime/Engine/Classes/Debug/DebugDrawService.h
@@ -7,58 +7,6 @@
 
 class FSceneView;
 
-// BEGIN - temporary solution for Orion, to draw all perception data correctly in client-server game. Changes to GameplayDebugger are coming on Main branch soon so it won't be merged back to Main.
-UENUM()
-enum class EDrawDebugShapeElement : uint8
-{
-	Invalid = 0,
-	SinglePoint, // individual points. 
-	Segment, // pairs of points 
-	Box,
-	Cone,
-	Cylinder,
-	Capsule,
-	Polygon,
-	String,
-};
-
-USTRUCT()
-struct ENGINE_API FDrawDebugShapeElement
-{
-	GENERATED_USTRUCT_BODY()
-
-	UPROPERTY()
-	FString Description;
-
-	UPROPERTY()
-	TArray<FVector> Points;
-
-	UPROPERTY()
-	FMatrix TransformationMatrix;
-
-	UPROPERTY()
-	EDrawDebugShapeElement Type;
-
-	UPROPERTY()
-	uint8 Color;
-
-	UPROPERTY()
-	uint16 ThicknesOrRadius;
-
-	FDrawDebugShapeElement(EDrawDebugShapeElement InType = EDrawDebugShapeElement::Invalid);
-	FDrawDebugShapeElement(const FString& InDescription, const FColor& InColor, uint16 InThickness);
-	void SetColor(const FColor& InColor);
-	EDrawDebugShapeElement GetType() const;
-	void SetType(EDrawDebugShapeElement InType);
-	FColor GetFColor() const;
-
-	static FDrawDebugShapeElement MakeString(const FString& Description, const FVector& Location);
-	static FDrawDebugShapeElement MakeLine(const FVector& Start, const FVector& End, FColor Color, float Thickness = 0);
-	static FDrawDebugShapeElement MakePoint(const FVector& Location, FColor Color, float Radius = 2);
-	static FDrawDebugShapeElement MakeCylinder(FVector const& Start, FVector const& End, float Radius, int32 Segments, FColor const& Color);
-};
-// END - temporary solution for Orion
-
 /** 
  * 
  */
