Copying //UE4/Dev-Framework to Dev-Main (//UE4/Dev-Main)

diff --git a/Engine/Source/Runtime/Engine/Classes/Debug/GameplayDebuggerBaseObject.h b/Engine/Source/Runtime/Engine/Classes/Debug/GameplayDebuggerBaseObject.h
new file mode 100644
--- /dev/null
+++ b/Engine/Source/Runtime/Engine/Classes/Debug/GameplayDebuggerBaseObject.h
@@ -0,0 +1,228 @@
+// Copyright 1998-2015 Epic Games, Inc. All Rights Reserved.
+
+#pragma once
+#include "Engine.h"
+#include "Net/UnrealNetwork.h"
+#include "Net/RepLayout.h"
+#include "Net/DataReplication.h"
+#include "Engine/ActorChannel.h"
+#include "GameplayDebuggerBaseObject.generated.h"
+
+#define ENABLED_GAMEPLAY_DEBUGGER (1 && !(UE_BUILD_SHIPPING || UE_BUILD_TEST))
+
+class AActorChannel;
+class UFont;
+class FCanvasItem;
+struct FCanvasIcon;
+class UCanvas;
+class UFont;
+
+namespace DebugTools
+{
+	static const FVector InvalidLocation = FVector(FLT_MAX);
+}
+
+ENGINE_API DECLARE_LOG_CATEGORY_EXTERN(LogGameplayDebugger, Display, All);
+
+UENUM()
+enum class EGameplayDebuggerShapeElement : uint8
+{
+	Invalid = 0,
+	SinglePoint, // individual points. 
+	Segment, // pairs of points 
+	Box,
+	Cone,
+	Cylinder,
+	Capsule,
+	Polygon,
+	String,
+};
+
+enum class EGameplayDebuggerReplicate : uint8
+{
+	WithoutCompression = 0,
+	WithCompressed 
+};
+
+
+USTRUCT()
+struct ENGINE_API FGameplayDebuggerShapeElement
+{
+	GENERATED_USTRUCT_BODY()
+
+	UPROPERTY()
+	FString Description;
+
+	UPROPERTY()
+	TArray<FVector> Points;
+
+	UPROPERTY()
+	FMatrix TransformationMatrix;
+
+	UPROPERTY()
+	TEnumAsByte<EGameplayDebuggerShapeElement> Type;
+
+	UPROPERTY()
+	uint8 Color;
+
+	UPROPERTY()
+	uint16 ThicknesOrRadius;
+
+	FGameplayDebuggerShapeElement(EGameplayDebuggerShapeElement InType = EGameplayDebuggerShapeElement::Invalid);
+	FGameplayDebuggerShapeElement(const FString& InDescription, const FColor& InColor, uint16 InThickness);
+	void SetColor(const FColor& InColor);
+	EGameplayDebuggerShapeElement GetType() const;
+	void SetType(EGameplayDebuggerShapeElement InType);
+	FColor GetFColor() const;
+};
+
+UCLASS(Transient, Abstract)
+class ENGINE_API UGameplayDebuggerHelper : public UObject
+{
+	GENERATED_UCLASS_BODY()
+public:
+	struct FPrintContext
+	{
+	public:
+		TWeakObjectPtr<UCanvas> Canvas;
+		TWeakObjectPtr<UFont> Font;
+		float CursorX, CursorY;
+		float DefaultX, DefaultY;
+		FFontRenderInfo FontRenderInfo;
+	public:
+		FPrintContext() {}
+		FPrintContext(class UFont* InFont, class UCanvas* InCanvas, float InX, float InY) : Canvas(InCanvas), Font(InFont), CursorX(InX), CursorY(InY), DefaultX(InX), DefaultY(InY) {}
+	};
+
+	// Begin - helper for gameplay debugger API
+	static void SetOverHeadContext(const FPrintContext& InContext) { OverHeadContext = InContext; }
+	static void SetDefaultContext(const FPrintContext& InContext) { DefaultContext = InContext; }
+	static FPrintContext& GetOverHeadContext() { return OverHeadContext; }
+	static FPrintContext& GetDefaultContext() { return DefaultContext; }
+
+	UFUNCTION(BlueprintCallable, Category = GameplayDebugger)
+	void PrintDebugString(const FString& InString);
+
+	static void PrintString(UGameplayDebuggerHelper::FPrintContext& Context, const FString& InString);
+	static void PrintString(UGameplayDebuggerHelper::FPrintContext& Context, const FString& InString, float X, float Y);
+	static void PrintString(UGameplayDebuggerHelper::FPrintContext& Context, const FColor& InColor, const FString& InString);
+	static void PrintString(UGameplayDebuggerHelper::FPrintContext& Context, const FColor& InColor, const FString& InString, float X, float Y);
+	static void CalulateStringSize(const UGameplayDebuggerHelper::FPrintContext& Context, UFont* Font, const FString& InString, float& OutX, float& OutY);
+	static void CalulateTextSize(const UGameplayDebuggerHelper::FPrintContext& Context, UFont* Font, const FText& InText, float& OutX, float& OutY);
+	static FVector ProjectLocation(const UGameplayDebuggerHelper::FPrintContext& Context, const FVector& Location);
+	static void DrawItem(const UGameplayDebuggerHelper::FPrintContext& Context, class FCanvasItem& Item, float X, float Y);
+	static void DrawIcon(const UGameplayDebuggerHelper::FPrintContext& Context, const FColor& InColor, const FCanvasIcon& Icon, float X, float Y, float Scale = 0.f);
+
+	static FGameplayDebuggerShapeElement MakeString(const FString& Description, const FVector& Location);
+	static FGameplayDebuggerShapeElement MakeLine(const FVector& Start, const FVector& End, FColor Color, float Thickness = 0);
+	static FGameplayDebuggerShapeElement MakePoint(const FVector& Location, FColor Color, float Radius = 2);
+	static FGameplayDebuggerShapeElement MakeCylinder(FVector const& Start, FVector const& End, float Radius, int32 Segments, FColor const& Color);
+	// End - helper for gameplay debugger API
+
+protected:
+	static FPrintContext OverHeadContext;
+	static FPrintContext DefaultContext;
+};
+
+ENGINE_API FArchive& operator<<(FArchive& Ar, FGameplayDebuggerShapeElement& ShapeElement);
+
+USTRUCT()
+struct FGameplayDebuggerReplicatedBlob
+{
+	GENERATED_USTRUCT_BODY()
+
+	FGameplayDebuggerReplicatedBlob() {}
+	FGameplayDebuggerReplicatedBlob(const TArray<uint8>& InRepData, int32 InVersionNum, int32 InRepDataOffset, int32 InRepDataSize) 
+		: RepData(InRepData), VersionNum(InVersionNum), RepDataOffset(InRepDataOffset), RepDataSize(InRepDataSize) {}
+
+	UPROPERTY()
+	TArray<uint8> RepData;
+
+	UPROPERTY()
+	int32 VersionNum;
+	
+	UPROPERTY()
+	int32 RepDataOffset;
+	
+	UPROPERTY()
+	int32 RepDataSize;
+};
+
+UCLASS(Transient, Abstract)
+class ENGINE_API UGameplayDebuggerBaseObject : public UGameplayDebuggerHelper
+{
+	GENERATED_UCLASS_BODY()
+
+public:
+	// Begin UObject interface
+	virtual bool IsSupportedForNetworking() const override { return true; }
+	int32 GetFunctionCallspace(UFunction* Function, void* Parameters, FFrame* Stack) override; //override needed to have Client/Server function for UObject derived class
+	bool CallRemoteFunction(UFunction* Function, void* Parameters, FOutParmRec* OutParms, FFrame* Stack) override;  //override needed to have Client/Server function for UObject derived class
+	virtual void PostLoad() override;
+	virtual void BeginDestroy() override;
+	virtual UWorld* GetWorld() const override
+	{
+		if (HasAllFlags(RF_ClassDefaultObject))
+		{
+			// If we are a CDO, we must return nullptr instead of calling Outer->GetWorld() to fool UObject::ImplementsGetWorld.
+			return nullptr;
+		}
+		return GetOuter()->GetWorld();
+	}
+	// End UObject interface
+
+	UFUNCTION(Reliable, Server, WithValidation)
+	void CollectDataToReplicateOnServer(APlayerController* MyPC, AActor *SelectedActor);
+
+	/** Set custom, additional keyboard bindings */
+	virtual void BindKeyboardImput(class UInputComponent*& InputComponent);
+
+	virtual void DrawCollectedData(APlayerController* MyPC, AActor *SelectedActor = nullptr);
+
+	virtual FString GetCategoryName() const { return TEXT("Basic"); };
+
+	virtual class FDebugRenderSceneProxy* CreateSceneProxy(const UPrimitiveComponent* InComponent, UWorld* World, AActor* SelectedActor) { return nullptr; }
+
+	void MarkRenderStateDirty() { bRenderStateDirty = true; }
+
+	void CleanRenderStateDirtyFlag() { bRenderStateDirty = false; }
+
+	bool IsRenderStateDirty() { return !!bRenderStateDirty; }
+
+	DECLARE_MULTICAST_DELEGATE_OneParam(FOnClassRegisterChangeEvent, class UGameplayDebuggerBaseObject*);
+	static FOnClassRegisterChangeEvent& GetOnClassRegisterEvent();
+	static FOnClassRegisterChangeEvent& GetOnClassUnregisterEvent();
+
+protected:
+	UPROPERTY(ReplicatedUsing = OnRep_GenericShapeElements)
+	TArray<FGameplayDebuggerShapeElement> GenericShapeElements;
+
+	UPROPERTY(ReplicatedUsing = OnRep_ReplicatedBlob)
+	FGameplayDebuggerReplicatedBlob ReplicatedBlob;
+
+	struct FReplicationDataHeader
+	{
+		uint32 UncompressedSize : 31;
+		uint32 bUseCompression : 1;
+	};
+
+	FGameplayDebuggerReplicatedBlob LocalReplicationData;
+	double DataPreparationStartTime;
+
+	virtual void CollectDataToReplicate(APlayerController* MyPC, AActor *SelectedActor);
+
+	UFUNCTION(Reliable, Server, WithValidation)
+	virtual void ServerRequestPacket(int32 Offset);
+
+	void RequestDataReplication(const FString& DataName, const TArray<uint8>& BufferToReplicate, EGameplayDebuggerReplicate Replicate = EGameplayDebuggerReplicate::WithCompressed);
+	virtual void OnDataReplicationComplited(const TArray<uint8>& ReplicatedData) {}
+
+	UFUNCTION()
+	virtual void OnRep_ReplicatedBlob();
+
+	UFUNCTION()
+	virtual void OnRep_GenericShapeElements();
+
+private:
+	int32 bRenderStateDirty : 1;
+};
diff --git a/Engine/Source/Runtime/Engine/Classes/Engine/DebugCameraHUD.h b/Engine/Source/Runtime/Engine/Classes/Engine/DebugCameraHUD.h
--- a/Engine/Source/Runtime/Engine/Classes/Engine/DebugCameraHUD.h
+++ b/Engine/Source/Runtime/Engine/Classes/Engine/DebugCameraHUD.h
@@ -9,7 +9,7 @@
 /**
  * HUD that displays info for the DebugCameraController view.
  */
-UCLASS(config=Game)
+UCLASS(config = Game, hidedropdown)
 class ENGINE_API ADebugCameraHUD
 	: public AHUD
 {
diff --git a/Engine/Source/Runtime/Engine/Classes/Engine/EngineBaseTypes.h b/Engine/Source/Runtime/Engine/Classes/Engine/EngineBaseTypes.h
--- a/Engine/Source/Runtime/Engine/Classes/Engine/EngineBaseTypes.h
+++ b/Engine/Source/Runtime/Engine/Classes/Engine/EngineBaseTypes.h
@@ -44,9 +44,6 @@ enum ETickingGroup
 	/** Any item that needs to be executed before physics simulation starts. */
 	TG_PrePhysics UMETA(DisplayName="Pre Physics"),
 
-	/** Any item that can safely run during parallel animation processing (i.e. that cannot touch the animation system) */
-	TG_DuringAnimation UMETA(DisplayName="During Animation"),
-
 	/** Special tick group that starts physics simulation. */							
 	TG_StartPhysics UMETA(Hidden, DisplayName="Start Physics"),
 
diff --git a/Engine/Source/Runtime/Engine/Classes/Engine/EngineTypes.h b/Engine/Source/Runtime/Engine/Classes/Engine/EngineTypes.h
--- a/Engine/Source/Runtime/Engine/Classes/Engine/EngineTypes.h
+++ b/Engine/Source/Runtime/Engine/Classes/Engine/EngineTypes.h
@@ -2495,7 +2495,15 @@ struct FMaterialSimplificationSettings
 	}
 };
 
-
+UENUM()
+enum ETextureSizingType
+{
+	TextureSizingType_UseSingleTextureSize UMETA(DisplayName = "Use TextureSize for all material properties"),
+	TextureSizingType_UseAutomaticBiasedSizes UMETA(DisplayName = "Use automatically biased texture sizes based on TextureSize"),
+	TextureSizingType_UseManualOverrideTextureSize UMETA(DisplayName = "Use per property manually overriden texture sizes"),
+	TextureSizingType_UseSimplygonAutomaticSizing UMETA(DisplayName = "Use Simplygon's automatic texture sizing"),
+	TextureSizingType_MAX,
+};
 
 USTRUCT()
 struct FMaterialProxySettings
@@ -2506,9 +2514,8 @@ struct FMaterialProxySettings
 	UPROPERTY(Category = Material, EditAnywhere)
 	FIntPoint TextureSize;
 
-	// Use automatic texure sizes
 	UPROPERTY(Category = Material, EditAnywhere)
-	bool bAutomaticTextureSizes;
+	TEnumAsByte<ETextureSizingType> TextureSizingType;
 
 	UPROPERTY(Category = Material, EditAnywhere)
 	float GutterSpace;
@@ -2549,25 +2556,60 @@ struct FMaterialProxySettings
 	UPROPERTY(Category = Material, EditAnywhere)
 	bool bOpacityMap;
 
-
+	// Override diffuse map size
+	UPROPERTY(Category = Material, AdvancedDisplay, EditAnywhere)
+	FIntPoint DiffuseTextureSize;
+
+	// Override normal map size
+	UPROPERTY(Category = Material, AdvancedDisplay, EditAnywhere)
+	FIntPoint NormalTextureSize;
+
+	// Override metallic map size
+	UPROPERTY(Category = Material, AdvancedDisplay, EditAnywhere)
+	FIntPoint MetallicTextureSize;
+
+	// Override roughness map size
+	UPROPERTY(Category = Material, AdvancedDisplay, EditAnywhere)
+	FIntPoint RoughnessTextureSize;
+
+	// Override specular map size
+	UPROPERTY(Category = Material, AdvancedDisplay, EditAnywhere)
+	FIntPoint SpecularTextureSize;				    
+						
+	// Override emissive map size
+	UPROPERTY(Category = Material, AdvancedDisplay, EditAnywhere)
+	FIntPoint EmissiveTextureSize;				    
+						
+	// Override opacity map size
+	UPROPERTY(Category = Material, AdvancedDisplay, EditAnywhere)
+	FIntPoint OpacityTextureSize;
+	
 	FMaterialProxySettings()
 		: TextureSize(1024, 1024)
-		, bAutomaticTextureSizes(false)
+		, TextureSizingType(TextureSizingType_UseSingleTextureSize)
 		, GutterSpace(4.0f)
 		, bNormalMap(true)
 		, bMetallicMap(false)
 		, MetallicConstant(0.0f)
 		, bRoughnessMap(false)
 		, RoughnessConstant(0.5f)				
-		, bSpecularMap(false)
+		, bSpecularMap(false)		
 		, SpecularConstant(0.5f)
+		, bEmissiveMap(false)
+		, bOpacityMap(false)
+		, DiffuseTextureSize(1024, 1024)
+		, NormalTextureSize(1024, 1024)
+		, MetallicTextureSize(1024, 1024)
+		, RoughnessTextureSize(1024, 1024)
+		, EmissiveTextureSize(1024, 1024)
+		, OpacityTextureSize(1024, 1024)
 	{
 	}
 
 	bool operator == (const FMaterialProxySettings& Other) const
 	{
 		return TextureSize == Other.TextureSize
-			&& bAutomaticTextureSizes == Other.bAutomaticTextureSizes
+			&& TextureSizingType == Other.TextureSizingType
 			&& GutterSpace == Other.GutterSpace
 			&& bNormalMap == Other.bNormalMap
 			&& MetallicConstant == Other.MetallicConstant
@@ -2575,7 +2617,15 @@ struct FMaterialProxySettings
 			&& RoughnessConstant == Other.RoughnessConstant
 			&& bRoughnessMap == Other.bRoughnessMap
 			&& SpecularConstant == Other.SpecularConstant
-			&& bSpecularMap == Other.bSpecularMap;
+			&& bSpecularMap == Other.bSpecularMap
+			&& bEmissiveMap == Other.bEmissiveMap
+			&& bOpacityMap == Other.bOpacityMap
+			&& DiffuseTextureSize == Other.DiffuseTextureSize
+			&& NormalTextureSize == Other.NormalTextureSize
+			&& MetallicTextureSize == Other.MetallicTextureSize
+			&& RoughnessTextureSize == Other.RoughnessTextureSize
+			&& EmissiveTextureSize == Other.EmissiveTextureSize
+			&& OpacityTextureSize == Other.OpacityTextureSize;
 	}
 };
 
diff --git a/Engine/Source/Runtime/Engine/Classes/Engine/LevelStreaming.h b/Engine/Source/Runtime/Engine/Classes/Engine/LevelStreaming.h
--- a/Engine/Source/Runtime/Engine/Classes/Engine/LevelStreaming.h
+++ b/Engine/Source/Runtime/Engine/Classes/Engine/LevelStreaming.h
@@ -97,6 +97,13 @@ public:
 	UPROPERTY(EditAnywhere, Category=LevelStreaming, BlueprintReadWrite)
 	FTransform LevelTransform;
 
+private:
+
+	/* Whether CachedWorldAssetPackageFName is valid */
+	uint32 bHasCachedWorldAssetPackageFName:1;
+
+public:
+
 	/** Whether we currently have a load request pending.																		*/
 	uint32 bHasLoadRequestPending:1;
 
diff --git a/Engine/Source/Runtime/Engine/Classes/Engine/MemberReference.h b/Engine/Source/Runtime/Engine/Classes/Engine/MemberReference.h
--- a/Engine/Source/Runtime/Engine/Classes/Engine/MemberReference.h
+++ b/Engine/Source/Runtime/Engine/Classes/Engine/MemberReference.h
@@ -392,7 +392,7 @@ public:
 				FString const StringName = MemberName.ToString();
 				for (TObjectIterator<UPackage> PackageIt; PackageIt && (ReturnField == nullptr); ++PackageIt)
 				{
-					if ((PackageIt->PackageFlags & PKG_CompiledIn) == 0x00)
+					if (PackageIt->HasAnyPackageFlags(PKG_CompiledIn) == false)
 					{
 						continue;
 					}
diff --git a/Engine/Source/Runtime/Engine/Classes/Engine/SceneCaptureCube.h b/Engine/Source/Runtime/Engine/Classes/Engine/SceneCaptureCube.h
--- a/Engine/Source/Runtime/Engine/Classes/Engine/SceneCaptureCube.h
+++ b/Engine/Source/Runtime/Engine/Classes/Engine/SceneCaptureCube.h
@@ -9,8 +9,8 @@
 #include "SceneCapture.h"
 #include "SceneCaptureCube.generated.h"
 
-UCLASS(hidecategories = (Collision, Material, Attachment, Actor), MinimalAPI)
-class ASceneCaptureCube : public ASceneCapture
+UCLASS(hidecategories = (Collision, Material, Attachment, Actor))
+class ENGINE_API ASceneCaptureCube : public ASceneCapture
 {
 	GENERATED_UCLASS_BODY()
 
@@ -30,7 +30,7 @@ public:
 	//~ Begin AActor Interface
 	virtual void PostActorCreated() override;
 #if WITH_EDITOR
-	ENGINE_API virtual void PostEditMove(bool bFinished) override;
+	virtual void PostEditMove(bool bFinished) override;
 #endif
 	//~ End AActor Interface.
 
@@ -41,9 +41,9 @@ public:
 	void OnInterpToggle(bool bEnable);
 
 	/** Returns CaptureComponentCube subobject **/
-	ENGINE_API class USceneCaptureComponentCube* GetCaptureComponentCube() const;
+	class USceneCaptureComponentCube* GetCaptureComponentCube() const;
 	/** Returns DrawFrustum subobject **/
-	ENGINE_API class UDrawFrustumComponent* GetDrawFrustum() const;
+	class UDrawFrustumComponent* GetDrawFrustum() const;
 };
 
 
diff --git a/Engine/Source/Runtime/Engine/Classes/Engine/World.h b/Engine/Source/Runtime/Engine/Classes/Engine/World.h
--- a/Engine/Source/Runtime/Engine/Classes/Engine/World.h
+++ b/Engine/Source/Runtime/Engine/Classes/Engine/World.h
@@ -417,13 +417,11 @@ struct FEndClothSimulationFunction : public FTickFunction
 };
 
 /* Struct of optional parameters passed to SpawnActor function(s). */
+PRAGMA_DISABLE_DEPRECATION_WARNINGS // Required for auto-generated functions referencing bNoCollisionFail
 struct ENGINE_API FActorSpawnParameters
 {
 	FActorSpawnParameters();
 
-	// Assignment operator overriden to allow disabling of deprecation warning when copying bNoCollisionFail
-	FActorSpawnParameters& operator=(const FActorSpawnParameters& Other);
-
 	/* A name to assign as the Name of the Actor being spawned. If no value is specified, the name of the spawned Actor will be automatically generated using the form [Class]_[Number]. */
 	FName Name;
 
@@ -461,6 +459,7 @@ struct ENGINE_API FActorSpawnParameters
 	/* Flags used to describe the spawned actor/object instance. */
 	EObjectFlags ObjectFlags;		
 };
+PRAGMA_ENABLE_DEPRECATION_WARNINGS
 
 
 /**
@@ -488,10 +487,6 @@ struct ENGINE_API FWorldAsyncTraceState
 
 	/** Used as counter for Buffer swap for DataBuffer. Right now it's only 2, but it can change. */
 	int32 CurrentFrame;
-
-	/** Next available index for each pool - used as ID for each trace query **/
-	int32 NextAvailableTraceIndex;
-	int32 NextAvailableOverlapIndex;
 };
 
 #if WITH_EDITOR
diff --git a/Engine/Source/Runtime/Engine/Classes/GameFramework/Actor.h b/Engine/Source/Runtime/Engine/Classes/GameFramework/Actor.h
--- a/Engine/Source/Runtime/Engine/Classes/GameFramework/Actor.h
+++ b/Engine/Source/Runtime/Engine/Classes/GameFramework/Actor.h
@@ -107,23 +107,23 @@ public:
 	 * @see SetActorHiddenInGame()
 	 */
 	UPROPERTY(Interp, EditAnywhere, Category=Rendering, BlueprintReadOnly, replicated, meta=(DisplayName = "Actor Hidden In Game", SequencerTrackClass = "MovieSceneVisibilityTrack"))
-	uint32 bHidden:1;
+	uint8 bHidden:1;
 
 	/** If true, when the actor is spawned it will be sent to the client but receive no further replication updates from the server afterwards. */
 	UPROPERTY()
-	uint32 bNetTemporary:1;
+	uint8 bNetTemporary:1;
 
 	/** If true, this actor was loaded directly from the map, and for networking purposes can be addressed by its full path name */
 	UPROPERTY()
-	uint32 bNetStartup:1;
+	uint8 bNetStartup:1;
 
 	/** If true, this actor is only relevant to its owner. If this flag is changed during play, all non-owner channels would need to be explicitly closed. */
 	UPROPERTY(Category=Replication, EditDefaultsOnly, BlueprintReadOnly)
-	uint32 bOnlyRelevantToOwner:1;
+	uint8 bOnlyRelevantToOwner:1;
 
 	/** Always relevant for network (overrides bOnlyRelevantToOwner). */
 	UPROPERTY(Category=Replication, EditDefaultsOnly, BlueprintReadWrite)
-	uint32 bAlwaysRelevant:1;    
+	uint8 bAlwaysRelevant:1;    
 
 	/**
 	 * If true, replicate movement/location related properties.
@@ -132,7 +132,7 @@ public:
 	 * @see https://docs.unrealengine.com/latest/INT/Gameplay/Networking/Replication/
 	 */
 	UPROPERTY(ReplicatedUsing=OnRep_ReplicateMovement, Category=Replication, EditDefaultsOnly)
-	uint32 bReplicateMovement:1;    
+	uint8 bReplicateMovement:1;    
 
 	/** Called on client when updated bReplicateMovement value is received for this actor. */
 	UFUNCTION()
@@ -143,7 +143,7 @@ public:
 	 * @see TornOff()
 	 */
 	UPROPERTY(replicated)
-	uint32 bTearOff:1;    
+	uint8 bTearOff:1;    
 
 	/** Networking - Server - TearOff this actor to stop replication to clients. Will set bTearOff to true. */
 	UFUNCTION(BlueprintCallable, Category = Replication)
@@ -154,26 +154,26 @@ public:
 	 * Causes all initialization to be performed again even though the actor may not have actually been reloaded.
 	 */
 	UPROPERTY(transient)
-	uint32 bExchangedRoles:1;
+	uint8 bExchangedRoles:1;
 
 	/** Is this actor still pending a full net update due to clients that weren't able to replicate the actor at the time of LastNetUpdateTime */
 	UPROPERTY(transient)
-	uint32 bPendingNetUpdate:1;
+	uint8 bPendingNetUpdate:1;
 
 	/** This actor will be loaded on network clients during map load */
 	UPROPERTY(Category=Replication, EditDefaultsOnly)
-	uint32 bNetLoadOnClient:1;
+	uint8 bNetLoadOnClient:1;
 
 	/** If actor has valid Owner, call Owner's IsNetRelevantFor and GetNetPriority */
 	UPROPERTY(Category=Replication, EditDefaultsOnly, BlueprintReadWrite)
-	uint32 bNetUseOwnerRelevancy:1;
+	uint8 bNetUseOwnerRelevancy:1;
 
 	/** If true, all input on the stack below this actor will not be considered */
 	UPROPERTY(EditDefaultsOnly, Category=Input)
-	uint32 bBlockInput:1;
+	uint8 bBlockInput:1;
 
 	/** True if this actor is currently running user construction script (used to defer component registration) */
-	uint32 bRunningUserConstructionScript:1;
+	uint8 bRunningUserConstructionScript:1;
 
 	/**
 	 * Whether we allow this Actor to tick before it receives the BeginPlay event.
@@ -181,24 +181,24 @@ public:
 	 * This Actor must be able to tick for this setting to be relevant.
 	 */
 	UPROPERTY(EditDefaultsOnly, Category="Tick")
-	uint32 bAllowTickBeforeBeginPlay:1;
+	uint8 bAllowTickBeforeBeginPlay:1;
 
 private:
 	/** Whether FinishSpawning has been called for this Actor.  If it has not, the Actor is in a mal-formed state */
-	uint32 bHasFinishedSpawning:1;
+	uint8 bHasFinishedSpawning:1;
 
 	/** Whether we've tried to register tick functions. Reset when they are unregistered. */
-	uint32 bTickFunctionsRegistered : 1;
+	uint8 bTickFunctionsRegistered : 1;
 
 	/**
 	 * Enables any collision on this actor.
 	 * @see SetActorEnableCollision(), GetActorEnableCollision()
 	 */
 	UPROPERTY()
-	uint32 bActorEnableCollision:1;
+	uint8 bActorEnableCollision:1;
 
 	/** Flag indicating we have checked initial simulating physics state to sync networked proxies to the server. */
-	uint32 bNetCheckedInitialPhysicsState:1;
+	uint8 bNetCheckedInitialPhysicsState:1;
 
 protected:
 	/**
@@ -206,7 +206,7 @@ protected:
 	 * @see SetReplicates()
 	 */
 	UPROPERTY(EditDefaultsOnly, BlueprintReadOnly, Category = "Replication")
-	uint32 bReplicates:1;
+	uint8 bReplicates:1;
 
 	/** This function should only be used in the constructor of classes that need to set the RemoteRole for backwards compatibility purposes */
 	void SetRemoteRoleForBackwardsCompat(const ENetRole InRemoteRole) { RemoteRole = InRemoteRole; }
@@ -348,7 +348,7 @@ public:
 
 	/** If true then destroy self when "finished", meaning all relevant components report that they are done and no timelines or timers are in flight. */
 	UPROPERTY(BlueprintReadWrite, Category=Actor)
-	uint32 bAutoDestroyWhenFinished:1;
+	uint8 bAutoDestroyWhenFinished:1;
 
 	/**
 	 * Whether this actor can take damage. Must be true for damage events (e.g. ReceiveDamage()) to be called.
@@ -356,28 +356,28 @@ public:
 	 * @see TakeDamage(), ReceiveDamage()
 	 */
 	UPROPERTY(EditAnywhere, BlueprintReadWrite, SaveGame, Replicated, Category=Actor)
-	uint32 bCanBeDamaged:1;
+	uint8 bCanBeDamaged:1;
 
 private:
 	/**
 	 * Set when actor is about to be deleted.
 	 */
 	UPROPERTY(Transient, DuplicateTransient)
-	uint32 bActorIsBeingDestroyed:1;    
+	uint8 bActorIsBeingDestroyed:1;    
 
 public:
 
 	/** This actor collides with the world when placing in the editor, even if RootComponent collision is disabled. Does not affect spawning, @see SpawnCollisionHandlingMethod */
 	UPROPERTY()
-	uint32 bCollideWhenPlacing:1;
+	uint8 bCollideWhenPlacing:1;
 
 	/** If true, this actor should search for an owned camera component to view through when used as a view target. */
 	UPROPERTY(EditAnywhere, BlueprintReadWrite, Category=Actor, AdvancedDisplay)
-	uint32 bFindCameraComponentWhenViewTarget:1;
+	uint8 bFindCameraComponentWhenViewTarget:1;
 	
 	/** If true, this actor will be replicated to network replays (default is true) */
 	UPROPERTY()
-	uint32 bRelevantForNetworkReplays:1;
+	uint8 bRelevantForNetworkReplays:1;
 
 	/** Controls how to handle spawning this actor in a situation where it's colliding with something else. "Default" means AlwaysSpawn here. */
 	UPROPERTY(EditAnywhere, BlueprintReadWrite, Category=Actor)
@@ -424,15 +424,6 @@ protected:
 	/** Handle for efficient management of LifeSpanExpired timer */
 	FTimerHandle TimerHandle_LifeSpanExpired;
 
-protected:
-
-	/**
-	 * If false, the Blueprint ReceiveTick() event will be disabled on dedicated servers.
-	 * @see AllowReceiveTickEventOnDedicatedServer()
-	 */
-	UPROPERTY()
-	uint32 bAllowReceiveTickEventOnDedicatedServer:1;
-
 public:
 
 	/** Return the value of bAllowReceiveTickEventOnDedicatedServer, indicating whether the Blueprint ReceiveTick() event will occur on dedicated servers. */
@@ -442,11 +433,19 @@ public:
 	UPROPERTY()
 	TArray< FName > Layers;
 
+	/** The Actor that owns the UChildActorComponent that owns this Actor. */
+	UPROPERTY()
+	TWeakObjectPtr<AActor> ParentComponentActor;	
+
 #if WITH_EDITORONLY_DATA
-protected:
 
-	UPROPERTY()
-	uint32 bActorLabelEditable:1;    // Is the actor label editable by the user?
+	/** The group this actor is a part of. */
+	UPROPERTY(transient)
+	AActor* GroupActor;
+
+	/** The scale to apply to any billboard components in editor builds (happens in any WITH_EDITOR build, including non-cooked games). */
+	UPROPERTY(EditAnywhere, BlueprintReadWrite, Category=Rendering, meta=(DisplayName="Editor Billboard Scale"))
+	float SpriteScale;
 
 private:
 	/**
@@ -460,55 +459,57 @@ private:
 	UPROPERTY()
 	FName FolderPath;
 
+protected:
+
+	UPROPERTY()
+	uint8 bActorLabelEditable:1;    // Is the actor label editable by the user?
+
 public:
 	/** Whether this actor is hidden within the editor viewport. */
 	UPROPERTY()
-	uint32 bHiddenEd:1;
+	uint8 bHiddenEd:1;
 
 protected:
 	/** Whether the actor can be manipulated by editor operations. */
 	UPROPERTY()
-	uint32 bEditable:1;
+	uint8 bEditable:1;
 
 	/** Whether this actor should be listed in the scene outliner. */
 	UPROPERTY()
-	uint32 bListedInSceneOutliner:1;
+	uint8 bListedInSceneOutliner:1;
 
 public:
 	/** Whether this actor is hidden by the layer browser. */
 	UPROPERTY()
-	uint32 bHiddenEdLayer:1;
+	uint8 bHiddenEdLayer:1;
 
 private:
 	/** Whether this actor is temporarily hidden within the editor; used for show/hide/etc functionality w/o dirtying the actor. */
 	UPROPERTY(transient)
-	uint32 bHiddenEdTemporary:1;
+	uint8 bHiddenEdTemporary:1;
 
 public:
 
 	/** Whether this actor is hidden by the level browser. */
 	UPROPERTY(transient)
-	uint32 bHiddenEdLevel:1;
+	uint8 bHiddenEdLevel:1;
 
 	/** If true, prevents the actor from being moved in the editor viewport. */
 	UPROPERTY()
-	uint32 bLockLocation:1;
-
-	/** The group this actor is a part of. */
-	UPROPERTY(transient)
-	AActor* GroupActor;
-
-	/** The scale to apply to any billboard components in editor builds (happens in any WITH_EDITOR build, including non-cooked games). */
-	UPROPERTY(EditAnywhere, BlueprintReadWrite, Category=Rendering, meta=(DisplayName="Editor Billboard Scale"))
-	float SpriteScale;
+	uint8 bLockLocation:1;
 
 	/** Returns how many lights are uncached for this actor. */
 	int32 GetNumUncachedLights();
 #endif // WITH_EDITORONLY_DATA
 
-	/** The Actor that owns the UChildActorComponent that owns this Actor. */
+protected:
+
+	/**
+	 * If false, the Blueprint ReceiveTick() event will be disabled on dedicated servers.
+	 * @see AllowReceiveTickEventOnDedicatedServer()
+	 */
 	UPROPERTY()
-	TWeakObjectPtr<AActor> ParentComponentActor;	
+	uint8 bAllowReceiveTickEventOnDedicatedServer:1;
 
 private:
 
@@ -516,26 +517,33 @@ private:
 	 *	Indicates that PreInitializeComponents/PostInitializeComponents have been called on this Actor 
 	 *	Prevents re-initializing of actors spawned during level startup
 	 */
-	uint32 bActorInitialized:1;
+	uint8 bActorInitialized:1;
 	
+	enum class EActorBeginPlayState : uint8
+	{
+		HasNotBegunPlay,
+		BeginningPlay,
+		HasBegunPlay,
+	};
+
 	/** 
 	 *	Indicates that BeginPlay has been called for this Actor.
 	 *  Set back to false once EndPlay has been called.
 	 */
-	uint32 bActorHasBegunPlay:1;
+	EActorBeginPlayState ActorHasBegunPlay:2;
 
 public:
 	/** Indicates the actor was pulled through a seamless travel.  */
 	UPROPERTY()
-	uint32 bActorSeamlessTraveled:1;
+	uint8 bActorSeamlessTraveled:1;
 
 	/** Whether this actor should no be affected by world origin shifting. */
 	UPROPERTY(EditAnywhere, AdvancedDisplay, Category=Actor)
-	uint32 bIgnoresOriginShifting:1;
+	uint8 bIgnoresOriginShifting:1;
 	
 	/** If true, and if World setting has bEnableHigerarhicalLOD is true, then it will generate LODActor from groups of clustered Actor */
 	UPROPERTY(EditAnywhere, AdvancedDisplay, Category=Actor)
-	uint32 bEnableAutoLODGeneration:1;
+	uint8 bEnableAutoLODGeneration:1;
 
 	/** Array of tags that can be used for grouping and categorizing. */
 	UPROPERTY(EditAnywhere, BlueprintReadWrite, AdvancedDisplay, Category=Actor)
@@ -1163,8 +1171,11 @@ public:
 	/** Returns whether an actor has been initialized */
 	bool IsActorInitialized() const { return bActorInitialized; }
 
+	/** Returns whether an actor is in the process of beginning play */
+	bool IsActorBeginningPlay() const { return ActorHasBegunPlay == EActorBeginPlayState::BeginningPlay; }
+
 	/** Returns whether an actor has had BeginPlay called on it (and not subsequently had EndPlay called) */
-	bool HasActorBegunPlay() const { return bActorHasBegunPlay; }
+	bool HasActorBegunPlay() const { return ActorHasBegunPlay == EActorBeginPlayState::HasBegunPlay; }
 
 	UFUNCTION(BlueprintCallable, Category="Game")
 	bool IsActorBeingDestroyed() const 
diff --git a/Engine/Source/Runtime/Engine/Classes/GameFramework/HUD.h b/Engine/Source/Runtime/Engine/Classes/GameFramework/HUD.h
--- a/Engine/Source/Runtime/Engine/Classes/GameFramework/HUD.h
+++ b/Engine/Source/Runtime/Engine/Classes/GameFramework/HUD.h
@@ -127,7 +127,8 @@ public:
 
 private:
 	/** if true show debug info for 'ShowDebugTargetActor', otherwise for Camera Viewtarget */
-	static bool bShowDebugForReticleTarget;
+	UPROPERTY(Transient)
+	bool bShowDebugForReticleTarget;
 
 	/** Class filter for selecting 'ShowDebugTargetActor' when 'bShowDebugForReticleTarget' is true. */
 	UPROPERTY(Transient)
diff --git a/Engine/Source/Runtime/Engine/Classes/GameFramework/WorldSettings.h b/Engine/Source/Runtime/Engine/Classes/GameFramework/WorldSettings.h
--- a/Engine/Source/Runtime/Engine/Classes/GameFramework/WorldSettings.h
+++ b/Engine/Source/Runtime/Engine/Classes/GameFramework/WorldSettings.h
@@ -200,6 +200,8 @@ struct ENGINE_API FNetViewer
 	FNetViewer(UNetConnection* InConnection, float DeltaSeconds);
 };
 
+PRAGMA_DISABLE_DEPRECATION_WARNINGS
+
 USTRUCT()
 struct ENGINE_API FHierarchicalSimplification
 {
@@ -252,26 +254,19 @@ struct ENGINE_API FHierarchicalSimplification
 		MergeSetting.bGenerateLightMapUV = true;
 	}
 
+private:
 
-
-	FHierarchicalSimplification& operator=(const FHierarchicalSimplification& Other)
+	// This function exists to force the compiler generated operators to be instantiated while the deprecation warning
+	// pragmas are disabled so no warnings are thrown when used elsewhere and the compiler is forced to generate them
+	void DummyFunction() const
 	{
-		PRAGMA_DISABLE_DEPRECATION_WARNINGS
-		
-		TransitionScreenSize = Other.TransitionScreenSize;
-		DrawDistance = Other.DrawDistance;
-		bSimplifyMesh = Other.bSimplifyMesh;
-		ProxySetting = Other.ProxySetting;
-		MergeSetting = Other.MergeSetting;
-		DesiredBoundRadius = Other.DesiredBoundRadius;
-		DesiredFillingPercentage = Other.DesiredFillingPercentage;
-		MinNumberOfActorsToBuild = Other.MinNumberOfActorsToBuild;
-		return *this;
-
-		PRAGMA_ENABLE_DEPRECATION_WARNINGS
-	};	
+		FHierarchicalSimplification ASP1, ASP2;
+		ASP1 = ASP2;
+	}
 };
 
+PRAGMA_ENABLE_DEPRECATION_WARNINGS
+
 /**
  * Actor containing all script accessible world properties.
  */
diff --git a/Engine/Source/Runtime/Engine/Classes/Materials/MaterialInterface.h b/Engine/Source/Runtime/Engine/Classes/Materials/MaterialInterface.h
--- a/Engine/Source/Runtime/Engine/Classes/Materials/MaterialInterface.h
+++ b/Engine/Source/Runtime/Engine/Classes/Materials/MaterialInterface.h
@@ -569,8 +569,8 @@ public:
 	/** Get bitfield indicating which feature levels should be compiled by default */
 	ENGINE_API static uint32 GetFeatureLevelsToCompileForAllMaterials() { return FeatureLevelsForAllMaterials | (1 << GMaxRHIFeatureLevel); }
 
-	/** Return number of used texture coordinates and whether or not the Vertex Color material node is used in the shader graph */
-	ENGINE_API void AnalyzeMaterialProperty(EMaterialProperty InProperty, int32& OutNumTextureCoordinates, bool& OutUseVertexColor);
+	/** Return number of used texture coordinates and whether or not the Vertex data is used in the shader graph */
+	ENGINE_API void AnalyzeMaterialProperty(EMaterialProperty InProperty, int32& OutNumTextureCoordinates, bool& bOutRequiresVertexData);
 
 	/** Iterate over all feature levels currently marked as active */
 	template <typename FunctionType>
diff --git a/Engine/Source/Runtime/Engine/Classes/Particles/EmitterCameraLensEffectBase.h b/Engine/Source/Runtime/Engine/Classes/Particles/EmitterCameraLensEffectBase.h
--- a/Engine/Source/Runtime/Engine/Classes/Particles/EmitterCameraLensEffectBase.h
+++ b/Engine/Source/Runtime/Engine/Classes/Particles/EmitterCameraLensEffectBase.h
@@ -2,76 +2,77 @@
 
 /**
  * Base class for Camera Lens Effects.  Needed so we can have AnimNotifies be able to show camera effects
- * in a nice drop down
- *
+ * in a nice drop down.
  */
 
 #pragma once
 #include "Particles/Emitter.h"
 #include "EmitterCameraLensEffectBase.generated.h"
 
-UCLASS(abstract, Blueprintable, MinimalAPI)
-class AEmitterCameraLensEffectBase : public AEmitter
+UCLASS(abstract, Blueprintable)
+class ENGINE_API AEmitterCameraLensEffectBase : public AEmitter
 {
 	GENERATED_UCLASS_BODY()
 
 protected:
 	/** Particle System to use */
-	UPROPERTY()
+	UPROPERTY(EditDefaultsOnly, Category = EmitterCameraLensEffectBase)
 	class UParticleSystem* PS_CameraEffect;
 
 	/** The effect to use for non extreme content */
-	UPROPERTY()
+	UPROPERTY(EditDefaultsOnly, Category = EmitterCameraLensEffectBase)
 	class UParticleSystem* PS_CameraEffectNonExtremeContent;
 
-public:
-	/** In order to get the particle effect looking correct we need to have a base FOV which we just to move the particle closer/further from the camera **/
-	UPROPERTY()
-	float BaseFOV;
+	/** Camera this emitter is attached to, will be notified when emitter is destroyed */
+	UPROPERTY(transient)
+	class APlayerCameraManager* BaseCamera;
 
 	/** 
-	 *  How far in front of the camera this emitter should live, assuming an FOV of 80 degrees. 
-	 *  Note that the actual distance will be automatically adjusted to account for the actual FOV.
+	 * Effect-to-camera transform to allow arbitrary placement of the particle system .
+	 * Note the X component of the location will be scaled with camera fov to keep the lens effect the same apparent size.
 	 */
-	UPROPERTY(EditAnywhere, Category=EmitterCameraLensEffectBase)
-	float DistFromCamera;
+	UPROPERTY(EditDefaultsOnly, Category = EmitterCameraLensEffectBase)
+	FTransform RelativeTransform;
+
+public:
+	/** This is the assumed FOV for which the effect was authored. The code will make automatic adjustments to make it look the same at different FOVs */
+	UPROPERTY(EditDefaultsOnly, Category = EmitterCameraLensEffectBase)
+	float BaseFOV;
 
 	/** true if multiple instances of this emitter can exist simultaneously, false otherwise.  */
-	UPROPERTY(EditAnywhere, Category=EmitterCameraLensEffectBase)
-	uint32 bAllowMultipleInstances:1;
+	UPROPERTY(EditAnywhere, Category = EmitterCameraLensEffectBase)
+	uint8 bAllowMultipleInstances:1;
 
 	/** 
 	 *  If an emitter class in this array is currently playing, do not play this effect.
 	 *  Useful for preventing multiple similar or expensive camera effects from playing simultaneously.
 	 */
-	UPROPERTY()
+	UPROPERTY(EditDefaultsOnly, Category = EmitterCameraLensEffectBase)
 	TArray<TSubclassOf<class AEmitterCameraLensEffectBase> > EmittersToTreatAsSame;
 
-protected:
-	/** Camera this emitter is attached to, will be notified when emitter is destroyed */
-	UPROPERTY(transient)
-	class APlayerCameraManager* BaseCamera;
-
 public:
-
-
 	//~ Begin AActor Interface
-	ENGINE_API virtual void EndPlay(const EEndPlayReason::Type EndPlayReason) override;
-	ENGINE_API virtual void PostInitializeComponents() override;
+	virtual void EndPlay(const EEndPlayReason::Type EndPlayReason) override;
+	virtual void PostInitializeComponents() override;
+	virtual void PostLoad() override;
 	//~ End AActor Interface
 
 	/** Tell the emitter what camera it is attached to. */
-	ENGINE_API virtual void RegisterCamera(APlayerCameraManager* C);
+	virtual void RegisterCamera(APlayerCameraManager* C);
 	
 	/** Called when this emitter is re-triggered, for bAllowMultipleInstances=false emitters. */
-	ENGINE_API virtual void NotifyRetriggered();
+	virtual void NotifyRetriggered();
 
 	/** This will actually activate the lens Effect.  We want this separated from PostInitializeComponents so we can cache these emitters **/
-	ENGINE_API virtual void ActivateLensEffect();
+	virtual void ActivateLensEffect();
 	
 	/** Given updated camera information, adjust this effect to display appropriately. */
-	ENGINE_API virtual void UpdateLocation(const FVector& CamLoc, const FRotator& CamRot, float CamFOVDeg);
-};
+	virtual void UpdateLocation(const FVector& CamLoc, const FRotator& CamRot, float CamFOVDeg);
 
+private:
+	/** DEPRECATED(4.11) */
+	UPROPERTY()
+	float DistFromCamera_DEPRECATED;
+};
 
 
diff --git a/Engine/Source/Runtime/Engine/Classes/Sound/SoundAttenuation.h b/Engine/Source/Runtime/Engine/Classes/Sound/SoundAttenuation.h
--- a/Engine/Source/Runtime/Engine/Classes/Sound/SoundAttenuation.h
+++ b/Engine/Source/Runtime/Engine/Classes/Sound/SoundAttenuation.h
@@ -128,6 +128,38 @@ struct ENGINE_API FAttenuationSettings
 	UPROPERTY(EditAnywhere, BlueprintReadWrite, Category=LowPassFilter )
 	float LPFRadiusMax;
 
+	/** Whether or not listener-focus calculations are enabled for this attenuation. */
+	UPROPERTY(EditAnywhere, BlueprintReadWrite, Category = Focus, meta = (EditCondition = "bSpatialize"))
+	uint32 bEnableListenerFocus : 1;
+
+	/** Azimuth angle (in degrees) relative to the listener forward vector which defines the focus region of sounds. Sounds playing at an angle less than this will be in focus. */
+	UPROPERTY(EditAnywhere, BlueprintReadWrite, Category = Focus, meta = (EditCondition = "bEnableListenerFocus"))
+	float FocusAzimuth;
+
+	/** Azimuth angle (in degrees) relative to the listener forward vector which defines the non-focus region of sounds. Sounds playing at an angle greater than this will be out of focus. */
+	UPROPERTY(EditAnywhere, BlueprintReadWrite, Category = Focus, meta = (EditCondition = "bEnableListenerFocus"))
+	float NonFocusAzimuth;
+
+	/** Amount to scale the distance calculation of sounds that are in-focus. Can be used to make in-focus sounds appear to be closer or further away than they actually are. */
+	UPROPERTY(EditAnywhere, BlueprintReadWrite, Category = Focus, meta = (EditCondition = "bEnableListenerFocus"))
+	float FocusDistanceScale;
+
+	/** Amount to scale the distance calculation of sounds that are not in-focus. Can be used to make in-focus sounds appear to be closer or further away than they actually are.  */
+	UPROPERTY(EditAnywhere, BlueprintReadWrite, Category = Focus, meta = (EditCondition = "bEnableListenerFocus"))
+	float NonFocusDistanceScale;
+
+	/** Amount to scale the priority of sounds that are in focus. Can be used to boost the priority of sounds that are in focus. */
+	UPROPERTY(EditAnywhere, BlueprintReadWrite, Category = Focus, meta = (EditCondition = "bEnableListenerFocus"))
+	float FocusPriorityScale;
+
+	/** Amount to scale the priority of sounds that are not in-focus. Can be used to reduce the priority of sounds that are not in focus. */
+	UPROPERTY(EditAnywhere, BlueprintReadWrite, Category = Focus, meta = (EditCondition = "bEnableListenerFocus"))
+	float NonFocusPriorityScale;
+
+	/** Amount to attenuate sounds that are not in focus. Can be overridden at the sound-level. */
+	UPROPERTY(EditAnywhere, BlueprintReadWrite, Category = Focus, meta = (EditCondition = "bEnableListenerFocus"))
+	float NonFocusVolumeAttenuation;
+
 	FAttenuationSettings()
 		: bAttenuate(true)
 		, bSpatialize(true)
@@ -144,14 +176,20 @@ struct ENGINE_API FAttenuationSettings
 		, FalloffDistance(3600.f)
 		, LPFRadiusMin(3000.f)
 		, LPFRadiusMax(6000.f)
+		, bEnableListenerFocus(false)
+		, FocusAzimuth(30.0f)
+		, NonFocusAzimuth(60.0f)
+		, FocusDistanceScale(1.0f)
+		, NonFocusDistanceScale(1.0f)
+		, FocusPriorityScale(1.0f)
+		, NonFocusPriorityScale(1.0f)
+		, NonFocusVolumeAttenuation(1.0f)
 	{
 	}
 
 	bool operator==(const FAttenuationSettings& Other) const;
 	void PostSerialize(const FArchive& Ar);
 
-	void ApplyAttenuation( const FTransform& SoundTransform, const FVector ListenerLocation, float& Volume, float& HighFrequencyGain ) const;
-
 	struct AttenuationShapeDetails
 	{
 		FVector Extents;
@@ -161,13 +199,10 @@ struct ENGINE_API FAttenuationSettings
 
 	void CollectAttenuationShapesForVisualization(TMultiMap<EAttenuationShape::Type, AttenuationShapeDetails>& ShapeDetailsMap) const;
 	float GetMaxDimension() const;
-
-private:
-
-	float AttenuationEval(const float Distance, const float Falloff) const;
-	float AttenuationEvalBox(const FTransform& SoundLocation, const FVector ListenerLocation) const;
-	float AttenuationEvalCapsule(const FTransform& SoundLocation, const FVector ListenerLocation) const;
-	float AttenuationEvalCone(const FTransform& SoundLocation, const FVector ListenerLocation) const;
+	float AttenuationEval(const float Distance, const float Falloff, const float DistanceScale) const;
+	float AttenuationEvalBox(const FTransform& SoundLocation, const FVector ListenerLocation, const float DistanceScale) const;
+	float AttenuationEvalCapsule(const FTransform& SoundLocation, const FVector ListenerLocation, const float DistanceScale) const;
+	float AttenuationEvalCone(const FTransform& SoundLocation, const FVector ListenerLocation, const float DistanceScale) const;
 };
 
 template<>
diff --git a/Engine/Source/Runtime/Engine/Classes/Sound/SoundBase.h b/Engine/Source/Runtime/Engine/Classes/Sound/SoundBase.h
--- a/Engine/Source/Runtime/Engine/Classes/Sound/SoundBase.h
+++ b/Engine/Source/Runtime/Engine/Classes/Sound/SoundBase.h
@@ -38,6 +38,9 @@ public:
 	UPROPERTY(EditAnywhere, BlueprintReadWrite, Category = Concurrency)
 	uint32 bOverrideConcurrency:1;
 
+	UPROPERTY(EditAnywhere, BlueprintReadWrite, Category = Attenuation)
+	uint32 bIgnoreFocus:1;
+
 	/** If bOverridePlayback is false, the sound concurrency settings to use for this sound. */
 	UPROPERTY(EditAnywhere, BlueprintReadWrite, Category = Concurrency, meta = (EditCondition = "!bOverrideConcurrency"))
 	class USoundConcurrency* SoundConcurrencySettings;
@@ -61,10 +64,17 @@ public:
 	UPROPERTY(EditAnywhere, Category=Attenuation)
 	USoundAttenuation* AttenuationSettings;
 
+	/** Sound priority (higher value is higher priority) used for concurrency resolution. This priority value is weighted against the final volume of the sound. */
+	UPROPERTY(EditAnywhere, BlueprintReadWrite, Category = Concurrency)
+	float Priority;
+
 public:	
 	/** Number of times this cue is currently being played. */
 	int32 CurrentPlayCount;
 
+	/** Amount to scale distance checks for this sound based on listener-focus */
+	float FocusDistanceScale;
+
 	//~ Begin UObject Interface.
 	virtual void PostInitProperties() override;
 	virtual void PostLoad() override;
@@ -115,7 +125,17 @@ public:
 	/** Returns the FSoundConcurrencySettings struct to use */
 	const FSoundConcurrencySettings* GetSoundConcurrencySettingsToApply();
 
+	/** Returns the priority to use when evaluating concurrency */
+	float GetPriority() const;
+
 	/** Returns the sound concurrency object ID if it exists. If it doesn't exist, returns 0. */
 	uint32 GetSoundConcurrencyObjectID() const;
+
+	/** Set the amount of distance scale this sound is using based on listener-focus */
+	void SetFocusDistanceScale(const float InFocusDistanceScale) { FocusDistanceScale = InFocusDistanceScale; }
+
+	/** Get the amount of distance scale this sound is using based on listener-focus */
+	float GetFocusDistanceScale() const { return FocusDistanceScale; }
+
 };
 
diff --git a/Engine/Source/Runtime/Engine/Private/AI/Navigation/NavigationSystem.cpp b/Engine/Source/Runtime/Engine/Private/AI/Navigation/NavigationSystem.cpp
--- a/Engine/Source/Runtime/Engine/Private/AI/Navigation/NavigationSystem.cpp
+++ b/Engine/Source/Runtime/Engine/Private/AI/Navigation/NavigationSystem.cpp
@@ -276,6 +276,7 @@ UNavigationSystem::UNavigationSystem(const FObjectInitializer& ObjectInitializer
 	, InitialNavBuildingLockFlags(0)
 	, bInitialSetupHasBeenPerformed(false)
 	, bInitialLevelsAdded(false)
+	, bWorldInitDone(false)
 	, CurrentlyDrawnNavDataIndex(0)
 	, DirtyAreasUpdateTime(0)
 {
@@ -724,6 +725,7 @@ void UNavigationSystem::OnWorldInitDone(FNavigationSystemRunMode Mode)
 		}
 	}
 
+	bWorldInitDone = true;
 	OnNavigationInitDone.Broadcast();
 }
 
@@ -3187,8 +3189,11 @@ ANavigationData* UNavigationSystem::CreateNavigationDataInstance(const FNavDataC
 			// Set descriptive name
 			Instance->Rename(*StrName, NULL, REN_DoNotDirty);
 #if WITH_EDITOR
-			const bool bMarkDirty = false;
-			Instance->SetActorLabel(StrName, bMarkDirty );
+			if (World->WorldType == EWorldType::Editor)
+			{
+				const bool bMarkDirty = false;
+				Instance->SetActorLabel(StrName, bMarkDirty);
+			}
 #endif // WITH_EDITOR
 		}
 	}
diff --git a/Engine/Source/Runtime/Engine/Private/AI/Navigation/PImplRecastNavMesh.cpp b/Engine/Source/Runtime/Engine/Private/AI/Navigation/PImplRecastNavMesh.cpp
--- a/Engine/Source/Runtime/Engine/Private/AI/Navigation/PImplRecastNavMesh.cpp
+++ b/Engine/Source/Runtime/Engine/Private/AI/Navigation/PImplRecastNavMesh.cpp
@@ -741,7 +741,7 @@ void FPImplRecastNavMesh::Raycast2D(const FVector& StartLoc, const FVector& EndL
 	const FVector RecastEnd = Unreal2RecastPoint(EndLoc);
 
 	NavNodeRef StartNode = INVALID_NAVNODEREF;
-	NavQuery.findNearestPoly(&RecastStart.X, Extent, QueryFilter, &StartNode, NULL);
+	NavQuery.findNearestContainingPoly(&RecastStart.X, Extent, QueryFilter, &StartNode, NULL);
 
 	if (StartNode != INVALID_NAVNODEREF)
 	{
diff --git a/Engine/Source/Runtime/Engine/Private/AI/Navigation/RecastNavMesh.cpp b/Engine/Source/Runtime/Engine/Private/AI/Navigation/RecastNavMesh.cpp
--- a/Engine/Source/Runtime/Engine/Private/AI/Navigation/RecastNavMesh.cpp
+++ b/Engine/Source/Runtime/Engine/Private/AI/Navigation/RecastNavMesh.cpp
@@ -1904,7 +1904,7 @@ void ARecastNavMesh::BatchRaycast(TArray<FNavigationRaycastWork>& Workload, FSha
 		const FVector RecastEnd = Unreal2RecastPoint(WorkItem.RayEnd);
 
 		NavNodeRef StartNode = INVALID_NAVNODEREF;
-		NavQuery.findNearestPoly(&RecastStart.X, Extent, QueryFilter, &StartNode, NULL);
+		NavQuery.findNearestContainingPoly(&RecastStart.X, Extent, QueryFilter, &StartNode, NULL);
 
 		if (StartNode != INVALID_NAVNODEREF)
 		{
diff --git a/Engine/Source/Runtime/Engine/Private/ActiveSound.cpp b/Engine/Source/Runtime/Engine/Private/ActiveSound.cpp
--- a/Engine/Source/Runtime/Engine/Private/ActiveSound.cpp
+++ b/Engine/Source/Runtime/Engine/Private/ActiveSound.cpp
@@ -49,7 +49,7 @@ FActiveSound::FActiveSound()
 	, HighFrequencyGainMultiplier(1.f)
 	, ConcurrencyVolumeScale(1.f)
 	, SubtitlePriority(0.f)
-	, VolumeWeightedPriorityScale(1.0f)
+	, Priority(1.0f)
 	, VolumeConcurrency(0.0f)
 	, OcclusionCheckInterval(0.f)
 	, LastOcclusionCheckTime(0.f)
@@ -201,7 +201,7 @@ void FActiveSound::UpdateWaveInstances( TArray<FWaveInstance*> &InWaveInstances,
 
 	int32 ClosestListenerIndex = 0;
 
-	if (AudioDevice->Listeners.Num() > 0)
+	if (AudioDevice->Listeners.Num() > 1)
 	{
 		SCOPE_CYCLE_COUNTER( STAT_AudioFindNearestLocation );
 		ClosestListenerIndex = FindClosestListener(AudioDevice->Listeners);
@@ -221,7 +221,7 @@ void FActiveSound::UpdateWaveInstances( TArray<FWaveInstance*> &InWaveInstances,
 	// final value that is correct
 	UpdateAdjustVolumeMultiplier(DeltaTime);
 	ParseParams.VolumeMultiplier = VolumeMultiplier * Sound->GetVolumeMultiplier() * CurrentAdjustVolumeMultiplier * AudioDevice->TransientMasterVolume * (bOccluded ? 0.5f : 1.0f) * ConcurrencyVolumeScale;
-	ParseParams.VolumeWeightedPriorityScale = VolumeWeightedPriorityScale;
+	ParseParams.Priority = Priority;
 	ParseParams.Pitch *= PitchMultiplier * Sound->GetPitchMultiplier();
 	ParseParams.HighFrequencyGain *= HighFrequencyGainMultiplier;
 
@@ -254,18 +254,7 @@ void FActiveSound::UpdateWaveInstances( TArray<FWaveInstance*> &InWaveInstances,
 	{
 		if (bHasAttenuationSettings)
 		{
-			AttenuationSettings.ApplyAttenuation(ParseParams.Transform, Listener.Transform.GetTranslation(), ParseParams.Volume, ParseParams.HighFrequencyGain);
-			ParseParams.OmniRadius = AttenuationSettings.OmniRadius;
-			ParseParams.StereoSpread = AttenuationSettings.StereoSpread;
-			ParseParams.bUseSpatialization = AttenuationSettings.bSpatialize;
-			if (AttenuationSettings.SpatializationAlgorithm == SPATIALIZATION_Default && AudioDevice->IsHRTFEnabledForAll())
-			{
-				ParseParams.SpatializationAlgorithm = SPATIALIZATION_HRTF;
-			}
-			else
-			{
-				ParseParams.SpatializationAlgorithm = AttenuationSettings.SpatializationAlgorithm;
-			}
+			ApplyAttenuation(ParseParams, ClosestListener);
 		}
 
 
@@ -668,3 +657,155 @@ void FActiveSound::CollectAttenuationShapesForVisualization(TMultiMap<EAttenuati
 		}
 	}
 }
+
+void FActiveSound::GetAttenuationListenerData(FAttenuationListenerData& ListenerData, const FSoundParseParameters& ParseParams, const FListener& Listener) const
+{
+	// Only perform this calculation once and explicitly require a reset.
+	if (ListenerData.bDataComputed)
+	{
+		return;
+	}
+
+	const FTransform& SoundTransform = ParseParams.Transform;
+	const FVector& ListenerLocation = Listener.Transform.GetTranslation();
+
+	FVector ListenerToSound = SoundTransform.GetTranslation() - ListenerLocation;
+	ListenerToSound.ToDirectionAndLength(ListenerData.ListenerToSoundDir, ListenerData.ListenerToSoundDistance);
+
+	ListenerData.AttenuationDistance = 0.0f;
+
+	if ((AttenuationSettings.bAttenuate && AttenuationSettings.AttenuationShape == EAttenuationShape::Sphere) || AttenuationSettings.bAttenuateWithLPF)
+	{
+		ListenerData.AttenuationDistance = FMath::Max(ListenerData.ListenerToSoundDistance - AttenuationSettings.AttenuationShapeExtents.X, 0.f);
+	}
+
+	ListenerData.bDataComputed = true;
+}
+
+void FActiveSound::ApplyAttenuation(FSoundParseParameters& ParseParams, const FListener& Listener)
+{
+	float& Volume = ParseParams.Volume;
+	const FTransform& SoundTransform = ParseParams.Transform;
+	const FVector& ListenerLocation = Listener.Transform.GetTranslation();
+	float& HighFrequencyGain = ParseParams.HighFrequencyGain;
+
+	FAttenuationListenerData ListenerData;
+
+	// Init focus factor to 1.0f in case this sound is ignoring focus
+	float FocusFactor = 1.0f;
+
+	// How much we want to scale the listener-emitter distance based on focus
+	float DistanceScale = 1.0f;
+
+	// Computer the focus factor if needed
+	check(Sound);
+	if (AttenuationSettings.bSpatialize && AttenuationSettings.bEnableListenerFocus && !Sound->bIgnoreFocus)
+	{
+		// Get the attenuation listener data
+		GetAttenuationListenerData(ListenerData, ParseParams, Listener);
+
+		// Compute the focus factor
+		const FVector& ListenerForwardDir = Listener.Transform.GetUnitAxis(EAxis::X);
+		const FVector& ListenerToSoundDir = ListenerData.ListenerToSoundDir;
+		const float FocusDotProduct = FVector::DotProduct(ListenerForwardDir, ListenerToSoundDir);
+		const float FocusAngle = FMath::RadiansToDegrees(FMath::Acos(FocusDotProduct));
+
+		const float FocusAzimuth = FMath::Clamp(AttenuationSettings.FocusAzimuth, 0.0f, 360.0f);
+		const float NonFocusAzimuth = FMath::Clamp(AttenuationSettings.NonFocusAzimuth, 0.0f, 360.0f);
+
+		if (FocusAzimuth != NonFocusAzimuth)
+		{
+			FocusFactor = (FocusAngle - FocusAzimuth) / (NonFocusAzimuth - FocusAzimuth);
+			FocusFactor = FMath::Clamp(FocusFactor, 0.0f, 1.0f);
+		}
+		else
+		{
+			FocusFactor = 0.0f;
+			if (FocusAngle < FocusAzimuth)
+			{
+				FocusFactor = 1.0f;
+			}
+		}
+
+		// Get the volume scale to apply the volume calculation based on the focus factor
+		const float NonFocusVolumeAttenuation = FMath::Clamp(AttenuationSettings.NonFocusVolumeAttenuation, 0.0f, 1.0f);
+		const float FocusVolumeAttenuation = FMath::Lerp(1.0f, NonFocusVolumeAttenuation, FocusFactor);
+		Volume *= FocusVolumeAttenuation;
+
+		// Scale the volume-weighted priority scale value we use for sorting this sound for voice-stealing
+		const float NonFocusPriorityScale = FMath::Max(AttenuationSettings.NonFocusPriorityScale, 0.0f);
+		const float FocusPriorityScale = FMath::Max(AttenuationSettings.FocusPriorityScale, 0.0f);
+		const float PriorityScale = FMath::Lerp(FocusPriorityScale, NonFocusPriorityScale, FocusFactor);
+		ParseParams.Priority *= PriorityScale;
+
+		// Get the distance scale to use when computing distance-calculations for 3d atttenuation
+		const float FocusDistanceScale = FMath::Max(AttenuationSettings.FocusDistanceScale, 0.0f);
+		const float NonFocusDistanceScale = FMath::Max(AttenuationSettings.NonFocusDistanceScale, 0.0f);
+		DistanceScale = FMath::Lerp(FocusDistanceScale, NonFocusDistanceScale, FocusFactor);
+
+		// Update the sound's distance scale value for checks based on max distance
+		Sound->SetFocusDistanceScale(DistanceScale);
+	}
+
+	// Attenuate the volume based on the model
+	if (AttenuationSettings.bAttenuate)
+	{
+		switch (AttenuationSettings.AttenuationShape)
+		{
+			case EAttenuationShape::Sphere:
+			{
+				// Update attenuation data in-case it hasn't been updated
+				GetAttenuationListenerData(ListenerData, ParseParams, Listener);
+				Volume *= AttenuationSettings.AttenuationEval(ListenerData.AttenuationDistance, AttenuationSettings.FalloffDistance, DistanceScale);
+				break;
+			}
+
+			case EAttenuationShape::Box:
+			Volume *= AttenuationSettings.AttenuationEvalBox(SoundTransform, ListenerLocation, DistanceScale);
+			break;
+
+			case EAttenuationShape::Capsule:
+			Volume *= AttenuationSettings.AttenuationEvalCapsule(SoundTransform, ListenerLocation, DistanceScale);
+			break;
+
+			case EAttenuationShape::Cone:
+			Volume *= AttenuationSettings.AttenuationEvalCone(SoundTransform, ListenerLocation, DistanceScale);
+			break;
+
+			default:
+			check(false);
+		}
+	}
+
+	// Attenuate with the low pass filter if necessary
+	if (AttenuationSettings.bAttenuateWithLPF)
+	{
+		GetAttenuationListenerData(ListenerData, ParseParams, Listener);
+		if (ListenerData.AttenuationDistance >= AttenuationSettings.LPFRadiusMax)
+		{
+			HighFrequencyGain = 0.0f;
+		}
+		// UsedLPFMinRadius is the point at which to start applying the low pass filter
+		else if (ListenerData.AttenuationDistance > AttenuationSettings.LPFRadiusMin)
+		{
+			HighFrequencyGain = 1.0f - ((ListenerData.AttenuationDistance - AttenuationSettings.LPFRadiusMin) / (AttenuationSettings.LPFRadiusMax - AttenuationSettings.LPFRadiusMin));
+		}
+		else
+		{
+			HighFrequencyGain = 1.0f;
+		}
+	}
+
+	ParseParams.OmniRadius = AttenuationSettings.OmniRadius;
+	ParseParams.StereoSpread = AttenuationSettings.StereoSpread;
+	ParseParams.bUseSpatialization |= AttenuationSettings.bSpatialize;
+
+	if (AttenuationSettings.SpatializationAlgorithm == SPATIALIZATION_Default && AudioDevice->IsHRTFEnabledForAll())
+	{
+		ParseParams.SpatializationAlgorithm = SPATIALIZATION_HRTF;
+	}
+	else
+	{
+		ParseParams.SpatializationAlgorithm = AttenuationSettings.SpatializationAlgorithm;
+	}
+}
\ No newline at end of file
diff --git a/Engine/Source/Runtime/Engine/Private/Actor.cpp b/Engine/Source/Runtime/Engine/Private/Actor.cpp
--- a/Engine/Source/Runtime/Engine/Private/Actor.cpp
+++ b/Engine/Source/Runtime/Engine/Private/Actor.cpp
@@ -587,7 +587,8 @@ void AActor::ProcessEvent(UFunction* Function, void* Parameters)
 	#else
 	const bool bAllowScriptExecution = GAllowActorScriptExecutionInEditor;
 	#endif
-	if( ((GetWorld() && (GetWorld()->AreActorsInitialized() || bAllowScriptExecution)) || HasAnyFlags(RF_ClassDefaultObject)) && !IsGarbageCollecting() )
+	UWorld* MyWorld = GetWorld();
+	if( ((MyWorld && (MyWorld->AreActorsInitialized() || bAllowScriptExecution)) || HasAnyFlags(RF_ClassDefaultObject)) && !IsGarbageCollecting() )
 	{
 #if !UE_BUILD_SHIPPING
 		if (!ProcessEventDelegate.IsBound() || !ProcessEventDelegate.Execute(this, Function, Parameters))
@@ -877,7 +878,8 @@ void AActor::GetComponentsBoundingCylinder(float& OutCollisionRadius, float& Out
 	if(IsTemplate())
 	{
 		// Editor code calls this function on default objects when placing them in the viewport, so no components will be registered in those cases.
-		if (!GetWorld() || !GetWorld()->IsGameWorld())
+		UWorld* MyWorld = GetWorld();
+		if (!MyWorld || !MyWorld->IsGameWorld())
 		{
 			bIgnoreRegistration = true;
 		}
@@ -1639,7 +1641,8 @@ void AActor::SetNetDormancy(ENetDormancy NewDormancy)
 		return;
 	}
 	
-	UNetDriver* NetDriver = GEngine->FindNamedNetDriver(GetWorld(), NetDriverName);
+	UWorld* MyWorld = GetWorld();
+	UNetDriver* NetDriver = GEngine->FindNamedNetDriver(MyWorld, NetDriverName);
 	if (NetDriver)
 	{
 		NetDormancy = NewDormancy;
@@ -1648,7 +1651,7 @@ void AActor::SetNetDormancy(ENetDormancy NewDormancy)
 		if (NewDormancy <= DORM_Awake)
 		{
 			// Since we are coming out of dormancy, make sure we are on the network actor list
-			GetWorld()->AddNetworkActor( this );
+			MyWorld->AddNetworkActor( this );
 
 			NetDriver->FlushActorDormancy(this);
 		}
@@ -1732,9 +1735,9 @@ void AActor::RouteEndPlay(const EEndPlayReason::Type EndPlayReason)
 
 void AActor::EndPlay(const EEndPlayReason::Type EndPlayReason)
 {
-	if (bActorHasBegunPlay)
+	if (ActorHasBegunPlay == EActorBeginPlayState::HasBegunPlay)
 	{
-		bActorHasBegunPlay = false;
+		ActorHasBegunPlay = EActorBeginPlayState::HasNotBegunPlay;
 
 		// Dispatch the blueprint events
 		ReceiveEndPlay(EndPlayReason);
@@ -2662,7 +2665,7 @@ void AActor::PostActorConstruction()
 				FRotator AdjustedRotation = GetActorRotation();
 				if (World->FindTeleportSpot(this, AdjustedLocation, AdjustedRotation))
 				{
-					SetActorLocationAndRotation(AdjustedLocation, AdjustedRotation);
+					SetActorLocationAndRotation(AdjustedLocation, AdjustedRotation, false, nullptr, ETeleportType::TeleportPhysics);
 				}
 			}
 			break;
@@ -2673,7 +2676,7 @@ void AActor::PostActorConstruction()
 				FRotator AdjustedRotation = GetActorRotation();
 				if (World->FindTeleportSpot(this, AdjustedLocation, AdjustedRotation))
 				{
-					SetActorLocationAndRotation(AdjustedLocation, AdjustedRotation);
+					SetActorLocationAndRotation(AdjustedLocation, AdjustedRotation, false, nullptr, ETeleportType::TeleportPhysics);
 				}
 				else
 				{
@@ -2739,9 +2742,12 @@ void AActor::SetReplicates(bool bInReplicates)
 { 
 	if (Role == ROLE_Authority || bInReplicates == false)
 	{
-		if (bReplicates == false && bInReplicates == true && GetWorld() != NULL)		// GetWorld will return NULL on CDO, FYI
+		if (bReplicates == false && bInReplicates == true)
 		{
-			GetWorld()->AddNetworkActor(this);
+			if (UWorld* MyWorld = GetWorld())		// GetWorld will return NULL on CDO, FYI
+			{
+				MyWorld->AddNetworkActor(this);
+			}
 		}
 
 		RemoteRole = (bInReplicates ? ROLE_SimulatedProxy : ROLE_None);
@@ -2787,9 +2793,13 @@ void AActor::PostNetInit()
 	}
 	check(RemoteRole == ROLE_Authority);
 
-	if (!HasActorBegunPlay() && GetWorld() && GetWorld()->HasBegunPlay())
+	if (!HasActorBegunPlay())
 	{
-		BeginPlay();
+		const UWorld* MyWorld = GetWorld();
+		if (MyWorld && MyWorld->HasBegunPlay())
+		{
+			BeginPlay();
+		}
 	}
 
 	UpdateOverlaps();
@@ -2816,7 +2826,7 @@ void AActor::SwapRolesForReplay()
 
 void AActor::BeginPlay()
 {
-	ensure(!bActorHasBegunPlay);
+	ensure(ActorHasBegunPlay == EActorBeginPlayState::HasNotBegunPlay);
 	SetLifeSpan( InitialLifeSpan );
 	RegisterAllActorTickFunctions(true, false); // Components are done below.
 
@@ -2841,9 +2851,10 @@ void AActor::BeginPlay()
 		}
 	}
 
+	ActorHasBegunPlay = EActorBeginPlayState::BeginningPlay;
 	ReceiveBeginPlay();
 
-	bActorHasBegunPlay = true;
+	ActorHasBegunPlay = EActorBeginPlayState::HasBegunPlay;
 }
 
 void AActor::EnableInput(APlayerController* PlayerController)
@@ -3867,27 +3878,28 @@ void AActor::DrawDebugComponents(FColor const& BaseColor) const
 	TInlineComponentArray<USceneComponent*> Components;
 	GetComponents(Components);
 
+	UWorld* MyWorld = GetWorld();
+
 	for(int32 ComponentIndex = 0; ComponentIndex < Components.Num(); ComponentIndex++)
 	{
 		USceneComponent const* const Component = Components[ComponentIndex]; 
 
-			FVector const Loc = Component->GetComponentLocation();
-			FRotator const Rot = Component->GetComponentRotation();
+		FVector const Loc = Component->GetComponentLocation();
+		FRotator const Rot = Component->GetComponentRotation();
 
-			// draw coord system at component loc
-			DrawDebugCoordinateSystem(GetWorld(), Loc, Rot, 10.f);
+		// draw coord system at component loc
+		DrawDebugCoordinateSystem(MyWorld, Loc, Rot, 10.f);
 
-			// draw line from me to my parent
-			USceneComponent const* const ParentComponent = Cast<USceneComponent>(Component->AttachParent);
-			if (ParentComponent)
-			{
-				DrawDebugLine(GetWorld(), ParentComponent->GetComponentLocation(), Loc, BaseColor);
-			}
-
-			// draw component name
-			DrawDebugString(GetWorld(), Loc+FVector(0,0,32), *Component->GetName());
+		// draw line from me to my parent
+		USceneComponent const* const ParentComponent = Cast<USceneComponent>(Component->AttachParent);
+		if (ParentComponent)
+		{
+			DrawDebugLine(MyWorld, ParentComponent->GetComponentLocation(), Loc, BaseColor);
 		}
 
+		// draw component name
+		DrawDebugString(MyWorld, Loc+FVector(0,0,32), *Component->GetName());
+	}
 }
 
 
diff --git a/Engine/Source/Runtime/Engine/Private/ActorEditor.cpp b/Engine/Source/Runtime/Engine/Private/ActorEditor.cpp
--- a/Engine/Source/Runtime/Engine/Private/ActorEditor.cpp
+++ b/Engine/Source/Runtime/Engine/Private/ActorEditor.cpp
@@ -181,7 +181,7 @@ void AActor::PostEditMove(bool bFinished)
 
 bool AActor::ReregisterComponentsWhenModified() const
 {
-	return !IsTemplate() && ( GetOutermost()->PackageFlags & PKG_PlayInEditor ) == 0 && GetWorld() != nullptr;
+	return !IsTemplate() && !GetOutermost()->HasAnyPackageFlags(PKG_PlayInEditor) && GetWorld() != nullptr;
 }
 
 void AActor::DebugShowComponentHierarchy(  const TCHAR* Info, bool bShowPosition )
diff --git a/Engine/Source/Runtime/Engine/Private/Animation/AnimBlueprintGeneratedClass.cpp b/Engine/Source/Runtime/Engine/Private/Animation/AnimBlueprintGeneratedClass.cpp
--- a/Engine/Source/Runtime/Engine/Private/Animation/AnimBlueprintGeneratedClass.cpp
+++ b/Engine/Source/Runtime/Engine/Private/Animation/AnimBlueprintGeneratedClass.cpp
@@ -130,6 +130,11 @@ void FAnimBlueprintDebugData::RecordNodeVisit(int32 TargetNodeIndex, int32 Sourc
 	new (UpdatedNodesThisFrame) FNodeVisit(SourceNodeIndex, TargetNodeIndex, BlendWeight);
 }
 
+void FAnimBlueprintDebugData::RecordNodeVisitArray(const TArray<FNodeVisit>& Nodes)
+{
+	UpdatedNodesThisFrame.Append(Nodes);
+}
+
 /////////////////////////////////////////////////////
 // FBinaryObjectWriter
 
diff --git a/Engine/Source/Runtime/Engine/Private/Animation/AnimComposite.cpp b/Engine/Source/Runtime/Engine/Private/Animation/AnimComposite.cpp
--- a/Engine/Source/Runtime/Engine/Private/Animation/AnimComposite.cpp
+++ b/Engine/Source/Runtime/Engine/Private/Animation/AnimComposite.cpp
@@ -27,9 +27,9 @@ void UAnimComposite::ReplaceReferredAnimations(const TMap<UAnimSequence*, UAnimS
 }
 #endif
 
-void UAnimComposite::OnAssetPlayerTickedInternal(FAnimAssetTickContext &Context, const float PreviousTime, const float MoveDelta, const FAnimTickRecord &Instance, class UAnimInstance* InstanceOwner) const
+void UAnimComposite::HandleAssetPlayerTickedInternal(FAnimAssetTickContext &Context, const float PreviousTime, const float MoveDelta, const FAnimTickRecord &Instance, struct FAnimNotifyQueue& NotifyQueue) const
 {
-	Super::OnAssetPlayerTickedInternal(Context, PreviousTime, MoveDelta, Instance, InstanceOwner);
+	Super::HandleAssetPlayerTickedInternal(Context, PreviousTime, MoveDelta, Instance, NotifyQueue);
 
 	ExtractRootMotionFromTrack(AnimationTrack, PreviousTime, PreviousTime + MoveDelta, Context.RootMotionMovementParams);
 }
