Copying //UE4/Dev-Framework to Dev-Main (//UE4/Dev-Main)

diff --git a/Engine/Source/Runtime/Engine/Private/Animation/AnimInstance.cpp b/Engine/Source/Runtime/Engine/Private/Animation/AnimInstance.cpp
--- a/Engine/Source/Runtime/Engine/Private/Animation/AnimInstance.cpp
+++ b/Engine/Source/Runtime/Engine/Private/Animation/AnimInstance.cpp
@@ -23,6 +23,8 @@
 #include "Animation/AnimNode_TransitionResult.h"
 #include "Animation/AnimNode_AssetPlayerBase.h"
 #include "Animation/AnimInstance.h"
+#include "Animation/AnimInstanceProxy.h"
+
 /** Anim stats */
 
 DEFINE_STAT(STAT_CalcSkelMeshBounds);
@@ -65,17 +67,8 @@ DEFINE_STAT(STAT_UpdateTriMeshVertices);
 
 DEFINE_STAT(STAT_AnimGameThreadTime);
 
-DECLARE_CYCLE_STAT(TEXT("TickAssetPlayerInstances"), STAT_TickAssetPlayerInstances, STATGROUP_Anim);
-DECLARE_CYCLE_STAT(TEXT("TickAssetPlayerInstance"), STAT_TickAssetPlayerInstance, STATGROUP_Anim);
-
-
-#define DO_ANIMSTAT_PROCESSING(StatName) DEFINE_STAT(STAT_ ## StatName)
-#include "AnimMTStats.h"
-#undef DO_ANIMSTAT_PROCESSING
-
-#define DO_ANIMSTAT_PROCESSING(StatName) DEFINE_STAT(STAT_ ## StatName ## _WorkerThread)
-#include "AnimMTStats.h"
-#undef DO_ANIMSTAT_PROCESSING
+DEFINE_STAT(STAT_TickAssetPlayerInstances);
+DEFINE_STAT(STAT_TickAssetPlayerInstance);
 
 // Define AnimNotify
 DEFINE_LOG_CATEGORY(LogAnimNotify);
@@ -84,6 +77,7 @@ DEFINE_LOG_CATEGORY(LogAnimNotify);
 
 extern TAutoConsoleVariable<int32> CVarUseParallelAnimUpdate;
 extern TAutoConsoleVariable<int32> CVarUseParallelAnimationEvaluation;
+extern TAutoConsoleVariable<int32> CVarForceUseParallelAnimUpdate;
 
 /////////////////////////////////////////////////////
 // UAnimInstance
@@ -92,36 +86,21 @@ extern TAutoConsoleVariable<int32> CVarUseParallelAnimationEvaluation;
 UAnimInstance::UAnimInstance(const FObjectInitializer& ObjectInitializer)
 	: Super(ObjectInitializer)
 {
-	RootNode = NULL;
+PRAGMA_DISABLE_DEPRECATION_WARNINGS
+	RootNode = nullptr;
+PRAGMA_ENABLE_DEPRECATION_WARNINGS
 	RootMotionMode = ERootMotionMode::RootMotionFromMontagesOnly;
-	SyncGroupWriteIndex = 0;
-	CurrentDeltaSeconds = 0.0f;
-	CurrentTimeDilation = 1.0f;
 	bNeedsUpdate = false;
 }
 
 void UAnimInstance::MakeSequenceTickRecord(FAnimTickRecord& TickRecord, class UAnimSequenceBase* Sequence, bool bLooping, float PlayRate, float FinalBlendWeight, float& CurrentTime, FMarkerTickRecord& MarkerTickRecord) const
 {
-	TickRecord.SourceAsset = Sequence;
-	TickRecord.TimeAccumulator = &CurrentTime;
-	TickRecord.MarkerTickRecord = &MarkerTickRecord;
-	TickRecord.PlayRateMultiplier = PlayRate;
-	TickRecord.EffectiveBlendWeight = FinalBlendWeight;
-	TickRecord.bLooping = bLooping;
+	GetProxyOnGameThread().MakeSequenceTickRecord(TickRecord, Sequence, bLooping, PlayRate, FinalBlendWeight, CurrentTime, MarkerTickRecord);
 }
 
 void UAnimInstance::MakeBlendSpaceTickRecord(FAnimTickRecord& TickRecord, class UBlendSpaceBase* BlendSpace, const FVector& BlendInput, TArray<FBlendSampleData>& BlendSampleDataCache, FBlendFilter& BlendFilter, bool bLooping, float PlayRate, float FinalBlendWeight, float& CurrentTime, FMarkerTickRecord& MarkerTickRecord) const
 {
-	TickRecord.SourceAsset = BlendSpace;
-	TickRecord.BlendSpace.BlendSpacePositionX = BlendInput.X;
-	TickRecord.BlendSpace.BlendSpacePositionY = BlendInput.Y;
-	TickRecord.BlendSpace.BlendSampleDataCache = &BlendSampleDataCache;
-	TickRecord.BlendSpace.BlendFilter = &BlendFilter;
-	TickRecord.TimeAccumulator = &CurrentTime;
-	TickRecord.MarkerTickRecord = &MarkerTickRecord;
-	TickRecord.PlayRateMultiplier = PlayRate;
-	TickRecord.EffectiveBlendWeight = FinalBlendWeight;
-	TickRecord.bLooping = bLooping;
+	GetProxyOnGameThread().MakeBlendSpaceTickRecord(TickRecord, BlendSpace, BlendInput, BlendSampleDataCache, BlendFilter, bLooping, PlayRate, FinalBlendWeight, CurrentTime, MarkerTickRecord);
 }
 
 // this is only used by montage marker based sync
@@ -140,40 +119,17 @@ void UAnimInstance::MakeMontageTickRecord(FAnimTickRecord& TickRecord, class UAn
 
 void UAnimInstance::SequenceAdvanceImmediate(UAnimSequenceBase* Sequence, bool bLooping, float PlayRate, float DeltaSeconds, float& CurrentTime, FMarkerTickRecord& MarkerTickRecord)
 {
-	FAnimTickRecord TickRecord;
-	MakeSequenceTickRecord(TickRecord, Sequence, bLooping, PlayRate, /*FinalBlendWeight=*/ 1.0f, CurrentTime, MarkerTickRecord);
-
-	FAnimAssetTickContext TickContext(DeltaSeconds, RootMotionMode, true);
-	TickRecord.SourceAsset->TickAssetPlayerInstance(TickRecord, this, TickContext);
+	GetProxyOnGameThread().SequenceAdvanceImmediate(Sequence, bLooping, PlayRate, DeltaSeconds, CurrentTime, MarkerTickRecord);
 }
 
 void UAnimInstance::BlendSpaceAdvanceImmediate(class UBlendSpaceBase* BlendSpace, const FVector& BlendInput, TArray<FBlendSampleData>& BlendSampleDataCache, FBlendFilter& BlendFilter, bool bLooping, float PlayRate, float DeltaSeconds, float& CurrentTime, FMarkerTickRecord& MarkerTickRecord)
 {
-	FAnimTickRecord TickRecord;
-	MakeBlendSpaceTickRecord(TickRecord, BlendSpace, BlendInput, BlendSampleDataCache, BlendFilter, bLooping, PlayRate, /*FinalBlendWeight=*/ 1.0f, CurrentTime, MarkerTickRecord);
-	
-	FAnimAssetTickContext TickContext(DeltaSeconds, RootMotionMode, true);
-	TickRecord.SourceAsset->TickAssetPlayerInstance(TickRecord, this, TickContext);
+	GetProxyOnGameThread().BlendSpaceAdvanceImmediate(BlendSpace, BlendInput, BlendSampleDataCache, BlendFilter, bLooping, PlayRate, DeltaSeconds, CurrentTime, MarkerTickRecord);
 }
 
-// Creates an uninitialized tick record in the list for the correct group or the ungrouped array.  If the group is valid, OutSyncGroupPtr will point to the group.
 FAnimTickRecord& UAnimInstance::CreateUninitializedTickRecord(int32 GroupIndex, FAnimGroupInstance*& OutSyncGroupPtr)
 {
-	// Find or create the sync group if there is one
-	OutSyncGroupPtr = NULL;
-	if (GroupIndex >= 0)
-	{
-		TArray<FAnimGroupInstance>& SyncGroups = SyncGroupArrays[GetSyncGroupWriteIndex()];
-		while (SyncGroups.Num() <= GroupIndex)
-		{
-			new (SyncGroups) FAnimGroupInstance();
-		}
-		OutSyncGroupPtr = &(SyncGroups[GroupIndex]);
-	}
-
-	// Create the record
-	FAnimTickRecord* TickRecord = new ((OutSyncGroupPtr != NULL) ? OutSyncGroupPtr->ActivePlayers : UngroupedActivePlayerArrays[GetSyncGroupWriteIndex()]) FAnimTickRecord();
-	return *TickRecord;
+	return GetProxyOnGameThread().CreateUninitializedTickRecord(GroupIndex, OutSyncGroupPtr);
 }
 
 AActor* UAnimInstance::GetOwningActor() const
@@ -224,32 +180,6 @@ void UAnimInstance::InitializeAnimation()
 
 	if (UAnimBlueprintGeneratedClass* AnimBlueprintClass = Cast<UAnimBlueprintGeneratedClass>(GetClass()))
 	{
-		// Grab a pointer to the root node
-		if (AnimBlueprintClass->RootAnimNodeProperty != NULL)
-		{
-			RootNode = AnimBlueprintClass->RootAnimNodeProperty->ContainerPtrToValuePtr<FAnimNode_Base>(this);
-		}
-		else
-		{
-			RootNode = NULL;
-		}
-
-		// if no mesh, use Blueprint Skeleton
-		if (CurrentSkeleton == NULL)
-		{
-			CurrentSkeleton = AnimBlueprintClass->TargetSkeleton;
-		}
-
-#if WITH_EDITORONLY_DATA
-		if (UAnimBlueprint* Blueprint = Cast<UAnimBlueprint>(AnimBlueprintClass->ClassGeneratedBy))
-		{
-			if (Blueprint->Status == BS_Error)
-			{
-				RootNode = NULL;
-			}
-		}
-#endif
-
 #if WITH_EDITOR
 		LifeTimer = 0.0;
 		CurrentLifeTimerScrubPosition = 0.0;
@@ -268,26 +198,29 @@ void UAnimInstance::InitializeAnimation()
 	// before initialize, need to recalculate required bone list
 	RecalcRequiredBones();
 
-	ReinitializeSlotNodes();
+	GetProxyOnGameThread().Initialize(this);
+
 	ClearMorphTargets();
 	NativeInitializeAnimation();
 	BlueprintInitializeAnimation();
 
-	if (RootNode != NULL)
-	{
-		InitializationCounter.Increment();
-		FAnimationInitializeContext InitContext(this);
-		RootNode->Initialize(InitContext);
-	}
+	GetProxyOnGameThread().InitializeRootNode();
+
+PRAGMA_DISABLE_DEPRECATION_WARNINGS
+	// @todo: remove after deprecation
+	RootNode = GetProxyOnGameThread().GetRootNode();
+PRAGMA_ENABLE_DEPRECATION_WARNINGS
 
 	// we can bind rules & events now the graph has been initialized
-	BindNativeDelegates();
+	GetProxyOnGameThread().BindNativeDelegates();
 }
 
 void UAnimInstance::UninitializeAnimation()
 {
 	NativeUninitializeAnimation();
 
+	GetProxyOnGameThread().Uninitialize(this);
+
 	StopAllMontages(0.f);
 
 	for(int32 Index = 0; Index < MontageInstances.Num(); ++Index)
@@ -316,10 +249,8 @@ void UAnimInstance::UninitializeAnimation()
 
 	ActiveAnimNotifyState.Reset();
 	EventCurves.Reset();
-	MorphTargetCurves.Reset();
-	MaterialParameterCurves.Reset();
 	MaterialParamatersToClear.Reset();
-	AnimNotifies.Reset();
+	NotifyQueue.Reset(SkelMeshComp);
 }
 
 #if WITH_EDITORONLY_DATA
@@ -351,13 +282,7 @@ bool UAnimInstance::UpdateSnapshotAndSkipRemainingUpdate()
 // update animation node, if you have your own anim instance, you could override this to update extra nodes if you'd like
 void UAnimInstance::UpdateAnimationNode(float DeltaSeconds)
 {
-	// Update the anim graph
-	if(RootNode != NULL)
-	{
-		UpdateCounter.Increment();
-		FAnimationUpdateContext UpdateContext(this, DeltaSeconds);
-		RootNode->Update(UpdateContext);
-	}
+	GetProxyOnGameThread().UpdateAnimationNode(DeltaSeconds);
 }
 
 void UAnimInstance::UpdateMontage(float DeltaSeconds)
@@ -382,7 +307,7 @@ void UAnimInstance::UpdateMontage(float DeltaSeconds)
 			if (ensure (GroupIndexToUse != INDEX_NONE))
 			{
 				FAnimGroupInstance* SyncGroup;
-				FAnimTickRecord& TickRecord = CreateUninitializedTickRecord(GroupIndexToUse, /*out*/ SyncGroup);
+				FAnimTickRecord& TickRecord = GetProxyOnGameThread().CreateUninitializedTickRecord(GroupIndexToUse, /*out*/ SyncGroup);
 				MakeMontageTickRecord(TickRecord, MontageInstance->Montage, MontageInstance->GetPosition(), 
 					MontageInstance->GetPreviousPosition(), MontageInstance->GetDeltaMoved(), MontageInstance->GetWeight(), 
 					MontageInstance->MarkersPassedThisTick, MontageInstance->MarkerTickRecord);
@@ -405,11 +330,14 @@ void UAnimInstance::UpdateAnimation(float DeltaSeconds, bool bNeedsValidRootMoti
 	SCOPE_CYCLE_COUNTER(STAT_UpdateAnimation);
 	FScopeCycleCounterUObject AnimScope(this);
 
+	// acquire the proxy as we need to update
+	FAnimInstanceProxy& Proxy = GetProxyOnGameThread();
+
 #if WITH_EDITOR
 	if (GIsEditor)
 	{
 		// Reset the anim graph visualization
-		if (RootNode != NULL)
+		if (Proxy.HasRootNode())
 		{
 			if (UAnimBlueprintGeneratedClass* AnimBlueprintClass = Cast<UAnimBlueprintGeneratedClass>(GetClass()))
 			{
@@ -437,6 +365,11 @@ void UAnimInstance::UpdateAnimation(float DeltaSeconds, bool bNeedsValidRootMoti
 	{
 		SCOPE_CYCLE_COUNTER(STAT_NativeUpdateAnimation);
 		NativeUpdateAnimation(DeltaSeconds);
+
+PRAGMA_DISABLE_DEPRECATION_WARNINGS
+		// @todo: remove once deprecated - called for backwards-compatibility
+		NativeUpdateAnimation_WorkerThread(DeltaSeconds);
+PRAGMA_ENABLE_DEPRECATION_WARNINGS
 	}
 	{
 		SCOPE_CYCLE_COUNTER(STAT_BlueprintUpdateAnimation);
@@ -450,13 +383,14 @@ void UAnimInstance::UpdateAnimation(float DeltaSeconds, bool bNeedsValidRootMoti
 	if(NeedsImmediateUpdate(DeltaSeconds))
 	{
 		// cant use parallel update, so just do the work here
-		UpdateAnimationInternal(DeltaSeconds);
+		Proxy.UpdateAnimation();
 		PostUpdateAnimation();
 	}
 	else if(bNeedsValidRootMotion && RootMotionMode == ERootMotionMode::RootMotionFromMontagesOnly)
 	{
 		// We may have just partially blended root motion, so make it up to 1 by
 		// blending in identity too
+		FRootMotionMovementParams& ExtractedRootMotion = Proxy.GetExtractedRootMotion();
 		if (ExtractedRootMotion.bHasRootMotion)
 		{
 			ExtractedRootMotion.MakeUpToFullWeight();
@@ -468,178 +402,38 @@ void UAnimInstance::PreUpdateAnimation(float DeltaSeconds)
 {
 	SCOPE_CYCLE_COUNTER(STAT_PreUpdateAnimation);
 
-	CurrentDeltaSeconds = DeltaSeconds;
-
-	const UWorld* World = GetSkelMeshComponent()->GetWorld();
-	check(World->GetWorldSettings());
-	CurrentTimeDilation = World->GetWorldSettings()->GetEffectiveTimeDilation();
 	bNeedsUpdate = true;
 
-	AnimNotifies.Reset();
-
-	ClearSlotNodeWeights();
-
-	// Reset the player tick list (but keep it presized)
-	TArray<FAnimTickRecord>& UngroupedActivePlayers = UngroupedActivePlayerArrays[GetSyncGroupWriteIndex()];
-	UngroupedActivePlayers.Reset();
+	NotifyQueue.Reset(GetSkelMeshComponent());
 
-	TArray<FAnimGroupInstance>& SyncGroups = SyncGroupArrays[GetSyncGroupWriteIndex()];
-	for (int32 GroupIndex = 0; GroupIndex < SyncGroups.Num(); ++GroupIndex)
-	{
-		SyncGroups[GroupIndex].Reset();
-	}
-}
-
-void UAnimInstance::UpdateAnimationInternal(float DeltaSeconds)
-{
-	ANIM_MT_SCOPE_CYCLE_COUNTER(UpdateAnimationInternal, IsRunningParallelEvaluation());
-	FScopeCycleCounterUObject AnimScope(this);
-
-	// update blueprint and native update
-	{
-		SCOPE_CYCLE_COUNTER(STAT_NativeUpdateAnimation);
-		NativeUpdateAnimation_WorkerThread(DeltaSeconds);
-	}
-
-	// update all nodes
-	UpdateAnimationNode(DeltaSeconds);
-
-	// tick all our active asset players
-	TickAssetPlayerInstances(DeltaSeconds);
-}
-
-void UAnimInstance::TickAssetPlayerInstances(float DeltaSeconds)
-{
-	SCOPE_CYCLE_COUNTER(STAT_TickAssetPlayerInstances);
-
-	// Handle all players inside sync groups
-	TArray<FAnimGroupInstance>& SyncGroups = SyncGroupArrays[GetSyncGroupWriteIndex()];
-	const TArray<FAnimGroupInstance>& PreviousSyncGroups = SyncGroupArrays[GetSyncGroupReadIndex()];
-	TArray<FAnimTickRecord>& UngroupedActivePlayers = UngroupedActivePlayerArrays[GetSyncGroupWriteIndex()];
-
-	for (int32 GroupIndex = 0; GroupIndex < SyncGroups.Num(); ++GroupIndex)
-	{
-		FAnimGroupInstance& SyncGroup = SyncGroups[GroupIndex];
-    
-		if (SyncGroup.ActivePlayers.Num() > 0)
-		{
-			const FAnimGroupInstance* PreviousGroup = PreviousSyncGroups.IsValidIndex(GroupIndex) ? &PreviousSyncGroups[GroupIndex] : nullptr;
-			SyncGroup.Prepare(PreviousGroup);
-
-			UE_LOG(LogAnimMarkerSync, Log, TEXT("Ticking Group [%d] GroupLeader [%d]"), GroupIndex, SyncGroup.GroupLeaderIndex);
-
-			const bool bOnlyOneAnimationInGroup = SyncGroup.ActivePlayers.Num() == 1;
-
-			// Tick the group leader
-			FAnimAssetTickContext TickContext(DeltaSeconds, RootMotionMode, bOnlyOneAnimationInGroup, SyncGroup.ValidMarkers);
-			// initialize to invalidate first
-			check(SyncGroup.GroupLeaderIndex == INDEX_NONE);
-			int32 GroupLeaderIndex = 0;
-			for (; GroupLeaderIndex < SyncGroup.ActivePlayers.Num(); ++GroupLeaderIndex)
-			{
-				FAnimTickRecord& GroupLeader = SyncGroup.ActivePlayers[GroupLeaderIndex];
-				// if it has leader score
-				SCOPE_CYCLE_COUNTER(STAT_TickAssetPlayerInstance);
-				FScopeCycleCounterUObject Scope(GroupLeader.SourceAsset);
-				GroupLeader.SourceAsset->TickAssetPlayerInstance(GroupLeader, this, TickContext);
-
-				if (RootMotionMode == ERootMotionMode::RootMotionFromEverything && TickContext.RootMotionMovementParams.bHasRootMotion)
-				{
-					ExtractedRootMotion.AccumulateWithBlend(TickContext.RootMotionMovementParams.RootMotionTransform, GroupLeader.EffectiveBlendWeight);
-				}
-
-				// if we're not using marker based sync, we don't care, get out
-				if (TickContext.CanUseMarkerPosition() == false)
-				{
-					SyncGroup.GroupLeaderIndex = GroupLeaderIndex;
-					break;
-				}
-				// otherwise, the new position should contain the valid position for end, otherwise, we don't know where to sync to
-				else if (TickContext.MarkerTickContext.IsMarkerSyncEndValid())
-				{
-					// if this leader contains correct position, break
-					SyncGroup.MarkerTickContext = TickContext.MarkerTickContext;
-					SyncGroup.GroupLeaderIndex = GroupLeaderIndex;
-					UE_LOG(LogAnimMarkerSync, Log, TEXT("Previous Sync Group Makrer Tick Context :\n%s"), *SyncGroup.MarkerTickContext.ToString());
-					UE_LOG(LogAnimMarkerSync, Log, TEXT("New Sync Group Makrer Tick Context :\n%s"), *TickContext.MarkerTickContext.ToString());
-					break;
-				}
-				else
-				{
-					SyncGroup.GroupLeaderIndex = GroupLeaderIndex;
-					UE_LOG(LogAnimMarkerSync, Log, TEXT("Invalid position from Leader %d. Trying next leader"), GroupLeaderIndex);
-				}
-			} 
-
-			check(SyncGroup.GroupLeaderIndex != INDEX_NONE);
-			// we found leader
-			SyncGroup.Finalize(PreviousGroup);
-
-			// Update everything else to follow the leader, if there is more followers
-			if (SyncGroup.ActivePlayers.Num() > GroupLeaderIndex + 1)
-			{
-				// if we don't have a good leader, no reason to convert to follower
-				// tick as leader
-				TickContext.ConvertToFollower();
-    
-				for (int32 TickIndex = GroupLeaderIndex + 1; TickIndex < SyncGroup.ActivePlayers.Num(); ++TickIndex)
-				{
-					FAnimTickRecord& AssetPlayer = SyncGroup.ActivePlayers[TickIndex];
-					{
-						SCOPE_CYCLE_COUNTER(STAT_TickAssetPlayerInstance);
-						FScopeCycleCounterUObject Scope(AssetPlayer.SourceAsset);
-						TickContext.RootMotionMovementParams.Clear();
-						AssetPlayer.SourceAsset->TickAssetPlayerInstance(AssetPlayer, this, TickContext);
-					}
-					if (RootMotionMode == ERootMotionMode::RootMotionFromEverything && TickContext.RootMotionMovementParams.bHasRootMotion)
-					{
-						ExtractedRootMotion.AccumulateWithBlend(TickContext.RootMotionMovementParams.RootMotionTransform, AssetPlayer.EffectiveBlendWeight);
-					}
-				}
-			}
-		}
-	}
-
-	// Handle the remaining ungrouped animation players
-	for (int32 TickIndex = 0; TickIndex < UngroupedActivePlayers.Num(); ++TickIndex)
-	{
-		FAnimTickRecord& AssetPlayerToTick = UngroupedActivePlayers[TickIndex];
-		static TArray<FName> TempNames;
-		const TArray<FName>* UniqueNames = AssetPlayerToTick.SourceAsset->GetUniqueMarkerNames();
-		const TArray<FName>& ValidMarkers = UniqueNames ? *UniqueNames : TempNames;
-
-		const bool bOnlyOneAnimationInGroup = true;
-		FAnimAssetTickContext TickContext(DeltaSeconds, RootMotionMode, bOnlyOneAnimationInGroup, ValidMarkers);
-		{
-			SCOPE_CYCLE_COUNTER(STAT_TickAssetPlayerInstance);
-			FScopeCycleCounterUObject Scope(AssetPlayerToTick.SourceAsset);
-			AssetPlayerToTick.SourceAsset->TickAssetPlayerInstance(AssetPlayerToTick, this, TickContext);
-		}
-		if (RootMotionMode == ERootMotionMode::RootMotionFromEverything && TickContext.RootMotionMovementParams.bHasRootMotion)
-		{
-			ExtractedRootMotion.AccumulateWithBlend(TickContext.RootMotionMovementParams.RootMotionTransform, AssetPlayerToTick.EffectiveBlendWeight);
-		}
-	}
+	GetProxyOnGameThread().PreUpdate(this, DeltaSeconds);
 }
 
 void UAnimInstance::PostUpdateAnimation()
 {
 	SCOPE_CYCLE_COUNTER(STAT_PostUpdateAnimation);
+	check(!IsRunningParallelEvaluation());
+
+	// acquire the proxy as we need to update
+	FAnimInstanceProxy& Proxy = GetProxyOnGameThread();
 
 	bNeedsUpdate = false;
 
+	Proxy.PostUpdate(this);
+
 	// We may have just partially blended root motion, so make it up to 1 by
 	// blending in identity too
+	FRootMotionMovementParams& ExtractedRootMotion = Proxy.GetExtractedRootMotion();
 	if (ExtractedRootMotion.bHasRootMotion)
 	{
 		ExtractedRootMotion.MakeUpToFullWeight();
 	}
 
 	// flip read/write index
-	TickSyncGroupWriteIndex();
+	Proxy.TickSyncGroupWriteIndex();
 
 	// now trigger Notifies
-	TriggerAnimNotifies(CurrentDeltaSeconds);
+	TriggerAnimNotifies(Proxy.GetDeltaSeconds());
 
 	// Trigger Montage end events after notifies. In case Montage ending ends abilities or other states, we make sure notifies are processed before montage events.
 	TriggerQueuedMontageEvents();
@@ -653,7 +447,7 @@ void UAnimInstance::PostUpdateAnimation()
 			{
 				if (Blueprint->GetObjectBeingDebugged() == this)
 				{
-					if ((CurrentLifeTimerScrubPosition == LifeTimer) && (CurrentDeltaSeconds > 0.0f))
+					if ((CurrentLifeTimerScrubPosition == LifeTimer) && (Proxy.GetDeltaSeconds() > 0.0f))
 					{
 						AnimBlueprintClass->GetAnimBlueprintDebugData().TakeSnapshot(this);
 					}
@@ -666,16 +460,18 @@ void UAnimInstance::PostUpdateAnimation()
 
 void UAnimInstance::ParallelUpdateAnimation()
 {
-	UpdateAnimationInternal(CurrentDeltaSeconds);
+	GetProxyOnAnyThread().UpdateAnimation();
 }
 
 bool UAnimInstance::NeedsImmediateUpdate(float DeltaSeconds) const
 {
+	const bool bUseParallelUpdateAnimation = bCanUseParallelUpdateAnimation || CVarForceUseParallelAnimUpdate.GetValueOnGameThread() != 0;
+
 	return
 		CVarUseParallelAnimUpdate.GetValueOnGameThread() == 0 ||
 		CVarUseParallelAnimationEvaluation.GetValueOnGameThread() == 0 ||
 		GetWorld()->IsServer() ||
-		!bCanUseParallelUpdateAnimation ||
+		!bUseParallelUpdateAnimation ||
 		DeltaSeconds == 0.0f ||
 		RootMotionMode == ERootMotionMode::RootMotionFromEverything;
 }
@@ -687,32 +483,52 @@ bool UAnimInstance::NeedsUpdate() const
 
 void UAnimInstance::EvaluateAnimation(FPoseContext& Output)
 {
-	ANIM_MT_SCOPE_CYCLE_COUNTER(EvaluateAnimInstance, IsRunningParallelEvaluation());
+	GetProxyOnGameThread().EvaluateAnimation(Output);
+}
 
-	// If bone caches have been invalidated, have AnimNodes refresh those.
-	if( bBoneCachesInvalidated && RootNode )
-	{
-		bBoneCachesInvalidated = false;
+bool UAnimInstance::ParallelCanEvaluate(const USkeletalMesh* InSkeletalMesh) const
+{
+	const FAnimInstanceProxy& Proxy = GetProxyOnAnyThread();
+	return Proxy.GetRequiredBones().IsValid() && (Proxy.GetRequiredBones().GetAsset() == InSkeletalMesh);
+}
 
-		CachedBonesCounter.Increment();
-		FAnimationCacheBonesContext UpdateContext(this);
-		RootNode->CacheBones(UpdateContext);
-	}
+void UAnimInstance::ParallelEvaluateAnimation(bool bForceRefPose, const USkeletalMesh* InSkeletalMesh, TArray<FTransform>& OutLocalAtoms, TArray<FActiveVertexAnim>& OutVertexAnims, FBlendedCurve& OutCurve)
+{
+	FAnimInstanceProxy& Proxy = GetProxyOnAnyThread();
 
-	// Evaluate native code if implemented, otherwise evaluate the node graph
-	if (!NativeEvaluateAnimation(Output))
+	if( !bForceRefPose )
 	{
-		if (RootNode != NULL)
+		// Create an evaluation context
+		FPoseContext EvaluationContext(&Proxy);
+		EvaluationContext.ResetToRefPose();
+			
+		// Run the anim blueprint
+		Proxy.EvaluateAnimation(EvaluationContext);
+		// Move the curves
+		OutCurve.MoveFrom(EvaluationContext.Curve);
+			
+		// can we avoid that copy?
+		if( EvaluationContext.Pose.GetNumBones() > 0 )
 		{
-			ANIM_MT_SCOPE_CYCLE_COUNTER(EvaluateAnimGraph, IsRunningParallelEvaluation());
-			EvaluationCounter.Increment();
-			RootNode->Evaluate(Output);
+			// Make sure rotations are normalized to account for accumulation of errors.
+			EvaluationContext.Pose.NormalizeRotations();
+			for (const FCompactPoseBoneIndex BoneIndex : EvaluationContext.Pose.ForEachBoneIndex())
+			{
+				FMeshPoseBoneIndex MeshPoseBoneIndex = EvaluationContext.Pose.GetBoneContainer().MakeMeshPoseIndex(BoneIndex);
+				OutLocalAtoms[MeshPoseBoneIndex.GetInt()] = EvaluationContext.Pose[BoneIndex];
+			}
 		}
 		else
 		{
-			Output.ResetToRefPose();
+			FAnimationRuntime::FillWithRefPose(OutLocalAtoms, Proxy.GetRequiredBones());
 		}
 	}
+	else
+	{
+		FAnimationRuntime::FillWithRefPose(OutLocalAtoms, Proxy.GetRequiredBones());
+	}
+
+	OutVertexAnims = FAnimationRuntime::UpdateActiveVertexAnims(InSkeletalMesh, Proxy.GetMorphTargetCurves(), Proxy.GetVertexAnims());
 }
 
 void UAnimInstance::PostEvaluateAnimation()
@@ -752,165 +568,32 @@ void UAnimInstance::NativeUninitializeAnimation()
 
 void UAnimInstance::AddNativeTransitionBinding(const FName& MachineName, const FName& PrevStateName, const FName& NextStateName, const FCanTakeTransition& NativeTransitionDelegate, const FName& TransitionName)
 {
-	NativeTransitionBindings.Add(FNativeTransitionBinding(MachineName, PrevStateName, NextStateName, NativeTransitionDelegate, TransitionName));
+	GetProxyOnGameThread().AddNativeTransitionBinding(MachineName, PrevStateName, NextStateName, NativeTransitionDelegate, TransitionName);
 }
 
 bool UAnimInstance::HasNativeTransitionBinding(const FName& MachineName, const FName& PrevStateName, const FName& NextStateName, FName& OutBindingName)
 {
-	for(const auto& Binding : NativeTransitionBindings)
-	{
-		if(Binding.MachineName == MachineName && Binding.PreviousStateName == PrevStateName && Binding.NextStateName == NextStateName)
-		{
-			IDelegateInstance* DelegateInstance = Binding.NativeTransitionDelegate.GetDelegateInstance();
-			if(DelegateInstance)
-			{
-				OutBindingName = DelegateInstance->GetFunctionName();
-			}
-#if WITH_EDITORONLY_DATA
-			// if the function wasnt forthcoming with a name, use the explicit name supplied in the binding
-			if(OutBindingName == NAME_None)
-			{
-				OutBindingName = Binding.TransitionName;
-			}
-#endif
-			return true;
-		}
-	}
-
-	return false;
+	return GetProxyOnGameThread().HasNativeTransitionBinding(MachineName, PrevStateName, NextStateName, OutBindingName);
 }
 
 void UAnimInstance::AddNativeStateEntryBinding(const FName& MachineName, const FName& StateName, const FOnGraphStateChanged& NativeEnteredDelegate)
 {
-	NativeStateEntryBindings.Add(FNativeStateBinding(MachineName, StateName, NativeEnteredDelegate));
+	GetProxyOnGameThread().AddNativeStateEntryBinding(MachineName, StateName, NativeEnteredDelegate);
 }
 	
 bool UAnimInstance::HasNativeStateEntryBinding(const FName& MachineName, const FName& StateName, FName& OutBindingName)
 {
-	for(const auto& Binding : NativeStateEntryBindings)
-	{
-		if(Binding.MachineName == MachineName && Binding.StateName == StateName)
-		{
-			IDelegateInstance* DelegateInstance = Binding.NativeStateDelegate.GetDelegateInstance();
-			if(DelegateInstance)
-			{
-				OutBindingName = DelegateInstance->GetFunctionName();
-			}
-			return true;
-		}
-	}
-
-	return false;
+	return GetProxyOnGameThread().HasNativeStateEntryBinding(MachineName, StateName, OutBindingName);
 }
 
 void UAnimInstance::AddNativeStateExitBinding(const FName& MachineName, const FName& StateName, const FOnGraphStateChanged& NativeExitedDelegate)
 {
-	NativeStateExitBindings.Add(FNativeStateBinding(MachineName, StateName, NativeExitedDelegate));
+	GetProxyOnGameThread().AddNativeStateExitBinding(MachineName, StateName, NativeExitedDelegate);
 }
 
 bool UAnimInstance::HasNativeStateExitBinding(const FName& MachineName, const FName& StateName, FName& OutBindingName)
 {
-	for(const auto& Binding : NativeStateExitBindings)
-	{
-		if(Binding.MachineName == MachineName && Binding.StateName == StateName)
-		{
-			IDelegateInstance* DelegateInstance = Binding.NativeStateDelegate.GetDelegateInstance();
-			if(DelegateInstance)
-			{
-				OutBindingName = DelegateInstance->GetFunctionName();
-			}
-			return true;
-		}
-	}
-
-	return false;
-}
-
-void UAnimInstance::BindNativeDelegates()
-{
-	auto ForEachStateLambda = [&](UAnimBlueprintGeneratedClass* AnimBlueprintGeneratedClass, const FName& MachineName, const FName& StateName, TFunctionRef<void (FAnimNode_StateMachine*, const FBakedAnimationState&, int32)> Predicate)
-	{
-		for(UStructProperty* Property : AnimBlueprintGeneratedClass->AnimNodeProperties)
-		{
-			if(Property && Property->Struct == FAnimNode_StateMachine::StaticStruct())
-			{
-				FAnimNode_StateMachine* StateMachine = Property->ContainerPtrToValuePtr<FAnimNode_StateMachine>(this);
-				if(StateMachine)
-				{
-					const FBakedAnimationStateMachine* MachineDescription = GetMachineDescription(AnimBlueprintGeneratedClass, StateMachine);
-					if(MachineDescription && MachineName == MachineDescription->MachineName)
-					{
-						// check each state transition for a match
-						int32 StateIndex = 0;
-						for(const FBakedAnimationState& State : MachineDescription->States)
-						{
-							if(State.StateName == StateName)
-							{
-								Predicate(StateMachine, State, StateIndex);
-							}
-							StateIndex++;
-						}
-					}
-				}
-			}
-		}
-	};
-
-	UAnimBlueprintGeneratedClass* AnimBlueprintGeneratedClass = Cast<UAnimBlueprintGeneratedClass>(GetClass());
-	if(AnimBlueprintGeneratedClass)
-	{
-		// transition delegates
-		for(const auto& Binding : NativeTransitionBindings)
-		{
-			ForEachStateLambda(AnimBlueprintGeneratedClass, Binding.MachineName, Binding.PreviousStateName, 
-				[&](FAnimNode_StateMachine* StateMachine, const FBakedAnimationState& State, int32 StateIndex)
-				{
-					for(const FBakedStateExitTransition& TransitionExit : State.Transitions)
-					{
-						if(TransitionExit.CanTakeDelegateIndex != INDEX_NONE)
-						{
-							// In case the state machine hasn't been initilized, we need to re-get the desc
-							const FBakedAnimationStateMachine* MachineDesc = GetMachineDescription(AnimBlueprintGeneratedClass, StateMachine);
-							const FAnimationTransitionBetweenStates& Transition = MachineDesc->Transitions[TransitionExit.TransitionIndex];
-							const FBakedAnimationState& BakedState = MachineDesc->States[Transition.NextState];
-							
-							if(BakedState.StateName == Binding.NextStateName)
-							{
-								FAnimNode_TransitionResult* ResultNode = GetNodeFromPropertyIndex<FAnimNode_TransitionResult>(this, AnimBlueprintGeneratedClass, TransitionExit.CanTakeDelegateIndex);
-								if(ResultNode)
-								{
-									ResultNode->NativeTransitionDelegate = Binding.NativeTransitionDelegate;
-								}
-							}
-						}
-					}
-				});
-		}
-
-		// state entry delegates
-		for(const auto& Binding : NativeStateEntryBindings)
-		{
-			ForEachStateLambda(AnimBlueprintGeneratedClass, Binding.MachineName, Binding.StateName, 
-				[&](FAnimNode_StateMachine* StateMachine, const FBakedAnimationState& State, int32 StateIndex)
-				{
-					// allocate enough space for all our states we need so far
-					StateMachine->OnGraphStatesEntered.SetNum(FMath::Max(StateIndex + 1, StateMachine->OnGraphStatesEntered.Num()));
-					StateMachine->OnGraphStatesEntered[StateIndex] = Binding.NativeStateDelegate;
-				});
-		}
-
-		// state exit delegates
-		for(const auto& Binding : NativeStateExitBindings)
-		{
-			ForEachStateLambda(AnimBlueprintGeneratedClass, Binding.MachineName, Binding.StateName, 
-				[&](FAnimNode_StateMachine* StateMachine, const FBakedAnimationState& State, int32 StateIndex)
-				{
-					// allocate enough space for all our states we need so far
-					StateMachine->OnGraphStatesExited.SetNum(FMath::Max(StateIndex + 1, StateMachine->OnGraphStatesExited.Num()));
-					StateMachine->OnGraphStatesExited[StateIndex] = Binding.NativeStateDelegate;
-				});
-		}
-	}
+	return GetProxyOnGameThread().HasNativeStateExitBinding(MachineName, StateName, OutBindingName);
 }
 
 void OutputCurveMap(TMap<FName, float>& CurveMap, UCanvas* Canvas, FDisplayDebugManager& DisplayDebugManager, float Indent)
@@ -999,6 +682,8 @@ void OutputTickRecords(const TArray<FAnimTickRecord>& Records, UCanvas* Canvas,
 
 void UAnimInstance::DisplayDebug(class UCanvas* Canvas, const FDebugDisplayInfo& DebugDisplay, float& YL, float& YPos)
 {
+	FAnimInstanceProxy& Proxy = GetProxyOnGameThread();
+
 	float Indent = 0.f;
 
 	FLinearColor TextYellow(0.86f, 0.69f, 0.f);
@@ -1030,7 +715,7 @@ void UAnimInstance::DisplayDebug(class UCanvas* Canvas, const FDebugDisplayInfo&
 	FString Heading = FString::Printf(TEXT("Animation: %s"), *GetName());
 	DisplayDebugManager.DrawString(Heading, Indent);
 
-	if (bShowGraph && RootNode != nullptr)
+	if (bShowGraph && Proxy.HasRootNode())
 	{
 		DisplayDebugManager.SetLinearDrawColor(TextYellow);
 
@@ -1045,7 +730,7 @@ void UAnimInstance::DisplayDebug(class UCanvas* Canvas, const FDebugDisplayInfo&
 
 		DebugDataCounter.Increment();
 		FNodeDebugData NodeDebugData(this);
-		RootNode->GatherDebugData(NodeDebugData);
+		Proxy.GatherDebugData(NodeDebugData);
 
 		TArray<FNodeDebugData::FFlattenedDebugData> FlattenedData = NodeDebugData.GetFlattenedDebugData();
 
@@ -1111,8 +796,8 @@ void UAnimInstance::DisplayDebug(class UCanvas* Canvas, const FDebugDisplayInfo&
 		FIndenter AnimIndent(Indent);
 
 		//Display Sync Groups
-		const TArray<FAnimGroupInstance>& SyncGroups = SyncGroupArrays[GetSyncGroupReadIndex()];
-		const TArray<FAnimTickRecord>& UngroupedActivePlayers = UngroupedActivePlayerArrays[GetSyncGroupReadIndex()];
+		const TArray<FAnimGroupInstance>& SyncGroups = GetProxyOnGameThread().GetSyncGroupRead();
+		const TArray<FAnimTickRecord>& UngroupedActivePlayers = GetProxyOnGameThread().GetUngroupedActivePlayersRead();
 
 		Heading = FString::Printf(TEXT("SyncGroups: %i"), SyncGroups.Num());
 		DisplayDebugManager.DrawString(Heading, Indent);
@@ -1194,26 +879,26 @@ void UAnimInstance::DisplayDebug(class UCanvas* Canvas, const FDebugDisplayInfo&
 		{
 			FIndenter CurveIndent(Indent);
 
-			Heading = FString::Printf(TEXT("Morph Curves: %i"), MorphTargetCurves.Num());
+			Heading = FString::Printf(TEXT("Morph Curves: %i"), Proxy.GetMorphTargetCurves().Num());
 			DisplayDebugManager.DrawString(Heading, Indent);
 
 			DisplayDebugManager.SetLinearDrawColor(TextWhite);
 
 			{
 				FIndenter MorphCurveIndent(Indent);
-				OutputCurveMap(MorphTargetCurves, Canvas, DisplayDebugManager, Indent);
+				OutputCurveMap(Proxy.GetMorphTargetCurves(), Canvas, DisplayDebugManager, Indent);
 			}
 
 			DisplayDebugManager.SetLinearDrawColor(TextYellow);
 
-			Heading = FString::Printf(TEXT("Material Curves: %i"), MaterialParameterCurves.Num());
+			Heading = FString::Printf(TEXT("Material Curves: %i"), Proxy.GetMaterialParameterCurves().Num());
 			DisplayDebugManager.DrawString(Heading, Indent);
 
 			DisplayDebugManager.SetLinearDrawColor(TextWhite);
 
 			{
 				FIndenter MaterialCurveIndent(Indent);
-				OutputCurveMap(MaterialParameterCurves, Canvas, DisplayDebugManager, Indent);
+				OutputCurveMap(Proxy.GetMaterialParameterCurves(), Canvas, DisplayDebugManager, Indent);
 			}
 
 			DisplayDebugManager.SetLinearDrawColor(TextYellow);
@@ -1238,15 +923,12 @@ void UAnimInstance::RecalcRequiredBones()
 
 	if( SkelMeshComp->SkeletalMesh && SkelMeshComp->SkeletalMesh->Skeleton )
 	{
-		RequiredBones.InitializeTo(SkelMeshComp->RequiredBones, *SkelMeshComp->SkeletalMesh);
+		GetProxyOnGameThread().RecalcRequiredBones(SkelMeshComp, SkelMeshComp->SkeletalMesh);
 	}
 	else if( CurrentSkeleton != NULL )
 	{
-		RequiredBones.InitializeTo(SkelMeshComp->RequiredBones, *CurrentSkeleton);
+		GetProxyOnGameThread().RecalcRequiredBones(SkelMeshComp, CurrentSkeleton);
 	}
-
-	// When RequiredBones mapping has changed, AnimNodes need to update their bones caches. 
-	bBoneCachesInvalidated = true;
 }
 
 void UAnimInstance::Serialize(FArchive& Ar)
@@ -1255,7 +937,7 @@ void UAnimInstance::Serialize(FArchive& Ar)
 
 	if (!Ar.IsLoading() || !Ar.IsSaving())
 	{
-		Ar << RequiredBones;
+		Ar << GetProxyOnGameThread().GetRequiredBones();
 	}
 }
 
@@ -1264,65 +946,53 @@ bool UAnimInstance::CanTransitionSignature() const
 	return false;
 }
 
-void UAnimInstance::AddAnimNotifies(const TArray<const FAnimNotifyEvent*>& NewNotifies, const float InstanceWeight)
+void UAnimInstance::BeginDestroy()
 {
-	// for now there is no filter whatsoever, it just adds everything requested
-	for (const FAnimNotifyEvent* Notify : NewNotifies)
+#if DO_CHECK
+	if(GetOuter() && GetOuter()->IsA<USkeletalMeshComponent>())
 	{
-		// only add if it is over TriggerWeightThreshold
-		const bool bPassesDedicatedServerCheck = Notify->bTriggerOnDedicatedServer || !IsRunningDedicatedServer();
-		if (bPassesDedicatedServerCheck && Notify->TriggerWeightThreshold <= InstanceWeight && PassesFiltering(Notify) && PassesChanceOfTriggering(Notify))
-		{
-			// Only add unique AnimNotifyState instances just once. We can get multiple triggers if looping over an animation.
-			// It is the same state, so just report it once.
-			Notify->NotifyStateClass ? AnimNotifies.AddUnique(Notify) : AnimNotifies.Add(Notify);
-		}
+		check(!IsRunningParallelEvaluation());
+	}
+#endif
+	if(AnimInstanceProxy != nullptr)
+	{
+		DestroyAnimInstanceProxy(AnimInstanceProxy);
+		AnimInstanceProxy = nullptr;
 	}
+
+	Super::BeginDestroy();
 }
 
-bool UAnimInstance::PassesFiltering(const FAnimNotifyEvent* Notify) const
+void UAnimInstance::PostInitProperties()
 {
-	switch (Notify->NotifyFilterType)
+	Super::PostInitProperties();
+
+	if(AnimInstanceProxy == nullptr)
 	{
-		case ENotifyFilterType::NoFiltering:
-		{
-			return true;
-		}
-		case ENotifyFilterType::LOD:
-		{
-			USkeletalMeshComponent* Comp = GetSkelMeshComponent();
-			return Comp ? Notify->NotifyFilterLOD > Comp->PredictedLODLevel : true;
-		}
-		default:
-		{
-			ensure(false); // Unknown Filter Type
-		}
+		AnimInstanceProxy = CreateAnimInstanceProxy();
+		check(AnimInstanceProxy != nullptr);
 	}
-	return true;
+}
+
+bool UAnimInstance::PassesFiltering(const FAnimNotifyEvent* Notify) const
+{
+	return NotifyQueue.PassesFiltering(Notify);
 }
 
 bool UAnimInstance::PassesChanceOfTriggering(const FAnimNotifyEvent* Event) const
 {
-	return Event->NotifyStateClass ? true : FMath::FRandRange(0.f, 1.f) < Event->NotifyTriggerChance;
+	return NotifyQueue.PassesChanceOfTriggering(Event);
 }
 
 void UAnimInstance::AddAnimNotifyFromGeneratedClass(int32 NotifyIndex)
 {
-	if(NotifyIndex==INDEX_NONE)
-	{
-		return;
-	}
-
-	if (UAnimBlueprintGeneratedClass* AnimBlueprintClass = Cast<UAnimBlueprintGeneratedClass>(GetClass()))
-	{
-		check(AnimBlueprintClass->AnimNotifies.IsValidIndex(NotifyIndex));
-		const FAnimNotifyEvent* Notify =& AnimBlueprintClass->AnimNotifies[NotifyIndex];
-		AnimNotifies.Add(Notify);
-	}
+	GetProxyOnGameThread().AddAnimNotifyFromGeneratedClass(NotifyIndex);
 }
 
 void UAnimInstance::AddCurveValue(const FName& CurveName, float Value, int32 CurveTypeFlags)
 {
+	FAnimInstanceProxy& Proxy = GetProxyOnGameThread();
+
 	// save curve value, it will overwrite if same exists, 
 	//CurveValues.Add(CurveName, Value);
 	if (CurveTypeFlags & ACF_TriggerEvent)
@@ -1342,7 +1012,7 @@ void UAnimInstance::AddCurveValue(const FName& CurveName, float Value, int32 Cur
 
 	if (CurveTypeFlags & ACF_DrivesMorphTarget)
 	{
-		float *CurveValPtr = MorphTargetCurves.Find(CurveName);
+		float *CurveValPtr = Proxy.GetMorphTargetCurves().Find(CurveName);
 		if ( CurveValPtr )
 		{
 			// sum up, in the future we might normalize, but for now this just sums up
@@ -1351,21 +1021,21 @@ void UAnimInstance::AddCurveValue(const FName& CurveName, float Value, int32 Cur
 		}
 		else
 		{
-			MorphTargetCurves.Add(CurveName, Value);
+			Proxy.GetMorphTargetCurves().Add(CurveName, Value);
 		}
 	}
 
 	if (CurveTypeFlags & ACF_DrivesMaterial)
 	{
 		MaterialParamatersToClear.RemoveSwap(CurveName);
-		float* CurveValPtr = MaterialParameterCurves.Find(CurveName);
+		float* CurveValPtr = Proxy.GetMaterialParameterCurves().Find(CurveName);
 		if( CurveValPtr)
 		{
 			*CurveValPtr += Value;
 		}
 		else
 		{
-			MaterialParameterCurves.Add(CurveName, Value);
+			Proxy.GetMaterialParameterCurves().Add(CurveName, Value);
 		}
 	}
 }
@@ -1382,13 +1052,30 @@ void UAnimInstance::AddCurveValue(const USkeleton::AnimCurveUID Uid, float Value
 	AddCurveValue(CurrentCurveName, Value, CurveTypeFlags);
 }
 
+int32 UAnimInstance::GetSyncGroupReadIndex() const
+{ 
+	return GetProxyOnGameThread().GetSyncGroupReadIndex(); 
+}
+
+int32 UAnimInstance::GetSyncGroupWriteIndex() const 
+{ 
+	return GetProxyOnGameThread().GetSyncGroupWriteIndex(); 
+}
+
+void UAnimInstance::TickSyncGroupWriteIndex() 
+{
+	GetProxyOnGameThread().TickSyncGroupWriteIndex();
+}
+
 void UAnimInstance::UpdateCurves(const FBlendedCurve& InCurve)
 {
 	SCOPE_CYCLE_COUNTER(STAT_UpdateCurves);
 
+	FAnimInstanceProxy& Proxy = GetProxyOnGameThread();
+
 	//Track material params we set last time round so we can clear them if they aren't set again.
 	MaterialParamatersToClear.Reset();
-	for(auto Iter = MaterialParameterCurves.CreateConstIterator(); Iter; ++Iter)
+	for(auto Iter = Proxy.GetMaterialParameterCurves().CreateConstIterator(); Iter; ++Iter)
 	{
 		if(Iter.Value() > 0.0f)
 		{
@@ -1397,8 +1084,8 @@ void UAnimInstance::UpdateCurves(const FBlendedCurve& InCurve)
 	}
 
 	EventCurves.Reset();
-	MorphTargetCurves.Reset();
-	MaterialParameterCurves.Reset();
+	Proxy.GetMorphTargetCurves().Reset();
+	Proxy.GetMaterialParameterCurves().Reset();
 
 	for(int32 CurveId=0; CurveId<InCurve.UIDList.Num(); ++CurveId)
 	{
@@ -1415,6 +1102,16 @@ void UAnimInstance::UpdateCurves(const FBlendedCurve& InCurve)
 	}
 }
 
+bool UAnimInstance::HasMorphTargetCurves() const
+{
+	return GetProxyOnGameThread().HasMorphTargetCurves();
+}
+
+TMap<FName, float>& UAnimInstance::GetMorphTargetCurves()
+{
+	return GetProxyOnGameThread().GetMorphTargetCurves();
+}
+
 void UAnimInstance::TriggerAnimNotifies(float DeltaSeconds)
 {
 	SCOPE_CYCLE_COUNTER(STAT_AnimTriggerAnimNotifies);
@@ -1425,9 +1122,9 @@ void UAnimInstance::TriggerAnimNotifies(float DeltaSeconds)
 	// AnimNotifyState freshly added that need their 'NotifyBegin' event called.
 	TArray<const FAnimNotifyEvent *> NotifyStateBeginEvent;
 
-	for (int32 Index=0; Index<AnimNotifies.Num(); Index++)
+	for (int32 Index=0; Index<NotifyQueue.AnimNotifies.Num(); Index++)
 	{
-		const FAnimNotifyEvent* AnimNotifyEvent = AnimNotifies[Index];
+		const FAnimNotifyEvent* AnimNotifyEvent = NotifyQueue.AnimNotifies[Index];
 
 		// AnimNotifyState
 		if( AnimNotifyEvent->NotifyStateClass )
@@ -1510,292 +1207,57 @@ void UAnimInstance::TriggerSingleAnimNotify(const FAnimNotifyEvent* AnimNotifyEv
 					// Actor has event, but with different parameters. Print warning
 					UE_LOG(LogAnimNotify, Warning, TEXT("Anim notifier named %s, but the parameter number does not match or not of the correct type"), *FuncName);
 				}
-			}
-		}
-	}
-}
-
-//to debug montage weight
-#define DEBUGMONTAGEWEIGHT 0
-
-void UAnimInstance::GetSlotWeight(FName const& SlotNodeName, float& out_SlotNodeWeight, float& out_SourceWeight, float& out_TotalNodeWeight) const
-{
-	// node total weight 
-	float NewSlotNodeWeight = 0.f;
-	// this is required to track, because it will be 1-SourceWeight
-	// if additive, it can be applied more
-	float NonAdditiveTotalWeight = 0.f;
-
-#if DEBUGMONTAGEWEIGHT
-	float TotalDesiredWeight = 0.f;
-#endif
-	// first get all the montage instance weight this slot node has
-	for (const FMontageEvaluationState& EvalState : MontageEvaluationData)
-	{
-		if (EvalState.Montage->IsValidSlot(SlotNodeName))
-		{
-			NewSlotNodeWeight += EvalState.MontageWeight;
-
-			if( !EvalState.Montage->IsValidAdditiveSlot(SlotNodeName) )
-			{
-				NonAdditiveTotalWeight += EvalState.MontageWeight;
-			}
-
-#if DEBUGMONTAGEWEIGHT			
-			TotalDesiredWeight += EvalState->DesiredWeight;
-#endif
-			UE_LOG(LogAnimMontage, Verbose, TEXT("GetSlotWeight : Owner: %s, AnimMontage: %s,  (DesiredWeight:%0.2f, Weight:%0.2f)"),
-						*GetNameSafe(GetOwningActor()), *EvalState.Montage->GetName(), EvalState.DesiredWeight, EvalState.MontageWeight);
-		}
-	}
-
-	// save the total node weight, it can be more than 1
-	// we need this so that when we eval, we normalized by this weight
-	// calculating there can cause inconsistency if some data changes
-	out_TotalNodeWeight = NewSlotNodeWeight;
-
-	// this can happen when it's blending in OR when newer animation comes in with shorter blendtime
-	// say #1 animation was blending out time with current blendtime 1.0 #2 animation was blending in with 1.0 (old) but got blend out with new blendtime 0.2f
-	// #3 animation was blending in with the new blendtime 0.2f, you'll have sum of #1, 2, 3 exceeds 1.f
-	if (NewSlotNodeWeight > (1.f + ZERO_ANIMWEIGHT_THRESH))
-	{
-		// you don't want to change weight of montage instance since it can play multiple slots
-		// if you change one, it will apply to all slots in that montage
-		// instead we should renormalize when we eval
-		// this should happen in the eval phase
-		NonAdditiveTotalWeight /= NewSlotNodeWeight;
-		// since we normalized, we reset
-		NewSlotNodeWeight = 1.f;
-	}
-#if DEBUGMONTAGEWEIGHT
-	else if (TotalDesiredWeight == 1.f && TotalSum < 1.f - ZERO_ANIMWEIGHT_THRESH)
-	{
-		// this can happen when it's blending in OR when newer animation comes in with longer blendtime
-		// say #1 animation was blending out time with current blendtime 0.2 #2 animation was blending in with 0.2 (old) but got blend out with new blendtime 1.f
-		// #3 animation was blending in with the new blendtime 1.f, you'll have sum of #1, 2, 3 doesn't meet 1.f
-		UE_LOG(LogAnimMontage, Warning, TEXT("[%s] Montage has less weight. Blending in?(%f)"), *SlotNodeName.ToString(), TotalSum);
-	}
-#endif
-
-	out_SlotNodeWeight = NewSlotNodeWeight;
-	out_SourceWeight = 1.f - NonAdditiveTotalWeight;
-}
-
-// for now disable becauase it will not work with single node instance
-#if (UE_BUILD_SHIPPING || UE_BUILD_TEST)
-#define DEBUG_MONTAGEINSTANCE_WEIGHT 0
-#else
-#define DEBUG_MONTAGEINSTANCE_WEIGHT 1
-#endif // !(UE_BUILD_SHIPPING || UE_BUILD_TEST)
-
-void UAnimInstance::SlotEvaluatePose(FName SlotNodeName, const FCompactPose& SourcePose, const FBlendedCurve& SourceCurve, float InSourceWeight, FCompactPose& BlendedPose, FBlendedCurve& BlendedCurve, float InBlendWeight, float InTotalNodeWeight)
-{
-	//Accessing MontageInstances from this function is not safe (as this can be called during Parallel Anim Evaluation!
-	//Any montage data you need to add should be part of MontageEvaluationData
-	// nothing to blend, just get it out
-	if (InBlendWeight <= ZERO_ANIMWEIGHT_THRESH)
-	{
-		BlendedPose = SourcePose;
-		BlendedCurve = SourceCurve;
-		return;
-	}
-
-	// Split our data into additive and non additive.
-	TArray<FSlotEvaluationPose> AdditivePoses;
-	TArray<FSlotEvaluationPose> NonAdditivePoses;
-
-	// first pass we go through collect weights and valid montages. 
-#if DEBUG_MONTAGEINSTANCE_WEIGHT
-	float TotalWeight = 0.f;
-#endif // DEBUG_MONTAGEINSTANCE_WEIGHT
-	for (const FMontageEvaluationState& EvalState : MontageEvaluationData)
-	{
-		if (EvalState.Montage->IsValidSlot(SlotNodeName))
-		{
-			FAnimTrack const* const AnimTrack = EvalState.Montage->GetAnimationData(SlotNodeName);
-
-			// Find out additive type for pose.
-			EAdditiveAnimationType const AdditiveAnimType = AnimTrack->IsAdditive() 
-				? (AnimTrack->IsRotationOffsetAdditive() ? AAT_RotationOffsetMeshSpace : AAT_LocalSpaceBase)
-				: AAT_None;
-
-			FSlotEvaluationPose NewPose(EvalState.MontageWeight, AdditiveAnimType);
-			
-			// Bone array has to be allocated prior to calling GetPoseFromAnimTrack
-			NewPose.Pose.SetBoneContainer(&RequiredBones);
-			NewPose.Curve.InitFrom(this);
-
-			// Extract pose from Track
-			FAnimExtractContext ExtractionContext(EvalState.MontagePosition, EvalState.Montage->HasRootMotion() && RootMotionMode != ERootMotionMode::NoRootMotionExtraction);
-			AnimTrack->GetAnimationPose(NewPose.Pose, NewPose.Curve, ExtractionContext);
-
-			// add montage curves 
-			FBlendedCurve MontageCurve;
-			MontageCurve.InitFrom(this);
-			EvalState.Montage->EvaluateCurveData(MontageCurve, EvalState.MontagePosition);
-			NewPose.Curve.Combine(MontageCurve);
-
-#if DEBUG_MONTAGEINSTANCE_WEIGHT
-			TotalWeight += EvalState.MontageWeight;
-#endif // DEBUG_MONTAGEINSTANCE_WEIGHT
-			if (AdditiveAnimType == AAT_None)
-			{
-				NonAdditivePoses.Add(NewPose);
-			}
-			else
-			{
-				AdditivePoses.Add(NewPose);
-			}
-		}
-	}
-
-	// allocate for blending
-	// If source has any weight, add it to the blend array.
-	float const SourceWeight = FMath::Clamp<float>(InSourceWeight, 0.f, 1.f);
-
-#if DEBUG_MONTAGEINSTANCE_WEIGHT
-	ensure (FMath::IsNearlyEqual(InTotalNodeWeight, TotalWeight, KINDA_SMALL_NUMBER));
-#endif // DEBUG_MONTAGEINSTANCE_WEIGHT
-	ensure (InTotalNodeWeight > ZERO_ANIMWEIGHT_THRESH);
-
-	if (InTotalNodeWeight > (1.f + ZERO_ANIMWEIGHT_THRESH))
-	{
-		// Re-normalize additive poses
-		for (int32 Index = 0; Index < AdditivePoses.Num(); Index++)
-		{
-			AdditivePoses[Index].Weight /= InTotalNodeWeight;
-		}
-		// Re-normalize non-additive poses
-		for (int32 Index = 0; Index < NonAdditivePoses.Num(); Index++)
-		{
-			NonAdditivePoses[Index].Weight /= InTotalNodeWeight;
-		}
-	}
-
-	// Make sure we have at least one montage here.
-	check((AdditivePoses.Num() > 0) || (NonAdditivePoses.Num() > 0));
-
-	// Second pass, blend non additive poses together
-	{
-		// If we're only playing additive animations, just copy source for base pose.
-		if (NonAdditivePoses.Num() == 0)
-		{
-			BlendedPose = SourcePose;
-			BlendedCurve = SourceCurve;
-		}
-		// Otherwise we need to blend non additive poses together
-		else
-		{
-			int32 const NumPoses = NonAdditivePoses.Num() + ((SourceWeight > ZERO_ANIMWEIGHT_THRESH) ? 1 : 0);
-
-			TArray<const FCompactPose*, TInlineAllocator<8>> BlendingPoses;
-			BlendingPoses.AddUninitialized(NumPoses);
-
-			TArray<float, TInlineAllocator<8>> BlendWeights;
-			BlendWeights.AddUninitialized(NumPoses);
-
-			TArray<const FBlendedCurve*, TInlineAllocator<8>> BlendingCurves;
-			BlendingCurves.AddUninitialized(NumPoses);
-
-			for (int32 Index = 0; Index < NonAdditivePoses.Num(); Index++)
-			{
-				BlendingPoses[Index] = &NonAdditivePoses[Index].Pose;
-				BlendingCurves[Index] = &NonAdditivePoses[Index].Curve;
-				BlendWeights[Index] = NonAdditivePoses[Index].Weight;
-			}
-
-			if (SourceWeight > ZERO_ANIMWEIGHT_THRESH)
-			{
-				int32 const SourceIndex = BlendWeights.Num() - 1;
-				BlendingPoses[SourceIndex] = &SourcePose;
-				BlendingCurves[SourceIndex] = &SourceCurve;
-				BlendWeights[SourceIndex] = SourceWeight;
-			}
-
-			// now time to blend all montages
-			FAnimationRuntime::BlendPosesTogetherIndirect(BlendingPoses, BlendingCurves, BlendWeights, BlendedPose, BlendedCurve);
-		}
-	}
-
-	// Third pass, layer on weighted additive poses.
-	if (AdditivePoses.Num() > 0)
-	{
-		for (int32 Index = 0; Index < AdditivePoses.Num(); Index++)
-		{
-			FSlotEvaluationPose const& AdditivePose = AdditivePoses[Index];
-			FAnimationRuntime::AccumulateAdditivePose(BlendedPose, AdditivePose.Pose, BlendedCurve, AdditivePose.Curve, AdditivePose.Weight, AdditivePose.AdditiveType);
+			}
 		}
 	}
+}
+
+//to debug montage weight
+#define DEBUGMONTAGEWEIGHT 0
+
+void UAnimInstance::GetSlotWeight(FName const& SlotNodeName, float& out_SlotNodeWeight, float& out_SourceWeight, float& out_TotalNodeWeight) const
+{
+	GetProxyOnGameThread().GetSlotWeight(SlotNodeName, out_SlotNodeWeight, out_SourceWeight, out_TotalNodeWeight);
+}
 
-	// Normalize rotations after blending/accumulation
-	BlendedPose.NormalizeRotations();
+void UAnimInstance::SlotEvaluatePose(FName SlotNodeName, const FCompactPose& SourcePose, const FBlendedCurve& SourceCurve, float InSourceWeight, FCompactPose& BlendedPose, FBlendedCurve& BlendedCurve, float InBlendWeight, float InTotalNodeWeight)
+{
+	GetProxyOnGameThread().SlotEvaluatePose(SlotNodeName, SourcePose, SourceCurve, InSourceWeight, BlendedPose, BlendedCurve, InBlendWeight, InTotalNodeWeight);
 }
 
 void UAnimInstance::ReinitializeSlotNodes()
 {
-	SlotWeightTracker.Reset();
-	
-	// Increment counter
-	SlotNodeInitializationCounter.Increment();
+	GetProxyOnGameThread().ReinitializeSlotNodes();
 }
 
 void UAnimInstance::RegisterSlotNodeWithAnimInstance(FName SlotNodeName)
 {
-	// verify if same slot node name exists
-	// then warn users, this is invalid
-	if (SlotWeightTracker.Contains(SlotNodeName))
-	{
-		UClass* InstanceClass = GetClass();
-
-		FString ClassNameString = InstanceClass ? InstanceClass->GetName() : FString("Unavailable");
-
-		FMessageLog("AnimBlueprint").Warning(FText::Format(LOCTEXT("AnimInstance_SlotNode", "SLOTNODE: '{0}' in animation instance class {1} already exists. Remove duplicates from the animation graph for this class."), FText::FromString(SlotNodeName.ToString()), FText::FromString(ClassNameString)));
-		return;
-	}
-
-	SlotWeightTracker.Add(SlotNodeName, FMontageActiveSlotTracker());
+	GetProxyOnGameThread().RegisterSlotNodeWithAnimInstance(SlotNodeName);
 }
 
 void UAnimInstance::UpdateSlotNodeWeight(FName SlotNodeName, float Weight)
 {
-	FMontageActiveSlotTracker* Tracker = SlotWeightTracker.Find(SlotNodeName);
-	if (Tracker)
-	{
-		// Count as relevant if we are weighted in
-		Tracker->bIsRelevantThisTick = Tracker->bIsRelevantThisTick || Weight > ZERO_ANIMWEIGHT_THRESH;
-	}
+	GetProxyOnGameThread().UpdateSlotNodeWeight(SlotNodeName, Weight);
 }
 
 void UAnimInstance::ClearSlotNodeWeights()
 {
-	for (auto Iter = SlotWeightTracker.CreateIterator(); Iter; ++Iter)
-	{
-		FMontageActiveSlotTracker& Tracker = Iter.Value();
-		Tracker.bWasRelevantOnPreviousTick = Tracker.bIsRelevantThisTick;
-		Tracker.bIsRelevantThisTick = false;
-		Tracker.RootMotionWeight = 0.f;
-	}
+	GetProxyOnGameThread().ClearSlotNodeWeights();
 }
 
 bool UAnimInstance::IsSlotNodeRelevantForNotifies(FName SlotNodeName) const
 {
-	const FMontageActiveSlotTracker* Tracker = SlotWeightTracker.Find(SlotNodeName);
-	return ( Tracker && (Tracker->bIsRelevantThisTick || Tracker->bWasRelevantOnPreviousTick) );
+	return GetProxyOnGameThread().IsSlotNodeRelevantForNotifies(SlotNodeName);
 }
 
 void UAnimInstance::UpdateSlotRootMotionWeight(FName SlotNodeName, float Weight)
 {
-	FMontageActiveSlotTracker* Tracker = SlotWeightTracker.Find(SlotNodeName);
-	if (Tracker)
-	{
-		Tracker->RootMotionWeight += Weight;
-	}
+	GetProxyOnGameThread().UpdateSlotRootMotionWeight(SlotNodeName, Weight);
 }
 
 float UAnimInstance::GetSlotRootMotionWeight(FName SlotNodeName) const
 {
-	const FMontageActiveSlotTracker* Tracker = SlotWeightTracker.Find(SlotNodeName);
-	return Tracker ? Tracker->RootMotionWeight : 0.f;
+	return GetProxyOnGameThread().GetSlotRootMotionWeight(SlotNodeName);
 }
 
 float UAnimInstance::GetCurveValue(FName CurveName)
@@ -1954,7 +1416,7 @@ void UAnimInstance::Montage_Advance(float DeltaSeconds)
 			FRootMotionMovementParams* RootMotionParams = NULL;
 			if (bExtractRootMotion)
 			{
-				RootMotionParams = (RootMotionMode != ERootMotionMode::IgnoreRootMotion) ? &ExtractedRootMotion : &LocalExtractedRootMotion;
+				RootMotionParams = (RootMotionMode != ERootMotionMode::IgnoreRootMotion) ? &GetProxyOnGameThread().GetExtractedRootMotion() : &LocalExtractedRootMotion;
 			}
 
 			MontageInstance->MontageSync_PreUpdate();
@@ -2736,7 +2198,7 @@ UAnimMontage* UAnimInstance::GetCurrentActiveMontage()
 	return NULL;
 }
 
-FAnimMontageInstance* UAnimInstance::GetActiveMontageInstance()
+FAnimMontageInstance* UAnimInstance::GetActiveMontageInstance() const
 {
 	// Start from end, as most recent instances are added at the end of the queue.
 	int32 const NumInstances = MontageInstances.Num();
@@ -2810,13 +2272,13 @@ FRootMotionMovementParams UAnimInstance::ConsumeExtractedRootMotion(float Alpha)
 	}
 	else if (Alpha > (1.f - ZERO_ANIMWEIGHT_THRESH))
 	{
-		FRootMotionMovementParams RootMotion = ExtractedRootMotion;
-		ExtractedRootMotion.Clear();
+		FRootMotionMovementParams RootMotion = GetProxyOnGameThread().GetExtractedRootMotion();
+		GetProxyOnGameThread().GetExtractedRootMotion().Clear();
 		return RootMotion;
 	}
 	else
 	{
-		return ExtractedRootMotion.ConsumeRootMotion(Alpha);
+		return GetProxyOnGameThread().GetExtractedRootMotion().ConsumeRootMotion(Alpha);
 	}
 }
 
@@ -2892,86 +2354,29 @@ void UAnimInstance::UnlockAIResources(bool bUnlockMovement, bool UnlockAILogic)
 
 bool UAnimInstance::GetTimeToClosestMarker(FName SyncGroup, FName MarkerName, float& OutMarkerTime) const
 {
-	const int32 SyncGroupIndex = GetSyncGroupIndexFromName(SyncGroup);
-	const TArray<FAnimGroupInstance>& SyncGroups = SyncGroupArrays[GetSyncGroupReadIndex()];
-
-	if (SyncGroups.IsValidIndex(SyncGroupIndex))
-	{
-		const FAnimGroupInstance& SyncGroupInstance = SyncGroups[SyncGroupIndex];
-		if (SyncGroupInstance.bCanUseMarkerSync && SyncGroupInstance.ActivePlayers.IsValidIndex(SyncGroupInstance.GroupLeaderIndex))
-		{
-			const FMarkerSyncAnimPosition& EndPosition = SyncGroupInstance.MarkerTickContext.GetMarkerSyncEndPosition();
-			const FAnimTickRecord& Leader = SyncGroupInstance.ActivePlayers[SyncGroupInstance.GroupLeaderIndex];
-			if (EndPosition.PreviousMarkerName == MarkerName)
-			{
-				OutMarkerTime = Leader.MarkerTickRecord->PreviousMarker.TimeToMarker;
-				return true;
-			}
-			else if (EndPosition.NextMarkerName == MarkerName)
-			{
-				OutMarkerTime = Leader.MarkerTickRecord->NextMarker.TimeToMarker;
-				return true;
-			}
-		}
-	}
-	return false;
+	return GetProxyOnAnyThread().GetTimeToClosestMarker(SyncGroup, MarkerName, OutMarkerTime);
 }
 
 bool UAnimInstance::HasMarkerBeenHitThisFrame(FName SyncGroup, FName MarkerName) const
 {
-	const int32 SyncGroupIndex = GetSyncGroupIndexFromName(SyncGroup);
-	const TArray<FAnimGroupInstance>& SyncGroups = SyncGroupArrays[GetSyncGroupReadIndex()];
-
-	if (SyncGroups.IsValidIndex(SyncGroupIndex))
-	{
-		const FAnimGroupInstance& SyncGroupInstance = SyncGroups[SyncGroupIndex];
-		if (SyncGroupInstance.bCanUseMarkerSync)
-		{
-			return SyncGroupInstance.MarkerTickContext.MarkersPassedThisTick.ContainsByPredicate([&MarkerName](const FPassedMarker& PassedMarker) -> bool
-			{
-				return PassedMarker.PassedMarkerName == MarkerName;
-			});
-		}
-	}
-	return false;
+	return GetProxyOnAnyThread().HasMarkerBeenHitThisFrame(SyncGroup, MarkerName);
 }
 
 bool UAnimInstance::IsSyncGroupBetweenMarkers(FName InSyncGroupName, FName PreviousMarker, FName NextMarker, bool bRespectMarkerOrder) const
 {
-	const FMarkerSyncAnimPosition& SyncGroupPosition = GetSyncGroupPosition(InSyncGroupName);
-	if ((SyncGroupPosition.PreviousMarkerName == PreviousMarker) && (SyncGroupPosition.NextMarkerName == NextMarker))
-	{
-		return true;
-	}
-
-	if (!bRespectMarkerOrder)
-	{
-		return ((SyncGroupPosition.PreviousMarkerName == NextMarker) && (SyncGroupPosition.NextMarkerName == PreviousMarker));
-	}
-
-	return false;
+	return GetProxyOnAnyThread().IsSyncGroupBetweenMarkers(InSyncGroupName, PreviousMarker, NextMarker, bRespectMarkerOrder);
 }
 
 FMarkerSyncAnimPosition UAnimInstance::GetSyncGroupPosition(FName InSyncGroupName) const
 {
-	const int32 SyncGroupIndex = GetSyncGroupIndexFromName(InSyncGroupName);
-	const TArray<FAnimGroupInstance>& SyncGroups = SyncGroupArrays[GetSyncGroupReadIndex()];
-
-	if (SyncGroups.IsValidIndex(SyncGroupIndex))
-	{
-		const FAnimGroupInstance& SyncGroupInstance = SyncGroups[SyncGroupIndex];
-		if (SyncGroupInstance.bCanUseMarkerSync)
-		{
-			return SyncGroupInstance.MarkerTickContext.GetMarkerSyncEndPosition();
-		}
-	}
-
-	return FMarkerSyncAnimPosition();
+	return GetProxyOnAnyThread().GetSyncGroupPosition(InSyncGroupName);
 }
 
 void UAnimInstance::UpdateMontageEvaluationData()
 {
-	MontageEvaluationData.Reset(MontageInstances.Num());
+	FAnimInstanceProxy& Proxy = GetProxyOnGameThread();
+
+	Proxy.GetMontageEvaluationData().Reset(MontageInstances.Num());
 	UE_LOG(LogAnimMontage, Verbose, TEXT("UpdateMontageEvaluationData Strting: Owner: %s"),	*GetNameSafe(GetOwningActor()));
 
 	for (FAnimMontageInstance* MontageInstance : MontageInstances)
@@ -2982,112 +2387,49 @@ void UAnimInstance::UpdateMontageEvaluationData()
 		{
 			UE_LOG(LogAnimMontage, Verbose, TEXT("UpdateMontageEvaluationData : AnimMontage: %s,  (DesiredWeight:%0.2f, Weight:%0.2f)"),
 						*MontageInstance->Montage->GetName(), MontageInstance->GetDesiredWeight(), MontageInstance->GetWeight());
-			MontageEvaluationData.Add(FMontageEvaluationState(MontageInstance->Montage, MontageInstance->GetWeight(), MontageInstance->GetDesiredWeight(), MontageInstance->GetPosition()));
+			Proxy.GetMontageEvaluationData().Add(FMontageEvaluationState(MontageInstance->Montage, MontageInstance->GetWeight(), MontageInstance->GetDesiredWeight(), MontageInstance->GetPosition(), MontageInstance->bPlaying, MontageInstance->IsActive()));
 		}
 	}
 }
 
 float UAnimInstance::GetInstanceAssetPlayerLength(int32 AssetPlayerIndex)
 {
-	if(FAnimNode_AssetPlayerBase* PlayerNode = GetNodeFromIndex<FAnimNode_AssetPlayerBase>(AssetPlayerIndex))
-	{
-		return PlayerNode->GetAnimAsset()->GetMaxCurrentTime();
-	}
-
-	return 0.0f;
+	return GetProxyOnAnyThread().GetInstanceAssetPlayerLength(AssetPlayerIndex);
 }
 
 float UAnimInstance::GetInstanceAssetPlayerTime(int32 AssetPlayerIndex)
 {
-	if(FAnimNode_AssetPlayerBase* PlayerNode = GetNodeFromIndex<FAnimNode_AssetPlayerBase>(AssetPlayerIndex))
-	{
-		return PlayerNode->GetAccumulatedTime();
-	}
-
-	return 0.0f;
+	return GetProxyOnAnyThread().GetInstanceAssetPlayerTime(AssetPlayerIndex);
 }
 
 float UAnimInstance::GetInstanceAssetPlayerTimeFraction(int32 AssetPlayerIndex)
 {
-	if(FAnimNode_AssetPlayerBase* PlayerNode = GetNodeFromIndex<FAnimNode_AssetPlayerBase>(AssetPlayerIndex))
-	{
-		UAnimationAsset* Asset = PlayerNode->GetAnimAsset();
-
-		float Length = Asset ? Asset->GetMaxCurrentTime() : 0.0f;
-
-		if(Length > 0.0f)
-		{
-			return PlayerNode->GetAccumulatedTime() / Length;
-		}
-	}
-
-	
-
-	return 0.0f;
+	return GetProxyOnAnyThread().GetInstanceAssetPlayerTimeFraction(AssetPlayerIndex);
 }
 
 float UAnimInstance::GetInstanceAssetPlayerTimeFromEndFraction(int32 AssetPlayerIndex)
 {
-	if(FAnimNode_AssetPlayerBase* PlayerNode = GetNodeFromIndex<FAnimNode_AssetPlayerBase>(AssetPlayerIndex))
-	{
-		UAnimationAsset* Asset = PlayerNode->GetAnimAsset();
-		
-		float Length = Asset ? Asset->GetMaxCurrentTime() : 0.f;
-
-		if(Length > 0.f)
-		{
-			return (Length - PlayerNode->GetAccumulatedTime()) / Length;
-		}
-	}
-
-	return 1.0f;
+	return GetProxyOnAnyThread().GetInstanceAssetPlayerTimeFromEndFraction(AssetPlayerIndex);
 }
 
 float UAnimInstance::GetInstanceAssetPlayerTimeFromEnd(int32 AssetPlayerIndex)
 {
-	if(FAnimNode_AssetPlayerBase* PlayerNode = GetNodeFromIndex<FAnimNode_AssetPlayerBase>(AssetPlayerIndex))
-	{
-		UAnimationAsset* Asset = PlayerNode->GetAnimAsset();
-		if(Asset)
-		{
-			return PlayerNode->GetAnimAsset()->GetMaxCurrentTime() - PlayerNode->GetAccumulatedTime();
-		}
-	}
-
-	return MAX_flt;
+	return GetProxyOnAnyThread().GetInstanceAssetPlayerTimeFromEnd(AssetPlayerIndex);
 }
 
 float UAnimInstance::GetInstanceStateWeight(int32 MachineIndex, int32 StateIndex)
 {
-	if(FAnimNode_StateMachine* MachineInstance = GetStateMachineInstance(MachineIndex))
-	{
-		return MachineInstance->GetStateWeight(StateIndex);
-	}
-	
-	return 0.0f;
+	return GetProxyOnAnyThread().GetInstanceStateWeight(MachineIndex, StateIndex);
 }
 
 float UAnimInstance::GetInstanceCurrentStateElapsedTime(int32 MachineIndex)
 {
-	if(FAnimNode_StateMachine* MachineInstance = GetStateMachineInstance(MachineIndex))
-	{
-		return MachineInstance->GetCurrentStateElapsedTime();
-	}
-
-	return 0.0f;
+	return GetProxyOnAnyThread().GetInstanceCurrentStateElapsedTime(MachineIndex);
 }
 
 float UAnimInstance::GetInstanceTransitionCrossfadeDuration(int32 MachineIndex, int32 TransitionIndex)
 {
-	if(FAnimNode_StateMachine* MachineInstance = GetStateMachineInstance(MachineIndex))
-	{
-		if(MachineInstance->IsValidTransitionIndex(TransitionIndex))
-		{
-			return MachineInstance->GetTransitionInfo(TransitionIndex).CrossfadeDuration;
-		}
-	}
-
-	return 0.0f;
+	return GetProxyOnAnyThread().GetInstanceTransitionCrossfadeDuration(MachineIndex, TransitionIndex);
 }
 
 float UAnimInstance::GetInstanceTransitionTimeElapsed(int32 MachineIndex, int32 TransitionIndex)
@@ -3098,305 +2440,82 @@ float UAnimInstance::GetInstanceTransitionTimeElapsed(int32 MachineIndex, int32
 
 float UAnimInstance::GetInstanceTransitionTimeElapsedFraction(int32 MachineIndex, int32 TransitionIndex)
 {
-	if(FAnimNode_StateMachine* MachineInstance = GetStateMachineInstance(MachineIndex))
-	{
-		if(MachineInstance->IsValidTransitionIndex(TransitionIndex))
-		{
-			return MachineInstance->GetCurrentStateElapsedTime() / MachineInstance->GetTransitionInfo(TransitionIndex).CrossfadeDuration;
-		}
-	}
-
-	return 0.0f;
+	return GetProxyOnAnyThread().GetInstanceTransitionTimeElapsedFraction(MachineIndex, TransitionIndex);
 }
 
 float UAnimInstance::GetRelevantAnimTimeRemaining(int32 MachineIndex, int32 StateIndex)
 {
-	if(FAnimNode_AssetPlayerBase* AssetPlayer = GetRelevantAssetPlayerFromState(MachineIndex, StateIndex))
-	{
-		if(AssetPlayer->GetAnimAsset())
-		{
-			return AssetPlayer->GetAnimAsset()->GetMaxCurrentTime() - AssetPlayer->GetAccumulatedTime();
-		}
-	}
-
-	return MAX_flt;
+	return GetProxyOnAnyThread().GetRelevantAnimTimeRemaining(MachineIndex, StateIndex);
 }
 
 float UAnimInstance::GetRelevantAnimTimeRemainingFraction(int32 MachineIndex, int32 StateIndex)
 {
-	if(FAnimNode_AssetPlayerBase* AssetPlayer = GetRelevantAssetPlayerFromState(MachineIndex, StateIndex))
-	{
-		if(AssetPlayer->GetAnimAsset())
-		{
-			float Length = AssetPlayer->GetAnimAsset()->GetMaxCurrentTime();
-			if(Length > 0.0f)
-			{
-				return (Length - AssetPlayer->GetAccumulatedTime()) / Length;
-			}
-		}
-	}
-
-	return 1.0f;
+	return GetProxyOnAnyThread().GetRelevantAnimTimeRemainingFraction(MachineIndex, StateIndex);
 }
 
 float UAnimInstance::GetRelevantAnimLength(int32 MachineIndex, int32 StateIndex)
 {
-	if(FAnimNode_AssetPlayerBase* AssetPlayer = GetRelevantAssetPlayerFromState(MachineIndex, StateIndex))
-	{
-		if(AssetPlayer->GetAnimAsset())
-		{
-			return AssetPlayer->GetAnimAsset()->GetMaxCurrentTime();
-		}
-	}
-
-	return 0.0f;
+	return GetProxyOnAnyThread().GetRelevantAnimLength(MachineIndex, StateIndex);
 }
 
 float UAnimInstance::GetRelevantAnimTime(int32 MachineIndex, int32 StateIndex)
 {
-	if(FAnimNode_AssetPlayerBase* AssetPlayer = GetRelevantAssetPlayerFromState(MachineIndex, StateIndex))
-	{
-		return AssetPlayer->GetAccumulatedTime();
-	}
-
-	return 0.0f;
+	return GetProxyOnAnyThread().GetRelevantAnimTime(MachineIndex, StateIndex);
 }
 
 float UAnimInstance::GetRelevantAnimTimeFraction(int32 MachineIndex, int32 StateIndex)
 {
-	if(FAnimNode_AssetPlayerBase* AssetPlayer = GetRelevantAssetPlayerFromState(MachineIndex, StateIndex))
-	{
-		float Length = AssetPlayer->GetAnimAsset()->GetMaxCurrentTime();
-		if(Length > 0.0f)
-		{
-			return AssetPlayer->GetAccumulatedTime() / Length;
-		}
-	}
-
-	return 0.0f;
+	return GetProxyOnAnyThread().GetRelevantAnimTimeFraction(MachineIndex, StateIndex);
 }
 
 FAnimNode_StateMachine* UAnimInstance::GetStateMachineInstance(int32 MachineIndex)
 {
-	if(UAnimBlueprintGeneratedClass* AnimBlueprintClass = Cast<UAnimBlueprintGeneratedClass>((UObject*)GetClass()))
-	{
-		if((MachineIndex >= 0) && (MachineIndex < AnimBlueprintClass->AnimNodeProperties.Num()))
-		{
-			const int32 InstancePropertyIndex = AnimBlueprintClass->AnimNodeProperties.Num() - 1 - MachineIndex;
-
-			UStructProperty* MachineInstanceProperty = AnimBlueprintClass->AnimNodeProperties[InstancePropertyIndex];
-			checkSlow(MachineInstanceProperty->Struct->IsChildOf(FAnimNode_StateMachine::StaticStruct()));
-
-			return MachineInstanceProperty->ContainerPtrToValuePtr<FAnimNode_StateMachine>(this);
-		}
-	}
-
-	return nullptr;
+	return GetProxyOnAnyThread().GetStateMachineInstance(MachineIndex);
 }
 
 FAnimNode_StateMachine* UAnimInstance::GetStateMachineInstanceFromName(FName MachineName)
 {
-	if (UAnimBlueprintGeneratedClass* AnimBlueprintClass = Cast<UAnimBlueprintGeneratedClass>((UObject*)GetClass()))
-	{
-		for (int32 MachineIndex = 0; MachineIndex < AnimBlueprintClass->AnimNodeProperties.Num(); MachineIndex++)
-		{
-			UStructProperty* Property = AnimBlueprintClass->AnimNodeProperties[AnimBlueprintClass->AnimNodeProperties.Num() - 1 - MachineIndex];
-			if (Property && Property->Struct == FAnimNode_StateMachine::StaticStruct())
-			{
-				FAnimNode_StateMachine* StateMachine = Property->ContainerPtrToValuePtr<FAnimNode_StateMachine>(this);
-				if (StateMachine)
-				{
-					if (const FBakedAnimationStateMachine* MachineDescription = GetMachineDescription(AnimBlueprintClass, StateMachine))
-					{
-						if (MachineDescription->MachineName == MachineName)
-						{
-							return StateMachine;
-						}
-					}
-				}
-			}
-		}
-	}
-
-	return nullptr;
+	return GetProxyOnAnyThread().GetStateMachineInstanceFromName(MachineName);
 }
 
 const FBakedAnimationStateMachine* UAnimInstance::GetStateMachineInstanceDesc(FName MachineName)
 {
-	if(UAnimBlueprintGeneratedClass* AnimBlueprintClass = Cast<UAnimBlueprintGeneratedClass>((UObject*)GetClass()))
-	{
-		for(int32 MachineIndex = 0; MachineIndex < AnimBlueprintClass->AnimNodeProperties.Num(); MachineIndex++)
-		{
-			UStructProperty* Property = AnimBlueprintClass->AnimNodeProperties[AnimBlueprintClass->AnimNodeProperties.Num() - 1 - MachineIndex];
-			if(Property && Property->Struct == FAnimNode_StateMachine::StaticStruct())
-			{
-				FAnimNode_StateMachine* StateMachine = Property->ContainerPtrToValuePtr<FAnimNode_StateMachine>(this);
-				if(StateMachine)
-				{
-					if(const FBakedAnimationStateMachine* MachineDescription = GetMachineDescription(AnimBlueprintClass, StateMachine))
-					{
-						if(MachineDescription->MachineName == MachineName)
-						{
-							return MachineDescription;
-						}
-					}
-				}
-			}
-		}
-	}
-
-	return nullptr;
+	return GetProxyOnAnyThread().GetStateMachineInstanceDesc(MachineName);
 }
 
 int32 UAnimInstance::GetStateMachineIndex(FName MachineName)
 {
-	if(UAnimBlueprintGeneratedClass* AnimBlueprintClass = Cast<UAnimBlueprintGeneratedClass>((UObject*)GetClass()))
-	{
-		for(int32 MachineIndex = 0; MachineIndex < AnimBlueprintClass->AnimNodeProperties.Num(); MachineIndex++)
-		{
-			UStructProperty* Property = AnimBlueprintClass->AnimNodeProperties[AnimBlueprintClass->AnimNodeProperties.Num() - 1 - MachineIndex];
-			if(Property && Property->Struct == FAnimNode_StateMachine::StaticStruct())
-			{
-				FAnimNode_StateMachine* StateMachine = Property->ContainerPtrToValuePtr<FAnimNode_StateMachine>(this);
-				if(StateMachine)
-				{
-					if(const FBakedAnimationStateMachine* MachineDescription = GetMachineDescription(AnimBlueprintClass, StateMachine))
-					{
-						if(MachineDescription->MachineName == MachineName)
-						{
-							return MachineIndex;
-						}
-					}
-				}
-			}
-		}
-	}
-
-	return INDEX_NONE;
+	return GetProxyOnAnyThread().GetStateMachineIndex(MachineName);
 }
 
 FAnimNode_Base* UAnimInstance::GetCheckedNodeFromIndexUntyped(int32 NodeIdx, UScriptStruct* RequiredStructType)
 {
-	FAnimNode_Base* NodePtr = nullptr;
-	if (UAnimBlueprintGeneratedClass* AnimBlueprintClass = Cast<UAnimBlueprintGeneratedClass>(GetClass()))
-	{
-		const int32 InstanceIdx = AnimBlueprintClass->AnimNodeProperties.Num() - 1 - NodeIdx;
-
-		if (AnimBlueprintClass->AnimNodeProperties.IsValidIndex(InstanceIdx))
-		{
-			UStructProperty* NodeProperty = AnimBlueprintClass->AnimNodeProperties[InstanceIdx];
-
-			if (NodeProperty->Struct->IsChildOf(RequiredStructType))
-			{
-				NodePtr = NodeProperty->ContainerPtrToValuePtr<FAnimNode_Base>(this);
-			}
-			else
-			{
-				checkfSlow(false, TEXT("Requested a node of type %s but found node of type %s"), *RequiredStructType->GetName(), *NodeProperty->Struct->GetName());
-			}
-		}
-		else
-		{
-			checkfSlow(false, TEXT("Requested node of type %s at index %d/%d, index out of bounds."), *RequiredStructType->GetName(), NodeIdx, InstanceIdx);
-		}
-	}
-
-	checkfSlow(NodePtr, TEXT("Requested node at index %d not found!"), NodeIdx);
-
-	return NodePtr;
+	return GetProxyOnGameThread().GetCheckedNodeFromIndexUntyped(NodeIdx, RequiredStructType);
 }
 
 FAnimNode_Base* UAnimInstance::GetNodeFromIndexUntyped(int32 NodeIdx, UScriptStruct* RequiredStructType)
 {
-	FAnimNode_Base* NodePtr = nullptr;
-	if (UAnimBlueprintGeneratedClass* AnimBlueprintClass = Cast<UAnimBlueprintGeneratedClass>(GetClass()))
-	{
-		const int32 InstanceIdx = AnimBlueprintClass->AnimNodeProperties.Num() - 1 - NodeIdx;
-
-		if (AnimBlueprintClass->AnimNodeProperties.IsValidIndex(InstanceIdx))
-		{
-			UStructProperty* NodeProperty = AnimBlueprintClass->AnimNodeProperties[InstanceIdx];
-
-			if (NodeProperty->Struct->IsChildOf(RequiredStructType))
-			{
-				NodePtr = NodeProperty->ContainerPtrToValuePtr<FAnimNode_Base>(this);
-			}
-		}
-	}
-
-	return NodePtr;
+	return GetProxyOnGameThread().GetNodeFromIndexUntyped(NodeIdx, RequiredStructType);
 }
 
 const FBakedAnimationStateMachine* UAnimInstance::GetMachineDescription(UAnimBlueprintGeneratedClass* AnimBlueprintClass, FAnimNode_StateMachine* MachineInstance)
 {
-	return AnimBlueprintClass->BakedStateMachines.IsValidIndex(MachineInstance->StateMachineIndexInClass) ? &(AnimBlueprintClass->BakedStateMachines[MachineInstance->StateMachineIndexInClass]) : nullptr;
+	return GetProxyOnGameThread().GetMachineDescription(AnimBlueprintClass, MachineInstance);
 }
 
 int32 UAnimInstance::GetInstanceAssetPlayerIndex(FName MachineName, FName StateName, FName AssetName)
 {
-	if(UAnimBlueprintGeneratedClass* AnimBlueprintClass = Cast<UAnimBlueprintGeneratedClass>((UObject*)GetClass()))
-	{
-		if(const FBakedAnimationStateMachine* MachineDescription = GetStateMachineInstanceDesc(MachineName))
-		{
-			for(int32 StateIndex = 0; StateIndex < MachineDescription->States.Num(); StateIndex++)
-			{
-				const FBakedAnimationState& State = MachineDescription->States[StateIndex];
-				if(State.StateName == StateName)
-				{
-					for(int32 PlayerIndex = 0; PlayerIndex < State.PlayerNodeIndices.Num(); PlayerIndex++)
-					{
-						checkSlow(State.PlayerNodeIndices[PlayerIndex] < AnimBlueprintClass->AnimNodeProperties.Num());
-						UStructProperty* AssetPlayerProperty = AnimBlueprintClass->AnimNodeProperties[AnimBlueprintClass->AnimNodeProperties.Num() - 1 - State.PlayerNodeIndices[PlayerIndex]];
-						if(AssetPlayerProperty && AssetPlayerProperty->Struct->IsChildOf(FAnimNode_AssetPlayerBase::StaticStruct()))
-						{
-							FAnimNode_AssetPlayerBase* AssetPlayer = AssetPlayerProperty->ContainerPtrToValuePtr<FAnimNode_AssetPlayerBase>(this);
-							if(AssetPlayer)
-							{
-								if(AssetName == NAME_None || AssetPlayer->GetAnimAsset()->GetFName() == AssetName)
-								{
-									return State.PlayerNodeIndices[PlayerIndex];
-								}
-							}
-						}
-					}
-				}
-			}
-		}
-	}
-
-	return INDEX_NONE;
+	return GetProxyOnGameThread().GetInstanceAssetPlayerIndex(MachineName, StateName, AssetName);
 }
 
 FAnimNode_AssetPlayerBase* UAnimInstance::GetRelevantAssetPlayerFromState(int32 MachineIndex, int32 StateIndex)
 {
-	FAnimNode_AssetPlayerBase* ResultPlayer = nullptr;
-	if(FAnimNode_StateMachine* MachineInstance = GetStateMachineInstance(MachineIndex))
-	{
-		float MaxWeight = 0.0f;
-		const FBakedAnimationState& State = MachineInstance->GetStateInfo(StateIndex);
-		for(const int32& PlayerIdx : State.PlayerNodeIndices)
-		{
-			if(FAnimNode_AssetPlayerBase* Player = GetNodeFromIndex<FAnimNode_AssetPlayerBase>(PlayerIdx))
-			{
-				if(!Player->bIgnoreForRelevancyTest && Player->GetCachedBlendWeight() > MaxWeight)
-				{
-					MaxWeight = Player->GetCachedBlendWeight();
-					ResultPlayer = Player;
-				}
-			}
-		}
-	}
-	return ResultPlayer;
+	return GetProxyOnGameThread().GetRelevantAssetPlayerFromState(MachineIndex, StateIndex);
 }
 
 int32 UAnimInstance::GetSyncGroupIndexFromName(FName SyncGroupName) const
 {
-
-	if (UAnimBlueprintGeneratedClass* AnimBlueprintClass = Cast<UAnimBlueprintGeneratedClass>(GetClass()))
-	{
-		return AnimBlueprintClass->GetSyncGroupIndex(SyncGroupName);
-	}
-	return INDEX_NONE;
+	return GetProxyOnGameThread().GetSyncGroupIndexFromName(SyncGroupName);
 }
 
 bool UAnimInstance::IsRunningParallelEvaluation() const
@@ -3409,4 +2528,35 @@ bool UAnimInstance::IsRunningParallelEvaluation() const
 	return false;
 }
 
+FAnimInstanceProxy* UAnimInstance::CreateAnimInstanceProxy()
+{
+	return new FAnimInstanceProxy(this);
+}
+
+void UAnimInstance::DestroyAnimInstanceProxy(FAnimInstanceProxy* InProxy)
+{
+	delete InProxy;
+}
+
+void UAnimInstance::UpdateMorphTargetCurves(const TMap<FName, float>& InMorphTargetCurves)
+{
+	GetProxyOnGameThread().UpdateMorphTargetCurves(InMorphTargetCurves);
+}
+
+void UAnimInstance::UpdateComponentsMaterialParameters(UPrimitiveComponent* Component)
+{
+	GetProxyOnGameThread().UpdateComponentsMaterialParameters(Component);
+}
+
+TArray<FActiveVertexAnim> UAnimInstance::UpdateActiveVertexAnims(const USkeletalMesh* SkeletalMesh)
+{
+	FAnimInstanceProxy& Proxy = GetProxyOnAnyThread();	// this can be called from CreateRenderState_Concurrent
+	return FAnimationRuntime::UpdateActiveVertexAnims(SkeletalMesh, Proxy.GetMorphTargetCurves(), Proxy.GetVertexAnims());
+}
+
+FBoneContainer& UAnimInstance::GetRequiredBones()
+{
+	return GetProxyOnGameThread().GetRequiredBones();
+}
+
 #undef LOCTEXT_NAMESPACE 
