Copying //UE4/Dev-Framework to Dev-Main (//UE4/Dev-Main)

diff --git a/Engine/Source/Runtime/Engine/Private/Animation/AnimInstanceProxy.cpp b/Engine/Source/Runtime/Engine/Private/Animation/AnimInstanceProxy.cpp
new file mode 100644
--- /dev/null
+++ b/Engine/Source/Runtime/Engine/Private/Animation/AnimInstanceProxy.cpp
@@ -0,0 +1,1457 @@
+// Copyright 1998-2015 Epic Games, Inc. All Rights Reserved.
+
+#include "EnginePrivate.h"
+#include "AnimInstanceProxy.h"
+#include "Animation/AnimStats.h"
+#include "Animation/AnimMontage.h"
+#include "Animation/BlendSpaceBase.h"
+#include "AnimationUtils.h"
+#include "MessageLog.h"
+#include "Animation/AnimNode_AssetPlayerBase.h"
+#include "Animation/AnimNode_StateMachine.h"
+#include "Animation/AnimNode_TransitionResult.h"
+
+#define DO_ANIMSTAT_PROCESSING(StatName) DEFINE_STAT(STAT_ ## StatName)
+#include "AnimMTStats.h"
+#undef DO_ANIMSTAT_PROCESSING
+
+#define DO_ANIMSTAT_PROCESSING(StatName) DEFINE_STAT(STAT_ ## StatName ## _WorkerThread)
+#include "AnimMTStats.h"
+#undef DO_ANIMSTAT_PROCESSING
+
+#define LOCTEXT_NAMESPACE "AnimInstance"
+
+void FAnimInstanceProxy::UpdateAnimationNode(float DeltaSeconds)
+{
+#if WITH_EDITORONLY_DATA
+	UpdatedNodesThisFrame.Reset();
+#endif
+
+	if(RootNode != nullptr)
+	{
+		UpdateCounter.Increment();
+		RootNode->Update(FAnimationUpdateContext(this, DeltaSeconds));
+	}
+}
+
+void FAnimInstanceProxy::Initialize(UAnimInstance* InAnimInstance)
+{
+	// copy anim instance object if it has not already been set up
+	AnimInstanceObject = InAnimInstance;
+
+	AnimBlueprintClass = Cast<UAnimBlueprintGeneratedClass>(InAnimInstance->GetClass());
+
+	USkeletalMeshComponent* OwnerComponent = InAnimInstance->GetSkelMeshComponent();
+	if (OwnerComponent->SkeletalMesh != NULL)
+	{
+		Skeleton = OwnerComponent->SkeletalMesh->Skeleton;
+	}
+	else
+	{
+		Skeleton = nullptr;
+	}
+
+	if(AnimBlueprintClass)
+	{
+		// Grab a pointer to the root node
+		if (AnimBlueprintClass->RootAnimNodeProperty != nullptr)
+		{
+			RootNode = AnimBlueprintClass->RootAnimNodeProperty->ContainerPtrToValuePtr<FAnimNode_Base>(InAnimInstance);
+		}
+		else
+		{
+			RootNode = nullptr;
+		}
+
+		// if no mesh, use Blueprint Skeleton
+		if (Skeleton == nullptr)
+		{
+			Skeleton = AnimBlueprintClass->TargetSkeleton;
+		}
+
+#if WITH_EDITORONLY_DATA
+		if (UAnimBlueprint* Blueprint = Cast<UAnimBlueprint>(AnimBlueprintClass->ClassGeneratedBy))
+		{
+			if (Blueprint->Status == BS_Error)
+			{
+				RootNode = nullptr;
+			}
+		}
+#endif
+	}
+
+	SkeletalMeshComponent = InAnimInstance->GetSkelMeshComponent();
+
+#if !NO_LOGGING
+	ActorName = GetNameSafe(InAnimInstance->GetOwningActor());
+#endif
+
+#if DO_CHECK
+	AnimInstanceName = *InAnimInstance->GetFullName();
+#endif
+
+	ReinitializeSlotNodes();
+}
+
+void FAnimInstanceProxy::InitializeRootNode()
+{
+	if (RootNode != nullptr)
+	{
+		GameThreadPreUpdateEvents.Reset();
+
+		InitializationCounter.Increment();
+		FAnimationInitializeContext InitContext(this);
+		RootNode->Initialize(InitContext);
+	}
+}
+
+void FAnimInstanceProxy::Uninitialize(UAnimInstance* InAnimInstance)
+{
+	MorphTargetCurves.Reset();
+	MaterialParameterCurves.Reset();
+}
+
+void FAnimInstanceProxy::PreUpdate(const UAnimInstance* InAnimInstance, float DeltaSeconds)
+{
+	CurrentDeltaSeconds = DeltaSeconds;
+	RootMotionMode = InAnimInstance->RootMotionMode;
+	bShouldExtractRootMotion = InAnimInstance->ShouldExtractRootMotion();
+
+	NotifyQueue.Reset(InAnimInstance->GetSkelMeshComponent());
+
+	ClearSlotNodeWeights();
+
+	// Reset the player tick list (but keep it presized)
+	TArray<FAnimTickRecord>& UngroupedActivePlayers = UngroupedActivePlayerArrays[GetSyncGroupWriteIndex()];
+	UngroupedActivePlayers.Reset();
+
+	TArray<FAnimGroupInstance>& SyncGroups = SyncGroupArrays[GetSyncGroupWriteIndex()];
+	for (int32 GroupIndex = 0; GroupIndex < SyncGroups.Num(); ++GroupIndex)
+	{
+		SyncGroups[GroupIndex].Reset();
+	}
+
+#if WITH_EDITORONLY_DATA
+	bIsBeingDebugged = false;
+	if(AnimBlueprintClass)
+	{
+		if (UAnimBlueprint* AnimBlueprint = Cast<UAnimBlueprint>(AnimBlueprintClass->ClassGeneratedBy))
+		{
+			bIsBeingDebugged = (InAnimInstance && (AnimBlueprint->GetObjectBeingDebugged() == InAnimInstance));
+		}
+	}
+#endif
+
+	// run preupdate events
+	for(const auto& Event : GameThreadPreUpdateEvents)
+	{
+		Event.ExecuteIfBound(InAnimInstance);
+	}
+}
+
+void FAnimInstanceProxy::PostUpdate(UAnimInstance* InAnimInstance) const
+{
+#if WITH_EDITORONLY_DATA
+	UAnimBlueprintGeneratedClass* AnimBlueprintGeneratedClass = Cast<UAnimBlueprintGeneratedClass>(InAnimInstance->GetClass());
+	AnimBlueprintGeneratedClass->GetAnimBlueprintDebugData().RecordNodeVisitArray(UpdatedNodesThisFrame);
+#endif
+
+	InAnimInstance->NotifyQueue.Append(NotifyQueue);
+}
+
+FAnimTickRecord& FAnimInstanceProxy::CreateUninitializedTickRecord(int32 GroupIndex, FAnimGroupInstance*& OutSyncGroupPtr)
+{
+	// Find or create the sync group if there is one
+	OutSyncGroupPtr = NULL;
+	if (GroupIndex >= 0)
+	{
+		TArray<FAnimGroupInstance>& SyncGroups = SyncGroupArrays[GetSyncGroupWriteIndex()];
+		while (SyncGroups.Num() <= GroupIndex)
+		{
+			new (SyncGroups) FAnimGroupInstance();
+		}
+		OutSyncGroupPtr = &(SyncGroups[GroupIndex]);
+	}
+
+	// Create the record
+	FAnimTickRecord* TickRecord = new ((OutSyncGroupPtr != NULL) ? OutSyncGroupPtr->ActivePlayers : UngroupedActivePlayerArrays[GetSyncGroupWriteIndex()]) FAnimTickRecord();
+	return *TickRecord;
+}
+
+void FAnimInstanceProxy::MakeSequenceTickRecord(FAnimTickRecord& TickRecord, class UAnimSequenceBase* Sequence, bool bLooping, float PlayRate, float FinalBlendWeight, float& CurrentTime, FMarkerTickRecord& MarkerTickRecord) const
+{
+	TickRecord.SourceAsset = Sequence;
+	TickRecord.TimeAccumulator = &CurrentTime;
+	TickRecord.MarkerTickRecord = &MarkerTickRecord;
+	TickRecord.PlayRateMultiplier = PlayRate;
+	TickRecord.EffectiveBlendWeight = FinalBlendWeight;
+	TickRecord.bLooping = bLooping;
+}
+
+void FAnimInstanceProxy::MakeBlendSpaceTickRecord(FAnimTickRecord& TickRecord, class UBlendSpaceBase* BlendSpace, const FVector& BlendInput, TArray<FBlendSampleData>& BlendSampleDataCache, FBlendFilter& BlendFilter, bool bLooping, float PlayRate, float FinalBlendWeight, float& CurrentTime, FMarkerTickRecord& MarkerTickRecord) const
+{
+	TickRecord.SourceAsset = BlendSpace;
+	TickRecord.BlendSpace.BlendSpacePositionX = BlendInput.X;
+	TickRecord.BlendSpace.BlendSpacePositionY = BlendInput.Y;
+	TickRecord.BlendSpace.BlendSampleDataCache = &BlendSampleDataCache;
+	TickRecord.BlendSpace.BlendFilter = &BlendFilter;
+	TickRecord.TimeAccumulator = &CurrentTime;
+	TickRecord.MarkerTickRecord = &MarkerTickRecord;
+	TickRecord.PlayRateMultiplier = PlayRate;
+	TickRecord.EffectiveBlendWeight = FinalBlendWeight;
+	TickRecord.bLooping = bLooping;
+}
+
+void FAnimInstanceProxy::SequenceAdvanceImmediate(UAnimSequenceBase* Sequence, bool bLooping, float PlayRate, float DeltaSeconds, float& CurrentTime, FMarkerTickRecord& MarkerTickRecord)
+{
+	FAnimTickRecord TickRecord;
+	MakeSequenceTickRecord(TickRecord, Sequence, bLooping, PlayRate, /*FinalBlendWeight=*/ 1.0f, CurrentTime, MarkerTickRecord);
+
+	FAnimAssetTickContext TickContext(DeltaSeconds, RootMotionMode, true);
+	TickRecord.SourceAsset->TickAssetPlayer(TickRecord, NotifyQueue, TickContext);
+}
+
+void FAnimInstanceProxy::BlendSpaceAdvanceImmediate(class UBlendSpaceBase* BlendSpace, const FVector& BlendInput, TArray<FBlendSampleData>& BlendSampleDataCache, FBlendFilter& BlendFilter, bool bLooping, float PlayRate, float DeltaSeconds, float& CurrentTime, FMarkerTickRecord& MarkerTickRecord)
+{
+	FAnimTickRecord TickRecord;
+	MakeBlendSpaceTickRecord(TickRecord, BlendSpace, BlendInput, BlendSampleDataCache, BlendFilter, bLooping, PlayRate, /*FinalBlendWeight=*/ 1.0f, CurrentTime, MarkerTickRecord);
+	
+	FAnimAssetTickContext TickContext(DeltaSeconds, RootMotionMode, true);
+	TickRecord.SourceAsset->TickAssetPlayer(TickRecord, NotifyQueue, TickContext);
+}
+
+void FAnimInstanceProxy::TickAssetPlayerInstances(float DeltaSeconds)
+{
+	SCOPE_CYCLE_COUNTER(STAT_TickAssetPlayerInstances);
+
+	// Handle all players inside sync groups
+	TArray<FAnimGroupInstance>& SyncGroups = SyncGroupArrays[GetSyncGroupWriteIndex()];
+	const TArray<FAnimGroupInstance>& PreviousSyncGroups = SyncGroupArrays[GetSyncGroupReadIndex()];
+	TArray<FAnimTickRecord>& UngroupedActivePlayers = UngroupedActivePlayerArrays[GetSyncGroupWriteIndex()];
+
+	for (int32 GroupIndex = 0; GroupIndex < SyncGroups.Num(); ++GroupIndex)
+	{
+		FAnimGroupInstance& SyncGroup = SyncGroups[GroupIndex];
+    
+		if (SyncGroup.ActivePlayers.Num() > 0)
+		{
+			const FAnimGroupInstance* PreviousGroup = PreviousSyncGroups.IsValidIndex(GroupIndex) ? &PreviousSyncGroups[GroupIndex] : nullptr;
+			SyncGroup.Prepare(PreviousGroup);
+
+			UE_LOG(LogAnimMarkerSync, Log, TEXT("Ticking Group [%d] GroupLeader [%d]"), GroupIndex, SyncGroup.GroupLeaderIndex);
+
+			const bool bOnlyOneAnimationInGroup = SyncGroup.ActivePlayers.Num() == 1;
+
+			// Tick the group leader
+			FAnimAssetTickContext TickContext(DeltaSeconds, RootMotionMode, bOnlyOneAnimationInGroup, SyncGroup.ValidMarkers);
+			// initialize to invalidate first
+			check(SyncGroup.GroupLeaderIndex == INDEX_NONE);
+			int32 GroupLeaderIndex = 0;
+			for (; GroupLeaderIndex < SyncGroup.ActivePlayers.Num(); ++GroupLeaderIndex)
+			{
+				FAnimTickRecord& GroupLeader = SyncGroup.ActivePlayers[GroupLeaderIndex];
+				// if it has leader score
+				SCOPE_CYCLE_COUNTER(STAT_TickAssetPlayerInstance);
+				FScopeCycleCounterUObject Scope(GroupLeader.SourceAsset);
+				GroupLeader.SourceAsset->TickAssetPlayer(GroupLeader, NotifyQueue, TickContext);
+
+				if (RootMotionMode == ERootMotionMode::RootMotionFromEverything && TickContext.RootMotionMovementParams.bHasRootMotion)
+				{
+					ExtractedRootMotion.AccumulateWithBlend(TickContext.RootMotionMovementParams.RootMotionTransform, GroupLeader.EffectiveBlendWeight);
+				}
+
+				// if we're not using marker based sync, we don't care, get out
+				if (TickContext.CanUseMarkerPosition() == false)
+				{
+					SyncGroup.GroupLeaderIndex = GroupLeaderIndex;
+					break;
+				}
+				// otherwise, the new position should contain the valid position for end, otherwise, we don't know where to sync to
+				else if (TickContext.MarkerTickContext.IsMarkerSyncEndValid())
+				{
+					// if this leader contains correct position, break
+					SyncGroup.MarkerTickContext = TickContext.MarkerTickContext;
+					SyncGroup.GroupLeaderIndex = GroupLeaderIndex;
+					UE_LOG(LogAnimMarkerSync, Log, TEXT("Previous Sync Group Makrer Tick Context :\n%s"), *SyncGroup.MarkerTickContext.ToString());
+					UE_LOG(LogAnimMarkerSync, Log, TEXT("New Sync Group Makrer Tick Context :\n%s"), *TickContext.MarkerTickContext.ToString());
+					break;
+				}
+				else
+				{
+					SyncGroup.GroupLeaderIndex = GroupLeaderIndex;
+					UE_LOG(LogAnimMarkerSync, Log, TEXT("Invalid position from Leader %d. Trying next leader"), GroupLeaderIndex);
+				}
+			} 
+
+			check(SyncGroup.GroupLeaderIndex != INDEX_NONE);
+			// we found leader
+			SyncGroup.Finalize(PreviousGroup);
+
+			// Update everything else to follow the leader, if there is more followers
+			if (SyncGroup.ActivePlayers.Num() > GroupLeaderIndex + 1)
+			{
+				// if we don't have a good leader, no reason to convert to follower
+				// tick as leader
+				TickContext.ConvertToFollower();
+    
+				for (int32 TickIndex = GroupLeaderIndex + 1; TickIndex < SyncGroup.ActivePlayers.Num(); ++TickIndex)
+				{
+					FAnimTickRecord& AssetPlayer = SyncGroup.ActivePlayers[TickIndex];
+					{
+						SCOPE_CYCLE_COUNTER(STAT_TickAssetPlayerInstance);
+						FScopeCycleCounterUObject Scope(AssetPlayer.SourceAsset);
+						TickContext.RootMotionMovementParams.Clear();
+						AssetPlayer.SourceAsset->TickAssetPlayer(AssetPlayer, NotifyQueue, TickContext);
+					}
+					if (RootMotionMode == ERootMotionMode::RootMotionFromEverything && TickContext.RootMotionMovementParams.bHasRootMotion)
+					{
+						ExtractedRootMotion.AccumulateWithBlend(TickContext.RootMotionMovementParams.RootMotionTransform, AssetPlayer.EffectiveBlendWeight);
+					}
+				}
+			}
+		}
+	}
+
+	// Handle the remaining ungrouped animation players
+	for (int32 TickIndex = 0; TickIndex < UngroupedActivePlayers.Num(); ++TickIndex)
+	{
+		FAnimTickRecord& AssetPlayerToTick = UngroupedActivePlayers[TickIndex];
+		static TArray<FName> TempNames;
+		const TArray<FName>* UniqueNames = AssetPlayerToTick.SourceAsset->GetUniqueMarkerNames();
+		const TArray<FName>& ValidMarkers = UniqueNames ? *UniqueNames : TempNames;
+
+		const bool bOnlyOneAnimationInGroup = true;
+		FAnimAssetTickContext TickContext(DeltaSeconds, RootMotionMode, bOnlyOneAnimationInGroup, ValidMarkers);
+		{
+			SCOPE_CYCLE_COUNTER(STAT_TickAssetPlayerInstance);
+			FScopeCycleCounterUObject Scope(AssetPlayerToTick.SourceAsset);
+			AssetPlayerToTick.SourceAsset->TickAssetPlayer(AssetPlayerToTick, NotifyQueue, TickContext);
+		}
+		if (RootMotionMode == ERootMotionMode::RootMotionFromEverything && TickContext.RootMotionMovementParams.bHasRootMotion)
+		{
+			ExtractedRootMotion.AccumulateWithBlend(TickContext.RootMotionMovementParams.RootMotionTransform, AssetPlayerToTick.EffectiveBlendWeight);
+		}
+	}
+}
+
+void FAnimInstanceProxy::AddAnimNotifies(const TArray<const FAnimNotifyEvent*>& NewNotifies, const float InstanceWeight)
+{
+	NotifyQueue.AddAnimNotifies(NewNotifies, InstanceWeight);
+}
+
+int32 FAnimInstanceProxy::GetSyncGroupIndexFromName(FName SyncGroupName) const
+{
+	if (AnimBlueprintClass)
+	{
+		return AnimBlueprintClass->GetSyncGroupIndex(SyncGroupName);
+	}
+	return INDEX_NONE;
+}
+
+bool FAnimInstanceProxy::GetTimeToClosestMarker(FName SyncGroup, FName MarkerName, float& OutMarkerTime) const
+{
+	const int32 SyncGroupIndex = GetSyncGroupIndexFromName(SyncGroup);
+	const TArray<FAnimGroupInstance>& SyncGroups = SyncGroupArrays[GetSyncGroupReadIndex()];
+
+	if (SyncGroups.IsValidIndex(SyncGroupIndex))
+	{
+		const FAnimGroupInstance& SyncGroupInstance = SyncGroups[SyncGroupIndex];
+		if (SyncGroupInstance.bCanUseMarkerSync && SyncGroupInstance.ActivePlayers.IsValidIndex(SyncGroupInstance.GroupLeaderIndex))
+		{
+			const FMarkerSyncAnimPosition& EndPosition = SyncGroupInstance.MarkerTickContext.GetMarkerSyncEndPosition();
+			const FAnimTickRecord& Leader = SyncGroupInstance.ActivePlayers[SyncGroupInstance.GroupLeaderIndex];
+			if (EndPosition.PreviousMarkerName == MarkerName)
+			{
+				OutMarkerTime = Leader.MarkerTickRecord->PreviousMarker.TimeToMarker;
+				return true;
+			}
+			else if (EndPosition.NextMarkerName == MarkerName)
+			{
+				OutMarkerTime = Leader.MarkerTickRecord->NextMarker.TimeToMarker;
+				return true;
+			}
+		}
+	}
+	return false;
+}
+
+void FAnimInstanceProxy::AddAnimNotifyFromGeneratedClass(int32 NotifyIndex)
+{
+	if(NotifyIndex==INDEX_NONE)
+	{
+		return;
+	}
+
+	if(AnimBlueprintClass)
+	{
+		check(AnimBlueprintClass->AnimNotifies.IsValidIndex(NotifyIndex));
+		const FAnimNotifyEvent* Notify =& AnimBlueprintClass->AnimNotifies[NotifyIndex];
+		NotifyQueue.AnimNotifies.Add(Notify);
+	}
+}
+
+bool FAnimInstanceProxy::HasMarkerBeenHitThisFrame(FName SyncGroup, FName MarkerName) const
+{
+	const int32 SyncGroupIndex = GetSyncGroupIndexFromName(SyncGroup);
+	const TArray<FAnimGroupInstance>& SyncGroups = SyncGroupArrays[GetSyncGroupReadIndex()];
+
+	if (SyncGroups.IsValidIndex(SyncGroupIndex))
+	{
+		const FAnimGroupInstance& SyncGroupInstance = SyncGroups[SyncGroupIndex];
+		if (SyncGroupInstance.bCanUseMarkerSync)
+		{
+			return SyncGroupInstance.MarkerTickContext.MarkersPassedThisTick.ContainsByPredicate([&MarkerName](const FPassedMarker& PassedMarker) -> bool
+			{
+				return PassedMarker.PassedMarkerName == MarkerName;
+			});
+		}
+	}
+	return false;
+}
+
+bool FAnimInstanceProxy::IsSyncGroupBetweenMarkers(FName InSyncGroupName, FName PreviousMarker, FName NextMarker, bool bRespectMarkerOrder) const
+{
+	const FMarkerSyncAnimPosition& SyncGroupPosition = GetSyncGroupPosition(InSyncGroupName);
+	if ((SyncGroupPosition.PreviousMarkerName == PreviousMarker) && (SyncGroupPosition.NextMarkerName == NextMarker))
+	{
+		return true;
+	}
+
+	if (!bRespectMarkerOrder)
+	{
+		return ((SyncGroupPosition.PreviousMarkerName == NextMarker) && (SyncGroupPosition.NextMarkerName == PreviousMarker));
+	}
+
+	return false;
+}
+
+FMarkerSyncAnimPosition FAnimInstanceProxy::GetSyncGroupPosition(FName InSyncGroupName) const
+{
+	const int32 SyncGroupIndex = GetSyncGroupIndexFromName(InSyncGroupName);
+	const TArray<FAnimGroupInstance>& SyncGroups = SyncGroupArrays[GetSyncGroupReadIndex()];
+
+	if (SyncGroups.IsValidIndex(SyncGroupIndex))
+	{
+		const FAnimGroupInstance& SyncGroupInstance = SyncGroups[SyncGroupIndex];
+		if (SyncGroupInstance.bCanUseMarkerSync)
+		{
+			return SyncGroupInstance.MarkerTickContext.GetMarkerSyncEndPosition();
+		}
+	}
+
+	return FMarkerSyncAnimPosition();
+}
+
+void FAnimInstanceProxy::ReinitializeSlotNodes()
+{
+	SlotWeightTracker.Reset();
+	
+	// Increment counter
+	SlotNodeInitializationCounter.Increment();
+}
+
+void FAnimInstanceProxy::RegisterSlotNodeWithAnimInstance(FName SlotNodeName)
+{
+	// message log access means we need to run this in the game thread
+	check(IsInGameThread());
+
+	// verify if same slot node name exists
+	// then warn users, this is invalid
+	if (SlotWeightTracker.Contains(SlotNodeName))
+	{
+		FString ClassNameString = GetAnimBlueprintClass() ? GetAnimBlueprintClass()->GetName() : FString("Unavailable");
+
+		FMessageLog("AnimBlueprint").Warning(FText::Format(LOCTEXT("AnimInstance_SlotNode", "SLOTNODE: '{0}' in animation instance class {1} already exists. Remove duplicates from the animation graph for this class."), FText::FromString(SlotNodeName.ToString()), FText::FromString(ClassNameString)));
+		return;
+	}
+
+	SlotWeightTracker.Add(SlotNodeName, FMontageActiveSlotTracker());
+}
+
+void FAnimInstanceProxy::AddGameThreadPreUpdateEvent(const FGameThreadPreUpdateEvent& InEvent)
+{
+	for(const auto& Event : GameThreadPreUpdateEvents)
+	{
+		if( Event.GetDelegateInstance()->GetRawMethodPtr() == InEvent.GetDelegateInstance()->GetRawMethodPtr() &&
+			Event.GetDelegateInstance()->GetRawUserObject() == InEvent.GetDelegateInstance()->GetRawUserObject())
+		{
+			// already added
+			return;
+		}
+	}
+
+	GameThreadPreUpdateEvents.Add(InEvent);
+}
+
+void FAnimInstanceProxy::UpdateSlotNodeWeight(FName SlotNodeName, float Weight)
+{
+	FMontageActiveSlotTracker* Tracker = SlotWeightTracker.Find(SlotNodeName);
+	if (Tracker)
+	{
+		// Count as relevant if we are weighted in
+		Tracker->bIsRelevantThisTick = Tracker->bIsRelevantThisTick || Weight > ZERO_ANIMWEIGHT_THRESH;
+	}
+}
+
+void FAnimInstanceProxy::ClearSlotNodeWeights()
+{
+	for (auto Iter = SlotWeightTracker.CreateIterator(); Iter; ++Iter)
+	{
+		FMontageActiveSlotTracker& Tracker = Iter.Value();
+		Tracker.bWasRelevantOnPreviousTick = Tracker.bIsRelevantThisTick;
+		Tracker.bIsRelevantThisTick = false;
+		Tracker.RootMotionWeight = 0.f;
+	}
+}
+
+bool FAnimInstanceProxy::IsSlotNodeRelevantForNotifies(FName SlotNodeName) const
+{
+	const FMontageActiveSlotTracker* Tracker = SlotWeightTracker.Find(SlotNodeName);
+	return ( Tracker && (Tracker->bIsRelevantThisTick || Tracker->bWasRelevantOnPreviousTick) );
+}
+
+void FAnimInstanceProxy::UpdateSlotRootMotionWeight(FName SlotNodeName, float Weight)
+{
+	FMontageActiveSlotTracker* Tracker = SlotWeightTracker.Find(SlotNodeName);
+	if (Tracker)
+	{
+		Tracker->RootMotionWeight += Weight;
+	}
+}
+
+float FAnimInstanceProxy::GetSlotRootMotionWeight(FName SlotNodeName) const
+{
+	const FMontageActiveSlotTracker* Tracker = SlotWeightTracker.Find(SlotNodeName);
+	return Tracker ? Tracker->RootMotionWeight : 0.f;
+}
+
+FAnimNode_Base* FAnimInstanceProxy::GetCheckedNodeFromIndexUntyped(int32 NodeIdx, UScriptStruct* RequiredStructType)
+{
+	FAnimNode_Base* NodePtr = nullptr;
+	if (AnimBlueprintClass)
+	{
+		const int32 InstanceIdx = AnimBlueprintClass->AnimNodeProperties.Num() - 1 - NodeIdx;
+
+		if (AnimBlueprintClass->AnimNodeProperties.IsValidIndex(InstanceIdx))
+		{
+			UStructProperty* NodeProperty = AnimBlueprintClass->AnimNodeProperties[InstanceIdx];
+
+			if (NodeProperty->Struct->IsChildOf(RequiredStructType))
+			{
+				NodePtr = NodeProperty->ContainerPtrToValuePtr<FAnimNode_Base>(AnimInstanceObject);
+			}
+			else
+			{
+				checkfSlow(false, TEXT("Requested a node of type %s but found node of type %s"), *RequiredStructType->GetName(), *NodeProperty->Struct->GetName());
+			}
+		}
+		else
+		{
+			checkfSlow(false, TEXT("Requested node of type %s at index %d/%d, index out of bounds."), *RequiredStructType->GetName(), NodeIdx, InstanceIdx);
+		}
+	}
+
+	checkfSlow(NodePtr, TEXT("Requested node at index %d not found!"), NodeIdx);
+
+	return NodePtr;
+}
+
+FAnimNode_Base* FAnimInstanceProxy::GetNodeFromIndexUntyped(int32 NodeIdx, UScriptStruct* RequiredStructType)
+{
+	FAnimNode_Base* NodePtr = nullptr;
+	if (AnimBlueprintClass)
+	{
+		const int32 InstanceIdx = AnimBlueprintClass->AnimNodeProperties.Num() - 1 - NodeIdx;
+
+		if (AnimBlueprintClass->AnimNodeProperties.IsValidIndex(InstanceIdx))
+		{
+			UStructProperty* NodeProperty = AnimBlueprintClass->AnimNodeProperties[InstanceIdx];
+
+			if (NodeProperty->Struct->IsChildOf(RequiredStructType))
+			{
+				NodePtr = NodeProperty->ContainerPtrToValuePtr<FAnimNode_Base>(AnimInstanceObject);
+			}
+		}
+	}
+
+	return NodePtr;
+}
+
+void FAnimInstanceProxy::RecalcRequiredBones(USkeletalMeshComponent* Component, UObject* Asset)
+{
+	RequiredBones.InitializeTo(Component->RequiredBones, *Asset);
+
+	// When RequiredBones mapping has changed, AnimNodes need to update their bones caches. 
+	bBoneCachesInvalidated = true;
+}
+
+void FAnimInstanceProxy::UpdateAnimation()
+{
+	ANIM_MT_SCOPE_CYCLE_COUNTER(ProxyUpdateAnimation, !IsInGameThread());
+	FScopeCycleCounterUObject AnimScope(GetAnimInstanceObject());
+
+	// update native update
+	{
+		SCOPE_CYCLE_COUNTER(STAT_NativeUpdateAnimation);
+		Update(CurrentDeltaSeconds);
+	}
+
+	// update all nodes
+	UpdateAnimationNode(CurrentDeltaSeconds);
+
+	// tick all our active asset players
+	TickAssetPlayerInstances(CurrentDeltaSeconds);
+}
+
+void FAnimInstanceProxy::EvaluateAnimation(FPoseContext& Output)
+{
+	ANIM_MT_SCOPE_CYCLE_COUNTER(EvaluateAnimInstance, !IsInGameThread());
+
+	// If bone caches have been invalidated, have AnimNodes refresh those.
+	if( bBoneCachesInvalidated && RootNode )
+	{
+		bBoneCachesInvalidated = false;
+
+		CachedBonesCounter.Increment();
+		FAnimationCacheBonesContext Proxy(this);
+		RootNode->CacheBones(Proxy);
+	}
+
+	// Evaluate native code if implemented, otherwise evaluate the node graph
+	if (!Evaluate(Output))
+	{
+		if (RootNode != NULL)
+		{
+			ANIM_MT_SCOPE_CYCLE_COUNTER(EvaluateAnimGraph, !IsInGameThread());
+			EvaluationCounter.Increment();
+			RootNode->Evaluate(Output);
+		}
+		else
+		{
+			Output.ResetToRefPose();
+		}
+	}
+}
+
+void FAnimInstanceProxy::UpdateMorphTargetCurves(const TMap<FName, float>& InMorphTargetCurves)
+{
+	QUICK_SCOPE_CYCLE_COUNTER(STAT_AnimInstanceProxy_UpdateMorphTargetCurves);
+
+	// now all tick/trigger/kismet is done
+	// add MorphTarget Curves from Kismet driven or any other source
+	// and overwrite if it exists
+	// Tick always should maintain this list, not Evaluate
+	for( auto Iter = InMorphTargetCurves.CreateConstIterator(); Iter; ++Iter )
+	{
+		float *CurveValPtr = MorphTargetCurves.Find(Iter.Key());
+		if ( CurveValPtr )
+		{
+			// override the value if Kismet request was made
+			*CurveValPtr = Iter.Value();
+		}
+		else
+		{
+			MorphTargetCurves.Add(Iter.Key(), Iter.Value());
+		}
+	}
+}
+
+void FAnimInstanceProxy::UpdateComponentsMaterialParameters(UPrimitiveComponent* Component)
+{
+	if(MaterialParameterCurves.Num() > 0)
+	{
+		QUICK_SCOPE_CYCLE_COUNTER(STAT_FAnimInstanceProxy_UpdateComponentsMaterialParameters);
+
+		for(auto Iter = MaterialParameterCurves.CreateConstIterator(); Iter; ++Iter)
+		{
+			FName ParameterName = Iter.Key();
+			float ParameterValue = Iter.Value();
+
+			UE_LOG(LogAnimation, Verbose, TEXT("Material Parameter change by Animation (%s : %0.2f)"), *ParameterName.ToString(), ParameterValue);
+			for(int32 MaterialIndex = 0; MaterialIndex < Component->GetNumMaterials(); ++MaterialIndex)
+			{
+				UMaterialInterface* MaterialInterface = Component->GetMaterial(MaterialIndex);
+				if(MaterialInterface)
+				{
+					float TestValue; //not used but needed for GetScalarParameterValue call
+					if(MaterialInterface->GetScalarParameterValue(ParameterName, TestValue))
+					{
+						UMaterialInstanceDynamic* DynamicMaterial = Cast<UMaterialInstanceDynamic>(MaterialInterface);
+						if(!DynamicMaterial) //Is it already a UMaterialInstanceDynamic (ie we used it last tick)
+						{
+							DynamicMaterial = Component->CreateAndSetMaterialInstanceDynamic(MaterialIndex);
+						}
+						DynamicMaterial->SetScalarParameterValue(ParameterName, ParameterValue);
+
+						// we don't break here because we can have multiple materials wanted to be driven by same parameter
+					}
+				}
+			}
+		}
+	}
+}
+
+// for now disable becauase it will not work with single node instance
+#if (UE_BUILD_SHIPPING || UE_BUILD_TEST)
+#define DEBUG_MONTAGEINSTANCE_WEIGHT 0
+#else
+#define DEBUG_MONTAGEINSTANCE_WEIGHT 1
+#endif // !(UE_BUILD_SHIPPING || UE_BUILD_TEST)
+
+void FAnimInstanceProxy::SlotEvaluatePose(FName SlotNodeName, const FCompactPose& SourcePose, const FBlendedCurve& SourceCurve, float InSourceWeight, FCompactPose& BlendedPose, FBlendedCurve& BlendedCurve, float InBlendWeight, float InTotalNodeWeight)
+{
+	//Accessing MontageInstances from this function is not safe (as this can be called during Parallel Anim Evaluation!
+	//Any montage data you need to add should be part of MontageEvaluationData
+	// nothing to blend, just get it out
+	if (InBlendWeight <= ZERO_ANIMWEIGHT_THRESH)
+	{
+		BlendedPose = SourcePose;
+		BlendedCurve = SourceCurve;
+		return;
+	}
+
+	// Split our data into additive and non additive.
+	TArray<FSlotEvaluationPose> AdditivePoses;
+	TArray<FSlotEvaluationPose> NonAdditivePoses;
+
+	// first pass we go through collect weights and valid montages. 
+#if DEBUG_MONTAGEINSTANCE_WEIGHT
+	float TotalWeight = 0.f;
+#endif // DEBUG_MONTAGEINSTANCE_WEIGHT
+	for (const FMontageEvaluationState& EvalState : MontageEvaluationData)
+	{
+		if (EvalState.Montage->IsValidSlot(SlotNodeName))
+		{
+			FAnimTrack const* const AnimTrack = EvalState.Montage->GetAnimationData(SlotNodeName);
+
+			// Find out additive type for pose.
+			EAdditiveAnimationType const AdditiveAnimType = AnimTrack->IsAdditive() 
+				? (AnimTrack->IsRotationOffsetAdditive() ? AAT_RotationOffsetMeshSpace : AAT_LocalSpaceBase)
+				: AAT_None;
+
+			FSlotEvaluationPose NewPose(EvalState.MontageWeight, AdditiveAnimType);
+			
+			// Bone array has to be allocated prior to calling GetPoseFromAnimTrack
+			NewPose.Pose.SetBoneContainer(&RequiredBones);
+			NewPose.Curve.InitFrom(GetSkeleton());
+
+			// Extract pose from Track
+			FAnimExtractContext ExtractionContext(EvalState.MontagePosition, EvalState.Montage->HasRootMotion() && RootMotionMode != ERootMotionMode::NoRootMotionExtraction);
+			AnimTrack->GetAnimationPose(NewPose.Pose, NewPose.Curve, ExtractionContext);
+
+			// add montage curves 
+			FBlendedCurve MontageCurve;
+			MontageCurve.InitFrom(GetSkeleton());
+			EvalState.Montage->EvaluateCurveData(MontageCurve, EvalState.MontagePosition);
+			NewPose.Curve.Combine(MontageCurve);
+
+#if DEBUG_MONTAGEINSTANCE_WEIGHT
+			TotalWeight += EvalState.MontageWeight;
+#endif // DEBUG_MONTAGEINSTANCE_WEIGHT
+			if (AdditiveAnimType == AAT_None)
+			{
+				NonAdditivePoses.Add(NewPose);
+			}
+			else
+			{
+				AdditivePoses.Add(NewPose);
+			}
+		}
+	}
+
+	// allocate for blending
+	// If source has any weight, add it to the blend array.
+	float const SourceWeight = FMath::Clamp<float>(InSourceWeight, 0.f, 1.f);
+
+#if DEBUG_MONTAGEINSTANCE_WEIGHT
+	ensure (FMath::IsNearlyEqual(InTotalNodeWeight, TotalWeight, KINDA_SMALL_NUMBER));
+#endif // DEBUG_MONTAGEINSTANCE_WEIGHT
+	ensure (InTotalNodeWeight > ZERO_ANIMWEIGHT_THRESH);
+
+	if (InTotalNodeWeight > (1.f + ZERO_ANIMWEIGHT_THRESH))
+	{
+		// Re-normalize additive poses
+		for (int32 Index = 0; Index < AdditivePoses.Num(); Index++)
+		{
+			AdditivePoses[Index].Weight /= InTotalNodeWeight;
+		}
+		// Re-normalize non-additive poses
+		for (int32 Index = 0; Index < NonAdditivePoses.Num(); Index++)
+		{
+			NonAdditivePoses[Index].Weight /= InTotalNodeWeight;
+		}
+	}
+
+	// Make sure we have at least one montage here.
+	check((AdditivePoses.Num() > 0) || (NonAdditivePoses.Num() > 0));
+
+	// Second pass, blend non additive poses together
+	{
+		// If we're only playing additive animations, just copy source for base pose.
+		if (NonAdditivePoses.Num() == 0)
+		{
+			BlendedPose = SourcePose;
+			BlendedCurve = SourceCurve;
+		}
+		// Otherwise we need to blend non additive poses together
+		else
+		{
+			int32 const NumPoses = NonAdditivePoses.Num() + ((SourceWeight > ZERO_ANIMWEIGHT_THRESH) ? 1 : 0);
+
+			TArray<const FCompactPose*, TInlineAllocator<8>> BlendingPoses;
+			BlendingPoses.AddUninitialized(NumPoses);
+
+			TArray<float, TInlineAllocator<8>> BlendWeights;
+			BlendWeights.AddUninitialized(NumPoses);
+
+			TArray<const FBlendedCurve*, TInlineAllocator<8>> BlendingCurves;
+			BlendingCurves.AddUninitialized(NumPoses);
+
+			for (int32 Index = 0; Index < NonAdditivePoses.Num(); Index++)
+			{
+				BlendingPoses[Index] = &NonAdditivePoses[Index].Pose;
+				BlendingCurves[Index] = &NonAdditivePoses[Index].Curve;
+				BlendWeights[Index] = NonAdditivePoses[Index].Weight;
+			}
+
+			if (SourceWeight > ZERO_ANIMWEIGHT_THRESH)
+			{
+				int32 const SourceIndex = BlendWeights.Num() - 1;
+				BlendingPoses[SourceIndex] = &SourcePose;
+				BlendingCurves[SourceIndex] = &SourceCurve;
+				BlendWeights[SourceIndex] = SourceWeight;
+			}
+
+			// now time to blend all montages
+			FAnimationRuntime::BlendPosesTogetherIndirect(BlendingPoses, BlendingCurves, BlendWeights, BlendedPose, BlendedCurve);
+		}
+	}
+
+	// Third pass, layer on weighted additive poses.
+	if (AdditivePoses.Num() > 0)
+	{
+		for (int32 Index = 0; Index < AdditivePoses.Num(); Index++)
+		{
+			FSlotEvaluationPose const& AdditivePose = AdditivePoses[Index];
+			FAnimationRuntime::AccumulateAdditivePose(BlendedPose, AdditivePose.Pose, BlendedCurve, AdditivePose.Curve, AdditivePose.Weight, AdditivePose.AdditiveType);
+		}
+	}
+
+	// Normalize rotations after blending/accumulation
+	BlendedPose.NormalizeRotations();
+}
+
+//to debug montage weight
+#define DEBUGMONTAGEWEIGHT 0
+
+void FAnimInstanceProxy::GetSlotWeight(FName const& SlotNodeName, float& out_SlotNodeWeight, float& out_SourceWeight, float& out_TotalNodeWeight) const
+{
+	// node total weight 
+	float NewSlotNodeWeight = 0.f;
+	// this is required to track, because it will be 1-SourceWeight
+	// if additive, it can be applied more
+	float NonAdditiveTotalWeight = 0.f;
+
+#if DEBUGMONTAGEWEIGHT
+	float TotalDesiredWeight = 0.f;
+#endif
+	// first get all the montage instance weight this slot node has
+	for (const FMontageEvaluationState& EvalState : MontageEvaluationData)
+	{
+		if (EvalState.Montage->IsValidSlot(SlotNodeName))
+		{
+			NewSlotNodeWeight += EvalState.MontageWeight;
+
+			if( !EvalState.Montage->IsValidAdditiveSlot(SlotNodeName) )
+			{
+				NonAdditiveTotalWeight += EvalState.MontageWeight;
+			}
+
+#if DEBUGMONTAGEWEIGHT			
+			TotalDesiredWeight += EvalState->DesiredWeight;
+#endif
+#if !NO_LOGGING
+			UE_LOG(LogAnimation, Verbose, TEXT("GetSlotWeight : Owner: %s, AnimMontage: %s,  (DesiredWeight:%0.2f, Weight:%0.2f)"),
+						*GetActorName(), *EvalState.Montage->GetName(), EvalState.DesiredWeight, EvalState.MontageWeight);
+#endif
+		}
+	}
+
+	// save the total node weight, it can be more than 1
+	// we need this so that when we eval, we normalized by this weight
+	// calculating there can cause inconsistency if some data changes
+	out_TotalNodeWeight = NewSlotNodeWeight;
+
+	// this can happen when it's blending in OR when newer animation comes in with shorter blendtime
+	// say #1 animation was blending out time with current blendtime 1.0 #2 animation was blending in with 1.0 (old) but got blend out with new blendtime 0.2f
+	// #3 animation was blending in with the new blendtime 0.2f, you'll have sum of #1, 2, 3 exceeds 1.f
+	if (NewSlotNodeWeight > (1.f + ZERO_ANIMWEIGHT_THRESH))
+	{
+		// you don't want to change weight of montage instance since it can play multiple slots
+		// if you change one, it will apply to all slots in that montage
+		// instead we should renormalize when we eval
+		// this should happen in the eval phase
+		NonAdditiveTotalWeight /= NewSlotNodeWeight;
+		// since we normalized, we reset
+		NewSlotNodeWeight = 1.f;
+	}
+#if DEBUGMONTAGEWEIGHT
+	else if (TotalDesiredWeight == 1.f && TotalSum < 1.f - ZERO_ANIMWEIGHT_THRESH)
+	{
+		// this can happen when it's blending in OR when newer animation comes in with longer blendtime
+		// say #1 animation was blending out time with current blendtime 0.2 #2 animation was blending in with 0.2 (old) but got blend out with new blendtime 1.f
+		// #3 animation was blending in with the new blendtime 1.f, you'll have sum of #1, 2, 3 doesn't meet 1.f
+		UE_LOG(LogAnimation, Warning, TEXT("[%s] Montage has less weight. Blending in?(%f)"), *SlotNodeName.ToString(), TotalSum);
+	}
+#endif
+
+	out_SlotNodeWeight = NewSlotNodeWeight;
+	out_SourceWeight = 1.f - NonAdditiveTotalWeight;
+}
+
+const FMontageEvaluationState* FAnimInstanceProxy::GetActiveMontageEvaluationState() const
+{
+	// Start from end, as most recent instances are added at the end of the queue.
+	int32 const NumInstances = MontageEvaluationData.Num();
+	for (int32 InstanceIndex = NumInstances - 1; InstanceIndex >= 0; InstanceIndex--)
+	{
+		const FMontageEvaluationState& EvaluationData = MontageEvaluationData[InstanceIndex];
+		if (EvaluationData.bIsActive)
+		{
+			return &EvaluationData;
+		}
+	}
+
+	return nullptr;
+}
+
+void FAnimInstanceProxy::AddVertexAnim(const FActiveVertexAnim& VertexAnim)
+{
+	VertexAnims.Add(VertexAnim);
+}
+
+void FAnimInstanceProxy::GatherDebugData(FNodeDebugData& DebugData)
+{
+	if(RootNode != nullptr)
+	{
+		 RootNode->GatherDebugData(DebugData); 
+	}
+}
+
+float FAnimInstanceProxy::GetInstanceAssetPlayerLength(int32 AssetPlayerIndex)
+{
+	if(FAnimNode_AssetPlayerBase* PlayerNode = GetNodeFromIndex<FAnimNode_AssetPlayerBase>(AssetPlayerIndex))
+	{
+		return PlayerNode->GetAnimAsset()->GetMaxCurrentTime();
+	}
+
+	return 0.0f;
+}
+
+float FAnimInstanceProxy::GetInstanceAssetPlayerTime(int32 AssetPlayerIndex)
+{
+	if(FAnimNode_AssetPlayerBase* PlayerNode = GetNodeFromIndex<FAnimNode_AssetPlayerBase>(AssetPlayerIndex))
+	{
+		return PlayerNode->GetAccumulatedTime();
+	}
+
+	return 0.0f;
+}
+
+float FAnimInstanceProxy::GetInstanceAssetPlayerTimeFraction(int32 AssetPlayerIndex)
+{
+	if(FAnimNode_AssetPlayerBase* PlayerNode = GetNodeFromIndex<FAnimNode_AssetPlayerBase>(AssetPlayerIndex))
+	{
+		UAnimationAsset* Asset = PlayerNode->GetAnimAsset();
+
+		float Length = Asset ? Asset->GetMaxCurrentTime() : 0.0f;
+
+		if(Length > 0.0f)
+		{
+			return PlayerNode->GetAccumulatedTime() / Length;
+		}
+	}
+
+	return 0.0f;
+}
+
+float FAnimInstanceProxy::GetInstanceAssetPlayerTimeFromEndFraction(int32 AssetPlayerIndex)
+{
+	if(FAnimNode_AssetPlayerBase* PlayerNode = GetNodeFromIndex<FAnimNode_AssetPlayerBase>(AssetPlayerIndex))
+	{
+		UAnimationAsset* Asset = PlayerNode->GetAnimAsset();
+		
+		float Length = Asset ? Asset->GetMaxCurrentTime() : 0.f;
+
+		if(Length > 0.f)
+		{
+			return (Length - PlayerNode->GetAccumulatedTime()) / Length;
+		}
+	}
+
+	return 1.0f;
+}
+
+float FAnimInstanceProxy::GetInstanceAssetPlayerTimeFromEnd(int32 AssetPlayerIndex)
+{
+	if(FAnimNode_AssetPlayerBase* PlayerNode = GetNodeFromIndex<FAnimNode_AssetPlayerBase>(AssetPlayerIndex))
+	{
+		UAnimationAsset* Asset = PlayerNode->GetAnimAsset();
+		if(Asset)
+		{
+			return PlayerNode->GetAnimAsset()->GetMaxCurrentTime() - PlayerNode->GetAccumulatedTime();
+		}
+	}
+
+	return MAX_flt;
+}
+
+float FAnimInstanceProxy::GetInstanceStateWeight(int32 MachineIndex, int32 StateIndex)
+{
+	if(FAnimNode_StateMachine* MachineInstance = GetStateMachineInstance(MachineIndex))
+	{
+		return MachineInstance->GetStateWeight(StateIndex);
+	}
+	
+	return 0.0f;
+}
+
+float FAnimInstanceProxy::GetInstanceCurrentStateElapsedTime(int32 MachineIndex)
+{
+	if(FAnimNode_StateMachine* MachineInstance = GetStateMachineInstance(MachineIndex))
+	{
+		return MachineInstance->GetCurrentStateElapsedTime();
+	}
+
+	return 0.0f;
+}
+
+float FAnimInstanceProxy::GetInstanceTransitionCrossfadeDuration(int32 MachineIndex, int32 TransitionIndex)
+{
+	if(FAnimNode_StateMachine* MachineInstance = GetStateMachineInstance(MachineIndex))
+	{
+		if(MachineInstance->IsValidTransitionIndex(TransitionIndex))
+		{
+			return MachineInstance->GetTransitionInfo(TransitionIndex).CrossfadeDuration;
+		}
+	}
+
+	return 0.0f;
+}
+
+float FAnimInstanceProxy::GetInstanceTransitionTimeElapsed(int32 MachineIndex, int32 TransitionIndex)
+{
+	// Just an alias for readability in the anim graph
+	return GetInstanceCurrentStateElapsedTime(MachineIndex);
+}
+
+float FAnimInstanceProxy::GetInstanceTransitionTimeElapsedFraction(int32 MachineIndex, int32 TransitionIndex)
+{
+	if(FAnimNode_StateMachine* MachineInstance = GetStateMachineInstance(MachineIndex))
+	{
+		if(MachineInstance->IsValidTransitionIndex(TransitionIndex))
+		{
+			return MachineInstance->GetCurrentStateElapsedTime() / MachineInstance->GetTransitionInfo(TransitionIndex).CrossfadeDuration;
+		}
+	}
+
+	return 0.0f;
+}
+
+float FAnimInstanceProxy::GetRelevantAnimTimeRemaining(int32 MachineIndex, int32 StateIndex)
+{
+	if(FAnimNode_AssetPlayerBase* AssetPlayer = GetRelevantAssetPlayerFromState(MachineIndex, StateIndex))
+	{
+		if(AssetPlayer->GetAnimAsset())
+		{
+			return AssetPlayer->GetAnimAsset()->GetMaxCurrentTime() - AssetPlayer->GetAccumulatedTime();
+		}
+	}
+
+	return MAX_flt;
+}
+
+float FAnimInstanceProxy::GetRelevantAnimTimeRemainingFraction(int32 MachineIndex, int32 StateIndex)
+{
+	if(FAnimNode_AssetPlayerBase* AssetPlayer = GetRelevantAssetPlayerFromState(MachineIndex, StateIndex))
+	{
+		if(AssetPlayer->GetAnimAsset())
+		{
+			float Length = AssetPlayer->GetAnimAsset()->GetMaxCurrentTime();
+			if(Length > 0.0f)
+			{
+				return (Length - AssetPlayer->GetAccumulatedTime()) / Length;
+			}
+		}
+	}
+
+	return 1.0f;
+}
+
+float FAnimInstanceProxy::GetRelevantAnimLength(int32 MachineIndex, int32 StateIndex)
+{
+	if(FAnimNode_AssetPlayerBase* AssetPlayer = GetRelevantAssetPlayerFromState(MachineIndex, StateIndex))
+	{
+		if(AssetPlayer->GetAnimAsset())
+		{
+			return AssetPlayer->GetAnimAsset()->GetMaxCurrentTime();
+		}
+	}
+
+	return 0.0f;
+}
+
+float FAnimInstanceProxy::GetRelevantAnimTime(int32 MachineIndex, int32 StateIndex)
+{
+	if(FAnimNode_AssetPlayerBase* AssetPlayer = GetRelevantAssetPlayerFromState(MachineIndex, StateIndex))
+	{
+		return AssetPlayer->GetAccumulatedTime();
+	}
+
+	return 0.0f;
+}
+
+float FAnimInstanceProxy::GetRelevantAnimTimeFraction(int32 MachineIndex, int32 StateIndex)
+{
+	if(FAnimNode_AssetPlayerBase* AssetPlayer = GetRelevantAssetPlayerFromState(MachineIndex, StateIndex))
+	{
+		float Length = AssetPlayer->GetAnimAsset()->GetMaxCurrentTime();
+		if(Length > 0.0f)
+		{
+			return AssetPlayer->GetAccumulatedTime() / Length;
+		}
+	}
+
+	return 0.0f;
+}
+
+FAnimNode_AssetPlayerBase* FAnimInstanceProxy::GetRelevantAssetPlayerFromState(int32 MachineIndex, int32 StateIndex)
+{
+	FAnimNode_AssetPlayerBase* ResultPlayer = nullptr;
+	if(FAnimNode_StateMachine* MachineInstance = GetStateMachineInstance(MachineIndex))
+	{
+		float MaxWeight = 0.0f;
+		const FBakedAnimationState& State = MachineInstance->GetStateInfo(StateIndex);
+		for(const int32& PlayerIdx : State.PlayerNodeIndices)
+		{
+			if(FAnimNode_AssetPlayerBase* Player = GetNodeFromIndex<FAnimNode_AssetPlayerBase>(PlayerIdx))
+			{
+				if(!Player->bIgnoreForRelevancyTest && Player->GetCachedBlendWeight() > MaxWeight)
+				{
+					MaxWeight = Player->GetCachedBlendWeight();
+					ResultPlayer = Player;
+				}
+			}
+		}
+	}
+	return ResultPlayer;
+}
+
+FAnimNode_StateMachine* FAnimInstanceProxy::GetStateMachineInstance(int32 MachineIndex)
+{
+	if(AnimBlueprintClass)
+	{
+		if((MachineIndex >= 0) && (MachineIndex < AnimBlueprintClass->AnimNodeProperties.Num()))
+		{
+			const int32 InstancePropertyIndex = AnimBlueprintClass->AnimNodeProperties.Num() - 1 - MachineIndex;
+
+			UStructProperty* MachineInstanceProperty = AnimBlueprintClass->AnimNodeProperties[InstancePropertyIndex];
+			checkSlow(MachineInstanceProperty->Struct->IsChildOf(FAnimNode_StateMachine::StaticStruct()));
+
+			return MachineInstanceProperty->ContainerPtrToValuePtr<FAnimNode_StateMachine>(AnimInstanceObject);
+		}
+	}
+
+	return nullptr;
+}
+
+void FAnimInstanceProxy::AddNativeTransitionBinding(const FName& MachineName, const FName& PrevStateName, const FName& NextStateName, const FCanTakeTransition& NativeTransitionDelegate, const FName& TransitionName)
+{
+	NativeTransitionBindings.Add(FNativeTransitionBinding(MachineName, PrevStateName, NextStateName, NativeTransitionDelegate, TransitionName));
+}
+
+bool FAnimInstanceProxy::HasNativeTransitionBinding(const FName& MachineName, const FName& PrevStateName, const FName& NextStateName, FName& OutBindingName)
+{
+	for(const auto& Binding : NativeTransitionBindings)
+	{
+		if(Binding.MachineName == MachineName && Binding.PreviousStateName == PrevStateName && Binding.NextStateName == NextStateName)
+		{
+			IDelegateInstance* DelegateInstance = Binding.NativeTransitionDelegate.GetDelegateInstance();
+			if(DelegateInstance)
+			{
+				OutBindingName = DelegateInstance->GetFunctionName();
+			}
+#if WITH_EDITORONLY_DATA
+			// if the function wasnt forthcoming with a name, use the explicit name supplied in the binding
+			if(OutBindingName == NAME_None)
+			{
+				OutBindingName = Binding.TransitionName;
+			}
+#endif
+			return true;
+		}
+	}
+
+	return false;
+}
+
+void FAnimInstanceProxy::AddNativeStateEntryBinding(const FName& MachineName, const FName& StateName, const FOnGraphStateChanged& NativeEnteredDelegate)
+{
+	NativeStateEntryBindings.Add(FNativeStateBinding(MachineName, StateName, NativeEnteredDelegate));
+}
+	
+bool FAnimInstanceProxy::HasNativeStateEntryBinding(const FName& MachineName, const FName& StateName, FName& OutBindingName)
+{
+	for(const auto& Binding : NativeStateEntryBindings)
+	{
+		if(Binding.MachineName == MachineName && Binding.StateName == StateName)
+		{
+			IDelegateInstance* DelegateInstance = Binding.NativeStateDelegate.GetDelegateInstance();
+			if(DelegateInstance)
+			{
+				OutBindingName = DelegateInstance->GetFunctionName();
+			}
+			return true;
+		}
+	}
+
+	return false;
+}
+
+void FAnimInstanceProxy::AddNativeStateExitBinding(const FName& MachineName, const FName& StateName, const FOnGraphStateChanged& NativeExitedDelegate)
+{
+	NativeStateExitBindings.Add(FNativeStateBinding(MachineName, StateName, NativeExitedDelegate));
+}
+
+bool FAnimInstanceProxy::HasNativeStateExitBinding(const FName& MachineName, const FName& StateName, FName& OutBindingName)
+{
+	for(const auto& Binding : NativeStateExitBindings)
+	{
+		if(Binding.MachineName == MachineName && Binding.StateName == StateName)
+		{
+			IDelegateInstance* DelegateInstance = Binding.NativeStateDelegate.GetDelegateInstance();
+			if(DelegateInstance)
+			{
+				OutBindingName = DelegateInstance->GetFunctionName();
+			}
+			return true;
+		}
+	}
+
+	return false;
+}
+
+void FAnimInstanceProxy::BindNativeDelegates()
+{
+	// if we have no root node, we are usually in error so early out
+	if(RootNode == nullptr)
+	{
+		return;
+	}
+
+	auto ForEachStateLambda = [&](UAnimBlueprintGeneratedClass* InAnimBlueprintClass, const FName& MachineName, const FName& StateName, TFunctionRef<void (FAnimNode_StateMachine*, const FBakedAnimationState&, int32)> Predicate)
+	{
+		for(UStructProperty* Property : InAnimBlueprintClass->AnimNodeProperties)
+		{
+			if(Property && Property->Struct == FAnimNode_StateMachine::StaticStruct())
+			{
+				FAnimNode_StateMachine* StateMachine = Property->ContainerPtrToValuePtr<FAnimNode_StateMachine>(AnimInstanceObject);
+				if(StateMachine)
+				{
+					const FBakedAnimationStateMachine* MachineDescription = GetMachineDescription(InAnimBlueprintClass, StateMachine);
+					if(MachineDescription && MachineName == MachineDescription->MachineName)
+					{
+						// check each state transition for a match
+						int32 StateIndex = 0;
+						for(const FBakedAnimationState& State : MachineDescription->States)
+						{
+							if(State.StateName == StateName)
+							{
+								Predicate(StateMachine, State, StateIndex);
+							}
+							StateIndex++;
+						}
+					}
+				}
+			}
+		}
+	};
+
+	if(AnimBlueprintClass)
+	{
+		// transition delegates
+		for(const auto& Binding : NativeTransitionBindings)
+		{
+			ForEachStateLambda(AnimBlueprintClass, Binding.MachineName, Binding.PreviousStateName, 
+				[&](FAnimNode_StateMachine* StateMachine, const FBakedAnimationState& State, int32 StateIndex)
+				{
+					for(const FBakedStateExitTransition& TransitionExit : State.Transitions)
+					{
+						if(TransitionExit.CanTakeDelegateIndex != INDEX_NONE)
+						{
+							// In case the state machine hasn't been initilized, we need to re-get the desc
+							const FBakedAnimationStateMachine* MachineDesc = GetMachineDescription(AnimBlueprintClass, StateMachine);
+							const FAnimationTransitionBetweenStates& Transition = MachineDesc->Transitions[TransitionExit.TransitionIndex];
+							const FBakedAnimationState& BakedState = MachineDesc->States[Transition.NextState];
+
+							if (BakedState.StateName == Binding.NextStateName)
+							{
+								FAnimNode_TransitionResult* ResultNode = GetNodeFromPropertyIndex<FAnimNode_TransitionResult>(AnimInstanceObject, AnimBlueprintClass, TransitionExit.CanTakeDelegateIndex);
+								if(ResultNode)
+								{
+									ResultNode->NativeTransitionDelegate = Binding.NativeTransitionDelegate;
+								}
+							}
+						}
+					}
+				});
+		}
+
+		// state entry delegates
+		for(const auto& Binding : NativeStateEntryBindings)
+		{
+			ForEachStateLambda(AnimBlueprintClass, Binding.MachineName, Binding.StateName, 
+				[&](FAnimNode_StateMachine* StateMachine, const FBakedAnimationState& State, int32 StateIndex)
+				{
+					// allocate enough space for all our states we need so far
+					StateMachine->OnGraphStatesEntered.SetNum(FMath::Max(StateIndex + 1, StateMachine->OnGraphStatesEntered.Num()));
+					StateMachine->OnGraphStatesEntered[StateIndex] = Binding.NativeStateDelegate;
+				});
+		}
+
+		// state exit delegates
+		for(const auto& Binding : NativeStateExitBindings)
+		{
+			ForEachStateLambda(AnimBlueprintClass, Binding.MachineName, Binding.StateName, 
+				[&](FAnimNode_StateMachine* StateMachine, const FBakedAnimationState& State, int32 StateIndex)
+				{
+					// allocate enough space for all our states we need so far
+					StateMachine->OnGraphStatesExited.SetNum(FMath::Max(StateIndex + 1, StateMachine->OnGraphStatesExited.Num()));
+					StateMachine->OnGraphStatesExited[StateIndex] = Binding.NativeStateDelegate;
+				});
+		}
+	}
+}
+
+const FBakedAnimationStateMachine* FAnimInstanceProxy::GetMachineDescription(UAnimBlueprintGeneratedClass* AnimBlueprintClass, FAnimNode_StateMachine* MachineInstance)
+{
+	return AnimBlueprintClass->BakedStateMachines.IsValidIndex(MachineInstance->StateMachineIndexInClass) ? &(AnimBlueprintClass->BakedStateMachines[MachineInstance->StateMachineIndexInClass]) : nullptr;
+}
+
+FAnimNode_StateMachine* FAnimInstanceProxy::GetStateMachineInstanceFromName(FName MachineName)
+{
+	if (AnimBlueprintClass)
+	{
+		for (int32 MachineIndex = 0; MachineIndex < AnimBlueprintClass->AnimNodeProperties.Num(); MachineIndex++)
+		{
+			UStructProperty* Property = AnimBlueprintClass->AnimNodeProperties[AnimBlueprintClass->AnimNodeProperties.Num() - 1 - MachineIndex];
+			if (Property && Property->Struct == FAnimNode_StateMachine::StaticStruct())
+			{
+				FAnimNode_StateMachine* StateMachine = Property->ContainerPtrToValuePtr<FAnimNode_StateMachine>(AnimInstanceObject);
+				if (StateMachine)
+				{
+					if (const FBakedAnimationStateMachine* MachineDescription = GetMachineDescription(AnimBlueprintClass, StateMachine))
+					{
+						if (MachineDescription->MachineName == MachineName)
+						{
+							return StateMachine;
+						}
+					}
+				}
+			}
+		}
+	}
+
+	return nullptr;
+}
+
+const FBakedAnimationStateMachine* FAnimInstanceProxy::GetStateMachineInstanceDesc(FName MachineName)
+{
+	if(AnimBlueprintClass)
+	{
+		for(int32 MachineIndex = 0; MachineIndex < AnimBlueprintClass->AnimNodeProperties.Num(); MachineIndex++)
+		{
+			UStructProperty* Property = AnimBlueprintClass->AnimNodeProperties[AnimBlueprintClass->AnimNodeProperties.Num() - 1 - MachineIndex];
+			if(Property && Property->Struct == FAnimNode_StateMachine::StaticStruct())
+			{
+				FAnimNode_StateMachine* StateMachine = Property->ContainerPtrToValuePtr<FAnimNode_StateMachine>(AnimInstanceObject);
+				if(StateMachine)
+				{
+					if(const FBakedAnimationStateMachine* MachineDescription = GetMachineDescription(AnimBlueprintClass, StateMachine))
+					{
+						if(MachineDescription->MachineName == MachineName)
+						{
+							return MachineDescription;
+						}
+					}
+				}
+			}
+		}
+	}
+
+	return nullptr;
+}
+
+int32 FAnimInstanceProxy::GetStateMachineIndex(FName MachineName)
+{
+	if(AnimBlueprintClass)
+	{
+		for(int32 MachineIndex = 0; MachineIndex < AnimBlueprintClass->AnimNodeProperties.Num(); MachineIndex++)
+		{
+			UStructProperty* Property = AnimBlueprintClass->AnimNodeProperties[AnimBlueprintClass->AnimNodeProperties.Num() - 1 - MachineIndex];
+			if(Property && Property->Struct == FAnimNode_StateMachine::StaticStruct())
+			{
+				FAnimNode_StateMachine* StateMachine = Property->ContainerPtrToValuePtr<FAnimNode_StateMachine>(AnimInstanceObject);
+				if(StateMachine)
+				{
+					if(const FBakedAnimationStateMachine* MachineDescription = GetMachineDescription(AnimBlueprintClass, StateMachine))
+					{
+						if(MachineDescription->MachineName == MachineName)
+						{
+							return MachineIndex;
+						}
+					}
+				}
+			}
+		}
+	}
+
+	return INDEX_NONE;
+}
+
+int32 FAnimInstanceProxy::GetInstanceAssetPlayerIndex(FName MachineName, FName StateName, FName AssetName)
+{
+	if(AnimBlueprintClass)
+	{
+		if(const FBakedAnimationStateMachine* MachineDescription = GetStateMachineInstanceDesc(MachineName))
+		{
+			for(int32 StateIndex = 0; StateIndex < MachineDescription->States.Num(); StateIndex++)
+			{
+				const FBakedAnimationState& State = MachineDescription->States[StateIndex];
+				if(State.StateName == StateName)
+				{
+					for(int32 PlayerIndex = 0; PlayerIndex < State.PlayerNodeIndices.Num(); PlayerIndex++)
+					{
+						checkSlow(State.PlayerNodeIndices[PlayerIndex] < AnimBlueprintClass->AnimNodeProperties.Num());
+						UStructProperty* AssetPlayerProperty = AnimBlueprintClass->AnimNodeProperties[AnimBlueprintClass->AnimNodeProperties.Num() - 1 - State.PlayerNodeIndices[PlayerIndex]];
+						if(AssetPlayerProperty && AssetPlayerProperty->Struct->IsChildOf(FAnimNode_AssetPlayerBase::StaticStruct()))
+						{
+							FAnimNode_AssetPlayerBase* AssetPlayer = AssetPlayerProperty->ContainerPtrToValuePtr<FAnimNode_AssetPlayerBase>(AnimInstanceObject);
+							if(AssetPlayer)
+							{
+								if(AssetName == NAME_None || AssetPlayer->GetAnimAsset()->GetFName() == AssetName)
+								{
+									return State.PlayerNodeIndices[PlayerIndex];
+								}
+							}
+						}
+					}
+				}
+			}
+		}
+	}
+
+	return INDEX_NONE;
+}
+
+#undef LOCTEXT_NAMESPACE
\ No newline at end of file
diff --git a/Engine/Source/Runtime/Engine/Private/Animation/AnimMontage.cpp b/Engine/Source/Runtime/Engine/Private/Animation/AnimMontage.cpp
--- a/Engine/Source/Runtime/Engine/Private/Animation/AnimMontage.cpp
+++ b/Engine/Source/Runtime/Engine/Private/Animation/AnimMontage.cpp
@@ -1012,7 +1012,7 @@ ENGINE_API void UAnimMontage::UpdateLinkableElements(int32 SlotIdx, int32 Segmen
 
 #endif
 
-void UAnimMontage::TickAssetPlayerInstance(FAnimTickRecord& Instance, class UAnimInstance* InstanceOwner, FAnimAssetTickContext& Context) const
+void UAnimMontage::TickAssetPlayer(FAnimTickRecord& Instance, struct FAnimNotifyQueue& NotifyQueue, FAnimAssetTickContext& Context) const
 {
 	// nothing has to happen here
 	// we just have to make sure we set Context data correct
@@ -1850,7 +1850,7 @@ void FAnimMontageInstance::HandleEvents(float PreviousTrackPos, float CurrentTra
 		}
 
 		// Queue all these notifies.
-		AnimInstance->AddAnimNotifies(Notifies, NotifyWeight);
+		AnimInstance->NotifyQueue.AddAnimNotifies(Notifies, NotifyWeight);
 	}
 
 	// Update active state branching points, before we handle the immediate tick marker.
@@ -1958,18 +1958,18 @@ void FAnimMontageInstance::SetMatineeAnimPositionInner(FName SlotName, USkeletal
 		UAnimSingleNodeInstance* SingleNodeInst = SkeletalMeshComponent->GetSingleNodeInstance();
 		if(SingleNodeInst)
 		{
-			if(SingleNodeInst->CurrentAsset != InAnimSequence)
+			if(SingleNodeInst->GetCurrentAsset() != InAnimSequence)
 			{
 				SingleNodeInst->SetAnimationAsset(InAnimSequence, bLooping);
 				SingleNodeInst->SetPosition(0.0f);
-				SingleNodeInst->bPlaying = false;
+				SingleNodeInst->SetPlaying(false);
 			}
 
-			if(SingleNodeInst->bLooping!=bLooping)
+			if(SingleNodeInst->IsLooping()!=bLooping)
 			{
 				SingleNodeInst->SetLooping(bLooping);
 			}
-			if(SingleNodeInst->CurrentTime != InPosition)
+			if(SingleNodeInst->GetCurrentTime() != InPosition)
 			{
 				SingleNodeInst->SetPosition(InPosition);
 			}
@@ -2025,7 +2025,7 @@ void FAnimMontageInstance::PreviewMatineeSetAnimPositionInner(FName SlotName, US
 	UAnimSingleNodeInstance * SingleNodeInst = SkeletalMeshComponent->GetSingleNodeInstance();
 	if(SingleNodeInst)
 	{
-		if(SingleNodeInst->CurrentAsset != InAnimSequence)
+		if(SingleNodeInst->GetCurrentAsset() != InAnimSequence)
 		{
 			SingleNodeInst->SetAnimationAsset(InAnimSequence, bLooping);
 		}
diff --git a/Engine/Source/Runtime/Engine/Private/Animation/AnimNodeBase.cpp b/Engine/Source/Runtime/Engine/Private/Animation/AnimNodeBase.cpp
--- a/Engine/Source/Runtime/Engine/Private/Animation/AnimNodeBase.cpp
+++ b/Engine/Source/Runtime/Engine/Private/Animation/AnimNodeBase.cpp
@@ -2,13 +2,85 @@
 
 #include "EnginePrivate.h"
 #include "Animation/AnimNodeBase.h"
+#include "Animation/AnimInstanceProxy.h"
+
+/////////////////////////////////////////////////////
+// FAnimationBaseContext
+
+// @todo: remove after deprecation
+FAnimationBaseContext::FAnimationBaseContext(UAnimInstance* InAnimInstance)
+	: AnimInstanceProxy(&InAnimInstance->GetProxyOnAnyThread())
+{
+	PRAGMA_DISABLE_DEPRECATION_WARNINGS
+	AnimInstance = InAnimInstance;
+	PRAGMA_ENABLE_DEPRECATION_WARNINGS
+}
+
+FAnimationBaseContext::FAnimationBaseContext(FAnimInstanceProxy* InAnimInstanceProxy)
+	: AnimInstanceProxy(InAnimInstanceProxy)
+{
+	// @todo: remove after deprecation
+	PRAGMA_DISABLE_DEPRECATION_WARNINGS
+	AnimInstance = CastChecked<UAnimInstance>(AnimInstanceProxy->GetAnimInstanceObject());
+	PRAGMA_ENABLE_DEPRECATION_WARNINGS
+}
+
+FAnimationBaseContext::FAnimationBaseContext(const FAnimationBaseContext& InContext)
+{
+	// @todo: remove after deprecation
+	PRAGMA_DISABLE_DEPRECATION_WARNINGS
+	AnimInstance = InContext.AnimInstance;
+	PRAGMA_ENABLE_DEPRECATION_WARNINGS
+	AnimInstanceProxy = InContext.AnimInstanceProxy;
+}
+
+UAnimBlueprintGeneratedClass* FAnimationBaseContext::GetAnimBlueprintClass() const
+{
+	return AnimInstanceProxy->GetAnimBlueprintClass();
+}
+
+#if WITH_EDITORONLY_DATA
+UAnimBlueprint* FAnimationBaseContext::GetAnimBlueprint() const
+{
+	if (UAnimBlueprintGeneratedClass* Class = AnimInstanceProxy->GetAnimBlueprintClass())
+	{
+		return Cast<UAnimBlueprint>(Class->ClassGeneratedBy);
+	}
+	return NULL;
+}
+#endif //WITH_EDITORONLY_DATA
+
+/////////////////////////////////////////////////////
+// FPoseContext
+
+void FPoseContext::Initialize(FAnimInstanceProxy* InAnimInstanceProxy)
+{
+	checkSlow(AnimInstanceProxy && AnimInstanceProxy->GetRequiredBones().IsValid());
+	Pose.SetBoneContainer(&AnimInstanceProxy->GetRequiredBones());
+	Curve.InitFrom(AnimInstanceProxy->GetSkeleton());
+}
+
+/////////////////////////////////////////////////////
+// FComponentSpacePoseContext
+
+void FComponentSpacePoseContext::ResetToRefPose()
+{
+	checkSlow(AnimInstanceProxy && AnimInstanceProxy->GetRequiredBones().IsValid());
+	Pose.InitPose(&AnimInstanceProxy->GetRequiredBones());
+	Curve.InitFrom(AnimInstanceProxy->GetSkeleton());
+}
 
 /////////////////////////////////////////////////////
 // FAnimNode_Base
 
 void FAnimNode_Base::Initialize(const FAnimationInitializeContext& Context)
 {
-	EvaluateGraphExposedInputs.Initialize(this, Context.AnimInstance);
+	EvaluateGraphExposedInputs.Initialize(this, Context.AnimInstanceProxy->GetAnimInstanceObject());
+}
+
+bool FAnimNode_Base::IsLODEnabled(FAnimInstanceProxy* AnimInstanceProxy, int32 InLODThreshold)
+{
+	return (InLODThreshold == INDEX_NONE || AnimInstanceProxy->GetSkelMeshComponent()->PredictedLODLevel <= InLODThreshold);
 }
 
 /////////////////////////////////////////////////////
@@ -26,7 +98,7 @@ void FPoseLinkBase::AttemptRelink(const FAnimationBaseContext& Context)
 		if ( ensure(AnimBlueprintClass->AnimNodeProperties.IsValidIndex(LinkID)) )
 		{
 			UProperty* LinkedProperty = AnimBlueprintClass->AnimNodeProperties[LinkID];
-			void* LinkedNodePtr = LinkedProperty->ContainerPtrToValuePtr<void>(Context.AnimInstance);
+			void* LinkedNodePtr = LinkedProperty->ContainerPtrToValuePtr<void>(Context.AnimInstanceProxy->GetAnimInstanceObject());
 			LinkedNode = (FAnimNode_Base*)LinkedNodePtr;
 		}
 	}
@@ -36,15 +108,14 @@ void FPoseLinkBase::Initialize(const FAnimationInitializeContext& Context)
 {
 #if DO_CHECK
 	checkf( !bProcessed, TEXT( "Initialize already in progress, circular link for AnimInstance [%s] Blueprint [%s]" ), \
-		Context.AnimInstance ? *Context.AnimInstance->GetFullName() : TEXT( "None" ), Context.GetAnimBlueprintClass() ? *Context.GetAnimBlueprintClass()->GetFullName() : TEXT( "None" ) );
+		*Context.AnimInstanceProxy->GetAnimInstanceName(), Context.GetAnimBlueprintClass() ? *Context.GetAnimBlueprintClass()->GetFullName() : TEXT( "None" ) );
 	TGuardValue<bool> CircularGuard(bProcessed, true);
 #endif
 
 	AttemptRelink(Context);
 
 #if ENABLE_ANIMGRAPH_TRAVERSAL_DEBUG
-	// 	checkf(!InitializationCounter.IsSynchronizedWith(Context.AnimInstance->InitializationCounter), TEXT("Already called Initialize on this node!"));
-	InitializationCounter.SynchronizeWith(Context.AnimInstance->InitializationCounter);
+	InitializationCounter.SynchronizeWith(Context.AnimInstanceProxy->GetInitializationCounter());
 #endif
 
 	// Do standard initialization
@@ -58,14 +129,12 @@ void FPoseLinkBase::CacheBones(const FAnimationCacheBonesContext& Context)
 {
 #if DO_CHECK
 	checkf( !bProcessed, TEXT( "CacheBones already in progress, circular link for AnimInstance [%s] Blueprint [%s]" ), \
-		Context.AnimInstance ? *Context.AnimInstance->GetFullName() : TEXT( "None" ), Context.GetAnimBlueprintClass() ? *Context.GetAnimBlueprintClass()->GetFullName() : TEXT( "None" ) );
+		*Context.AnimInstanceProxy->GetAnimInstanceName(), Context.GetAnimBlueprintClass() ? *Context.GetAnimBlueprintClass()->GetFullName() : TEXT( "None" ) );
 	TGuardValue<bool> CircularGuard(bProcessed, true);
 #endif
 
 #if ENABLE_ANIMGRAPH_TRAVERSAL_DEBUG
-	// 	checkf(InitializationCounter.IsSynchronizedWith(Context.AnimInstance->InitializationCounter), TEXT("Calling CachedBones without initialization!"));
-	// 	checkf(!CachedBonesCounter.IsSynchronizedWith(Context.AnimInstance->CachedBonesCounter), TEXT("Already called CachedBones for this node!"));
-	CachedBonesCounter.SynchronizeWith(Context.AnimInstance->CachedBonesCounter);
+	CachedBonesCounter.SynchronizeWith(Context.AnimInstanceProxy->GetCachedBonesCounter());
 #endif
 
 	if (LinkedNode != NULL)
@@ -79,7 +148,7 @@ void FPoseLinkBase::Update(const FAnimationUpdateContext& Context)
 	QUICK_SCOPE_CYCLE_COUNTER(STAT_FPoseLinkBase_Update);
 #if DO_CHECK
 	checkf( !bProcessed, TEXT( "Update already in progress, circular link for AnimInstance [%s] Blueprint [%s]" ), \
-		Context.AnimInstance ? *Context.AnimInstance->GetFullName() : TEXT( "None" ), Context.GetAnimBlueprintClass() ? *Context.GetAnimBlueprintClass()->GetFullName() : TEXT( "None" ) );
+		*Context.AnimInstanceProxy->GetAnimInstanceName(), Context.GetAnimBlueprintClass() ? *Context.GetAnimBlueprintClass()->GetFullName() : TEXT( "None" ) );
 	TGuardValue<bool> CircularGuard(bProcessed, true);
 #endif
 
@@ -95,22 +164,18 @@ void FPoseLinkBase::Update(const FAnimationUpdateContext& Context)
 		// Record the node line activation
 		if (LinkedNode != NULL)
 		{
-			//@TODO: Might like a cheaper way to check this, like a bool on the Anim Blueprint or the context
-			if (UAnimBlueprint* BP = Context.GetAnimBlueprint())
+			if (Context.AnimInstanceProxy->IsBeingDebugged())
 			{
-				if( Context.AnimInstance && (BP->GetObjectBeingDebugged() == Context.AnimInstance))
-				{
-					Context.GetAnimBlueprintClass()->GetAnimBlueprintDebugData().RecordNodeVisit(LinkID, SourceLinkID, Context.GetFinalBlendWeight());
-				}
+				Context.AnimInstanceProxy->RecordNodeVisit(LinkID, SourceLinkID, Context.GetFinalBlendWeight());
 			}
 		}
 	}
 #endif
 
 #if ENABLE_ANIMGRAPH_TRAVERSAL_DEBUG
-	checkf(InitializationCounter.IsSynchronizedWith(Context.AnimInstance->InitializationCounter), TEXT("Calling Update without initialization!"));
-	checkf(!UpdateCounter.IsSynchronizedWith(Context.AnimInstance->UpdateCounter), TEXT("Already called Update for this node!"));
-	UpdateCounter.SynchronizeWith(Context.AnimInstance->UpdateCounter);
+	checkf(InitializationCounter.IsSynchronizedWith(Context.AnimInstanceProxy->GetInitializationCounter()), TEXT("Calling Update without initialization!"));
+	checkf(!UpdateCounter.IsSynchronizedWith(Context.AnimInstanceProxy->GetUpdateCounter()), TEXT("Already called Update for this node!"));
+	UpdateCounter.SynchronizeWith(Context.AnimInstanceProxy->GetUpdateCounter());
 #endif
 
 	if (LinkedNode != NULL)
@@ -134,7 +199,7 @@ void FPoseLink::Evaluate(FPoseContext& Output)
 {
 #if DO_CHECK
 	checkf( !bProcessed, TEXT( "Evaluate already in progress, circular link for AnimInstance [%s] Blueprint [%s]" ), \
-		Output.AnimInstance ? *Output.AnimInstance->GetFullName() : TEXT( "None" ), Output.GetAnimBlueprintClass() ? *Output.GetAnimBlueprintClass()->GetFullName() : TEXT( "None" ) );
+		*Output.AnimInstanceProxy->GetAnimInstanceName(), Output.GetAnimBlueprintClass() ? *Output.GetAnimBlueprintClass()->GetFullName() : TEXT( "None" ) );
 	TGuardValue<bool> CircularGuard(bProcessed, true);
 #endif
 
@@ -147,11 +212,11 @@ void FPoseLink::Evaluate(FPoseContext& Output)
 #endif
 
 #if ENABLE_ANIMGRAPH_TRAVERSAL_DEBUG
-	checkf(InitializationCounter.IsSynchronizedWith(Output.AnimInstance->InitializationCounter), TEXT("Calling Evaluate without initialization!"));
-	checkf(CachedBonesCounter.IsSynchronizedWith(Output.AnimInstance->CachedBonesCounter), TEXT("Calling Evaluate without CachedBones!"));
-	checkf(UpdateCounter.IsSynchronizedWith(Output.AnimInstance->UpdateCounter), TEXT("Calling Evaluate without Update for this node!"));
-	checkf(!EvaluationCounter.IsSynchronizedWith(Output.AnimInstance->EvaluationCounter), TEXT("Already called Evaluate for this node!"));
-	EvaluationCounter.SynchronizeWith(Output.AnimInstance->EvaluationCounter);
+	checkf(InitializationCounter.IsSynchronizedWith(Output.AnimInstanceProxy->GetInitializationCounter()), TEXT("Calling Evaluate without initialization!"));
+	checkf(CachedBonesCounter.IsSynchronizedWith(Output.AnimInstanceProxy->GetCachedBonesCounter()), TEXT("Calling Evaluate without CachedBones!"));
+	checkf(UpdateCounter.IsSynchronizedWith(Output.AnimInstanceProxy->GetUpdateCounter()), TEXT("Calling Evaluate without Update for this node!"));
+	checkf(!EvaluationCounter.IsSynchronizedWith(Output.AnimInstanceProxy->GetEvaluationCounter()), TEXT("Already called Evaluate for this node!"));
+	EvaluationCounter.SynchronizeWith(Output.AnimInstanceProxy->GetEvaluationCounter());
 #endif
 
 	if (LinkedNode != NULL)
@@ -182,16 +247,16 @@ void FComponentSpacePoseLink::EvaluateComponentSpace(FComponentSpacePoseContext&
 {
 #if DO_CHECK
 	checkf( !bProcessed, TEXT( "EvaluateComponentSpace already in progress, circular link for AnimInstance [%s] Blueprint [%s]" ), \
-		Output.AnimInstance ? *Output.AnimInstance->GetFullName() : TEXT( "None" ), Output.GetAnimBlueprintClass() ? *Output.GetAnimBlueprintClass()->GetFullName() : TEXT( "None" ) );
+		*Output.AnimInstanceProxy->GetAnimInstanceName(), Output.GetAnimBlueprintClass() ? *Output.GetAnimBlueprintClass()->GetFullName() : TEXT( "None" ) );
 	TGuardValue<bool> CircularGuard(bProcessed, true);
 #endif
 
 #if ENABLE_ANIMGRAPH_TRAVERSAL_DEBUG
-	checkf(InitializationCounter.IsSynchronizedWith(Output.AnimInstance->InitializationCounter), TEXT("Calling EvaluateComponentSpace without initialization!"));
-	checkf(CachedBonesCounter.IsSynchronizedWith(Output.AnimInstance->CachedBonesCounter), TEXT("Calling EvaluateComponentSpace without CachedBones!"));
-	checkf(UpdateCounter.IsSynchronizedWith(Output.AnimInstance->UpdateCounter), TEXT("Calling EvaluateComponentSpace without Update for this node!"));
-	checkf(!EvaluationCounter.IsSynchronizedWith(Output.AnimInstance->EvaluationCounter), TEXT("Already called EvaluateComponentSpace for this node!"));
-	EvaluationCounter.SynchronizeWith(Output.AnimInstance->EvaluationCounter);
+	checkf(InitializationCounter.IsSynchronizedWith(Output.AnimInstanceProxy->GetInitializationCounter()), TEXT("Calling EvaluateComponentSpace without initialization!"));
+	checkf(CachedBonesCounter.IsSynchronizedWith(Output.AnimInstanceProxy->GetCachedBonesCounter()), TEXT("Calling EvaluateComponentSpace without CachedBones!"));
+	checkf(UpdateCounter.IsSynchronizedWith(Output.AnimInstanceProxy->GetUpdateCounter()), TEXT("Calling EvaluateComponentSpace without Update for this node!"));
+	checkf(!EvaluationCounter.IsSynchronizedWith(Output.AnimInstanceProxy->GetEvaluationCounter()), TEXT("Already called EvaluateComponentSpace for this node!"));
+	EvaluationCounter.SynchronizeWith(Output.AnimInstanceProxy->GetEvaluationCounter());
 #endif
 
 	if (LinkedNode != NULL)
@@ -258,4 +323,92 @@ void FNodeDebugData::GetFlattenedDebugData(TArray<FFlattenedDebugData>& Flattene
 			Child.GetFlattenedDebugData(FlattenedDebugData, ChildIndent, ChainID);
 		}
 	}
+}
+
+void FExposedValueHandler::Initialize(FAnimNode_Base* AnimNode, UObject* AnimInstanceObject) 
+{
+	if(bInitialized)
+	{
+		return;
+	}
+
+	if (BoundFunction != NAME_None)
+	{
+		Function = AnimInstanceObject->FindFunction(BoundFunction);
+	}
+	else
+	{
+		Function = NULL;
+	}
+
+	// initialize copy records
+	for(auto& CopyRecord : CopyRecords)
+	{
+		UProperty* SourceProperty = AnimInstanceObject->GetClass()->FindPropertyByName(CopyRecord.SourcePropertyName);
+		check(SourceProperty);
+		if(UArrayProperty* SourceArrayProperty = Cast<UArrayProperty>(SourceProperty))
+		{
+			FScriptArrayHelper ArrayHelper(SourceArrayProperty, SourceProperty->ContainerPtrToValuePtr<uint8>(AnimInstanceObject));
+			check(ArrayHelper.IsValidIndex(CopyRecord.SourceArrayIndex));
+			CopyRecord.Source = ArrayHelper.GetRawPtr(CopyRecord.SourceArrayIndex);
+		}
+		else
+		{
+			if(CopyRecord.SourceSubPropertyName != NAME_None)
+			{
+				void* Source = SourceProperty->ContainerPtrToValuePtr<uint8>(AnimInstanceObject, 0);
+				UStructProperty* SourceStructProperty = CastChecked<UStructProperty>(SourceProperty);
+				UProperty* SourceStructSubProperty = SourceStructProperty->Struct->FindPropertyByName(CopyRecord.SourceSubPropertyName);
+				CopyRecord.Source = SourceStructSubProperty->ContainerPtrToValuePtr<uint8>(Source, CopyRecord.SourceArrayIndex);
+			}
+			else
+			{
+				CopyRecord.Source = SourceProperty->ContainerPtrToValuePtr<uint8>(AnimInstanceObject, CopyRecord.SourceArrayIndex);
+			}
+		}
+
+		if(UArrayProperty* DestArrayProperty = Cast<UArrayProperty>(CopyRecord.DestProperty))
+		{
+			FScriptArrayHelper ArrayHelper(DestArrayProperty, CopyRecord.DestProperty->ContainerPtrToValuePtr<uint8>(AnimNode));
+			check(ArrayHelper.IsValidIndex(CopyRecord.DestArrayIndex));
+			CopyRecord.Dest = ArrayHelper.GetRawPtr(CopyRecord.DestArrayIndex);
+		}
+		else
+		{
+			CopyRecord.Dest = CopyRecord.DestProperty->ContainerPtrToValuePtr<uint8>(AnimNode, CopyRecord.DestArrayIndex);
+		}
+	}
+
+	bInitialized = true;
+}
+
+void FExposedValueHandler::Execute(const FAnimationBaseContext& Context) const
+{
+	if (Function != nullptr)
+	{
+		Context.AnimInstanceProxy->GetAnimInstanceObject()->ProcessEvent(Function, NULL);
+	}
+
+	for(const FExposedValueCopyRecord& CopyRecord : CopyRecords)
+	{
+		// if any of these checks fail then it's likely that Initialize has not been called.
+		// has new anim node type been added that doesnt call the base class Initialize()?
+		checkSlow(CopyRecord.Dest != nullptr);
+		checkSlow(CopyRecord.Source != nullptr);
+		checkSlow(CopyRecord.Size != 0);
+		
+		switch(CopyRecord.PostCopyOperation)
+		{
+		case EPostCopyOperation::None:
+			{
+				FMemory::Memcpy(CopyRecord.Dest, CopyRecord.Source, CopyRecord.Size);
+			}
+			break;
+		case EPostCopyOperation::LogicalNegateBool:
+			{
+				*(bool*)CopyRecord.Dest = !(*(bool*)CopyRecord.Source);
+			}
+			break;
+		}
+	}
 }
\ No newline at end of file
diff --git a/Engine/Source/Runtime/Engine/Private/Animation/AnimNodeSpaceConversions.cpp b/Engine/Source/Runtime/Engine/Private/Animation/AnimNodeSpaceConversions.cpp
--- a/Engine/Source/Runtime/Engine/Private/Animation/AnimNodeSpaceConversions.cpp
+++ b/Engine/Source/Runtime/Engine/Private/Animation/AnimNodeSpaceConversions.cpp
@@ -28,7 +28,7 @@ void FAnimNode_ConvertComponentToLocalSpace::Update(const FAnimationUpdateContex
 void FAnimNode_ConvertComponentToLocalSpace::Evaluate(FPoseContext & Output)
 {
 	// Evaluate the child and convert
-	FComponentSpacePoseContext InputCSPose(Output.AnimInstance);
+	FComponentSpacePoseContext InputCSPose(Output.AnimInstanceProxy);
 	ComponentPose.EvaluateComponentSpace(InputCSPose);
 
 	checkSlow( InputCSPose.Pose.GetPose().IsValid() );
@@ -75,7 +75,7 @@ void FAnimNode_ConvertLocalToComponentSpace::GatherDebugData(FNodeDebugData& Deb
 void FAnimNode_ConvertLocalToComponentSpace::EvaluateComponentSpace(FComponentSpacePoseContext & OutputCSPose)
 {
 	// Evaluate the child and convert
-	FPoseContext InputPose(OutputCSPose.AnimInstance);
+	FPoseContext InputPose(OutputCSPose.AnimInstanceProxy);
 	LocalPose.Evaluate(InputPose);
 
 	OutputCSPose.Pose.InitPose(InputPose.Pose);
diff --git a/Engine/Source/Runtime/Engine/Private/Animation/AnimNode_ApplyMeshSpaceAdditive.cpp b/Engine/Source/Runtime/Engine/Private/Animation/AnimNode_ApplyMeshSpaceAdditive.cpp
--- a/Engine/Source/Runtime/Engine/Private/Animation/AnimNode_ApplyMeshSpaceAdditive.cpp
+++ b/Engine/Source/Runtime/Engine/Private/Animation/AnimNode_ApplyMeshSpaceAdditive.cpp
@@ -24,7 +24,7 @@ void FAnimNode_ApplyMeshSpaceAdditive::Update(const FAnimationUpdateContext& Con
 {
 	Base.Update(Context);
 
-	if (IsLODEnabled(Context.AnimInstance, LODThreshold))
+	if (IsLODEnabled(Context.AnimInstanceProxy, LODThreshold))
 	{
 		// @note: If you derive this class, and if you have input that you rely on for base
 		// this is not going to work	
@@ -39,7 +39,7 @@ void FAnimNode_ApplyMeshSpaceAdditive::Update(const FAnimationUpdateContext& Con
 
 void FAnimNode_ApplyMeshSpaceAdditive::Evaluate(FPoseContext& Output)
 {
-	if (IsLODEnabled(Output.AnimInstance, LODThreshold))
+	if (IsLODEnabled(Output.AnimInstanceProxy, LODThreshold))
 	{
 		//@TODO: Could evaluate Base into Output and save a copy
 		const float ActualAlpha = AlphaScaleBias.ApplyTo(Alpha);
diff --git a/Engine/Source/Runtime/Engine/Private/Animation/AnimNode_SequencePlayer.cpp b/Engine/Source/Runtime/Engine/Private/Animation/AnimNode_SequencePlayer.cpp
--- a/Engine/Source/Runtime/Engine/Private/Animation/AnimNode_SequencePlayer.cpp
+++ b/Engine/Source/Runtime/Engine/Private/Animation/AnimNode_SequencePlayer.cpp
@@ -2,6 +2,7 @@
 
 #include "EnginePrivate.h"
 #include "Animation/AnimNode_SequencePlayer.h"
+#include "Animation/AnimInstanceProxy.h"
 
 /////////////////////////////////////////////////////
 // FAnimSequencePlayerNode
@@ -31,7 +32,7 @@ void FAnimNode_SequencePlayer::UpdateAssetPlayer(const FAnimationUpdateContext&
 {
 	EvaluateGraphExposedInputs.Execute(Context);
 
-	if ((Sequence != NULL) && (Context.AnimInstance->CurrentSkeleton->IsCompatible(Sequence->GetSkeleton())))
+	if ((Sequence != NULL) && (Context.AnimInstanceProxy->IsSkeletonCompatible(Sequence->GetSkeleton())))
 	{
 		if (InternalTimeAccumulator > Sequence->SequenceLength)
 		{
@@ -44,9 +45,9 @@ void FAnimNode_SequencePlayer::UpdateAssetPlayer(const FAnimationUpdateContext&
 
 void FAnimNode_SequencePlayer::Evaluate(FPoseContext& Output)
 {
-	if ((Sequence != NULL) && (Output.AnimInstance->CurrentSkeleton->IsCompatible(Sequence->GetSkeleton())))
+	if ((Sequence != NULL) && (Output.AnimInstanceProxy->IsSkeletonCompatible(Sequence->GetSkeleton())))
 	{
-		Sequence->GetAnimationPose(Output.Pose, Output.Curve, FAnimExtractContext(InternalTimeAccumulator, Output.AnimInstance->ShouldExtractRootMotion()));
+		Sequence->GetAnimationPose(Output.Pose, Output.Curve, FAnimExtractContext(InternalTimeAccumulator, Output.AnimInstanceProxy->ShouldExtractRootMotion()));
 	}
 	else
 	{
diff --git a/Engine/Source/Runtime/Engine/Private/Animation/AnimNode_StateMachine.cpp b/Engine/Source/Runtime/Engine/Private/Animation/AnimNode_StateMachine.cpp
--- a/Engine/Source/Runtime/Engine/Private/Animation/AnimNode_StateMachine.cpp
+++ b/Engine/Source/Runtime/Engine/Private/Animation/AnimNode_StateMachine.cpp
@@ -8,6 +8,7 @@
 #include "AnimationRuntime.h"
 #include "Animation/AnimStats.h"
 #include "Animation/BlendProfile.h"
+#include "Animation/AnimInstanceProxy.h"
 
 DECLARE_CYCLE_STAT(TEXT("StateMachine SetState"), Stat_StateMachineSetState, STATGROUP_Anim);
 
@@ -78,12 +79,12 @@ void FAnimationActiveTransitionEntry::InitializeCustomGraphLinks(const FAnimatio
 	{
 		const UAnimBlueprintGeneratedClass* AnimBlueprintClass = Context.GetAnimBlueprintClass();
 		CustomTransitionGraph.LinkID = AnimBlueprintClass->AnimNodeProperties.Num() - 1 - TransitionRule.CustomResultNodeIndex; //@TODO: Crazysauce
-		FAnimationInitializeContext InitContext(Context.AnimInstance);
+		FAnimationInitializeContext InitContext(Context.AnimInstanceProxy);
 		CustomTransitionGraph.Initialize(InitContext);
 
 		for (int32 Index = 0; Index < TransitionRule.PoseEvaluatorLinks.Num(); ++Index)
 		{
-			FAnimNode_TransitionPoseEvaluator* PoseEvaluator = GetNodeFromPropertyIndex<FAnimNode_TransitionPoseEvaluator>(Context.AnimInstance, AnimBlueprintClass, TransitionRule.PoseEvaluatorLinks[Index]);
+			FAnimNode_TransitionPoseEvaluator* PoseEvaluator = GetNodeFromPropertyIndex<FAnimNode_TransitionPoseEvaluator>(Context.AnimInstanceProxy->GetAnimInstanceObject(), AnimBlueprintClass, TransitionRule.PoseEvaluatorLinks[Index]);
 			PoseEvaluators.Add(PoseEvaluator);
 		}
 	}
@@ -192,16 +193,16 @@ void FAnimationPotentialTransition::Clear()
 // FAnimNode_StateMachine
 
 // Tries to get the instance information for the state machine
-FBakedAnimationStateMachine* FAnimNode_StateMachine::GetMachineDescription()
+const FBakedAnimationStateMachine* FAnimNode_StateMachine::GetMachineDescription() const
 {
-	if (PRIVATE_MachineDescription != NULL)
+	if (PRIVATE_MachineDescription != nullptr)
 	{
 		return PRIVATE_MachineDescription;
 	}
 	else
 	{
 		UE_LOG(LogAnimation, Warning, TEXT("FAnimNode_StateMachine: Bad machine ptr"));
-		return NULL;
+		return nullptr;
 	}
 }
 
@@ -209,10 +210,10 @@ void FAnimNode_StateMachine::Initialize(const FAnimationInitializeContext& Conte
 {
 	FAnimNode_Base::Initialize(Context);
 
-	UAnimBlueprintGeneratedClass* AnimBlueprintClass = Context.GetAnimBlueprintClass();
+	const UAnimBlueprintGeneratedClass* AnimBlueprintClass = Context.GetAnimBlueprintClass();
 	PRIVATE_MachineDescription = AnimBlueprintClass->BakedStateMachines.IsValidIndex(StateMachineIndexInClass) ? &(AnimBlueprintClass->BakedStateMachines[StateMachineIndexInClass]) : NULL;
 
-	if (FBakedAnimationStateMachine* Machine = GetMachineDescription())
+	if (const FBakedAnimationStateMachine* Machine = GetMachineDescription())
 	{
 		ElapsedTime = 0.0f;
 
@@ -237,7 +238,7 @@ void FAnimNode_StateMachine::Initialize(const FAnimationInitializeContext& Conte
 				// also initialize transitions
 				if(State.EntryRuleNodeIndex != INDEX_NONE)
 				{
-					if (FAnimNode_TransitionResult* TransitionNode = GetNodeFromPropertyIndex<FAnimNode_TransitionResult>(Context.AnimInstance, AnimBlueprintClass, State.EntryRuleNodeIndex))
+					if (FAnimNode_TransitionResult* TransitionNode = GetNodeFromPropertyIndex<FAnimNode_TransitionResult>(Context.AnimInstanceProxy->GetAnimInstanceObject(), AnimBlueprintClass, State.EntryRuleNodeIndex))
 					{
 						TransitionNode->Initialize(Context);
 					}
@@ -248,7 +249,7 @@ void FAnimNode_StateMachine::Initialize(const FAnimationInitializeContext& Conte
 					const FBakedStateExitTransition& TransitionRule = State.Transitions[TransitionIndex];
 					if (TransitionRule.CanTakeDelegateIndex != INDEX_NONE)
 					{
-						if (FAnimNode_TransitionResult* TransitionNode = GetNodeFromPropertyIndex<FAnimNode_TransitionResult>(Context.AnimInstance, AnimBlueprintClass, TransitionRule.CanTakeDelegateIndex))
+						if (FAnimNode_TransitionResult* TransitionNode = GetNodeFromPropertyIndex<FAnimNode_TransitionResult>(Context.AnimInstanceProxy->GetAnimInstanceObject(), AnimBlueprintClass, TransitionRule.CanTakeDelegateIndex))
 						{
 							TransitionNode->Initialize(Context);
 						}
@@ -271,7 +272,7 @@ void FAnimNode_StateMachine::Initialize(const FAnimationInitializeContext& Conte
 
 void FAnimNode_StateMachine::CacheBones(const FAnimationCacheBonesContext& Context) 
 {
-	if (FBakedAnimationStateMachine* Machine = GetMachineDescription())
+	if (const FBakedAnimationStateMachine* Machine = GetMachineDescription())
 	{
 		for (int32 StateIndex = 0; StateIndex < Machine->States.Num(); ++StateIndex)
 		{
@@ -316,7 +317,7 @@ const FAnimationTransitionBetweenStates& FAnimNode_StateMachine::GetTransitionIn
 
 void FAnimNode_StateMachine::Update(const FAnimationUpdateContext& Context)
 {
-	FBakedAnimationStateMachine* Machine = GetMachineDescription();
+	const FBakedAnimationStateMachine* Machine = GetMachineDescription();
 	if (Machine != nullptr)
 	{
 		if (Machine->States.Num() == 0)
@@ -327,7 +328,7 @@ void FAnimNode_StateMachine::Update(const FAnimationUpdateContext& Context)
 		{
 			// Attempting to catch a crash where the state machine has been freed.
 			UE_LOG(LogAnimation, Warning, TEXT("FAnimNode_StateMachine::Update - Invalid current state, please report. Attempting to use state %d in state machine %d"), CurrentState, StateMachineIndexInClass);
-			UE_LOG(LogAnimation, Warning, TEXT("\t\tWhen updating AnimInstance: %s"), *Context.AnimInstance->GetName())
+			UE_LOG(LogAnimation, Warning, TEXT("\t\tWhen updating AnimInstance: %s"), *Context.AnimInstanceProxy->GetAnimInstanceObject()->GetName())
 
 			return;
 		}
@@ -370,7 +371,7 @@ void FAnimNode_StateMachine::Update(const FAnimationUpdateContext& Context)
 			// let the latest transition know it has been interrupted
 			if ((ActiveTransitionArray.Num() > 0) && ActiveTransitionArray[ActiveTransitionArray.Num()-1].bActive)
 			{
-				Context.AnimInstance->AddAnimNotifyFromGeneratedClass(ActiveTransitionArray[ActiveTransitionArray.Num()-1].InterruptNotify);
+				Context.AnimInstanceProxy->AddAnimNotifyFromGeneratedClass(ActiveTransitionArray[ActiveTransitionArray.Num()-1].InterruptNotify);
 			}
 
 			const int32 PreviousState = CurrentState;
@@ -379,8 +380,8 @@ void FAnimNode_StateMachine::Update(const FAnimationUpdateContext& Context)
 			// Fire off Notifies for state transition
 			if (!bFirstUpdate)
 			{
-				Context.AnimInstance->AddAnimNotifyFromGeneratedClass(GetStateInfo(PreviousState).EndNotify);
-				Context.AnimInstance->AddAnimNotifyFromGeneratedClass(GetStateInfo(NextState).StartNotify);
+				Context.AnimInstanceProxy->AddAnimNotifyFromGeneratedClass(GetStateInfo(PreviousState).EndNotify);
+				Context.AnimInstanceProxy->AddAnimNotifyFromGeneratedClass(GetStateInfo(NextState).StartNotify);
 			}
 			
 			// Get the current weight of the next state, which may be non-zero
@@ -408,7 +409,7 @@ void FAnimNode_StateMachine::Update(const FAnimationUpdateContext& Context)
 
 			if(!bFirstUpdate)
 			{
-				Context.AnimInstance->AddAnimNotifyFromGeneratedClass(NewTransition->StartNotify);
+				Context.AnimInstanceProxy->AddAnimNotifyFromGeneratedClass(NewTransition->StartNotify);
 			}
 			
 			SetState(Context, NextState);
@@ -421,7 +422,7 @@ void FAnimNode_StateMachine::Update(const FAnimationUpdateContext& Context)
 	if (bFirstUpdate)
 	{
 		//Handle enter notify for "first" (after initial transitions) state
-		Context.AnimInstance->AddAnimNotifyFromGeneratedClass(GetStateInfo().StartNotify);
+		Context.AnimInstanceProxy->AddAnimNotifyFromGeneratedClass(GetStateInfo().StartNotify);
 		// in the first update, we don't like to transition from entry state
 		// so we throw out any transition data at the first update
 		ActiveTransitionArray.Reset();
@@ -446,8 +447,8 @@ void FAnimNode_StateMachine::Update(const FAnimationUpdateContext& Context)
 				// only play these events if it is the last transition (most recent, going to current state)
 				if (Index == (ActiveTransitionArray.Num() - 1))
 				{
-					Context.AnimInstance->AddAnimNotifyFromGeneratedClass(ActiveTransitionArray[Index].EndNotify);
-					Context.AnimInstance->AddAnimNotifyFromGeneratedClass(GetStateInfo().FullyBlendedNotify);
+					Context.AnimInstanceProxy->AddAnimNotifyFromGeneratedClass(ActiveTransitionArray[Index].EndNotify);
+					Context.AnimInstanceProxy->AddAnimNotifyFromGeneratedClass(GetStateInfo().FullyBlendedNotify);
 				}
 			}
 			else
@@ -485,7 +486,7 @@ FAnimNode_AssetPlayerBase* FAnimNode_StateMachine::GetRelevantAssetPlayerFromSta
 	float MaxWeight = 0.0f;
 	for (const int32& PlayerIdx : StateInfo.PlayerNodeIndices)
 	{
-		if (FAnimNode_AssetPlayerBase* Player = Context.AnimInstance->GetNodeFromIndex<FAnimNode_AssetPlayerBase>(PlayerIdx))
+		if (FAnimNode_AssetPlayerBase* Player = Context.AnimInstanceProxy->GetNodeFromIndex<FAnimNode_AssetPlayerBase>(PlayerIdx))
 		{
 			if (!Player->bIgnoreForRelevancyTest && (Player->GetCachedBlendWeight() > MaxWeight))
 			{
@@ -512,7 +513,7 @@ bool FAnimNode_StateMachine::FindValidTransition(const FAnimationUpdateContext&
 
 	// Conduit 'states' have an additional entry rule which must be true to consider taking any transitions via the conduit
 	//@TODO: It would add flexibility to be able to define this on normal state nodes as well, assuming the dual-graph editing is sorted out
-	if (FAnimNode_TransitionResult* StateEntryRuleNode = GetNodeFromPropertyIndex<FAnimNode_TransitionResult>(Context.AnimInstance, AnimBlueprintClass, StateInfo.EntryRuleNodeIndex))
+	if (FAnimNode_TransitionResult* StateEntryRuleNode = GetNodeFromPropertyIndex<FAnimNode_TransitionResult>(Context.AnimInstanceProxy->GetAnimInstanceObject(), AnimBlueprintClass, StateInfo.EntryRuleNodeIndex))
 	{
 		if (StateEntryRuleNode->NativeTransitionDelegate.IsBound())
 		{
@@ -541,7 +542,7 @@ bool FAnimNode_StateMachine::FindValidTransition(const FAnimationUpdateContext&
 			continue;
 		}
 
-		FAnimNode_TransitionResult* ResultNode = GetNodeFromPropertyIndex<FAnimNode_TransitionResult>(Context.AnimInstance, AnimBlueprintClass, TransitionRule.CanTakeDelegateIndex);
+		FAnimNode_TransitionResult* ResultNode = GetNodeFromPropertyIndex<FAnimNode_TransitionResult>(Context.AnimInstanceProxy->GetAnimInstanceObject(), AnimBlueprintClass, TransitionRule.CanTakeDelegateIndex);
 
 		if (ResultNode->NativeTransitionDelegate.IsBound())
 		{
@@ -648,7 +649,7 @@ void FAnimNode_StateMachine::UpdateTransitionStates(const FAnimationUpdateContex
 
 void FAnimNode_StateMachine::Evaluate(FPoseContext& Output)
 {
-	if (FBakedAnimationStateMachine* Machine = GetMachineDescription())
+	if (const FBakedAnimationStateMachine* Machine = GetMachineDescription())
 	{
 		if (Machine->States.Num() == 0 || !Machine->States.IsValidIndex(CurrentState))
 		{
@@ -662,11 +663,11 @@ void FAnimNode_StateMachine::Evaluate(FPoseContext& Output)
 		return;
 	}
 
-	ANIM_MT_SCOPE_CYCLE_COUNTER(EvaluateAnimStateMachine, Output.AnimInstance->IsRunningParallelEvaluation());
+	ANIM_MT_SCOPE_CYCLE_COUNTER(EvaluateAnimStateMachine, !IsInGameThread());
 	
 	if (ActiveTransitionArray.Num() > 0)
 	{
-		check(Output.AnimInstance->CurrentSkeleton);
+		check(Output.AnimInstanceProxy->GetSkeleton());
 		
 		check(StateCachedPoses.Num() == 0);
 		StateCachedPoses.AddZeroed(StatePoseLinks.Num());
@@ -743,7 +744,7 @@ void FAnimNode_StateMachine::EvaluateTransitionStandardBlendInternal(FPoseContex
 	{
 		for(FCompactPoseBoneIndex BoneIndex : Output.Pose.ForEachBoneIndex())
 		{
-			const int32 PerBoneIndex = Transition.BlendProfile->GetPerBoneInterpolationIndex(BoneIndex.GetInt(), Output.AnimInstance->RequiredBones);
+			const int32 PerBoneIndex = Transition.BlendProfile->GetPerBoneInterpolationIndex(BoneIndex.GetInt(), Output.AnimInstanceProxy->GetRequiredBones());
 
 			// Use defined per-bone scale if the bone has a scale specified in the blend profile
 			ScalarRegister FirstWeight = PerBoneIndex != INDEX_NONE ? ScalarRegister(Transition.StateBlendData[1].PerBoneBlendData[PerBoneIndex]) : ScalarRegister(VPreviousWeight);
@@ -847,7 +848,7 @@ void FAnimNode_StateMachine::SetState(const FAnimationBaseContext& Context, int3
 		// This stops any zero length blends holding on to old weights
 		for(const int32& PlayerIndex : GetStateInfo(CurrentState).PlayerNodeIndices)
 		{
-			if(FAnimNode_AssetPlayerBase* Player = Context.AnimInstance->GetNodeFromIndex<FAnimNode_AssetPlayerBase>(PlayerIndex))
+			if(FAnimNode_AssetPlayerBase* Player = Context.AnimInstanceProxy->GetNodeFromIndex<FAnimNode_AssetPlayerBase>(PlayerIndex))
 			{
 				Player->ClearCachedBlendWeight();
 			}
@@ -856,11 +857,11 @@ void FAnimNode_StateMachine::SetState(const FAnimationBaseContext& Context, int3
 		if (!bAlreadyActive && !IsAConduitState(NewStateIndex))
 		{
 			// Initialize the new state since it's not part of an active transition (and thus not still initialized)
-			FAnimationInitializeContext InitContext(Context.AnimInstance);
+			FAnimationInitializeContext InitContext(Context.AnimInstanceProxy);
 			StatePoseLinks[NewStateIndex].Initialize(InitContext);
 
 			// Also update BoneCaching.
-			FAnimationCacheBonesContext CacheBoneContext(Context.AnimInstance);
+			FAnimationCacheBonesContext CacheBoneContext(Context.AnimInstanceProxy);
 			StatePoseLinks[NewStateIndex].CacheBones(CacheBoneContext);
 		}
 
@@ -944,7 +945,7 @@ const FPoseContext& FAnimNode_StateMachine::EvaluateState(int32 StateIndex, cons
 	FPoseContext* CachePosePtr = StateCachedPoses[StateIndex];
 	if (CachePosePtr == nullptr)
 	{
-		CachePosePtr = new FPoseContext(Context.AnimInstance);
+		CachePosePtr = new FPoseContext(Context.AnimInstanceProxy);
 		check(CachePosePtr);
 		StateCachedPoses[StateIndex] = CachePosePtr;
 
diff --git a/Engine/Source/Runtime/Engine/Private/Animation/AnimNode_TransitionPoseEvaluator.cpp b/Engine/Source/Runtime/Engine/Private/Animation/AnimNode_TransitionPoseEvaluator.cpp
--- a/Engine/Source/Runtime/Engine/Private/Animation/AnimNode_TransitionPoseEvaluator.cpp
+++ b/Engine/Source/Runtime/Engine/Private/Animation/AnimNode_TransitionPoseEvaluator.cpp
@@ -2,6 +2,7 @@
 
 #include "EnginePrivate.h"
 #include "Animation/AnimNode_TransitionPoseEvaluator.h"
+#include "Animation/AnimInstanceProxy.h"
 
 /////////////////////////////////////////////////////
 // FAnimNode_TransitionPoseEvaluator
@@ -32,8 +33,8 @@ void FAnimNode_TransitionPoseEvaluator::Initialize(const FAnimationInitializeCon
 
 void FAnimNode_TransitionPoseEvaluator::CacheBones(const FAnimationCacheBonesContext& Context) 
 {
-	CachedPose.SetBoneContainer(&Context.AnimInstance->RequiredBones);
-	CachedCurve.InitFrom(Context.AnimInstance);
+	CachedPose.SetBoneContainer(&Context.AnimInstanceProxy->GetRequiredBones());
+	CachedCurve.InitFrom(Context.AnimInstanceProxy->GetSkeleton());
 }
 
 void FAnimNode_TransitionPoseEvaluator::Update(const FAnimationUpdateContext& Context)
diff --git a/Engine/Source/Runtime/Engine/Private/Animation/AnimNotifyQueue.cpp b/Engine/Source/Runtime/Engine/Private/Animation/AnimNotifyQueue.cpp
new file mode 100644
--- /dev/null
+++ b/Engine/Source/Runtime/Engine/Private/Animation/AnimNotifyQueue.cpp
@@ -0,0 +1,59 @@
+
+#include "EnginePrivate.h"
+#include "Animation/AnimNotifyQueue.h"
+
+bool FAnimNotifyQueue::PassesFiltering(const FAnimNotifyEvent* Notify) const
+{
+	switch (Notify->NotifyFilterType)
+	{
+		case ENotifyFilterType::NoFiltering:
+		{
+			return true;
+		}
+		case ENotifyFilterType::LOD:
+		{
+			return Notify->NotifyFilterLOD > PredictedLODLevel;
+		}
+		default:
+		{
+			ensure(false); // Unknown Filter Type
+		}
+	}
+	return true;
+}
+
+bool FAnimNotifyQueue::PassesChanceOfTriggering(const FAnimNotifyEvent* Event) const
+{
+	return Event->NotifyStateClass ? true : RandomStream.FRandRange(0.f, 1.f) < Event->NotifyTriggerChance;
+}
+
+void FAnimNotifyQueue::AddAnimNotifies(const TArray<const FAnimNotifyEvent*>& NewNotifies, const float InstanceWeight)
+{
+	// for now there is no filter whatsoever, it just adds everything requested
+	for (const FAnimNotifyEvent* Notify : NewNotifies)
+	{
+		// only add if it is over TriggerWeightThreshold
+		const bool bPassesDedicatedServerCheck = Notify->bTriggerOnDedicatedServer || !IsRunningDedicatedServer();
+		if (bPassesDedicatedServerCheck && Notify->TriggerWeightThreshold <= InstanceWeight && PassesFiltering(Notify) && PassesChanceOfTriggering(Notify))
+		{
+			// Only add unique AnimNotifyState instances just once. We can get multiple triggers if looping over an animation.
+			// It is the same state, so just report it once.
+			Notify->NotifyStateClass ? AnimNotifies.AddUnique(Notify) : AnimNotifies.Add(Notify);
+		}
+	}
+}
+
+void FAnimNotifyQueue::Reset(USkeletalMeshComponent* Component)
+{
+	AnimNotifies.Reset();
+	PredictedLODLevel = Component ? Component->PredictedLODLevel : -1;
+}
+
+void FAnimNotifyQueue::Append(const FAnimNotifyQueue& Queue)
+{
+	// we dont just append here - we need to preserve uniqueness for AnimNotifyState instances
+	for (const FAnimNotifyEvent* Notify : Queue.AnimNotifies)
+	{
+		Notify->NotifyStateClass ? AnimNotifies.AddUnique(Notify) : AnimNotifies.Add(Notify);
+	}
+}
\ No newline at end of file
diff --git a/Engine/Source/Runtime/Engine/Private/Animation/AnimPhysicsSolver.cpp b/Engine/Source/Runtime/Engine/Private/Animation/AnimPhysicsSolver.cpp
--- a/Engine/Source/Runtime/Engine/Private/Animation/AnimPhysicsSolver.cpp
+++ b/Engine/Source/Runtime/Engine/Private/Animation/AnimPhysicsSolver.cpp
@@ -208,7 +208,7 @@ FAnimPhysAngularLimit::FAnimPhysAngularLimit(FAnimPhysRigidBody* InFirstBody, FA
 , MinimumTorque(InMinimumTorque)
 , MaximumTorque(InMaximumTorque)
 {
-
+	UpdateCachedData();
 }
 
 void FAnimPhysAngularLimit::RemoveBias()
@@ -229,9 +229,9 @@ void FAnimPhysAngularLimit::Iter(float DeltaTime)
 
 	// the amount of spin we have to add to satisfy the limit.
 	float DeltaSpin = TargetSpin - CurrentSpin;
-	float SpinToTorque = 1.0f / (((Bodies[0]) ? FVector::DotProduct(WorldSpaceAxis, Bodies[0]->InverseWorldSpaceTensor.TransformVector(WorldSpaceAxis)) : 0.0f) + ((Bodies[1]) ? FVector::DotProduct(WorldSpaceAxis, Bodies[1]->InverseWorldSpaceTensor.TransformVector(WorldSpaceAxis)) : 0.0f));
+
 	// how we have to change the angular impulse
-	float DeltaTorque = DeltaSpin * SpinToTorque;
+	float DeltaTorque = DeltaSpin * CachedSpinToTorque;
 
 	// total torque cannot exceed maxtorque
 	DeltaTorque = FMath::Min(DeltaTorque, MaximumTorque * DeltaTime - Torque);
@@ -251,6 +251,11 @@ void FAnimPhysAngularLimit::Iter(float DeltaTime)
 	Torque += DeltaTorque;
 }
 
+void FAnimPhysAngularLimit::UpdateCachedData()
+{
+	CachedSpinToTorque = 1.0f / (((Bodies[0]) ? FVector::DotProduct(WorldSpaceAxis, Bodies[0]->InverseWorldSpaceTensor.TransformVector(WorldSpaceAxis)) : 0.0f) + ((Bodies[1]) ? FVector::DotProduct(WorldSpaceAxis, Bodies[1]->InverseWorldSpaceTensor.TransformVector(WorldSpaceAxis)) : 0.0f));
+}
+
 //////////////////////////////////////////////////////////////////////////
 
 FAnimPhysLinearLimit::FAnimPhysLinearLimit() 
@@ -270,7 +275,7 @@ FAnimPhysLinearLimit::FAnimPhysLinearLimit(FAnimPhysRigidBody* InFirstBody, FAni
 , Maximumforce(FMath::Max(InForceRange.X, InForceRange.Y))
 , SumImpulses(0.0f)
 {
-
+	UpdateCachedData();
 }
 
 void FAnimPhysLinearLimit::RemoveBias()
@@ -280,39 +285,39 @@ void FAnimPhysLinearLimit::RemoveBias()
 
 void FAnimPhysLinearLimit::Iter(float DeltaTime)
 {
-	// Local body positions rotated into world space
-	FVector Position0 = (Bodies[0]) ? Bodies[0]->Pose.Orientation * FirstPosition : FirstPosition;
-	FVector Position1 = (Bodies[1]) ? Bodies[1]->Pose.Orientation * SecondPosition : SecondPosition;
-
 	// Instantaneous linear velocity at point of constraint
-	FVector Velocity0 = Bodies[0] ? FVector::CrossProduct(Bodies[0]->Spin(), Position0) + Bodies[0]->LinearMomentum * Bodies[0]->InverseMass : FVector(0.0f, 0.0f, 0.0f);
-	FVector Velocity1 = Bodies[1] ? FVector::CrossProduct(Bodies[1]->Spin(), Position1) + Bodies[1]->LinearMomentum * Bodies[1]->InverseMass : FVector(0.0f, 0.0f, 0.0f);
+	FVector Velocity0 = Bodies[0] ? FVector::CrossProduct(Bodies[0]->Spin(), WorldSpacePosition0) + Bodies[0]->LinearMomentum * Bodies[0]->InverseMass : FVector(0.0f, 0.0f, 0.0f);
+	FVector Velocity1 = Bodies[1] ? FVector::CrossProduct(Bodies[1]->Spin(), WorldSpacePosition1) + Bodies[1]->LinearMomentum * Bodies[1]->InverseMass : FVector(0.0f, 0.0f, 0.0f);
 
 	// velocity of rb1 wrt rb0
 	float  DeltaVelocity = FVector::DotProduct(Velocity1 - Velocity0, LimitNormal);
 
 	float VelocityImpulse = -TargetSpeed - DeltaVelocity;
 
-	float InertiaImpulse = (Bodies[0] ? Bodies[0]->InverseMass + FVector::DotProduct(FVector::CrossProduct(Bodies[0]->InverseWorldSpaceTensor.TransformVector(FVector::CrossProduct(Position0, LimitNormal)), Position0), LimitNormal) : 0)
-		+ (Bodies[1] ? Bodies[1]->InverseMass + FVector::DotProduct(FVector::CrossProduct(Bodies[1]->InverseWorldSpaceTensor.TransformVector(FVector::CrossProduct(Position1, LimitNormal)), Position1), LimitNormal) : 0);
-
-	float ResultantImpulse = VelocityImpulse / InertiaImpulse;
+	float ResultantImpulse = VelocityImpulse * InverseInertiaImpulse;
 
 	ResultantImpulse = FMath::Min(Maximumforce * DeltaTime - SumImpulses, ResultantImpulse);
 	ResultantImpulse = FMath::Max(MinimumForce * DeltaTime - SumImpulses, ResultantImpulse);
 
 	if (Bodies[0])
 	{
-		FAnimPhys::ApplyImpulse(Bodies[0], Position0, LimitNormal * -ResultantImpulse);
+		FAnimPhys::ApplyImpulse(Bodies[0], WorldSpacePosition0, LimitNormal * -ResultantImpulse);
 	}
 
 	if (Bodies[1])
 	{
-		FAnimPhys::ApplyImpulse(Bodies[1], Position1, LimitNormal * ResultantImpulse);
+		FAnimPhys::ApplyImpulse(Bodies[1], WorldSpacePosition1, LimitNormal * ResultantImpulse);
 	}
 
 	SumImpulses += ResultantImpulse;
 }
+void FAnimPhysLinearLimit::UpdateCachedData()
+{
+	WorldSpacePosition0 = (Bodies[0]) ? Bodies[0]->Pose.Orientation * FirstPosition : FirstPosition;
+	WorldSpacePosition1 = (Bodies[1]) ? Bodies[1]->Pose.Orientation * SecondPosition : SecondPosition;
+	InverseInertiaImpulse = 1.0f / ((Bodies[0] ? Bodies[0]->InverseMass + FVector::DotProduct(FVector::CrossProduct(Bodies[0]->InverseWorldSpaceTensor.TransformVector(FVector::CrossProduct(WorldSpacePosition0, LimitNormal)), WorldSpacePosition0), LimitNormal) : 0)
+		+ (Bodies[1] ? Bodies[1]->InverseMass + FVector::DotProduct(FVector::CrossProduct(Bodies[1]->InverseWorldSpaceTensor.TransformVector(FVector::CrossProduct(WorldSpacePosition1, LimitNormal)), WorldSpacePosition1), LimitNormal) : 0));
+}
 
 //////////////////////////////////////////////////////////////////////////
 
diff --git a/Engine/Source/Runtime/Engine/Private/Animation/AnimSequence.cpp b/Engine/Source/Runtime/Engine/Private/Animation/AnimSequence.cpp
--- a/Engine/Source/Runtime/Engine/Private/Animation/AnimSequence.cpp
+++ b/Engine/Source/Runtime/Engine/Private/Animation/AnimSequence.cpp
@@ -811,9 +811,9 @@ void UAnimSequence::ExtractBoneTransform(const TArray<struct FRawAnimSequenceTra
 	OutAtom.NormalizeRotation();
 }
 
-void UAnimSequence::OnAssetPlayerTickedInternal(FAnimAssetTickContext &Context, const float PreviousTime, const float MoveDelta, const FAnimTickRecord &Instance, class UAnimInstance* InstanceOwner) const
+void UAnimSequence::HandleAssetPlayerTickedInternal(FAnimAssetTickContext &Context, const float PreviousTime, const float MoveDelta, const FAnimTickRecord &Instance, struct FAnimNotifyQueue& NotifyQueue) const
 {
-	Super::OnAssetPlayerTickedInternal(Context, PreviousTime, MoveDelta, Instance, InstanceOwner);
+	Super::HandleAssetPlayerTickedInternal(Context, PreviousTime, MoveDelta, Instance, NotifyQueue);
 
 	if (bEnableRootMotion)
 	{
