Copying //UE4/Dev-Framework to Dev-Main (//UE4/Dev-Main)

diff --git a/Engine/Source/Runtime/Engine/Private/Animation/AnimSequenceBase.cpp b/Engine/Source/Runtime/Engine/Private/Animation/AnimSequenceBase.cpp
--- a/Engine/Source/Runtime/Engine/Private/Animation/AnimSequenceBase.cpp
+++ b/Engine/Source/Runtime/Engine/Private/Animation/AnimSequenceBase.cpp
@@ -219,7 +219,7 @@ void UAnimSequenceBase::GetAnimNotifiesFromDeltaPositions(const float& PreviousP
 	}
 }
 
-void UAnimSequenceBase::TickAssetPlayerInstance(FAnimTickRecord& Instance, class UAnimInstance* InstanceOwner, FAnimAssetTickContext& Context) const
+void UAnimSequenceBase::TickAssetPlayer(FAnimTickRecord& Instance, struct FAnimNotifyQueue& NotifyQueue, FAnimAssetTickContext& Context) const
 {
 	float& CurrentTime = *(Instance.TimeAccumulator);
 	float PreviousTime = CurrentTime;
@@ -277,7 +277,7 @@ void UAnimSequenceBase::TickAssetPlayerInstance(FAnimTickRecord& Instance, class
 		}
 	}
 
-	OnAssetPlayerTickedInternal(Context, PreviousTime, MoveDelta, Instance, InstanceOwner);
+	HandleAssetPlayerTickedInternal(Context, PreviousTime, MoveDelta, Instance, NotifyQueue);
 }
 
 void UAnimSequenceBase::TickByMarkerAsFollower(FMarkerTickRecord &Instance, FMarkerTickContext &MarkerContext, float& CurrentTime, float& OutPreviousTime, const float MoveDelta, const bool bLooping) const
@@ -520,13 +520,20 @@ void UAnimSequenceBase::Serialize(FArchive& Ar)
 	RawCurveData.Serialize(Ar);
 }
 
-void UAnimSequenceBase::OnAssetPlayerTickedInternal(FAnimAssetTickContext &Context, const float PreviousTime, const float MoveDelta, const FAnimTickRecord &Instance, class UAnimInstance* InstanceOwner) const
+void UAnimSequenceBase::OnAssetPlayerTickedInternal(FAnimAssetTickContext &Context, const float PreviousTime, const float MoveDelta, const FAnimTickRecord &Instance, class UAnimInstance* InAnimInstance) const
+{
+	// @todo: remove after deprecation
+	// forward to non-deprecated version
+	HandleAssetPlayerTickedInternal(Context, PreviousTime, MoveDelta, Instance, InAnimInstance->NotifyQueue);
+}
+
+void UAnimSequenceBase::HandleAssetPlayerTickedInternal(FAnimAssetTickContext &Context, const float PreviousTime, const float MoveDelta, const FAnimTickRecord &Instance, struct FAnimNotifyQueue& NotifyQueue) const
 {
 	if (Context.ShouldGenerateNotifies())
 	{
 		// Harvest and record notifies
 		TArray<const FAnimNotifyEvent*> AnimNotifies;
 		GetAnimNotifies(PreviousTime, MoveDelta, Instance.bLooping, AnimNotifies);
-		InstanceOwner->AddAnimNotifies(AnimNotifies, Instance.EffectiveBlendWeight);
+		NotifyQueue.AddAnimNotifies(AnimNotifies, Instance.EffectiveBlendWeight);
 	}
 }
diff --git a/Engine/Source/Runtime/Engine/Private/Animation/AnimSingleNodeInstance.cpp b/Engine/Source/Runtime/Engine/Private/Animation/AnimSingleNodeInstance.cpp
--- a/Engine/Source/Runtime/Engine/Private/Animation/AnimSingleNodeInstance.cpp
+++ b/Engine/Source/Runtime/Engine/Private/Animation/AnimSingleNodeInstance.cpp
@@ -16,6 +16,7 @@
 #include "Animation/AimOffsetBlendSpace.h"
 #include "Animation/AimOffsetBlendSpace1D.h"
 #include "Animation/AnimMontage.h"
+#include "Animation/AnimSingleNodeInstanceProxy.h"
 
 /////////////////////////////////////////////////////
 // UAnimSingleNodeInstance
@@ -23,65 +24,31 @@
 
 UAnimSingleNodeInstance::UAnimSingleNodeInstance(const FObjectInitializer& ObjectInitializer)
 	: Super(ObjectInitializer)
-	, PlayRate(1.f)
-	, bLooping(true)
-	, bPlaying(true)
-	, bReverse(false)
 {
 }
 
-void UAnimSingleNodeInstance::SetAnimationAsset(class UAnimationAsset* NewAsset,bool bIsLooping,float InPlayRate)
+void UAnimSingleNodeInstance::SetAnimationAsset(class UAnimationAsset* NewAsset, bool bInIsLooping, float InPlayRate)
 {
-	if (NewAsset != CurrentAsset)
-	{
-		CurrentAsset = NewAsset;
-	}
-
-	if (USkeletalMeshComponent * MeshComponent = GetSkelMeshComponent())
-	{
-		if (MeshComponent->SkeletalMesh == NULL)
-		{
-			// if it does not have SkeletalMesh, we nullify it
-			CurrentAsset = NULL;
-		}
-		else if (CurrentAsset != NULL)
-		{
-			// if we have an asset, make sure their skeleton matches, otherwise, null it
-			if (CurrentSkeleton != CurrentAsset->GetSkeleton())
-			{
-				// clear asset since we do not have matching skeleton
-				CurrentAsset = NULL;
-			}
-		}
-	}
+	FAnimSingleNodeInstanceProxy& Proxy = GetProxyOnGameThread<FAnimSingleNodeInstanceProxy>();
+	Proxy.SetAnimationAsset(NewAsset, GetSkelMeshComponent(), bInIsLooping, InPlayRate);
 
-	bLooping = bIsLooping;
-	PlayRate = InPlayRate;
-	CurrentTime = 0.f;
-	BlendSpaceInput = FVector::ZeroVector;
-	MarkerTickRecord.Reset();
-
-#if WITH_EDITORONLY_DATA
-	PreviewPoseCurrentTime = 0.0f;
-#endif
-
-	UAnimMontage* Montage = Cast<UAnimMontage>(CurrentAsset);
+	UAnimMontage* Montage = Cast<UAnimMontage>(NewAsset);
 	if ( Montage!=NULL )
 	{
-		ReinitializeSlotNodes();
+		Proxy.ReinitializeSlotNodes();
 		if ( Montage->SlotAnimTracks.Num() > 0 )
 		{
-			RegisterSlotNodeWithAnimInstance(Montage->SlotAnimTracks[0].SlotName);
+			Proxy.RegisterSlotNodeWithAnimInstance(Montage->SlotAnimTracks[0].SlotName);
 		}
 		RestartMontage( Montage );
-		SetPlaying(bPlaying);
+		SetPlaying(IsPlaying());
 	}
 	else
 	{
 		// otherwise stop all montages
 		StopAllMontages(0.25f);
 
-		UBlendSpaceBase * BlendSpace = Cast<UBlendSpaceBase>(CurrentAsset);
+		UBlendSpaceBase * BlendSpace = Cast<UBlendSpaceBase>(NewAsset);
 		if(BlendSpace)
 		{
 			BlendSpace->InitializeFilter(&BlendFilter);
@@ -91,31 +58,7 @@ void UAnimSingleNodeInstance::SetAnimationAsset(class UAnimationAsset* NewAsset,
 
 void UAnimSingleNodeInstance::SetVertexAnimation(UVertexAnimation * NewVertexAnim, bool bIsLooping, float InPlayRate)
 {
-	if (NewVertexAnim != CurrentVertexAnim)
-	{
-		CurrentVertexAnim = NewVertexAnim;
-	}
-
-	if (USkeletalMeshComponent * MeshComponent = GetSkelMeshComponent())
-	{
-		if (MeshComponent->SkeletalMesh == NULL)
-		{
-			// if it does not have SkeletalMesh, we nullify it
-			CurrentVertexAnim = NULL;
-		}
-		else if (CurrentVertexAnim != NULL)
-		{
-			// if we have an anim, make sure their mesh matches, otherwise, null it
-			if (MeshComponent->SkeletalMesh != CurrentVertexAnim->BaseSkelMesh)
-			{
-				// clear asset since we do not have matching skeleton
-				CurrentVertexAnim = NULL;
-			}
-		}
-	}
-
-	bLooping = bIsLooping;
-	PlayRate = InPlayRate;
+	GetProxyOnGameThread<FAnimSingleNodeInstanceProxy>().SetVertexAnimation(NewVertexAnim, bIsLooping, InPlayRate);
 
 	// reinitialize
 	InitializeAnimation();
@@ -181,243 +124,31 @@ void UAnimSingleNodeInstance::UpdateMontageWeightForTimeSkip(float TimeDifferenc
 
 void UAnimSingleNodeInstance::UpdateBlendspaceSamples(FVector InBlendInput)
 {
-	if (UBlendSpaceBase* BlendSpace = Cast<UBlendSpaceBase>(CurrentAsset))
-	{
-		float OutCurrentTime = 0.f;
-		FMarkerTickRecord TempMarkerTickRecord;
-		BlendSpaceAdvanceImmediate(BlendSpace, InBlendInput, BlendSampleData, BlendFilter, false, 1.f, 0.f, OutCurrentTime, TempMarkerTickRecord);
-	}
+	GetProxyOnGameThread<FAnimSingleNodeInstanceProxy>().UpdateBlendspaceSamples(InBlendInput);
 }
 
-void UAnimSingleNodeInstance::RestartMontage(UAnimMontage * Montage, FName FromSection)
+void UAnimSingleNodeInstance::RestartMontage(UAnimMontage* Montage, FName FromSection)
 {
-	if( Montage == CurrentAsset )
+	FAnimSingleNodeInstanceProxy& Proxy = GetProxyOnGameThread<FAnimSingleNodeInstanceProxy>();
+
+	if( Montage == Proxy.GetCurrentAsset() )
 	{
-		WeightInfo.Reset();
-		Montage_Play(Montage, PlayRate);
+		Proxy.ResetWeightInfo();
+		Montage_Play(Montage, Proxy.GetPlayRate());
 		if( FromSection != NAME_None )
 		{
 			Montage_JumpToSection(FromSection);
 		}
-		SetMontageLoop(Montage, bLooping, FromSection);
+		SetMontageLoop(Montage, Proxy.IsLooping(), FromSection);
 	}
 }
 
 void UAnimSingleNodeInstance::NativeInitializeAnimation()
 {
-	CurrentAsset = NULL;
-	CurrentVertexAnim = NULL;
-#if WITH_EDITORONLY_DATA
-	PreviewPoseCurrentTime = 0.0f;
-#endif
-
-	// it's already doing it when evaluate
-	BlendSpaceInput = FVector::ZeroVector;
-	CurrentTime = 0.f;
 	USkeletalMeshComponent* SkelComp = GetSkelMeshComponent();
 	SkelComp->AnimationData.Initialize(this);
 }
 
-void UAnimSingleNodeInstance::UpdateAnimationNode(float DeltaTimeX)
-{
-	float NewPlayRate = PlayRate;
-	UAnimSequence* PreviewBasePose = NULL;
-
-	if (bPlaying == false)
-	{
-		// we still have to tick animation when bPlaying is false because 
-		NewPlayRate = 0.f;
-	}
-
-	if(CurrentAsset != NULL)
-	{
-		FAnimGroupInstance* SyncGroup;
-		if (UBlendSpaceBase* BlendSpace = Cast<UBlendSpaceBase>(CurrentAsset))
-		{
-			FAnimTickRecord& TickRecord = CreateUninitializedTickRecord(INDEX_NONE, /*out*/ SyncGroup);
-			MakeBlendSpaceTickRecord(TickRecord, BlendSpace, BlendSpaceInput, BlendSampleData, BlendFilter, bLooping, NewPlayRate, 1.f, /*inout*/ CurrentTime, MarkerTickRecord);
-#if WITH_EDITORONLY_DATA
-			PreviewBasePose = BlendSpace->PreviewBasePose;
-#endif
-		}
-		else if (UAnimSequence* Sequence = Cast<UAnimSequence>(CurrentAsset))
-		{
-			FAnimTickRecord& TickRecord = CreateUninitializedTickRecord(INDEX_NONE, /*out*/ SyncGroup);
-			MakeSequenceTickRecord(TickRecord, Sequence, bLooping, NewPlayRate, 1.f, /*inout*/ CurrentTime, MarkerTickRecord);
-			// if it's not looping, just set play to be false when reached to end
-			if (!bLooping)
-			{
-				const float CombinedPlayRate = NewPlayRate*Sequence->RateScale;
-				if ((CombinedPlayRate < 0.f && CurrentTime <= 0.f) || (CombinedPlayRate > 0.f && CurrentTime >= Sequence->SequenceLength))
-				{
-					SetPlaying(false);
-				}
-			}
-		}
-		else if(UAnimComposite* Composite = Cast<UAnimComposite>(CurrentAsset))
-		{
-			FAnimTickRecord& TickRecord = CreateUninitializedTickRecord(INDEX_NONE, /*out*/ SyncGroup);
-			MakeSequenceTickRecord(TickRecord, Composite, bLooping, NewPlayRate, 1.f, /*inout*/ CurrentTime, MarkerTickRecord);
-			// if it's not looping, just set play to be false when reached to end
-			if (!bLooping)
-			{
-				const float CombinedPlayRate = NewPlayRate*Composite->RateScale;
-				if ((CombinedPlayRate < 0.f && CurrentTime <= 0.f) || (CombinedPlayRate > 0.f && CurrentTime >= Composite->SequenceLength))
-				{
-					SetPlaying(false);
-				}
-			}
-		}
-		else if (UAnimMontage * Montage = Cast<UAnimMontage>(CurrentAsset))
-		{
-			// Full weight , if you don't have slot track, you won't be able to see animation playing
-			if ( Montage->SlotAnimTracks.Num() > 0 )
-			{
-				// in the future, maybe we can support which slot
-				const FName CurrentSlotNodeName = Montage->SlotAnimTracks[0].SlotName;
-				GetSlotWeight(CurrentSlotNodeName, WeightInfo.SlotNodeWeight, WeightInfo.SourceWeight, WeightInfo.TotalNodeWeight);
-				UpdateSlotNodeWeight(CurrentSlotNodeName, WeightInfo.SlotNodeWeight);		
-			}
-			// get the montage position
-			// @todo anim: temporarily just choose first slot and show the location
-			FAnimMontageInstance * ActiveMontageInstance = GetActiveMontageInstance();
-			if (!ActiveMontageInstance && bPlaying)
-			{
-				SetPlaying(false);
-			}
-#if WITH_EDITORONLY_DATA
-			PreviewBasePose = Montage->PreviewBasePose;
-#endif
-		}
-	}
-	else if(CurrentVertexAnim != NULL)
-	{
-		float MoveDelta = DeltaTimeX * NewPlayRate;
-		FAnimationRuntime::AdvanceTime(bLooping, MoveDelta, CurrentTime, CurrentVertexAnim->GetAnimLength());
-	}
-
-#if WITH_EDITORONLY_DATA
-	if(PreviewBasePose)
-	{
-		float MoveDelta = DeltaTimeX * NewPlayRate;
-		const bool bIsPreviewPoseLooping = true;
-
-		FAnimationRuntime::AdvanceTime(bIsPreviewPoseLooping, MoveDelta, PreviewPoseCurrentTime, PreviewBasePose->SequenceLength);
-	}
-#endif
-}
-
-bool UAnimSingleNodeInstance::NativeEvaluateAnimation(FPoseContext& Output)
-{
-	if (CurrentAsset != NULL)
-	{
-		//@TODO: animrefactor: Seems like more code duplication than we need
-		if (UBlendSpaceBase* BlendSpace = Cast<UBlendSpaceBase>(CurrentAsset))
-		{
-			InternalBlendSpaceEvaluatePose(BlendSpace, BlendSampleData, Output);
-		}
-		else if (UAnimSequence* Sequence = Cast<UAnimSequence>(CurrentAsset))
-		{
-			if (Sequence->IsValidAdditive())
-			{
-				FAnimExtractContext ExtractionContext(CurrentTime, Sequence->bEnableRootMotion);
-
-				Sequence->GetAdditiveBasePose(Output.Pose, Output.Curve, ExtractionContext);
-
-				FCompactPose AdditivePose;
-				FBlendedCurve AdditiveCurve;
-				AdditivePose.SetBoneContainer(&Output.Pose.GetBoneContainer());
-				AdditiveCurve.InitFrom(Output.Curve);
-				Sequence->GetAnimationPose(AdditivePose, AdditiveCurve, ExtractionContext);
-
-				FAnimationRuntime::AccumulateAdditivePose(Output.Pose, AdditivePose, Output.Curve, AdditiveCurve, 1.f, Sequence->AdditiveAnimType);
-				Output.Pose.NormalizeRotations();
-			}
-			else
-			{
-				// if SkeletalMesh isn't there, we'll need to use skeleton
-				Sequence->GetAnimationPose(Output.Pose, Output.Curve, FAnimExtractContext(CurrentTime, Sequence->bEnableRootMotion));
-			}
-		}
-		else if (UAnimComposite* Composite = Cast<UAnimComposite>(CurrentAsset))
-		{
-			FAnimExtractContext ExtractionContext(CurrentTime, ShouldExtractRootMotion());
-			const FAnimTrack& AnimTrack = Composite->AnimationTrack;
-
-			// find out if this is additive animation
-			if (AnimTrack.IsAdditive())
-			{
-#if WITH_EDITORONLY_DATA
-				if (Composite->PreviewBasePose)
-				{
-					Composite->PreviewBasePose->GetAdditiveBasePose(Output.Pose, Output.Curve, ExtractionContext);
-				}
-				else
-#endif
-				{
-					// get base pose - for now we only support ref pose as base
-					Output.Pose.ResetToRefPose();
-				}
-
-				EAdditiveAnimationType AdditiveAnimType = AnimTrack.IsRotationOffsetAdditive()? AAT_RotationOffsetMeshSpace : AAT_LocalSpaceBase;
-
-				FCompactPose AdditivePose;
-				FBlendedCurve AdditiveCurve;
-				AdditivePose.SetBoneContainer(&Output.Pose.GetBoneContainer());
-				AdditiveCurve.InitFrom(Output.Curve);
-				Composite->GetAnimationPose(AdditivePose, AdditiveCurve, ExtractionContext);
-
-				FAnimationRuntime::AccumulateAdditivePose(Output.Pose, AdditivePose, Output.Curve, AdditiveCurve, 1.f, AdditiveAnimType);
-			}
-			else
-			{
-				//doesn't handle additive yet
-				Composite->GetAnimationPose(Output.Pose, Output.Curve, ExtractionContext);
-			}
-		}
-		else if (UAnimMontage* Montage = Cast<UAnimMontage>(CurrentAsset))
-		{
-			// for now only update first slot
-			// in the future, add option to see which slot to see
-			if (Montage->SlotAnimTracks.Num() > 0)
-			{
-				FCompactPose SourcePose;
-				FBlendedCurve SourceCurve;
-				SourcePose.SetBoneContainer(&Output.Pose.GetBoneContainer());
-				SourceCurve.InitFrom(Output.Curve);
-			
-				if (Montage->IsValidAdditive())
-				{
-#if WITH_EDITORONLY_DATA
-					// if montage is additive, we need to have base pose for the slot pose evaluate
-					if (Montage->PreviewBasePose && Montage->SequenceLength > 0.f)
-					{
-						Montage->PreviewBasePose->GetBonePose(SourcePose, SourceCurve, FAnimExtractContext(CurrentTime));
-					}
-					else
-#endif // WITH_EDITORONLY_DATA
-					{
-						SourcePose.ResetToRefPose();				
-					}
-				}
-				else
-				{
-					SourcePose.ResetToRefPose();			
-				}
-
-				SlotEvaluatePose(Montage->SlotAnimTracks[0].SlotName, SourcePose, SourceCurve, WeightInfo.SourceWeight, Output.Pose, Output.Curve, WeightInfo.SlotNodeWeight, WeightInfo.TotalNodeWeight);
-			}
-		}
-	}
-
-	if(CurrentVertexAnim != NULL)
-	{
-		VertexAnims.Add(FActiveVertexAnim(CurrentVertexAnim, 1.f, CurrentTime));
-	}
-
-	return true;
-}
-
 void UAnimSingleNodeInstance::NativePostEvaluateAnimation()
 {
 	PostEvaluateAnimEvent.ExecuteIfBound();
@@ -427,9 +158,10 @@ void UAnimSingleNodeInstance::NativePostEvaluateAnimation()
 
 void UAnimSingleNodeInstance::OnMontageInstanceStopped(FAnimMontageInstance& StoppedMontageInstance) 
 {
-	if (StoppedMontageInstance.Montage == CurrentAsset)
+	FAnimSingleNodeInstanceProxy& Proxy = GetProxyOnGameThread<FAnimSingleNodeInstanceProxy>();
+	if (StoppedMontageInstance.Montage == Proxy.GetCurrentAsset())
 	{
-		CurrentTime = StoppedMontageInstance.GetPosition();
+		Proxy.SetCurrentTime(StoppedMontageInstance.GetPosition());
 	}
 
 	Super::OnMontageInstanceStopped(StoppedMontageInstance);
@@ -442,7 +174,8 @@ void UAnimSingleNodeInstance::Montage_Advance(float DeltaTime)
 	FAnimMontageInstance * CurMontageInstance = GetActiveMontageInstance();
 	if ( CurMontageInstance )
 	{
-		CurrentTime = CurMontageInstance->GetPosition();
+		FAnimSingleNodeInstanceProxy& Proxy = GetProxyOnGameThread<FAnimSingleNodeInstanceProxy>();
+		Proxy.SetCurrentTime(CurMontageInstance->GetPosition());
 	}
 }
 
@@ -461,25 +194,27 @@ void UAnimSingleNodeInstance::StopAnim()
 
 void UAnimSingleNodeInstance::SetLooping(bool bIsLooping)
 {
-	bLooping = bIsLooping;
+	FAnimSingleNodeInstanceProxy& Proxy = GetProxyOnGameThread<FAnimSingleNodeInstanceProxy>();
+	Proxy.SetLooping(bIsLooping);
 
-	if (UAnimMontage* Montage = Cast<UAnimMontage>(CurrentAsset))
+	if (UAnimMontage* Montage = Cast<UAnimMontage>(GetCurrentAsset()))
 	{
-		SetMontageLoop(Montage, bLooping, Montage_GetCurrentSection());
+		SetMontageLoop(Montage, Proxy.IsLooping(), Montage_GetCurrentSection());
 	}
 }
 
 void UAnimSingleNodeInstance::SetPlaying(bool bIsPlaying)
 {
-	bPlaying = bIsPlaying;
+	FAnimSingleNodeInstanceProxy& Proxy = GetProxyOnGameThread<FAnimSingleNodeInstanceProxy>();
+	Proxy.SetPlaying(bIsPlaying);
 
 	if (FAnimMontageInstance* CurMontageInstance = GetActiveMontageInstance())
 	{
 		CurMontageInstance->bPlaying = bIsPlaying;
 	}
-	else if (bPlaying)
+	else if (Proxy.IsPlaying())
 	{
-		UAnimMontage* Montage = Cast<UAnimMontage>(CurrentAsset);
+		UAnimMontage* Montage = Cast<UAnimMontage>(GetCurrentAsset());
 		if (Montage)
 		{
 			RestartMontage(Montage);
@@ -487,9 +222,35 @@ void UAnimSingleNodeInstance::SetPlaying(bool bIsPlaying)
 	}
 }
 
+bool UAnimSingleNodeInstance::IsPlaying() const
+{
+	// since setPlaying is setting to montage, we should get it as symmmetry
+	if (FAnimMontageInstance* CurMontageInstance = GetActiveMontageInstance())
+	{
+		return CurMontageInstance->bPlaying;
+	}
+
+	return GetProxyOnGameThread<FAnimSingleNodeInstanceProxy>().IsPlaying();
+}
+
+bool UAnimSingleNodeInstance::IsReverse() const
+{
+	return GetProxyOnGameThread<FAnimSingleNodeInstanceProxy>().IsReverse();
+}
+
+bool UAnimSingleNodeInstance::IsLooping() const
+{
+	return GetProxyOnGameThread<FAnimSingleNodeInstanceProxy>().IsLooping();
+}
+
+float UAnimSingleNodeInstance::GetPlayRate() const
+{
+	return GetProxyOnGameThread<FAnimSingleNodeInstanceProxy>().GetPlayRate();
+}
+
 void UAnimSingleNodeInstance::SetPlayRate(float InPlayRate)
 {
-	PlayRate = InPlayRate;
+	GetProxyOnGameThread<FAnimSingleNodeInstanceProxy>().SetPlayRate(InPlayRate);
 
 	if (FAnimMontageInstance* CurMontageInstance = GetActiveMontageInstance())
 	{
@@ -497,42 +258,36 @@ void UAnimSingleNodeInstance::SetPlayRate(float InPlayRate)
 	}
 }
 
-void UAnimSingleNodeInstance::SetReverse(bool bInReverse)
+UAnimationAsset* UAnimSingleNodeInstance::GetCurrentAsset()
 {
-	bReverse = bInReverse;
-	if (bInReverse)
-	{
-		PlayRate = -FMath::Abs(PlayRate);
-	}
-	else
-	{
-		PlayRate = FMath::Abs(PlayRate);
-	}
+	return GetProxyOnGameThread<FAnimSingleNodeInstanceProxy>().GetCurrentAsset();
+}
 
-// reverse support is a bit tricky for montage
-// since we don't have delegate when it reached to the beginning
-// for now I comment this out and do not support
-// I'd like the buttons to be customizable per asset types -
-// 	TTP 233456	ANIM: support different scrub controls per asset type
-/*
-	FAnimMontageInstance * CurMontageInstance = GetActiveMontageInstance();
-	if ( CurMontageInstance )
-	{
-		if ( bReverse == (CurMontageInstance->PlayRate > 0.f) )
-		{
-			CurMontageInstance->PlayRate *= -1.f;
-		}
-	}*/
+UVertexAnimation* UAnimSingleNodeInstance::GetCurrentVertexAnimation()
+{
+	return GetProxyOnGameThread<FAnimSingleNodeInstanceProxy>().GetCurrentVertexAnimation();
+}
+
+float UAnimSingleNodeInstance::GetCurrentTime() const
+{
+	return GetProxyOnGameThread<FAnimSingleNodeInstanceProxy>().GetCurrentTime();
+}
+
+void UAnimSingleNodeInstance::SetReverse(bool bInReverse)
+{
+	GetProxyOnGameThread<FAnimSingleNodeInstanceProxy>().SetReverse(bInReverse);
 }
 
 void UAnimSingleNodeInstance::SetPosition(float InPosition, bool bFireNotifies)
 {
-	float PreviousTime = CurrentTime;
-	CurrentTime = FMath::Clamp<float>(InPosition, 0.f, GetLength());
+	FAnimSingleNodeInstanceProxy& Proxy = GetProxyOnGameThread<FAnimSingleNodeInstanceProxy>();
+
+	float PreviousTime = Proxy.GetCurrentTime();
+	Proxy.SetCurrentTime(FMath::Clamp<float>(InPosition, 0.f, GetLength()));
 
 	if (FAnimMontageInstance* CurMontageInstance = GetActiveMontageInstance())
 	{
-		CurMontageInstance->SetPosition(CurrentTime);
+		CurMontageInstance->SetPosition(Proxy.GetCurrentTime());
 	}
 
 	// Handle notifies
@@ -540,17 +295,17 @@ void UAnimSingleNodeInstance::SetPosition(float InPosition, bool bFireNotifies)
 	// this will need to handle manually, emptying, it and collect it, and trigger them at once. 
 	if (bFireNotifies)
 	{
-		UAnimSequenceBase * SequenceBase = Cast<UAnimSequenceBase> (CurrentAsset);
+		UAnimSequenceBase * SequenceBase = Cast<UAnimSequenceBase> (Proxy.GetCurrentAsset());
 		if (SequenceBase)
 		{
-			AnimNotifies.Empty();
+			NotifyQueue.Reset(GetSkelMeshComponent());
 
 			TArray<const FAnimNotifyEvent*> Notifies;
-			SequenceBase->GetAnimNotifiesFromDeltaPositions(PreviousTime, CurrentTime, Notifies);
+			SequenceBase->GetAnimNotifiesFromDeltaPositions(PreviousTime, Proxy.GetCurrentTime(), Notifies);
 			if ( Notifies.Num() > 0 )
 			{
 				// single node instance only has 1 asset at a time
-				AddAnimNotifies(Notifies, 1.0f);
+				NotifyQueue.AddAnimNotifies(Notifies, 1.0f);
 			}
 
 			TriggerAnimNotifies(0.f);
@@ -561,91 +316,41 @@ void UAnimSingleNodeInstance::SetPosition(float InPosition, bool bFireNotifies)
 
 void UAnimSingleNodeInstance::SetBlendSpaceInput(const FVector& InBlendInput)
 {
-	BlendSpaceInput = InBlendInput;
-}
-
-void UAnimSingleNodeInstance::InternalBlendSpaceEvaluatePose(class UBlendSpaceBase* BlendSpace, TArray<FBlendSampleData>& BlendSampleDataCache, FPoseContext& OutContext)
-{
-	USkeletalMeshComponent* Component = GetSkelMeshComponent();
-
-	if (BlendSpace->IsValidAdditive())
-	{
-		FCompactPose& OutPose = OutContext.Pose;
-		FBlendedCurve& OutCurve = OutContext.Curve;
-		FCompactPose AdditivePose;
-		FBlendedCurve AdditiveCurve;
-		AdditivePose.SetBoneContainer(&OutPose.GetBoneContainer());
-		AdditiveCurve.InitFrom(OutCurve);
-
-#if WITH_EDITORONLY_DATA
-		if (BlendSpace->PreviewBasePose)
-		{
-			BlendSpace->PreviewBasePose->GetBonePose(/*out*/ OutPose, /*out*/OutCurve, FAnimExtractContext(PreviewPoseCurrentTime));
-		}
-		else
-#endif // WITH_EDITORONLY_DATA
-		{
-			// otherwise, get ref pose
-			OutPose.ResetToRefPose();
-		}
-
-		BlendSpace->GetAnimationPose(BlendSampleDataCache, AdditivePose, AdditiveCurve);
-
-		enum EAdditiveAnimationType AdditiveType = BlendSpace->bRotationBlendInMeshSpace? AAT_RotationOffsetMeshSpace : AAT_LocalSpaceBase;
-
-		FAnimationRuntime::AccumulateAdditivePose(OutPose, AdditivePose, OutCurve, AdditiveCurve, 1.f, AdditiveType);
-	}
-	else
-	{
-		BlendSpace->GetAnimationPose(BlendSampleDataCache, OutContext.Pose, OutContext.Curve);
-	}
+	GetProxyOnGameThread<FAnimSingleNodeInstanceProxy>().SetBlendSpaceInput(InBlendInput);
 }
 
 float UAnimSingleNodeInstance::GetLength()
 {
-	if ((CurrentAsset != NULL))
-	{
-		if (UBlendSpace* BlendSpace = Cast<UBlendSpace>(CurrentAsset))
-		{
-			return BlendSpace->AnimLength;
-		}
-		else if (UAnimSequenceBase* SequenceBase = Cast<UAnimSequenceBase>(CurrentAsset))
-		{
-			return SequenceBase->SequenceLength;
-		}
-	}	
-
-	return 0.f;
+	return GetProxyOnGameThread<FAnimSingleNodeInstanceProxy>().GetLength();
 }
 
 void UAnimSingleNodeInstance::StepForward()
 {
-	if ((CurrentAsset != NULL))
+	if (UAnimSequence* Sequence = Cast<UAnimSequence>(GetCurrentAsset()))
 	{
-		if (UAnimSequence* Sequence = Cast<UAnimSequence>(CurrentAsset))
-		{
-			float KeyLength = Sequence->SequenceLength/Sequence->NumFrames+SMALL_NUMBER;
-			float Fraction = (CurrentTime+KeyLength)/Sequence->SequenceLength;
-			int32 Frames = FMath::Clamp<int32>((float)(Sequence->NumFrames*Fraction), 0, Sequence->NumFrames);
-			SetPosition(Frames*KeyLength);
-		}
+		FAnimSingleNodeInstanceProxy& Proxy = GetProxyOnGameThread<FAnimSingleNodeInstanceProxy>();
+
+		float KeyLength = Sequence->SequenceLength/Sequence->NumFrames+SMALL_NUMBER;
+		float Fraction = (Proxy.GetCurrentTime()+KeyLength)/Sequence->SequenceLength;
+		int32 Frames = FMath::Clamp<int32>((float)(Sequence->NumFrames*Fraction), 0, Sequence->NumFrames);
+		SetPosition(Frames*KeyLength);
 	}	
 }
 
 void UAnimSingleNodeInstance::StepBackward()
 {
-	if ((CurrentAsset != NULL))
+	if (UAnimSequence* Sequence = Cast<UAnimSequence>(GetCurrentAsset()))
 	{
-		if (UAnimSequence* Sequence = Cast<UAnimSequence>(CurrentAsset))
-		{
-			float KeyLength = Sequence->SequenceLength/Sequence->NumFrames+SMALL_NUMBER;
-			float Fraction = (CurrentTime-KeyLength)/Sequence->SequenceLength;
-			int32 Frames = FMath::Clamp<int32>((float)(Sequence->NumFrames*Fraction), 0, Sequence->NumFrames);
-			SetPosition(Frames*KeyLength);
-		}
-	}	
+		FAnimSingleNodeInstanceProxy& Proxy = GetProxyOnGameThread<FAnimSingleNodeInstanceProxy>();
+
+		float KeyLength = Sequence->SequenceLength/Sequence->NumFrames+SMALL_NUMBER;
+		float Fraction = (Proxy.GetCurrentTime()-KeyLength)/Sequence->SequenceLength;
+		int32 Frames = FMath::Clamp<int32>((float)(Sequence->NumFrames*Fraction), 0, Sequence->NumFrames);
+		SetPosition(Frames*KeyLength);
+	}
 }
-/////////////////////////////////////////////////////
-/////////////////////////////////////////////////////
-/////////////////////////////////////////////////////
 
+FAnimInstanceProxy* UAnimSingleNodeInstance::CreateAnimInstanceProxy()
+{
+	return new FAnimSingleNodeInstanceProxy(this);
+}
diff --git a/Engine/Source/Runtime/Engine/Private/Animation/AnimSingleNodeInstanceProxy.cpp b/Engine/Source/Runtime/Engine/Private/Animation/AnimSingleNodeInstanceProxy.cpp
new file mode 100644
--- /dev/null
+++ b/Engine/Source/Runtime/Engine/Private/Animation/AnimSingleNodeInstanceProxy.cpp
@@ -0,0 +1,406 @@
+// Copyright 1998-2015 Epic Games, Inc. All Rights Reserved.
+
+#include "EnginePrivate.h"
+#include "AnimSingleNodeInstanceProxy.h"
+#include "Animation/AnimComposite.h"
+#include "Animation/AnimMontage.h"
+#include "Animation/VertexAnim/VertexAnimation.h"
+#include "Animation/BlendSpace.h"
+
+void FAnimSingleNodeInstanceProxy::Initialize(UAnimInstance* InAnimInstance)
+{
+	FAnimInstanceProxy::Initialize(InAnimInstance);
+
+	CurrentAsset = NULL;
+	CurrentVertexAnim = NULL;
+#if WITH_EDITORONLY_DATA
+	PreviewPoseCurrentTime = 0.0f;
+#endif
+
+	// it's already doing it when evaluate
+	BlendSpaceInput = FVector::ZeroVector;
+	CurrentTime = 0.f;
+}
+
+bool FAnimSingleNodeInstanceProxy::Evaluate(FPoseContext& Output)
+{
+	if (CurrentAsset != NULL)
+	{
+		//@TODO: animrefactor: Seems like more code duplication than we need
+		if (UBlendSpaceBase* BlendSpace = Cast<UBlendSpaceBase>(CurrentAsset))
+		{
+			InternalBlendSpaceEvaluatePose(BlendSpace, BlendSampleData, Output);
+		}
+		else if (UAnimSequence* Sequence = Cast<UAnimSequence>(CurrentAsset))
+		{
+			if (Sequence->IsValidAdditive())
+			{
+				FAnimExtractContext ExtractionContext(CurrentTime, Sequence->bEnableRootMotion);
+
+				Sequence->GetAdditiveBasePose(Output.Pose, Output.Curve, ExtractionContext);
+
+				FCompactPose AdditivePose;
+				FBlendedCurve AdditiveCurve;
+				AdditivePose.SetBoneContainer(&Output.Pose.GetBoneContainer());
+				AdditiveCurve.InitFrom(Output.Curve);
+				Sequence->GetAnimationPose(AdditivePose, AdditiveCurve, ExtractionContext);
+
+				FAnimationRuntime::AccumulateAdditivePose(Output.Pose, AdditivePose, Output.Curve, AdditiveCurve, 1.f, Sequence->AdditiveAnimType);
+				Output.Pose.NormalizeRotations();
+			}
+			else
+			{
+				// if SkeletalMesh isn't there, we'll need to use skeleton
+				Sequence->GetAnimationPose(Output.Pose, Output.Curve, FAnimExtractContext(CurrentTime, Sequence->bEnableRootMotion));
+			}
+		}
+		else if (UAnimComposite* Composite = Cast<UAnimComposite>(CurrentAsset))
+		{
+			FAnimExtractContext ExtractionContext(CurrentTime, ShouldExtractRootMotion());
+			const FAnimTrack& AnimTrack = Composite->AnimationTrack;
+
+			// find out if this is additive animation
+			if (AnimTrack.IsAdditive())
+			{
+#if WITH_EDITORONLY_DATA
+				if (Composite->PreviewBasePose)
+				{
+					Composite->PreviewBasePose->GetAdditiveBasePose(Output.Pose, Output.Curve, ExtractionContext);
+				}
+				else
+#endif
+				{
+					// get base pose - for now we only support ref pose as base
+					Output.Pose.ResetToRefPose();
+				}
+
+				EAdditiveAnimationType AdditiveAnimType = AnimTrack.IsRotationOffsetAdditive()? AAT_RotationOffsetMeshSpace : AAT_LocalSpaceBase;
+
+				FCompactPose AdditivePose;
+				FBlendedCurve AdditiveCurve;
+				AdditivePose.SetBoneContainer(&Output.Pose.GetBoneContainer());
+				AdditiveCurve.InitFrom(Output.Curve);
+				Composite->GetAnimationPose(AdditivePose, AdditiveCurve, ExtractionContext);
+
+				FAnimationRuntime::AccumulateAdditivePose(Output.Pose, AdditivePose, Output.Curve, AdditiveCurve, 1.f, AdditiveAnimType);
+			}
+			else
+			{
+				//doesn't handle additive yet
+				Composite->GetAnimationPose(Output.Pose, Output.Curve, ExtractionContext);
+			}
+		}
+		else if (UAnimMontage* Montage = Cast<UAnimMontage>(CurrentAsset))
+		{
+			// for now only update first slot
+			// in the future, add option to see which slot to see
+			if (Montage->SlotAnimTracks.Num() > 0)
+			{
+				FCompactPose SourcePose;
+				FBlendedCurve SourceCurve;
+				SourcePose.SetBoneContainer(&Output.Pose.GetBoneContainer());
+				SourceCurve.InitFrom(Output.Curve);
+			
+				if (Montage->IsValidAdditive())
+				{
+#if WITH_EDITORONLY_DATA
+					// if montage is additive, we need to have base pose for the slot pose evaluate
+					if (Montage->PreviewBasePose && Montage->SequenceLength > 0.f)
+					{
+						Montage->PreviewBasePose->GetBonePose(SourcePose, SourceCurve, FAnimExtractContext(CurrentTime));
+					}
+					else
+#endif // WITH_EDITORONLY_DATA
+					{
+						SourcePose.ResetToRefPose();				
+					}
+				}
+				else
+				{
+					SourcePose.ResetToRefPose();			
+				}
+
+				SlotEvaluatePose(Montage->SlotAnimTracks[0].SlotName, SourcePose, SourceCurve, WeightInfo.SourceWeight, Output.Pose, Output.Curve, WeightInfo.SlotNodeWeight, WeightInfo.TotalNodeWeight);
+			}
+		}
+	}
+
+	if(CurrentVertexAnim != NULL)
+	{
+		AddVertexAnim(FActiveVertexAnim(CurrentVertexAnim, 1.f, CurrentTime));
+	}
+
+	return true;
+}
+
+void FAnimSingleNodeInstanceProxy::UpdateAnimationNode(float DeltaSeconds)
+{
+	float NewPlayRate = PlayRate;
+	UAnimSequence* PreviewBasePose = NULL;
+
+	if (bPlaying == false)
+	{
+		// we still have to tick animation when bPlaying is false because 
+		NewPlayRate = 0.f;
+	}
+
+	if(CurrentAsset != NULL)
+	{
+		FAnimGroupInstance* SyncGroup;
+		if (UBlendSpaceBase* BlendSpace = Cast<UBlendSpaceBase>(CurrentAsset))
+		{
+			FAnimTickRecord& TickRecord = CreateUninitializedTickRecord(INDEX_NONE, /*out*/ SyncGroup);
+			MakeBlendSpaceTickRecord(TickRecord, BlendSpace, BlendSpaceInput, BlendSampleData, BlendFilter, bLooping, NewPlayRate, 1.f, /*inout*/ CurrentTime, MarkerTickRecord);
+#if WITH_EDITORONLY_DATA
+			PreviewBasePose = BlendSpace->PreviewBasePose;
+#endif
+		}
+		else if (UAnimSequence* Sequence = Cast<UAnimSequence>(CurrentAsset))
+		{
+			FAnimTickRecord& TickRecord = CreateUninitializedTickRecord(INDEX_NONE, /*out*/ SyncGroup);
+			MakeSequenceTickRecord(TickRecord, Sequence, bLooping, NewPlayRate, 1.f, /*inout*/ CurrentTime, MarkerTickRecord);
+			// if it's not looping, just set play to be false when reached to end
+			if (!bLooping)
+			{
+				const float CombinedPlayRate = NewPlayRate*Sequence->RateScale;
+				if ((CombinedPlayRate < 0.f && CurrentTime <= 0.f) || (CombinedPlayRate > 0.f && CurrentTime >= Sequence->SequenceLength))
+				{
+					SetPlaying(false);
+				}
+			}
+		}
+		else if(UAnimComposite* Composite = Cast<UAnimComposite>(CurrentAsset))
+		{
+			FAnimTickRecord& TickRecord = CreateUninitializedTickRecord(INDEX_NONE, /*out*/ SyncGroup);
+			MakeSequenceTickRecord(TickRecord, Composite, bLooping, NewPlayRate, 1.f, /*inout*/ CurrentTime, MarkerTickRecord);
+			// if it's not looping, just set play to be false when reached to end
+			if (!bLooping)
+			{
+				const float CombinedPlayRate = NewPlayRate*Composite->RateScale;
+				if ((CombinedPlayRate < 0.f && CurrentTime <= 0.f) || (CombinedPlayRate > 0.f && CurrentTime >= Composite->SequenceLength))
+				{
+					SetPlaying(false);
+				}
+			}
+		}
+		else if (UAnimMontage * Montage = Cast<UAnimMontage>(CurrentAsset))
+		{
+			// Full weight , if you don't have slot track, you won't be able to see animation playing
+			if ( Montage->SlotAnimTracks.Num() > 0 )
+			{
+				// in the future, maybe we can support which slot
+				const FName CurrentSlotNodeName = Montage->SlotAnimTracks[0].SlotName;
+				GetSlotWeight(CurrentSlotNodeName, WeightInfo.SlotNodeWeight, WeightInfo.SourceWeight, WeightInfo.TotalNodeWeight);
+				UpdateSlotNodeWeight(CurrentSlotNodeName, WeightInfo.SlotNodeWeight);
+			}
+			// get the montage position
+			// @todo anim: temporarily just choose first slot and show the location
+			const FMontageEvaluationState* ActiveMontageEvaluationState = GetActiveMontageEvaluationState();
+			if (ActiveMontageEvaluationState)
+			{
+				CurrentTime = ActiveMontageEvaluationState->MontagePosition;
+			}
+			else if (bPlaying)
+			{
+				SetPlaying(false);
+			}
+#if WITH_EDITORONLY_DATA
+			PreviewBasePose = Montage->PreviewBasePose;
+#endif
+		}
+	}
+	else if(CurrentVertexAnim != NULL)
+	{
+		float MoveDelta = DeltaSeconds * NewPlayRate;
+		FAnimationRuntime::AdvanceTime(bLooping, MoveDelta, CurrentTime, CurrentVertexAnim->GetAnimLength());
+	}
+
+#if WITH_EDITORONLY_DATA
+	if(PreviewBasePose)
+	{
+		float MoveDelta = DeltaSeconds * NewPlayRate;
+		const bool bIsPreviewPoseLooping = true;
+
+		FAnimationRuntime::AdvanceTime(bIsPreviewPoseLooping, MoveDelta, PreviewPoseCurrentTime, PreviewBasePose->SequenceLength);
+	}
+#endif
+}
+
+void FAnimSingleNodeInstanceProxy::PostUpdate(UAnimInstance* InAnimInstance) const
+{
+	FAnimInstanceProxy::PostUpdate(InAnimInstance);
+
+	// sync up playing state for active montage instances
+	int32 EvaluationDataIndex = 0;
+	const TArray<FMontageEvaluationState>& EvaluationData = GetMontageEvaluationData();
+	for (FAnimMontageInstance* MontageInstance : InAnimInstance->MontageInstances)
+	{
+		if (MontageInstance->Montage && MontageInstance->GetWeight() > ZERO_ANIMWEIGHT_THRESH)
+		{
+			// sanity check we are playing the same montage
+			check(MontageInstance->Montage == EvaluationData[EvaluationDataIndex].Montage);
+			MontageInstance->bPlaying = EvaluationData[EvaluationDataIndex].bIsPlaying;
+			EvaluationDataIndex++;
+		}
+	}
+}
+
+void FAnimSingleNodeInstanceProxy::InternalBlendSpaceEvaluatePose(class UBlendSpaceBase* BlendSpace, TArray<FBlendSampleData>& BlendSampleDataCache, FPoseContext& OutContext)
+{
+	if (BlendSpace->IsValidAdditive())
+	{
+		FCompactPose& OutPose = OutContext.Pose;
+		FBlendedCurve& OutCurve = OutContext.Curve;
+		FCompactPose AdditivePose;
+		FBlendedCurve AdditiveCurve;
+		AdditivePose.SetBoneContainer(&OutPose.GetBoneContainer());
+		AdditiveCurve.InitFrom(OutCurve);
+
+#if WITH_EDITORONLY_DATA
+		if (BlendSpace->PreviewBasePose)
+		{
+			BlendSpace->PreviewBasePose->GetBonePose(/*out*/ OutPose, /*out*/OutCurve, FAnimExtractContext(PreviewPoseCurrentTime));
+		}
+		else
+#endif // WITH_EDITORONLY_DATA
+		{
+			// otherwise, get ref pose
+			OutPose.ResetToRefPose();
+		}
+
+		BlendSpace->GetAnimationPose(BlendSampleDataCache, AdditivePose, AdditiveCurve);
+
+		enum EAdditiveAnimationType AdditiveType = BlendSpace->bRotationBlendInMeshSpace? AAT_RotationOffsetMeshSpace : AAT_LocalSpaceBase;
+
+		FAnimationRuntime::AccumulateAdditivePose(OutPose, AdditivePose, OutCurve, AdditiveCurve, 1.f, AdditiveType);
+	}
+	else
+	{
+		BlendSpace->GetAnimationPose(BlendSampleDataCache, OutContext.Pose, OutContext.Curve);
+	}
+}
+
+void FAnimSingleNodeInstanceProxy::SetAnimationAsset(class UAnimationAsset* NewAsset, USkeletalMeshComponent* MeshComponent, bool bIsLooping, float InPlayRate)
+{
+	if (NewAsset != CurrentAsset)
+	{
+		CurrentAsset = NewAsset;
+	}
+
+	if (MeshComponent)
+	{
+		if (MeshComponent->SkeletalMesh == NULL)
+		{
+			// if it does not have SkeletalMesh, we nullify it
+			CurrentAsset = NULL;
+		}
+		else if (CurrentAsset != NULL)
+		{
+			// if we have an asset, make sure their skeleton matches, otherwise, null it
+			if (GetSkeleton() != CurrentAsset->GetSkeleton())
+			{
+				// clear asset since we do not have matching skeleton
+				CurrentAsset = NULL;
+			}
+		}
+	}
+
+	bLooping = bIsLooping;
+	PlayRate = InPlayRate;
+	CurrentTime = 0.f;
+	BlendSpaceInput = FVector::ZeroVector;
+	MarkerTickRecord.Reset();
+
+#if WITH_EDITORONLY_DATA
+	PreviewPoseCurrentTime = 0.0f;
+#endif
+}
+
+void FAnimSingleNodeInstanceProxy::UpdateBlendspaceSamples(FVector InBlendInput)
+{
+	if (UBlendSpaceBase* BlendSpace = Cast<UBlendSpaceBase>(CurrentAsset))
+	{
+		float OutCurrentTime = 0.f;
+		FMarkerTickRecord TempMarkerTickRecord;
+		BlendSpaceAdvanceImmediate(BlendSpace, InBlendInput, BlendSampleData, BlendFilter, false, 1.f, 0.f, OutCurrentTime, TempMarkerTickRecord);
+	}
+}
+
+void FAnimSingleNodeInstanceProxy::SetVertexAnimation(UVertexAnimation * NewVertexAnim, bool bIsLooping, float InPlayRate)
+{
+	if (NewVertexAnim != CurrentVertexAnim)
+	{
+		CurrentVertexAnim = NewVertexAnim;
+	}
+
+	if (USkeletalMeshComponent * MeshComponent = GetSkelMeshComponent())
+	{
+		if (MeshComponent->SkeletalMesh == NULL)
+		{
+			// if it does not have SkeletalMesh, we nullify it
+			CurrentVertexAnim = NULL;
+		}
+		else if (CurrentVertexAnim != NULL)
+		{
+			// if we have an anim, make sure their mesh matches, otherwise, null it
+			if (MeshComponent->SkeletalMesh != CurrentVertexAnim->BaseSkelMesh)
+			{
+				// clear asset since we do not have matching skeleton
+				CurrentVertexAnim = NULL;
+			}
+		}
+	}
+
+	bLooping = bIsLooping;
+	PlayRate = InPlayRate;
+}
+
+void FAnimSingleNodeInstanceProxy::SetReverse(bool bInReverse)
+{
+	bReverse = bInReverse;
+	if (bInReverse)
+	{
+		PlayRate = -FMath::Abs(PlayRate);
+	}
+	else
+	{
+		PlayRate = FMath::Abs(PlayRate);
+	}
+
+// reverse support is a bit tricky for montage
+// since we don't have delegate when it reached to the beginning
+// for now I comment this out and do not support
+// I'd like the buttons to be customizable per asset types -
+// 	TTP 233456	ANIM: support different scrub controls per asset type
+/*
+	FAnimMontageInstance * CurMontageInstance = GetActiveMontageInstance();
+	if ( CurMontageInstance )
+	{
+		if ( bReverse == (CurMontageInstance->PlayRate > 0.f) )
+		{
+			CurMontageInstance->PlayRate *= -1.f;
+		}
+	}*/
+}
+
+float FAnimSingleNodeInstanceProxy::GetLength()
+{
+	if ((CurrentAsset != NULL))
+	{
+		if (UBlendSpace* BlendSpace = Cast<UBlendSpace>(CurrentAsset))
+		{
+			return BlendSpace->AnimLength;
+		}
+		else if (UAnimSequenceBase* SequenceBase = Cast<UAnimSequenceBase>(CurrentAsset))
+		{
+			return SequenceBase->SequenceLength;
+		}
+	}	
+
+	return 0.f;
+}
+
+void FAnimSingleNodeInstanceProxy::SetBlendSpaceInput(const FVector& InBlendInput)
+{
+	BlendSpaceInput = InBlendInput;
+}
diff --git a/Engine/Source/Runtime/Engine/Private/Animation/AnimationRuntime.cpp b/Engine/Source/Runtime/Engine/Private/Animation/AnimationRuntime.cpp
--- a/Engine/Source/Runtime/Engine/Private/Animation/AnimationRuntime.cpp
+++ b/Engine/Source/Runtime/Engine/Private/Animation/AnimationRuntime.cpp
@@ -15,6 +15,9 @@
 #include "Animation/AnimInstance.h"
 #include "BonePose.h"
 #include "Animation/BlendProfile.h"
+#include "SkeletalRender.h"
+#include "Animation/VertexAnim/VertexAnimBase.h"
+#include "Animation/VertexAnim/MorphTarget.h"
 
 DEFINE_LOG_CATEGORY(LogAnimation);
 DEFINE_LOG_CATEGORY(LogRootMotion);
@@ -1277,6 +1280,73 @@ void FAnimationRuntime::FillUpSpaceBasesRetargetBasePose(const USkeleton* Skelet
 }
 #endif // WITH_EDITOR
 
+/** See if an array of ActiveVertexAnims already contains the supplied anim */
+static int32 FindVertexAnim(const TArray<FActiveVertexAnim>& ActiveAnims, UVertexAnimBase* Anim)
+{
+	for(int32 i=0; i<ActiveAnims.Num(); i++)
+	{
+		if(ActiveAnims[i].VertAnim == Anim)
+		{
+			return i;
+		}
+	}
+
+	return INDEX_NONE;
+}
+
+TArray<FActiveVertexAnim> FAnimationRuntime::UpdateActiveVertexAnims(const USkeletalMesh* InSkeletalMesh, const TMap<FName, float>& MorphCurveAnims, const TArray<FActiveVertexAnim>& ActiveAnims)
+{
+	TArray<struct FActiveVertexAnim> OutVertexAnims;
+
+	// First copy ActiveAnims
+	for(int32 AnimIdx=0; AnimIdx < ActiveAnims.Num(); AnimIdx++)
+	{
+		const FActiveVertexAnim& ActiveAnim = ActiveAnims[AnimIdx];
+		const float ActiveAnimAbsWeight = FMath::Abs(ActiveAnim.Weight);
+
+		// Check it has valid weight, and works on this SkeletalMesh
+		if (	ActiveAnimAbsWeight > MinVertexAnimBlendWeight &&
+			ActiveAnim.VertAnim != NULL &&
+			ActiveAnim.VertAnim->BaseSkelMesh == InSkeletalMesh)
+		{
+			OutVertexAnims.Add(ActiveAnim);
+		}
+		// @TODO Need to check for duplicates here?
+	}
+
+	// Then go over the CurveKeys finding morph targets by name
+	for(auto CurveIter=MorphCurveAnims.CreateConstIterator(); CurveIter; ++CurveIter)
+	{
+		const FName& CurveName	= (CurveIter).Key();
+		const float& Weight	= (CurveIter).Value();
+
+		// If it has a valid weight
+		if(FMath::Abs(Weight) > MinVertexAnimBlendWeight)
+		{
+			// Find morph reference
+			UMorphTarget* Target = InSkeletalMesh ? InSkeletalMesh->FindMorphTarget(CurveName) : NULL;
+			if(Target != NULL)				
+			{
+				// See if this morph target already has an entry
+				int32 AnimIndex = FindVertexAnim(OutVertexAnims, Target);
+				// If not, add it
+				if(AnimIndex == INDEX_NONE)
+				{
+					OutVertexAnims.Add(FActiveVertexAnim(Target, Weight));
+				}
+				// If it does, use the max weight
+				else
+				{
+					const float CurrentWeight = OutVertexAnims[AnimIndex].Weight;
+					OutVertexAnims[AnimIndex].Weight = FMath::Max<float>(CurrentWeight, Weight);
+				}
+			}
+		}
+	}
+
+	return OutVertexAnims;
+}
+
 /////////////////////////////////////////////////////////////////////////////////////////
 // FA2CSPose
 /////////////////////////////////////////////////////////////////////////////////////////
diff --git a/Engine/Source/Runtime/Engine/Private/Animation/BlendSpaceBase.cpp b/Engine/Source/Runtime/Engine/Private/Animation/BlendSpaceBase.cpp
--- a/Engine/Source/Runtime/Engine/Private/Animation/BlendSpaceBase.cpp
+++ b/Engine/Source/Runtime/Engine/Private/Animation/BlendSpaceBase.cpp
@@ -159,7 +159,7 @@ int32 GetHighestWeightSample(TArray<FBlendSampleData> &SampleDataList)
 }
 
 /////////////////////////////////////////////////////
-void UBlendSpaceBase::TickAssetPlayerInstance(FAnimTickRecord& Instance, class UAnimInstance* InstanceOwner, FAnimAssetTickContext& Context) const
+void UBlendSpaceBase::TickAssetPlayer(FAnimTickRecord& Instance, struct FAnimNotifyQueue& NotifyQueue, FAnimAssetTickContext& Context) const
 {
 	const float DeltaTime = Context.GetDeltaTime();
 	float MoveDelta = Instance.PlayRateMultiplier * DeltaTime;
@@ -457,7 +457,7 @@ void UBlendSpaceBase::TickAssetPlayerInstance(FAnimTickRecord& Instance, class U
 
 				if (bGenerateNotifies && Notifies.Num() > 0)
 				{
-					InstanceOwner->AddAnimNotifies(Notifies, Instance.EffectiveBlendWeight);
+					NotifyQueue.AddAnimNotifies(Notifies, Instance.EffectiveBlendWeight);
 				}
 			}
 		}
diff --git a/Engine/Source/Runtime/Engine/Private/Animation/SkeletalControl.cpp b/Engine/Source/Runtime/Engine/Private/Animation/SkeletalControl.cpp
--- a/Engine/Source/Runtime/Engine/Private/Animation/SkeletalControl.cpp
+++ b/Engine/Source/Runtime/Engine/Private/Animation/SkeletalControl.cpp
@@ -30,12 +30,15 @@ bool FBoneReference::Initialize(const FBoneContainer& RequiredBones)
 	// If bone name is not found, look into the master skeleton to see if it's found there.
 	// SkeletalMeshes can exclude bones from the master skeleton, and that's OK.
 	// If it's not found in the master skeleton, the bone does not exist at all! so we should report it as a warning.
-	if( (BoneIndex == INDEX_NONE) && RequiredBones.GetSkeletonAsset() )
+	if (BoneIndex == INDEX_NONE)
 	{
-		if( RequiredBones.GetSkeletonAsset()->GetReferenceSkeleton().FindBoneIndex(BoneName) == INDEX_NONE )
+		if (USkeleton* SkeletonAsset = RequiredBones.GetSkeletonAsset())
 		{
-			UE_LOG(LogAnimation, Warning, TEXT("FBoneReference::Initialize BoneIndex for Bone '%s' does not exist in Skeleton '%s'"), 
-				*BoneName.ToString(), *GetNameSafe(RequiredBones.GetSkeletonAsset()));
+			if (SkeletonAsset->GetReferenceSkeleton().FindBoneIndex(BoneName) == INDEX_NONE)
+			{
+				UE_LOG(LogAnimation, Warning, TEXT("FBoneReference::Initialize BoneIndex for Bone '%s' does not exist in Skeleton '%s'"), 
+					*BoneName.ToString(), *GetNameSafe(SkeletonAsset));
+			}
 		}
 	}
 
diff --git a/Engine/Source/Runtime/Engine/Private/AnimationAsset.cpp b/Engine/Source/Runtime/Engine/Private/AnimationAsset.cpp
--- a/Engine/Source/Runtime/Engine/Private/AnimationAsset.cpp
+++ b/Engine/Source/Runtime/Engine/Private/AnimationAsset.cpp
@@ -3,6 +3,7 @@
 #include "EnginePrivate.h"
 #include "Animation/AnimSequence.h"
 #include "AnimationUtils.h"
+#include "Animation/AnimInstanceProxy.h"
 
 #define LEADERSCORE_ALWAYSLEADER  	2.f
 #define LEADERSCORE_MONTAGE			3.f
@@ -204,6 +205,13 @@ void UAnimationAsset::SetSkeleton(USkeleton* NewSkeleton)
 	}
 }
 
+void UAnimationAsset::TickAssetPlayerInstance(FAnimTickRecord& Instance, class UAnimInstance* AnimInstance, FAnimAssetTickContext& Context) const
+{ 
+	// @todo: remove after deprecation
+	// Forward to non-deprecated function
+	TickAssetPlayer(Instance, AnimInstance->NotifyQueue, Context); 
+}
+
 #if WITH_EDITOR
 bool UAnimationAsset::ReplaceSkeleton(USkeleton* NewSkeleton, bool bConvertSpaces/*=false*/)
 {
@@ -296,6 +304,7 @@ FBoneContainer::FBoneContainer()
 , RefSkeleton(NULL)
 , bDisableRetargeting(false)
 , bUseRAWData(false)
+, bUseSourceData(false)
 {
 	BoneIndicesArray.Empty();
 	BoneSwitchArray.Empty();
@@ -311,6 +320,7 @@ FBoneContainer::FBoneContainer(const TArray<FBoneIndexType>& InRequiredBoneIndex
 , RefSkeleton(NULL)
 , bDisableRetargeting(false)
 , bUseRAWData(false)
+, bUseSourceData(false)
 {
 	Initialize();
 }
@@ -326,29 +336,36 @@ void FBoneContainer::InitializeTo(const TArray<FBoneIndexType>& InRequiredBoneIn
 void FBoneContainer::Initialize()
 {
 	RefSkeleton = NULL;
-	AssetSkeletalMesh = Cast<USkeletalMesh>(Asset.Get());
-	if( AssetSkeletalMesh.IsValid() )
+	UObject* AssetObj = Asset.Get();
+	USkeletalMesh* AssetSkeletalMeshObj = Cast<USkeletalMesh>(AssetObj);
+	USkeleton* AssetSkeletonObj = nullptr;
+
+	if (AssetSkeletalMeshObj)
 	{
-		RefSkeleton = &AssetSkeletalMesh->RefSkeleton;
-		AssetSkeleton = AssetSkeletalMesh->Skeleton;
+		RefSkeleton = &AssetSkeletalMeshObj->RefSkeleton;
+		AssetSkeletonObj = AssetSkeletalMeshObj->Skeleton;
 	}
 	else
 	{
-		AssetSkeleton = Cast<USkeleton>(Asset.Get());
-		if( AssetSkeleton.IsValid() )
+		AssetSkeletonObj = Cast<USkeleton>(AssetObj);
+		if (AssetSkeletonObj)
 		{
-			RefSkeleton = &AssetSkeleton->GetReferenceSkeleton();
+			RefSkeleton = &AssetSkeletonObj->GetReferenceSkeleton();
 		}
 	}
+
 	// Only supports SkeletalMeshes or Skeletons.
-	check( AssetSkeletalMesh.Get() || AssetSkeleton.Get() );
+	check( AssetSkeletalMeshObj || AssetSkeletonObj );
 	// Skeleton should always be there.
-	checkf( AssetSkeleton.Get(), TEXT("%s missing skeleton"), *GetNameSafe(AssetSkeletalMesh.Get()));
+	checkf( AssetSkeletonObj, TEXT("%s missing skeleton"), *GetNameSafe(AssetSkeletalMeshObj));
 	check( RefSkeleton );
 
+	AssetSkeleton = AssetSkeletonObj;
+	AssetSkeletalMesh = AssetSkeletalMeshObj;
+
 	// Take biggest amount of bones between SkeletalMesh and Skeleton for BoneSwitchArray.
 	// SkeletalMesh can have less, but AnimSequences tracks will map to Skeleton which can have more.
-	const int32 MaxBones = AssetSkeleton.IsValid() ? FMath::Max<int32>(RefSkeleton->GetNum(), AssetSkeleton->GetReferenceSkeleton().GetNum()) : RefSkeleton->GetNum();
+	const int32 MaxBones = AssetSkeletonObj ? FMath::Max<int32>(RefSkeleton->GetNum(), AssetSkeletonObj->GetReferenceSkeleton().GetNum()) : RefSkeleton->GetNum();
 
 	// Initialize BoneSwitchArray.
 	BoneSwitchArray.Init(false, MaxBones);
@@ -366,15 +383,15 @@ void FBoneContainer::Initialize()
 	// Cache our mapping tables
 	// Here we create look up tables between our target asset and its USkeleton's refpose.
 	// Most times our Target is a SkeletalMesh
-	if( AssetSkeletalMesh.IsValid() )
+	if (AssetSkeletalMeshObj)
 	{
-		RemapFromSkelMesh(*AssetSkeletalMesh.Get(), *AssetSkeleton.Get());
+		RemapFromSkelMesh(*AssetSkeletalMeshObj, *AssetSkeletonObj);
 	}
 	// But we also support a Skeleton's RefPose.
 	else
 	{
 		// Right now we only support a single Skeleton. Skeleton hierarchy coming soon!
-		RemapFromSkeleton(*AssetSkeleton.Get());
+		RemapFromSkeleton(*AssetSkeletonObj);
 	}
 
 	//Set up compact pose data
diff --git a/Engine/Source/Runtime/Engine/Private/Audio.cpp b/Engine/Source/Runtime/Engine/Private/Audio.cpp
--- a/Engine/Source/Runtime/Engine/Private/Audio.cpp
+++ b/Engine/Source/Runtime/Engine/Private/Audio.cpp
@@ -364,7 +364,22 @@ FSpatializationParams FSoundSource::GetSpatializationParams()
 	if (WaveInstance->bUseSpatialization)
 	{
 		FVector EmitterPosition = AudioDevice->GetListenerTransformedDirection(WaveInstance->Location, &Params.Distance);
-		Params.NormalizedOmniRadius = (Params.Distance > 0) ? (WaveInstance->OmniRadius / Params.Distance) : FLT_MAX;
+		
+		// If we are using the OmniRadius feature
+		if (WaveInstance->OmniRadius > 0.0f)
+		{
+			static const float MaxNormalizedRadius = 1000000.0f;
+			Params.NormalizedOmniRadius = MaxNormalizedRadius;
+
+			if (Params.Distance > 0)
+			{
+				Params.NormalizedOmniRadius = FMath::Clamp(WaveInstance->OmniRadius / Params.Distance, 0.0f, MaxNormalizedRadius);
+			}
+		}
+		else
+		{
+			Params.NormalizedOmniRadius = 0.0f;
+		}
 
 		if (Buffer->NumChannels == 2)
 		{
@@ -466,7 +481,7 @@ FWaveInstance::FWaveInstance( FActiveSound* InActiveSound )
 ,	ActiveSound( InActiveSound )
 ,	Volume( 0.0f )
 ,	VolumeMultiplier( 1.0f )
-,	PlayPriority( 0.0f )
+,	Priority( 1.0f )
 ,	VoiceCenterChannelVolume( 0.0f )
 ,	RadioFilterVolume( 0.0f )
 ,	RadioFilterVolumeThreshold( 0.0f )
@@ -568,7 +583,7 @@ float FWaveInstance::GetVolumeWeightedPriority() const
 	check(ActiveSound);
 	// If this wave instance's active sound should stop due to max concurrency, return a large negative weighted priority so 
 	// that it will always be sorted to the bottom of the WaveInstance list and stopped.
-	return ActiveSound->bShouldStopDueToMaxConcurrency ? -100.0f : GetActualVolume() * VolumeWeightedPriorityScale;
+	return ActiveSound->bShouldStopDueToMaxConcurrency ? -100.0f : GetActualVolume() * Priority;
 }
 
 bool FWaveInstance::IsStreaming() const
diff --git a/Engine/Source/Runtime/Engine/Private/BlueprintGeneratedClass.cpp b/Engine/Source/Runtime/Engine/Private/BlueprintGeneratedClass.cpp
--- a/Engine/Source/Runtime/Engine/Private/BlueprintGeneratedClass.cpp
+++ b/Engine/Source/Runtime/Engine/Private/BlueprintGeneratedClass.cpp
@@ -76,7 +76,7 @@ void UBlueprintGeneratedClass::PostLoad()
 
 	if (const UPackage* Package = GetOutermost())
 	{
-		if (Package->PackageFlags & PKG_ForDiffing)
+		if (Package->HasAnyPackageFlags(PKG_ForDiffing))
 		{
 			ClassFlags |= CLASS_Deprecated;
 		}
diff --git a/Engine/Source/Runtime/Engine/Private/Camera/CameraComponent.cpp b/Engine/Source/Runtime/Engine/Private/Camera/CameraComponent.cpp
--- a/Engine/Source/Runtime/Engine/Private/Camera/CameraComponent.cpp
+++ b/Engine/Source/Runtime/Engine/Private/Camera/CameraComponent.cpp
@@ -43,6 +43,16 @@ UCameraComponent::UCameraComponent(const FObjectInitializer& ObjectInitializer)
 }
 
 #if WITH_EDITORONLY_DATA
+
+void UCameraComponent::AddReferencedObjects(UObject* InThis, FReferenceCollector& Collector)
+{
+	UCameraComponent* This = CastChecked<UCameraComponent>(InThis);
+	Collector.AddReferencedObject(This->ProxyMeshComponent);
+	Collector.AddReferencedObject(This->DrawFrustum);
+
+	Super::AddReferencedObjects(InThis, Collector);
+}
+
 void UCameraComponent::OnComponentDestroyed()
 {
 	Super::OnComponentDestroyed();
@@ -61,29 +71,32 @@ void UCameraComponent::OnComponentDestroyed()
 void UCameraComponent::OnRegister()
 {
 #if WITH_EDITORONLY_DATA
-	if (ProxyMeshComponent == NULL)
+	if (AActor* MyOwner = GetOwner())
 	{
-		ProxyMeshComponent = NewObject<UStaticMeshComponent>(GetOuter(), NAME_None, RF_Transactional);
-		ProxyMeshComponent->AttachTo(this);
-		ProxyMeshComponent->AlwaysLoadOnClient = false;
-		ProxyMeshComponent->AlwaysLoadOnServer = false;
-		ProxyMeshComponent->StaticMesh = CameraMesh;
-		ProxyMeshComponent->SetCollisionProfileName(UCollisionProfile::NoCollision_ProfileName);
-		ProxyMeshComponent->bHiddenInGame = true;
-		ProxyMeshComponent->CastShadow = false;
-		ProxyMeshComponent->PostPhysicsComponentTick.bCanEverTick = false;
-		ProxyMeshComponent->CreationMethod = CreationMethod;
-		ProxyMeshComponent->RegisterComponentWithWorld(GetWorld());
-	}
+		if (ProxyMeshComponent == nullptr)
+		{
+			ProxyMeshComponent = NewObject<UStaticMeshComponent>(MyOwner, NAME_None, RF_Transactional | RF_TextExportTransient);
+			ProxyMeshComponent->AttachTo(this);
+			ProxyMeshComponent->AlwaysLoadOnClient = false;
+			ProxyMeshComponent->AlwaysLoadOnServer = false;
+			ProxyMeshComponent->StaticMesh = CameraMesh;
+			ProxyMeshComponent->SetCollisionProfileName(UCollisionProfile::NoCollision_ProfileName);
+			ProxyMeshComponent->bHiddenInGame = true;
+			ProxyMeshComponent->CastShadow = false;
+			ProxyMeshComponent->PostPhysicsComponentTick.bCanEverTick = false;
+			ProxyMeshComponent->CreationMethod = CreationMethod;
+			ProxyMeshComponent->RegisterComponentWithWorld(GetWorld());
+		}
 
-	if (DrawFrustum == NULL)
-	{
-		DrawFrustum = NewObject<UDrawFrustumComponent>(GetOuter(), NAME_None, RF_Transactional);
-		DrawFrustum->AttachTo(this);
-		DrawFrustum->AlwaysLoadOnClient = false;
-		DrawFrustum->AlwaysLoadOnServer = false;
-		DrawFrustum->CreationMethod = CreationMethod;
-		DrawFrustum->RegisterComponentWithWorld(GetWorld());
+		if (DrawFrustum == nullptr)
+		{
+			DrawFrustum = NewObject<UDrawFrustumComponent>(MyOwner, NAME_None, RF_Transactional | RF_TextExportTransient);
+			DrawFrustum->AttachTo(this);
+			DrawFrustum->AlwaysLoadOnClient = false;
+			DrawFrustum->AlwaysLoadOnServer = false;
+			DrawFrustum->CreationMethod = CreationMethod;
+			DrawFrustum->RegisterComponentWithWorld(GetWorld());
+		}
 	}
 
 	RefreshVisualRepresentation();
@@ -95,31 +108,6 @@ void UCameraComponent::OnRegister()
 	SetDeprecatedControllerViewRotation(*this, bUsePawnControlRotation);
 }
 
-void UCameraComponent::OnUnregister()
-{
-	Super::OnUnregister();
-
-#if WITH_EDITORONLY_DATA
-	// have to removed the sub-components that we added in OnRegister (for 
-	// reinstancing, where we CopyPropertiesForUnrelatedObjects()... don't want
-	// these copied since we'll generate them on the next OnRegister)
-	if (ProxyMeshComponent != NULL)
-	{
-		ProxyMeshComponent->DetachFromParent();
-		ProxyMeshComponent->DestroyComponent();
-		ProxyMeshComponent = NULL;
-	}
-
-	if (DrawFrustum != NULL)
-	{
-		DrawFrustum->DetachFromParent();
-		DrawFrustum->DestroyComponent();
-		DrawFrustum = NULL;
-	}
-#endif
-}
-
-
 void UCameraComponent::PostLoad()
 {
 	Super::PostLoad();
@@ -151,7 +139,7 @@ void UCameraComponent::PostLoad()
  }
 void UCameraComponent::RefreshVisualRepresentation()
 {
-	if (DrawFrustum != NULL)
+	if (DrawFrustum != nullptr)
 	{
 		const float FrustumDrawDistance = 1000.0f;
 		if (ProjectionMode == ECameraProjectionMode::Perspective)
@@ -173,7 +161,7 @@ void UCameraComponent::RefreshVisualRepresentation()
 
 void UCameraComponent::OverrideFrustumColor(FColor OverrideColor)
 {
-	if (DrawFrustum != NULL)
+	if (DrawFrustum != nullptr)
 	{
 		DrawFrustum->FrustumColor = OverrideColor;
 	}
@@ -181,7 +169,7 @@ void UCameraComponent::OverrideFrustumColor(FColor OverrideColor)
 
 void UCameraComponent::RestoreFrustumColor()
 {
-	if (DrawFrustum != NULL)
+	if (DrawFrustum != nullptr)
 	{
 		//@TODO: 
 		const FColor DefaultFrustumColor(255, 0, 255, 255);
@@ -192,22 +180,20 @@ void UCameraComponent::RestoreFrustumColor()
 }
 #endif	// WITH_EDITORONLY_DATA
 
-#if WITH_EDITOR
+#if WITH_EDITORONLY_DATA
 void UCameraComponent::PostEditChangeProperty(FPropertyChangedEvent& PropertyChangedEvent)
 {
 	Super::PostEditChangeProperty(PropertyChangedEvent);
 
-#if WITH_EDITORONLY_DATA
 	RefreshVisualRepresentation();
-#endif	// WITH_EDITORONLY_DATA
 }
-#endif	// WITH_EDITOR
+#endif	// WITH_EDITORONLY_DATA
 
 void UCameraComponent::Serialize(FArchive& Ar)
 {
 	Super::Serialize(Ar);
 
-	if(Ar.IsLoading())
+	if (Ar.IsLoading())
 	{
 		PostProcessSettings.OnAfterLoad();
 	}
diff --git a/Engine/Source/Runtime/Engine/Private/Collision/WorldCollisionAsync.cpp b/Engine/Source/Runtime/Engine/Private/Collision/WorldCollisionAsync.cpp
--- a/Engine/Source/Runtime/Engine/Private/Collision/WorldCollisionAsync.cpp
+++ b/Engine/Source/Runtime/Engine/Private/Collision/WorldCollisionAsync.cpp
@@ -38,9 +38,9 @@ namespace
 	template <>                   FORCEINLINE TArray<TUniquePtr<TTraceThreadData<FOverlapDatum>>>& GetTraceContainer<FOverlapDatum>(AsyncTraceData& DataBuffer) { return DataBuffer.OverlapData; }
 
 	// Helper functions to return the right named member trace index based on a datum type
-	template <typename DatumType> FORCEINLINE int32& GetTraceIndex               (FWorldAsyncTraceState& State);
-	template <>                   FORCEINLINE int32& GetTraceIndex<FTraceDatum>  (FWorldAsyncTraceState& State) { return State.NextAvailableTraceIndex;   }
-	template <>                   FORCEINLINE int32& GetTraceIndex<FOverlapDatum>(FWorldAsyncTraceState& State) { return State.NextAvailableOverlapIndex; }
+	template <typename DatumType> FORCEINLINE int32& GetTraceIndex(AsyncTraceData& DataBuffer);
+	template <>                   FORCEINLINE int32& GetTraceIndex<FTraceDatum>(AsyncTraceData& DataBuffer) { return DataBuffer.NumQueuedTraceData; }
+	template <>                   FORCEINLINE int32& GetTraceIndex<FOverlapDatum>(AsyncTraceData& DataBuffer) { return DataBuffer.NumQueuedOverlapData; }
 
 	/** For referencing a thread data buffer and a datum within it */
 	struct FBufferIndexPair
@@ -230,7 +230,9 @@ namespace
 	template <typename DatumType>
 	void ExecuteAsyncTraceIfAvailable(FWorldAsyncTraceState& State, bool bExecuteAll)
 	{
-		FBufferIndexPair Next(GetTraceIndex<DatumType>(State));
+		auto& DataBuffer = State.GetBufferForCurrentFrame();
+
+		FBufferIndexPair Next(GetTraceIndex<DatumType>(DataBuffer));
 
 		// when Next.Index == 0, and Next.Block > 0 , that means next one will be in the next buffer
 		// but that case we'd like to send to thread
@@ -245,7 +247,6 @@ namespace
 			return;
 		}
 
-		auto& DataBuffer = State.GetBufferForCurrentFrame();
 		auto* Datum      = GetTraceContainer<DatumType>(DataBuffer)[Next.Block]->Buffer;
 		#if RUN_ASYNC_TRACE
 			DataBuffer.AsyncTraceCompletionEvent.Emplace(TGraphTask<FAsyncTraceTask>::CreateTask(NULL, ENamedThreads::GameThread).ConstructAndDispatchWhenReady(Datum, Next.Index));
@@ -264,7 +265,7 @@ namespace
 		check(DataBuffer.bAsyncAllowed);
 
 		auto&  TraceData  = GetTraceContainer<DatumType>(DataBuffer);
-		int32& TraceIndex = GetTraceIndex    <DatumType>(State);
+		int32& TraceIndex = GetTraceIndex<DatumType>(DataBuffer);
 
 		// we calculate index as continuous, not as each chunk, but continuous way 
 		int32 LastAvailableIndex = TraceData.Num() * ASYNC_TRACE_BUFFER_SIZE;
@@ -289,8 +290,6 @@ namespace
 
 FWorldAsyncTraceState::FWorldAsyncTraceState()
 	: CurrentFrame             (0)
-	, NextAvailableTraceIndex  (0)
-	, NextAvailableOverlapIndex(0)
 {
 	// initial buffer is open for business
 	DataBuffer[CurrentFrame].bAsyncAllowed = true;
@@ -408,12 +407,13 @@ void UWorld::ResetAsyncTrace()
 	WaitForAllAsyncTraceTasks();
 
 	// do run delegates before starting next round
-	for (int32 Idx = 0; Idx != DataBufferExecuted.TraceData.Num(); ++Idx)
+	for (int32 Idx = 0; Idx != DataBufferExecuted.NumQueuedTraceData; ++Idx)
 	{
 		auto& TraceData = FBufferIndexPair(Idx).DatumLookupChecked(DataBufferExecuted.TraceData);
 		TraceData.Delegate.ExecuteIfBound(FTraceHandle(TraceData.FrameNumber, Idx), TraceData);
 	}
-	for (int32 Idx = 0; Idx != DataBufferExecuted.OverlapData.Num(); ++Idx)
+
+	for (int32 Idx = 0; Idx != DataBufferExecuted.NumQueuedOverlapData; ++Idx)
 	{
 		auto& TraceData = FBufferIndexPair(Idx).DatumLookupChecked(DataBufferExecuted.OverlapData);
 		TraceData.Delegate.ExecuteIfBound(FTraceHandle(TraceData.FrameNumber, Idx), TraceData);
@@ -433,8 +433,10 @@ void UWorld::FinishAsyncTrace()
 	++AsyncTraceState.CurrentFrame;
 
 	// set up new buffer to accept trace requests
-	AsyncTraceState.GetBufferForCurrentFrame().bAsyncAllowed = true;
-	AsyncTraceState.NextAvailableTraceIndex = 0;
-	AsyncTraceState.NextAvailableOverlapIndex = 0;
+	AsyncTraceData& NewAsyncBuffer = AsyncTraceState.GetBufferForCurrentFrame();
+	NewAsyncBuffer.bAsyncAllowed = true;
+	NewAsyncBuffer.NumQueuedTraceData = 0;
+	NewAsyncBuffer.NumQueuedOverlapData = 0;
 
 }
+
diff --git a/Engine/Source/Runtime/Engine/Private/Components/ActorComponent.cpp b/Engine/Source/Runtime/Engine/Private/Components/ActorComponent.cpp
--- a/Engine/Source/Runtime/Engine/Private/Components/ActorComponent.cpp
+++ b/Engine/Source/Runtime/Engine/Private/Components/ActorComponent.cpp
@@ -868,10 +868,10 @@ void UActorComponent::RegisterComponentWithWorld(UWorld* InWorld)
 		}
 	}
 
-	if (MyOwner && MyOwner->HasActorBegunPlay() && !bHasBegunPlay)
+	if (MyOwner && (MyOwner->HasActorBegunPlay() || MyOwner->IsActorBeginningPlay()))
 	{
 		RegisterAllComponentTickFunctions(true);
-		if (bWantsBeginPlay)
+		if (bWantsBeginPlay && !bHasBegunPlay)
 		{
 			BeginPlay();
 		}
diff --git a/Engine/Source/Runtime/Engine/Private/Components/AudioComponent.cpp b/Engine/Source/Runtime/Engine/Private/Components/AudioComponent.cpp
--- a/Engine/Source/Runtime/Engine/Private/Components/AudioComponent.cpp
+++ b/Engine/Source/Runtime/Engine/Private/Components/AudioComponent.cpp
@@ -22,7 +22,8 @@ UAudioComponent::UAudioComponent(const FObjectInitializer& ObjectInitializer)
 	bVisualizeComponent = true;
 #endif
 	VolumeMultiplier = 1.f;
-	VolumeWeightedPriorityScale = 1.f;
+	bOverridePriority = false;
+	Priority = 1.f;
 	PitchMultiplier = 1.f;
 	VolumeModulationMin = 1.f;
 	VolumeModulationMax = 1.f;
@@ -151,7 +152,15 @@ void UAudioComponent::PlayInternal(const float StartTime, const float FadeInDura
 			NewActiveSound.ConcurrencySettings = ConcurrencySettings;
 
 			NewActiveSound.VolumeMultiplier = (VolumeModulationMax + ((VolumeModulationMin - VolumeModulationMax) * FMath::SRand())) * VolumeMultiplier;
-			NewActiveSound.VolumeWeightedPriorityScale = VolumeWeightedPriorityScale;
+			// The priority used for the active sound is the audio component's priority scaled with the sound's priority
+			if (bOverridePriority)
+			{
+				NewActiveSound.Priority = Priority;
+			}
+			else
+			{
+				NewActiveSound.Priority = Sound->Priority;
+			}
 			NewActiveSound.PitchMultiplier = (PitchModulationMax + ((PitchModulationMin - PitchModulationMax) * FMath::SRand())) * PitchMultiplier;
 			NewActiveSound.HighFrequencyGainMultiplier = HighFrequencyGainMultiplier;
 
diff --git a/Engine/Source/Runtime/Engine/Private/Components/NavMovementComponent.cpp b/Engine/Source/Runtime/Engine/Private/Components/NavMovementComponent.cpp
--- a/Engine/Source/Runtime/Engine/Private/Components/NavMovementComponent.cpp
+++ b/Engine/Source/Runtime/Engine/Private/Components/NavMovementComponent.cpp
@@ -18,7 +18,7 @@ void FNavAgentProperties::UpdateWithCollisionComponent(UShapeComponent* Collisio
 
 bool FNavAgentProperties::IsNavDataMatching(const FNavAgentProperties& Other) const
 {
-	return (PreferredNavData == nullptr || Other.PreferredNavData == nullptr || PreferredNavData == Other.PreferredNavData);
+	return (PreferredNavData == Other.PreferredNavData || PreferredNavData == nullptr || Other.PreferredNavData == nullptr);
 }
 
 //----------------------------------------------------------------------//
diff --git a/Engine/Source/Runtime/Engine/Private/Components/PoseableMeshComponent.cpp b/Engine/Source/Runtime/Engine/Private/Components/PoseableMeshComponent.cpp
--- a/Engine/Source/Runtime/Engine/Private/Components/PoseableMeshComponent.cpp
+++ b/Engine/Source/Runtime/Engine/Private/Components/PoseableMeshComponent.cpp
@@ -223,7 +223,7 @@ FTransform UPoseableMeshComponent::GetBoneTransformByName(FName BoneName, EBoneS
 		{
 			if (UAnimInstance* AnimInstance = SMC->GetAnimInstance())
 			{
-				return GetBoneTransformByNameHelper(BoneName, BoneSpace, AnimInstance->RequiredBones, SMC);
+				return GetBoneTransformByNameHelper(BoneName, BoneSpace, AnimInstance->GetRequiredBones(), SMC);
 			}
 			FString Message = FString::Printf(TEXT("Cannot return valid bone transform. Master Pose Component has no anim instance"));
 			FFrame::KismetExecutionMessage(*Message, ELogVerbosity::Warning);
diff --git a/Engine/Source/Runtime/Engine/Private/Components/PrimitiveComponent.cpp b/Engine/Source/Runtime/Engine/Private/Components/PrimitiveComponent.cpp
--- a/Engine/Source/Runtime/Engine/Private/Components/PrimitiveComponent.cpp
+++ b/Engine/Source/Runtime/Engine/Private/Components/PrimitiveComponent.cpp
@@ -1104,6 +1104,13 @@ ECollisionChannel UPrimitiveComponent::GetCollisionObjectType() const
 	return ECollisionChannel(BodyInstance.GetObjectType());
 }
 
+void UPrimitiveComponent::SetBoundsScale(float NewBoundsScale)
+{
+	BoundsScale = NewBoundsScale;
+	UpdateBounds();
+	MarkRenderTransformDirty();
+}
+
 UMaterialInterface* UPrimitiveComponent::GetMaterial(int32 Index) const
 {
 	return NULL;
@@ -1415,7 +1422,7 @@ FCollisionShape UPrimitiveComponent::GetCollisionShape(float Inflation) const
 		Extent = Extent.ComponentMax(FVector::ZeroVector);
 	}
 
-	return FCollisionShape::MakeBox(Bounds.BoxExtent + Inflation);
+	return FCollisionShape::MakeBox(Extent);
 }
 
 
diff --git a/Engine/Source/Runtime/Engine/Private/Components/SceneComponent.cpp b/Engine/Source/Runtime/Engine/Private/Components/SceneComponent.cpp
--- a/Engine/Source/Runtime/Engine/Private/Components/SceneComponent.cpp
+++ b/Engine/Source/Runtime/Engine/Private/Components/SceneComponent.cpp
@@ -59,14 +59,14 @@ USceneComponent::USceneComponent(const FObjectInitializer& ObjectInitializer /*=
 	bAutoActivate = false;
 }
 
+#if WITH_EDITORONLY_DATA
 void USceneComponent::AddReferencedObjects(UObject* InThis, FReferenceCollector& Collector)
 {
-#if WITH_EDITORONLY_DATA
 	USceneComponent* This = CastChecked<USceneComponent>(InThis);
 	Collector.AddReferencedObject(This->SpriteComponent);
-#endif
 	Super::AddReferencedObjects(InThis, Collector);
 }
+#endif
 
 #if WITH_EDITOR
 
@@ -96,7 +96,7 @@ DECLARE_DELEGATE_RetVal_OneParam(bool, FMobilityQueryDelegate, EComponentMobilit
  */
 static int32 SetDecedentMobility(USceneComponent const* SceneComponentObject, EComponentMobility::Type NewMobilityType, FMobilityQueryDelegate ShouldOverrideMobility = FMobilityQueryDelegate())
 {
-	if (!ensure(SceneComponentObject != NULL))
+	if (!ensure(SceneComponentObject != nullptr))
 	{
 		return 0;
 	}
@@ -104,13 +104,13 @@ static int32 SetDecedentMobility(USceneComponent const* SceneComponentObject, EC
 	TArray<USceneComponent*> AttachedChildren = SceneComponentObject->AttachChildren;
 	// gather children for component templates
 	USCS_Node* SCSNode = ComponentUtils::FindCorrespondingSCSNode(SceneComponentObject);
-	if (SCSNode != NULL)
+	if (SCSNode != nullptr)
 	{
 		// gather children from the SCSNode
 		for (USCS_Node* SCSChild : SCSNode->GetChildNodes())
 		{
 			USceneComponent* ChildSceneComponent = Cast<USceneComponent>(SCSChild->ComponentTemplate);
-			if (ChildSceneComponent != NULL)
+			if (ChildSceneComponent != nullptr)
 			{
 				AttachedChildren.Add(ChildSceneComponent);
 			}
@@ -160,7 +160,7 @@ static int32 SetDecedentMobility(USceneComponent const* SceneComponentObject, EC
  */
 static int32 SetAncestorMobility(USceneComponent const* SceneComponentObject, EComponentMobility::Type NewMobilityType, FMobilityQueryDelegate ShouldOverrideMobility = FMobilityQueryDelegate())
 {
-	if (!ensure(SceneComponentObject != NULL))
+	if (!ensure(SceneComponentObject != nullptr))
 	{
 		return 0;
 	}
@@ -305,7 +305,7 @@ void USceneComponent::PostEditChangeChainProperty(FPropertyChangedChainEvent& Pr
 
 FTransform USceneComponent::CalcNewComponentToWorld_GeneralCase(const FTransform& NewRelativeTransform, const USceneComponent* Parent, FName SocketName) const
 {
-	if (Parent != NULL)
+	if (Parent != nullptr)
 	{
 		const FTransform ParentToWorld = Parent->GetSocketTransform(SocketName);
 		FTransform NewCompToWorld = NewRelativeTransform * ParentToWorld;
@@ -602,7 +602,7 @@ void USceneComponent::DestroyComponent(bool bPromoteChildren/*= false*/)
 	if (bPromoteChildren)
 	{
 		AActor* Owner = GetOwner();
-		if (Owner != NULL)
+		if (Owner != nullptr)
 		{
 			Owner->Modify();
 			USceneComponent* ChildToPromote = nullptr;
@@ -623,10 +623,10 @@ void USceneComponent::DestroyComponent(bool bPromoteChildren/*= false*/)
 				{
 					// Didn't find a suitable component to promote so create a new default component
 
-					Rename(NULL, GetOuter(), REN_DoNotDirty | REN_DontCreateRedirectors);
+					Rename(nullptr, GetOuter(), REN_DoNotDirty | REN_DontCreateRedirectors);
 
 					// Construct a new default root component
-					auto NewRootComponent = NewObject<USceneComponent>(Owner, USceneComponent::GetDefaultSceneRootVariableName(), RF_Transactional);
+					USceneComponent* NewRootComponent = NewObject<USceneComponent>(Owner, USceneComponent::GetDefaultSceneRootVariableName(), RF_Transactional);
 					NewRootComponent->Mobility = Mobility;
 					NewRootComponent->SetWorldLocationAndRotation(GetComponentLocation(), GetComponentRotation());
 #if WITH_EDITORONLY_DATA
@@ -720,10 +720,10 @@ void USceneComponent::OnComponentDestroyed()
 
 	ScopedMovementStack.Reset();
 
-	// Detach children before destroying
-	for (int32 Index = AttachChildren.Num()-1; Index >= 0; --Index)
+	int32 ChildCount = AttachChildren.Num();
+	while (ChildCount > 0)
 	{
-		if (USceneComponent* Child = AttachChildren[Index])
+		if (USceneComponent* Child = AttachChildren.Last())
 		{
 			if (AttachParent)
 			{
@@ -733,7 +733,13 @@ void USceneComponent::OnComponentDestroyed()
 			{
 				Child->DetachFromParent();
 			}
+			checkf(ChildCount > AttachChildren.Num(), TEXT("AttachChildren count increased while detaching '%s', likely caused by OnAttachmentChanged introducing new children, which could lead to an infinite loop."), *Child->GetName());
 		}
+		else
+		{
+			AttachChildren.Pop(false);
+		}
+		ChildCount = AttachChildren.Num();
 	}
 
 	// Ensure we are detached before destroying
@@ -763,7 +769,7 @@ void USceneComponent::UpdateBounds()
 
 	// if use parent bound if attach parent exists, and the flag is set
 	// since parents tick first before child, this should work correctly
-	if ( bUseAttachParentBound && AttachParent != NULL )
+	if ( bUseAttachParentBound && AttachParent != nullptr )
 	{
 		Bounds = AttachParent->Bounds;
 	}
@@ -784,7 +790,7 @@ void USceneComponent::UpdateBounds()
 
 #if WITH_EDITOR
 	// If bounds have changed (in editor), trigger data rebuild
-	if ( IsRegistered() && (World != NULL) && !World->IsGameWorld() &&
+	if ( IsRegistered() && (World != nullptr) && !World->IsGameWorld() &&
 		(OriginalBounds.Origin.Equals(Bounds.Origin) == false || OriginalBounds.BoxExtent.Equals(Bounds.BoxExtent) == false) )
 	{
 		GEngine->TriggerStreamingDataRebuild();
@@ -959,7 +965,7 @@ void USceneComponent::SetWorldLocation(FVector NewLocation, bool bSweep, FHitRes
 	FVector NewRelLocation = NewLocation;
 
 	// If attached to something, transform into local space
-	if (AttachParent != NULL && !bAbsoluteLocation)
+	if (AttachParent != nullptr && !bAbsoluteLocation)
 	{
 		FTransform ParentToWorld = AttachParent->GetSocketTransform(AttachSocketName);
 		NewRelLocation = ParentToWorld.InverseTransformPosition(NewLocation);
@@ -973,7 +979,7 @@ void USceneComponent::SetWorldRotation(const FQuat& NewRotation, bool bSweep, FH
 	FQuat NewRelRotation = NewRotation;
 
 	// If already attached to something, transform into local space
-	if (AttachParent != NULL && !bAbsoluteRotation)
+	if (AttachParent != nullptr && !bAbsoluteRotation)
 	{
 		const FQuat ParentToWorldQuat = AttachParent->GetSocketQuaternion(AttachSocketName);
 		// Quat multiplication works reverse way, make sure you do Parent(-1) * World = Local, not World*Parent(-) = Local (the way matrix does)
@@ -1003,7 +1009,7 @@ void USceneComponent::SetWorldScale3D(FVector NewScale)
 	FVector NewRelScale = NewScale;
 
 	// If attached to something, transform into local space
-	if(AttachParent != NULL && !bAbsoluteScale)
+	if(AttachParent != nullptr && !bAbsoluteScale)
 	{
 		FTransform ParentToWorld = AttachParent->GetSocketTransform(AttachSocketName);
 		NewRelScale = NewScale * ParentToWorld.GetSafeScaleReciprocal(ParentToWorld.GetScale3D());
@@ -1019,7 +1025,7 @@ void USceneComponent::SetWorldTransform(const FTransform& NewTransform, bool bSw
 	FVector NewLocation = NewTransform.GetTranslation();
 	FVector NewScale = NewTransform.GetScale3D();
 
-	if (AttachParent != NULL)
+	if (AttachParent != nullptr)
 	{
 		FTransform ParentToWorld = AttachParent->GetSocketTransform(AttachSocketName);
 		FTransform RelativeTM = NewTransform.GetRelativeTransform(ParentToWorld);
@@ -1060,7 +1066,7 @@ void USceneComponent::SetWorldLocationAndRotation(FVector NewLocation, const FQu
 {
 	// If attached to something, transform into local space
 	FQuat NewFinalRotation = NewRotation;
-	if (AttachParent != NULL)
+	if (AttachParent != nullptr)
 	{
 		FTransform ParentToWorld = AttachParent->GetSocketTransform(AttachSocketName);
 
@@ -1083,7 +1089,7 @@ void USceneComponent::SetWorldLocationAndRotation(FVector NewLocation, const FQu
 void USceneComponent::SetWorldLocationAndRotationNoPhysics(const FVector& NewLocation, const FRotator& NewRotation)
 {
 	// If attached to something, transform into local space
-	if (AttachParent != NULL)
+	if (AttachParent != nullptr)
 	{
 		const FTransform ParentToWorld = AttachParent->GetSocketTransform(AttachSocketName);
 
@@ -1165,7 +1171,7 @@ void USceneComponent::GetParentComponents(TArray<class USceneComponent*>& Parent
 	Parents.Empty();
 
 	USceneComponent* ParentIterator = AttachParent;
-	while (ParentIterator != NULL)
+	while (ParentIterator != nullptr)
 	{
 		Parents.Add(ParentIterator);
 		ParentIterator = ParentIterator->AttachParent;
@@ -1182,12 +1188,12 @@ USceneComponent* USceneComponent::GetChildComponent(int32 ChildIndex) const
 	if (ChildIndex < 0)
 	{
 		UE_LOG(LogBlueprint, Log, TEXT("SceneComponent::GetChild called with a negative ChildIndex: %d"), ChildIndex);
-		return NULL;
+		return nullptr;
 	}
 	else if (ChildIndex >= AttachChildren.Num())
 	{
 		UE_LOG(LogBlueprint, Log, TEXT("SceneComponent::GetChild called with an out of range ChildIndex: %d; Number of children is %d."), ChildIndex, AttachChildren.Num());
-		return NULL;
+		return nullptr;
 	}
 
 	return AttachChildren[ChildIndex];
@@ -1229,7 +1235,7 @@ void USceneComponent::K2_AttachTo(class USceneComponent* InParent, FName InSocke
 
 bool USceneComponent::AttachTo(class USceneComponent* Parent, FName InSocketName, EAttachLocation::Type AttachType /*= EAttachLocation::KeepRelativeOffset */, bool bWeldSimulatedBodies /*= false*/)
 {
-	if(Parent != NULL)
+	if(Parent != nullptr)
 	{
 		if (Parent == AttachParent && InSocketName == AttachSocketName && Parent->AttachChildren.Contains(this))
 		{
@@ -1474,7 +1480,7 @@ void USceneComponent::SnapTo(class USceneComponent* Parent, FName InSocketName)
 
 void USceneComponent::DetachFromParent(bool bMaintainWorldPosition, bool bCallModify)
 {
-	if(AttachParent != NULL)
+	if(AttachParent != nullptr)
 	{
 		AActor* Owner = GetOwner();
 
@@ -1506,7 +1512,7 @@ void USceneComponent::DetachFromParent(bool bMaintainWorldPosition, bool bCallMo
 			}
 		}
 #endif
-		AttachParent = NULL;
+		AttachParent = nullptr;
 		AttachSocketName = NAME_None;
 
 		OnAttachmentChanged();
@@ -1540,14 +1546,14 @@ USceneComponent* USceneComponent::GetAttachmentRoot() const
 AActor* USceneComponent::GetAttachmentRootActor() const
 {
 	const USceneComponent* const AttachmentRootComponent = GetAttachmentRoot();
-	return AttachmentRootComponent ? AttachmentRootComponent->GetOwner() : NULL;
+	return AttachmentRootComponent ? AttachmentRootComponent->GetOwner() : nullptr;
 }
 
 bool USceneComponent::IsAttachedTo(class USceneComponent* TestComp) const
 {
-	if(TestComp != NULL)
+	if(TestComp != nullptr)
 	{
-		for( const USceneComponent* Comp=this->AttachParent; Comp!=NULL; Comp=Comp->AttachParent )
+		for( const USceneComponent* Comp=this->AttachParent; Comp!=nullptr; Comp=Comp->AttachParent )
 		{
 			if( TestComp == Comp )
 			{
@@ -1644,7 +1650,7 @@ void USceneComponent::UpdateChildTransforms(bool bSkipPhysicsMove, ETeleportType
 	for(int32 i=0; i<AttachChildren.Num(); i++)
 	{
 		USceneComponent* ChildComp = AttachChildren[i];
-		if (ChildComp != NULL)
+		if (ChildComp != nullptr)
 		{
 			// Don't update the child if it uses a completely absolute (world-relative) scheme, unless it has never been updated.
 			if (!(ChildComp->bAbsoluteLocation && ChildComp->bAbsoluteRotation && ChildComp->bAbsoluteScale) || !ChildComp->bWorldToComponentUpdated)
@@ -1796,7 +1802,7 @@ ECollisionChannel USceneComponent::GetCollisionObjectType() const
 ECollisionResponse USceneComponent::GetCollisionResponseToComponent(class USceneComponent* OtherComponent) const
 {
 	// Ignore if no component, or either component has no collision
-	if(OtherComponent == NULL || GetCollisionEnabled() == ECollisionEnabled::NoCollision || OtherComponent->GetCollisionEnabled() == ECollisionEnabled::NoCollision)
+	if(OtherComponent == nullptr || GetCollisionEnabled() == ECollisionEnabled::NoCollision || OtherComponent->GetCollisionEnabled() == ECollisionEnabled::NoCollision)
 	{
 		return ECR_Ignore;
 	}
@@ -1862,7 +1868,7 @@ APhysicsVolume* USceneComponent::GetPhysicsVolume() const
 		return MyWorld->GetDefaultPhysicsVolume();
 	}
 
-	return NULL;
+	return nullptr;
 }
 
 void USceneComponent::UpdatePhysicsVolume( bool bTriggerNotifiers )
@@ -1995,7 +2001,7 @@ bool USceneComponent::InternalSetWorldLocationAndRotation(FVector NewLocation, c
 #endif
 
 	// If attached to something, transform into local space
-	if (AttachParent != NULL)
+	if (AttachParent != nullptr)
 	{
 		FTransform const ParentToWorld = AttachParent->GetSocketTransform(AttachSocketName);
 
@@ -2055,7 +2061,7 @@ void USceneComponent::UpdateOverlaps(TArray<FOverlapInfo> const* PendingOverlaps
 		if (AttachChildren[ChildIdx])
 		{
 			// Do not pass on OverlapsAtEndLocation, it only applied to this component.
-			AttachChildren[ChildIdx]->UpdateOverlaps(NULL, bDoNotifies);
+			AttachChildren[ChildIdx]->UpdateOverlaps(nullptr, bDoNotifies);
 		}
 	}
 
@@ -2329,7 +2335,7 @@ void USceneComponent::ApplyWorldOffset(const FVector& InOffset, bool bWorldShift
 	for(int32 i=0; i<AttachChildren.Num(); i++)
 	{
 		USceneComponent* ChildComp = AttachChildren[i];
-		if(ChildComp != NULL)
+		if(ChildComp != nullptr)
 		{
 			ChildComp->ApplyWorldOffset(InOffset, bWorldShift);
 		}
@@ -2378,7 +2384,7 @@ void USceneComponent::PostNetReceive()
 	Super::PostNetReceive();
 
 	// If we have no attach parent, attach to parent's root component.
-	if (AttachParent == NULL)
+	if (AttachParent == nullptr)
 	{
 		USceneComponent * ParentRoot = GetOwner()->GetRootComponent();
 		if (ParentRoot != this)
@@ -2426,10 +2432,10 @@ void USceneComponent::GetLifetimeReplicatedProps(TArray< FLifetimeProperty > & O
 bool USceneComponent::CanEditChange( const UProperty* Property ) const
 {
 	bool bIsEditable = Super::CanEditChange( Property );
-	if( bIsEditable && Property != NULL )
+	if( bIsEditable && Property != nullptr )
 	{
 		AActor* Owner = GetOwner();
-		if(Owner != NULL)
+		if(Owner != nullptr)
 		{
 			if(Property->GetFName() == TEXT( "RelativeLocation" ) ||
 			   Property->GetFName() == TEXT( "RelativeRotation" ) ||
@@ -2524,7 +2530,7 @@ FScopedMovementUpdate::FScopedMovementUpdate( class USceneComponent* Component,
 	}
 	else
 	{
-		Owner = NULL;
+		Owner = nullptr;
 	}
 }
 
@@ -2535,7 +2541,7 @@ FScopedMovementUpdate::~FScopedMovementUpdate()
 	{
 		Owner->EndScopedMovementUpdate(*this);
 	}
-	Owner = NULL;
+	Owner = nullptr;
 }
 
 
diff --git a/Engine/Source/Runtime/Engine/Private/Components/SkeletalMeshComponent.cpp b/Engine/Source/Runtime/Engine/Private/Components/SkeletalMeshComponent.cpp
--- a/Engine/Source/Runtime/Engine/Private/Components/SkeletalMeshComponent.cpp
+++ b/Engine/Source/Runtime/Engine/Private/Components/SkeletalMeshComponent.cpp
@@ -38,7 +38,8 @@
 #endif
 
 TAutoConsoleVariable<int32> CVarUseParallelAnimationEvaluation(TEXT("a.ParallelAnimEvaluation"), 1, TEXT("If 1, animation evaluation will be run across the task graph system. If 0, evaluation will run purely on the game thread"));
-TAutoConsoleVariable<int32> CVarUseParallelAnimUpdate(TEXT("a.ParallelAnimUpdate"), 0, TEXT("If != 0, then we update animation blend tree, native update, asset players and montages (is possible) on worker threads."));
+TAutoConsoleVariable<int32> CVarUseParallelAnimUpdate(TEXT("a.ParallelAnimUpdate"), 1, TEXT("If != 0, then we update animation blend tree, native update, asset players and montages (is possible) on worker threads."));
+TAutoConsoleVariable<int32> CVarForceUseParallelAnimUpdate(TEXT("a.ForceParallelAnimUpdate"), 1, TEXT("If != 0, then we update animations on worker threads regardless of the setting on the anim blueprint."));
 
 DECLARE_CYCLE_STAT(TEXT("Swap Anim Buffers"), STAT_CompleteAnimSwapBuffers, STATGROUP_Anim);
 DECLARE_CYCLE_STAT_EXTERN(TEXT("Anim Instance Spawn Time"), STAT_AnimSpawnTime, STATGROUP_Anim, );
@@ -121,7 +122,7 @@ USkeletalMeshComponent::USkeletalMeshComponent(const FObjectInitializer& ObjectI
 {
 	bAutoActivate = true;
 	PrimaryComponentTick.bCanEverTick = true;
-	PrimaryComponentTick.TickGroup = CVarUseParallelAnimationEvaluation.GetValueOnGameThread() != 0 && CVarUseParallelAnimUpdate.GetValueOnGameThread() != 0 ? TG_DuringAnimation : TG_PrePhysics;
+	PrimaryComponentTick.TickGroup = TG_PrePhysics;
 	bWantsInitializeComponent = true;
 	GlobalAnimRateScale = 1.0f;
 	bNoSkeletonUpdate = false;
@@ -168,6 +169,9 @@ USkeletalMeshComponent::USkeletalMeshComponent(const FObjectInitializer& ObjectI
 	bHasCustomNavigableGeometry = EHasCustomNavigableGeometry::Yes;
 
 	bTickInEditor = true;
+
+	RootBodyData.BodyIndex = INDEX_NONE;
+	RootBodyData.TransformToRoot = FTransform::Identity;
 }
 
 
@@ -176,7 +180,6 @@ void USkeletalMeshComponent::RegisterComponentTickFunctions(bool bRegister)
 	Super::RegisterComponentTickFunctions(bRegister);
 
 	UpdatePreClothTickRegisteredState();
-	UpdateDuringAnimationTickRegisteredState();
 }
 
 void USkeletalMeshComponent::RegisterPreClothTick(bool bRegister)
@@ -214,20 +217,6 @@ void USkeletalMeshComponent::UpdatePreClothTickRegisteredState()
 	RegisterPreClothTick(bShouldRunClothTick && PrimaryComponentTick.IsTickFunctionRegistered());
 }
 
-void USkeletalMeshComponent::UpdateDuringAnimationTickRegisteredState()
-{
-	if( CVarUseParallelAnimationEvaluation.GetValueOnGameThread() != 0 && 
-		CVarUseParallelAnimUpdate.GetValueOnGameThread() != 0 &&
-		FApp::ShouldUseThreadingForPerformance())
-	{
-		PrimaryComponentTick.TickGroup = TG_DuringAnimation;
-	}
-	else
-	{
-		PrimaryComponentTick.TickGroup = TG_PrePhysics;
-	}
-}
-
 bool USkeletalMeshComponent::NeedToSpawnAnimScriptInstance(bool bForceInit) const
 {
 	if (AnimationMode == EAnimationMode::AnimationBlueprint && (AnimBlueprintGeneratedClass != NULL) && 
@@ -511,30 +500,10 @@ void USkeletalMeshComponent::TickAnimation(float DeltaTime, bool bNeedsValidRoot
 			// AnimNotifies?  Please take a look and fix as we discussed.  Temporary fix:
 			if (AnimScriptInstance != NULL)
 			{
-				{
-					QUICK_SCOPE_CYCLE_COUNTER(STAT_USkeletalMeshComponent_TickAnimation_MorphTargetCurves);
-
-				    // now all tick/trigger/kismet is done
-				    // add MorphTarget Curves from Kismet driven or any other source
-				    // and overwrite if it exists
-				    // Tick always should maintain this list, not Evaluate
-				    for( auto Iter = MorphTargetCurves.CreateConstIterator(); Iter; ++Iter )
-				    {
-					    float *CurveValPtr = AnimScriptInstance->MorphTargetCurves.Find(Iter.Key());
-					    if ( CurveValPtr )
-					    {
-						    // override the value if Kismet request was made
-						    *CurveValPtr = Iter.Value();
-					    }
-					    else
-					    {
-						    AnimScriptInstance->MorphTargetCurves.Add(Iter.Key(), Iter.Value());
-					    }				
-				    }
-				}
+				AnimScriptInstance->UpdateMorphTargetCurves(MorphTargetCurves);
 
 				//Update material parameters
-				UpdateMaterialParameters();
+				AnimScriptInstance->UpdateComponentsMaterialParameters(this);
 			}
 		}
 	}
@@ -542,37 +511,7 @@ void USkeletalMeshComponent::TickAnimation(float DeltaTime, bool bNeedsValidRoot
 
 void USkeletalMeshComponent::UpdateMaterialParameters()
 {
-	if(AnimScriptInstance->MaterialParameterCurves.Num() > 0)
-	{
-		QUICK_SCOPE_CYCLE_COUNTER(STAT_USkeletalMeshComponent_UpdateMaterialParameters);
-
-		for(auto Iter = AnimScriptInstance->MaterialParameterCurves.CreateConstIterator(); Iter; ++Iter)
-		{
-			FName ParameterName = Iter.Key();
-			float ParameterValue = Iter.Value();
-
-			UE_LOG(LogAnimation, Verbose, TEXT("Material Parameter change by Animation (%s : %0.2f)"), *ParameterName.ToString(), ParameterValue);
-			for(int32 MaterialIndex = 0; MaterialIndex < GetNumMaterials(); ++MaterialIndex)
-			{
-				UMaterialInterface* MaterialInterface = GetMaterial(MaterialIndex);
-				if(MaterialInterface)
-				{
-					float TestValue; //not used but needed for GetScalarParameterValue call
-					if(MaterialInterface->GetScalarParameterValue(ParameterName, TestValue))
-					{
-						UMaterialInstanceDynamic* DynamicMaterial = Cast<UMaterialInstanceDynamic>(MaterialInterface);
-						if(!DynamicMaterial) //Is it already a UMaterialInstanceDynamic (ie we used it last tick)
-						{
-							DynamicMaterial = CreateAndSetMaterialInstanceDynamic(MaterialIndex);
-						}
-						DynamicMaterial->SetScalarParameterValue(ParameterName, ParameterValue);
-
-						// we don't break here because we can have multiple materials wanted to be driven by same parameter
-					}
-				}
-			}
-		}
-	}
+	AnimScriptInstance->UpdateComponentsMaterialParameters(this);
 }
 
 bool USkeletalMeshComponent::UpdateLODStatus()
@@ -602,7 +541,7 @@ bool USkeletalMeshComponent::ShouldUpdateTransform(bool bLODHasChanged) const
 #endif
 
 	// If forcing RefPose we can skip updating the skeleton for perf, except if it's using MorphTargets.
-	const bool bSkipBecauseOfRefPose = bForceRefpose && bOldForceRefPose && (MorphTargetCurves.Num() == 0) && ((AnimScriptInstance)? AnimScriptInstance->MorphTargetCurves.Num() == 0 : true);
+	const bool bSkipBecauseOfRefPose = bForceRefpose && bOldForceRefPose && (MorphTargetCurves.Num() == 0) && ((AnimScriptInstance) ? !AnimScriptInstance->HasMorphTargetCurves() : true);
 
 	// LOD changing should always trigger an update.
 	return (bLODHasChanged || (!bNoSkeletonUpdate && !bSkipBecauseOfRefPose && Super::ShouldUpdateTransform(bLODHasChanged)));
@@ -972,44 +911,12 @@ void USkeletalMeshComponent::EvaluateAnimation(const USkeletalMesh* InSkeletalMe
 	}
 
 	// We can only evaluate animation if RequiredBones is properly setup for the right mesh!
-	if( InSkeletalMesh->Skeleton && InAnimInstance
-		&& ensure(bRequiredBonesUpToDate)
-		&& InAnimInstance->RequiredBones.IsValid()
-		&& (InAnimInstance->RequiredBones.GetAsset() == InSkeletalMesh) )
-	{
-		if( !bForceRefpose )
-		{
-			// Create an evaluation context
-			FPoseContext EvaluationContext(InAnimInstance);
-			EvaluationContext.ResetToRefPose();
-			
-			// Run the anim blueprint
-			InAnimInstance->EvaluateAnimation(EvaluationContext);
-			// Move the curves
-			OutCurve.MoveFrom(EvaluationContext.Curve);
-			
-			// can we avoid that copy?
-			if( EvaluationContext.Pose.GetNumBones() > 0 )
-			{
-				// Make sure rotations are normalized to account for accumulation of errors.
-				EvaluationContext.Pose.NormalizeRotations();
-				for (const FCompactPoseBoneIndex BoneIndex : EvaluationContext.Pose.ForEachBoneIndex())
-				{
-					FMeshPoseBoneIndex MeshPoseBoneIndex = EvaluationContext.Pose.GetBoneContainer().MakeMeshPoseIndex(BoneIndex);
-					OutLocalAtoms[MeshPoseBoneIndex.GetInt()] = EvaluationContext.Pose[BoneIndex];
-				}
-			}
-			else
-			{
-				FAnimationRuntime::FillWithRefPose(OutLocalAtoms, InAnimInstance->RequiredBones);
-			}
-		}
-		else
-		{
-			FAnimationRuntime::FillWithRefPose(OutLocalAtoms, InAnimInstance->RequiredBones);
-		}
-
-		OutVertexAnims = UpdateActiveVertexAnims(InSkeletalMesh, InAnimInstance->MorphTargetCurves, InAnimInstance->VertexAnims);
+	if( InSkeletalMesh->Skeleton && 
+		InAnimInstance &&
+		ensure(bRequiredBonesUpToDate) &&
+		InAnimInstance->ParallelCanEvaluate(InSkeletalMesh))
+	{
+		InAnimInstance->ParallelEvaluateAnimation(bForceRefpose, InSkeletalMesh, OutLocalAtoms, OutVertexAnims, OutCurve);
 	}
 	else
 	{
@@ -1019,7 +926,7 @@ void USkeletalMeshComponent::EvaluateAnimation(const USkeletalMesh* InSkeletalMe
 		if ( MorphTargetCurves.Num() > 0 )
 		{
 			TArray<struct FActiveVertexAnim> EmptyVertexAnims;
-			OutVertexAnims = UpdateActiveVertexAnims(InSkeletalMesh, MorphTargetCurves, EmptyVertexAnims);
+			OutVertexAnims = FAnimationRuntime::UpdateActiveVertexAnims(InSkeletalMesh, MorphTargetCurves, EmptyVertexAnims);
 		}
 	}
 
@@ -1037,7 +944,7 @@ void USkeletalMeshComponent::UpdateSlaveComponent()
 
 		if ( MasterSMC->AnimScriptInstance )
 		{
-			ActiveVertexAnims = UpdateActiveVertexAnims(SkeletalMesh, MasterSMC->AnimScriptInstance->MorphTargetCurves, MasterSMC->AnimScriptInstance->VertexAnims);
+			ActiveVertexAnims = MasterSMC->AnimScriptInstance->UpdateActiveVertexAnims(SkeletalMesh);
 		}
 	}
 
@@ -1060,7 +967,7 @@ void USkeletalMeshComponent::PerformAnimationEvaluation(const USkeletalMesh* InS
 	// update anim instance
 	if(AnimEvaluationContext.bDoUpdate)
 	{
-		AnimScriptInstance->ParallelUpdateAnimation();
+		InAnimInstance->ParallelUpdateAnimation();
 	}
 
 	// evaluate pure animations, and fill up LocalAtoms
@@ -1485,13 +1392,6 @@ void USkeletalMeshComponent::SetAnimInstanceClass(class UClass* NewClass)
 
 UAnimInstance* USkeletalMeshComponent::GetAnimInstance() const
 {
-	// check for concurrent access
-#if DO_CHECK
-	if(AnimScriptInstance && PrimaryComponentTick.TickGroup == TG_DuringAnimation)
-	{
-		check(!AnimScriptInstance->IsRunningParallelEvaluation());
-	}
-#endif
 	return AnimScriptInstance;
 }
 
@@ -1757,7 +1657,7 @@ void USkeletalMeshComponent::SetAnimation(UAnimationAsset* NewAnimToPlay)
 	if (SingleNodeInstance)
 	{
 		SingleNodeInstance->SetAnimationAsset(NewAnimToPlay, false);
-		SingleNodeInstance->bPlaying = false;
+		SingleNodeInstance->SetPlaying(false);
 	}
 	else if( AnimScriptInstance != NULL )
 	{
@@ -1772,7 +1672,7 @@ void USkeletalMeshComponent::SetVertexAnimation(UVertexAnimation* NewVertexAnima
 	{
 		SingleNodeInstance->SetVertexAnimation(NewVertexAnimation, false);
 		// when set the asset, we shouldn't automatically play. 
-		SingleNodeInstance->bPlaying = false;
+		SingleNodeInstance->SetPlaying(false);
 	}
 	else if( AnimScriptInstance != NULL )
 	{
@@ -1785,8 +1685,8 @@ void USkeletalMeshComponent::Play(bool bLooping)
 	UAnimSingleNodeInstance* SingleNodeInstance = GetSingleNodeInstance();
 	if (SingleNodeInstance)
 	{
-		SingleNodeInstance->bPlaying = true;
-		SingleNodeInstance->bLooping = bLooping;
+		SingleNodeInstance->SetPlaying(true);
+		SingleNodeInstance->SetLooping(bLooping);
 	}
 	else if( AnimScriptInstance != NULL )
 	{
@@ -1799,7 +1699,7 @@ void USkeletalMeshComponent::Stop()
 	UAnimSingleNodeInstance* SingleNodeInstance = GetSingleNodeInstance();
 	if (SingleNodeInstance)
 	{
-		SingleNodeInstance->bPlaying = false;
+		SingleNodeInstance->SetPlaying(false);
 	}
 	else if( AnimScriptInstance != NULL )
 	{
@@ -1812,7 +1712,7 @@ bool USkeletalMeshComponent::IsPlaying() const
 	UAnimSingleNodeInstance* SingleNodeInstance = GetSingleNodeInstance();
 	if (SingleNodeInstance)
 	{
-		return SingleNodeInstance->bPlaying;
+		return SingleNodeInstance->IsPlaying();
 	}
 	else if( AnimScriptInstance != NULL )
 	{
@@ -1840,7 +1740,7 @@ float USkeletalMeshComponent::GetPosition() const
 	UAnimSingleNodeInstance* SingleNodeInstance = GetSingleNodeInstance();
 	if (SingleNodeInstance)
 	{
-		return SingleNodeInstance->CurrentTime;
+		return SingleNodeInstance->GetCurrentTime();
 	}
 	else if( AnimScriptInstance != NULL )
 	{
@@ -1868,7 +1768,7 @@ float USkeletalMeshComponent::GetPlayRate() const
 	UAnimSingleNodeInstance* SingleNodeInstance = GetSingleNodeInstance();
 	if (SingleNodeInstance)
 	{
-		return SingleNodeInstance->PlayRate;
+		return SingleNodeInstance->GetPlayRate();
 	}
 	else if( AnimScriptInstance != NULL )
 	{
@@ -2176,39 +2076,41 @@ bool USkeletalMeshComponent::IsPlayingRootMotionFromEverything()
 
 void USkeletalMeshComponent::SetRootBodyIndex(int32 InBodyIndex)
 {
-	RootBodyData.BodyIndex = InBodyIndex;
-
-	if(Bodies.IsValidIndex(RootBodyData.BodyIndex) && SkeletalMesh && 
-		Bodies[RootBodyData.BodyIndex]->BodySetup.IsValid() && Bodies[RootBodyData.BodyIndex]->BodySetup.Get()->BoneName != NAME_None)
+	// this is getting called prior to initialization. 
+	// @todo : better fix is to initialize it? overkilling it though. 
+	if (InBodyIndex != INDEX_NONE)
 	{
-		int32 BoneIndex = GetBoneIndex(Bodies[RootBodyData.BodyIndex]->BodySetup->BoneName);
-		// if bone index is valid and not 0, it SHOULD have parnet index
-		if (ensure (BoneIndex != INDEX_NONE))
+		RootBodyData.BodyIndex = InBodyIndex;
+		RootBodyData.TransformToRoot = FTransform::Identity;
+
+		// Only need to do further work if we have any bodies at all (ie physics state is created)
+		if(Bodies.Num() > 0)
 		{
-			int32 ParentIndex = SkeletalMesh->RefSkeleton.GetParentIndex(BoneIndex);
-			if (BoneIndex != 0 && ensure (ParentIndex != INDEX_NONE))
+			if (Bodies.IsValidIndex(RootBodyData.BodyIndex) && SkeletalMesh &&
+				Bodies[RootBodyData.BodyIndex]->BodySetup.IsValid() && Bodies[RootBodyData.BodyIndex]->BodySetup.Get()->BoneName != NAME_None)
 			{
-				const TArray<FTransform>& RefPose = SkeletalMesh->RefSkeleton.GetRefBonePose();
+				int32 BoneIndex = GetBoneIndex(Bodies[RootBodyData.BodyIndex]->BodySetup->BoneName);
+				// if bone index is valid and not 0, it SHOULD have parnet index
+				if (ensure(BoneIndex != INDEX_NONE))
+				{
+					int32 ParentIndex = SkeletalMesh->RefSkeleton.GetParentIndex(BoneIndex);
+					if (BoneIndex != 0 && ensure(ParentIndex != INDEX_NONE))
+					{
+						const TArray<FTransform>& RefPose = SkeletalMesh->RefSkeleton.GetRefBonePose();
 
-				FTransform RelativeTransform = FTransform(SkeletalMesh->RefBasesInvMatrix[BoneIndex]) * RefPose[ParentIndex];
-				// now get offset 
-				RootBodyData.TransformToRoot = RelativeTransform;
+						const TArray<FTransform>& SpaceBases = GetSpaceBases();
+						FTransform RelativeTransform = SpaceBases[BoneIndex].GetRelativeTransformReverse(SpaceBases[ParentIndex]);
+						// now get offset 
+						RootBodyData.TransformToRoot = RelativeTransform;
+					}
+				}
 			}
 			else
 			{
-				RootBodyData.TransformToRoot = FTransform::Identity;
+				//@todo : this needs a bit of investigation. We used to call this function only in Init, but nowi t's called 
+				ensure(false);
 			}
 		}
-		else
-		{
-			RootBodyData.TransformToRoot = FTransform::Identity;
-		}
-	}
-	else
-	{
-		// error - this should not happen
-		ensure(false);
-		RootBodyData.TransformToRoot = FTransform::Identity;
 	}
 }
 
@@ -2216,7 +2118,9 @@ void USkeletalMeshComponent::RefreshActiveVertexAnims()
 {
 	if (SkeletalMesh && AnimScriptInstance)
 	{
-		ActiveVertexAnims = UpdateActiveVertexAnims(SkeletalMesh, AnimScriptInstance->MorphTargetCurves, AnimScriptInstance->VertexAnims);
+		// as this can be called from any worker thread (i.e. from CreateRenderState_Concurrent) we cant currently be doing parallel evaluation
+		check(!IsRunningParallelEvaluation());
+		ActiveVertexAnims = AnimScriptInstance->UpdateActiveVertexAnims(SkeletalMesh);
 	}
 	else
 	{
diff --git a/Engine/Source/Runtime/Engine/Private/Components/SkinnedMeshComponent.cpp b/Engine/Source/Runtime/Engine/Private/Components/SkinnedMeshComponent.cpp
--- a/Engine/Source/Runtime/Engine/Private/Components/SkinnedMeshComponent.cpp
+++ b/Engine/Source/Runtime/Engine/Private/Components/SkinnedMeshComponent.cpp
@@ -1887,73 +1887,6 @@ bool USkinnedMeshComponent::UpdateLODStatus()
 	return bLODChanged;
 }
 
-/** See if an array of ActiveVertexAnims already contains the supplied anim */
-int32 FindVertexAnim(const TArray<FActiveVertexAnim>& ActiveAnims, UVertexAnimBase* Anim)
-{
-	for(int32 i=0; i<ActiveAnims.Num(); i++)
-	{
-		if(ActiveAnims[i].VertAnim == Anim)
-		{
-			return i;
-		}
-	}
-
-	return INDEX_NONE;
-}
-
-TArray<FActiveVertexAnim> USkinnedMeshComponent::UpdateActiveVertexAnims(const USkeletalMesh* InSkeletalMesh, const TMap<FName, float>& MorphCurveAnims, const TArray<FActiveVertexAnim>& ActiveAnims)
-{
-	TArray<struct FActiveVertexAnim> OutVertexAnims;
-
-	// First copy ActiveAnims
-	for(int32 AnimIdx=0; AnimIdx < ActiveAnims.Num(); AnimIdx++)
-	{
-		const FActiveVertexAnim& ActiveAnim = ActiveAnims[AnimIdx];
-		const float ActiveAnimAbsWeight = FMath::Abs(ActiveAnim.Weight);
-
-		// Check it has valid weight, and works on this SkeletalMesh
-		if (	ActiveAnimAbsWeight > MinVertexAnimBlendWeight &&
-			ActiveAnim.VertAnim != NULL &&
-			ActiveAnim.VertAnim->BaseSkelMesh == InSkeletalMesh)
-		{
-			OutVertexAnims.Add(ActiveAnim);
-		}
-		// @TODO Need to check for duplicates here?
-	}
-
-	// Then go over the CurveKeys finding morph targets by name
-	for(auto CurveIter=MorphCurveAnims.CreateConstIterator(); CurveIter; ++CurveIter)
-	{
-		const FName& CurveName	= (CurveIter).Key();
-		const float& Weight	= (CurveIter).Value();
-
-		// If it has a valid weight
-		if(FMath::Abs(Weight) > MinVertexAnimBlendWeight)
-		{
-			// Find morph reference
-			UMorphTarget* Target = InSkeletalMesh ? InSkeletalMesh->FindMorphTarget(CurveName) : NULL;
-			if(Target != NULL)				
-			{
-				// See if this morph target already has an entry
-				int32 AnimIndex = FindVertexAnim(OutVertexAnims, Target);
-				// If not, add it
-				if(AnimIndex == INDEX_NONE)
-				{
-					OutVertexAnims.Add(FActiveVertexAnim(Target, Weight));
-				}
-				// If it does, use the max weight
-				else
-				{
-					const float CurrentWeight = OutVertexAnims[AnimIndex].Weight;
-					OutVertexAnims[AnimIndex].Weight = FMath::Max<float>(CurrentWeight, Weight);
-				}
-			}
-		}
-	}
-
-	return OutVertexAnims;
-}
-
 void USkinnedMeshComponent::FinalizeBoneTransform()
 {
 	FlipEditableSpaceBases();
