Copying //UE4/Dev-Framework to Dev-Main (//UE4/Dev-Main)

diff --git a/Engine/Source/Runtime/Engine/Private/Debug/DebugDrawService.cpp b/Engine/Source/Runtime/Engine/Private/Debug/DebugDrawService.cpp
--- a/Engine/Source/Runtime/Engine/Private/Debug/DebugDrawService.cpp
+++ b/Engine/Source/Runtime/Engine/Private/Debug/DebugDrawService.cpp
@@ -98,91 +98,3 @@ void UDebugDrawService::Draw(const FEngineShowFlags Flags, UCanvas* Canvas)
 		}
 	}
 }
-
-//----------------------------------------------------------------------//
-// FDebugShapeElement
-//----------------------------------------------------------------------//
-
-FDrawDebugShapeElement::FDrawDebugShapeElement(EDrawDebugShapeElement InType)
-	: TransformationMatrix(FMatrix::Identity)
-	, Type(InType)
-	, Color(0xff)
-	, ThicknesOrRadius(0)
-{
-	// Empty
-}
-
-FDrawDebugShapeElement::FDrawDebugShapeElement(const FString& InDescription, const FColor& InColor, uint16 InThickness)
-	: TransformationMatrix(FMatrix::Identity)
-	, Type(EDrawDebugShapeElement::Invalid)
-	, ThicknesOrRadius(InThickness)
-{
-	if (InDescription.IsEmpty() == false)
-	{
-		Description = InDescription;
-	}
-	SetColor(InColor);
-}
-
-void FDrawDebugShapeElement::SetColor(const FColor& InColor)
-{
-	Color = ((InColor.DWColor() >> 30) << 6) | (((InColor.DWColor() & 0x00ff0000) >> 22) << 4) | (((InColor.DWColor() & 0x0000ff00) >> 14) << 2) | ((InColor.DWColor() & 0x000000ff) >> 6);
-}
-
-EDrawDebugShapeElement FDrawDebugShapeElement::GetType() const
-{
-	return Type;
-}
-
-void FDrawDebugShapeElement::SetType(EDrawDebugShapeElement InType)
-{
-	Type = InType;
-}
-
-FColor FDrawDebugShapeElement::GetFColor() const
-{
-	FColor RetColor(((Color & 0xc0) << 24) | ((Color & 0x30) << 18) | ((Color & 0x0c) << 12) | ((Color & 0x03) << 6));
-	RetColor.A = (RetColor.A * 255) / 192; // convert alpha to 0-255 range
-	return RetColor;
-}
-
-FDrawDebugShapeElement FDrawDebugShapeElement::MakeString(const FString& Description, const FVector& Location)
-{
-	FDrawDebugShapeElement StringElement(Description, FColor::White, 0);
-	StringElement.SetType(EDrawDebugShapeElement::String);
-	StringElement.Points.Add(Location);
-
-	return StringElement;
-}
-
-FDrawDebugShapeElement FDrawDebugShapeElement::MakeLine(const FVector& Start, const FVector& End, FColor Color, float Thickness)
-{
-	FDrawDebugShapeElement LineElement(EDrawDebugShapeElement::Segment);
-	LineElement.Points.Add(Start);
-	LineElement.Points.Add(End);
-	LineElement.SetColor(Color);
-	LineElement.ThicknesOrRadius = Thickness;
-
-	return LineElement;
-}
-
-FDrawDebugShapeElement FDrawDebugShapeElement::MakePoint(const FVector& Location, FColor Color, float Radius)
-{
-	FDrawDebugShapeElement PointElement(EDrawDebugShapeElement::SinglePoint);
-	PointElement.Points.Add(Location);
-	PointElement.ThicknesOrRadius = Radius;
-	PointElement.SetColor(Color);
-
-	return PointElement;
-}
-
-FDrawDebugShapeElement FDrawDebugShapeElement::MakeCylinder(FVector const& Start, FVector const& End, float Radius, int32 Segments, FColor const& Color)
-{
-	FDrawDebugShapeElement ShapeElement(EDrawDebugShapeElement::Cylinder);
-	ShapeElement.Points.Add(Start);
-	ShapeElement.Points.Add(End);
-	ShapeElement.ThicknesOrRadius = Radius;
-	ShapeElement.SetColor(Color);
-
-	return ShapeElement;
-}
diff --git a/Engine/Source/Runtime/Engine/Private/Debug/GameplayDebuggerBaseObject.cpp b/Engine/Source/Runtime/Engine/Private/Debug/GameplayDebuggerBaseObject.cpp
new file mode 100644
--- /dev/null
+++ b/Engine/Source/Runtime/Engine/Private/Debug/GameplayDebuggerBaseObject.cpp
@@ -0,0 +1,761 @@
+// Copyright 1998-2015 Epic Games, Inc. All Rights Reserved.
+#include "EnginePrivate.h"
+#include "Net/UnrealNetwork.h"
+#include "Misc/CoreMisc.h"
+#include "Regex.h"
+#include "Debug/GameplayDebuggerBaseObject.h"
+#include "AssetRegistryModule.h"
+
+DEFINE_LOG_CATEGORY(LogGameplayDebugger);
+
+UGameplayDebuggerBaseObject::UGameplayDebuggerBaseObject(const FObjectInitializer& ObjectInitializer) : Super(ObjectInitializer)
+{
+
+}
+
+void UGameplayDebuggerBaseObject::GetLifetimeReplicatedProps(TArray< FLifetimeProperty > & OutLifetimeProps) const
+{
+	Super::GetLifetimeReplicatedProps(OutLifetimeProps);
+#if ENABLED_GAMEPLAY_DEBUGGER
+	DOREPLIFETIME(UGameplayDebuggerBaseObject, GenericShapeElements);
+	DOREPLIFETIME(UGameplayDebuggerBaseObject, ReplicatedBlob);
+#endif
+}
+
+int32 UGameplayDebuggerBaseObject::GetFunctionCallspace(UFunction* Function, void* Parameters, FFrame* Stack)
+{
+#if ENABLED_GAMEPLAY_DEBUGGER
+	if (HasAnyFlags(RF_ClassDefaultObject))
+	{
+		return FunctionCallspace::Local;
+	}
+	check(GetOuter() != NULL);
+	return GetOuter()->GetFunctionCallspace(Function, Parameters, Stack);
+#endif //ENABLED_GAMEPLAY_DEBUGGER
+
+	return FunctionCallspace::Local;
+}
+
+bool UGameplayDebuggerBaseObject::CallRemoteFunction(UFunction* Function, void* Parameters, FOutParmRec* OutParms, FFrame* Stack)
+{
+#if ENABLED_GAMEPLAY_DEBUGGER
+	check(!HasAnyFlags(RF_ClassDefaultObject));
+	check(GetOuter() != NULL);
+
+	AActor* Owner = CastChecked<AActor>(GetOuter());
+	UNetDriver* NetDriver = Owner->GetNetDriver();
+	if (NetDriver)
+	{
+		NetDriver->ProcessRemoteFunction(Owner, Function, Parameters, OutParms, Stack, this);
+		return true;
+	}
+#endif
+	return false;
+}
+
+void UGameplayDebuggerBaseObject::OnRep_GenericShapeElements()
+{
+#if ENABLED_GAMEPLAY_DEBUGGER
+	MarkRenderStateDirty();
+#endif
+}
+
+void UGameplayDebuggerBaseObject::OnRep_ReplicatedBlob()
+{
+#if ENABLED_GAMEPLAY_DEBUGGER
+	if (LocalReplicationData.VersionNum != ReplicatedBlob.VersionNum || LocalReplicationData.RepDataOffset != ReplicatedBlob.RepDataOffset)
+	{
+		LocalReplicationData.VersionNum = ReplicatedBlob.VersionNum;
+		LocalReplicationData.RepDataOffset = ReplicatedBlob.RepDataOffset;
+		LocalReplicationData.RepDataSize = ReplicatedBlob.RepDataSize;
+		LocalReplicationData.RepData.Reset();
+		if (ReplicatedBlob.RepDataOffset != 0)
+		{
+			ServerRequestPacket(0);
+			return;
+		}
+	}
+
+	const int32 StartIdx = LocalReplicationData.RepData.Num();
+	LocalReplicationData.RepData.Append(ReplicatedBlob.RepData);
+	LocalReplicationData.RepDataOffset += ReplicatedBlob.RepData.Num();
+
+	const int32 HeaderSize = sizeof(FReplicationDataHeader);
+	if (LocalReplicationData.RepDataSize == LocalReplicationData.RepData.Num() && LocalReplicationData.RepData.Num() > HeaderSize)
+	{
+		TArray<uint8> UncompressedBuffer;
+		FReplicationDataHeader DataHeader;
+
+		uint8* SourceBuffer = (uint8*)LocalReplicationData.RepData.GetData();
+		FMemory::Memcpy(&DataHeader, SourceBuffer, sizeof(FReplicationDataHeader));
+		SourceBuffer += sizeof(FReplicationDataHeader);
+
+		UncompressedBuffer.AddZeroed(DataHeader.UncompressedSize);
+		if (DataHeader.bUseCompression)
+		{
+			const int32 CompressedSize = LocalReplicationData.RepData.Num() - HeaderSize;
+			FCompression::UncompressMemory((ECompressionFlags)(COMPRESS_ZLIB),
+				(void*)UncompressedBuffer.GetData(), DataHeader.UncompressedSize, SourceBuffer, CompressedSize);
+		}
+		else
+		{
+			FMemory::Memcpy(UncompressedBuffer.GetData(), SourceBuffer, DataHeader.UncompressedSize);
+		}
+		LocalReplicationData.RepData = UncompressedBuffer;
+		LocalReplicationData.RepDataSize = UncompressedBuffer.Num();
+
+		OnDataReplicationComplited(LocalReplicationData.RepData);
+	}
+	else
+	{
+		ServerRequestPacket(LocalReplicationData.RepDataOffset);
+	}
+
+#endif // ENABLED_GAMEPLAY_DEBUGGER
+}
+
+bool UGameplayDebuggerBaseObject::CollectDataToReplicateOnServer_Validate(APlayerController* MyPC, AActor *SelectedActor)
+{
+#if ENABLED_GAMEPLAY_DEBUGGER
+	return true;
+#else
+	return false;
+#endif
+}
+
+void UGameplayDebuggerBaseObject::CollectDataToReplicateOnServer_Implementation(APlayerController* MyPC, AActor *SelectedActor)
+{
+#if ENABLED_GAMEPLAY_DEBUGGER
+	CollectDataToReplicate(MyPC, SelectedActor);
+#endif
+}
+
+bool UGameplayDebuggerBaseObject::ServerRequestPacket_Validate(int32 Offset)
+{
+#if ENABLED_GAMEPLAY_DEBUGGER
+	return true;
+#else
+	return false;
+#endif
+}
+
+void UGameplayDebuggerBaseObject::ServerRequestPacket_Implementation(int32 Offset)
+{
+#if ENABLED_GAMEPLAY_DEBUGGER
+	if (LocalReplicationData.RepData.IsValidIndex(Offset) == false)
+	{
+		return;
+	}
+	LocalReplicationData.RepDataOffset = Offset;
+
+	int32 MaxPacketSizeAsFunctionArgument = MAX_PACKET_SIZE;
+	int32 NumToSend = LocalReplicationData.RepData.Num() - LocalReplicationData.RepDataOffset;
+	int32 PacketSize = FMath::Min(MaxPacketSizeAsFunctionArgument, NumToSend);
+
+	TArray<uint8> PacketData;
+	PacketData.AddZeroed(PacketSize);
+	const uint8*SrcBuffer = LocalReplicationData.RepData.GetData() + LocalReplicationData.RepDataOffset;
+	uint8 *DestBuffer = PacketData.GetData();
+	FMemory::Memcpy(DestBuffer, SrcBuffer, PacketSize);
+
+	ReplicatedBlob = FGameplayDebuggerReplicatedBlob(PacketData, LocalReplicationData.VersionNum, LocalReplicationData.RepDataOffset, LocalReplicationData.RepDataSize);
+#endif
+}
+
+void UGameplayDebuggerBaseObject::RequestDataReplication(const FString& DataName, const TArray<uint8>& UncompressedBuffer, EGameplayDebuggerReplicate Replicate)
+{
+#if ENABLED_GAMEPLAY_DEBUGGER
+	LocalReplicationData.VersionNum++;
+	LocalReplicationData.RepDataOffset = 0;
+	ENetMode NetMode = GetWorld()->GetNetMode();
+	if (NetMode != NM_Standalone)
+	{
+		const double Timer2 = FPlatformTime::Seconds();
+
+		FReplicationDataHeader DataHeader;
+		DataHeader.bUseCompression = Replicate == EGameplayDebuggerReplicate::WithCompressed;
+		DataHeader.UncompressedSize = UncompressedBuffer.Num();
+
+		const int32 HeaderSize = sizeof(FReplicationDataHeader);
+		LocalReplicationData.RepData.Init(0, HeaderSize + FMath::TruncToInt(1.1f * UncompressedBuffer.Num()));
+
+		const int32 UncompressedSize = UncompressedBuffer.Num();
+		int32 CompressedSize = LocalReplicationData.RepData.Num() - HeaderSize;
+		uint8* DestBuffer = LocalReplicationData.RepData.GetData();
+		FMemory::Memcpy(DestBuffer, &DataHeader, HeaderSize);
+		DestBuffer += HeaderSize;
+
+		if (Replicate == EGameplayDebuggerReplicate::WithCompressed)
+		{
+			FCompression::CompressMemory((ECompressionFlags)(COMPRESS_ZLIB | COMPRESS_BiasMemory),
+				(void*)DestBuffer, CompressedSize, (void*)UncompressedBuffer.GetData(), UncompressedSize);
+
+			LocalReplicationData.RepData.SetNum(CompressedSize + HeaderSize, false);
+		}
+		else
+		{
+			FMemory::Memcpy(DestBuffer, (void*)UncompressedBuffer.GetData(), UncompressedSize);
+			LocalReplicationData.RepData.SetNum(UncompressedSize + HeaderSize, false);
+		}
+
+		LocalReplicationData.RepDataSize = LocalReplicationData.RepData.Num();
+		ServerRequestPacket(/*int32 Offset*/0);
+
+#if UE_BUILD_DEBUG
+		const double Timer3 = FPlatformTime::Seconds();
+		UE_LOG(LogGameplayDebugger, Verbose, TEXT("Preparing '%s' data to replicate: %.1fkB took %.3fms (collect: %.3fms + %s compress %d%%: %.3fms)"),
+			*DataName,
+			LocalReplicationData.RepData.Num() / 1024.0f, 
+			1000.0f * (Timer3 - DataPreparationStartTime),
+			1000.0f * (Timer2 - DataPreparationStartTime),
+			Replicate == EGameplayDebuggerReplicate::WithCompressed ? TEXT("") : TEXT("no"),
+			FMath::TruncToInt(100.0f * LocalReplicationData.RepData.Num() / UncompressedBuffer.Num()),
+			1000.0f * (Timer3 - Timer2)
+		);
+#endif
+	}
+	else
+	{
+		LocalReplicationData.RepData = UncompressedBuffer;
+		LocalReplicationData.RepDataSize = LocalReplicationData.RepData.Num();
+		OnDataReplicationComplited(LocalReplicationData.RepData);
+	}
+#endif
+}
+
+void UGameplayDebuggerBaseObject::CollectDataToReplicate(APlayerController* MyPC, AActor *SelectedActor)
+{
+#if ENABLED_GAMEPLAY_DEBUGGER
+	DataPreparationStartTime = FPlatformTime::Seconds();
+#endif
+}
+
+void UGameplayDebuggerBaseObject::DrawCollectedData(APlayerController* MyPC, AActor* SelectedActor)
+{
+#if ENABLED_GAMEPLAY_DEBUGGER
+	float VerticalLabelOffset = 0.f;
+	for (const auto& Shape : GenericShapeElements)
+	{
+		if (Shape.GetType() == EGameplayDebuggerShapeElement::String)
+		{
+			const FVector& Loc = Shape.Points[0];
+			const FVector ScreenLoc = DefaultContext.Canvas->Project(Loc);
+
+			PrintString(DefaultContext, Shape.GetFColor(), Shape.Description, ScreenLoc.X, ScreenLoc.Y + VerticalLabelOffset);
+			VerticalLabelOffset += 17;
+		}
+		else if (Shape.GetType() == EGameplayDebuggerShapeElement::Segment)
+		{
+			DrawDebugLine(GetWorld(), Shape.Points[0], Shape.Points[1], Shape.GetFColor());
+		}
+		else if (Shape.GetType() == EGameplayDebuggerShapeElement::SinglePoint)
+		{
+			DrawDebugSphere(GetWorld(), Shape.Points[0], Shape.ThicknesOrRadius, 16, Shape.GetFColor());
+		}
+		else if (Shape.GetType() == EGameplayDebuggerShapeElement::Cylinder)
+		{
+			DrawDebugCylinder(GetWorld(), Shape.Points[0], Shape.Points[1], Shape.ThicknesOrRadius, 16, Shape.GetFColor());
+		}
+	}
+#endif
+}
+
+void UGameplayDebuggerBaseObject::BindKeyboardImput(class UInputComponent*& InputComponent)
+{
+
+}
+
+void UGameplayDebuggerBaseObject::PostLoad()
+{
+	Super::PostLoad();
+
+#if ENABLED_GAMEPLAY_DEBUGGER
+	if (HasAnyFlags(RF_ClassDefaultObject)
+#	if WITH_HOT_RELOAD
+		&& !GIsHotReload
+#	endif // WITH_HOT_RELOAD
+		)
+	{
+		UGameplayDebuggerBaseObject::GetOnClassRegisterEvent().Broadcast(this);
+	}
+#endif
+}
+
+void UGameplayDebuggerBaseObject::BeginDestroy()
+{
+#if ENABLED_GAMEPLAY_DEBUGGER
+	if (HasAnyFlags(RF_ClassDefaultObject)
+#	if WITH_HOT_RELOAD
+		&& !GIsHotReload
+#	endif // WITH_HOT_RELOAD
+		)
+	{
+		UGameplayDebuggerBaseObject::GetOnClassUnregisterEvent().Broadcast(this);
+	}
+#endif
+	Super::BeginDestroy();
+}
+
+UGameplayDebuggerBaseObject::FOnClassRegisterChangeEvent& UGameplayDebuggerBaseObject::GetOnClassRegisterEvent()
+{
+	static UGameplayDebuggerBaseObject::FOnClassRegisterChangeEvent OnClassRegisterEvent;
+	return OnClassRegisterEvent;
+}
+
+UGameplayDebuggerBaseObject::FOnClassRegisterChangeEvent& UGameplayDebuggerBaseObject::GetOnClassUnregisterEvent()
+{
+	static UGameplayDebuggerBaseObject::FOnClassRegisterChangeEvent OnClassUnregisterEvent;
+	return OnClassUnregisterEvent;
+}
+
+
+//----------------------------------------------------------------------//
+// FDebugShapeElement
+//----------------------------------------------------------------------//
+
+FArchive& operator<<(FArchive& Ar, FGameplayDebuggerShapeElement& ShapeElement)
+{
+	Ar << ShapeElement.Description;
+	Ar << ShapeElement.Points;
+	Ar << ShapeElement.TransformationMatrix;
+	Ar << ShapeElement.Type;
+	Ar << ShapeElement.Color;
+	Ar << ShapeElement.ThicknesOrRadius;
+
+	return Ar;
+}
+
+FGameplayDebuggerShapeElement::FGameplayDebuggerShapeElement(EGameplayDebuggerShapeElement InType)
+	: TransformationMatrix(FMatrix::Identity)
+	, Type(InType)
+	, Color(0xff)
+	, ThicknesOrRadius(0)
+{
+	// Empty
+}
+
+FGameplayDebuggerShapeElement::FGameplayDebuggerShapeElement(const FString& InDescription, const FColor& InColor, uint16 InThickness)
+	: TransformationMatrix(FMatrix::Identity)
+	, Type(EGameplayDebuggerShapeElement::Invalid)
+	, ThicknesOrRadius(InThickness)
+{
+	if (InDescription.IsEmpty() == false)
+	{
+		Description = InDescription;
+	}
+	SetColor(InColor);
+}
+
+void FGameplayDebuggerShapeElement::SetColor(const FColor& InColor)
+{
+	Color = ((InColor.DWColor() >> 30) << 6) | (((InColor.DWColor() & 0x00ff0000) >> 22) << 4) | (((InColor.DWColor() & 0x0000ff00) >> 14) << 2) | ((InColor.DWColor() & 0x000000ff) >> 6);
+}
+
+EGameplayDebuggerShapeElement FGameplayDebuggerShapeElement::GetType() const
+{
+	return Type;
+}
+
+void FGameplayDebuggerShapeElement::SetType(EGameplayDebuggerShapeElement InType)
+{
+	Type = InType;
+}
+
+FColor FGameplayDebuggerShapeElement::GetFColor() const
+{
+	FColor RetColor(((Color & 0xc0) << 24) | ((Color & 0x30) << 18) | ((Color & 0x0c) << 12) | ((Color & 0x03) << 6));
+	RetColor.A = (RetColor.A * 255) / 192; // convert alpha to 0-255 range
+	return RetColor;
+}
+
+//////////////////////////////////////////////////////////////////////////
+UGameplayDebuggerHelper::FPrintContext UGameplayDebuggerHelper::OverHeadContext;
+UGameplayDebuggerHelper::FPrintContext UGameplayDebuggerHelper::DefaultContext;
+
+UGameplayDebuggerHelper::UGameplayDebuggerHelper(const FObjectInitializer& ObjectInitializer) : Super(ObjectInitializer)
+{
+
+}
+
+void UGameplayDebuggerHelper::PrintDebugString(const FString& InString)
+{
+	PrintString(GetDefaultContext(), InString);
+}
+
+void UGameplayDebuggerHelper::PrintString(UGameplayDebuggerHelper::FPrintContext& Context, const FString& InString)
+{
+#if ENABLED_GAMEPLAY_DEBUGGER
+	class FParserHelper
+	{
+		enum TokenType
+		{
+			OpenTag,
+			CloseTag,
+			NewLine,
+			EndOfString,
+			RegularChar,
+			Tab
+		};
+
+		enum Tag
+		{
+			DefinedColor,
+			OtherColor,
+			ErrorTag,
+		};
+
+		uint8 ReadToken()
+		{
+			uint8 OutToken = RegularChar;
+			TCHAR ch = Index < DataString.Len() ? DataString[Index] : '\0';
+			switch (ch)
+			{
+			case '\0':
+				OutToken = EndOfString;
+				break;
+			case '{':
+				OutToken = OpenTag;
+				Index++;
+				break;
+			case '}':
+				OutToken = CloseTag;
+				Index++;
+				break;
+			case '\n':
+				OutToken = NewLine;
+				Index++;
+				break;
+			case '\t':
+				OutToken = Tab;
+				Index++;
+				break;
+			default:
+				OutToken = RegularChar;
+				break;
+			}
+			return OutToken;
+		}
+
+		uint32 ParseTag(FString& OutData)
+		{
+			FString TagString;
+			int32 OryginalIndex = Index;
+			uint8 token = ReadToken();
+			while (token != EndOfString && token != CloseTag)
+			{
+				if (token == RegularChar)
+				{
+					TagString.AppendChar(DataString[Index++]);
+				}
+				token = ReadToken();
+			}
+
+			int OutTag = ErrorTag;
+
+			if (token != CloseTag)
+			{
+				Index = OryginalIndex;
+				OutData = FString::Printf(TEXT("{%s"), *TagString);
+				OutData.AppendChar(DataString[Index - 1]);
+				return OutTag;
+			}
+
+			if (GColorList.IsValidColorName(*TagString.ToLower()))
+			{
+				OutTag = DefinedColor;
+				OutData = TagString;
+			}
+			else
+			{
+				FColor Color;
+				if (Color.InitFromString(TagString))
+				{
+					OutTag = OtherColor;
+					OutData = TagString;
+				}
+				else
+				{
+					OutTag = ErrorTag;
+					OutData = FString::Printf(TEXT("{%s"), *TagString);
+					OutData.AppendChar(DataString[Index - 1]);
+				}
+			}
+			//Index++;
+			return OutTag;
+		}
+
+		struct StringNode
+		{
+			FString String;
+			FColor Color;
+			bool bNewLine;
+			StringNode() : Color(FColor::White), bNewLine(false) {}
+		};
+		int32 Index;
+		FString DataString;
+	public:
+		TArray<StringNode> Strings;
+
+		void ParseString(const FString& StringToParse)
+		{
+			Index = 0;
+			DataString = StringToParse;
+			Strings.Add(StringNode());
+			if (Index >= DataString.Len())
+				return;
+
+			uint8 Token = ReadToken();
+			while (Token != EndOfString)
+			{
+				switch (Token)
+				{
+				case RegularChar:
+					Strings[Strings.Num() - 1].String.AppendChar(DataString[Index++]);
+					break;
+				case NewLine:
+					Strings.Add(StringNode());
+					Strings[Strings.Num() - 1].bNewLine = true;
+					Strings[Strings.Num() - 1].Color = Strings[Strings.Num() - 2].Color;
+					break;
+				case EndOfString:
+					break;
+				case Tab:
+				{
+					const FString TabString(TEXT("     "));
+					Strings[Strings.Num() - 1].String.Append(TabString);
+					static bool sbTest = false;
+					if (sbTest)
+					{
+						Index++;
+					}
+					break;
+				}
+				case OpenTag:
+				{
+					FString OutData;
+					switch (ParseTag(OutData))
+					{
+					case DefinedColor:
+					{
+						int32 i = Strings.Add(StringNode());
+						Strings[i].Color = GColorList.GetFColorByName(*OutData.ToLower());
+					}
+						break;
+					case OtherColor:
+					{
+						FColor NewColor;
+						if (NewColor.InitFromString(OutData))
+						{
+							int32 i = Strings.Add(StringNode());
+							Strings[i].Color = NewColor;
+							break;
+						}
+					}
+					default:
+						Strings[Strings.Num() - 1].String += OutData;
+						break;
+					}
+				}
+					break;
+				}
+				Token = ReadToken();
+			}
+		}
+	};
+
+	FParserHelper Helper;
+	Helper.ParseString(InString);
+
+	float YMovement = 0;
+	float XL = 0.f, YL = 0.f;
+	CalulateStringSize(Context, Context.Font.Get(), TEXT("X"), XL, YL);
+
+	for (int32 Index = 0; Index < Helper.Strings.Num(); ++Index)
+	{
+		if (Index > 0 && Helper.Strings[Index].bNewLine)
+		{
+			if (Context.Canvas != NULL && YMovement + Context.CursorY > Context.Canvas->ClipY)
+			{
+				Context.DefaultX += Context.Canvas->ClipX / 2;
+				Context.CursorX = Context.DefaultX;
+				Context.CursorY = Context.DefaultY;
+				YMovement = 0;
+			}
+			YMovement += YL;
+			Context.CursorX = Context.DefaultX;
+		}
+		const FString Str = Helper.Strings[Index].String;
+
+		if (Str.Len() > 0 && Context.Canvas != NULL)
+		{
+			float SizeX, SizeY;
+			CalulateStringSize(Context, Context.Font.Get(), Str, SizeX, SizeY);
+
+			FCanvasTextItem TextItem(FVector2D::ZeroVector, FText::FromString(Str), Context.Font.Get(), FLinearColor::White);
+			if (Context.FontRenderInfo.bEnableShadow)
+			{
+				TextItem.EnableShadow(FColor::Black, FVector2D(1, 1));
+			}
+
+			TextItem.SetColor(Helper.Strings[Index].Color);
+			DrawItem(Context, TextItem, Context.CursorX, YMovement + Context.CursorY);
+			Context.CursorX += SizeX;
+		}
+	}
+
+	Context.CursorY += YMovement;
+#endif //!(UE_BUILD_SHIPPING || UE_BUILD_TEST)
+}
+
+void UGameplayDebuggerHelper::PrintString(UGameplayDebuggerHelper::FPrintContext& Context, const FColor& InColor, const FString& InString)
+{
+#if ENABLED_GAMEPLAY_DEBUGGER
+	PrintString(Context, FString::Printf(TEXT("{%s}%s"), *InColor.ToString(), *InString));
+#endif //!(UE_BUILD_SHIPPING || UE_BUILD_TEST)
+}
+
+void UGameplayDebuggerHelper::PrintString(UGameplayDebuggerHelper::FPrintContext& Context, const FString& InString, float X, float Y)
+{
+#if ENABLED_GAMEPLAY_DEBUGGER
+	const float OldX = Context.CursorX, OldY = Context.CursorY;
+	const float OldDX = Context.DefaultX, OldDY = Context.DefaultY;
+
+	Context.DefaultX = Context.CursorX = X;
+	Context.DefaultY = Context.CursorY = Y;
+	PrintString(Context, InString);
+
+	Context.CursorX = OldX;
+	Context.CursorY = OldY;
+	Context.DefaultX = OldDX;
+	Context.DefaultY = OldDY;
+#endif
+}
+
+void UGameplayDebuggerHelper::PrintString(UGameplayDebuggerHelper::FPrintContext& Context, const FColor& InColor, const FString& InString, float X, float Y)
+{
+#if ENABLED_GAMEPLAY_DEBUGGER
+	const float OldX = Context.CursorX, OldY = Context.CursorY;
+	const float OldDX = Context.DefaultX, OldDY = Context.DefaultY;
+
+	Context.DefaultX = Context.CursorX = X;
+	Context.DefaultY = Context.CursorY = Y;
+	PrintString(Context, InColor, InString);
+
+	Context.CursorX = OldX;
+	Context.CursorY = OldY;
+	Context.DefaultX = OldDX;
+	Context.DefaultY = OldDY;
+#endif
+}
+
+void UGameplayDebuggerHelper::CalulateStringSize(const UGameplayDebuggerHelper::FPrintContext& Context, UFont* Font, const FString& InString, float& OutX, float& OutY)
+{
+#if ENABLED_GAMEPLAY_DEBUGGER
+	FString String = InString;
+	const FRegexPattern ElementRegexPattern(TEXT("\\{.*?\\}"));
+	FRegexMatcher ElementRegexMatcher(ElementRegexPattern, String);
+
+	while (ElementRegexMatcher.FindNext())
+	{
+		int32 AttributeListBegin = ElementRegexMatcher.GetCaptureGroupBeginning(0);
+		int32 AttributeListEnd = ElementRegexMatcher.GetCaptureGroupEnding(0);
+		String.RemoveAt(AttributeListBegin, AttributeListEnd - AttributeListBegin);
+		ElementRegexMatcher = FRegexMatcher(ElementRegexPattern, String);
+	}
+
+	OutX = OutY = 0;
+	if (Context.Canvas != NULL)
+	{
+		TArray<FString> Lines;
+		String.ParseIntoArrayLines(Lines);
+		for (FString Line : Lines)
+		{
+			float Width, Height;
+			Context.Canvas->StrLen(Font != NULL ? Font : Context.Font.Get(), Line, Width, Height);
+			OutX = FMath::Max(OutX, Width);
+			OutY += Height;
+		}
+	}
+#endif
+}
+
+void UGameplayDebuggerHelper::CalulateTextSize(const UGameplayDebuggerHelper::FPrintContext& Context, UFont* Font, const FText& InText, float& OutX, float& OutY)
+{
+#if ENABLED_GAMEPLAY_DEBUGGER
+	CalulateStringSize(Context, Font, InText.ToString(), OutX, OutY);
+#endif
+}
+
+FVector UGameplayDebuggerHelper::ProjectLocation(const UGameplayDebuggerHelper::FPrintContext& Context, const FVector& Location)
+{
+#if ENABLED_GAMEPLAY_DEBUGGER
+	if (Context.Canvas != NULL)
+	{
+		return Context.Canvas->Project(Location);
+	}
+#endif
+	return FVector();
+}
+
+void UGameplayDebuggerHelper::DrawItem(const UGameplayDebuggerHelper::FPrintContext& Context, class FCanvasItem& Item, float X, float Y)
+{
+#if ENABLED_GAMEPLAY_DEBUGGER
+	if (Context.Canvas.IsValid())
+	{
+		Context.Canvas->DrawItem(Item, FVector2D(X, Y));
+	}
+#endif
+}
+
+void UGameplayDebuggerHelper::DrawIcon(const UGameplayDebuggerHelper::FPrintContext& Context, const FColor& InColor, const FCanvasIcon& Icon, float X, float Y, float Scale)
+{
+#if ENABLED_GAMEPLAY_DEBUGGER
+	if (Context.Canvas.IsValid())
+	{
+		Context.Canvas->SetDrawColor(InColor);
+		Context.Canvas->DrawIcon(Icon, X, Y, Scale);
+	}
+#endif
+}
+
+FGameplayDebuggerShapeElement UGameplayDebuggerHelper::MakeString(const FString& Description, const FVector& Location)
+{
+	FGameplayDebuggerShapeElement StringElement(Description, FColor::White, 0);
+	StringElement.SetType(EGameplayDebuggerShapeElement::String);
+	StringElement.Points.Add(Location);
+
+	return StringElement;
+}
+
+FGameplayDebuggerShapeElement UGameplayDebuggerHelper::MakeLine(const FVector& Start, const FVector& End, FColor Color, float Thickness)
+{
+	FGameplayDebuggerShapeElement LineElement(EGameplayDebuggerShapeElement::Segment);
+	LineElement.Points.Add(Start);
+	LineElement.Points.Add(End);
+	LineElement.SetColor(Color);
+	LineElement.ThicknesOrRadius = Thickness;
+
+	return LineElement;
+}
+
+FGameplayDebuggerShapeElement UGameplayDebuggerHelper::MakePoint(const FVector& Location, FColor Color, float Radius)
+{
+	FGameplayDebuggerShapeElement PointElement(EGameplayDebuggerShapeElement::SinglePoint);
+	PointElement.Points.Add(Location);
+	PointElement.ThicknesOrRadius = Radius;
+	PointElement.SetColor(Color);
+
+	return PointElement;
+}
+
+FGameplayDebuggerShapeElement UGameplayDebuggerHelper::MakeCylinder(FVector const& Start, FVector const& End, float Radius, int32 Segments, FColor const& Color)
+{
+	FGameplayDebuggerShapeElement ShapeElement(EGameplayDebuggerShapeElement::Cylinder);
+	ShapeElement.Points.Add(Start);
+	ShapeElement.Points.Add(End);
+	ShapeElement.ThicknesOrRadius = Radius;
+	ShapeElement.SetColor(Color);
+
+	return ShapeElement;
+}
+
diff --git a/Engine/Source/Runtime/Engine/Private/ErrorChecking.cpp b/Engine/Source/Runtime/Engine/Private/ErrorChecking.cpp
--- a/Engine/Source/Runtime/Engine/Private/ErrorChecking.cpp
+++ b/Engine/Source/Runtime/Engine/Private/ErrorChecking.cpp
@@ -243,7 +243,7 @@ void ALevelStreamingVolume::CheckForErrors()
 	bool bHasAssociatedLevels = false;
 	for (auto It = GetWorld()->StreamingLevels.CreateConstIterator(); It; ++It)
 	{
-		if ((*It)->EditorStreamingVolumes.Find(this) != INDEX_NONE)
+		if ((*It)->EditorStreamingVolumes.Contains(this))
 		{
 			bHasAssociatedLevels = true;
 			break;
diff --git a/Engine/Source/Runtime/Engine/Private/GameplayStatics.cpp b/Engine/Source/Runtime/Engine/Private/GameplayStatics.cpp
--- a/Engine/Source/Runtime/Engine/Private/GameplayStatics.cpp
+++ b/Engine/Source/Runtime/Engine/Private/GameplayStatics.cpp
@@ -690,21 +690,25 @@ UParticleSystemComponent* UGameplayStatics::SpawnEmitterAttached(UParticleSystem
 		{
 			UE_LOG(LogScript, Warning, TEXT("UGameplayStatics::SpawnEmitterAttached: NULL AttachComponent specified!"));
 		}
-		else if (AttachToComponent->GetWorld() && AttachToComponent->GetWorld()->GetNetMode() != NM_DedicatedServer)
+		else
 		{
-			PSC = CreateParticleSystem(EmitterTemplate, AttachToComponent->GetWorld(), AttachToComponent->GetOwner(), bAutoDestroy);
-			PSC->AttachTo(AttachToComponent, AttachPointName);
-			if (LocationType == EAttachLocation::KeepWorldPosition)
-			{
-				PSC->SetWorldLocationAndRotation(Location, Rotation);
-			}
-			else
+			UWorld* const World = AttachToComponent->GetWorld();
+			if (World && World->GetNetMode() != NM_DedicatedServer)
 			{
-				PSC->SetRelativeLocationAndRotation(Location, Rotation);
+				PSC = CreateParticleSystem(EmitterTemplate, World, AttachToComponent->GetOwner(), bAutoDestroy);
+				PSC->AttachTo(AttachToComponent, AttachPointName);
+				if (LocationType == EAttachLocation::KeepWorldPosition)
+				{
+					PSC->SetWorldLocationAndRotation(Location, Rotation);
+				}
+				else
+				{
+					PSC->SetRelativeLocationAndRotation(Location, Rotation);
+				}
+				PSC->SetRelativeScale3D(FVector(1.f));
+				PSC->RegisterComponentWithWorld(World);
+				PSC->ActivateSystem(true);
 			}
-			PSC->SetRelativeScale3D(FVector(1.f));
-			PSC->RegisterComponentWithWorld(AttachToComponent->GetWorld());
-			PSC->ActivateSystem(true);
 		}
 	}
 	return PSC;
@@ -802,6 +806,7 @@ void UGameplayStatics::PlaySound2D(UObject* WorldContextObject, class USoundBase
 
 		NewActiveSound.bIsUISound = true;
 		NewActiveSound.ConcurrencySettings = ConcurrencySettings;
+		NewActiveSound.Priority = Sound->Priority;
 
 		AudioDevice->AddNewActiveSound(NewActiveSound);
 	}
@@ -889,6 +894,7 @@ void UGameplayStatics::PlaySoundAtLocation(UObject* WorldContextObject, class US
 			}
 
 			NewActiveSound.ConcurrencySettings = ConcurrencySettings;
+			NewActiveSound.Priority = Sound->Priority;
 
 			// TODO - Audio Threading. This call would be a task call to dispatch to the audio thread
 			AudioDevice->AddNewActiveSound(NewActiveSound);
diff --git a/Engine/Source/Runtime/Engine/Private/HUD.cpp b/Engine/Source/Runtime/Engine/Private/HUD.cpp
--- a/Engine/Source/Runtime/Engine/Private/HUD.cpp
+++ b/Engine/Source/Runtime/Engine/Private/HUD.cpp
@@ -18,7 +18,6 @@ DEFINE_LOG_CATEGORY_STATIC(LogHUD, Log, All);
 
 #define LOCTEXT_NAMESPACE "HUD"
 
-bool AHUD::bShowDebugForReticleTarget = false;
 FOnShowDebugInfo AHUD::OnShowDebugInfo;
 
 
diff --git a/Engine/Source/Runtime/Engine/Private/LODActor.cpp b/Engine/Source/Runtime/Engine/Private/LODActor.cpp
--- a/Engine/Source/Runtime/Engine/Private/LODActor.cpp
+++ b/Engine/Source/Runtime/Engine/Private/LODActor.cpp
@@ -125,6 +125,7 @@ ALODActor::ALODActor(const FObjectInitializer& ObjectInitializer)
 	StaticMeshComponent->CastShadow = bCastsShadow;
 	StaticMeshComponent->bCastStaticShadow = bCastsStaticShadow;
 	StaticMeshComponent->bCastDynamicShadow = bCastsDynamicShadow;
+	StaticMeshComponent->bAllowCullDistanceVolume = false;
 
 	RootComponent = StaticMeshComponent;	
 }
diff --git a/Engine/Source/Runtime/Engine/Private/Level.cpp b/Engine/Source/Runtime/Engine/Private/Level.cpp
--- a/Engine/Source/Runtime/Engine/Private/Level.cpp
+++ b/Engine/Source/Runtime/Engine/Private/Level.cpp
@@ -602,7 +602,7 @@ void ULevel::PostLoad()
 	}
 
 #if WITH_EDITOR
-	if (!(GetOutermost()->PackageFlags & PKG_PlayInEditor))
+	if (!GetOutermost()->HasAnyPackageFlags(PKG_PlayInEditor))
 	{
 		// Rename the LevelScriptBlueprint after the outer world.
 		UWorld* OuterWorld = Cast<UWorld>(GetOuter());
@@ -822,9 +822,10 @@ void ULevel::IncrementalUpdateComponents(int32 NumComponentsToUpdate, bool bReru
 			// Don't rerun construction scripts until after all actors' components have been registered.  This
 			// is necessary because child attachment lists are populated during registration, and running construction
 			// scripts requires that the attachments are correctly initialized.
-			for (AActor* Actor : Actors)
+			// Don't use ranged for as construction scripts can manipulate the actor array
+			for (int32 ActorIndex = 0; ActorIndex < Actors.Num(); ++ActorIndex)
 			{
-				if (Actor)
+				if (AActor* Actor = Actors[ActorIndex])
 				{
 #if PERF_TRACK_DETAILED_ASYNC_STATS
 					FScopeCycleCounterUObject ContextScope(Actor);
diff --git a/Engine/Source/Runtime/Engine/Private/LevelStreaming.cpp b/Engine/Source/Runtime/Engine/Private/LevelStreaming.cpp
--- a/Engine/Source/Runtime/Engine/Private/LevelStreaming.cpp
+++ b/Engine/Source/Runtime/Engine/Private/LevelStreaming.cpp
@@ -196,7 +196,7 @@ void ULevelStreaming::PostLoad()
 {
 	Super::PostLoad();
 
-	const bool PIESession = GetWorld()->WorldType == EWorldType::PIE || (GetOutermost()->PackageFlags & PKG_PlayInEditor) != 0;
+	const bool PIESession = GetWorld()->WorldType == EWorldType::PIE || GetOutermost()->HasAnyPackageFlags(PKG_PlayInEditor);
 
 	// If this streaming level was saved with a short package name, try to convert it to a long package name
 	if ( !PIESession && PackageName_DEPRECATED != NAME_None )
@@ -215,12 +215,9 @@ void ULevelStreaming::PostLoad()
 		}
 	}
 
-	if ( !WorldAsset.IsNull() )
+	if ( !PIESession && !WorldAsset.IsNull() )
 	{
 		const FString WorldPackageName = GetWorldAssetPackageName();
-		CachedWorldAssetPackageFName = FName(*WorldPackageName);
-		if ( !PIESession )
-		{
 		if (FPackageName::DoesPackageExist(WorldPackageName) == false)
 		{
 			UE_LOG(LogLevelStreaming, Display, TEXT("Failed to find streaming level package file: %s. This streaming level may not load or save properly."), *WorldPackageName);
@@ -238,7 +235,6 @@ void ULevelStreaming::PostLoad()
 #endif // WITH_EDITOR
 		}
 	}
-	}
 
 	if (GetLinkerUE4Version() < VER_UE4_LEVEL_STREAMING_DRAW_COLOR_TYPE_CHANGE)
 	{
@@ -267,7 +263,7 @@ void ULevelStreaming::Serialize( FArchive& Ar )
 	
 	if (Ar.IsLoading())
 	{
-		if ((GetOutermost()->PackageFlags & PKG_PlayInEditor) != 0 && GetOutermost()->PIEInstanceID != INDEX_NONE)
+		if (GetOutermost()->HasAnyPackageFlags(PKG_PlayInEditor) && GetOutermost()->PIEInstanceID != INDEX_NONE)
 		{
 			RenameForPIE(GetOutermost()->PIEInstanceID);
 		}
@@ -415,7 +411,7 @@ bool ULevelStreaming::RequestLevel(UWorld* PersistentWorld, bool bAllowLevelLoad
 	if ( PersistentWorld->IsPlayInEditor() )
 	{
 #if WITH_EDITOR
-		if ((PersistentWorld->GetOutermost()->PackageFlags & PKG_PlayInEditor) != 0)
+		if (PersistentWorld->GetOutermost()->HasAnyPackageFlags(PKG_PlayInEditor))
 		{
 			PackageFlags |= PKG_PlayInEditor;
 		}
@@ -762,17 +758,22 @@ void ULevelStreaming::BroadcastLevelVisibleStatus(UWorld* PersistentWorld, FName
 void ULevelStreaming::SetWorldAsset(const TAssetPtr<UWorld>& NewWorldAsset)
 {
 	WorldAsset = NewWorldAsset;
-	CachedWorldAssetPackageFName = FName(*GetWorldAssetPackageName());
+	bHasCachedWorldAssetPackageFName = false;
 }
 
 FString ULevelStreaming::GetWorldAssetPackageName() const
 {
-	const FString WorldAssetPath = WorldAsset.ToStringReference().ToString();
-	return FPackageName::ObjectPathToPackageName(WorldAssetPath);
+	return GetWorldAssetPackageFName().ToString();
 }
 
 FName ULevelStreaming::GetWorldAssetPackageFName() const
 {
+	if (!bHasCachedWorldAssetPackageFName)
+	{
+		ULevelStreaming* MutableThis = const_cast<ULevelStreaming*>(this);
+		MutableThis->CachedWorldAssetPackageFName = FName(*FPackageName::ObjectPathToPackageName(WorldAsset.ToString()));
+		MutableThis->bHasCachedWorldAssetPackageFName = true;
+	}
 	return CachedWorldAssetPackageFName;
 }
 
@@ -886,7 +887,7 @@ void ULevelStreaming::PostEditChangeProperty(FPropertyChangedEvent& PropertyChan
 		}
 		else if (PropertyName == GET_MEMBER_NAME_CHECKED(ULevelStreaming, WorldAsset))
 		{
-			CachedWorldAssetPackageFName = FName(*GetWorldAssetPackageName());
+			bHasCachedWorldAssetPackageFName = false;
 		}
 	}
 
diff --git a/Engine/Source/Runtime/Engine/Private/LevelTick.cpp b/Engine/Source/Runtime/Engine/Private/LevelTick.cpp
--- a/Engine/Source/Runtime/Engine/Private/LevelTick.cpp
+++ b/Engine/Source/Runtime/Engine/Private/LevelTick.cpp
@@ -1011,7 +1011,6 @@ static TAutoConsoleVariable<int32> CVarStressTestGCWhileStreaming(
 #endif
 
 DECLARE_CYCLE_STAT(TEXT("TG_PrePhysics"), STAT_TG_PrePhysics, STATGROUP_TickGroups);
-DECLARE_CYCLE_STAT(TEXT("TG_DuringAnimation"), STAT_TG_DuringAnimation, STATGROUP_TickGroups);
 DECLARE_CYCLE_STAT(TEXT("TG_StartPhysics"), STAT_TG_StartPhysics, STATGROUP_TickGroups);
 DECLARE_CYCLE_STAT(TEXT("Start TG_DuringPhysics"), STAT_TG_DuringPhysics, STATGROUP_TickGroups);
 DECLARE_CYCLE_STAT(TEXT("TG_EndPhysics"), STAT_TG_EndPhysics, STATGROUP_TickGroups);
@@ -1171,10 +1170,6 @@ void UWorld::Tick( ELevelTick TickType, float DeltaSeconds )
 			SCOPE_CYCLE_COUNTER(STAT_TG_PrePhysics);
 			RunTickGroup(TG_PrePhysics);
 		}
-		{
-			SCOPE_CYCLE_COUNTER(STAT_TG_DuringAnimation);
-			RunTickGroup(TG_DuringAnimation);		
-		}
         bInTick = false;
         EnsureCollisionTreeIsBuilt();
         bInTick = true;
diff --git a/Engine/Source/Runtime/Engine/Private/Materials/MaterialShared.cpp b/Engine/Source/Runtime/Engine/Private/Materials/MaterialShared.cpp
--- a/Engine/Source/Runtime/Engine/Private/Materials/MaterialShared.cpp
+++ b/Engine/Source/Runtime/Engine/Private/Materials/MaterialShared.cpp
@@ -2371,7 +2371,7 @@ int32 UMaterialInterface::CompileProperty(FMaterialCompiler* Compiler, EMaterial
 	}
 }
 
-void UMaterialInterface::AnalyzeMaterialProperty(EMaterialProperty InProperty, int32& OutNumTextureCoordinates, bool& OutUseVertexColor)
+void UMaterialInterface::AnalyzeMaterialProperty(EMaterialProperty InProperty, int32& OutNumTextureCoordinates, bool& bOutRequiresVertexData)
 {
 #if WITH_EDITORONLY_DATA
 	// FHLSLMaterialTranslator collects all required information during translation, but these data are protected. Needs to
@@ -2390,6 +2390,16 @@ void UMaterialInterface::AnalyzeMaterialProperty(EMaterialProperty InProperty, i
 		{
 			return bUsesVertexColor;
 		}
+
+		bool UsesTransformVector() const
+		{
+			return bUsesTransformVector;
+		}
+
+		bool UsesWorldPositionExcludingShaderOffsets() const
+		{
+			return bNeedsWorldPositionExcludingShaderOffsets;
+		}
 	};
 
 	FMaterialCompilationOutput TempOutput;
@@ -2402,7 +2412,7 @@ void UMaterialInterface::AnalyzeMaterialProperty(EMaterialProperty InProperty, i
 	CompileProperty(&MaterialTranslator, InProperty);
 	// Request data from translator
 	OutNumTextureCoordinates = MaterialTranslator.GetTextureCoordsCount();
-	OutUseVertexColor = MaterialTranslator.UsesVertexColor();
+	bOutRequiresVertexData = MaterialTranslator.UsesVertexColor() || MaterialTranslator.UsesTransformVector() || MaterialTranslator.UsesWorldPositionExcludingShaderOffsets();
 #endif
 }
 
diff --git a/Engine/Source/Runtime/Engine/Private/Particles/ParticleComponents.cpp b/Engine/Source/Runtime/Engine/Private/Particles/ParticleComponents.cpp
--- a/Engine/Source/Runtime/Engine/Private/Particles/ParticleComponents.cpp
+++ b/Engine/Source/Runtime/Engine/Private/Particles/ParticleComponents.cpp
@@ -6157,32 +6157,40 @@ AEmitterCameraLensEffectBase::AEmitterCameraLensEffectBase(const FObjectInitiali
 		.DoNotCreateDefaultSubobject(TEXT("ArrowComponent0"))
 	)
 {
-	DistFromCamera = 90.0f;
 	InitialLifeSpan = 10.0f;
 	BaseFOV = 80.0f;
 	bDestroyOnSystemFinish = true;
 
+	// default transform is a 180 yaw to flip the system around to face the camera
+	// and 90 units pushed out
+	// (we assume it by default that the effect was authored facing down the +X, due to legacy reasons)
+	RelativeTransform = FTransform( 
+		FRotator(0.f, 180.f, 0.f),
+		FVector(90.f, 0.f, 0.f)
+		);
+
 	GetParticleSystemComponent()->bOnlyOwnerSee = true;
 	GetParticleSystemComponent()->SecondsBeforeInactive = 0.0f;
+
+	// this property is deprecated, give it the sentinel value to indicate it doesn't need to be migrated
+	DistFromCamera_DEPRECATED = TNumericLimits<float>::Max();
 }
 
 void AEmitterCameraLensEffectBase::UpdateLocation(const FVector& CamLoc, const FRotator& CamRot, float CamFOVDeg)
 {
-	FRotationMatrix M(CamRot);
-
-	// the particle is FACING X being parallel to the Y axis.  So we just flip the entire thing to face toward the player who is looking down X
-	const FVector& X = M.GetScaledAxis( EAxis::X );
-	M.SetAxis(0, -X);
-	M.SetAxis(1, -M.GetScaledAxis( EAxis::Y ));
-
-	const FRotator& NewRot = M.Rotator();
-
+	// adjust for FOV
 	// base dist uses BaseFOV which is set on the indiv camera lens effect class
-	const float DistAdjustedForFOV = DistFromCamera * ( FMath::Tan(float(BaseFOV*0.5f*PI/180.f)) / FMath::Tan(float(CamFOVDeg*0.5f*PI/180.f)) );
+	FTransform RelativeTransformAdjustedForFOV = RelativeTransform;
+	FVector AdjustedRelativeLoc = RelativeTransformAdjustedForFOV.GetLocation();
+	AdjustedRelativeLoc.X *= FMath::Tan(BaseFOV*0.5f*PI / 180.f) / FMath::Tan(CamFOVDeg*0.5f*PI / 180.f);
+	RelativeTransformAdjustedForFOV.SetLocation(AdjustedRelativeLoc);
+
+	FTransform const CameraToWorld(CamRot, CamLoc);
 
-	//UE_LOG(LogParticles, Warning, TEXT("DistAdjustedForFOV: %f  BaseFOV: %f  CamFOVDeg: %f"), DistAdjustedForFOV, BaseFOV, CamFOVDeg );
+	// RelativeTransform is "effect to camera"
+	FTransform const EffectToWorld = RelativeTransformAdjustedForFOV * CameraToWorld;
 
-	SetActorLocationAndRotation( CamLoc + X * DistAdjustedForFOV, NewRot, false );
+	SetActorTransform(EffectToWorld);
 }
 
 void AEmitterCameraLensEffectBase::EndPlay(const EEndPlayReason::Type EndPlayReason)
@@ -6208,10 +6216,28 @@ void AEmitterCameraLensEffectBase::PostInitializeComponents()
 	ActivateLensEffect();
 }
 
+void AEmitterCameraLensEffectBase::PostLoad()
+{
+	Super::PostLoad();
+
+	// using TNumericLimits<float>::Max() as a sentinel value to indicate this deprecated data has been 
+	// migrated to the new format
+	if (DistFromCamera_DEPRECATED != TNumericLimits<float>::Max())
+	{
+		// copy old data into the new transform
+		FVector Loc = RelativeTransform.GetLocation();
+		Loc.X = DistFromCamera_DEPRECATED;
+		RelativeTransform.SetLocation(Loc);
+
+		// don't copy again (just in case this gets saved, which is shouldn't)
+		DistFromCamera_DEPRECATED = TNumericLimits<float>::Max();
+	}
+}
+
 void AEmitterCameraLensEffectBase::ActivateLensEffect()
 {
 	// only play the camera effect on clients
-	UWorld* World = GetWorld();
+	UWorld const* const World = GetWorld();
 	check(World);
 	if( GetNetMode() != NM_DedicatedServer )
 	{
diff --git a/Engine/Source/Runtime/Engine/Private/PhysicsEngine/BodySetup.cpp b/Engine/Source/Runtime/Engine/Private/PhysicsEngine/BodySetup.cpp
--- a/Engine/Source/Runtime/Engine/Private/PhysicsEngine/BodySetup.cpp
+++ b/Engine/Source/Runtime/Engine/Private/PhysicsEngine/BodySetup.cpp
@@ -443,7 +443,7 @@ public:
 			// first apply the scale first 
 			float Radius = FMath::Max(SphylElem.Radius * ScaleRadius, 0.1f);
 			float Length = SphylElem.Length + SphylElem.Radius * 2.f;
-			float HalfLength = Length * ScaleLength * 0.5f;
+			float HalfLength = FMath::Max(Length * ScaleLength * 0.5f, 0.1f);
 			Radius = FMath::Clamp(Radius, 0.1f, HalfLength);	//radius is capped by half length
 			float HalfHeight = HalfLength - Radius;
 			HalfHeight = FMath::Max(0.1f, HalfHeight);
diff --git a/Engine/Source/Runtime/Engine/Private/Selection.cpp b/Engine/Source/Runtime/Engine/Private/Selection.cpp
--- a/Engine/Source/Runtime/Engine/Private/Selection.cpp
+++ b/Engine/Source/Runtime/Engine/Private/Selection.cpp
@@ -26,7 +26,7 @@ void USelection::Select(UObject* InObject)
 	check( InObject );
 
 	// Warn if we attempt to select a PIE object.
-	if ( InObject->GetOutermost()->PackageFlags & PKG_PlayInEditor )
+	if ( InObject->GetOutermost()->HasAnyPackageFlags(PKG_PlayInEditor) )
 	{
 		UE_LOG(LogSelection, Warning, TEXT("PIE object was selected: \"%s\""), *InObject->GetFullName() );
 	}
@@ -170,7 +170,7 @@ bool USelection::Modify(bool bAlwaysMarkDirty/* =true */)
 	for (auto ObjectPtr : SelectedObjects)
 	{
 		UObject* Object = ObjectPtr.Get();
-		if (Object && (Object->GetOutermost()->PackageFlags & (PKG_PlayInEditor | PKG_ContainsScript | PKG_CompiledIn)) != 0)
+		if (Object && Object->GetOutermost()->HasAnyPackageFlags(PKG_PlayInEditor | PKG_ContainsScript | PKG_CompiledIn))
 		{
 			return false;
 		}
diff --git a/Engine/Source/Runtime/Engine/Private/SingleAnimationPlayData.cpp b/Engine/Source/Runtime/Engine/Private/SingleAnimationPlayData.cpp
--- a/Engine/Source/Runtime/Engine/Private/SingleAnimationPlayData.cpp
+++ b/Engine/Source/Runtime/Engine/Private/SingleAnimationPlayData.cpp
@@ -18,11 +18,11 @@ void FSingleAnimationPlayData::Initialize(UAnimSingleNodeInstance* Instance)
 
 void FSingleAnimationPlayData::PopulateFrom(UAnimSingleNodeInstance* Instance)
 {
-	AnimToPlay = Instance->CurrentAsset;
-	SavedPosition = Instance->CurrentTime;
-	SavedPlayRate = Instance->PlayRate;
-	bSavedPlaying = Instance->bPlaying;
-	bSavedLooping = Instance->bLooping;
+	AnimToPlay = Instance->GetCurrentAsset();
+	SavedPosition = Instance->GetCurrentTime();
+	SavedPlayRate = Instance->GetPlayRate();
+	bSavedPlaying = Instance->IsPlaying();
+	bSavedLooping = Instance->IsLooping();
 }
 
 void FSingleAnimationPlayData::ValidatePosition()
diff --git a/Engine/Source/Runtime/Engine/Private/SkeletalMeshComponentPhysics.cpp b/Engine/Source/Runtime/Engine/Private/SkeletalMeshComponentPhysics.cpp
--- a/Engine/Source/Runtime/Engine/Private/SkeletalMeshComponentPhysics.cpp
+++ b/Engine/Source/Runtime/Engine/Private/SkeletalMeshComponentPhysics.cpp
@@ -463,6 +463,11 @@ void USkeletalMeshComponent::SetSimulatePhysics(bool bSimulate)
 		}
 	}
 
+	if(IsSimulatingPhysics())
+	{
+		SetRootBodyIndex(RootBodyData.BodyIndex);	//Update the root body data cache in case animation has moved root body relative to root joint
+	}
+
 	UpdatePreClothTickRegisteredState();
 }
 
@@ -870,7 +875,7 @@ void USkeletalMeshComponent::InitArticulated(FPhysScene* PhysScene)
 				Body->CreateDOFLock();
 
 				// Set to sleep if necessary
-				if (bShouldSleep)
+				if (bShouldSleep  && Body->GetPxRigidDynamic_AssumesLocked())
 				{
 					Body->GetPxRigidDynamic_AssumesLocked()->putToSleep();
 				}
@@ -1051,6 +1056,8 @@ void USkeletalMeshComponent::SetAllBodiesSimulatePhysics(bool bNewSimulate)
 		Bodies[i]->SetInstanceSimulatePhysics(bNewSimulate);
 	}
 
+	SetRootBodyIndex(RootBodyData.BodyIndex);	//Update the root body data cache in case animation has moved root body relative to root joint
+
 	UpdatePreClothTickRegisteredState();
 }
 
@@ -1096,6 +1103,11 @@ void USkeletalMeshComponent::SetAllBodiesBelowSimulatePhysics( const FName& InBo
 		Bodies[BodyIndices[i]]->SetInstanceSimulatePhysics(bNewSimulate);
 	}
 
+	if (IsSimulatingPhysics())
+	{
+		SetRootBodyIndex(RootBodyData.BodyIndex);	//Update the root body data cache in case animation has moved root body relative to root joint
+	}
+
 	UpdatePreClothTickRegisteredState();
 }
 
@@ -4230,6 +4242,37 @@ void USkeletalMeshComponent::DrawClothingTangents(FPrimitiveDrawInterface* PDI)
 #endif // #if WITH_APEX_CLOTHING
 }
 
+
+void DrawWireTaperedCapsuleSleeve(const FSphere& S0, const FSphere& S1, TFunction<void(const FVector&, const FVector&)> DrawLine, bool TrueHull = true)
+{
+	// Draws just the sides of a tapered capsule specified by provided Spheres that can have different radii.  Does not draw the spheres, just the sleeve.
+	// If Using TrueHull (above console var) Extent geometry endpoints not necessarily coplanar with sphere origins (uses hull horizon)
+	// Otherwise uses the great-circle cap assumption.
+	const float AngleIncrement = 30.0f;   // if parameter added for number of sides, then set this to be:  = 360.0f / NumSides; 
+	FVector Separation = S1.Center - S0.Center;
+	float Distance = Separation.Size();
+	if (Separation.IsNearlyZero() || Distance <= FMath::Abs(S0.W - S1.W))
+	{
+		return;
+	}
+	FQuat CapsuleOrientation = FQuat::FindBetween(FVector(0, 0, 1), Separation.GetSafeNormal());
+	float OffsetZ = TrueHull ? -(S1.W - S0.W) / Distance : 0.0f;
+	float ScaleXY = FMath::Sqrt(1.0f - FMath::Square(OffsetZ));
+	FVector VertexPrevious = CapsuleOrientation.RotateVector(FVector(ScaleXY, 0, OffsetZ));
+	for (float Angle = AngleIncrement; Angle <= 360.0f; Angle += AngleIncrement)  // iterate over unit circle about capsule's major axis (which is orientation.AxisZ)
+	{
+		FVector VertexCurrent = CapsuleOrientation.RotateVector(FVector(FMath::Cos(FMath::DegreesToRadians(Angle))*ScaleXY, FMath::Sin(FMath::DegreesToRadians(Angle))*ScaleXY, OffsetZ));
+		DrawLine(S0.Center + VertexCurrent  * S0.W, S1.Center + VertexCurrent * S1.W);  // capsule side segment between spheres
+		DrawLine(S0.Center + VertexPrevious * S0.W, S0.Center + VertexCurrent * S0.W);  // cap-circle segment on sphere S0
+		DrawLine(S1.Center + VertexPrevious * S1.W, S1.Center + VertexCurrent * S1.W);  // cap-circle segment on sphere S1
+		VertexPrevious = VertexCurrent;
+	}
+}
+
+TAutoConsoleVariable<int32> CVarEnableTaperedCapsulesTrueHull(TEXT("p.TaperedCapsulesTrueHull"), 1, TEXT("If 1, debug rendering of tapered capsules will use correct horizon to generate shaft geometry rather than using great-arc assumption where cap is coplanar with sphere origins. "));
+
+
+
 void USkeletalMeshComponent::DrawClothingCollisionVolumes(FPrimitiveDrawInterface* PDI)
 {
 #if WITH_APEX_CLOTHING
@@ -4388,10 +4431,10 @@ void USkeletalMeshComponent::DrawClothingCollisionVolumes(FPrimitiveDrawInterfac
 
 		for(int32 i=0; i<NumConnections; i+=2)
 		{
-			uint16 Index1 = Connections[i];
-			uint16 Index2 = Connections[i+1];
-
-			DrawDebugLine(GetWorld(), SpherePositions[Index1], SpherePositions[Index2], FColor::Magenta, false, -1.0f, SDPG_Foreground);
+			DrawWireTaperedCapsuleSleeve(FSphere(SpherePositions[Connections[i + 0]], Spheres[Connections[i + 0]].Radius),
+			                             FSphere(SpherePositions[Connections[i + 1]], Spheres[Connections[i + 1]].Radius),
+			                             [&](const FVector& V0, const FVector& V1)->void{PDI->DrawLine(V0, V1, Colors[AssetIdx % 3], SDPG_World); },
+			                             (CVarEnableTaperedCapsulesTrueHull.GetValueOnAnyThread()) != 0);
 		}
 	}
 #endif // #if WITH_APEX_CLOTHING
diff --git a/Engine/Source/Runtime/Engine/Private/Slate/SlateSoundDevice.cpp b/Engine/Source/Runtime/Engine/Private/Slate/SlateSoundDevice.cpp
--- a/Engine/Source/Runtime/Engine/Private/Slate/SlateSoundDevice.cpp
+++ b/Engine/Source/Runtime/Engine/Private/Slate/SlateSoundDevice.cpp
@@ -24,6 +24,7 @@ void FSlateSoundDevice::PlaySound(const FSlateSound& Sound, int32 UserIndex) con
 				NewActiveSound.Sound = SoundResource;
 				NewActiveSound.bIsUISound = true;
 				NewActiveSound.UserIndex = UserIndex;
+				NewActiveSound.Priority = NewActiveSound.Sound->Priority;
 
 				AudioDevice->AddNewActiveSound(NewActiveSound);
 			}
diff --git a/Engine/Source/Runtime/Engine/Private/SoundAttenuation.cpp b/Engine/Source/Runtime/Engine/Private/SoundAttenuation.cpp
--- a/Engine/Source/Runtime/Engine/Private/SoundAttenuation.cpp
+++ b/Engine/Source/Runtime/Engine/Private/SoundAttenuation.cpp
@@ -77,7 +77,7 @@ float FAttenuationSettings::GetMaxDimension() const
  *
  * @return Attenuation value (between 0.0 and 1.0)
  */
-float FAttenuationSettings::AttenuationEval(const float Distance, const float Falloff) const
+float FAttenuationSettings::AttenuationEval(const float Distance, const float Falloff, const float DistanceScale) const
 {
 	// Clamp the input distance between 0.0f and Falloff. If the Distance
 	// is actually less than the min value, it will use the min-value of the algorithm/curve
@@ -86,6 +86,7 @@ float FAttenuationSettings::AttenuationEval(const float Distance, const float Fa
 	// which could be 1.0 (and not 0.0f).
 
 	float DistanceCopy = FMath::Clamp(Distance, 0.0f, Falloff);
+	DistanceCopy *= DistanceScale;
 
 	float Result = 0.0f;
 	switch(DistanceAlgorithm)
@@ -132,18 +133,18 @@ float FAttenuationSettings::AttenuationEval(const float Distance, const float Fa
 	return FMath::Clamp(Result, 0.0f, 1.0f);
 }
 
-float FAttenuationSettings::AttenuationEvalBox(const FTransform& SoundTransform, const FVector ListenerLocation) const
+float FAttenuationSettings::AttenuationEvalBox(const FTransform& SoundTransform, const FVector ListenerLocation, const float DistanceScale) const
 {
 	const float DistanceSq = ComputeSquaredDistanceFromBoxToPoint(-AttenuationShapeExtents, AttenuationShapeExtents,SoundTransform.InverseTransformPositionNoScale(ListenerLocation));
 	if (DistanceSq < FalloffDistance * FalloffDistance)
 	{ 
-		return AttenuationEval(FMath::Sqrt(DistanceSq), FalloffDistance);
+		return AttenuationEval(FMath::Sqrt(DistanceSq), FalloffDistance, DistanceScale);
 	}
 
 	return 0.f;
 }
 
-float FAttenuationSettings::AttenuationEvalCapsule(const FTransform& SoundTransform, const FVector ListenerLocation) const
+float FAttenuationSettings::AttenuationEvalCapsule(const FTransform& SoundTransform, const FVector ListenerLocation, const float DistanceScale) const
 {
 	float Distance = 0.f;
 	const float CapsuleHalfHeight = AttenuationShapeExtents.X;
@@ -163,10 +164,10 @@ float FAttenuationSettings::AttenuationEvalCapsule(const FTransform& SoundTransf
 		Distance = FMath::PointDistToSegment(ListenerLocation, StartPoint, EndPoint) - CapsuleRadius;
 	}
 
-	return AttenuationEval(Distance, FalloffDistance);
+	return AttenuationEval(Distance, FalloffDistance, DistanceScale);
 }
 
-float FAttenuationSettings::AttenuationEvalCone(const FTransform& SoundTransform, const FVector ListenerLocation) const
+float FAttenuationSettings::AttenuationEvalCone(const FTransform& SoundTransform, const FVector ListenerLocation, const float DistanceScale) const
 {
 	const FVector SoundForward = SoundTransform.GetUnitAxis( EAxis::X );
 
@@ -175,12 +176,12 @@ float FAttenuationSettings::AttenuationEvalCone(const FTransform& SoundTransform
 	const FVector Origin = SoundTransform.GetTranslation() - (SoundForward * ConeOffset);
 
 	const float Distance = FMath::Max(FVector::Dist( Origin, ListenerLocation ) - AttenuationShapeExtents.X, 0.f);
-	VolumeMultiplier *= AttenuationEval(Distance, FalloffDistance);
+	VolumeMultiplier *= AttenuationEval(Distance, FalloffDistance, DistanceScale);
 
 	if (VolumeMultiplier > 0.f)
 	{
 		const float theta = FMath::RadiansToDegrees(fabsf(FMath::Acos( FVector::DotProduct(SoundForward, (ListenerLocation - Origin).GetSafeNormal()))));
-		VolumeMultiplier *= AttenuationEval(theta - AttenuationShapeExtents.Y, AttenuationShapeExtents.Z);
+		VolumeMultiplier *= AttenuationEval(theta - AttenuationShapeExtents.Y, AttenuationShapeExtents.Z, 1.0f);
 	}
 
 	return VolumeMultiplier;
@@ -200,60 +201,14 @@ bool FAttenuationSettings::operator==(const FAttenuationSettings& Other) const
 			&& LPFRadiusMax			    == Other.LPFRadiusMax
 			&& FalloffDistance		    == Other.FalloffDistance
 			&& AttenuationShapeExtents	== Other.AttenuationShapeExtents
-			&& SpatializationAlgorithm  == Other.SpatializationAlgorithm);
-}
-
-void FAttenuationSettings::ApplyAttenuation( const FTransform& SoundTransform, const FVector ListenerLocation, float& Volume, float& HighFrequencyGain ) const
-{
-	// Attenuate the volume based on the model
-	if( bAttenuate )
-	{
-		switch(AttenuationShape)
-		{
-		case EAttenuationShape::Sphere:
-			{
-				const float Distance = FMath::Max(FVector::Dist( SoundTransform.GetTranslation(), ListenerLocation ) - AttenuationShapeExtents.X, 0.f);
-				Volume *= AttenuationEval(Distance, FalloffDistance);
-				break;
-			}
-
-		case EAttenuationShape::Box:
-			Volume *= AttenuationEvalBox(SoundTransform, ListenerLocation);
-			break;
-
-		case EAttenuationShape::Capsule:
-			Volume *= AttenuationEvalCapsule(SoundTransform, ListenerLocation);
-			break;
-
-		case EAttenuationShape::Cone:
-			Volume *= AttenuationEvalCone(SoundTransform, ListenerLocation);
-			break;
-
-		default:
-			check(false);
-		}
-
-	}
-
-	// Attenuate with the low pass filter if necessary
-	if( bAttenuateWithLPF )
-	{
-		const float Distance = FMath::Max(FVector::Dist( SoundTransform.GetTranslation(), ListenerLocation ) - AttenuationShapeExtents.X, 0.f);
-
-		if( Distance >= LPFRadiusMax )
-		{
-			HighFrequencyGain = 0.0f;
-		}
-		// UsedLPFMinRadius is the point at which to start applying the low pass filter
-		else if( Distance > LPFRadiusMin )
-		{
-			HighFrequencyGain = 1.0f - ( ( Distance - LPFRadiusMin ) / ( LPFRadiusMax - LPFRadiusMin ) );
-		}
-		else
-		{
-			HighFrequencyGain = 1.0f;
-		}
-	}
+			&& SpatializationAlgorithm  == Other.SpatializationAlgorithm
+			&& bEnableListenerFocus		== Other.bEnableListenerFocus
+			&& FocusAzimuth				== Other.FocusAzimuth
+			&& NonFocusAzimuth			== Other.NonFocusAzimuth
+			&& FocusDistanceScale		== Other.FocusDistanceScale
+			&& FocusPriorityScale		== Other.FocusPriorityScale
+			&& NonFocusPriorityScale	== Other.NonFocusPriorityScale
+			&& NonFocusVolumeAttenuation == Other.NonFocusVolumeAttenuation);
 }
 
 void FAttenuationSettings::CollectAttenuationShapesForVisualization(TMultiMap<EAttenuationShape::Type, AttenuationShapeDetails>& ShapeDetailsMap) const
diff --git a/Engine/Source/Runtime/Engine/Private/SoundBase.cpp b/Engine/Source/Runtime/Engine/Private/SoundBase.cpp
--- a/Engine/Source/Runtime/Engine/Private/SoundBase.cpp
+++ b/Engine/Source/Runtime/Engine/Private/SoundBase.cpp
@@ -7,6 +7,9 @@
 
 USoundBase::USoundBase(const FObjectInitializer& ObjectInitializer)
 	: Super(ObjectInitializer)
+	, bIgnoreFocus(false)
+	, FocusDistanceScale(1.0f)
+	, Priority(1.0f)
 {
 	MaxConcurrentPlayCount_DEPRECATED = 16;
 }
@@ -57,6 +60,13 @@ bool USoundBase::IsAudibleSimple(class FAudioDevice* AudioDevice, const FVector
 
 	// Is this SourceActor within the MaxAudibleDistance of any of the listeners?
 	float MaxAudibleDistance = InAttenuationSettings != nullptr ? InAttenuationSettings->Attenuation.GetMaxDimension() : GetMaxAudibleDistance();
+
+	// Scale the max audible distance by the distance scale 
+	if (!bIgnoreFocus)
+	{
+		MaxAudibleDistance *= FocusDistanceScale;
+	}
+
 	return AudioDevice->LocationIsAudible(Location, MaxAudibleDistance);
 }
 
@@ -68,7 +78,13 @@ bool USoundBase::IsAudible( const FVector &SourceLocation, const FVector &Listen
 	// Account for any portals
 	const FVector ModifiedSourceLocation = SourceLocation;
 
-	const float MaxDist = GetMaxAudibleDistance();
+	float MaxDist = GetMaxAudibleDistance();
+
+	if (!bIgnoreFocus)
+	{
+		MaxDist *= FocusDistanceScale;
+	}
+
 	if( MaxDist * MaxDist >= ( ListenerLocation - ModifiedSourceLocation ).SizeSquared() )
 	{
 		// Can't line check through portals
@@ -120,6 +136,11 @@ const FSoundConcurrencySettings* USoundBase::GetSoundConcurrencySettingsToApply(
 	return nullptr;
 }
 
+float USoundBase::GetPriority() const
+{
+	return Priority;
+}
+
 uint32 USoundBase::GetSoundConcurrencyObjectID() const
 {
 	if (SoundConcurrencySettings != nullptr && !bOverrideConcurrency)
diff --git a/Engine/Source/Runtime/Engine/Private/SoundConcurrency.cpp b/Engine/Source/Runtime/Engine/Private/SoundConcurrency.cpp
--- a/Engine/Source/Runtime/Engine/Private/SoundConcurrency.cpp
+++ b/Engine/Source/Runtime/Engine/Private/SoundConcurrency.cpp
@@ -236,14 +236,14 @@ FActiveSound* FSoundConcurrencyManager::ResolveConcurrency(const FActiveSound& N
 				// Find the current lowest priority sound.
 				for (FActiveSound* CurrSound : ActiveSounds)
 				{
-					if (SoundToStop == nullptr || CurrSound->VolumeWeightedPriorityScale < SoundToStop->VolumeWeightedPriorityScale)
+					if (SoundToStop == nullptr || CurrSound->Priority < SoundToStop->Priority)
 					{
 						SoundToStop = CurrSound;
 					}
 				}
 
 				// Only stop any sounds if the *lowest* priority is lower than the incoming NewActiveSound
-				if (SoundToStop->VolumeWeightedPriorityScale >= NewActiveSound.VolumeWeightedPriorityScale)
+				if (SoundToStop->Priority >= NewActiveSound.Priority)
 				{
 					SoundToStop = nullptr;
 				}
diff --git a/Engine/Source/Runtime/Engine/Private/SoundCue.cpp b/Engine/Source/Runtime/Engine/Private/SoundCue.cpp
--- a/Engine/Source/Runtime/Engine/Private/SoundCue.cpp
+++ b/Engine/Source/Runtime/Engine/Private/SoundCue.cpp
@@ -98,7 +98,7 @@ void USoundCue::PostLoad()
 		else
 		{
 			// we should have a soundcuegraph unless we are contained in a package which is missing editor only data
-			check( GetOutermost()->PackageFlags & PKG_FilterEditorOnly );
+			check( GetOutermost()->HasAnyPackageFlags(PKG_FilterEditorOnly) );
 		}
 
 		// Always load all sound waves in the editor
diff --git a/Engine/Source/Runtime/Engine/Private/SoundNodeAttenuation.cpp b/Engine/Source/Runtime/Engine/Private/SoundNodeAttenuation.cpp
--- a/Engine/Source/Runtime/Engine/Private/SoundNodeAttenuation.cpp
+++ b/Engine/Source/Runtime/Engine/Private/SoundNodeAttenuation.cpp
@@ -54,11 +54,7 @@ void USoundNodeAttenuation::ParseNodes( FAudioDevice* AudioDevice, const UPTRINT
 	if (Settings)
 	{
 		const FListener& Listener = AudioDevice->Listeners[0];
-		Settings->ApplyAttenuation(UpdatedParseParams.Transform, Listener.Transform.GetTranslation(), UpdatedParseParams.Volume, UpdatedParseParams.HighFrequencyGain);
-		UpdatedParseParams.OmniRadius = Settings->OmniRadius;
-		UpdatedParseParams.StereoSpread = Settings->StereoSpread;
-		UpdatedParseParams.SpatializationAlgorithm = Settings->SpatializationAlgorithm;
-		UpdatedParseParams.bUseSpatialization |= Settings->bSpatialize;
+		ActiveSound.ApplyAttenuation(UpdatedParseParams, Listener);
 	}
 	else
 	{
diff --git a/Engine/Source/Runtime/Engine/Private/SoundWave.cpp b/Engine/Source/Runtime/Engine/Private/SoundWave.cpp
--- a/Engine/Source/Runtime/Engine/Private/SoundWave.cpp
+++ b/Engine/Source/Runtime/Engine/Private/SoundWave.cpp
@@ -230,7 +230,7 @@ void USoundWave::Serialize( FArchive& Ar )
 		}
 
 #if WITH_EDITORONLY_DATA	
-		if (Ar.IsLoading() && !Ar.IsTransacting() && !bCooked && !(GetOutermost()->PackageFlags & PKG_ReloadingForCooker))
+		if (Ar.IsLoading() && !Ar.IsTransacting() && !bCooked && !GetOutermost()->HasAnyPackageFlags(PKG_ReloadingForCooker))
 		{
 			BeginCachePlatformData();
 		}
@@ -324,7 +324,7 @@ void USoundWave::PostLoad()
 {
 	Super::PostLoad();
 
-	if (GetOutermost()->PackageFlags & PKG_ReloadingForCooker)
+	if (GetOutermost()->HasAnyPackageFlags(PKG_ReloadingForCooker))
 	{
 		return;
 	}
@@ -652,15 +652,16 @@ void USoundWave::Parse( FAudioDevice* AudioDevice, const UPTRINT NodeWaveInstanc
 			bAlwaysPlay = ActiveSound.bAlwaysPlay;
 		}
 
-		// This is a first-guess at priority, this will later change according to VolumeWeightedPriorityScale
-		WaveInstance->PlayPriority = WaveInstance->Volume + (bAlwaysPlay ? 1.0f : 0.0f) + WaveInstance->RadioFilterVolume;
-
-		// If set to bAlwaysPlay, double the current sound's priority scale. This will still result in a possible 0-priority output if the sound has 0 actual volume
-		WaveInstance->VolumeWeightedPriorityScale = ParseParams.VolumeWeightedPriorityScale;
+		// If set to bAlwaysPlay, increase the current sound's priority scale by 10x. This will still result in a possible 0-priority output if the sound has 0 actual volume
 		if (bAlwaysPlay)
 		{
-			WaveInstance->VolumeWeightedPriorityScale *= 2.0f;
+			WaveInstance->Priority = MAX_FLT;
 		}
+		else
+		{
+			WaveInstance->Priority = ParseParams.Priority;
+		}
+
 		WaveInstance->Location = ParseParams.Transform.GetTranslation();
 		WaveInstance->bIsStarted = true;
 		WaveInstance->bAlreadyNotifiedHook = false;
@@ -680,7 +681,7 @@ void USoundWave::Parse( FAudioDevice* AudioDevice, const UPTRINT NodeWaveInstanc
 		}
 
 		// Don't add wave instances that are not going to be played at this point.
-		if( WaveInstance->PlayPriority > KINDA_SMALL_NUMBER )
+		if( WaveInstance->Volume > KINDA_SMALL_NUMBER )
 		{
 			WaveInstances.Add( WaveInstance );
 		}
diff --git a/Engine/Source/Runtime/Engine/Private/StaticMesh.cpp b/Engine/Source/Runtime/Engine/Private/StaticMesh.cpp
--- a/Engine/Source/Runtime/Engine/Private/StaticMesh.cpp
+++ b/Engine/Source/Runtime/Engine/Private/StaticMesh.cpp
@@ -1201,9 +1201,9 @@ static FString BuildDistanceFieldDerivedDataKey(const FString& InMeshKey)
 
 void FStaticMeshRenderData::Cache(UStaticMesh* Owner, const FStaticMeshLODSettings& LODSettings)
 {
-	if (Owner->GetOutermost()->PackageFlags & PKG_FilterEditorOnly)
+	if (Owner->GetOutermost()->HasAnyPackageFlags(PKG_FilterEditorOnly))
 	{
-		// Don't chache for cooked packages
+		// Don't cache for cooked packages
 		return;
 	}
 
@@ -1760,7 +1760,7 @@ static FStaticMeshRenderData& GetPlatformStaticMeshRenderData(UStaticMesh* Mesh,
 	FString PlatformDerivedDataKey = BuildStaticMeshDerivedDataKey(Mesh, PlatformLODSettings.GetLODGroup(Mesh->LODGroup));
 	FStaticMeshRenderData* PlatformRenderData = Mesh->RenderData;
 
-	if (Mesh->GetOutermost()->PackageFlags & PKG_FilterEditorOnly)
+	if (Mesh->GetOutermost()->HasAnyPackageFlags(PKG_FilterEditorOnly))
 	{
 		check(PlatformRenderData);
 		return *PlatformRenderData;
@@ -2000,7 +2000,7 @@ void UStaticMesh::PostLoad()
 	Super::PostLoad();
 
 #if WITH_EDITOR
-	if (!(GetOutermost()->PackageFlags & PKG_FilterEditorOnly))
+	if (!GetOutermost()->HasAnyPackageFlags(PKG_FilterEditorOnly))
 	{
 		// Needs to happen before 'CacheDerivedData'
 		if (GetLinkerUE4Version() < VER_UE4_BUILD_SCALE_VECTOR)
diff --git a/Engine/Source/Runtime/Engine/Private/StatsRender2.cpp b/Engine/Source/Runtime/Engine/Private/StatsRender2.cpp
--- a/Engine/Source/Runtime/Engine/Private/StatsRender2.cpp
+++ b/Engine/Source/Runtime/Engine/Private/StatsRender2.cpp
@@ -710,62 +710,87 @@ static void RenderGroupedWithHierarchy(const FGameThreadHudData& ViewData, FView
 	// Render all groups.
 	for( int32 GroupIndex = 0; GroupIndex < ViewData.HudGroups.Num(); ++GroupIndex )
 	{
-		// If the stat isn't enabled for this particular viewport, skip
-		FString StatGroupName = ViewData.GroupNames[GroupIndex].ToString();
-		StatGroupName.RemoveFromStart(TEXT("STATGROUP_"), ESearchCase::CaseSensitive);
-		if (!Viewport->GetClient() || !Viewport->GetClient()->IsStatEnabled(StatGroupName))
-		{
-			continue;
-		}
+		const FHudGroup& HudGroup = ViewData.HudGroups[GroupIndex];
+		const bool bBudget = HudGroup.ThreadBudgetMap.Num() > 0;
+		const int32 NumThreadsBreakdown = bBudget ? HudGroup.FlatAggregateThreadBreakdown.Num() : 1;
+		TArray<FName> ThreadNames;
+		HudGroup.FlatAggregateThreadBreakdown.GetKeys(ThreadNames);
 
-		// Render header.
-		const FName& GroupName = ViewData.GroupNames[GroupIndex];
-		const FString& GroupDesc = ViewData.GroupDescriptions[GroupIndex];
-		const FString GroupLongName = !ViewData.RootFilter.IsEmpty() ? FString::Printf(TEXT("%s [%s] ROOT=%s"), *GroupDesc, *GroupName.GetPlainNameString(), *ViewData.RootFilter)
-																	   : FString::Printf(TEXT("%s [%s]"), *GroupDesc, *GroupName.GetPlainNameString());
+		for(int32 ThreadBreakdownIdx = 0; ThreadBreakdownIdx < NumThreadsBreakdown; ++ThreadBreakdownIdx)
+		{
+			// If the stat isn't enabled for this particular viewport, skip
+			FString StatGroupName = ViewData.GroupNames[GroupIndex].ToString();
+			StatGroupName.RemoveFromStart(TEXT("STATGROUP_"), ESearchCase::CaseSensitive);
+			if (!Viewport->GetClient() || !Viewport->GetClient()->IsStatEnabled(StatGroupName))
+			{
+				continue;
+			}
 
-		Canvas->DrawShadowedString(X, Y, *GroupLongName, Globals.StatFont, Globals.GroupColor);
-		Y += Globals.GetFontHeight();
+			// Render header.
+			const FName& GroupName = ViewData.GroupNames[GroupIndex];
+			const FString& GroupDesc = ViewData.GroupDescriptions[GroupIndex];
+			FString GroupLongName = FString::Printf(TEXT("%s [%s]"), *GroupDesc, *GroupName.GetPlainNameString());
+			
+			FName ThreadName;
+			FName ShortThreadName;
+			if(NumThreadsBreakdown)
+			{
+				ThreadName = ThreadNames[ThreadBreakdownIdx];
+				ShortThreadName = FStatNameAndInfo::GetShortNameFrom(ThreadName);
+				GroupLongName += FString::Printf(TEXT(" - %s"), *ShortThreadName.ToString());
+			}
 
-		const FHudGroup& HudGroup = ViewData.HudGroups[GroupIndex];
-		const bool bHasHierarchy = !!HudGroup.HierAggregate.Num();
-		const bool bHasFlat = !!HudGroup.FlatAggregate.Num();
-		const bool bBudget = HudGroup.TotalGroupBudget >= 0.f;
-		if (bHasHierarchy || bHasFlat)
-		{
-			// Render grouped headings.
-			Y += RenderGroupedHeadings( Canvas, X, Y, bHasHierarchy, bBudget );
-		}
+			if(!ViewData.RootFilter.IsEmpty())
+			{
+				GroupLongName += FString::Printf(TEXT(" ROOT=%s"), *ViewData.RootFilter);
+			}
 
-		// Render hierarchy.
-		if( bHasHierarchy )
-		{
-			RenderHierCycles( Canvas, X, Y, HudGroup );
+			Canvas->DrawShadowedString(X, Y, *GroupLongName, Globals.StatFont, Globals.GroupColor);
 			Y += Globals.GetFontHeight();
-		}
 
-		// Render flat.
-		if( bHasFlat )
-		{
-			RenderArrayOfStats(Canvas,X,Y,HudGroup.FlatAggregate, ViewData, HudGroup.BudgetIgnoreStats, HudGroup.TotalGroupBudget, RenderFlatCycle);
-			Y += Globals.GetFontHeight();
-		}
+			
+			const bool bHasHierarchy = !!HudGroup.HierAggregate.Num();
+			const bool bHasFlat = !!HudGroup.FlatAggregate.Num();
 
-		// Render memory counters.
-		if( HudGroup.MemoryAggregate.Num() )
-		{
-			Y += RenderMemoryHeadings(Canvas,X,Y);
-			RenderArrayOfStats(Canvas,X,Y,HudGroup.MemoryAggregate, ViewData, HudGroup.BudgetIgnoreStats, HudGroup.TotalGroupBudget, RenderMemoryCounter);
-			Y += Globals.GetFontHeight();
-		}
+			if (bHasHierarchy || bHasFlat)
+			{
+				// Render grouped headings.
+				Y += RenderGroupedHeadings(Canvas, X, Y, bHasHierarchy, bBudget);
+			}
 
-		// Render remaining counters.
-		if( HudGroup.CountersAggregate.Num() )
-		{
-			Y += RenderCounterHeadings(Canvas,X,Y);
-			RenderArrayOfStats(Canvas,X,Y,HudGroup.CountersAggregate, ViewData, HudGroup.BudgetIgnoreStats, HudGroup.TotalGroupBudget, RenderCounter);
-			Y += Globals.GetFontHeight();
+			// Render hierarchy.
+			if (bHasHierarchy)
+			{
+				RenderHierCycles(Canvas, X, Y, HudGroup);
+				Y += Globals.GetFontHeight();
+			}
+
+			const float* BudgetPtr = ShortThreadName != NAME_None ? HudGroup.ThreadBudgetMap.Find(ShortThreadName) : nullptr;
+			const float Budget = BudgetPtr ? *BudgetPtr : -1.f;
+			// Render flat.
+			if (bHasFlat)
+			{
+				RenderArrayOfStats(Canvas, X, Y, bBudget ? HudGroup.FlatAggregateThreadBreakdown[ThreadName] : HudGroup.FlatAggregate, ViewData, HudGroup.BudgetIgnoreStats, Budget, RenderFlatCycle);
+				Y += Globals.GetFontHeight();
+			}
+
+			// Render memory counters.
+			if (HudGroup.MemoryAggregate.Num())
+			{
+				Y += RenderMemoryHeadings(Canvas, X, Y);
+				RenderArrayOfStats(Canvas, X, Y, HudGroup.MemoryAggregate, ViewData, HudGroup.BudgetIgnoreStats, Budget, RenderMemoryCounter);
+				Y += Globals.GetFontHeight();
+			}
+
+			// Render remaining counters.
+			if (HudGroup.CountersAggregate.Num())
+			{
+				Y += RenderCounterHeadings(Canvas, X, Y);
+				RenderArrayOfStats(Canvas, X, Y, HudGroup.CountersAggregate, ViewData, HudGroup.BudgetIgnoreStats, Budget, RenderCounter);
+				Y += Globals.GetFontHeight();
+			}
 		}
+		
 	}
 }
 
diff --git a/Engine/Source/Runtime/Engine/Private/Texture2D.cpp b/Engine/Source/Runtime/Engine/Private/Texture2D.cpp
--- a/Engine/Source/Runtime/Engine/Private/Texture2D.cpp
+++ b/Engine/Source/Runtime/Engine/Private/Texture2D.cpp
@@ -193,7 +193,7 @@ void UTexture2D::Serialize(FArchive& Ar)
 	}
 
 #if WITH_EDITOR	
-	if (Ar.IsLoading() && !Ar.IsTransacting() && !bCooked && !(GetOutermost()->PackageFlags & PKG_ReloadingForCooker))
+	if (Ar.IsLoading() && !Ar.IsTransacting() && !bCooked && !GetOutermost()->HasAnyPackageFlags(PKG_ReloadingForCooker))
 	{
 		// The composite texture may not have been loaded yet. We have to defer caching platform
 		// data until post load.
diff --git a/Engine/Source/Runtime/Engine/Private/TextureDerivedData.cpp b/Engine/Source/Runtime/Engine/Private/TextureDerivedData.cpp
--- a/Engine/Source/Runtime/Engine/Private/TextureDerivedData.cpp
+++ b/Engine/Source/Runtime/Engine/Private/TextureDerivedData.cpp
@@ -1479,7 +1479,7 @@ void UTexture::BeginCachePlatformData()
 void UTexture::BeginCacheForCookedPlatformData( const ITargetPlatform *TargetPlatform )
 {
 	TMap<FString, FTexturePlatformData*>* CookedPlatformDataPtr = GetCookedPlatformData();
-	if (CookedPlatformDataPtr && !(GetOutermost()->PackageFlags & PKG_FilterEditorOnly))
+	if (CookedPlatformDataPtr && !GetOutermost()->HasAnyPackageFlags(PKG_FilterEditorOnly))
 	{
 		TMap<FString,FTexturePlatformData*>& CookedPlatformData = *CookedPlatformDataPtr;
 		
@@ -1734,12 +1734,12 @@ void UTexture::FinishCachePlatformData()
 			}
 
 #if DO_CHECK
-			if (!(GetOutermost()->PackageFlags & PKG_FilterEditorOnly))
+			if (!GetOutermost()->HasAnyPackageFlags(PKG_FilterEditorOnly))
 			{
-			FString DerivedDataKey;
-			FTextureBuildSettings BuildSettings;
-			GetBuildSettingsForRunningPlatform(*this, BuildSettings);
-			GetTextureDerivedDataKey(*this, BuildSettings, DerivedDataKey);
+				FString DerivedDataKey;
+				FTextureBuildSettings BuildSettings;
+				GetBuildSettingsForRunningPlatform(*this, BuildSettings);
+				GetTextureDerivedDataKey(*this, BuildSettings, DerivedDataKey);
 
 				check(RunningPlatformData->DerivedDataKey == DerivedDataKey);
 			}
diff --git a/Engine/Source/Runtime/Engine/Private/TickTaskManager.cpp b/Engine/Source/Runtime/Engine/Private/TickTaskManager.cpp
--- a/Engine/Source/Runtime/Engine/Private/TickTaskManager.cpp
+++ b/Engine/Source/Runtime/Engine/Private/TickTaskManager.cpp
@@ -49,7 +49,6 @@ static TAutoConsoleVariable<int32> CVarAllowConcurrentQueue(
 static bool GCanDemoteInto[TG_MAX] =
 {
 	true,		// 	TG_PrePhysics,
-	false,		// 	TG_DuringAnimation,
 	true,		// 	TG_StartPhysics,
 	false,		// 	TG_DuringPhysics,
 	true,		// 	TG_EndPhysics,
diff --git a/Engine/Source/Runtime/Engine/Private/TriangleRendering.cpp b/Engine/Source/Runtime/Engine/Private/TriangleRendering.cpp
--- a/Engine/Source/Runtime/Engine/Private/TriangleRendering.cpp
+++ b/Engine/Source/Runtime/Engine/Private/TriangleRendering.cpp
@@ -171,6 +171,37 @@ void FTriangleRenderer::DrawTriangle(FRHICommandListImmediate& RHICmdList, const
 	GetRendererModule().DrawTileMesh(RHICmdList, View, TriMesh, bIsHitTesting, HitProxyId);
 }
 
+void FTriangleRenderer::DrawTriangleUsingVertexColor(FRHICommandListImmediate& RHICmdList, const class FSceneView& View, const FMaterialRenderProxy* MaterialRenderProxy, bool bNeedsToSwitchVerticalAxis, const FCanvasUVTri& Tri, bool bIsHitTesting, const FHitProxyId HitProxyId)
+{
+	FMaterialTriangleVertex DestVertex[3];
+
+	// create verts
+	if (bNeedsToSwitchVerticalAxis)
+	{
+		DestVertex[0].Initialize(Tri.V1_Pos.X, View.ViewRect.Height() - Tri.V1_Pos.Y, Tri.V1_UV.X, Tri.V1_UV.Y);
+		DestVertex[1].Initialize(Tri.V0_Pos.X, View.ViewRect.Height() - Tri.V0_Pos.Y, Tri.V0_UV.X, Tri.V0_UV.Y);
+		DestVertex[2].Initialize(Tri.V2_Pos.X, View.ViewRect.Height() - Tri.V2_Pos.Y, Tri.V2_UV.X, Tri.V2_UV.Y);
+	}
+	else
+	{
+		DestVertex[0].Initialize(Tri.V1_Pos.X, Tri.V1_Pos.Y, Tri.V1_UV.X, Tri.V1_UV.Y);
+		DestVertex[1].Initialize(Tri.V0_Pos.X, Tri.V0_Pos.Y, Tri.V0_UV.X, Tri.V0_UV.Y);
+		DestVertex[2].Initialize(Tri.V2_Pos.X, Tri.V2_Pos.Y, Tri.V2_UV.X, Tri.V2_UV.Y);
+	}
+
+	DestVertex[0].Color = Tri.V0_Color.ToFColor(true).DWColor();
+	DestVertex[1].Color = Tri.V1_Color.ToFColor(true).DWColor();
+	DestVertex[2].Color = Tri.V2_Color.ToFColor(true).DWColor();
+
+	// update the FMeshBatch
+	FMeshBatch& TriMesh = GTriangleMesh.TriMeshElement;
+	TriMesh.UseDynamicData = true;
+	TriMesh.DynamicVertexData = DestVertex;
+	TriMesh.MaterialRenderProxy = MaterialRenderProxy;
+
+	GetRendererModule().DrawTileMesh(RHICmdList, View, TriMesh, bIsHitTesting, HitProxyId);
+}
+
 bool FCanvasTriangleRendererItem::Render_RenderThread(FRHICommandListImmediate& RHICmdList, const FCanvas* Canvas)
 {
 	float CurrentRealTime = 0.f;
@@ -215,7 +246,7 @@ bool FCanvasTriangleRendererItem::Render_RenderThread(FRHICommandListImmediate&
 	for (int32 TriIdx = 0; TriIdx < Data->Triangles.Num(); TriIdx++)
 	{
 		const FRenderData::FTriangleInst& Tri = Data->Triangles[TriIdx];
-		FTriangleRenderer::DrawTriangle(
+		FTriangleRenderer::DrawTriangleUsingVertexColor(
 			RHICmdList,
 			*View,
 			Data->MaterialRenderProxy,
@@ -300,7 +331,7 @@ bool FCanvasTriangleRendererItem::Render_GameThread(const FCanvas* Canvas)
 		for (int32 TriIdx = 0; TriIdx < Parameters.RenderData->Triangles.Num(); TriIdx++)
 		{
 			const FRenderData::FTriangleInst& Tri = Parameters.RenderData->Triangles[TriIdx];
-			FTriangleRenderer::DrawTriangle(
+			FTriangleRenderer::DrawTriangleUsingVertexColor(
 				RHICmdList,
 				*Parameters.View,
 				Parameters.RenderData->MaterialRenderProxy,
diff --git a/Engine/Source/Runtime/Engine/Private/UnrealEngine.cpp b/Engine/Source/Runtime/Engine/Private/UnrealEngine.cpp
--- a/Engine/Source/Runtime/Engine/Private/UnrealEngine.cpp
+++ b/Engine/Source/Runtime/Engine/Private/UnrealEngine.cpp
@@ -478,7 +478,7 @@ namespace
 		UPackage* NewPackage = CreatePackage(NULL, *PIEPackageName);
 		if (NewPackage != nullptr && WorldContext.WorldType == EWorldType::PIE)
 		{
-			NewPackage->PackageFlags |= PKG_PlayInEditor;
+			NewPackage->SetPackageFlags(PKG_PlayInEditor);
 			LoadFlags |= LOAD_PackageForPIE;
 		}
 		OutPackage = LoadPackage(NewPackage, *SourceWorldPackage, LoadFlags);
@@ -9472,7 +9472,7 @@ bool UEngine::LoadMap( FWorldContext& WorldContext, FURL URL, class UPendingNetG
 #if WITH_EDITOR
 		if (WorldContext.WorldType == EWorldType::PIE)
 		{
-			MapOuter->PackageFlags |= PKG_PlayInEditor;
+			MapOuter->SetPackageFlags(PKG_PlayInEditor);
 		}
 		MapOuter->PIEInstanceID = WorldContext.PIEInstance;
 #endif
@@ -9562,7 +9562,7 @@ bool UEngine::LoadMap( FWorldContext& WorldContext, FURL URL, class UPendingNetG
 
 		FScopeCycleCounterUObject MapScope(WorldPackage);
 
-		if (FPlatformProperties::RequiresCookedData() && GUseSeekFreeLoading && !(WorldPackage->PackageFlags & PKG_DisallowLazyLoading))
+		if (FPlatformProperties::RequiresCookedData() && GUseSeekFreeLoading && !WorldPackage->HasAnyPackageFlags(PKG_DisallowLazyLoading))
 		{
 			UE_LOG(LogLoad, Fatal, TEXT("Map '%s' has not been cooked correctly! Most likely stale version on the XDK."), *WorldPackage->GetName());
 		}
@@ -9609,10 +9609,10 @@ bool UEngine::LoadMap( FWorldContext& WorldContext, FURL URL, class UPendingNetG
 	WorldContext.World()->WorldType = WorldContext.WorldType;
 	
 	// Fixme: hacky but we need to set PackageFlags here if we are in a PIE Context.
-	// Also, dont add to root when in PIE, since PIE doesn't remove world from root
+	// Also, don't add to root when in PIE, since PIE doesn't remove world from root
 	if (WorldContext.WorldType == EWorldType::PIE)
 	{
-		check((CastChecked<UPackage>(WorldContext.World()->GetOutermost())->PackageFlags & PKG_PlayInEditor) == PKG_PlayInEditor);
+		check(WorldContext.World()->GetOutermost()->HasAnyPackageFlags(PKG_PlayInEditor));
 		WorldContext.World()->ClearFlags(RF_Standalone);
 	}
 	else
diff --git a/Engine/Source/Runtime/Engine/Private/Vehicles/PhysXVehicleManager.cpp b/Engine/Source/Runtime/Engine/Private/Vehicles/PhysXVehicleManager.cpp
--- a/Engine/Source/Runtime/Engine/Private/Vehicles/PhysXVehicleManager.cpp
+++ b/Engine/Source/Runtime/Engine/Private/Vehicles/PhysXVehicleManager.cpp
@@ -245,6 +245,9 @@ void FPhysXVehicleManager::RemoveVehicle( TWeakObjectPtr<UWheeledVehicleMovement
 	case PxVehicleTypes::eDRIVETANK:
 		((PxVehicleDriveTank*)PVehicle)->free();
 		break;
+	case PxVehicleTypes::eDRIVENW:
+		((PxVehicleDriveNW*)PVehicle)->free();
+		break;
 	case PxVehicleTypes::eNODRIVE:
 		((PxVehicleNoDrive*)PVehicle)->free();
 		break;
diff --git a/Engine/Source/Runtime/Engine/Private/VisualLogger/VisualLogger.cpp b/Engine/Source/Runtime/Engine/Private/VisualLogger/VisualLogger.cpp
--- a/Engine/Source/Runtime/Engine/Private/VisualLogger/VisualLogger.cpp
+++ b/Engine/Source/Runtime/Engine/Private/VisualLogger/VisualLogger.cpp
@@ -20,7 +20,7 @@ DEFINE_STAT(STAT_VisualLog);
 
 namespace
 {
-	static UWorld* GetWorldForVisualLogger(const class UObject* Object)
+	UWorld* GetWorldForVisualLogger(const UObject* Object)
 	{
 		UWorld* World = Object ? GEngine->GetWorldFromContextObject(Object, false) : nullptr;
 #if WITH_EDITOR
@@ -41,10 +41,10 @@ namespace
 	}
 }
 
-TMap<UObject*, TArray<TWeakObjectPtr<const UObject> > > FVisualLogger::RedirectionMap;
+TMap<const UWorld*, FVisualLogger::RedirectionMapType> FVisualLogger::WorldToRedirectionMap;
 int32 FVisualLogger::bIsRecording = false;
 
-bool FVisualLogger::CheckVisualLogInputInternal(const class UObject* Object, const struct FLogCategoryBase& Category, ELogVerbosity::Type Verbosity, UWorld **World, FVisualLogEntry **CurrentEntry)
+bool FVisualLogger::CheckVisualLogInputInternal(const UObject* Object, const FLogCategoryBase& Category, ELogVerbosity::Type Verbosity, UWorld **World, FVisualLogEntry **CurrentEntry)
 {
 	if (FVisualLogger::IsRecording() == false || !Object || (GEngine && GEngine->bDisableAILogging) || Object->HasAnyFlags(RF_ClassDefaultObject))
 	{
@@ -73,13 +73,13 @@ bool FVisualLogger::CheckVisualLogInputInternal(const class UObject* Object, con
 	return true;
 }
 
-FVisualLogEntry* FVisualLogger::GetLastEntryForObject(const class UObject* Object)
+FVisualLogEntry* FVisualLogger::GetLastEntryForObject(const UObject* Object)
 {
 	UObject * LogOwner = FVisualLogger::FindRedirection(Object);
 	return CurrentEntryPerObject.Contains(LogOwner) ? &CurrentEntryPerObject[LogOwner] : nullptr;
 }
 
-FVisualLogEntry* FVisualLogger::GetEntryToWrite(const class UObject* Object, float TimeStamp, ECreateIfNeeded ShouldCreate)
+FVisualLogEntry* FVisualLogger::GetEntryToWrite(const UObject* Object, float TimeStamp, ECreateIfNeeded ShouldCreate)
 {
 	FVisualLogEntry* CurrentEntry = nullptr;
 	UObject * LogOwner = FVisualLogger::FindRedirection(Object);
@@ -90,13 +90,13 @@ FVisualLogEntry* FVisualLogger::GetEntryToWrite(const class UObject* Object, flo
 
 	bool InitializeNewEntry = false;
 
-	TWeakObjectPtr<UWorld> World = GetWorldForVisualLogger(Object);
+	const UWorld* World = GetWorldForVisualLogger(Object);
 
 	if (CurrentEntryPerObject.Contains(LogOwner))
 	{
 		CurrentEntry = &CurrentEntryPerObject[LogOwner];
 		InitializeNewEntry = TimeStamp > CurrentEntry->TimeStamp && ShouldCreate == ECreateIfNeeded::Create;
-		if (World.IsValid())
+		if (World && IsInGameThread())
 		{
 			World->GetTimerManager().ClearTimer(VisualLoggerCleanupTimerHandle);
 			for (auto& CurrentPair : CurrentEntryPerObject)
@@ -104,7 +104,7 @@ FVisualLogEntry* FVisualLogger::GetEntryToWrite(const class UObject* Object, flo
 				FVisualLogEntry* Entry = &CurrentPair.Value;
 				if (Entry->TimeStamp >= 0 && Entry->TimeStamp < TimeStamp)
 				{
-					for (auto* Device : OutputDevices)
+					for (FVisualLogDevice* Device : OutputDevices)
 					{
 						Device->Serialize(CurrentPair.Key, ObjectToNameMap[CurrentPair.Key], ObjectToClassNameMap[CurrentPair.Key], *Entry);
 					}
@@ -121,6 +121,7 @@ FVisualLogEntry* FVisualLogger::GetEntryToWrite(const class UObject* Object, flo
 		ObjectToNameMap.Add(LogOwner, LogOwner->GetFName());
 		ObjectToClassNameMap.Add(LogOwner, *(LogOwner->GetClass()->GetName()));
 		ObjectToPointerMap.Add(LogOwner, LogOwner);
+		ObjectToWorldMap.Add(LogOwner, World);
 		InitializeNewEntry = true;
 	}
 
@@ -129,6 +130,7 @@ FVisualLogEntry* FVisualLogger::GetEntryToWrite(const class UObject* Object, flo
 		CurrentEntry->Reset();
 		CurrentEntry->TimeStamp = TimeStamp;
 
+		auto& RedirectionMap = GetRedirectionMap(LogOwner);
 		if (RedirectionMap.Contains(LogOwner))
 		{
 			if (ObjectToPointerMap.Contains(LogOwner) && ObjectToPointerMap[LogOwner].IsValid())
@@ -153,7 +155,7 @@ FVisualLogEntry* FVisualLogger::GetEntryToWrite(const class UObject* Object, flo
 		}
 		else
 		{
-			const class AActor* ObjectAsActor = Cast<class AActor>(Object);
+			const AActor* ObjectAsActor = Cast<AActor>(Object);
 			if (ObjectAsActor)
 			{
 				CurrentEntry->Location = ObjectAsActor->GetActorLocation();
@@ -166,7 +168,7 @@ FVisualLogEntry* FVisualLogger::GetEntryToWrite(const class UObject* Object, flo
 		}
 	}
 
-	if (World.IsValid())
+	if (World && IsInGameThread())
 	{
 		//set next tick timer to flush obsolete/old entries
 		World->GetTimerManager().SetTimer(VisualLoggerCleanupTimerHandle, FTimerDelegate::CreateLambda(
@@ -174,9 +176,9 @@ FVisualLogEntry* FVisualLogger::GetEntryToWrite(const class UObject* Object, flo
 			for (auto& CurrentPair : CurrentEntryPerObject)
 			{
 				FVisualLogEntry* Entry = &CurrentPair.Value;
-				if (Entry->TimeStamp >= 0 && (!World.IsValid() || Entry->TimeStamp < World->GetTimeSeconds())) // CurrentEntry->TimeStamp == -1 means it's not initialized entry information
+				if (Entry->TimeStamp >= 0 && (!World || Entry->TimeStamp < World->GetTimeSeconds())) // CurrentEntry->TimeStamp == -1 means it's not initialized entry information
 				{
-					for (auto* Device : OutputDevices)
+					for (FVisualLogDevice* Device : OutputDevices)
 					{
 						Device->Serialize(CurrentPair.Key, ObjectToNameMap[CurrentPair.Key], ObjectToClassNameMap[CurrentPair.Key], *Entry);
 					}
@@ -197,7 +199,7 @@ void FVisualLogger::Flush()
 	{
 		if (CurrentEntry.Value.TimeStamp >= 0)
 		{
-			for (auto* Device : OutputDevices)
+			for (FVisualLogDevice* Device : OutputDevices)
 			{
 				Device->Serialize(CurrentEntry.Key, ObjectToNameMap[CurrentEntry.Key], ObjectToClassNameMap[CurrentEntry.Key], CurrentEntry.Value);
 			}
@@ -207,48 +209,48 @@ void FVisualLogger::Flush()
 }
 
 
-void FVisualLogger::EventLog(const class UObject* Object, const FName EventTag1, const FVisualLogEventBase& Event1, const FVisualLogEventBase& Event2, const FVisualLogEventBase& Event3, const FVisualLogEventBase& Event4, const FVisualLogEventBase& Event5, const FVisualLogEventBase& Event6)
+void FVisualLogger::EventLog(const UObject* Object, const FName EventTag1, const FVisualLogEventBase& Event1, const FVisualLogEventBase& Event2, const FVisualLogEventBase& Event3, const FVisualLogEventBase& Event4, const FVisualLogEventBase& Event5, const FVisualLogEventBase& Event6)
 {
 	EventLog(Object, EventTag1, Event1, Event2, Event3, Event4, Event5);
 	EventLog(Object, EventTag1, Event6);
 }
 
 
-void FVisualLogger::EventLog(const class UObject* Object, const FName EventTag1, const FVisualLogEventBase& Event1, const FVisualLogEventBase& Event2, const FVisualLogEventBase& Event3, const FVisualLogEventBase& Event4, const FVisualLogEventBase& Event5)
+void FVisualLogger::EventLog(const UObject* Object, const FName EventTag1, const FVisualLogEventBase& Event1, const FVisualLogEventBase& Event2, const FVisualLogEventBase& Event3, const FVisualLogEventBase& Event4, const FVisualLogEventBase& Event5)
 {
 	EventLog(Object, EventTag1, Event1, Event2, Event3, Event4);
 	EventLog(Object, EventTag1, Event5);
 }
 
 
-void FVisualLogger::EventLog(const class UObject* Object, const FName EventTag1, const FVisualLogEventBase& Event1, const FVisualLogEventBase& Event2, const FVisualLogEventBase& Event3, const FVisualLogEventBase& Event4)
+void FVisualLogger::EventLog(const UObject* Object, const FName EventTag1, const FVisualLogEventBase& Event1, const FVisualLogEventBase& Event2, const FVisualLogEventBase& Event3, const FVisualLogEventBase& Event4)
 {
 	EventLog(Object, EventTag1, Event1, Event2, Event3);
 	EventLog(Object, EventTag1, Event4);
 }
 
 
-void FVisualLogger::EventLog(const class UObject* Object, const FName EventTag1, const FVisualLogEventBase& Event1, const FVisualLogEventBase& Event2, const FVisualLogEventBase& Event3)
+void FVisualLogger::EventLog(const UObject* Object, const FName EventTag1, const FVisualLogEventBase& Event1, const FVisualLogEventBase& Event2, const FVisualLogEventBase& Event3)
 {
 	EventLog(Object, EventTag1, Event1, Event2);
 	EventLog(Object, EventTag1, Event3);
 }
 
 
-void FVisualLogger::EventLog(const class UObject* Object, const FName EventTag1, const FVisualLogEventBase& Event1, const FVisualLogEventBase& Event2)
+void FVisualLogger::EventLog(const UObject* Object, const FName EventTag1, const FVisualLogEventBase& Event1, const FVisualLogEventBase& Event2)
 {
 	EventLog(Object, EventTag1, Event1);
 	EventLog(Object, EventTag1, Event2);
 }
 
 
-void FVisualLogger::EventLog(const class UObject* LogOwner, const FVisualLogEventBase& Event1, const FName EventTag1, const FName EventTag2, const FName EventTag3, const FName EventTag4, const FName EventTag5, const FName EventTag6)
+void FVisualLogger::EventLog(const UObject* LogOwner, const FVisualLogEventBase& Event1, const FName EventTag1, const FName EventTag2, const FName EventTag3, const FName EventTag4, const FName EventTag5, const FName EventTag6)
 {
 	EventLog(LogOwner, EventTag1, Event1, EventTag2, EventTag3, EventTag4, EventTag5, EventTag6);
 }
 
 
-void FVisualLogger::EventLog(const class UObject* Object, const FName EventTag1, const FVisualLogEventBase& Event, const FName EventTag2, const FName EventTag3, const FName EventTag4, const FName EventTag5, const FName EventTag6)
+void FVisualLogger::EventLog(const UObject* Object, const FName EventTag1, const FVisualLogEventBase& Event, const FName EventTag2, const FName EventTag3, const FName EventTag4, const FName EventTag5, const FName EventTag6)
 {
 	SCOPE_CYCLE_COUNTER(STAT_VisualLog);
 	UWorld *World = nullptr;
@@ -278,7 +280,7 @@ void FVisualLogger::EventLog(const class UObject* Object, const FName EventTag1,
 	CurrentEntry->Events[Index].EventTags.Remove(NAME_None);
 }
 
-void FVisualLogger::NavigationDataDump(const class UObject* Object, const struct FLogCategoryBase& Category, ELogVerbosity::Type Verbosity, int32 UniqueLogId, const FBox& Box)
+void FVisualLogger::NavigationDataDump(const UObject* Object, const FLogCategoryBase& Category, ELogVerbosity::Type Verbosity, int32 UniqueLogId, const FBox& Box)
 {
 	SCOPE_CYCLE_COUNTER(STAT_VisualLog);
 	UWorld *World = nullptr;
@@ -322,7 +324,7 @@ void FVisualLogger::Shutdown()
 	}
 }
 
-void FVisualLogger::Cleanup(bool bReleaseMemory)
+void FVisualLogger::Cleanup(UWorld* OldWorld, bool bReleaseMemory)
 {
 	const bool WasRecordingToFile = IsRecordingToFile();
 	if (WasRecordingToFile)
@@ -335,9 +337,34 @@ void FVisualLogger::Cleanup(bool bReleaseMemory)
 		Device->Cleanup(bReleaseMemory);
 	}
 
-	RedirectionMap.Reset();
+	if (OldWorld != nullptr)
+	{
+		WorldToRedirectionMap.Remove(OldWorld);
+
+		for (auto It = ObjectToWorldMap.CreateIterator(); It; ++It)
+		{
+			if (It.Value() == OldWorld)
+			{
+				const UObject* Obj = It.Key();
+				ObjectToWorldMap.Remove(Obj);
+				CurrentEntryPerObject.Remove(Obj);
+				ObjectToNameMap.Remove(Obj);
+				ObjectToClassNameMap.Remove(Obj);
+				ObjectToPointerMap.Remove(Obj);
+			}
+		}
+	}
+	else
+	{
+		WorldToRedirectionMap.Reset();
+		ObjectToWorldMap.Reset();
+		CurrentEntryPerObject.Reset();
+		ObjectToNameMap.Reset();
+		ObjectToClassNameMap.Reset();
+		ObjectToPointerMap.Reset();
+	}
+
 	LastUniqueIds.Reset();
-	CurrentEntryPerObject.Reset();
 
 	if (WasRecordingToFile)
 	{
@@ -350,6 +377,22 @@ int32 FVisualLogger::GetUniqueId(float Timestamp)
 	return LastUniqueIds.FindOrAdd(Timestamp)++;
 }
 
+FVisualLogger::RedirectionMapType& FVisualLogger::GetRedirectionMap(const UObject* InObject)
+{
+	const UWorld* World = nullptr;
+	if (FVisualLogger::Get().ObjectToWorldMap.Contains(InObject))
+	{
+		World = FVisualLogger::Get().ObjectToWorldMap[InObject].Get();
+	}
+
+	if (World == nullptr)
+	{
+		World = GetWorldForVisualLogger(nullptr);
+	}
+
+	return WorldToRedirectionMap.FindOrAdd(World);
+}
+
 void FVisualLogger::Redirect(UObject* FromObject, UObject* ToObject)
 {
 	if (FromObject == ToObject || FromObject == nullptr || ToObject == nullptr)
@@ -359,6 +402,7 @@ void FVisualLogger::Redirect(UObject* FromObject, UObject* ToObject)
 
 	UObject* OldRedirection = FindRedirection(FromObject);
 	UObject* NewRedirection = FindRedirection(ToObject);
+	auto& RedirectionMap = GetRedirectionMap(FromObject);
 
 	if (OldRedirection != NewRedirection)
 	{
@@ -374,8 +418,9 @@ void FVisualLogger::Redirect(UObject* FromObject, UObject* ToObject)
 	}
 }
 
-class UObject* FVisualLogger::FindRedirection(const UObject* Object)
+UObject* FVisualLogger::FindRedirection(const UObject* Object)
 {
+	auto& RedirectionMap = GetRedirectionMap(Object);
 	if (RedirectionMap.Contains(Object) == false)
 	{
 		for (auto& Redirection : RedirectionMap)
@@ -387,7 +432,7 @@ class UObject* FVisualLogger::FindRedirection(const UObject* Object)
 		}
 	}
 
-	return const_cast<class UObject*>(Object);
+	return const_cast<UObject*>(Object);
 }
 
 void FVisualLogger::SetIsRecording(bool InIsRecording) 
@@ -419,7 +464,7 @@ void FVisualLogger::SetIsRecordingToFile(bool InIsRecording)
 
 	if (bIsRecordingToFile && !InIsRecording)
 	{
-		for (auto* Device : OutputDevices)
+		for (FVisualLogDevice* Device : OutputDevices)
 		{
 			if (Device->HasFlags(EVisualLoggerDeviceFlags::CanSaveToFile))
 			{
@@ -431,7 +476,7 @@ void FVisualLogger::SetIsRecordingToFile(bool InIsRecording)
 	else if (!bIsRecordingToFile && InIsRecording)
 	{
 		StartRecordingToFileTime = World ? World->TimeSeconds : 0;
-		for (auto* Device : OutputDevices)
+		for (FVisualLogDevice* Device : OutputDevices)
 		{
 			if (Device->HasFlags(EVisualLoggerDeviceFlags::CanSaveToFile))
 			{
@@ -443,7 +488,7 @@ void FVisualLogger::SetIsRecordingToFile(bool InIsRecording)
 	bIsRecordingToFile = InIsRecording;
 }
 
-bool FVisualLogger::IsCategoryLogged(const struct FLogCategoryBase& Category) const
+bool FVisualLogger::IsCategoryLogged(const FLogCategoryBase& Category) const
 {
 	if ((GEngine && GEngine->bDisableAILogging) || IsRecording() == false)
 	{
diff --git a/Engine/Source/Runtime/Engine/Private/VisualLogger/VisualLoggerAutomationTests.cpp b/Engine/Source/Runtime/Engine/Private/VisualLogger/VisualLoggerAutomationTests.cpp
--- a/Engine/Source/Runtime/Engine/Private/VisualLogger/VisualLoggerAutomationTests.cpp
+++ b/Engine/Source/Runtime/Engine/Private/VisualLogger/VisualLoggerAutomationTests.cpp
@@ -83,7 +83,7 @@ struct FTestDeviceContext
 	{
 		Device.Cleanup();
 		FVisualLogger::Get().SetIsRecording(false);
-		FVisualLogger::Get().Cleanup();
+		FVisualLogger::Get().Cleanup(nullptr);
 		FVisualLogger::Get().AddDevice(&Device);
 
 		EngineDisableAILoggingFlag = GEngine->bDisableAILogging;
@@ -95,7 +95,7 @@ struct FTestDeviceContext
 	{
 		FVisualLogger::Get().SetIsRecording(false);
 		FVisualLogger::Get().RemoveDevice(&Device);
-		FVisualLogger::Get().Cleanup();
+		FVisualLogger::Get().Cleanup(nullptr);
 		Device.Cleanup();
 
 		GEngine->bDisableAILogging = EngineDisableAILoggingFlag;
diff --git a/Engine/Source/Runtime/Engine/Private/World.cpp b/Engine/Source/Runtime/Engine/Private/World.cpp
--- a/Engine/Source/Runtime/Engine/Private/World.cpp
+++ b/Engine/Source/Runtime/Engine/Private/World.cpp
@@ -111,25 +111,6 @@ FActorSpawnParameters::FActorSpawnParameters()
 
 PRAGMA_ENABLE_DEPRECATION_WARNINGS
 
-FActorSpawnParameters& FActorSpawnParameters::operator=(const FActorSpawnParameters& Other)
-{
-	Name = Other.Name;
-	Template = Other.Template;
-	Owner = Other.Owner;
-	Instigator = Other.Instigator;
-	OverrideLevel = Other.OverrideLevel;
-	SpawnCollisionHandlingOverride = Other.SpawnCollisionHandlingOverride;
-PRAGMA_DISABLE_DEPRECATION_WARNINGS
-	bNoCollisionFail = Other.bNoCollisionFail;
-PRAGMA_ENABLE_DEPRECATION_WARNINGS
-	bRemoteOwned = Other.bRemoteOwned;
-	bNoFail = Other.bNoFail;
-	bDeferConstruction = Other.bDeferConstruction;
-	bAllowDuringConstructionScript = Other.bAllowDuringConstructionScript;
-	ObjectFlags = Other.ObjectFlags;
-	return *this;
-}
-
 /*-----------------------------------------------------------------------------
 	UWorld implementation.
 -----------------------------------------------------------------------------*/
@@ -272,7 +253,7 @@ void UWorld::Serialize( FArchive& Ar )
 	}
 	
 	// UWorlds loaded/duplicated for PIE must lose RF_Public and RF_Standalone since they should not be referenced by objects in other packages and they should be GCed normally
-	if (GetOutermost()->PackageFlags & PKG_PlayInEditor)
+	if (GetOutermost()->HasAnyPackageFlags(PKG_PlayInEditor))
 	{
 		ClearFlags(RF_Public|RF_Standalone);
 	}
@@ -396,7 +377,7 @@ bool UWorld::Rename(const TCHAR* InName, UObject* NewOuter, ERenameFlags Flags)
 		// If this is the last world removed from a package, clear the PKG_ContainsMap flag
 		if ( UWorld::FindWorldInPackage(OldPackage) == NULL )
 		{
-			OldPackage->PackageFlags &= ~PKG_ContainsMap;
+			OldPackage->ClearPackageFlags(PKG_ContainsMap);
 		}
 
 		// Set the PKG_ContainsMap flag in the new package
@@ -586,7 +567,7 @@ void UWorld::FinishDestroy()
 
 		if ( !bContainsAnotherWorld )
 		{
-			WorldPackage->PackageFlags &= ~PKG_ContainsMap;
+			WorldPackage->ClearPackageFlags(PKG_ContainsMap);
 		}
 	}
 
@@ -643,7 +624,7 @@ void UWorld::PostLoad()
 #if WITH_EDITOR
 	if (GIsEditor)
 	{
-		if (!(GetOutermost()->PackageFlags & PKG_PlayInEditor))
+		if (!GetOutermost()->HasAnyPackageFlags(PKG_PlayInEditor))
 		{
 			// Needed for VER_UE4_WORLD_NAMED_AFTER_PACKAGE. If this file was manually renamed outside of the editor, this is needed anyway
 			const FString ShortPackageName = FPackageName::GetLongPackageAssetName(GetOutermost()->GetName());
@@ -1172,7 +1153,7 @@ UWorld* UWorld::CreateWorld(const EWorldType::Type InWorldType, bool bInformEngi
 
 	if (InWorldType == EWorldType::PIE)
 	{
-		WorldPackage->PackageFlags |= PKG_PlayInEditor;
+		WorldPackage->SetPackageFlags(PKG_PlayInEditor);
 	}
 
 	// Mark the package as containing a world.  This has to happen here rather than at serialization time,
@@ -2144,7 +2125,7 @@ void FLevelStreamingGCHelper::PrepareStreamedOutLevelsForGC()
 	{
 		ULevel*	Level = LevelsPendingUnload[LevelIndex].Get();
 
-		if( Level && (!GIsEditor || (Level->GetOutermost()->PackageFlags & PKG_PlayInEditor) ))
+		if( Level && (!GIsEditor || Level->GetOutermost()->HasAnyPackageFlags(PKG_PlayInEditor) ))
 		{
 			UPackage* LevelPackage = Cast<UPackage>(Level->GetOutermost());
 			UE_LOG(LogStreaming, Log, TEXT("PrepareStreamedOutLevelsForGC called on '%s'"), *LevelPackage->GetName() );
@@ -2321,7 +2302,7 @@ UWorld* UWorld::DuplicateWorldForPIE(const FString& PackageName, UWorld* OwningW
 	const FName PrefixedLevelFName = FName(*PrefixedLevelName);
 	UWorld::WorldTypePreLoadMap.FindOrAdd(PrefixedLevelFName) = EWorldType::PIE;
 	UPackage* PIELevelPackage = CastChecked<UPackage>(CreatePackage(NULL,*PrefixedLevelName));
-	PIELevelPackage->PackageFlags |= PKG_PlayInEditor;
+	PIELevelPackage->SetPackageFlags(PKG_PlayInEditor);
 #if WITH_EDITOR
 	PIELevelPackage->PIEInstanceID = WorldContext.PIEInstance;
 #endif
@@ -3233,7 +3214,7 @@ void UWorld::CleanupWorld(bool bSessionEnded, bool bCleanupResources, UWorld* Ne
 #endif //WITH_EDITOR
 
 #if ENABLE_VISUAL_LOG
-	FVisualLogger::Get().Cleanup();
+	FVisualLogger::Get().Cleanup(this);
 #endif // ENABLE_VISUAL_LOG	
 
 	// Tell actors to remove their components from the scene.
@@ -4910,7 +4891,7 @@ UWorld* FSeamlessTravelHandler::Tick()
 			{
 				UPackage * WorldPackage = LoadedWorld->GetOutermost();
 				check(WorldPackage);
-				WorldPackage->PackageFlags |= PKG_PlayInEditor;
+				WorldPackage->SetPackageFlags(PKG_PlayInEditor);
 			}
 
 			// Clear any world specific state from NetDriver before switching World
diff --git a/Engine/Source/Runtime/Engine/Private/WorldComposition.cpp b/Engine/Source/Runtime/Engine/Private/WorldComposition.cpp
--- a/Engine/Source/Runtime/Engine/Private/WorldComposition.cpp
+++ b/Engine/Source/Runtime/Engine/Private/WorldComposition.cpp
@@ -28,7 +28,7 @@ void UWorldComposition::PostInitProperties()
 {
 	Super::PostInitProperties();
 
-	if (!IsTemplate() && (GetOutermost()->PackageFlags & PKG_PlayInEditor) == 0)
+	if (!IsTemplate() && !GetOutermost()->HasAnyPackageFlags(PKG_PlayInEditor))
 	{
 		// Tiles information is not serialized to disk, and should be regenerated on world composition object construction
 		Rescan();
