Copying //UE4/Dev-Framework to Dev-Main (//UE4/Dev-Main)

diff --git a/Engine/Source/Runtime/Engine/Public/ActiveSound.h b/Engine/Source/Runtime/Engine/Public/ActiveSound.h
--- a/Engine/Source/Runtime/Engine/Public/ActiveSound.h
+++ b/Engine/Source/Runtime/Engine/Public/ActiveSound.h
@@ -27,7 +27,8 @@ struct FSoundParseParameters
 	float Volume;
 	float VolumeMultiplier;
 
-	float VolumeWeightedPriorityScale;
+	// The priority of sound, which is the product of the component priority and the USoundBased priority
+	float Priority;
 
 	float Pitch;
 	float HighFrequencyGain;
@@ -187,7 +188,9 @@ public:
 	float ConcurrencyVolumeScale;
 
 	float SubtitlePriority;
-	float VolumeWeightedPriorityScale;
+
+	/** The product of the component priority and the USoundBase priority */
+	float Priority;
 
 	// The volume used to determine concurrency resolution for "quietest" active sound
 	float VolumeConcurrency;
@@ -303,6 +306,9 @@ public:
 	/** Returns the sound concurrency object ID if it exists. If it doesn't exist, returns 0. */
 	uint32 GetSoundConcurrencyObjectID() const;
 
+	/** Applies the active sound's attenuation settings to the input parse params using the given listener */
+	void ApplyAttenuation(FSoundParseParameters& ParseParams, const FListener& Listener);
+
 private:
 
 	/** This is a friend so the audio device can call Stop() on the active sound. */
@@ -323,4 +329,23 @@ private:
 
 	/** Apply the interior settings to the ambient sound as appropriate */
 	void HandleInteriorVolumes( const FListener& Listener, struct FSoundParseParameters& ParseParams );
+
+	/** Struct used to cache listener attenuation vector math results */
+	struct FAttenuationListenerData
+	{
+		FVector ListenerToSoundDir;
+		float AttenuationDistance;
+		float ListenerToSoundDistance;
+		bool bDataComputed;
+
+		FAttenuationListenerData()
+			: ListenerToSoundDir(FVector::ZeroVector)
+			, AttenuationDistance(0.0f)
+			, ListenerToSoundDistance(0.0f)
+			, bDataComputed(false)
+		{}
+	};
+
+	/** Performs listener-based calculations for spatializing and attenuating the active sound */
+	void GetAttenuationListenerData(FAttenuationListenerData& Data, const FSoundParseParameters& ParseParams, const FListener& Listener) const;
 };
\ No newline at end of file
diff --git a/Engine/Source/Runtime/Engine/Public/Animation/AnimInstanceProxy.h b/Engine/Source/Runtime/Engine/Public/Animation/AnimInstanceProxy.h
new file mode 100644
--- /dev/null
+++ b/Engine/Source/Runtime/Engine/Public/Animation/AnimInstanceProxy.h
@@ -0,0 +1,553 @@
+// Copyright 1998-2015 Epic Games, Inc. All Rights Reserved.
+
+#pragma once
+
+#include "Animation/AnimNodeBase.h"
+#include "Animation/AnimationAsset.h"
+#include "Animation/AnimInstance.h"
+#include "Animation/AnimNotifyQueue.h"
+#include "AnimInstanceProxy.generated.h"
+
+struct FAnimNode_Base;
+
+/** Delegate fired for nodes to gather any data they need on the game thread before */
+DECLARE_DELEGATE_OneParam(FGameThreadPreUpdateEvent, const UAnimInstance* /*InAnimInstance*/);
+
+/** Proxy object passed around during animation tree update in lieu of a UAnimInstance */
+USTRUCT(meta = (DisplayName = "Native Variables"))
+struct ENGINE_API FAnimInstanceProxy
+{
+	GENERATED_BODY()
+
+public:
+	FAnimInstanceProxy()
+	{
+	}
+
+	FAnimInstanceProxy(UAnimInstance* Instance)
+		: AnimInstanceObject(Instance)
+		, AnimBlueprintClass(Cast<UAnimBlueprintGeneratedClass>(Instance->GetClass()))
+		, Skeleton(nullptr)
+		, SkeletalMeshComponent(nullptr)
+		, CurrentDeltaSeconds(0.0f)
+#if WITH_EDITORONLY_DATA
+		, bIsBeingDebugged(false)
+#endif
+		, RootNode(nullptr)
+		, SyncGroupWriteIndex(0)
+		, RootMotionMode(ERootMotionMode::NoRootMotionExtraction)
+		, bShouldExtractRootMotion(false)
+		, bBoneCachesInvalidated(false)
+	{
+	}
+
+	virtual ~FAnimInstanceProxy() {}
+
+	// Get the Blueprint Generated Class associated with this context, if there is one.
+	// Note: This can return NULL, so check the result.
+	UAnimBlueprintGeneratedClass* GetAnimBlueprintClass() const
+	{
+		return AnimBlueprintClass;
+	}
+
+	/** Get the last DeltaSeconds passed into PreUpdate() */
+	float GetDeltaSeconds() const
+	{
+		return CurrentDeltaSeconds;
+	}
+
+#if WITH_EDITORONLY_DATA
+	/** Whether the UAnimInstance this context refers to is currently being debugged in the editor */
+	bool IsBeingDebugged() const
+	{
+		return bIsBeingDebugged;
+	}
+
+	/** Record a visited node in the debugger */
+	void RecordNodeVisit(int32 TargetNodeIndex, int32 SourceNodeIndex, float BlendWeight)
+	{
+		new (UpdatedNodesThisFrame) FAnimBlueprintDebugData::FNodeVisit(SourceNodeIndex, TargetNodeIndex, BlendWeight);
+	}
+#endif
+
+	// flip sync group read/write indices
+	void TickSyncGroupWriteIndex()
+	{ 
+		SyncGroupWriteIndex = GetSyncGroupReadIndex();
+	}
+
+	/** Get the sync group we are currently reading from */
+	const TArray<FAnimGroupInstance>& GetSyncGroupRead() const
+	{ 
+		return SyncGroupArrays[GetSyncGroupReadIndex()]; 
+	}
+
+	/** Get the ungrouped active player we are currently reading from */
+	const TArray<FAnimTickRecord>& GetUngroupedActivePlayersRead() 
+	{ 
+		return UngroupedActivePlayerArrays[GetSyncGroupReadIndex()]; 
+	}
+
+	/** Tick active asset players */
+	void TickAssetPlayerInstances(float DeltaSeconds);
+
+	/** Queues an Anim Notify from the shared list on our generated class */
+	void AddAnimNotifyFromGeneratedClass(int32 NotifyIndex);
+
+	/** Trigger any anim notifies */
+	void TriggerAnimNotifies(USkeletalMeshComponent* SkelMeshComp, float DeltaSeconds);
+
+	/** Check whether the supplied skeleton is compatible with this instance's skeleton */
+	bool IsSkeletonCompatible(USkeleton const* InSkeleton) const 
+	{ 
+		return InSkeleton->IsCompatible(Skeleton); 
+	}
+
+	/** Check whether we should extract root motion */
+	bool ShouldExtractRootMotion() const 
+	{ 
+		return bShouldExtractRootMotion;
+	}
+
+	/** Access various counters */
+	const FGraphTraversalCounter& GetInitializationCounter() const { return InitializationCounter; }
+	const FGraphTraversalCounter& GetCachedBonesCounter() const  { return CachedBonesCounter; }
+	const FGraphTraversalCounter& GetUpdateCounter() const  { return UpdateCounter; }
+	const FGraphTraversalCounter& GetEvaluationCounter() const { return EvaluationCounter; }
+	const FGraphTraversalCounter& GetSlotNodeInitializationCounter() const { return SlotNodeInitializationCounter; }
+
+	/** Access root motion params */
+	FRootMotionMovementParams& GetExtractedRootMotion() { return ExtractedRootMotion; }
+
+	/** Access UObject base of UAnimInstance */
+	UObject* GetAnimInstanceObject() { return AnimInstanceObject; }
+
+	/** Gets an unchecked (can return nullptr) node given an index into the node property array */
+	FAnimNode_Base* GetNodeFromIndexUntyped(int32 NodeIdx, UScriptStruct* RequiredStructType);
+
+	/** Gets a checked node given an index into the node property array */
+	FAnimNode_Base* GetCheckedNodeFromIndexUntyped(int32 NodeIdx, UScriptStruct* RequiredStructType);
+
+	/** Gets a checked node given an index into the node property array */
+	template<class NodeType>
+	NodeType* GetCheckedNodeFromIndex(int32 NodeIdx)
+	{
+		return (NodeType*)GetCheckedNodeFromIndexUntyped(NodeIdx, NodeType::StaticStruct());
+	}
+
+	/** Gets an unchecked (can return nullptr) node given an index into the node property array */
+	template<class NodeType>
+	NodeType* GetNodeFromIndex(int32 NodeIdx)
+	{
+		return (NodeType*)GetNodeFromIndexUntyped(NodeIdx, NodeType::StaticStruct());
+	}
+
+	/** const access to required bones array */
+	const FBoneContainer& GetRequiredBones() const 
+	{ 
+		return RequiredBones; 
+	}
+
+	/** access to required bones array */
+	FBoneContainer& GetRequiredBones() 
+	{ 
+		return RequiredBones;
+	}
+
+	/** Get the current skeleton we are using */
+	USkeleton* GetSkeleton() 
+	{ 
+		return Skeleton; 
+	}
+
+	/** Get the current skeletal mesh component we are running on */
+	USkeletalMeshComponent* GetSkelMeshComponent() 
+	{ 
+		return SkeletalMeshComponent; 
+	}
+
+	// Creates an uninitialized tick record in the list for the correct group or the ungrouped array.  If the group is valid, OutSyncGroupPtr will point to the group.
+	FAnimTickRecord& CreateUninitializedTickRecord(int32 GroupIndex, FAnimGroupInstance*& OutSyncGroupPtr);
+
+	/** Helper function: make a tick record for a sequence */
+	void MakeSequenceTickRecord(FAnimTickRecord& TickRecord, UAnimSequenceBase* Sequence, bool bLooping, float PlayRate, float FinalBlendWeight, float& CurrentTime, FMarkerTickRecord& MarkerTickRecord) const;
+
+	/** Helper function: make a tick record for a blend space */
+	void MakeBlendSpaceTickRecord(FAnimTickRecord& TickRecord, UBlendSpaceBase* BlendSpace, const FVector& BlendInput, TArray<FBlendSampleData>& BlendSampleDataCache, FBlendFilter& BlendFilter, bool bLooping, float PlayRate, float FinalBlendWeight, float& CurrentTime, FMarkerTickRecord& MarkerTickRecord) const;
+
+	/**
+	 * Get Slot Node Weight : this returns new Slot Node Weight, Source Weight, Original TotalNodeWeight
+	 *							this 3 values can't be derived from each other
+	 *
+	 * @param SlotNodeName : the name of the slot node you're querying
+	 * @param out_SlotNodeWeight : The node weight for this slot node in the range of [0, 1]
+	 * @param out_SourceWeight : The Source weight for this node. 
+	 * @param out_TotalNodeWeight : Total weight of this node
+	 */
+	void GetSlotWeight(FName const& SlotNodeName, float& out_SlotNodeWeight, float& out_SourceWeight, float& out_TotalNodeWeight) const;
+
+	/** Evaluate a pose for a named montage slot */
+	void SlotEvaluatePose(FName SlotNodeName, const FCompactPose& SourcePose, const FBlendedCurve& SourceCurve, float InSourceWeight, FCompactPose& BlendedPose, FBlendedCurve& BlendedCurve, float InBlendWeight, float InTotalNodeWeight);
+	
+	// Allow slot nodes to store off their weight during ticking
+	void UpdateSlotNodeWeight(FName SlotNodeName, float Weight);
+
+	// Allow slot nodes to store off their root motion weight during ticking
+	void UpdateSlotRootMotionWeight(FName SlotNodeName, float Weight);
+
+	/** Register a named slot */
+	void RegisterSlotNodeWithAnimInstance(FName SlotNodeName);
+
+	/** 
+	 * Add a pre-update event. Events can be added more than once. Subsequent adds will be ignored.
+	 * Events are called from PreUpdate() on the game thread.
+	 * This is intended to be used by FAnimNode_Base-derived classes that need to perform some Update() work on the game thread.
+	 * @param	InEvent		The event to register
+	 */
+	void AddGameThreadPreUpdateEvent(const FGameThreadPreUpdateEvent& InEvent);
+
+	/** Check whether we have a valid root node */
+	bool HasRootNode() const
+	{ 
+		return RootNode != nullptr; 
+	}
+
+	/** @todo: remove after deprecation */
+	FAnimNode_Base* GetRootNode() 
+	{ 
+		return RootNode;
+	}
+
+	/** Gather debug data from this instance proxy and the blend tree for display */
+	void GatherDebugData(FNodeDebugData& DebugData);
+
+#if !NO_LOGGING
+	const FString& GetActorName() const 
+	{
+		return ActorName; 
+	}
+#endif
+
+#if DO_CHECK
+	const FString& GetAnimInstanceName() const 
+	{ 
+		return AnimInstanceName;
+	}
+#endif
+
+	/** Gets the runtime instance of the specified state machine by Name */
+	FAnimNode_StateMachine* GetStateMachineInstanceFromName(FName MachineName);
+
+	/** Get the machine description for the specified instance. Does not rely on PRIVATE_MachineDescription being initialized */
+	static const FBakedAnimationStateMachine* GetMachineDescription(UAnimBlueprintGeneratedClass* AnimBlueprintClass, FAnimNode_StateMachine* MachineInstance);
+
+	/** 
+	 * Get the index of the specified instance asset player. Useful to pass to GetInstanceAssetPlayerLength (etc.).
+	 * Passing NAME_None to InstanceName will return the first (assumed only) player instance index found.
+	 */
+	int32 GetInstanceAssetPlayerIndex(FName MachineName, FName StateName, FName InstanceName = NAME_None);
+
+	/** Only restricted classes can access the protected interface */
+	friend class UAnimInstance;
+	friend class UAnimSingleNodeInstance;
+	friend class USkeletalMeshComponent;
+
+protected:
+	/** Called when our anim instance is being initialized */
+	virtual void Initialize(UAnimInstance* InAnimInstance);
+
+	/** Called when our anim instance is being uninitialized */
+	virtual void Uninitialize(UAnimInstance* InAnimInstance);
+
+	/** Called before update so we can copy any data we need */
+	virtual void PreUpdate(const UAnimInstance* InAnimInstance, float DeltaSeconds);
+
+	/** Update override point */
+	virtual void Update(float DeltaSeconds) {}
+
+	/** Updates the anim graph */
+	virtual void UpdateAnimationNode(float DeltaSeconds);
+
+	/** 
+	 * Evaluate override point 
+	 * @return true if this function is implemented, false otherwise.
+	 * Note: the node graph will not be evaluated if this function returns true
+	 */
+	virtual bool Evaluate(FPoseContext& Output) { return false; }
+
+	/** Called after update so we can copy any data we need */
+	virtual void PostUpdate(UAnimInstance* InAnimInstance) const;
+
+	/** Calls Update(), updates the anim graph, ticks asset players */
+	void UpdateAnimation();
+
+	/** Evaluates the anim graph */
+	void EvaluateAnimation(FPoseContext& Output);
+
+	// @todo document
+	void SequenceAdvanceImmediate(UAnimSequenceBase* Sequence, bool bLooping, float PlayRate, float DeltaSeconds, /*inout*/ float& CurrentTime, FMarkerTickRecord& MarkerTickRecord);
+
+	// @todo document
+	void BlendSpaceAdvanceImmediate(UBlendSpaceBase* BlendSpace, const FVector& BlendInput, TArray<FBlendSampleData> & BlendSampleDataCache, FBlendFilter & BlendFilter, bool bLooping, float PlayRate, float DeltaSeconds, /*inout*/ float& CurrentTime, FMarkerTickRecord& MarkerTickRecord);
+
+	// Gets the sync group we should be reading from
+	int32 GetSyncGroupReadIndex() const 
+	{ 
+		return 1 - SyncGroupWriteIndex; 
+	}
+
+	// Gets the sync group we should be writing to
+	int32 GetSyncGroupWriteIndex() const 
+	{ 
+		return SyncGroupWriteIndex; 
+	}
+
+	/** Add anim notifier **/
+	void AddAnimNotifies(const TArray<const FAnimNotifyEvent*>& NewNotifies, const float InstanceWeight);
+
+	/** Returns the baked sync group index from the compile step */
+	int32 GetSyncGroupIndexFromName(FName SyncGroupName) const;
+
+	bool GetTimeToClosestMarker(FName SyncGroup, FName MarkerName, float& OutMarkerTime) const;
+
+	bool HasMarkerBeenHitThisFrame(FName SyncGroup, FName MarkerName) const;
+
+	bool IsSyncGroupBetweenMarkers(FName InSyncGroupName, FName PreviousMarker, FName NextMarker, bool bRespectMarkerOrder = true) const;
+
+	FMarkerSyncAnimPosition GetSyncGroupPosition(FName InSyncGroupName) const;
+
+	// slot node run-time functions
+	void ReinitializeSlotNodes();
+
+	// if it doesn't tick, it will keep old weight, so we'll have to clear it in the beginning of tick
+	void ClearSlotNodeWeights();
+	bool IsSlotNodeRelevantForNotifies(FName SlotNodeName) const;
+
+	// Get the root motion weight for the montage slot
+	float GetSlotRootMotionWeight(FName SlotNodeName) const;
+
+	/** 
+	 * Recalculate Required Bones [RequiredBones]
+	 * Is called when bRequiredBonesUpToDate = false
+	 */
+	void RecalcRequiredBones(USkeletalMeshComponent* Component, UObject* Asset);
+
+	/** Update morph target curves for external map (from USkeletalMeshComponent) */
+	void UpdateMorphTargetCurves(const TMap<FName, float>& InMorphTargetCurves);
+
+	/** Update the material parameters of the supplied component from this instance */
+	void UpdateComponentsMaterialParameters(UPrimitiveComponent* Component);
+
+	/** Get Currently active montage evaluation state.
+		Note that there might be multiple Active at the same time. This will only return the first active one it finds. **/
+	const FMontageEvaluationState* GetActiveMontageEvaluationState() const;
+
+	/** Access montage array data */
+	TArray<FMontageEvaluationState>& GetMontageEvaluationData() { return MontageEvaluationData; }
+
+	/** Access montage array data */
+	const TArray<FMontageEvaluationState>& GetMontageEvaluationData() const { return MontageEvaluationData; }
+
+	/** Check whether we have active morph target curves */
+	bool HasMorphTargetCurves() const { return MorphTargetCurves.Num() > 0; }
+
+	/** Access the active morph target curves */
+	TMap<FName, float>& GetMorphTargetCurves() { return MorphTargetCurves; }
+
+	/** Access the active material parameter curves */
+	TMap<FName, float>& GetMaterialParameterCurves() { return MaterialParameterCurves; }
+
+	/** Check whether we have active morph target curves */
+	TArray<struct FActiveVertexAnim>& GetVertexAnims() { return VertexAnims; }
+
+	/** Adds an active vertex anim */
+	void AddVertexAnim(const FActiveVertexAnim& VertexAnim);
+
+	/** Gets the most relevant asset player in a specified state */
+	FAnimNode_AssetPlayerBase* GetRelevantAssetPlayerFromState(int32 MachineIndex, int32 StateIndex);
+
+	/** Gets the runtime instance of the specified state machine */
+	FAnimNode_StateMachine* GetStateMachineInstance(int32 MachineIndex);
+
+	/** Gets an unchecked (can return nullptr) node given a property of the anim instance */
+	template<class NodeType>
+	NodeType* GetNodeFromProperty(UProperty* Property)
+	{
+		return (NodeType*)Property->ContainerPtrToValuePtr<NodeType>(AnimInstanceObject);
+	}
+
+	/** Gets the length in seconds of the asset referenced in an asset player node */
+	float GetInstanceAssetPlayerLength(int32 AssetPlayerIndex);
+
+	/** Get the current accumulated time in seconds for an asset player node */
+	float GetInstanceAssetPlayerTime(int32 AssetPlayerIndex);
+
+	/** Get the current accumulated time as a fraction for an asset player node */
+	float GetInstanceAssetPlayerTimeFraction(int32 AssetPlayerIndex);
+
+	/** Get the time in seconds from the end of an animation in an asset player node */
+	float GetInstanceAssetPlayerTimeFromEnd(int32 AssetPlayerIndex);
+
+	/** Get the time as a fraction of the asset length of an animation in an asset player node */
+	float GetInstanceAssetPlayerTimeFromEndFraction(int32 AssetPlayerIndex);
+
+	/** Get the blend weight of a specified state */
+	float GetInstanceStateWeight(int32 MachineIndex, int32 StateIndex);
+
+	/** Get the current elapsed time of a state within the specified state machine */
+	float GetInstanceCurrentStateElapsedTime(int32 MachineIndex);
+
+	/** Get the crossfade duration of a specified transition */
+	float GetInstanceTransitionCrossfadeDuration(int32 MachineIndex, int32 TransitionIndex);
+
+	/** Get the elapsed time in seconds of a specified transition */
+	float GetInstanceTransitionTimeElapsed(int32 MachineIndex, int32 TransitionIndex);
+
+	/** Get the elapsed time as a fraction of the crossfade duration of a specified transition */
+	float GetInstanceTransitionTimeElapsedFraction(int32 MachineIndex, int32 TransitionIndex);
+
+	/** Get the time remaining in seconds for the most relevant animation in the source state */
+	float GetRelevantAnimTimeRemaining(int32 MachineIndex, int32 StateIndex);
+
+	/** Get the time remaining as a fraction of the duration for the most relevant animation in the source state */
+	float GetRelevantAnimTimeRemainingFraction(int32 MachineIndex, int32 StateIndex);
+
+	/** Get the length in seconds of the most relevant animation in the source state */
+	float GetRelevantAnimLength(int32 MachineIndex, int32 StateIndex);
+
+	/** Get the current accumulated time in seconds for the most relevant animation in the source state */
+	float GetRelevantAnimTime(int32 MachineIndex, int32 StateIndex);
+
+	/** Get the current accumulated time as a fraction of the length of the most relevant animation in the source state */
+	float GetRelevantAnimTimeFraction(int32 MachineIndex, int32 StateIndex);
+
+	// Sets up a native transition delegate between states with PrevStateName and NextStateName, in the state machine with name MachineName.
+	// Note that a transition already has to exist for this to succeed
+	void AddNativeTransitionBinding(const FName& MachineName, const FName& PrevStateName, const FName& NextStateName, const FCanTakeTransition& NativeTransitionDelegate, const FName& TransitionName = NAME_None);
+
+	// Check for whether a native rule is bound to the specified transition
+	bool HasNativeTransitionBinding(const FName& MachineName, const FName& PrevStateName, const FName& NextStateName, FName& OutBindingName);
+
+	// Sets up a native state entry delegate from state with StateName, in the state machine with name MachineName.
+	void AddNativeStateEntryBinding(const FName& MachineName, const FName& StateName, const FOnGraphStateChanged& NativeEnteredDelegate);
+	
+	// Check for whether a native entry delegate is bound to the specified state
+	bool HasNativeStateEntryBinding(const FName& MachineName, const FName& StateName, FName& OutBindingName);
+
+	// Sets up a native state exit delegate from state with StateName, in the state machine with name MachineName.
+	void AddNativeStateExitBinding(const FName& MachineName, const FName& StateName, const FOnGraphStateChanged& NativeExitedDelegate);
+
+	// Check for whether a native exit delegate is bound to the specified state
+	bool HasNativeStateExitBinding(const FName& MachineName, const FName& StateName, FName& OutBindingName);
+
+	/** Bind any native delegates that we have set up */
+	void BindNativeDelegates();
+
+	/** Gets the runtime instance desc of the state machine specified by name */
+	const FBakedAnimationStateMachine* GetStateMachineInstanceDesc(FName MachineName);
+
+	/** Gets the index of the state machine matching MachineName */
+	int32 GetStateMachineIndex(FName MachineName);
+
+	/** Initialize the root node - split into a separate function for backwards compatibility (initialization order) reasons */
+	void InitializeRootNode();
+
+private:
+	/** Object ptr to our UAnimInstance */
+	mutable UObject* AnimInstanceObject;
+
+	/** Our anim blueprint generated class */
+	UAnimBlueprintGeneratedClass* AnimBlueprintClass;
+
+	/** Skeleton we are using, only used for comparison purposes */
+	USkeleton* Skeleton;
+
+	/** Skeletal mesh component we are attached to */
+	USkeletalMeshComponent* SkeletalMeshComponent;
+
+	/** The last time passed into PreUpdate() */
+	float CurrentDeltaSeconds;
+
+#if WITH_EDITORONLY_DATA
+	/** Whether this UAnimInstance is currently being debugged in the editor */
+	bool bIsBeingDebugged;
+
+	/** Array of visited nodes this frame */
+	TArray<FAnimBlueprintDebugData::FNodeVisit> UpdatedNodesThisFrame;
+#endif
+
+#if !NO_LOGGING
+	/** Actor name for debug logging purposes */
+	FString ActorName;
+#endif
+
+#if DO_CHECK
+	/** Anim instance name for debug purposes */
+	FString AnimInstanceName;
+#endif
+
+	/** Anim graph */
+	FAnimNode_Base* RootNode;
+
+	/** The list of animation assets which are going to be evaluated this frame and need to be ticked (ungrouped) */
+	TArray<FAnimTickRecord> UngroupedActivePlayerArrays[2];
+
+	/** The set of tick groups for this anim instance */
+	TArray<FAnimGroupInstance> SyncGroupArrays[2];
+
+	// Current sync group buffer index
+	int32 SyncGroupWriteIndex;
+
+	/** Animation Notifies that has been triggered in the latest tick **/
+	FAnimNotifyQueue NotifyQueue;
+
+	// Root motion mode duplicated from the anim instance
+	ERootMotionMode::Type RootMotionMode;
+
+	// Diplicate of bool result of ShouldExtractRootMotion()
+	bool bShouldExtractRootMotion;
+
+	// Tracker map for slot name->weights/relevancy
+	TMap<FName, FMontageActiveSlotTracker> SlotWeightTracker;
+
+	// Counters for synchronization
+	FGraphTraversalCounter InitializationCounter;
+	FGraphTraversalCounter CachedBonesCounter;
+	FGraphTraversalCounter UpdateCounter;
+	FGraphTraversalCounter EvaluationCounter;
+	FGraphTraversalCounter SlotNodeInitializationCounter;
+
+	// Root motion extracted from animation since the last time ConsumeExtractedRootMotion was called
+	FRootMotionMovementParams ExtractedRootMotion;
+
+	/** Temporary array of bone indices required this frame. Should be subset of Skeleton and Mesh's RequiredBones */
+	FBoneContainer RequiredBones;
+
+	/** When RequiredBones mapping has changed, AnimNodes need to update their bones caches. */
+	bool bBoneCachesInvalidated;
+
+	/** Morph Target Curves that will be used for SkeletalMeshComponent **/
+	TMap<FName, float> MorphTargetCurves;
+
+	/** Material Curves that will be used for SkeletalMeshComponent **/
+	TMap<FName, float> MaterialParameterCurves;
+
+	/** Copy of UAnimInstance::MontageInstances data used for update & evaluation */
+	TArray<FMontageEvaluationState> MontageEvaluationData;
+
+	/** Array indicating active vertex anims (by reference) generated by anim instance. */
+	TArray<struct FActiveVertexAnim> VertexAnims;
+
+	/** Delegate fired on the game thread before update occurs */
+	TArray<FGameThreadPreUpdateEvent> GameThreadPreUpdateEvents;
+
+	/** Native transition rules */
+	TArray<FNativeTransitionBinding> NativeTransitionBindings;
+
+	/** Native state entry bindings */
+	TArray<FNativeStateBinding> NativeStateEntryBindings;
+
+	/** Native state exit bindings */
+	TArray<FNativeStateBinding> NativeStateExitBindings;
+
+};
diff --git a/Engine/Source/Runtime/Engine/Public/Animation/AnimMTStats.h b/Engine/Source/Runtime/Engine/Public/Animation/AnimMTStats.h
--- a/Engine/Source/Runtime/Engine/Public/Animation/AnimMTStats.h
+++ b/Engine/Source/Runtime/Engine/Public/Animation/AnimMTStats.h
@@ -8,4 +8,4 @@ DO_ANIMSTAT_PROCESSING(DecompressAnimationPose);
 DO_ANIMSTAT_PROCESSING(BlendPosesInGraph);
 DO_ANIMSTAT_PROCESSING(EvaluateAnimStateMachine);
 DO_ANIMSTAT_PROCESSING(FillSpaceBases);
-DO_ANIMSTAT_PROCESSING(UpdateAnimationInternal);
\ No newline at end of file
+DO_ANIMSTAT_PROCESSING(ProxyUpdateAnimation);
\ No newline at end of file
diff --git a/Engine/Source/Runtime/Engine/Public/Animation/AnimPhysicsSolver.h b/Engine/Source/Runtime/Engine/Public/Animation/AnimPhysicsSolver.h
--- a/Engine/Source/Runtime/Engine/Public/Animation/AnimPhysicsSolver.h
+++ b/Engine/Source/Runtime/Engine/Public/Animation/AnimPhysicsSolver.h
@@ -30,7 +30,7 @@
 
 #pragma once
 
-DECLARE_CYCLE_STAT_EXTERN(TEXT("Anim Dynamics Update"), STAT_AnimDynamicsUpdate, STATGROUP_Physics, ENGINE_API);
+DECLARE_CYCLE_STAT_EXTERN(TEXT("Anim Dynamics Physics Update"), STAT_AnimDynamicsUpdate, STATGROUP_Physics, ENGINE_API);
 DECLARE_CYCLE_STAT_EXTERN(TEXT("Linear Limit Pre-Update"), STAT_AnimDynamicsLinearPre, STATGROUP_Physics, ENGINE_API);
 DECLARE_CYCLE_STAT_EXTERN(TEXT("Linear Limit Post-Update"), STAT_AnimDynamicsLinearPost, STATGROUP_Physics, ENGINE_API);
 DECLARE_CYCLE_STAT_EXTERN(TEXT("Angular Limit Pre-Update"), STAT_AnimDynamicsAngularPre, STATGROUP_Physics, ENGINE_API);
@@ -289,6 +289,9 @@ public:
 	// Maximum torque this limit can apply
 	float  MaximumTorque;
 
+	// Cached spin to torque value that is independant of iterations
+	float CachedSpinToTorque;
+
 	FAnimPhysAngularLimit();
 	FAnimPhysAngularLimit(FAnimPhysRigidBody* InFirstBody, FAnimPhysRigidBody* InSecondBody, const FVector& InWorldSpaceAxis, float InTargetSpin = 0, float InMinimumTorque = -MAX_flt, float InMaximumTorque = MAX_flt);	
 
@@ -299,6 +302,8 @@ public:
 	/** Solve the limit
 	 */
 	void Iter(float DeltaTime);
+
+	void UpdateCachedData();
 };
 
 class ENGINE_API FAnimPhysLinearLimit : public FAnimPhysLimit
@@ -329,6 +334,15 @@ class ENGINE_API FAnimPhysLinearLimit : public FAnimPhysLimit
 	// Sum of impulses applied
 	float  SumImpulses;
 
+	// Cached denominator of the impulse calculation, doesn't change across iterations
+	float InverseInertiaImpulse;
+
+	// Cached world space position on body 0
+	FVector WorldSpacePosition0;
+
+	// Cached world space position on body 1
+	FVector WorldSpacePosition1;
+
 	FAnimPhysLinearLimit();
 	FAnimPhysLinearLimit(FAnimPhysRigidBody* InFirstBody, FAnimPhysRigidBody* InSecondBody, const FVector& InFirstPosition, const FVector& InSecondPosition,
 		const FVector& InNormal = FVector(0.0f, 0.0f, 1.0f), float InTargetSpeed = 0.0f, float InTargetSpeedWithoutBias = 0.0f, const FVector2D& InForceRange = FVector2D(-MAX_flt, MAX_flt));
@@ -340,6 +354,8 @@ class ENGINE_API FAnimPhysLinearLimit : public FAnimPhysLimit
 	/** Solve the limit
 	*/
 	void Iter(float DetlaTime);
+
+	void UpdateCachedData();
 };
 
 struct ENGINE_API FAnimPhysSpring
diff --git a/Engine/Source/Runtime/Engine/Public/Animation/AnimSingleNodeInstanceProxy.h b/Engine/Source/Runtime/Engine/Public/Animation/AnimSingleNodeInstanceProxy.h
new file mode 100644
--- /dev/null
+++ b/Engine/Source/Runtime/Engine/Public/Animation/AnimSingleNodeInstanceProxy.h
@@ -0,0 +1,144 @@
+// Copyright 1998-2015 Epic Games, Inc. All Rights Reserved.
+
+#pragma once
+
+#include "AnimInstanceProxy.h"
+#include "AnimSingleNodeInstanceProxy.generated.h"
+
+/** Proxy override for this UAnimInstance-derived class */
+USTRUCT()
+struct ENGINE_API FAnimSingleNodeInstanceProxy : public FAnimInstanceProxy
+{
+	GENERATED_BODY()
+
+public:
+	FAnimSingleNodeInstanceProxy()
+	{
+	}
+
+	FAnimSingleNodeInstanceProxy(UAnimInstance* InAnimInstance)
+		: FAnimInstanceProxy(InAnimInstance)
+		, CurrentAsset(nullptr)
+		, CurrentVertexAnim(nullptr)
+		, BlendSpaceInput(0.0f, 0.0f, 0.0f)
+		, CurrentTime(0.0f)
+#if WITH_EDITORONLY_DATA
+		, PreviewPoseCurrentTime(0.0f)
+#endif
+		, PlayRate(1.f)
+		, bLooping(true)
+		, bPlaying(true)
+		, bReverse(false)
+	{}
+
+	virtual void Initialize(UAnimInstance* InAnimInstance) override;
+	virtual bool Evaluate(FPoseContext& Output) override;
+	virtual void UpdateAnimationNode(float DeltaSeconds) override;
+	virtual void PostUpdate(UAnimInstance* InAnimInstance) const override;
+
+	void SetPlaying(bool bIsPlaying)
+	{
+		bPlaying = bIsPlaying;
+	}
+
+	bool IsPlaying() const
+	{
+		return bPlaying;
+	}
+
+	bool IsReverse() const
+	{
+		return bReverse;
+	}
+
+	void SetLooping(bool bIsLooping)
+	{
+		bLooping = bIsLooping;
+	}
+
+	bool IsLooping() const
+	{
+		return bLooping;
+	}
+
+	void ResetWeightInfo()
+	{
+		WeightInfo.Reset();
+	}
+
+	virtual void SetAnimationAsset(UAnimationAsset* NewAsset, USkeletalMeshComponent* MeshComponent, bool bIsLooping, float InPlayRate);
+
+	UAnimationAsset* GetCurrentAsset() { return CurrentAsset; }
+
+	UVertexAnimation* GetCurrentVertexAnimation() { return CurrentVertexAnim; }
+
+	void UpdateBlendspaceSamples(FVector InBlendInput);
+
+	void SetCurrentTime(float InCurrentTime)
+	{
+		CurrentTime = InCurrentTime;
+	}
+
+	float GetCurrentTime() const
+	{
+		return CurrentTime;
+	}
+
+	float GetPlayRate() const
+	{
+		return PlayRate;
+	}
+
+	void SetPlayRate(float InPlayRate)
+	{
+		PlayRate = InPlayRate;
+	}
+
+	void SetVertexAnimation(UVertexAnimation * NewVertexAnim, bool bIsLooping, float InPlayRate);
+
+	void SetReverse(bool bInReverse);
+
+	float GetLength();
+
+	void SetBlendSpaceInput(const FVector& InBlendInput);
+
+private:
+	void InternalBlendSpaceEvaluatePose(class UBlendSpaceBase* BlendSpace, TArray<FBlendSampleData>& BlendSampleDataCache, FPoseContext& OutContext);
+
+private:
+	/** Current Asset being played **/
+	UAnimationAsset* CurrentAsset;
+
+	/** Current vertex anim being played **/
+	UVertexAnimation* CurrentVertexAnim;
+
+	/** Random cached values to play each asset **/
+	FVector BlendSpaceInput;
+
+	/** Random cached values to play each asset **/
+	TArray<FBlendSampleData> BlendSampleData;
+
+	/** Random cached values to play each asset **/
+	FBlendFilter BlendFilter;
+
+	/** Slot node weight transient data */
+	FSlotNodeWeightInfo WeightInfo;
+
+	/** Shared parameters for previewing blendspace or animsequence **/
+	float CurrentTime;
+
+#if WITH_EDITORONLY_DATA
+	float PreviewPoseCurrentTime;
+#endif
+
+	/** Cache for data needed during marker sync */
+	FMarkerTickRecord MarkerTickRecord;
+
+	float PlayRate;
+
+	uint32 bLooping:1;
+
+	uint32 bPlaying:1;
+
+	uint32 bReverse:1;
+};
\ No newline at end of file
diff --git a/Engine/Source/Runtime/Engine/Public/Animation/AnimStats.h b/Engine/Source/Runtime/Engine/Public/Animation/AnimStats.h
--- a/Engine/Source/Runtime/Engine/Public/Animation/AnimStats.h
+++ b/Engine/Source/Runtime/Engine/Public/Animation/AnimStats.h
@@ -40,6 +40,8 @@ DECLARE_CYCLE_STAT_EXTERN(TEXT("AnimMontageInstance_TickBranchPoints"), STAT_Ani
 DECLARE_CYCLE_STAT_EXTERN(TEXT("AnimMontageInstance_Advance_Iteration"), STAT_AnimMontageInstance_Advance_Iteration, STATGROUP_Anim, );
 DECLARE_CYCLE_STAT_EXTERN(TEXT("UpdateCurves"), STAT_UpdateCurves, STATGROUP_Anim, );
 DECLARE_CYCLE_STAT_EXTERN(TEXT("LocalBlendCSBoneTransforms"), STAT_LocalBlendCSBoneTransforms, STATGROUP_Anim, ENGINE_API);
+DECLARE_CYCLE_STAT_EXTERN(TEXT("TickAssetPlayerInstances"), STAT_TickAssetPlayerInstances, STATGROUP_Anim, );
+DECLARE_CYCLE_STAT_EXTERN(TEXT("TickAssetPlayerInstance"), STAT_TickAssetPlayerInstance, STATGROUP_Anim, );
 
 #define DO_ANIMSTAT_PROCESSING(StatName) DECLARE_CYCLE_STAT_EXTERN(TEXT(#StatName), STAT_ ## StatName, STATGROUP_Anim, ENGINE_API)
 #include "AnimMTStats.h"
diff --git a/Engine/Source/Runtime/Engine/Public/AnimationRuntime.h b/Engine/Source/Runtime/Engine/Public/AnimationRuntime.h
--- a/Engine/Source/Runtime/Engine/Public/AnimationRuntime.h
+++ b/Engine/Source/Runtime/Engine/Public/AnimationRuntime.h
@@ -263,6 +263,11 @@ public:
 	static void FillUpSpaceBasesRetargetBasePose(const USkeleton* Skeleton, TArray<FTransform> &SpaceBaseRefPose);
 #endif
 
+	/**
+	* Combine CurveKeys (that reference morph targets by name) and ActiveAnims (that reference vertex anims by reference) into the ActiveVertexAnims array.
+	*/
+	static TArray<struct FActiveVertexAnim> UpdateActiveVertexAnims(const USkeletalMesh* InSkeletalMesh, const TMap<FName, float>& InMorphCurveAnims, const TArray<FActiveVertexAnim>& InActiveAnims);
+
 private:
 	/** 
 	* Blend Poses per bone weights : The BasePose + BlendPoses(SourceIndex) * Blend Weights(BoneIndex)
diff --git a/Engine/Source/Runtime/Engine/Public/Audio.h b/Engine/Source/Runtime/Engine/Public/Audio.h
--- a/Engine/Source/Runtime/Engine/Public/Audio.h
+++ b/Engine/Source/Runtime/Engine/Public/Audio.h
@@ -191,10 +191,8 @@ struct ENGINE_API FWaveInstance
 	float				Volume;
 	/** Current volume multiplier - used to zero the volume without stopping the source */
 	float				VolumeMultiplier;
-	/** Current priority (deprecated) */
-	float				PlayPriority;
-	/** A priority value that scales with volume (post all gain stages) and is used to determin voice playback priority. */
-	float				VolumeWeightedPriorityScale;
+	/** An audio component priority value that scales with volume (post all gain stages) and is used to determine voice playback priority. */
+	float				Priority;
 	/** Voice center channel volume */
 	float				VoiceCenterChannelVolume;
 	/** Volume of the radio filter effect */
diff --git a/Engine/Source/Runtime/Engine/Public/BoneContainer.h b/Engine/Source/Runtime/Engine/Public/BoneContainer.h
--- a/Engine/Source/Runtime/Engine/Public/BoneContainer.h
+++ b/Engine/Source/Runtime/Engine/Public/BoneContainer.h
@@ -210,6 +210,7 @@ public:
 			<< B.PoseToSkeletonBoneIndexArray
 			<< B.bDisableRetargeting
 			<< B.bUseRAWData
+			<< B.bUseSourceData
 			;
 		return Ar;
 	}
diff --git a/Engine/Source/Runtime/Engine/Public/Engine.h b/Engine/Source/Runtime/Engine/Public/Engine.h
--- a/Engine/Source/Runtime/Engine/Public/Engine.h
+++ b/Engine/Source/Runtime/Engine/Public/Engine.h
@@ -112,7 +112,6 @@
 #include "Components/ExponentialHeightFogComponent.h"
 #include "Components/DirectionalLightComponent.h"
 #include "Components/SkyLightComponent.h"
-#include "PhysicsEngine/PhysicsThrusterComponent.h"
 #include "Components/PostProcessComponent.h"
 #include "Components/ArrowComponent.h"
 #include "Components/BillboardComponent.h"
diff --git a/Engine/Source/Runtime/Engine/Public/MaterialShared.h b/Engine/Source/Runtime/Engine/Public/MaterialShared.h
--- a/Engine/Source/Runtime/Engine/Public/MaterialShared.h
+++ b/Engine/Source/Runtime/Engine/Public/MaterialShared.h
@@ -358,6 +358,8 @@ namespace EMaterialShaderMapUsage
 		MaterialExportRoughness,
 		MaterialExportAO,
 		MaterialExportEmissive,
+		MaterialExportOpacity,
+		MaterialExportSubSurfaceColor
 	};
 }
 
diff --git a/Engine/Source/Runtime/Engine/Public/PhysXPublic.h b/Engine/Source/Runtime/Engine/Public/PhysXPublic.h
--- a/Engine/Source/Runtime/Engine/Public/PhysXPublic.h
+++ b/Engine/Source/Runtime/Engine/Public/PhysXPublic.h
@@ -157,7 +157,7 @@ private:
 #endif
 
 /** Get a pointer to the PxScene from an SceneIndex (will be NULL if scene already shut down) */
-PxScene* GetPhysXSceneFromIndex(int32 InSceneIndex);
+ENGINE_API PxScene* GetPhysXSceneFromIndex(int32 InSceneIndex);
 
 template <bool NeedsLock>
 struct FPhysXSupport
diff --git a/Engine/Source/Runtime/Engine/Public/TriangleRendering.h b/Engine/Source/Runtime/Engine/Public/TriangleRendering.h
--- a/Engine/Source/Runtime/Engine/Public/TriangleRendering.h
+++ b/Engine/Source/Runtime/Engine/Public/TriangleRendering.h
@@ -1,7 +1,7 @@
 // Copyright 1998-2015 Epic Games, Inc. All Rights Reserved.
 
 /*=============================================================================
-	TriangleRendering.h: Simple tile rendering implementation.
+	TriangleRendering.h: Simple triangle rendering implementation.
 =============================================================================*/
 
 #ifndef _INC_TRIANGLERENDERING
@@ -12,10 +12,14 @@ class FTriangleRenderer
 public:
 
 	/**
-	 * Draw a tile at the given location and size, using the given UVs
-	 * (UV = [0..1]
+	 * Draw a triangle using the given vertices using optional color.
 	 */
 	ENGINE_API static void DrawTriangle(FRHICommandListImmediate& RHICmdList, const class FSceneView& View, const FMaterialRenderProxy* MaterialRenderProxy, bool bNeedsToSwitchVerticalAxis, const FCanvasUVTri& Tri, bool bIsHitTesting = false, const FHitProxyId HitProxyId = FHitProxyId(), const FColor InVertexColor = FColor(255, 255, 255, 255));
+	
+	/**
+	* Draw a triangle using the given vertices. Same as DrawTriangle but uses the colors in the supplied vertices
+	*/
+	ENGINE_API static void DrawTriangleUsingVertexColor(FRHICommandListImmediate& RHICmdList, const class FSceneView& View, const FMaterialRenderProxy* MaterialRenderProxy, bool bNeedsToSwitchVerticalAxis, const FCanvasUVTri& Tri, bool bIsHitTesting, const FHitProxyId HitProxyId);
 
 private:
 
diff --git a/Engine/Source/Runtime/Engine/Public/VisualLogger/VisualLogger.h b/Engine/Source/Runtime/Engine/Public/VisualLogger/VisualLogger.h
--- a/Engine/Source/Runtime/Engine/Public/VisualLogger/VisualLogger.h
+++ b/Engine/Source/Runtime/Engine/Public/VisualLogger/VisualLogger.h
@@ -126,6 +126,9 @@ DECLARE_LOG_CATEGORY_EXTERN(LogVisual, Display, All);
 
 class FVisualLogDevice;
 class FVisualLogExtensionInterface;
+class UObject;
+class UWorld;
+struct FLogCategoryBase;
 
 DECLARE_DELEGATE_RetVal(FString, FVisualLogFilenameGetterDelegate);
 
@@ -133,39 +136,39 @@ class ENGINE_API FVisualLogger : public FOutputDevice
 {
 public:
 	// Regular text log
-	VARARG_DECL(static void, static void, return, CategorizedLogf, VARARG_NONE, const TCHAR*, VARARG_EXTRA(const class UObject* LogOwner) VARARG_EXTRA(const struct FLogCategoryBase& Category) VARARG_EXTRA(ELogVerbosity::Type Verbosity) VARARG_EXTRA(int32 UniqueLogId), VARARG_EXTRA(LogOwner) VARARG_EXTRA(Category) VARARG_EXTRA(Verbosity) VARARG_EXTRA(UniqueLogId));
+	VARARG_DECL(static void, static void, return, CategorizedLogf, VARARG_NONE, const TCHAR*, VARARG_EXTRA(const UObject* LogOwner) VARARG_EXTRA(const FLogCategoryBase& Category) VARARG_EXTRA(ELogVerbosity::Type Verbosity) VARARG_EXTRA(int32 UniqueLogId), VARARG_EXTRA(LogOwner) VARARG_EXTRA(Category) VARARG_EXTRA(Verbosity) VARARG_EXTRA(UniqueLogId));
 	// Segment log
-	VARARG_DECL(static void, static void, return, GeometryShapeLogf, VARARG_NONE, const TCHAR*, VARARG_EXTRA(const class UObject* LogOwner) VARARG_EXTRA(const struct FLogCategoryBase& Category) VARARG_EXTRA(ELogVerbosity::Type Verbosity) VARARG_EXTRA(int32 UniqueLogId) VARARG_EXTRA(const FVector& Start) VARARG_EXTRA(const FVector& End) VARARG_EXTRA(const FColor& Color), VARARG_EXTRA(LogOwner) VARARG_EXTRA(Category) VARARG_EXTRA(Verbosity) VARARG_EXTRA(UniqueLogId) VARARG_EXTRA(Start) VARARG_EXTRA(End) VARARG_EXTRA(Color));
+	VARARG_DECL(static void, static void, return, GeometryShapeLogf, VARARG_NONE, const TCHAR*, VARARG_EXTRA(const UObject* LogOwner) VARARG_EXTRA(const FLogCategoryBase& Category) VARARG_EXTRA(ELogVerbosity::Type Verbosity) VARARG_EXTRA(int32 UniqueLogId) VARARG_EXTRA(const FVector& Start) VARARG_EXTRA(const FVector& End) VARARG_EXTRA(const FColor& Color), VARARG_EXTRA(LogOwner) VARARG_EXTRA(Category) VARARG_EXTRA(Verbosity) VARARG_EXTRA(UniqueLogId) VARARG_EXTRA(Start) VARARG_EXTRA(End) VARARG_EXTRA(Color));
 	// Location/Shpere log
-	VARARG_DECL(static void, static void, return, GeometryShapeLogf, VARARG_NONE, const TCHAR*, VARARG_EXTRA(const class UObject* LogOwner) VARARG_EXTRA(const struct FLogCategoryBase& Category) VARARG_EXTRA(ELogVerbosity::Type Verbosity) VARARG_EXTRA(int32 UniqueLogId) VARARG_EXTRA(const FVector& Location) VARARG_EXTRA(float Radius) VARARG_EXTRA(const FColor& Color), VARARG_EXTRA(LogOwner) VARARG_EXTRA(Category) VARARG_EXTRA(Verbosity) VARARG_EXTRA(UniqueLogId) VARARG_EXTRA(Location) VARARG_EXTRA(Radius) VARARG_EXTRA(Color));
+	VARARG_DECL(static void, static void, return, GeometryShapeLogf, VARARG_NONE, const TCHAR*, VARARG_EXTRA(const UObject* LogOwner) VARARG_EXTRA(const FLogCategoryBase& Category) VARARG_EXTRA(ELogVerbosity::Type Verbosity) VARARG_EXTRA(int32 UniqueLogId) VARARG_EXTRA(const FVector& Location) VARARG_EXTRA(float Radius) VARARG_EXTRA(const FColor& Color), VARARG_EXTRA(LogOwner) VARARG_EXTRA(Category) VARARG_EXTRA(Verbosity) VARARG_EXTRA(UniqueLogId) VARARG_EXTRA(Location) VARARG_EXTRA(Radius) VARARG_EXTRA(Color));
 	// Box log
-	VARARG_DECL(static void, static void, return, GeometryShapeLogf, VARARG_NONE, const TCHAR*, VARARG_EXTRA(const class UObject* LogOwner) VARARG_EXTRA(const struct FLogCategoryBase& Category) VARARG_EXTRA(ELogVerbosity::Type Verbosity) VARARG_EXTRA(int32 UniqueLogId) VARARG_EXTRA(const FBox& Box) VARARG_EXTRA(const FMatrix& Matrix) VARARG_EXTRA(const FColor& Color), VARARG_EXTRA(LogOwner) VARARG_EXTRA(Category) VARARG_EXTRA(Verbosity) VARARG_EXTRA(UniqueLogId) VARARG_EXTRA(Box) VARARG_EXTRA(Matrix) VARARG_EXTRA(Color));
+	VARARG_DECL(static void, static void, return, GeometryShapeLogf, VARARG_NONE, const TCHAR*, VARARG_EXTRA(const UObject* LogOwner) VARARG_EXTRA(const FLogCategoryBase& Category) VARARG_EXTRA(ELogVerbosity::Type Verbosity) VARARG_EXTRA(int32 UniqueLogId) VARARG_EXTRA(const FBox& Box) VARARG_EXTRA(const FMatrix& Matrix) VARARG_EXTRA(const FColor& Color), VARARG_EXTRA(LogOwner) VARARG_EXTRA(Category) VARARG_EXTRA(Verbosity) VARARG_EXTRA(UniqueLogId) VARARG_EXTRA(Box) VARARG_EXTRA(Matrix) VARARG_EXTRA(Color));
 	// Cone log
-	VARARG_DECL(static void, static void, return, GeometryShapeLogf, VARARG_NONE, const TCHAR*, VARARG_EXTRA(const class UObject* LogOwner) VARARG_EXTRA(const struct FLogCategoryBase& Category) VARARG_EXTRA(ELogVerbosity::Type Verbosity) VARARG_EXTRA(int32 UniqueLogId) VARARG_EXTRA(const FVector& Orgin) VARARG_EXTRA(const FVector& Direction) VARARG_EXTRA(const float Length) VARARG_EXTRA(const float Angle) VARARG_EXTRA(const FColor& Color), VARARG_EXTRA(LogOwner) VARARG_EXTRA(Category) VARARG_EXTRA(Verbosity) VARARG_EXTRA(UniqueLogId) VARARG_EXTRA(Orgin) VARARG_EXTRA(Direction) VARARG_EXTRA(Length) VARARG_EXTRA(Angle) VARARG_EXTRA(Color));
+	VARARG_DECL(static void, static void, return, GeometryShapeLogf, VARARG_NONE, const TCHAR*, VARARG_EXTRA(const UObject* LogOwner) VARARG_EXTRA(const FLogCategoryBase& Category) VARARG_EXTRA(ELogVerbosity::Type Verbosity) VARARG_EXTRA(int32 UniqueLogId) VARARG_EXTRA(const FVector& Orgin) VARARG_EXTRA(const FVector& Direction) VARARG_EXTRA(const float Length) VARARG_EXTRA(const float Angle) VARARG_EXTRA(const FColor& Color), VARARG_EXTRA(LogOwner) VARARG_EXTRA(Category) VARARG_EXTRA(Verbosity) VARARG_EXTRA(UniqueLogId) VARARG_EXTRA(Orgin) VARARG_EXTRA(Direction) VARARG_EXTRA(Length) VARARG_EXTRA(Angle) VARARG_EXTRA(Color));
 	// Cylinder log
-	VARARG_DECL(static void, static void, return, GeometryShapeLogf, VARARG_NONE, const TCHAR*, VARARG_EXTRA(const class UObject* LogOwner) VARARG_EXTRA(const struct FLogCategoryBase& Category) VARARG_EXTRA(ELogVerbosity::Type Verbosity) VARARG_EXTRA(int32 UniqueLogId) VARARG_EXTRA(const FVector& Start) VARARG_EXTRA(const FVector& End) VARARG_EXTRA(const float Radius) VARARG_EXTRA(const FColor& Color), VARARG_EXTRA(LogOwner) VARARG_EXTRA(Category) VARARG_EXTRA(Verbosity) VARARG_EXTRA(UniqueLogId) VARARG_EXTRA(Start) VARARG_EXTRA(End) VARARG_EXTRA(Radius) VARARG_EXTRA(Color));
+	VARARG_DECL(static void, static void, return, GeometryShapeLogf, VARARG_NONE, const TCHAR*, VARARG_EXTRA(const UObject* LogOwner) VARARG_EXTRA(const FLogCategoryBase& Category) VARARG_EXTRA(ELogVerbosity::Type Verbosity) VARARG_EXTRA(int32 UniqueLogId) VARARG_EXTRA(const FVector& Start) VARARG_EXTRA(const FVector& End) VARARG_EXTRA(const float Radius) VARARG_EXTRA(const FColor& Color), VARARG_EXTRA(LogOwner) VARARG_EXTRA(Category) VARARG_EXTRA(Verbosity) VARARG_EXTRA(UniqueLogId) VARARG_EXTRA(Start) VARARG_EXTRA(End) VARARG_EXTRA(Radius) VARARG_EXTRA(Color));
 	// Capsule log
-	VARARG_DECL(static void, static void, return, GeometryShapeLogf, VARARG_NONE, const TCHAR*, VARARG_EXTRA(const class UObject* LogOwner) VARARG_EXTRA(const struct FLogCategoryBase& Category) VARARG_EXTRA(ELogVerbosity::Type Verbosity) VARARG_EXTRA(int32 UniqueLogId) VARARG_EXTRA(const FVector& Center) VARARG_EXTRA(float HalfHeight) VARARG_EXTRA(float Radius) VARARG_EXTRA(const FQuat & Rotation) VARARG_EXTRA(const FColor& Color), VARARG_EXTRA(LogOwner) VARARG_EXTRA(Category) VARARG_EXTRA(Verbosity) VARARG_EXTRA(UniqueLogId) VARARG_EXTRA(Center) VARARG_EXTRA(HalfHeight) VARARG_EXTRA(Radius) VARARG_EXTRA(Rotation) VARARG_EXTRA(Color));
+	VARARG_DECL(static void, static void, return, GeometryShapeLogf, VARARG_NONE, const TCHAR*, VARARG_EXTRA(const UObject* LogOwner) VARARG_EXTRA(const FLogCategoryBase& Category) VARARG_EXTRA(ELogVerbosity::Type Verbosity) VARARG_EXTRA(int32 UniqueLogId) VARARG_EXTRA(const FVector& Center) VARARG_EXTRA(float HalfHeight) VARARG_EXTRA(float Radius) VARARG_EXTRA(const FQuat & Rotation) VARARG_EXTRA(const FColor& Color), VARARG_EXTRA(LogOwner) VARARG_EXTRA(Category) VARARG_EXTRA(Verbosity) VARARG_EXTRA(UniqueLogId) VARARG_EXTRA(Center) VARARG_EXTRA(HalfHeight) VARARG_EXTRA(Radius) VARARG_EXTRA(Rotation) VARARG_EXTRA(Color));
 	// NavArea/Extruded convex log
-	VARARG_DECL(static void, static void, return, NavAreaShapeLogf, VARARG_NONE, const TCHAR*, VARARG_EXTRA(const class UObject* LogOwner) VARARG_EXTRA(const struct FLogCategoryBase& Category) VARARG_EXTRA(ELogVerbosity::Type Verbosity) VARARG_EXTRA(int32 UniqueLogId) VARARG_EXTRA(const TArray<FVector> &ConvexPoints) VARARG_EXTRA(float MinZ) VARARG_EXTRA(float MaxZ) VARARG_EXTRA(const FColor& Color), VARARG_EXTRA(LogOwner) VARARG_EXTRA(Category) VARARG_EXTRA(Verbosity) VARARG_EXTRA(UniqueLogId) VARARG_EXTRA(ConvexPoints) VARARG_EXTRA(MinZ) VARARG_EXTRA(MaxZ) VARARG_EXTRA(Color));
+	VARARG_DECL(static void, static void, return, NavAreaShapeLogf, VARARG_NONE, const TCHAR*, VARARG_EXTRA(const UObject* LogOwner) VARARG_EXTRA(const FLogCategoryBase& Category) VARARG_EXTRA(ELogVerbosity::Type Verbosity) VARARG_EXTRA(int32 UniqueLogId) VARARG_EXTRA(const TArray<FVector> &ConvexPoints) VARARG_EXTRA(float MinZ) VARARG_EXTRA(float MaxZ) VARARG_EXTRA(const FColor& Color), VARARG_EXTRA(LogOwner) VARARG_EXTRA(Category) VARARG_EXTRA(Verbosity) VARARG_EXTRA(UniqueLogId) VARARG_EXTRA(ConvexPoints) VARARG_EXTRA(MinZ) VARARG_EXTRA(MaxZ) VARARG_EXTRA(Color));
 	// 3d Mesh log
-	VARARG_DECL(static void, static void, return, GeometryShapeLogf, VARARG_NONE, const TCHAR*, VARARG_EXTRA(const class UObject* LogOwner) VARARG_EXTRA(const struct FLogCategoryBase& Category) VARARG_EXTRA(ELogVerbosity::Type Verbosity) VARARG_EXTRA(int32 UniqueLogId) VARARG_EXTRA(const TArray<FVector> &Vertices) VARARG_EXTRA(const TArray<int32> &Indices) VARARG_EXTRA(const FColor& Color), VARARG_EXTRA(LogOwner) VARARG_EXTRA(Category) VARARG_EXTRA(Verbosity) VARARG_EXTRA(UniqueLogId) VARARG_EXTRA(Vertices) VARARG_EXTRA(Indices) VARARG_EXTRA(Color));
+	VARARG_DECL(static void, static void, return, GeometryShapeLogf, VARARG_NONE, const TCHAR*, VARARG_EXTRA(const UObject* LogOwner) VARARG_EXTRA(const FLogCategoryBase& Category) VARARG_EXTRA(ELogVerbosity::Type Verbosity) VARARG_EXTRA(int32 UniqueLogId) VARARG_EXTRA(const TArray<FVector> &Vertices) VARARG_EXTRA(const TArray<int32> &Indices) VARARG_EXTRA(const FColor& Color), VARARG_EXTRA(LogOwner) VARARG_EXTRA(Category) VARARG_EXTRA(Verbosity) VARARG_EXTRA(UniqueLogId) VARARG_EXTRA(Vertices) VARARG_EXTRA(Indices) VARARG_EXTRA(Color));
 	// 2d Convex shape
-	VARARG_DECL(static void, static void, return, GeometryConvexLogf, VARARG_NONE, const TCHAR*, VARARG_EXTRA(const class UObject* LogOwner) VARARG_EXTRA(const struct FLogCategoryBase& Category) VARARG_EXTRA(ELogVerbosity::Type Verbosity) VARARG_EXTRA(int32 UniqueLogId) VARARG_EXTRA(const TArray<FVector> &Points) VARARG_EXTRA(const FColor& Color), VARARG_EXTRA(LogOwner) VARARG_EXTRA(Category) VARARG_EXTRA(Verbosity) VARARG_EXTRA(UniqueLogId) VARARG_EXTRA(Points) VARARG_EXTRA(Color));
+	VARARG_DECL(static void, static void, return, GeometryConvexLogf, VARARG_NONE, const TCHAR*, VARARG_EXTRA(const UObject* LogOwner) VARARG_EXTRA(const FLogCategoryBase& Category) VARARG_EXTRA(ELogVerbosity::Type Verbosity) VARARG_EXTRA(int32 UniqueLogId) VARARG_EXTRA(const TArray<FVector> &Points) VARARG_EXTRA(const FColor& Color), VARARG_EXTRA(LogOwner) VARARG_EXTRA(Category) VARARG_EXTRA(Verbosity) VARARG_EXTRA(UniqueLogId) VARARG_EXTRA(Points) VARARG_EXTRA(Color));
 	//Histogram data
-	VARARG_DECL(static void, static void, return, HistogramDataLogf, VARARG_NONE, const TCHAR*, VARARG_EXTRA(const class UObject* LogOwner) VARARG_EXTRA(const struct FLogCategoryBase& Category) VARARG_EXTRA(ELogVerbosity::Type Verbosity) VARARG_EXTRA(int32 UniqueLogId) VARARG_EXTRA(FName GraphName) VARARG_EXTRA(FName DataName) VARARG_EXTRA(const FVector2D& Data) VARARG_EXTRA(const FColor& Color), VARARG_EXTRA(LogOwner) VARARG_EXTRA(Category) VARARG_EXTRA(Verbosity) VARARG_EXTRA(UniqueLogId) VARARG_EXTRA(GraphName) VARARG_EXTRA(DataName) VARARG_EXTRA(Data) VARARG_EXTRA(Color));
+	VARARG_DECL(static void, static void, return, HistogramDataLogf, VARARG_NONE, const TCHAR*, VARARG_EXTRA(const UObject* LogOwner) VARARG_EXTRA(const FLogCategoryBase& Category) VARARG_EXTRA(ELogVerbosity::Type Verbosity) VARARG_EXTRA(int32 UniqueLogId) VARARG_EXTRA(FName GraphName) VARARG_EXTRA(FName DataName) VARARG_EXTRA(const FVector2D& Data) VARARG_EXTRA(const FColor& Color), VARARG_EXTRA(LogOwner) VARARG_EXTRA(Category) VARARG_EXTRA(Verbosity) VARARG_EXTRA(UniqueLogId) VARARG_EXTRA(GraphName) VARARG_EXTRA(DataName) VARARG_EXTRA(Data) VARARG_EXTRA(Color));
 	// Navigation data debug snapshot
-	static void NavigationDataDump(const class UObject* LogOwner, const struct FLogCategoryBase& Category, ELogVerbosity::Type Verbosity, int32 UniqueLogId, const FBox& Box);
+	static void NavigationDataDump(const UObject* LogOwner, const FLogCategoryBase& Category, ELogVerbosity::Type Verbosity, int32 UniqueLogId, const FBox& Box);
 
 	/** Log events */
-	static void EventLog(const class UObject* LogOwner, const FName EventTag1, const FVisualLogEventBase& Event1, const FName EventTag2 = NAME_None, const FName EventTag3 = NAME_None, const FName EventTag4 = NAME_None, const FName EventTag5 = NAME_None, const FName EventTag6 = NAME_None);
-	static void EventLog(const class UObject* LogOwner, const FName EventTag1, const FVisualLogEventBase& Event1, const FVisualLogEventBase& Event2);
-	static void EventLog(const class UObject* LogOwner, const FName EventTag1, const FVisualLogEventBase& Event1, const FVisualLogEventBase& Event2, const FVisualLogEventBase& Event3);
-	static void EventLog(const class UObject* LogOwner, const FName EventTag1, const FVisualLogEventBase& Event1, const FVisualLogEventBase& Event2, const FVisualLogEventBase& Event3, const FVisualLogEventBase& Event4);
-	static void EventLog(const class UObject* LogOwner, const FName EventTag1, const FVisualLogEventBase& Event1, const FVisualLogEventBase& Event2, const FVisualLogEventBase& Event3, const FVisualLogEventBase& Event4, const FVisualLogEventBase& Event5);
-	static void EventLog(const class UObject* LogOwner, const FName EventTag1, const FVisualLogEventBase& Event1, const FVisualLogEventBase& Event2, const FVisualLogEventBase& Event3, const FVisualLogEventBase& Event4, const FVisualLogEventBase& Event5, const FVisualLogEventBase& Event6);
+	static void EventLog(const UObject* LogOwner, const FName EventTag1, const FVisualLogEventBase& Event1, const FName EventTag2 = NAME_None, const FName EventTag3 = NAME_None, const FName EventTag4 = NAME_None, const FName EventTag5 = NAME_None, const FName EventTag6 = NAME_None);
+	static void EventLog(const UObject* LogOwner, const FName EventTag1, const FVisualLogEventBase& Event1, const FVisualLogEventBase& Event2);
+	static void EventLog(const UObject* LogOwner, const FName EventTag1, const FVisualLogEventBase& Event1, const FVisualLogEventBase& Event2, const FVisualLogEventBase& Event3);
+	static void EventLog(const UObject* LogOwner, const FName EventTag1, const FVisualLogEventBase& Event1, const FVisualLogEventBase& Event2, const FVisualLogEventBase& Event3, const FVisualLogEventBase& Event4);
+	static void EventLog(const UObject* LogOwner, const FName EventTag1, const FVisualLogEventBase& Event1, const FVisualLogEventBase& Event2, const FVisualLogEventBase& Event3, const FVisualLogEventBase& Event4, const FVisualLogEventBase& Event5);
+	static void EventLog(const UObject* LogOwner, const FName EventTag1, const FVisualLogEventBase& Event1, const FVisualLogEventBase& Event2, const FVisualLogEventBase& Event3, const FVisualLogEventBase& Event4, const FVisualLogEventBase& Event5, const FVisualLogEventBase& Event6);
 	
-	static void EventLog(const class UObject* LogOwner, const FVisualLogEventBase& Event1, const FName EventTag1 = NAME_None, const FName EventTag2 = NAME_None, const FName EventTag3 = NAME_None, const FName EventTag4 = NAME_None, const FName EventTag5 = NAME_None, const FName EventTag6 = NAME_None);
+	static void EventLog(const UObject* LogOwner, const FVisualLogEventBase& Event1, const FName EventTag1 = NAME_None, const FName EventTag2 = NAME_None, const FName EventTag3 = NAME_None, const FName EventTag4 = NAME_None, const FName EventTag5 = NAME_None, const FName EventTag6 = NAME_None);
 
 	// static getter
 	static FVisualLogger& Get()
@@ -180,13 +183,13 @@ public:
 	virtual void Shutdown();
 
 	// Removes all logged data 
-	void Cleanup(bool bReleaseMemory = false);
+	void Cleanup(UWorld* OldWorld, bool bReleaseMemory = false);
 
 	/** Set log owner redirection from one object to another, to combine logs */
-	static void Redirect(class UObject* FromObject, class UObject* ToObject);
+	static void Redirect(UObject* FromObject, UObject* ToObject);
 
 	/** find and return redirection object for given object*/
-	static class UObject* FindRedirection(const class UObject* Object);
+	static UObject* FindRedirection(const UObject* Object);
 
 	/** blocks all categories from logging. It can be bypassed with white list for categories */
 	void BlockAllCategories(bool InBlock) { bBlockedAllCategories = InBlock; }
@@ -195,11 +198,11 @@ public:
 	bool IsBlockedForAllCategories() const { return !!bBlockedAllCategories; }
 
 	/** Returns white list for modifications */
-	const TArray<class FName>& GetWhiteList() const { return CategoriesWhiteList; }
+	const TArray<FName>& GetWhiteList() const { return CategoriesWhiteList; }
 
 	bool IsWhiteListed(const FName& Name) const { return CategoriesWhiteList.Find(Name) != INDEX_NONE; }
 
-	void AddCategortyToWhiteList(class FName Category) { CategoriesWhiteList.AddUnique(Category); }
+	void AddCategortyToWhiteList(FName Category) { CategoriesWhiteList.AddUnique(Category); }
 
 	void ClearWhiteList() { CategoriesWhiteList.Reset(); }
 
@@ -226,11 +229,11 @@ public:
 	/** Remove visual logger output device */
 	const TArray<FVisualLogDevice*>& GetDevices() const { return OutputDevices; }
 	/** Check if log category can be recorded, verify before using GetEntryToWrite! */
-	bool IsCategoryLogged(const struct FLogCategoryBase& Category) const;
+	bool IsCategoryLogged(const FLogCategoryBase& Category) const;
 	/** Returns  current entry for given TimeStap or creates another one  but first it serialize previous entry as completed to vislog devices. Use VisualLogger::DontCreate to get current entry without serialization*/
-	FVisualLogEntry* GetEntryToWrite(const class UObject* Object, float TimeStamp, ECreateIfNeeded ShouldCreate = ECreateIfNeeded::Create);
+	FVisualLogEntry* GetEntryToWrite(const UObject* Object, float TimeStamp, ECreateIfNeeded ShouldCreate = ECreateIfNeeded::Create);
 	/** Retrieves last used entry for given UObject*/
-	FVisualLogEntry* GetLastEntryForObject(const class UObject* Object);
+	FVisualLogEntry* GetLastEntryForObject(const UObject* Object);
 	/** flush and serialize data if timestamp allows it */
 	virtual void Flush() override;
 
@@ -247,31 +250,39 @@ public:
 	const TMap<FName, FVisualLogExtensionInterface*>& GetAllExtensions() const { return AllExtensions; }
 
 	/** internal check for each usage of visual logger */
-	static bool CheckVisualLogInputInternal(const class UObject* Object, const struct FLogCategoryBase& Category, ELogVerbosity::Type Verbosity, UWorld **World, FVisualLogEntry **CurrentEntry);
+	static bool CheckVisualLogInputInternal(const UObject* Object, const FLogCategoryBase& Category, ELogVerbosity::Type Verbosity, UWorld **World, FVisualLogEntry **CurrentEntry);
 
+	typedef TMap<UObject*, TArray<TWeakObjectPtr<const UObject> > > RedirectionMapType;
+	static RedirectionMapType& GetRedirectionMap(const UObject* InObject);
+
+	typedef TMap<const UObject*, TWeakObjectPtr<const UWorld> > FObjectToWorldMapType;
+	FObjectToWorldMapType& GetObjectToWorldMap() { return ObjectToWorldMap; }
 private:
 	FVisualLogger();
-	virtual void Serialize(const TCHAR* V, ELogVerbosity::Type Verbosity, const class FName& Category) override { ensureMsgf(0, TEXT("Regular serialize is forbiden for visual logs")); }
+	virtual void Serialize(const TCHAR* V, ELogVerbosity::Type Verbosity, const FName& Category) override { ensureMsgf(0, TEXT("Regular serialize is forbiden for visual logs")); }
 
 protected:
 	/** Array of output devices to redirect to */
 	TArray<FVisualLogDevice*> OutputDevices;
 	// Map for inter-objects redirections
-	static TMap<class UObject*, TArray<TWeakObjectPtr<const class UObject> > > RedirectionMap;
+	static TMap<const UWorld*, RedirectionMapType> WorldToRedirectionMap;
+
 	// white list of categories to bypass blocking
-	TArray<class FName>	CategoriesWhiteList;
+	TArray<FName>	CategoriesWhiteList;
 	// Visual Logger extensions map
 	TMap<FName, FVisualLogExtensionInterface*> AllExtensions;
 	// last generated unique id for given times tamp
 	TMap<float, int32> LastUniqueIds;
 	// Current entry with all data
-	TMap<const class UObject*, FVisualLogEntry>	CurrentEntryPerObject;
+	TMap<const UObject*, FVisualLogEntry>	CurrentEntryPerObject;
 	// Map to contain names for Objects (they can be destroyed after while)
-	TMap<const class UObject*, FName> ObjectToNameMap;
+	TMap<const UObject*, FName> ObjectToNameMap;
 	// Map to contain class names for Objects (they can be destroyed after while)
-	TMap<const class UObject*, FName> ObjectToClassNameMap;
+	TMap<const UObject*, FName> ObjectToClassNameMap;
 	// Map to contain information about pointers in game
-	TMap<const class UObject*, TWeakObjectPtr<const class UObject> > ObjectToPointerMap;
+	TMap<const UObject*, TWeakObjectPtr<const UObject> > ObjectToPointerMap;
+	// Cached map to world information because it's just raw pointer and not real object
+	FObjectToWorldMapType ObjectToWorldMap;
 	// if set all categories are blocked from logging
 	int32 bBlockedAllCategories : 1;
 	// if set we are recording to file
@@ -321,7 +332,7 @@ protected:
 	FMemory::SystemFree(AllocatedBuffer);
 
 inline
-VARARG_BODY(void, FVisualLogger::CategorizedLogf, const TCHAR*, VARARG_EXTRA(const class UObject* Object) VARARG_EXTRA(const struct FLogCategoryBase& Category) VARARG_EXTRA(ELogVerbosity::Type Verbosity) VARARG_EXTRA(int32 UniqueLogId))
+VARARG_BODY(void, FVisualLogger::CategorizedLogf, const TCHAR*, VARARG_EXTRA(const UObject* Object) VARARG_EXTRA(const FLogCategoryBase& Category) VARARG_EXTRA(ELogVerbosity::Type Verbosity) VARARG_EXTRA(int32 UniqueLogId))
 {
 	SCOPE_CYCLE_COUNTER(STAT_VisualLog);
 	UWorld *World = nullptr;
@@ -337,7 +348,7 @@ VARARG_BODY(void, FVisualLogger::CategorizedLogf, const TCHAR*, VARARG_EXTRA(con
 }
 
 inline
-VARARG_BODY(void, FVisualLogger::GeometryShapeLogf, const TCHAR*, VARARG_EXTRA(const class UObject* Object) VARARG_EXTRA(const struct FLogCategoryBase& Category) VARARG_EXTRA(ELogVerbosity::Type Verbosity) VARARG_EXTRA(int32 UniqueLogId) VARARG_EXTRA(const FVector& Start) VARARG_EXTRA(const FVector& End) VARARG_EXTRA(const FColor& Color))
+VARARG_BODY(void, FVisualLogger::GeometryShapeLogf, const TCHAR*, VARARG_EXTRA(const UObject* Object) VARARG_EXTRA(const FLogCategoryBase& Category) VARARG_EXTRA(ELogVerbosity::Type Verbosity) VARARG_EXTRA(int32 UniqueLogId) VARARG_EXTRA(const FVector& Start) VARARG_EXTRA(const FVector& End) VARARG_EXTRA(const FColor& Color))
 {
 	SCOPE_CYCLE_COUNTER(STAT_VisualLog);
 	UWorld *World = nullptr;
@@ -353,7 +364,7 @@ VARARG_BODY(void, FVisualLogger::GeometryShapeLogf, const TCHAR*, VARARG_EXTRA(c
 }
 
 inline
-VARARG_BODY(void, FVisualLogger::GeometryShapeLogf, const TCHAR*, VARARG_EXTRA(const class UObject* Object) VARARG_EXTRA(const struct FLogCategoryBase& Category) VARARG_EXTRA(ELogVerbosity::Type Verbosity) VARARG_EXTRA(int32 UniqueLogId) VARARG_EXTRA(const FVector& Location) VARARG_EXTRA(float Radius) VARARG_EXTRA(const FColor& Color))
+VARARG_BODY(void, FVisualLogger::GeometryShapeLogf, const TCHAR*, VARARG_EXTRA(const UObject* Object) VARARG_EXTRA(const FLogCategoryBase& Category) VARARG_EXTRA(ELogVerbosity::Type Verbosity) VARARG_EXTRA(int32 UniqueLogId) VARARG_EXTRA(const FVector& Location) VARARG_EXTRA(float Radius) VARARG_EXTRA(const FColor& Color))
 {
 	SCOPE_CYCLE_COUNTER(STAT_VisualLog);
 	UWorld *World = nullptr;
@@ -369,7 +380,7 @@ VARARG_BODY(void, FVisualLogger::GeometryShapeLogf, const TCHAR*, VARARG_EXTRA(c
 }
 
 inline
-VARARG_BODY(void, FVisualLogger::GeometryShapeLogf, const TCHAR*, VARARG_EXTRA(const class UObject* Object) VARARG_EXTRA(const struct FLogCategoryBase& Category) VARARG_EXTRA(ELogVerbosity::Type Verbosity) VARARG_EXTRA(int32 UniqueLogId) VARARG_EXTRA(const FBox& Box) VARARG_EXTRA(const FMatrix& Matrix) VARARG_EXTRA(const FColor& Color))
+VARARG_BODY(void, FVisualLogger::GeometryShapeLogf, const TCHAR*, VARARG_EXTRA(const UObject* Object) VARARG_EXTRA(const FLogCategoryBase& Category) VARARG_EXTRA(ELogVerbosity::Type Verbosity) VARARG_EXTRA(int32 UniqueLogId) VARARG_EXTRA(const FBox& Box) VARARG_EXTRA(const FMatrix& Matrix) VARARG_EXTRA(const FColor& Color))
 {
 	SCOPE_CYCLE_COUNTER(STAT_VisualLog);
 	UWorld *World = nullptr;
@@ -385,7 +396,7 @@ VARARG_BODY(void, FVisualLogger::GeometryShapeLogf, const TCHAR*, VARARG_EXTRA(c
 }
 
 inline
-VARARG_BODY(void, FVisualLogger::GeometryShapeLogf, const TCHAR*, VARARG_EXTRA(const class UObject* Object) VARARG_EXTRA(const struct FLogCategoryBase& Category) VARARG_EXTRA(ELogVerbosity::Type Verbosity) VARARG_EXTRA(int32 UniqueLogId) VARARG_EXTRA(const FVector& Orgin) VARARG_EXTRA(const FVector& Direction) VARARG_EXTRA(const float Length) VARARG_EXTRA(const float Angle)  VARARG_EXTRA(const FColor& Color))
+VARARG_BODY(void, FVisualLogger::GeometryShapeLogf, const TCHAR*, VARARG_EXTRA(const UObject* Object) VARARG_EXTRA(const FLogCategoryBase& Category) VARARG_EXTRA(ELogVerbosity::Type Verbosity) VARARG_EXTRA(int32 UniqueLogId) VARARG_EXTRA(const FVector& Orgin) VARARG_EXTRA(const FVector& Direction) VARARG_EXTRA(const float Length) VARARG_EXTRA(const float Angle)  VARARG_EXTRA(const FColor& Color))
 {
 	SCOPE_CYCLE_COUNTER(STAT_VisualLog);
 	UWorld *World = nullptr;
@@ -401,7 +412,7 @@ VARARG_BODY(void, FVisualLogger::GeometryShapeLogf, const TCHAR*, VARARG_EXTRA(c
 }
 
 inline
-VARARG_BODY(void, FVisualLogger::GeometryShapeLogf, const TCHAR*, VARARG_EXTRA(const class UObject* Object) VARARG_EXTRA(const struct FLogCategoryBase& Category) VARARG_EXTRA(ELogVerbosity::Type Verbosity) VARARG_EXTRA(int32 UniqueLogId) VARARG_EXTRA(const FVector& Start) VARARG_EXTRA(const FVector& End) VARARG_EXTRA(const float Radius) VARARG_EXTRA(const FColor& Color))
+VARARG_BODY(void, FVisualLogger::GeometryShapeLogf, const TCHAR*, VARARG_EXTRA(const UObject* Object) VARARG_EXTRA(const FLogCategoryBase& Category) VARARG_EXTRA(ELogVerbosity::Type Verbosity) VARARG_EXTRA(int32 UniqueLogId) VARARG_EXTRA(const FVector& Start) VARARG_EXTRA(const FVector& End) VARARG_EXTRA(const float Radius) VARARG_EXTRA(const FColor& Color))
 {
 	SCOPE_CYCLE_COUNTER(STAT_VisualLog);
 	UWorld *World = nullptr;
@@ -417,7 +428,7 @@ VARARG_BODY(void, FVisualLogger::GeometryShapeLogf, const TCHAR*, VARARG_EXTRA(c
 }
 
 inline
-VARARG_BODY(void, FVisualLogger::GeometryShapeLogf, const TCHAR*, VARARG_EXTRA(const class UObject* Object) VARARG_EXTRA(const struct FLogCategoryBase& Category) VARARG_EXTRA(ELogVerbosity::Type Verbosity) VARARG_EXTRA(int32 UniqueLogId) VARARG_EXTRA(const FVector& Center) VARARG_EXTRA(float HalfHeight) VARARG_EXTRA(float Radius) VARARG_EXTRA(const FQuat& Rotation) VARARG_EXTRA(const FColor& Color))
+VARARG_BODY(void, FVisualLogger::GeometryShapeLogf, const TCHAR*, VARARG_EXTRA(const UObject* Object) VARARG_EXTRA(const FLogCategoryBase& Category) VARARG_EXTRA(ELogVerbosity::Type Verbosity) VARARG_EXTRA(int32 UniqueLogId) VARARG_EXTRA(const FVector& Center) VARARG_EXTRA(float HalfHeight) VARARG_EXTRA(float Radius) VARARG_EXTRA(const FQuat& Rotation) VARARG_EXTRA(const FColor& Color))
 {
 	SCOPE_CYCLE_COUNTER(STAT_VisualLog);
 	UWorld *World = nullptr;
@@ -433,7 +444,7 @@ VARARG_BODY(void, FVisualLogger::GeometryShapeLogf, const TCHAR*, VARARG_EXTRA(c
 }
 
 inline
-VARARG_BODY(void, FVisualLogger::NavAreaShapeLogf, const TCHAR*, VARARG_EXTRA(const class UObject* Object) VARARG_EXTRA(const struct FLogCategoryBase& Category) VARARG_EXTRA(ELogVerbosity::Type Verbosity) VARARG_EXTRA(int32 UniqueLogId) VARARG_EXTRA(const TArray<FVector> &ConvexPoints) VARARG_EXTRA(float MinZ) VARARG_EXTRA(float MaxZ) VARARG_EXTRA(const FColor& Color))
+VARARG_BODY(void, FVisualLogger::NavAreaShapeLogf, const TCHAR*, VARARG_EXTRA(const UObject* Object) VARARG_EXTRA(const FLogCategoryBase& Category) VARARG_EXTRA(ELogVerbosity::Type Verbosity) VARARG_EXTRA(int32 UniqueLogId) VARARG_EXTRA(const TArray<FVector> &ConvexPoints) VARARG_EXTRA(float MinZ) VARARG_EXTRA(float MaxZ) VARARG_EXTRA(const FColor& Color))
 {
 	SCOPE_CYCLE_COUNTER(STAT_VisualLog);
 	UWorld *World = nullptr;
@@ -449,7 +460,7 @@ VARARG_BODY(void, FVisualLogger::NavAreaShapeLogf, const TCHAR*, VARARG_EXTRA(co
 }
 
 inline
-VARARG_BODY(void, FVisualLogger::GeometryShapeLogf, const TCHAR*, VARARG_EXTRA(const class UObject* Object) VARARG_EXTRA(const struct FLogCategoryBase& Category) VARARG_EXTRA(ELogVerbosity::Type Verbosity) VARARG_EXTRA(int32 UniqueLogId) VARARG_EXTRA(const TArray<FVector> &Vertices) VARARG_EXTRA(const TArray<int32> &Indices) VARARG_EXTRA(const FColor& Color))
+VARARG_BODY(void, FVisualLogger::GeometryShapeLogf, const TCHAR*, VARARG_EXTRA(const UObject* Object) VARARG_EXTRA(const FLogCategoryBase& Category) VARARG_EXTRA(ELogVerbosity::Type Verbosity) VARARG_EXTRA(int32 UniqueLogId) VARARG_EXTRA(const TArray<FVector> &Vertices) VARARG_EXTRA(const TArray<int32> &Indices) VARARG_EXTRA(const FColor& Color))
 {
 	SCOPE_CYCLE_COUNTER(STAT_VisualLog);
 	UWorld *World = nullptr;
@@ -465,7 +476,7 @@ VARARG_BODY(void, FVisualLogger::GeometryShapeLogf, const TCHAR*, VARARG_EXTRA(c
 }
 
 inline
-VARARG_BODY(void, FVisualLogger::GeometryConvexLogf, const TCHAR*, VARARG_EXTRA(const class UObject* Object) VARARG_EXTRA(const struct FLogCategoryBase& Category) VARARG_EXTRA(ELogVerbosity::Type Verbosity) VARARG_EXTRA(int32 UniqueLogId) VARARG_EXTRA(const TArray<FVector> &Points) VARARG_EXTRA(const FColor& Color))
+VARARG_BODY(void, FVisualLogger::GeometryConvexLogf, const TCHAR*, VARARG_EXTRA(const UObject* Object) VARARG_EXTRA(const FLogCategoryBase& Category) VARARG_EXTRA(ELogVerbosity::Type Verbosity) VARARG_EXTRA(int32 UniqueLogId) VARARG_EXTRA(const TArray<FVector> &Points) VARARG_EXTRA(const FColor& Color))
 {
 	SCOPE_CYCLE_COUNTER(STAT_VisualLog);
 	UWorld *World = nullptr;
@@ -482,7 +493,7 @@ VARARG_BODY(void, FVisualLogger::GeometryConvexLogf, const TCHAR*, VARARG_EXTRA(
 
 
 inline
-VARARG_BODY(void, FVisualLogger::HistogramDataLogf, const TCHAR*, VARARG_EXTRA(const class UObject* Object) VARARG_EXTRA(const struct FLogCategoryBase& Category) VARARG_EXTRA(ELogVerbosity::Type Verbosity) VARARG_EXTRA(int32 UniqueLogId) VARARG_EXTRA(FName GraphName) VARARG_EXTRA(FName DataName) VARARG_EXTRA(const FVector2D& Data) VARARG_EXTRA(const FColor& Color))
+VARARG_BODY(void, FVisualLogger::HistogramDataLogf, const TCHAR*, VARARG_EXTRA(const UObject* Object) VARARG_EXTRA(const FLogCategoryBase& Category) VARARG_EXTRA(ELogVerbosity::Type Verbosity) VARARG_EXTRA(int32 UniqueLogId) VARARG_EXTRA(FName GraphName) VARARG_EXTRA(FName DataName) VARARG_EXTRA(const FVector2D& Data) VARARG_EXTRA(const FColor& Color))
 {
 	SCOPE_CYCLE_COUNTER(STAT_VisualLog);
 	UWorld *World = nullptr;
diff --git a/Engine/Source/Runtime/Engine/Public/VisualLogger/VisualLoggerTypes.h b/Engine/Source/Runtime/Engine/Public/VisualLogger/VisualLoggerTypes.h
--- a/Engine/Source/Runtime/Engine/Public/VisualLogger/VisualLoggerTypes.h
+++ b/Engine/Source/Runtime/Engine/Public/VisualLogger/VisualLoggerTypes.h
@@ -2,6 +2,9 @@
 #pragma once
 #include "EngineDefines.h"
 
+class UObject;
+class UWorld;
+
 enum class ECreateIfNeeded : int8
 {
 	Invalid = -1,
diff --git a/Engine/Source/Runtime/Engine/Public/WorldCollision.h b/Engine/Source/Runtime/Engine/Public/WorldCollision.h
--- a/Engine/Source/Runtime/Engine/Public/WorldCollision.h
+++ b/Engine/Source/Runtime/Engine/Public/WorldCollision.h
@@ -435,6 +435,11 @@ struct AsyncTraceData : FNoncopyable
 	TArray<TUniquePtr<TTraceThreadData<FTraceDatum>>>			TraceData;
 	TArray<TUniquePtr<TTraceThreadData<FOverlapDatum>>>			OverlapData;
 
+	/** Datum entries in TraceData are persistent for efficiency. This is the number of them that are actually in use (rather than TraceData.Num()). */
+	int32 NumQueuedTraceData;
+	/** Datum entries in OverlapData are persistent for efficiency. This is the number of them that are actually in use (rather than OverlapData.Num()). */
+	int32 NumQueuedOverlapData;
+
 	/**
 	 * if Execution is all done, set this to be true
 	 * 
@@ -446,7 +451,11 @@ struct AsyncTraceData : FNoncopyable
 	/**  Thread completion event for batch **/
 	FGraphEventArray		AsyncTraceCompletionEvent;
 
-	AsyncTraceData() : bAsyncAllowed(false) {}
+	AsyncTraceData() 
+		: NumQueuedTraceData(0)
+		, NumQueuedOverlapData(0)
+		, bAsyncAllowed(false)
+	{}
 };
 
 
diff --git a/Engine/Source/Runtime/GameplayAbilities/Private/AbilitiesGameplayDebuggerObject.cpp b/Engine/Source/Runtime/GameplayAbilities/Private/AbilitiesGameplayDebuggerObject.cpp
new file mode 100644
--- /dev/null
+++ b/Engine/Source/Runtime/GameplayAbilities/Private/AbilitiesGameplayDebuggerObject.cpp
@@ -0,0 +1,47 @@
+// Copyright 1998-2015 Epic Games, Inc. All Rights Reserved.
+#include "AbilitySystemPrivatePCH.h"
+#include "Misc/CoreMisc.h"
+#include "Net/UnrealNetwork.h"
+#include "DebugRenderSceneProxy.h"
+#include "GameplayTasksComponent.h"
+#include "AbilitiesGameplayDebuggerObject.h"
+
+UAbilitiesGameplayDebuggerObject::UAbilitiesGameplayDebuggerObject(const FObjectInitializer& ObjectInitializer) : Super(ObjectInitializer)
+{
+}
+
+void UAbilitiesGameplayDebuggerObject::GetLifetimeReplicatedProps(TArray< FLifetimeProperty > & OutLifetimeProps) const
+{
+	Super::GetLifetimeReplicatedProps(OutLifetimeProps);
+#if !(UE_BUILD_SHIPPING || UE_BUILD_TEST)
+	DOREPLIFETIME(UAbilitiesGameplayDebuggerObject, AbilityInfo);
+	DOREPLIFETIME(UAbilitiesGameplayDebuggerObject, bIsUsingAbilities);
+#endif
+}
+
+void UAbilitiesGameplayDebuggerObject::CollectDataToReplicate(APlayerController* MyPC, AActor *SelectedActor)
+{
+#if !(UE_BUILD_SHIPPING || UE_BUILD_TEST)
+	APawn* MyPawn = Cast<APawn>(SelectedActor);
+	if (MyPawn)
+	{
+		bool bUsingAbilities;
+		IGameplayAbilitiesModule& AbilitiesModule = IGameplayAbilitiesModule::Get();
+		AbilitiesModule.GetActiveAbilitiesDebugDataForActor(MyPawn, AbilityInfo, bUsingAbilities);
+		bIsUsingAbilities = bUsingAbilities;
+	}
+	else
+	{
+		bIsUsingAbilities = false;
+		AbilityInfo = TEXT("None");
+	}
+#endif
+}
+
+void UAbilitiesGameplayDebuggerObject::DrawCollectedData(APlayerController* MyPC, AActor* SelectedActor)
+{
+	if (bIsUsingAbilities)
+	{
+		PrintString(DefaultContext, FString::Printf(TEXT("Ability: {yellow}%s\n"), *AbilityInfo));
+	}
+}
diff --git a/Engine/Source/Runtime/GameplayAbilities/Private/AbilitiesGameplayDebuggerObject.h b/Engine/Source/Runtime/GameplayAbilities/Private/AbilitiesGameplayDebuggerObject.h
new file mode 100644
--- /dev/null
+++ b/Engine/Source/Runtime/GameplayAbilities/Private/AbilitiesGameplayDebuggerObject.h
@@ -0,0 +1,23 @@
+// Copyright 1998-2015 Epic Games, Inc. All Rights Reserved.
+
+#pragma once
+#include "Debug/GameplayDebuggerBaseObject.h"
+#include "AbilitiesGameplayDebuggerObject.generated.h"
+
+UCLASS()
+class UAbilitiesGameplayDebuggerObject : public UGameplayDebuggerBaseObject
+{
+	GENERATED_UCLASS_BODY()
+
+public:
+	UPROPERTY(Replicated)
+	FString AbilityInfo;
+
+	UPROPERTY(Replicated)
+	uint32 bIsUsingAbilities : 1;
+
+	virtual void CollectDataToReplicate(APlayerController* MyPC, AActor *SelectedActor) override;
+
+	virtual void DrawCollectedData(APlayerController* MyPC, AActor* SelectedActor) override;
+
+};
diff --git a/Engine/Source/Runtime/Navmesh/Private/Detour/DetourNavMeshQuery.cpp b/Engine/Source/Runtime/Navmesh/Private/Detour/DetourNavMeshQuery.cpp
--- a/Engine/Source/Runtime/Navmesh/Private/Detour/DetourNavMeshQuery.cpp
+++ b/Engine/Source/Runtime/Navmesh/Private/Detour/DetourNavMeshQuery.cpp
@@ -31,10 +31,6 @@
 #include "DetourAssert.h"
 #include <new>
 
-#if TRACK_PATHFINDING_PERF
-#include "PerfTimer.h"
-#endif
-
 /// @class dtQueryFilter
 ///
 /// <b>The Default Implementation</b>
@@ -194,7 +190,6 @@ dtNavMeshQuery::dtNavMeshQuery() :
 	m_tinyNodePool(0),
 	m_nodePool(0),
 	m_openList(0),
-	m_queryTime(0.0f),
 	m_queryNodes(0)
 {
 	memset(&m_query, 0, sizeof(dtQueryData));
@@ -1069,6 +1064,61 @@ dtStatus dtNavMeshQuery::findNearestPoly(const float* center, const float* exten
 	return DT_SUCCESS;
 }
 
+/// @par 
+///
+/// @note If the search box does not intersect any polygons the search will 
+/// return #DT_SUCCESS, but @p nearestRef will be zero. So if in doubt, check 
+/// @p nearestRef before using @p nearestPt.
+///
+/// @warning This function is not suitable for large area searches.  If the search
+/// extents overlaps more than 128 polygons it may return an invalid result.
+///
+dtStatus dtNavMeshQuery::findNearestContainingPoly(const float* center, const float* extents,
+												   const dtQueryFilter* filter,
+												   dtPolyRef* nearestRef, float* nearestPt) const
+{
+	dtAssert(m_nav);
+
+	*nearestRef = 0;
+
+	// Get nearby polygons from proximity grid.
+	dtPolyRef polys[128];
+	int polyCount = 0;
+	if (dtStatusFailed(queryPolygons(center, extents, filter, polys, &polyCount, 128)))
+		return DT_FAILURE | DT_INVALID_PARAM;
+
+	// Find nearest polygon amongst the nearby polygons.
+	dtPolyRef nearest = 0;
+	float nearestDistanceSqr = FLT_MAX;
+	for (int i = 0; i < polyCount; ++i)
+	{
+		dtPolyRef ref = polys[i];
+
+		bool inPoly = false;
+		isPointInsidePoly(ref, center, inPoly);
+
+		if (inPoly)
+		{
+			float closestPtPoly[3];
+			closestPointOnPoly(ref, center, closestPtPoly);
+			float d = dtVdistSqr(center, closestPtPoly);
+			float h = dtAbs(center[1] - closestPtPoly[1]);
+			if (d < nearestDistanceSqr && h < extents[1])
+			{
+				if (nearestPt)
+					dtVcopy(nearestPt, closestPtPoly);
+				nearestDistanceSqr = d;
+				nearest = ref;
+			}
+		}
+	}
+
+	if (nearestRef)
+		*nearestRef = nearest;
+
+	return DT_SUCCESS;
+}
+
 dtPolyRef dtNavMeshQuery::findNearestPolyInTile(const dtMeshTile* tile, const float* center, const float* extents,
 												const dtQueryFilter* filter, float* nearestPt) const
 {
@@ -1277,7 +1327,6 @@ dtStatus dtNavMeshQuery::findPath(dtPolyRef startRef, dtPolyRef endRef,
 	dtAssert(m_nodePool);
 	dtAssert(m_openList);
 	
-	m_queryTime = 0.0f;
 	m_queryNodes = 0;
 
 	if (!startRef || !endRef)
@@ -1297,10 +1346,6 @@ dtStatus dtNavMeshQuery::findPath(dtPolyRef startRef, dtPolyRef endRef,
 	const float H_SCALE = filter->getModifiedHeuristicScale();
 	//@UE4 END
 
-#if TRACK_PATHFINDING_PERF
-	const TimeVal startTime = getPerfTime();
-#endif
-
 	m_nodePool->clear();
 	m_openList->clear();
 	
@@ -1515,11 +1560,6 @@ dtStatus dtNavMeshQuery::findPath(dtPolyRef startRef, dtPolyRef endRef,
 		*totalCost = lastBestNode->total;
 	}
 
-#if TRACK_PATHFINDING_PERF
-	const TimeVal endTime = getPerfTime();
-	m_queryTime = getPerfDeltaTimeUsec(startTime, endTime) / 1000.0f;
-#endif
-
 	return status;
 }
 
@@ -1530,7 +1570,6 @@ dtStatus dtNavMeshQuery::testClusterPath(dtPolyRef startRef, dtPolyRef endRef) c
 	const unsigned int startPolyIdx = m_nav->decodePolyIdPoly(startRef);
 	const unsigned int endPolyIdx = m_nav->decodePolyIdPoly(endRef);
 
-	m_queryTime = 0.0f;
 	m_queryNodes = 0;
 
 	if (startTile == 0 || endTile == 0 ||
@@ -1554,10 +1593,6 @@ dtStatus dtNavMeshQuery::testClusterPath(dtPolyRef startRef, dtPolyRef endRef) c
 		return DT_SUCCESS;
 	}
 
-#if TRACK_PATHFINDING_PERF
-	const TimeVal startTime = getPerfTime();
-#endif
-
 	m_nodePool->clear();
 	m_openList->clear();
 
@@ -1675,11 +1710,6 @@ dtStatus dtNavMeshQuery::testClusterPath(dtPolyRef startRef, dtPolyRef endRef) c
 		}
 	}
 
-#if TRACK_PATHFINDING_PERF
-	const TimeVal endTime = getPerfTime();
-	m_queryTime = getPerfDeltaTimeUsec(startTime, endTime) / 1000.0f;
-#endif
-
 	if (lastBestNode->id == endCRef)
 		status = DT_SUCCESS;
 
diff --git a/Engine/Source/Runtime/Navmesh/Private/DetourCrowd/DetourCrowd.cpp b/Engine/Source/Runtime/Navmesh/Private/DetourCrowd/DetourCrowd.cpp
--- a/Engine/Source/Runtime/Navmesh/Private/DetourCrowd/DetourCrowd.cpp
+++ b/Engine/Source/Runtime/Navmesh/Private/DetourCrowd/DetourCrowd.cpp
@@ -1287,16 +1287,7 @@ void dtCrowd::updateStepProximityData(const float dt, dtCrowdAgentDebugInfo* deb
 		if (dtVdist2DSqr(ag->npos, ag->boundary.getCenter()) > dtSqr(updateThr) ||
 			!ag->boundary.isValid(m_navquery, &m_filters[ag->params.filter]))
 		{
-			// UE4: force removing segments too close to offmesh links
-			float linkV0[3] = { 0.0f };
-			float linkV1[3] = { 0.0f };
-			
-			bool bHasOffmeshLink = ag->ncorners && (ag->cornerFlags[ag->ncorners - 1] & DT_STRAIGHTPATH_OFFMESH_CONNECTION);
-			if (bHasOffmeshLink)
-			{
-				const dtStatus linkStatus = m_navquery->getAttachedNavMesh()->getOffMeshConnectionPolyEndPoints(0, ag->cornerPolys[ag->ncorners - 1], 0, linkV0, linkV1);
-				bHasOffmeshLink = dtStatusSucceed(linkStatus);
-			}
+			const bool bIgnoreEdgesNearLastCorner = ag->ncorners && (ag->cornerFlags[ag->ncorners - 1] & (DT_STRAIGHTPATH_OFFMESH_CONNECTION | DT_STRAIGHTPATH_END));
 
 			// UE4: move dir for segment scoring
 			float moveDir[3] = { 0.0f };
@@ -1311,7 +1302,7 @@ void dtCrowd::updateStepProximityData(const float dt, dtCrowdAgentDebugInfo* deb
 			dtVnormalize(moveDir);
 
 			ag->boundary.update(&m_sharedBoundary, sharedDataIdx, ag->npos, ag->params.collisionQueryRange,
-				bHasOffmeshLink, linkV0, linkV1,
+				bIgnoreEdgesNearLastCorner, &ag->cornerVerts[(ag->ncorners - 1) * 3],
 				ag->corridor.getPath(), m_raycastSingleArea ? ag->corridor.getPathCount() : 0,
 				moveDir, m_navquery, &m_filters[ag->params.filter]);
 		}
@@ -1553,7 +1544,7 @@ void dtCrowd::updateStepAvoidance(const float dt, dtCrowdAgentDebugInfo* debug)
 				const float* s = ag->boundary.getSegment(j);
 				if (dtTriArea2D(ag->npos, s, s + 3) < 0.0f)
 					continue;
-				m_obstacleQuery->addSegment(s, s + 3);
+				m_obstacleQuery->addSegment(s, s + 3, ag->boundary.getSegmentFlags(j));
 			}
 
 			dtObstacleAvoidanceDebugData* vod = 0;
diff --git a/Engine/Source/Runtime/Navmesh/Private/DetourCrowd/DetourLocalBoundary.cpp b/Engine/Source/Runtime/Navmesh/Private/DetourCrowd/DetourLocalBoundary.cpp
--- a/Engine/Source/Runtime/Navmesh/Private/DetourCrowd/DetourLocalBoundary.cpp
+++ b/Engine/Source/Runtime/Navmesh/Private/DetourCrowd/DetourLocalBoundary.cpp
@@ -25,6 +25,7 @@
 #include "DetourLocalBoundary.h"
 #include "DetourNavMeshQuery.h"
 #include "DetourCommon.h"
+#include "DetourCrowd.h"
 #include "DetourAssert.h"
 
 
@@ -46,7 +47,7 @@ void dtLocalBoundary::reset()
 	m_nsegs = 0;
 }
 
-void dtLocalBoundary::addSegment(const float dist, const float* s)
+void dtLocalBoundary::addSegment(const float dist, const float* s, int flags)
 {
 	// Insert neighbour based on the distance.
 	Segment* seg = 0;
@@ -79,6 +80,7 @@ void dtLocalBoundary::addSegment(const float dist, const float* s)
 	}
 	
 	seg->d = dist;
+	seg->flags = flags;
 	memcpy(seg->s, s, sizeof(float)*6);
 	
 	if (m_nsegs + 1 < MAX_LOCAL_SEGS)
@@ -87,8 +89,30 @@ void dtLocalBoundary::addSegment(const float dist, const float* s)
 	}
 }
 
+namespace LocalBoundaryHelpers
+{
+	inline bool IsOutsideHeightLimit(const float* pos, const float* p, const float* q, float tseg)
+	{
+		const float closestPtHeight = p[1] + (q[1] - p[1]) * tseg;
+		const float segHeightDiff = dtAbs(closestPtHeight - pos[1]);
+		const float maxHeightDiff = 50.0f;
+
+		return segHeightDiff > maxHeightDiff;
+	}
+
+	inline int GetSegmentFlags(const float* endPos, const float* p, const float* q, float QueryRange)
+	{
+		const float ignoreDistancePct = 0.25f;
+		const float maxDistSq = dtSqr(QueryRange * ignoreDistancePct);
+
+		float tsegEnd = 0.0f;
+		const float distEndSqr = dtDistancePtSegSqr2D(endPos, p, q, tsegEnd);
+		return (distEndSqr < maxDistSq) ? DT_CROWD_BOUNDARY_IGNORE : 0;
+	}
+}
+
 void dtLocalBoundary::update(dtPolyRef ref, const float* pos, const float collisionQueryRange,
-	const bool bRemoveNearLink, const float* linkV0, const float* linkV1,
+	const bool bIgnoreAtEnd, const float* endPos,
 	const dtPolyRef* path, const int npath,
 	const float* moveDir,
 	dtNavMeshQuery* navquery, const dtQueryFilter* filter)
@@ -129,18 +153,16 @@ void dtLocalBoundary::update(dtPolyRef ref, const float* pos, const float collis
 			if (distSqr > dtSqr(collisionQueryRange))
 				continue;
 
-			// [UE4] remove segments too close to requested position
-			if (bRemoveNearLink)
+			// [UE4] handle segments too far in Y (height diff)
+			const bool bOutsideHeightLimit = LocalBoundaryHelpers::IsOutsideHeightLimit(pos, s, s + 3, tseg);
+			if (bOutsideHeightLimit)
 			{
-				float dummyS = 0.0f;
-				float dummyT = 0.0f;
-				const bool bIntersect = dtIntersectSegSeg2D(s, s+3, linkV0, linkV1, dummyS, dummyT);
-				if (bIntersect)
-				{
-					continue;
-				}
+				continue;
 			}
 
+			// [UE4] handle segments too close to requested position
+			const int segFlags = bIgnoreAtEnd ? LocalBoundaryHelpers::GetSegmentFlags(endPos, s, s + 3, collisionQueryRange) : 0;
+
 			// [UE4] include direction to segment in score
 			dtVlerp(closestPt, s, s + 3, tseg);
 			dtVsub(dirToSeg, closestPt, pos);
@@ -148,14 +170,14 @@ void dtLocalBoundary::update(dtPolyRef ref, const float* pos, const float collis
 			const float dseg = dtVdot2D(dirToSeg, moveDir);
 			const float score = distSqr * ((1.0f - dseg) * 0.5f);
 
-			addSegment(score, s);
+			addSegment(score, s, segFlags);
 		}
 	}
 }
 
 void dtLocalBoundary::update(const dtSharedBoundary* sharedData, const int sharedIdx,
 	const float* pos, const float collisionQueryRange,
-	const bool bRemoveNearLink, const float* linkV0, const float* linkV1,
+	const bool bIgnoreAtEnd, const float* endPos,
 	const dtPolyRef* path, const int npath, const float* moveDir,
 	dtNavMeshQuery* navquery, const dtQueryFilter* filter)
 {
@@ -196,18 +218,16 @@ void dtLocalBoundary::update(const dtSharedBoundary* sharedData, const int share
 		if (distSqr > dtSqr(collisionQueryRange))
 			continue;
 
-		// remove segments too close to requested position
-		if (bRemoveNearLink)
+		// handle segments too far in Y (height diff)
+		const bool bOutsideHeightLimit = LocalBoundaryHelpers::IsOutsideHeightLimit(pos, Data.Edges[Idx].v0, Data.Edges[Idx].v1, tseg);
+		if (bOutsideHeightLimit)
 		{
-			float dummyS = 0.0f;
-			float dummyT = 0.0f;
-			const bool bIntersect = dtIntersectSegSeg2D(Data.Edges[Idx].v0, Data.Edges[Idx].v1, linkV0, linkV1, dummyS, dummyT);
-			if (bIntersect)
-			{
-				continue;
-			}
+			continue;
 		}
 
+		// handle segments too close to requested position
+		const int segFlags = bIgnoreAtEnd ? LocalBoundaryHelpers::GetSegmentFlags(endPos, Data.Edges[Idx].v0, Data.Edges[Idx].v1, collisionQueryRange) : 0;
+
 		// remove segments when both sides are on path (single area trace)
 		if (PathLookup.Contains(Data.Edges[Idx].p0) && PathLookup.Contains(Data.Edges[Idx].p1))
 		{
@@ -223,7 +243,7 @@ void dtLocalBoundary::update(const dtSharedBoundary* sharedData, const int share
 
 		dtVcopy(s, Data.Edges[Idx].v0);
 		dtVcopy(s + 3, Data.Edges[Idx].v1);
-		addSegment(score, s);
+		addSegment(score, s, segFlags);
 	}
 }
 
diff --git a/Engine/Source/Runtime/Navmesh/Private/DetourCrowd/DetourObstacleAvoidance.cpp b/Engine/Source/Runtime/Navmesh/Private/DetourCrowd/DetourObstacleAvoidance.cpp
--- a/Engine/Source/Runtime/Navmesh/Private/DetourCrowd/DetourObstacleAvoidance.cpp
+++ b/Engine/Source/Runtime/Navmesh/Private/DetourCrowd/DetourObstacleAvoidance.cpp
@@ -24,6 +24,7 @@
 #include "DetourCommon.h"
 #include "DetourAlloc.h"
 #include "DetourAssert.h"
+#include "DetourCrowd.h"
 #include <string.h>
 #include <math.h>
 #include <float.h>
@@ -273,7 +274,7 @@ void dtObstacleAvoidanceQuery::addCircle(const float* pos, const float rad,
 	dtVcopy(cir->dvel, dvel);
 }
 
-void dtObstacleAvoidanceQuery::addSegment(const float* p, const float* q)
+void dtObstacleAvoidanceQuery::addSegment(const float* p, const float* q, int flags)
 {
 	// [UE4] fixed condition below. Used to be strict > comparison
 	if (m_nsegments >= m_maxSegments)
@@ -282,6 +283,7 @@ void dtObstacleAvoidanceQuery::addSegment(const float* p, const float* q)
 	dtObstacleSegment* seg = &m_segments[m_nsegments++];
 	dtVcopy(seg->p, p);
 	dtVcopy(seg->q, q);
+	seg->canIgnore = (flags & DT_CROWD_BOUNDARY_IGNORE) != 0;
 }
 
 void dtObstacleAvoidanceQuery::prepare(const float* pos, const float* dvel)
@@ -368,6 +370,7 @@ float dtObstacleAvoidanceQuery::processSample(const float* vcand, const float cs
 		}
 	}
 
+	const float TooCloseToSegmentDistPct = 0.1f;
 	for (int i = 0; i < m_nsegments; ++i)
 	{
 		const dtObstacleSegment* seg = &m_segments[i];
@@ -390,10 +393,15 @@ float dtObstacleAvoidanceQuery::processSample(const float* vcand, const float cs
 		{
 			if (!isectRaySeg(pos, vcand, seg->p, seg->q, htmin))
 				continue;
+
+			if (seg->canIgnore && htmin > TooCloseToSegmentDistPct)
+			{
+				htmin = 1.0f;
+			}
 		}
 		
 		// UE4: when sample is too close to segment (navmesh wall) - disable it completely
-		if (htmin < 0.1f)
+		if (htmin < TooCloseToSegmentDistPct)
 		{
 			return -1.0f;
 		}
diff --git a/Engine/Source/Runtime/Navmesh/Public/Detour/DetourNavMeshQuery.h b/Engine/Source/Runtime/Navmesh/Public/Detour/DetourNavMeshQuery.h
--- a/Engine/Source/Runtime/Navmesh/Public/Detour/DetourNavMeshQuery.h
+++ b/Engine/Source/Runtime/Navmesh/Public/Detour/DetourNavMeshQuery.h
@@ -475,7 +475,20 @@ public:
 	dtStatus findNearestPoly(const float* center, const float* extents,
 							 const dtQueryFilter* filter,
 							 dtPolyRef* nearestRef, float* nearestPt) const;
-	
+
+	//@UE4 BEGIN
+	/// Finds the nearest polygon containing the specified center point.
+	///  @param[in]		center		The center of the search box. [(x, y, z)]
+	///  @param[in]		extents		The search distance along each axis. [(x, y, z)]
+	///  @param[in]		filter		The polygon filter to apply to the query.
+	///  @param[out]	nearestRef	The reference id of the nearest polygon.
+	///  @param[out]	nearestPt	The nearest point on the polygon. [opt] [(x, y, z)]
+	/// @returns The status flags for the query.
+	dtStatus findNearestContainingPoly(const float* center, const float* extents,
+									   const dtQueryFilter* filter,
+									   dtPolyRef* nearestRef, float* nearestPt) const;
+	//@UE4 END
+
 	/// Finds polygons that overlap the search box.
 	///  @param[in]		center		The center of the search box. [(x, y, z)]
 	///  @param[in]		extents		The search distance along each axis. [(x, y, z)]
@@ -674,7 +687,6 @@ public:
 	/// Gets best node ref and cost from sliced pathfinding data
 	void getCurrentBestResult(struct dtNode*& bestNode, float& bestCost) const { bestNode = m_query.lastBestNode; bestCost = m_query.lastBestNodeCost; }
 
-	float getQueryTime() const { return m_queryTime; }
 	int getQueryNodes() const { return m_queryNodes; }
 	/// @}
 	
@@ -754,7 +766,6 @@ private:
 	class dtNodePool* m_nodePool;		///< Pointer to node pool.
 	class dtNodeQueue* m_openList;		///< Pointer to open list queue.
 
-	mutable float m_queryTime;
 	mutable int m_queryNodes;
 };
 
diff --git a/Engine/Source/Runtime/Navmesh/Public/DetourCrowd/DetourCrowd.h b/Engine/Source/Runtime/Navmesh/Public/DetourCrowd/DetourCrowd.h
--- a/Engine/Source/Runtime/Navmesh/Public/DetourCrowd/DetourCrowd.h
+++ b/Engine/Source/Runtime/Navmesh/Public/DetourCrowd/DetourCrowd.h
@@ -207,6 +207,12 @@ enum UpdateFlags
 	DT_CROWD_SLOWDOWN_AT_GOAL		= 1 << 7,	///< [UE4] Slowdown before reaching goal
 };
 
+// [UE4] Flags used by boundary segments (dtLocalBoundary::Segment)
+enum CrowdBoundaryFlags
+{
+	DT_CROWD_BOUNDARY_IGNORE = 1 << 0,
+};
+
 struct NAVMESH_API dtCrowdAgentDebugInfo
 {
 	int idx;
diff --git a/Engine/Source/Runtime/Navmesh/Public/DetourCrowd/DetourLocalBoundary.h b/Engine/Source/Runtime/Navmesh/Public/DetourCrowd/DetourLocalBoundary.h
--- a/Engine/Source/Runtime/Navmesh/Public/DetourCrowd/DetourLocalBoundary.h
+++ b/Engine/Source/Runtime/Navmesh/Public/DetourCrowd/DetourLocalBoundary.h
@@ -34,6 +34,7 @@ class NAVMESH_API dtLocalBoundary
 	{
 		float s[6];	///< Segment start/end
 		float d;	///< Distance for pruning.
+		int flags;
 	};
 	
 	float m_center[3];
@@ -43,7 +44,7 @@ class NAVMESH_API dtLocalBoundary
 	dtPolyRef m_polys[MAX_LOCAL_POLYS];
 	int m_npolys;
 
-	void addSegment(const float dist, const float* seg);
+	void addSegment(const float dist, const float* seg, int flags = 0);
 	
 public:
 	dtLocalBoundary();
@@ -53,14 +54,14 @@ public:
 
 	// [UE4: new sections: link removal, path corridor, direction]
 	void update(dtPolyRef ref, const float* pos, const float collisionQueryRange,
-		const bool bRemoveNearLink, const float* linkV0, const float* linkV1,
+		const bool bIgnoreAtEnd, const float* endPos,
 		const dtPolyRef* path, const int npath,
 		const float* moveDir,
 		dtNavMeshQuery* navquery, const dtQueryFilter* filter);
 
 	void update(const dtSharedBoundary* sharedData, const int sharedIdx,
 		const float* pos, const float collisionQueryRange,
-		const bool bRemoveNearLink, const float* linkV0, const float* linkV1,
+		const bool bIgnoreAtEnd, const float* endPos,
 		const dtPolyRef* path, const int npath, const float* moveDir,
 		dtNavMeshQuery* navquery, const dtQueryFilter* filter);
 	
@@ -69,6 +70,7 @@ public:
 	inline const float* getCenter() const { return m_center; }
 	inline int getSegmentCount() const { return m_nsegs; }
 	inline const float* getSegment(int i) const { return m_segs[i].s; }
+	inline const int getSegmentFlags(int i) const { return m_segs[i].flags; }
 };
 
 #endif // DETOURLOCALBOUNDARY_H
diff --git a/Engine/Source/Runtime/Navmesh/Public/DetourCrowd/DetourObstacleAvoidance.h b/Engine/Source/Runtime/Navmesh/Public/DetourCrowd/DetourObstacleAvoidance.h
--- a/Engine/Source/Runtime/Navmesh/Public/DetourCrowd/DetourObstacleAvoidance.h
+++ b/Engine/Source/Runtime/Navmesh/Public/DetourCrowd/DetourObstacleAvoidance.h
@@ -34,7 +34,8 @@ struct dtObstacleCircle
 struct dtObstacleSegment
 {
 	float p[3], q[3];		///< End points of the obstacle segment
-	bool touch;
+	unsigned char touch : 1;
+	unsigned char canIgnore : 1;
 };
 
 class NAVMESH_API dtObstacleAvoidanceDebugData
@@ -114,7 +115,7 @@ public:
 	void addCircle(const float* pos, const float rad,
 				   const float* vel, const float* dvel);
 				   
-	void addSegment(const float* p, const float* q);
+	void addSegment(const float* p, const float* q, int flags = 0);
 
 	// [UE4] store new sampling pattern
 	bool setCustomSamplingPattern(int idx, const float* angles, const float* radii, int nsamples);
diff --git a/Engine/Source/Runtime/Windows/XAudio2/Private/XAudio2Source.cpp b/Engine/Source/Runtime/Windows/XAudio2/Private/XAudio2Source.cpp
--- a/Engine/Source/Runtime/Windows/XAudio2/Private/XAudio2Source.cpp
+++ b/Engine/Source/Runtime/Windows/XAudio2/Private/XAudio2Source.cpp
@@ -246,6 +246,12 @@ void FXAudio2SoundSource::SubmitPCMRTBuffers( void )
 
 	ReadMorePCMData(2, DataReadMode);
 
+	if (DataReadMode == EDataReadMode::Synchronous)
+	{
+		AudioDevice->ValidateAPICall(TEXT("SubmitSourceBuffer - PCMRT"),
+									 Source->SubmitSourceBuffer(&XAudio2Buffers[2]));
+	}
+
 	bResourcesNeedFreeing = true;
 }
 
@@ -1550,6 +1556,9 @@ void FSpatializationHelper::Init()
 	Emitter.CurveDistanceScaler = 1.0f;
 	Emitter.DopplerScaler = 1.0f;
 
+	// Zero the matrix coefficients
+	FMemory::Memzero(MatrixCoefficients, sizeof(float)*ARRAY_COUNT(MatrixCoefficients));
+
 	DSPSettings.SrcChannelCount = UE4_XAUDIO3D_INPUTCHANNELS;
 	DSPSettings.DstChannelCount = SPEAKER_COUNT;
 	DSPSettings.pMatrixCoefficients = MatrixCoefficients;
@@ -1711,6 +1720,17 @@ void FSpatializationHelper::DumpSpatializationState() const
 
 void FSpatializationHelper::CalculateDolbySurroundRate( const FVector& OrientFront, const FVector& ListenerPosition, const FVector& EmitterPosition, float OmniRadius, float* OutVolumes  )
 {
+#if ENABLE_NAN_DIAGNOSTIC
+	OrientFront.DiagnosticCheckNaN(TEXT("FSpatializationHelper: OrientFront"));
+	ListenerPosition.DiagnosticCheckNaN(TEXT("FSpatializationHelper: ListenerPosition"));
+	EmitterPosition.DiagnosticCheckNaN(TEXT("FSpatializationHelper: EmitterPosition"));
+	if (!FMath::IsFinite(OmniRadius))
+	{
+		const FString NaNorINF = FMath::IsNaN(OmniRadius) ? TEXT("NaN") : TEXT("INF");
+		UE_LOG(LogXAudio2, Warning, TEXT("OmniRadius generated a %s: %f"), *NaNorINF, OmniRadius);
+	}
+#endif
+
 	uint32 CalculateFlags = X3DAUDIO_CALCULATE_MATRIX | X3DAUDIO_CALCULATE_DOPPLER | X3DAUDIO_CALCULATE_REVERB;
 
 	Listener.OrientFront.x = OrientFront.X;
@@ -1745,6 +1765,8 @@ void FSpatializationHelper::CalculateDolbySurroundRate( const FVector& OrientFro
 			//logOrEnsureNanError(TEXT("CalculateDolbySurroundRate generated a %s in channel %d. OmniRadius:%f MatrixCoefficient:%f"),
 			//	*NaNorINF, SpeakerIndex, OmniRadius, DSPSettings.pMatrixCoefficients[SpeakerIndex]);
 #endif
+			// Zero the coefficients so we don't continue getting bad values
+			FMemory::Memzero(MatrixCoefficients, sizeof(float)*ARRAY_COUNT(MatrixCoefficients));
 		}
 #endif
 	}
diff --git a/Engine/Source/UE4Client.Target.cs b/Engine/Source/UE4Client.Target.cs
--- a/Engine/Source/UE4Client.Target.cs
+++ b/Engine/Source/UE4Client.Target.cs
@@ -48,7 +48,6 @@ public class UE4ClientTarget : TargetRules
 		ModuleNames.Add("XmlParser");
 		ModuleNames.Add("UE4Game");
 		ModuleNames.Add("AITestSuite");
-		ModuleNames.Add("GameplayDebugger");
 		ModuleNames.Add("RuntimeAssetCache");
 		ModuleNames.Add("UnrealCodeAnalyzerTests");
 		if (Target.Platform == UnrealTargetPlatform.Win32 || Target.Platform == UnrealTargetPlatform.Win64)
diff --git a/Engine/Source/UE4Game.Target.cs b/Engine/Source/UE4Game.Target.cs
--- a/Engine/Source/UE4Game.Target.cs
+++ b/Engine/Source/UE4Game.Target.cs
@@ -67,7 +67,6 @@ public class UE4GameTarget : TargetRules
 		ModuleNames.Add("XmlParser");
 		ModuleNames.Add("UE4Game");
 		ModuleNames.Add("AITestSuite");
-		ModuleNames.Add("GameplayDebugger");
 		ModuleNames.Add("RuntimeAssetCache");
 		ModuleNames.Add("UnrealCodeAnalyzerTests");
 		ModuleNames.Add("JsonUtilities");
