CIS fix for deprecation warning regarding FHierarchicalSimplification::DrawDistance

diff --git a/Engine/Source/Runtime/Engine/Classes/GameFramework/WorldSettings.h b/Engine/Source/Runtime/Engine/Classes/GameFramework/WorldSettings.h
--- a/Engine/Source/Runtime/Engine/Classes/GameFramework/WorldSettings.h
+++ b/Engine/Source/Runtime/Engine/Classes/GameFramework/WorldSettings.h
@@ -262,6 +262,8 @@ private:
 	{
 		FHierarchicalSimplification ASP1, ASP2;
 		ASP1 = ASP2;
+
+		FHierarchicalSimplification ASP3(ASP2);
 	}
 };
 
