Fix Xbox One CIS errors

Function template defaults are not supported in VS2012. Updating callsites of GetProxyOnGameThread() and GetProxyOnAnyThread() to specify template arguments.

diff --git a/Engine/Source/Runtime/Engine/Classes/Animation/AnimInstance.h b/Engine/Source/Runtime/Engine/Classes/Animation/AnimInstance.h
--- a/Engine/Source/Runtime/Engine/Classes/Animation/AnimInstance.h
+++ b/Engine/Source/Runtime/Engine/Classes/Animation/AnimInstance.h
@@ -1226,7 +1226,7 @@ protected:
 	virtual void DestroyAnimInstanceProxy(FAnimInstanceProxy* InProxy);
 
 	/** Access the proxy but block if a task is currently in progress as it wouldn't be safe to access it */
-	template <typename T = FAnimInstanceProxy>
+	template <typename T /*= FAnimInstanceProxy*/>	// @TODO: Cant default parameters to this function on Xbox One until we move off the VS2012 compiler
 	FORCEINLINE T& GetProxyOnGameThread()
 	{
 		check(IsInGameThread());
@@ -1244,7 +1244,7 @@ protected:
 	}
 
 	/** Access the proxy but block if a task is currently in progress as it wouldn't be safe to access it */
-	template <typename T = FAnimInstanceProxy>
+	template <typename T/* = FAnimInstanceProxy*/>	// @TODO: Cant default parameters to this function on Xbox One until we move off the VS2012 compiler
 	FORCEINLINE const T& GetProxyOnGameThread() const
 	{
 		check(IsInGameThread());
@@ -1262,7 +1262,7 @@ protected:
 	}
 
 	/** Access the proxy but block if a task is currently in progress (and we are on the game thread) as it wouldn't be safe to access it */
-	template <typename T = FAnimInstanceProxy>
+	template <typename T/* = FAnimInstanceProxy*/>	// @TODO: Cant default parameters to this function on Xbox One until we move off the VS2012 compiler
 	FORCEINLINE T& GetProxyOnAnyThread()
 	{
 		if(GetOuter() && GetOuter()->IsA<USkeletalMeshComponent>())
@@ -1282,7 +1282,7 @@ protected:
 	}
 
 	/** Access the proxy but block if a task is currently in progress (and we are on the game thread) as it wouldn't be safe to access it */
-	template <typename T = FAnimInstanceProxy>
+	template <typename T/* = FAnimInstanceProxy*/>	// @TODO: Cant default parameters to this function on Xbox One until we move off the VS2012 compiler
 	FORCEINLINE const T& GetProxyOnAnyThread() const
 	{
 		if(GetOuter() && GetOuter()->IsA<USkeletalMeshComponent>())
diff --git a/Engine/Source/Runtime/Engine/Private/Animation/AnimInstance.cpp b/Engine/Source/Runtime/Engine/Private/Animation/AnimInstance.cpp
--- a/Engine/Source/Runtime/Engine/Private/Animation/AnimInstance.cpp
+++ b/Engine/Source/Runtime/Engine/Private/Animation/AnimInstance.cpp
@@ -95,12 +95,12 @@ PRAGMA_ENABLE_DEPRECATION_WARNINGS
 
 void UAnimInstance::MakeSequenceTickRecord(FAnimTickRecord& TickRecord, class UAnimSequenceBase* Sequence, bool bLooping, float PlayRate, float FinalBlendWeight, float& CurrentTime, FMarkerTickRecord& MarkerTickRecord) const
 {
-	GetProxyOnGameThread().MakeSequenceTickRecord(TickRecord, Sequence, bLooping, PlayRate, FinalBlendWeight, CurrentTime, MarkerTickRecord);
+	GetProxyOnGameThread<FAnimInstanceProxy>().MakeSequenceTickRecord(TickRecord, Sequence, bLooping, PlayRate, FinalBlendWeight, CurrentTime, MarkerTickRecord);
 }
 
 void UAnimInstance::MakeBlendSpaceTickRecord(FAnimTickRecord& TickRecord, class UBlendSpaceBase* BlendSpace, const FVector& BlendInput, TArray<FBlendSampleData>& BlendSampleDataCache, FBlendFilter& BlendFilter, bool bLooping, float PlayRate, float FinalBlendWeight, float& CurrentTime, FMarkerTickRecord& MarkerTickRecord) const
 {
-	GetProxyOnGameThread().MakeBlendSpaceTickRecord(TickRecord, BlendSpace, BlendInput, BlendSampleDataCache, BlendFilter, bLooping, PlayRate, FinalBlendWeight, CurrentTime, MarkerTickRecord);
+	GetProxyOnGameThread<FAnimInstanceProxy>().MakeBlendSpaceTickRecord(TickRecord, BlendSpace, BlendInput, BlendSampleDataCache, BlendFilter, bLooping, PlayRate, FinalBlendWeight, CurrentTime, MarkerTickRecord);
 }
 
 // this is only used by montage marker based sync
@@ -119,17 +119,17 @@ void UAnimInstance::MakeMontageTickRecord(FAnimTickRecord& TickRecord, class UAn
 
 void UAnimInstance::SequenceAdvanceImmediate(UAnimSequenceBase* Sequence, bool bLooping, float PlayRate, float DeltaSeconds, float& CurrentTime, FMarkerTickRecord& MarkerTickRecord)
 {
-	GetProxyOnGameThread().SequenceAdvanceImmediate(Sequence, bLooping, PlayRate, DeltaSeconds, CurrentTime, MarkerTickRecord);
+	GetProxyOnGameThread<FAnimInstanceProxy>().SequenceAdvanceImmediate(Sequence, bLooping, PlayRate, DeltaSeconds, CurrentTime, MarkerTickRecord);
 }
 
 void UAnimInstance::BlendSpaceAdvanceImmediate(class UBlendSpaceBase* BlendSpace, const FVector& BlendInput, TArray<FBlendSampleData>& BlendSampleDataCache, FBlendFilter& BlendFilter, bool bLooping, float PlayRate, float DeltaSeconds, float& CurrentTime, FMarkerTickRecord& MarkerTickRecord)
 {
-	GetProxyOnGameThread().BlendSpaceAdvanceImmediate(BlendSpace, BlendInput, BlendSampleDataCache, BlendFilter, bLooping, PlayRate, DeltaSeconds, CurrentTime, MarkerTickRecord);
+	GetProxyOnGameThread<FAnimInstanceProxy>().BlendSpaceAdvanceImmediate(BlendSpace, BlendInput, BlendSampleDataCache, BlendFilter, bLooping, PlayRate, DeltaSeconds, CurrentTime, MarkerTickRecord);
 }
 
 FAnimTickRecord& UAnimInstance::CreateUninitializedTickRecord(int32 GroupIndex, FAnimGroupInstance*& OutSyncGroupPtr)
 {
-	return GetProxyOnGameThread().CreateUninitializedTickRecord(GroupIndex, OutSyncGroupPtr);
+	return GetProxyOnGameThread<FAnimInstanceProxy>().CreateUninitializedTickRecord(GroupIndex, OutSyncGroupPtr);
 }
 
 AActor* UAnimInstance::GetOwningActor() const
@@ -198,28 +198,28 @@ void UAnimInstance::InitializeAnimation()
 	// before initialize, need to recalculate required bone list
 	RecalcRequiredBones();
 
-	GetProxyOnGameThread().Initialize(this);
+	GetProxyOnGameThread<FAnimInstanceProxy>().Initialize(this);
 
 	ClearMorphTargets();
 	NativeInitializeAnimation();
 	BlueprintInitializeAnimation();
 
-	GetProxyOnGameThread().InitializeRootNode();
+	GetProxyOnGameThread<FAnimInstanceProxy>().InitializeRootNode();
 
 PRAGMA_DISABLE_DEPRECATION_WARNINGS
 	// @todo: remove after deprecation
-	RootNode = GetProxyOnGameThread().GetRootNode();
+	RootNode = GetProxyOnGameThread<FAnimInstanceProxy>().GetRootNode();
 PRAGMA_ENABLE_DEPRECATION_WARNINGS
 
 	// we can bind rules & events now the graph has been initialized
-	GetProxyOnGameThread().BindNativeDelegates();
+	GetProxyOnGameThread<FAnimInstanceProxy>().BindNativeDelegates();
 }
 
 void UAnimInstance::UninitializeAnimation()
 {
 	NativeUninitializeAnimation();
 
-	GetProxyOnGameThread().Uninitialize(this);
+	GetProxyOnGameThread<FAnimInstanceProxy>().Uninitialize(this);
 
 	StopAllMontages(0.f);
 
@@ -282,7 +282,7 @@ bool UAnimInstance::UpdateSnapshotAndSkipRemainingUpdate()
 // update animation node, if you have your own anim instance, you could override this to update extra nodes if you'd like
 void UAnimInstance::UpdateAnimationNode(float DeltaSeconds)
 {
-	GetProxyOnGameThread().UpdateAnimationNode(DeltaSeconds);
+	GetProxyOnGameThread<FAnimInstanceProxy>().UpdateAnimationNode(DeltaSeconds);
 }
 
 void UAnimInstance::UpdateMontage(float DeltaSeconds)
@@ -307,7 +307,7 @@ void UAnimInstance::UpdateMontage(float DeltaSeconds)
 			if (ensure (GroupIndexToUse != INDEX_NONE))
 			{
 				FAnimGroupInstance* SyncGroup;
-				FAnimTickRecord& TickRecord = GetProxyOnGameThread().CreateUninitializedTickRecord(GroupIndexToUse, /*out*/ SyncGroup);
+				FAnimTickRecord& TickRecord = GetProxyOnGameThread<FAnimInstanceProxy>().CreateUninitializedTickRecord(GroupIndexToUse, /*out*/ SyncGroup);
 				MakeMontageTickRecord(TickRecord, MontageInstance->Montage, MontageInstance->GetPosition(), 
 					MontageInstance->GetPreviousPosition(), MontageInstance->GetDeltaMoved(), MontageInstance->GetWeight(), 
 					MontageInstance->MarkersPassedThisTick, MontageInstance->MarkerTickRecord);
@@ -331,7 +331,7 @@ void UAnimInstance::UpdateAnimation(float DeltaSeconds, bool bNeedsValidRootMoti
 	FScopeCycleCounterUObject AnimScope(this);
 
 	// acquire the proxy as we need to update
-	FAnimInstanceProxy& Proxy = GetProxyOnGameThread();
+	FAnimInstanceProxy& Proxy = GetProxyOnGameThread<FAnimInstanceProxy>();
 
 #if WITH_EDITOR
 	if (GIsEditor)
@@ -406,7 +406,7 @@ void UAnimInstance::PreUpdateAnimation(float DeltaSeconds)
 
 	NotifyQueue.Reset(GetSkelMeshComponent());
 
-	GetProxyOnGameThread().PreUpdate(this, DeltaSeconds);
+	GetProxyOnGameThread<FAnimInstanceProxy>().PreUpdate(this, DeltaSeconds);
 }
 
 void UAnimInstance::PostUpdateAnimation()
@@ -415,7 +415,7 @@ void UAnimInstance::PostUpdateAnimation()
 	check(!IsRunningParallelEvaluation());
 
 	// acquire the proxy as we need to update
-	FAnimInstanceProxy& Proxy = GetProxyOnGameThread();
+	FAnimInstanceProxy& Proxy = GetProxyOnGameThread<FAnimInstanceProxy>();
 
 	bNeedsUpdate = false;
 
@@ -460,7 +460,7 @@ void UAnimInstance::PostUpdateAnimation()
 
 void UAnimInstance::ParallelUpdateAnimation()
 {
-	GetProxyOnAnyThread().UpdateAnimation();
+	GetProxyOnAnyThread<FAnimInstanceProxy>().UpdateAnimation();
 }
 
 bool UAnimInstance::NeedsImmediateUpdate(float DeltaSeconds) const
@@ -483,18 +483,18 @@ bool UAnimInstance::NeedsUpdate() const
 
 void UAnimInstance::EvaluateAnimation(FPoseContext& Output)
 {
-	GetProxyOnGameThread().EvaluateAnimation(Output);
+	GetProxyOnGameThread<FAnimInstanceProxy>().EvaluateAnimation(Output);
 }
 
 bool UAnimInstance::ParallelCanEvaluate(const USkeletalMesh* InSkeletalMesh) const
 {
-	const FAnimInstanceProxy& Proxy = GetProxyOnAnyThread();
+	const FAnimInstanceProxy& Proxy = GetProxyOnAnyThread<FAnimInstanceProxy>();
 	return Proxy.GetRequiredBones().IsValid() && (Proxy.GetRequiredBones().GetAsset() == InSkeletalMesh);
 }
 
 void UAnimInstance::ParallelEvaluateAnimation(bool bForceRefPose, const USkeletalMesh* InSkeletalMesh, TArray<FTransform>& OutLocalAtoms, TArray<FActiveVertexAnim>& OutVertexAnims, FBlendedCurve& OutCurve)
 {
-	FAnimInstanceProxy& Proxy = GetProxyOnAnyThread();
+	FAnimInstanceProxy& Proxy = GetProxyOnAnyThread<FAnimInstanceProxy>();
 
 	if( !bForceRefPose )
 	{
@@ -568,32 +568,32 @@ void UAnimInstance::NativeUninitializeAnimation()
 
 void UAnimInstance::AddNativeTransitionBinding(const FName& MachineName, const FName& PrevStateName, const FName& NextStateName, const FCanTakeTransition& NativeTransitionDelegate, const FName& TransitionName)
 {
-	GetProxyOnGameThread().AddNativeTransitionBinding(MachineName, PrevStateName, NextStateName, NativeTransitionDelegate, TransitionName);
+	GetProxyOnGameThread<FAnimInstanceProxy>().AddNativeTransitionBinding(MachineName, PrevStateName, NextStateName, NativeTransitionDelegate, TransitionName);
 }
 
 bool UAnimInstance::HasNativeTransitionBinding(const FName& MachineName, const FName& PrevStateName, const FName& NextStateName, FName& OutBindingName)
 {
-	return GetProxyOnGameThread().HasNativeTransitionBinding(MachineName, PrevStateName, NextStateName, OutBindingName);
+	return GetProxyOnGameThread<FAnimInstanceProxy>().HasNativeTransitionBinding(MachineName, PrevStateName, NextStateName, OutBindingName);
 }
 
 void UAnimInstance::AddNativeStateEntryBinding(const FName& MachineName, const FName& StateName, const FOnGraphStateChanged& NativeEnteredDelegate)
 {
-	GetProxyOnGameThread().AddNativeStateEntryBinding(MachineName, StateName, NativeEnteredDelegate);
+	GetProxyOnGameThread<FAnimInstanceProxy>().AddNativeStateEntryBinding(MachineName, StateName, NativeEnteredDelegate);
 }
 	
 bool UAnimInstance::HasNativeStateEntryBinding(const FName& MachineName, const FName& StateName, FName& OutBindingName)
 {
-	return GetProxyOnGameThread().HasNativeStateEntryBinding(MachineName, StateName, OutBindingName);
+	return GetProxyOnGameThread<FAnimInstanceProxy>().HasNativeStateEntryBinding(MachineName, StateName, OutBindingName);
 }
 
 void UAnimInstance::AddNativeStateExitBinding(const FName& MachineName, const FName& StateName, const FOnGraphStateChanged& NativeExitedDelegate)
 {
-	GetProxyOnGameThread().AddNativeStateExitBinding(MachineName, StateName, NativeExitedDelegate);
+	GetProxyOnGameThread<FAnimInstanceProxy>().AddNativeStateExitBinding(MachineName, StateName, NativeExitedDelegate);
 }
 
 bool UAnimInstance::HasNativeStateExitBinding(const FName& MachineName, const FName& StateName, FName& OutBindingName)
 {
-	return GetProxyOnGameThread().HasNativeStateExitBinding(MachineName, StateName, OutBindingName);
+	return GetProxyOnGameThread<FAnimInstanceProxy>().HasNativeStateExitBinding(MachineName, StateName, OutBindingName);
 }
 
 void OutputCurveMap(TMap<FName, float>& CurveMap, UCanvas* Canvas, FDisplayDebugManager& DisplayDebugManager, float Indent)
@@ -682,7 +682,7 @@ void OutputTickRecords(const TArray<FAnimTickRecord>& Records, UCanvas* Canvas,
 
 void UAnimInstance::DisplayDebug(class UCanvas* Canvas, const FDebugDisplayInfo& DebugDisplay, float& YL, float& YPos)
 {
-	FAnimInstanceProxy& Proxy = GetProxyOnGameThread();
+	FAnimInstanceProxy& Proxy = GetProxyOnGameThread<FAnimInstanceProxy>();
 
 	float Indent = 0.f;
 
@@ -796,8 +796,8 @@ void UAnimInstance::DisplayDebug(class UCanvas* Canvas, const FDebugDisplayInfo&
 		FIndenter AnimIndent(Indent);
 
 		//Display Sync Groups
-		const TArray<FAnimGroupInstance>& SyncGroups = GetProxyOnGameThread().GetSyncGroupRead();
-		const TArray<FAnimTickRecord>& UngroupedActivePlayers = GetProxyOnGameThread().GetUngroupedActivePlayersRead();
+		const TArray<FAnimGroupInstance>& SyncGroups = GetProxyOnGameThread<FAnimInstanceProxy>().GetSyncGroupRead();
+		const TArray<FAnimTickRecord>& UngroupedActivePlayers = GetProxyOnGameThread<FAnimInstanceProxy>().GetUngroupedActivePlayersRead();
 
 		Heading = FString::Printf(TEXT("SyncGroups: %i"), SyncGroups.Num());
 		DisplayDebugManager.DrawString(Heading, Indent);
@@ -923,11 +923,11 @@ void UAnimInstance::RecalcRequiredBones()
 
 	if( SkelMeshComp->SkeletalMesh && SkelMeshComp->SkeletalMesh->Skeleton )
 	{
-		GetProxyOnGameThread().RecalcRequiredBones(SkelMeshComp, SkelMeshComp->SkeletalMesh);
+		GetProxyOnGameThread<FAnimInstanceProxy>().RecalcRequiredBones(SkelMeshComp, SkelMeshComp->SkeletalMesh);
 	}
 	else if( CurrentSkeleton != NULL )
 	{
-		GetProxyOnGameThread().RecalcRequiredBones(SkelMeshComp, CurrentSkeleton);
+		GetProxyOnGameThread<FAnimInstanceProxy>().RecalcRequiredBones(SkelMeshComp, CurrentSkeleton);
 	}
 }
 
@@ -937,7 +937,7 @@ void UAnimInstance::Serialize(FArchive& Ar)
 
 	if (!Ar.IsLoading() || !Ar.IsSaving())
 	{
-		Ar << GetProxyOnGameThread().GetRequiredBones();
+		Ar << GetProxyOnGameThread<FAnimInstanceProxy>().GetRequiredBones();
 	}
 }
 
@@ -986,12 +986,12 @@ bool UAnimInstance::PassesChanceOfTriggering(const FAnimNotifyEvent* Event) cons
 
 void UAnimInstance::AddAnimNotifyFromGeneratedClass(int32 NotifyIndex)
 {
-	GetProxyOnGameThread().AddAnimNotifyFromGeneratedClass(NotifyIndex);
+	GetProxyOnGameThread<FAnimInstanceProxy>().AddAnimNotifyFromGeneratedClass(NotifyIndex);
 }
 
 void UAnimInstance::AddCurveValue(const FName& CurveName, float Value, int32 CurveTypeFlags)
 {
-	FAnimInstanceProxy& Proxy = GetProxyOnGameThread();
+	FAnimInstanceProxy& Proxy = GetProxyOnGameThread<FAnimInstanceProxy>();
 
 	// save curve value, it will overwrite if same exists, 
 	//CurveValues.Add(CurveName, Value);
@@ -1054,24 +1054,24 @@ void UAnimInstance::AddCurveValue(const USkeleton::AnimCurveUID Uid, float Value
 
 int32 UAnimInstance::GetSyncGroupReadIndex() const
 { 
-	return GetProxyOnGameThread().GetSyncGroupReadIndex(); 
+	return GetProxyOnGameThread<FAnimInstanceProxy>().GetSyncGroupReadIndex(); 
 }
 
 int32 UAnimInstance::GetSyncGroupWriteIndex() const 
 { 
-	return GetProxyOnGameThread().GetSyncGroupWriteIndex(); 
+	return GetProxyOnGameThread<FAnimInstanceProxy>().GetSyncGroupWriteIndex(); 
 }
 
 void UAnimInstance::TickSyncGroupWriteIndex() 
 {
-	GetProxyOnGameThread().TickSyncGroupWriteIndex();
+	GetProxyOnGameThread<FAnimInstanceProxy>().TickSyncGroupWriteIndex();
 }
 
 void UAnimInstance::UpdateCurves(const FBlendedCurve& InCurve)
 {
 	SCOPE_CYCLE_COUNTER(STAT_UpdateCurves);
 
-	FAnimInstanceProxy& Proxy = GetProxyOnGameThread();
+	FAnimInstanceProxy& Proxy = GetProxyOnGameThread<FAnimInstanceProxy>();
 
 	//Track material params we set last time round so we can clear them if they aren't set again.
 	MaterialParamatersToClear.Reset();
@@ -1104,12 +1104,12 @@ void UAnimInstance::UpdateCurves(const FBlendedCurve& InCurve)
 
 bool UAnimInstance::HasMorphTargetCurves() const
 {
-	return GetProxyOnGameThread().HasMorphTargetCurves();
+	return GetProxyOnGameThread<FAnimInstanceProxy>().HasMorphTargetCurves();
 }
 
 TMap<FName, float>& UAnimInstance::GetMorphTargetCurves()
 {
-	return GetProxyOnGameThread().GetMorphTargetCurves();
+	return GetProxyOnGameThread<FAnimInstanceProxy>().GetMorphTargetCurves();
 }
 
 void UAnimInstance::TriggerAnimNotifies(float DeltaSeconds)
@@ -1217,47 +1217,47 @@ void UAnimInstance::TriggerSingleAnimNotify(const FAnimNotifyEvent* AnimNotifyEv
 
 void UAnimInstance::GetSlotWeight(FName const& SlotNodeName, float& out_SlotNodeWeight, float& out_SourceWeight, float& out_TotalNodeWeight) const
 {
-	GetProxyOnGameThread().GetSlotWeight(SlotNodeName, out_SlotNodeWeight, out_SourceWeight, out_TotalNodeWeight);
+	GetProxyOnGameThread<FAnimInstanceProxy>().GetSlotWeight(SlotNodeName, out_SlotNodeWeight, out_SourceWeight, out_TotalNodeWeight);
 }
 
 void UAnimInstance::SlotEvaluatePose(FName SlotNodeName, const FCompactPose& SourcePose, const FBlendedCurve& SourceCurve, float InSourceWeight, FCompactPose& BlendedPose, FBlendedCurve& BlendedCurve, float InBlendWeight, float InTotalNodeWeight)
 {
-	GetProxyOnGameThread().SlotEvaluatePose(SlotNodeName, SourcePose, SourceCurve, InSourceWeight, BlendedPose, BlendedCurve, InBlendWeight, InTotalNodeWeight);
+	GetProxyOnGameThread<FAnimInstanceProxy>().SlotEvaluatePose(SlotNodeName, SourcePose, SourceCurve, InSourceWeight, BlendedPose, BlendedCurve, InBlendWeight, InTotalNodeWeight);
 }
 
 void UAnimInstance::ReinitializeSlotNodes()
 {
-	GetProxyOnGameThread().ReinitializeSlotNodes();
+	GetProxyOnGameThread<FAnimInstanceProxy>().ReinitializeSlotNodes();
 }
 
 void UAnimInstance::RegisterSlotNodeWithAnimInstance(FName SlotNodeName)
 {
-	GetProxyOnGameThread().RegisterSlotNodeWithAnimInstance(SlotNodeName);
+	GetProxyOnGameThread<FAnimInstanceProxy>().RegisterSlotNodeWithAnimInstance(SlotNodeName);
 }
 
 void UAnimInstance::UpdateSlotNodeWeight(FName SlotNodeName, float Weight)
 {
-	GetProxyOnGameThread().UpdateSlotNodeWeight(SlotNodeName, Weight);
+	GetProxyOnGameThread<FAnimInstanceProxy>().UpdateSlotNodeWeight(SlotNodeName, Weight);
 }
 
 void UAnimInstance::ClearSlotNodeWeights()
 {
-	GetProxyOnGameThread().ClearSlotNodeWeights();
+	GetProxyOnGameThread<FAnimInstanceProxy>().ClearSlotNodeWeights();
 }
 
 bool UAnimInstance::IsSlotNodeRelevantForNotifies(FName SlotNodeName) const
 {
-	return GetProxyOnGameThread().IsSlotNodeRelevantForNotifies(SlotNodeName);
+	return GetProxyOnGameThread<FAnimInstanceProxy>().IsSlotNodeRelevantForNotifies(SlotNodeName);
 }
 
 void UAnimInstance::UpdateSlotRootMotionWeight(FName SlotNodeName, float Weight)
 {
-	GetProxyOnGameThread().UpdateSlotRootMotionWeight(SlotNodeName, Weight);
+	GetProxyOnGameThread<FAnimInstanceProxy>().UpdateSlotRootMotionWeight(SlotNodeName, Weight);
 }
 
 float UAnimInstance::GetSlotRootMotionWeight(FName SlotNodeName) const
 {
-	return GetProxyOnGameThread().GetSlotRootMotionWeight(SlotNodeName);
+	return GetProxyOnGameThread<FAnimInstanceProxy>().GetSlotRootMotionWeight(SlotNodeName);
 }
 
 float UAnimInstance::GetCurveValue(FName CurveName)
@@ -1416,7 +1416,7 @@ void UAnimInstance::Montage_Advance(float DeltaSeconds)
 			FRootMotionMovementParams* RootMotionParams = NULL;
 			if (bExtractRootMotion)
 			{
-				RootMotionParams = (RootMotionMode != ERootMotionMode::IgnoreRootMotion) ? &GetProxyOnGameThread().GetExtractedRootMotion() : &LocalExtractedRootMotion;
+				RootMotionParams = (RootMotionMode != ERootMotionMode::IgnoreRootMotion) ? &GetProxyOnGameThread<FAnimInstanceProxy>().GetExtractedRootMotion() : &LocalExtractedRootMotion;
 			}
 
 			MontageInstance->MontageSync_PreUpdate();
@@ -2272,13 +2272,13 @@ FRootMotionMovementParams UAnimInstance::ConsumeExtractedRootMotion(float Alpha)
 	}
 	else if (Alpha > (1.f - ZERO_ANIMWEIGHT_THRESH))
 	{
-		FRootMotionMovementParams RootMotion = GetProxyOnGameThread().GetExtractedRootMotion();
-		GetProxyOnGameThread().GetExtractedRootMotion().Clear();
+		FRootMotionMovementParams RootMotion = GetProxyOnGameThread<FAnimInstanceProxy>().GetExtractedRootMotion();
+		GetProxyOnGameThread<FAnimInstanceProxy>().GetExtractedRootMotion().Clear();
 		return RootMotion;
 	}
 	else
 	{
-		return GetProxyOnGameThread().GetExtractedRootMotion().ConsumeRootMotion(Alpha);
+		return GetProxyOnGameThread<FAnimInstanceProxy>().GetExtractedRootMotion().ConsumeRootMotion(Alpha);
 	}
 }
 
@@ -2354,27 +2354,27 @@ void UAnimInstance::UnlockAIResources(bool bUnlockMovement, bool UnlockAILogic)
 
 bool UAnimInstance::GetTimeToClosestMarker(FName SyncGroup, FName MarkerName, float& OutMarkerTime) const
 {
-	return GetProxyOnAnyThread().GetTimeToClosestMarker(SyncGroup, MarkerName, OutMarkerTime);
+	return GetProxyOnAnyThread<FAnimInstanceProxy>().GetTimeToClosestMarker(SyncGroup, MarkerName, OutMarkerTime);
 }
 
 bool UAnimInstance::HasMarkerBeenHitThisFrame(FName SyncGroup, FName MarkerName) const
 {
-	return GetProxyOnAnyThread().HasMarkerBeenHitThisFrame(SyncGroup, MarkerName);
+	return GetProxyOnAnyThread<FAnimInstanceProxy>().HasMarkerBeenHitThisFrame(SyncGroup, MarkerName);
 }
 
 bool UAnimInstance::IsSyncGroupBetweenMarkers(FName InSyncGroupName, FName PreviousMarker, FName NextMarker, bool bRespectMarkerOrder) const
 {
-	return GetProxyOnAnyThread().IsSyncGroupBetweenMarkers(InSyncGroupName, PreviousMarker, NextMarker, bRespectMarkerOrder);
+	return GetProxyOnAnyThread<FAnimInstanceProxy>().IsSyncGroupBetweenMarkers(InSyncGroupName, PreviousMarker, NextMarker, bRespectMarkerOrder);
 }
 
 FMarkerSyncAnimPosition UAnimInstance::GetSyncGroupPosition(FName InSyncGroupName) const
 {
-	return GetProxyOnAnyThread().GetSyncGroupPosition(InSyncGroupName);
+	return GetProxyOnAnyThread<FAnimInstanceProxy>().GetSyncGroupPosition(InSyncGroupName);
 }
 
 void UAnimInstance::UpdateMontageEvaluationData()
 {
-	FAnimInstanceProxy& Proxy = GetProxyOnGameThread();
+	FAnimInstanceProxy& Proxy = GetProxyOnGameThread<FAnimInstanceProxy>();
 
 	Proxy.GetMontageEvaluationData().Reset(MontageInstances.Num());
 	UE_LOG(LogAnimMontage, Verbose, TEXT("UpdateMontageEvaluationData Strting: Owner: %s"),	*GetNameSafe(GetOwningActor()));
@@ -2394,42 +2394,42 @@ void UAnimInstance::UpdateMontageEvaluationData()
 
 float UAnimInstance::GetInstanceAssetPlayerLength(int32 AssetPlayerIndex)
 {
-	return GetProxyOnAnyThread().GetInstanceAssetPlayerLength(AssetPlayerIndex);
+	return GetProxyOnAnyThread<FAnimInstanceProxy>().GetInstanceAssetPlayerLength(AssetPlayerIndex);
 }
 
 float UAnimInstance::GetInstanceAssetPlayerTime(int32 AssetPlayerIndex)
 {
-	return GetProxyOnAnyThread().GetInstanceAssetPlayerTime(AssetPlayerIndex);
+	return GetProxyOnAnyThread<FAnimInstanceProxy>().GetInstanceAssetPlayerTime(AssetPlayerIndex);
 }
 
 float UAnimInstance::GetInstanceAssetPlayerTimeFraction(int32 AssetPlayerIndex)
 {
-	return GetProxyOnAnyThread().GetInstanceAssetPlayerTimeFraction(AssetPlayerIndex);
+	return GetProxyOnAnyThread<FAnimInstanceProxy>().GetInstanceAssetPlayerTimeFraction(AssetPlayerIndex);
 }
 
 float UAnimInstance::GetInstanceAssetPlayerTimeFromEndFraction(int32 AssetPlayerIndex)
 {
-	return GetProxyOnAnyThread().GetInstanceAssetPlayerTimeFromEndFraction(AssetPlayerIndex);
+	return GetProxyOnAnyThread<FAnimInstanceProxy>().GetInstanceAssetPlayerTimeFromEndFraction(AssetPlayerIndex);
 }
 
 float UAnimInstance::GetInstanceAssetPlayerTimeFromEnd(int32 AssetPlayerIndex)
 {
-	return GetProxyOnAnyThread().GetInstanceAssetPlayerTimeFromEnd(AssetPlayerIndex);
+	return GetProxyOnAnyThread<FAnimInstanceProxy>().GetInstanceAssetPlayerTimeFromEnd(AssetPlayerIndex);
 }
 
 float UAnimInstance::GetInstanceStateWeight(int32 MachineIndex, int32 StateIndex)
 {
-	return GetProxyOnAnyThread().GetInstanceStateWeight(MachineIndex, StateIndex);
+	return GetProxyOnAnyThread<FAnimInstanceProxy>().GetInstanceStateWeight(MachineIndex, StateIndex);
 }
 
 float UAnimInstance::GetInstanceCurrentStateElapsedTime(int32 MachineIndex)
 {
-	return GetProxyOnAnyThread().GetInstanceCurrentStateElapsedTime(MachineIndex);
+	return GetProxyOnAnyThread<FAnimInstanceProxy>().GetInstanceCurrentStateElapsedTime(MachineIndex);
 }
 
 float UAnimInstance::GetInstanceTransitionCrossfadeDuration(int32 MachineIndex, int32 TransitionIndex)
 {
-	return GetProxyOnAnyThread().GetInstanceTransitionCrossfadeDuration(MachineIndex, TransitionIndex);
+	return GetProxyOnAnyThread<FAnimInstanceProxy>().GetInstanceTransitionCrossfadeDuration(MachineIndex, TransitionIndex);
 }
 
 float UAnimInstance::GetInstanceTransitionTimeElapsed(int32 MachineIndex, int32 TransitionIndex)
@@ -2440,82 +2440,82 @@ float UAnimInstance::GetInstanceTransitionTimeElapsed(int32 MachineIndex, int32
 
 float UAnimInstance::GetInstanceTransitionTimeElapsedFraction(int32 MachineIndex, int32 TransitionIndex)
 {
-	return GetProxyOnAnyThread().GetInstanceTransitionTimeElapsedFraction(MachineIndex, TransitionIndex);
+	return GetProxyOnAnyThread<FAnimInstanceProxy>().GetInstanceTransitionTimeElapsedFraction(MachineIndex, TransitionIndex);
 }
 
 float UAnimInstance::GetRelevantAnimTimeRemaining(int32 MachineIndex, int32 StateIndex)
 {
-	return GetProxyOnAnyThread().GetRelevantAnimTimeRemaining(MachineIndex, StateIndex);
+	return GetProxyOnAnyThread<FAnimInstanceProxy>().GetRelevantAnimTimeRemaining(MachineIndex, StateIndex);
 }
 
 float UAnimInstance::GetRelevantAnimTimeRemainingFraction(int32 MachineIndex, int32 StateIndex)
 {
-	return GetProxyOnAnyThread().GetRelevantAnimTimeRemainingFraction(MachineIndex, StateIndex);
+	return GetProxyOnAnyThread<FAnimInstanceProxy>().GetRelevantAnimTimeRemainingFraction(MachineIndex, StateIndex);
 }
 
 float UAnimInstance::GetRelevantAnimLength(int32 MachineIndex, int32 StateIndex)
 {
-	return GetProxyOnAnyThread().GetRelevantAnimLength(MachineIndex, StateIndex);
+	return GetProxyOnAnyThread<FAnimInstanceProxy>().GetRelevantAnimLength(MachineIndex, StateIndex);
 }
 
 float UAnimInstance::GetRelevantAnimTime(int32 MachineIndex, int32 StateIndex)
 {
-	return GetProxyOnAnyThread().GetRelevantAnimTime(MachineIndex, StateIndex);
+	return GetProxyOnAnyThread<FAnimInstanceProxy>().GetRelevantAnimTime(MachineIndex, StateIndex);
 }
 
 float UAnimInstance::GetRelevantAnimTimeFraction(int32 MachineIndex, int32 StateIndex)
 {
-	return GetProxyOnAnyThread().GetRelevantAnimTimeFraction(MachineIndex, StateIndex);
+	return GetProxyOnAnyThread<FAnimInstanceProxy>().GetRelevantAnimTimeFraction(MachineIndex, StateIndex);
 }
 
 FAnimNode_StateMachine* UAnimInstance::GetStateMachineInstance(int32 MachineIndex)
 {
-	return GetProxyOnAnyThread().GetStateMachineInstance(MachineIndex);
+	return GetProxyOnAnyThread<FAnimInstanceProxy>().GetStateMachineInstance(MachineIndex);
 }
 
 FAnimNode_StateMachine* UAnimInstance::GetStateMachineInstanceFromName(FName MachineName)
 {
-	return GetProxyOnAnyThread().GetStateMachineInstanceFromName(MachineName);
+	return GetProxyOnAnyThread<FAnimInstanceProxy>().GetStateMachineInstanceFromName(MachineName);
 }
 
 const FBakedAnimationStateMachine* UAnimInstance::GetStateMachineInstanceDesc(FName MachineName)
 {
-	return GetProxyOnAnyThread().GetStateMachineInstanceDesc(MachineName);
+	return GetProxyOnAnyThread<FAnimInstanceProxy>().GetStateMachineInstanceDesc(MachineName);
 }
 
 int32 UAnimInstance::GetStateMachineIndex(FName MachineName)
 {
-	return GetProxyOnAnyThread().GetStateMachineIndex(MachineName);
+	return GetProxyOnAnyThread<FAnimInstanceProxy>().GetStateMachineIndex(MachineName);
 }
 
 FAnimNode_Base* UAnimInstance::GetCheckedNodeFromIndexUntyped(int32 NodeIdx, UScriptStruct* RequiredStructType)
 {
-	return GetProxyOnGameThread().GetCheckedNodeFromIndexUntyped(NodeIdx, RequiredStructType);
+	return GetProxyOnGameThread<FAnimInstanceProxy>().GetCheckedNodeFromIndexUntyped(NodeIdx, RequiredStructType);
 }
 
 FAnimNode_Base* UAnimInstance::GetNodeFromIndexUntyped(int32 NodeIdx, UScriptStruct* RequiredStructType)
 {
-	return GetProxyOnGameThread().GetNodeFromIndexUntyped(NodeIdx, RequiredStructType);
+	return GetProxyOnGameThread<FAnimInstanceProxy>().GetNodeFromIndexUntyped(NodeIdx, RequiredStructType);
 }
 
 const FBakedAnimationStateMachine* UAnimInstance::GetMachineDescription(UAnimBlueprintGeneratedClass* AnimBlueprintClass, FAnimNode_StateMachine* MachineInstance)
 {
-	return GetProxyOnGameThread().GetMachineDescription(AnimBlueprintClass, MachineInstance);
+	return GetProxyOnGameThread<FAnimInstanceProxy>().GetMachineDescription(AnimBlueprintClass, MachineInstance);
 }
 
 int32 UAnimInstance::GetInstanceAssetPlayerIndex(FName MachineName, FName StateName, FName AssetName)
 {
-	return GetProxyOnGameThread().GetInstanceAssetPlayerIndex(MachineName, StateName, AssetName);
+	return GetProxyOnGameThread<FAnimInstanceProxy>().GetInstanceAssetPlayerIndex(MachineName, StateName, AssetName);
 }
 
 FAnimNode_AssetPlayerBase* UAnimInstance::GetRelevantAssetPlayerFromState(int32 MachineIndex, int32 StateIndex)
 {
-	return GetProxyOnGameThread().GetRelevantAssetPlayerFromState(MachineIndex, StateIndex);
+	return GetProxyOnGameThread<FAnimInstanceProxy>().GetRelevantAssetPlayerFromState(MachineIndex, StateIndex);
 }
 
 int32 UAnimInstance::GetSyncGroupIndexFromName(FName SyncGroupName) const
 {
-	return GetProxyOnGameThread().GetSyncGroupIndexFromName(SyncGroupName);
+	return GetProxyOnGameThread<FAnimInstanceProxy>().GetSyncGroupIndexFromName(SyncGroupName);
 }
 
 bool UAnimInstance::IsRunningParallelEvaluation() const
@@ -2540,23 +2540,23 @@ void UAnimInstance::DestroyAnimInstanceProxy(FAnimInstanceProxy* InProxy)
 
 void UAnimInstance::UpdateMorphTargetCurves(const TMap<FName, float>& InMorphTargetCurves)
 {
-	GetProxyOnGameThread().UpdateMorphTargetCurves(InMorphTargetCurves);
+	GetProxyOnGameThread<FAnimInstanceProxy>().UpdateMorphTargetCurves(InMorphTargetCurves);
 }
 
 void UAnimInstance::UpdateComponentsMaterialParameters(UPrimitiveComponent* Component)
 {
-	GetProxyOnGameThread().UpdateComponentsMaterialParameters(Component);
+	GetProxyOnGameThread<FAnimInstanceProxy>().UpdateComponentsMaterialParameters(Component);
 }
 
 TArray<FActiveVertexAnim> UAnimInstance::UpdateActiveVertexAnims(const USkeletalMesh* SkeletalMesh)
 {
-	FAnimInstanceProxy& Proxy = GetProxyOnAnyThread();	// this can be called from CreateRenderState_Concurrent
+	FAnimInstanceProxy& Proxy = GetProxyOnAnyThread<FAnimInstanceProxy>();	// this can be called from CreateRenderState_Concurrent
 	return FAnimationRuntime::UpdateActiveVertexAnims(SkeletalMesh, Proxy.GetMorphTargetCurves(), Proxy.GetVertexAnims());
 }
 
 FBoneContainer& UAnimInstance::GetRequiredBones()
 {
-	return GetProxyOnGameThread().GetRequiredBones();
+	return GetProxyOnGameThread<FAnimInstanceProxy>().GetRequiredBones();
 }
 
 #undef LOCTEXT_NAMESPACE 
diff --git a/Engine/Source/Runtime/Engine/Private/Animation/AnimNodeBase.cpp b/Engine/Source/Runtime/Engine/Private/Animation/AnimNodeBase.cpp
--- a/Engine/Source/Runtime/Engine/Private/Animation/AnimNodeBase.cpp
+++ b/Engine/Source/Runtime/Engine/Private/Animation/AnimNodeBase.cpp
@@ -9,7 +9,7 @@
 
 // @todo: remove after deprecation
 FAnimationBaseContext::FAnimationBaseContext(UAnimInstance* InAnimInstance)
-	: AnimInstanceProxy(&InAnimInstance->GetProxyOnAnyThread())
+	: AnimInstanceProxy(&InAnimInstance->GetProxyOnAnyThread<FAnimInstanceProxy>())
 {
 	PRAGMA_DISABLE_DEPRECATION_WARNINGS
 	AnimInstance = InAnimInstance;
