Fixes crash when rendering stats with no cycle counter (for example stat RHI is 100% crash)

diff --git a/Engine/Source/Runtime/Engine/Private/StatsRender2.cpp b/Engine/Source/Runtime/Engine/Private/StatsRender2.cpp
--- a/Engine/Source/Runtime/Engine/Private/StatsRender2.cpp
+++ b/Engine/Source/Runtime/Engine/Private/StatsRender2.cpp
@@ -733,7 +733,7 @@ static void RenderGroupedWithHierarchy(const FGameThreadHudData& ViewData, FView
 			
 			FName ThreadName;
 			FName ShortThreadName;
-			if(NumThreadsBreakdown)
+			if(bBudget)
 			{
 				ThreadName = ThreadNames[ThreadBreakdownIdx];
 				ShortThreadName = FStatNameAndInfo::GetShortNameFrom(ThreadName);
@@ -773,24 +773,24 @@ static void RenderGroupedWithHierarchy(const FGameThreadHudData& ViewData, FView
 				RenderArrayOfStats(Canvas, X, Y, bBudget ? HudGroup.FlatAggregateThreadBreakdown[ThreadName] : HudGroup.FlatAggregate, ViewData, HudGroup.BudgetIgnoreStats, Budget, RenderFlatCycle);
 				Y += Globals.GetFontHeight();
 			}
+		}
+		
 
-			// Render memory counters.
-			if (HudGroup.MemoryAggregate.Num())
-			{
-				Y += RenderMemoryHeadings(Canvas, X, Y);
-				RenderArrayOfStats(Canvas, X, Y, HudGroup.MemoryAggregate, ViewData, HudGroup.BudgetIgnoreStats, Budget, RenderMemoryCounter);
-				Y += Globals.GetFontHeight();
-			}
+		// Render memory counters.
+		if (HudGroup.MemoryAggregate.Num())
+		{
+			Y += RenderMemoryHeadings(Canvas, X, Y);
+			RenderArrayOfStats(Canvas, X, Y, HudGroup.MemoryAggregate, ViewData, HudGroup.BudgetIgnoreStats, -1.f, RenderMemoryCounter);
+			Y += Globals.GetFontHeight();
+		}
 
-			// Render remaining counters.
-			if (HudGroup.CountersAggregate.Num())
-			{
-				Y += RenderCounterHeadings(Canvas, X, Y);
-				RenderArrayOfStats(Canvas, X, Y, HudGroup.CountersAggregate, ViewData, HudGroup.BudgetIgnoreStats, Budget, RenderCounter);
-				Y += Globals.GetFontHeight();
-			}
+		// Render remaining counters.
+		if (HudGroup.CountersAggregate.Num())
+		{
+			Y += RenderCounterHeadings(Canvas, X, Y);
+			RenderArrayOfStats(Canvas, X, Y, HudGroup.CountersAggregate, ViewData, HudGroup.BudgetIgnoreStats, -1.f, RenderCounter);
+			Y += Globals.GetFontHeight();
 		}
-		
 	}
 }
 
