UE-23360 Merging fix from CL 2770604 up from orion stream

diff --git a/Engine/Source/Runtime/Engine/Private/SoundWave.cpp b/Engine/Source/Runtime/Engine/Private/SoundWave.cpp
--- a/Engine/Source/Runtime/Engine/Private/SoundWave.cpp
+++ b/Engine/Source/Runtime/Engine/Private/SoundWave.cpp
@@ -505,13 +505,10 @@ void USoundWave::FreeResources()
 	DEC_DWORD_STAT_BY(STAT_AudioMemory, TrackedMemoryUsage);
 	TrackedMemoryUsage = 0;
 
-	SampleRate = 0;
-	Duration = 0.0f;
 	ResourceID = 0;
 	bDynamicResource = false;
 	DecompressionType = DTYPE_Setup;
 	bDecompressedFromOgg = 0;
-	RawPCMDataSize = 0;
 }
 
 FWaveInstance* USoundWave::HandleStart( FActiveSound& ActiveSound, const UPTRINT WaveInstanceHash )
