Copying //UE4/Orion-Staging to Dev-Main (Originating from //Orion/Main at CL-2777663)

diff --git a/Engine/Source/Runtime/Engine/Private/Particles/ParticleComponents.cpp b/Engine/Source/Runtime/Engine/Private/Particles/ParticleComponents.cpp
--- a/Engine/Source/Runtime/Engine/Private/Particles/ParticleComponents.cpp
+++ b/Engine/Source/Runtime/Engine/Private/Particles/ParticleComponents.cpp
@@ -1697,8 +1697,10 @@ void UParticleSpriteEmitter::SetToSensibleDefaults()
 /*-----------------------------------------------------------------------------
 	UParticleSystem implementation.
 -----------------------------------------------------------------------------*/
+
 UParticleSystem::UParticleSystem(const FObjectInitializer& ObjectInitializer)
 	: Super(ObjectInitializer)
+	, bAnyEmitterLoopsForever(false)
 {
 #if WITH_EDITORONLY_DATA
 	ThumbnailDistance = 200.0;
@@ -1810,6 +1812,29 @@ void UParticleSystem::PostEditChangeProperty(FPropertyChangedEvent& PropertyChan
 		}
 	}
 
+	// Recompute the looping flag
+	bAnyEmitterLoopsForever = false;
+	for (UParticleEmitter* Emitter : Emitters)
+	{
+		if (Emitter != nullptr)
+		{
+			for (const UParticleLODLevel* LODLevel : Emitter->LODLevels)
+			{
+				if (LODLevel != nullptr)
+				{
+					if (LODLevel->bEnabled)
+					{
+						const UParticleModuleRequired* RequiredModule = LODLevel->RequiredModule;
+						if ((RequiredModule != nullptr) && (RequiredModule->EmitterLoops == 0))
+						{
+							bAnyEmitterLoopsForever = true;
+						}
+					}
+				}
+			}
+		}
+	}
+
 	//cap the WarmupTickRate to realistic values
 	if (WarmupTickRate <= 0)
 	{
@@ -1854,10 +1879,12 @@ void UParticleSystem::PostLoad()
 {
 	Super::PostLoad();
 
-	// Remove any old emitters
+	// Run thru all of the emitters, load them up and compute some flags based on them
 	bHasPhysics = false;
+	bAnyEmitterLoopsForever = false;
 	for (int32 i = Emitters.Num() - 1; i >= 0; i--)
 	{
+		// Remove any old emitters
 		UParticleEmitter* Emitter = Emitters[i];
 		if (Emitter == NULL)
 		{
@@ -1871,38 +1898,47 @@ void UParticleSystem::PostLoad()
 
 		Emitter->ConditionalPostLoad();
 
-		if (Emitter->IsA(UParticleSpriteEmitter::StaticClass()))
+		bool bCookedOut = false;
+		if (UParticleSpriteEmitter* SpriteEmitter = Cast<UParticleSpriteEmitter>(Emitter))
 		{
-			UParticleSpriteEmitter* SpriteEmitter = Cast<UParticleSpriteEmitter>(Emitter);
+			bCookedOut = SpriteEmitter->bCookedOut;
+		}
 
-			if (SpriteEmitter->bCookedOut == false)
-			{
-				UParticleLODLevel* LODLevel = SpriteEmitter->LODLevels[0];
-				check(LODLevel);
+		if (!bCookedOut)
+		{
+			UParticleLODLevel* LODLevel = Emitter->LODLevels[0];
+			check(LODLevel);
 
-				LODLevel->ConditionalPostLoad();
+			LODLevel->ConditionalPostLoad();
 				
-				//@todo. Move this into the editor and serialize?
-				for (int32 LODIndex = 0; (LODIndex < Emitter->LODLevels.Num()) && (bHasPhysics == false); LODIndex++)
+			//@todo. Move these flag calculations into the editor and serialize?
+			for (UParticleLODLevel* ParticleLODLevel : Emitter->LODLevels)
+			{
+				if (ParticleLODLevel)
 				{
 					//@todo. This is a temporary fix for emitters that apply physics.
 					// Check for collision modules with bApplyPhysics set to true
-					UParticleLODLevel*  ParticleLODLevel = Emitter->LODLevels[LODIndex];
-					if (ParticleLODLevel)
+					for (int32 ModuleIndex = 0; ModuleIndex < LODLevel->Modules.Num(); ModuleIndex++)
 					{
-						for (int32 ModuleIndex = 0; ModuleIndex < LODLevel->Modules.Num(); ModuleIndex++)
+						UParticleModuleCollision* CollisionModule = Cast<UParticleModuleCollision>(ParticleLODLevel->Modules[ModuleIndex]);
+						if (CollisionModule)
 						{
-							UParticleModuleCollision* CollisionModule = Cast<UParticleModuleCollision>(ParticleLODLevel->Modules[ModuleIndex]);
-							if (CollisionModule)
+							if (CollisionModule->bApplyPhysics == true)
 							{
-								if (CollisionModule->bApplyPhysics == true)
-								{
-									bHasPhysics = true;
-									break;
-								}
+								bHasPhysics = true;
+								break;
 							}
 						}
 					}
+
+					if (LODLevel->bEnabled)
+					{
+						const UParticleModuleRequired* RequiredModule = LODLevel->RequiredModule;
+						if ((RequiredModule != nullptr) && (RequiredModule->EmitterLoops == 0))
+						{
+							bAnyEmitterLoopsForever = true;
+						}
+					}
 				}
 			}
 		}
@@ -2090,9 +2126,8 @@ void UParticleSystem::GetAssetRegistryTags(TArray<FAssetRegistryTag>& OutTags) c
 	OutTags.Add(FAssetRegistryTag("LODMethod", LODMethodString, FAssetRegistryTag::TT_Alphabetical));
 
 	// Run thru the emitters and see if any will loop forever
-	bool bAnyEmitterLoopsForever = false;
 	bool bEmitterIsDangerous = false;
-	for (int32 EmitterIndex = 0; (EmitterIndex < Emitters.Num()) && !bAnyEmitterLoopsForever; ++EmitterIndex)
+	for (int32 EmitterIndex = 0; EmitterIndex < Emitters.Num(); ++EmitterIndex)
 	{
 		if (const UParticleEmitter* Emitter = Emitters[EmitterIndex])
 		{
@@ -2103,13 +2138,7 @@ void UParticleSystem::GetAssetRegistryTags(TArray<FAssetRegistryTag>& OutTags) c
 					const UParticleModuleRequired* RequiredModule = LODLevel->RequiredModule;
 					check(RequiredModule);
 
-					if (RequiredModule->EmitterLoops == 0)
-					{
-						bAnyEmitterLoopsForever = true;
-						break;
-					}
-
-					UParticleModuleSpawn *SpawnModule = LODLevel->SpawnModule;
+					UParticleModuleSpawn* SpawnModule = LODLevel->SpawnModule;
 					check(SpawnModule);
 
 					if (SpawnModule->GetMaximumSpawnRate() == 0 
@@ -2119,7 +2148,6 @@ void UParticleSystem::GetAssetRegistryTags(TArray<FAssetRegistryTag>& OutTags) c
 					{
 						bEmitterIsDangerous = true;
 					}
-
 				}
 			}
 		}
@@ -2630,8 +2658,7 @@ void UParticleSystem::BuildEmitters()
 	const int32 EmitterCount = Emitters.Num();
 	for ( int32 EmitterIndex = 0; EmitterIndex < EmitterCount; ++EmitterIndex )
 	{
-		UParticleEmitter* Emitter = Emitters[EmitterIndex];
-		if ( Emitter )
+		if (UParticleEmitter* Emitter = Emitters[EmitterIndex])
 		{
 			Emitter->Build();
 		}
@@ -6207,7 +6234,14 @@ void AEmitterCameraLensEffectBase::RegisterCamera(APlayerCameraManager* C)
 	BaseCamera = C;
 }
 
-void AEmitterCameraLensEffectBase::NotifyRetriggered() {}
+void AEmitterCameraLensEffectBase::NotifyRetriggered() 
+{
+	UParticleSystemComponent* const PSC = GetParticleSystemComponent();
+	if (PSC && PSC->bWasDeactivated)
+	{
+		PSC->Activate(false);
+	}
+}
 
 void AEmitterCameraLensEffectBase::PostInitializeComponents()
 {
@@ -6258,5 +6292,29 @@ void AEmitterCameraLensEffectBase::ActivateLensEffect()
 	}
 }
 
+void AEmitterCameraLensEffectBase::DeactivateLensEffect()
+{
+	UParticleSystemComponent* const PSC = GetParticleSystemComponent();
+	if (PSC)
+	{
+		PSC->DeactivateSystem();
+	}
+}
+
+bool AEmitterCameraLensEffectBase::IsLooping() const
+{
+	if ((PS_CameraEffect != nullptr) && PS_CameraEffect->IsLooping())
+	{
+		return true;
+	}
+
+	if ((PS_CameraEffectNonExtremeContent != nullptr) && PS_CameraEffectNonExtremeContent->IsLooping())
+	{
+		return true;
+	}
+
+	return false;
+}
+
 #undef LOCTEXT_NAMESPACE
 
diff --git a/Engine/Source/Runtime/Engine/Private/Particles/ParticleEmitterInstances.cpp b/Engine/Source/Runtime/Engine/Private/Particles/ParticleEmitterInstances.cpp
--- a/Engine/Source/Runtime/Engine/Private/Particles/ParticleEmitterInstances.cpp
+++ b/Engine/Source/Runtime/Engine/Private/Particles/ParticleEmitterInstances.cpp
@@ -28,6 +28,8 @@
 #include "Particles/ParticleSpriteEmitter.h"
 #include "Particles/ParticleSystemComponent.h"
 
+#include "Components/PointLightComponent.h"
+
 /*-----------------------------------------------------------------------------
 FParticlesStatGroup
 -----------------------------------------------------------------------------*/
@@ -347,6 +349,16 @@ FParticleEmitterInstance::FParticleEmitterInstance() :
 /** Destructor	*/
 FParticleEmitterInstance::~FParticleEmitterInstance()
 {
+	for (int32 i = 0; i < HighQualityLights.Num(); ++i)
+	{
+		UPointLightComponent* PointLightComponent = HighQualityLights[i];
+		{
+			PointLightComponent->Modify();
+			PointLightComponent->DestroyComponent(false);
+		}
+	}
+	HighQualityLights.Reset();
+
 	FMemory::Free(ParticleData);
 	FMemory::Free(ParticleIndices);
 	FMemory::Free(InstanceData);
@@ -2733,6 +2745,24 @@ bool FParticleEmitterInstance::Tick_MaterialOverrides()
 	return bOverridden;
 }
 
+bool FParticleEmitterInstance::UseLocalSpace()
+{
+	const UParticleLODLevel* LODLevel = GetCurrentLODLevelChecked();
+	return LODLevel->RequiredModule->bUseLocalSpace;
+}
+
+void FParticleEmitterInstance::GetScreenAlignmentAndScale(int32& OutScreenAlign, FVector& OutScale)
+{
+	const UParticleLODLevel* LODLevel = GetCurrentLODLevelChecked();
+	OutScreenAlign = (int32)LODLevel->RequiredModule->ScreenAlignment;
+
+	OutScale = FVector(1.0f, 1.0f, 1.0f);
+	if (Component)
+	{
+		OutScale = Component->ComponentToWorld.GetScale3D();
+	}
+}
+
 /*-----------------------------------------------------------------------------
 	ParticleSpriteEmitterInstance
 -----------------------------------------------------------------------------*/
diff --git a/Engine/Source/Runtime/Engine/Private/Particles/ParticleModules.cpp b/Engine/Source/Runtime/Engine/Private/Particles/ParticleModules.cpp
--- a/Engine/Source/Runtime/Engine/Private/Particles/ParticleModules.cpp
+++ b/Engine/Source/Runtime/Engine/Private/Particles/ParticleModules.cpp
@@ -76,6 +76,7 @@
 #include "Distributions/DistributionFloatUniformCurve.h"
 #include "Engine/InterpCurveEdSetup.h"
 #include "Distributions/DistributionFloatConstantCurve.h"
+#include "Components/PointLightComponent.h"
 
 /*-----------------------------------------------------------------------------
 	Abstract base modules used for categorization.
@@ -2769,6 +2770,11 @@ void UParticleModuleLight::InitializeDefaults()
 	{
 		LightExponent.Distribution = NewObject<UDistributionFloatConstant>(this, TEXT("DistributionLightExponent"));
 	}
+
+	if (!bHighQualityLights)
+	{
+		bShadowCastingLights = false;
+	}
 }
 
 void UParticleModuleLight::PostInitProperties()
@@ -2791,7 +2797,7 @@ void UParticleModuleLight::PostEditChangeProperty(FPropertyChangedEvent& Propert
 
 bool UParticleModuleLight::CanTickInAnyThread()
 {
-	return BrightnessOverLife.OkForParallel() && ColorScaleOverLife.OkForParallel() && RadiusScale.OkForParallel() && LightExponent.OkForParallel();
+	return !bHighQualityLights && BrightnessOverLife.OkForParallel() && ColorScaleOverLife.OkForParallel() && RadiusScale.OkForParallel() && LightExponent.OkForParallel();
 }
 
 void UParticleModuleLight::SpawnEx(FParticleEmitterInstance* Owner, int32 Offset, float SpawnTime, struct FRandomStream* InRandomStream, FBaseParticle* ParticleBase)
@@ -2806,6 +2812,112 @@ void UParticleModuleLight::SpawnEx(FParticleEmitterInstance* Owner, int32 Offset
 	const float RandomNumber = InRandomStream ? InRandomStream->GetFraction() : FMath::SRand();
 	LightData.bValid = RandomNumber < SpawnFraction;
 	LightData.bAffectsTranslucency = bAffectsTranslucency;
+	LightData.bHighQuality = bHighQualityLights;
+	LightData.LightId = 0;
+
+	if (bHighQualityLights)
+	{		
+		LightData.LightId = SpawnHQLight(LightData, Particle, Owner);
+	}
+}
+
+uint64 UParticleModuleLight::SpawnHQLight(const FLightParticlePayload& Payload, const FBaseParticle& Particle, FParticleEmitterInstance* Owner)
+{
+	uint64 LightId = 0;
+	if (!Owner)
+	{
+		return 0;
+	}
+	UParticleSystemComponent* ParticleSystem = Owner->Component;
+	if (!ParticleSystem)
+	{
+		return 0;
+	}
+	AActor* HQLightContainer = ParticleSystem->GetOwner();
+	if (!HQLightContainer || HQLightContainer->IsPendingKillPending())
+	{
+		return 0;
+	}
+
+	// Construct the new component and attach as needed				
+	UPointLightComponent* PointLightComponent = NewObject<UPointLightComponent>(HQLightContainer, NAME_None, RF_NoFlags);
+	if (PointLightComponent)
+	{
+		LightId = (uint64)PointLightComponent;
+						
+		USceneComponent* RootComponent = HQLightContainer->GetRootComponent();
+		if (RootComponent)
+		{
+			PointLightComponent->AttachTo(RootComponent, NAME_None, EAttachLocation::KeepRelativeOffset);
+		}			
+		PointLightComponent->CreationMethod = EComponentCreationMethod::UserConstructionScript;
+		PointLightComponent->RegisterComponent();
+		Owner->HighQualityLights.Add(PointLightComponent);
+			
+			
+		PointLightComponent->bUseInverseSquaredFalloff = bUseInverseSquaredFalloff;
+		PointLightComponent->bAffectTranslucentLighting = bAffectsTranslucency;
+		PointLightComponent->SetCastShadows(bShadowCastingLights);
+
+		int32 ScreenAlignment;
+		FVector ComponentScale;
+		Owner->GetScreenAlignmentAndScale(ScreenAlignment, ComponentScale);
+		UpdateHQLight(PointLightComponent, Payload, Particle, ScreenAlignment, ComponentScale, Owner->UseLocalSpace(), nullptr, false);
+	}	
+	return LightId;
+}
+
+void UParticleModuleLight::UpdateHQLight(UPointLightComponent* PointLightComponent, const FLightParticlePayload& Payload, const FBaseParticle& Particle, int32 ScreenAlignment, FVector ComponentScale, bool bLocalSpace, FSceneInterface* OwnerScene, bool bDoRTUpdate)
+{
+	if (bLocalSpace)
+	{
+		PointLightComponent->SetRelativeLocation(Particle.Location);
+	}
+	else
+	{
+		PointLightComponent->SetWorldLocation(Particle.Location);
+	}
+	
+	FLinearColor DesiredFinalColor = FVector(Particle.Color) * Particle.Color.A * Payload.ColorScale;
+	if (bUseInverseSquaredFalloff)
+	{
+		//later in light rendering HQ lights are multiplied by 16 in inverse falloff mode to adjust for lumens.  
+		//We want our particle lights to match simple lights as much as possible when toggling so remove that here.
+		const float fLumenAdjust = 1.0f / 16.0f;
+		DesiredFinalColor *= fLumenAdjust;
+	}
+
+	//light color on HQ lights is just a uint32 and our light scalars can be huge.  To preserve the color control and range from the particles we need to normalize
+	//around the full range multiplied value, and set the scalar intensity such that it will bring things back into line later.
+	FVector AdjustedColor(DesiredFinalColor.R, DesiredFinalColor.G, DesiredFinalColor.B);
+	float Intensity = AdjustedColor.Size();
+	AdjustedColor.Normalize();	
+	
+	//light module currently needs to run AFTER any size modification modules to get a value that matches 'simple' lights.
+	FVector2D Size;
+	Size.X = FMath::Abs(Particle.Size.X * ComponentScale.X);
+	Size.Y = FMath::Abs(Particle.Size.Y * ComponentScale.Y);
+	if (ScreenAlignment == PSA_Square || ScreenAlignment == PSA_FacingCameraPosition)
+	{
+		Size.Y = Size.X;
+	}
+	float Radius = Payload.RadiusScale * (Size.X + Size.X) * 0.5f;
+
+	//HQ light color is an FColor which is a uint32.  Thus we have to break out the out of range value into the intensity scalar.
+	FColor NormalizedColor;
+	NormalizedColor.R = (uint8)FMath::Min(AdjustedColor.X * 255.0f, 255.0f);
+	NormalizedColor.G = (uint8)FMath::Min(AdjustedColor.Y * 255.0f, 255.0f);
+	NormalizedColor.B = (uint8)FMath::Min(AdjustedColor.Z * 255.0f, 255.0f);
+	PointLightComponent->SetIntensity(Intensity);
+	PointLightComponent->SetLightColor(NormalizedColor);
+	PointLightComponent->SetAttenuationRadius(Radius);
+	PointLightComponent->SetLightFalloffExponent(Payload.LightExponent);
+
+	if (OwnerScene && bDoRTUpdate)
+	{
+		OwnerScene->UpdateLightTransform(PointLightComponent);
+		OwnerScene->UpdateLightColorAndBrightness(PointLightComponent);
+	}
 }
 
 void UParticleModuleLight::Spawn(FParticleEmitterInstance* Owner, int32 Offset, float SpawnTime, FBaseParticle* ParticleBase)
@@ -2820,17 +2932,56 @@ void UParticleModuleLight::Update(FParticleEmitterInstance* Owner, int32 Offset,
 	{
 		return;
 	}
+
+	UWorld* OwnerWorld = Owner->GetWorld();
+	FSceneInterface* OwnerScene = nullptr;
+	if (OwnerWorld)
+	{
+		OwnerScene = OwnerWorld->Scene;
+	}
+	
+	TSet<uint64> ActiveLights;
 	UParticleLODLevel* LODLevel	= Owner->SpriteTemplate->GetCurrentLODLevel(Owner);
 	check(LODLevel);
 	FPlatformMisc::Prefetch(Owner->ParticleData, (Owner->ParticleIndices[0] * Owner->ParticleStride));
 	FPlatformMisc::Prefetch(Owner->ParticleData, (Owner->ParticleIndices[0] * Owner->ParticleStride) + CACHE_LINE_SIZE);
+	const bool bUseLocalSpace = Owner->UseLocalSpace();
+	int32 ScreenAlignment;
+	FVector ComponentScale;
+	Owner->GetScreenAlignmentAndScale(ScreenAlignment, ComponentScale);
+
 	BEGIN_UPDATE_LOOP;
 	{
 		PARTICLE_ELEMENT(FLightParticlePayload,	Data);
 		const float Brightness = BrightnessOverLife.GetValue(Particle.RelativeTime, Owner->Component);
 		Data.ColorScale = ColorScaleOverLife.GetValue(Particle.RelativeTime, Owner->Component) * Brightness;
+
+		if (bHighQualityLights && (Data.LightId != 0))
+		{
+			ActiveLights.Add(Data.LightId);
+
+			//for now we can do this
+			UPointLightComponent* PointLightComponent = (UPointLightComponent*)Data.LightId;
+			UpdateHQLight(PointLightComponent, Data, Particle, ScreenAlignment, ComponentScale, bUseLocalSpace, OwnerScene, true);
+		}
 	}
 	END_UPDATE_LOOP;
+
+	//remove any dead lights.
+	if (bHighQualityLights)
+	{
+		for (int32 i = 0; i < Owner->HighQualityLights.Num(); ++i)
+		{
+			UPointLightComponent* PointLightComponent = Owner->HighQualityLights[i];
+			if (PointLightComponent && !ActiveLights.Contains((uint64)PointLightComponent))
+			{
+				PointLightComponent->Modify();
+				PointLightComponent->DestroyComponent(false);
+				Owner->HighQualityLights.RemoveAtSwap(i);
+				--i;
+			}
+		}
+	}
 }
 
 uint32 UParticleModuleLight::RequiredBytes(FParticleEmitterInstance* Owner)
diff --git a/Engine/Source/Runtime/Engine/Private/Particles/ParticleSystemRender.cpp b/Engine/Source/Runtime/Engine/Private/Particles/ParticleSystemRender.cpp
--- a/Engine/Source/Runtime/Engine/Private/Particles/ParticleSystemRender.cpp
+++ b/Engine/Source/Runtime/Engine/Private/Particles/ParticleSystemRender.cpp
@@ -842,8 +842,8 @@ void GatherParticleLightData(const FDynamicSpriteEmitterReplayDataBase& Source,
 				ParticleLight.bAffectTranslucency = LightPayload->bAffectsTranslucency;
 
 				// Early out if the light will have no visible contribution
-				if (ParticleLight.Radius <= KINDA_SMALL_NUMBER
-					&& ParticleLight.Color.GetMax() <= KINDA_SMALL_NUMBER)
+				if (LightPayload->bHighQuality || 
+					(ParticleLight.Radius <= KINDA_SMALL_NUMBER && ParticleLight.Color.GetMax() <= KINDA_SMALL_NUMBER))
 				{
 					continue;
 				}
diff --git a/Engine/Source/Runtime/Engine/Private/PlayerController.cpp b/Engine/Source/Runtime/Engine/Private/PlayerController.cpp
--- a/Engine/Source/Runtime/Engine/Private/PlayerController.cpp
+++ b/Engine/Source/Runtime/Engine/Private/PlayerController.cpp
@@ -2323,8 +2323,16 @@ void APlayerController::PostProcessInput(const float DeltaTime, const bool bGame
 
 void APlayerController::ResetIgnoreInputFlags()
 {
-	IgnoreMoveInput = GetDefault<APlayerController>()->IgnoreMoveInput;
-	IgnoreLookInput = GetDefault<APlayerController>()->IgnoreLookInput;
+	// The movement locks can be set in cinematic mode, but if a restart occurs, we don't want them to be reset.
+	if (!bCinemaDisableInputMove)
+	{
+		IgnoreMoveInput = GetDefault<APlayerController>()->IgnoreMoveInput;
+	}
+
+	if (!bCinemaDisableInputLook)
+	{
+		IgnoreLookInput = GetDefault<APlayerController>()->IgnoreLookInput;
+	}
 }
 
 void APlayerController::SetCinematicMode( bool bInCinematicMode, bool bAffectsMovement, bool bAffectsTurning)
diff --git a/Engine/Source/Runtime/Engine/Private/PrimitiveDrawingUtils.cpp b/Engine/Source/Runtime/Engine/Private/PrimitiveDrawingUtils.cpp
--- a/Engine/Source/Runtime/Engine/Private/PrimitiveDrawingUtils.cpp
+++ b/Engine/Source/Runtime/Engine/Private/PrimitiveDrawingUtils.cpp
@@ -11,25 +11,41 @@
 #include "Engine/LightMapTexture2D.h"
 
 /** Emits draw events for a given FMeshBatch and the FPrimitiveSceneProxy corresponding to that mesh element. */
-#if !(UE_BUILD_SHIPPING || UE_BUILD_TEST)
-void EmitMeshDrawEvents_Inner(FRHICommandList& RHICmdList, const FPrimitiveSceneProxy* PrimitiveSceneProxy, const FMeshBatch& Mesh)
+#if WANTS_DRAW_MESH_EVENTS
+
+void BeginMeshDrawEvent_Inner(FRHICommandList& RHICmdList, const FPrimitiveSceneProxy* PrimitiveSceneProxy, const FMeshBatch& Mesh, TDrawEvent<FRHICommandList>& MeshEvent)
 {
-	// Only show material name at the top level
-	// Note: this is the parent's material name, not the material instance
-	SCOPED_DRAW_EVENTF(RHICmdList, MaterialEvent, *Mesh.MaterialRenderProxy->GetMaterial(PrimitiveSceneProxy ? PrimitiveSceneProxy->GetScene().GetFeatureLevel() : GMaxRHIFeatureLevel)->GetFriendlyName());
+	// Only show material and resource name at the top level
 	if (PrimitiveSceneProxy)
 	{
-		// Show Actor, level and resource name inside the material name
-		// These are separate draw events since some platforms only allow 32 character event names (xenon)
-		{
-			SCOPED_CONDITIONAL_DRAW_EVENTF(RHICmdList, LevelEvent, PrimitiveSceneProxy->GetLevelName() != NAME_None, PrimitiveSceneProxy->GetLevelName().IsValid() ? *PrimitiveSceneProxy->GetLevelName().ToString() : TEXT(""));
-		}
-		{
-			SCOPED_CONDITIONAL_DRAW_EVENTF(RHICmdList, OwnerEvent,PrimitiveSceneProxy->GetOwnerName() != NAME_None, *PrimitiveSceneProxy->GetOwnerName().ToString());
-		}
-		SCOPED_CONDITIONAL_DRAW_EVENTF(RHICmdList, ResourceEvent,PrimitiveSceneProxy->GetResourceName() != NAME_None, PrimitiveSceneProxy->GetResourceName().IsValid() ? *PrimitiveSceneProxy->GetResourceName().ToString() : TEXT(""));
+		BEGIN_DRAW_EVENTF(
+			RHICmdList, 
+			MaterialEvent, 
+			MeshEvent, 
+			TEXT("%s %s"), 
+			// Note: this is the parent's material name, not the material instance
+			*Mesh.MaterialRenderProxy->GetMaterial(PrimitiveSceneProxy ? PrimitiveSceneProxy->GetScene().GetFeatureLevel() : GMaxRHIFeatureLevel)->GetFriendlyName(),
+			PrimitiveSceneProxy->GetResourceName().IsValid() ? *PrimitiveSceneProxy->GetResourceName().ToString() : TEXT(""));
+
+			// Show Actor, level and resource name inside the material name
+			// These are separate draw events since some platforms have a limit on draw event length
+			// Note: empty leaf events are culled from profilegpu by default so these won't show up
+			{
+				//SCOPED_CONDITIONAL_DRAW_EVENTF(RHICmdList, LevelEvent, PrimitiveSceneProxy->GetLevelName() != NAME_None, PrimitiveSceneProxy->GetLevelName().IsValid() ? *PrimitiveSceneProxy->GetLevelName().ToString() : TEXT(""));
+			}
+			//SCOPED_CONDITIONAL_DRAW_EVENTF(RHICmdList, OwnerEvent,PrimitiveSceneProxy->GetOwnerName() != NAME_None, *PrimitiveSceneProxy->GetOwnerName().ToString());
+	}
+	else
+	{
+		BEGIN_DRAW_EVENTF(
+			RHICmdList, 
+			MaterialEvent, 
+			MeshEvent, 
+			// Note: this is the parent's material name, not the material instance
+			*Mesh.MaterialRenderProxy->GetMaterial(GMaxRHIFeatureLevel)->GetFriendlyName());
 	}
 }
+
 #endif
 
 void DrawPlane10x10(class FPrimitiveDrawInterface* PDI,const FMatrix& ObjectToWorld, float Radii, FVector2D UVMin, FVector2D UVMax, const FMaterialRenderProxy* MaterialRenderProxy, uint8 DepthPriorityGroup)
diff --git a/Engine/Source/Runtime/Engine/Private/SceneUtils.cpp b/Engine/Source/Runtime/Engine/Private/SceneUtils.cpp
--- a/Engine/Source/Runtime/Engine/Private/SceneUtils.cpp
+++ b/Engine/Source/Runtime/Engine/Private/SceneUtils.cpp
@@ -4,7 +4,7 @@
 #include "RHI.h"
 #include "SceneUtils.h"
 
-#if WANTS_DRAW_MESH_EVENTS && PLATFORM_SUPPORTS_DRAW_MESH_EVENTS
+#if WANTS_DRAW_MESH_EVENTS
 
 template<typename TRHICmdList>
 void TDrawEvent<TRHICmdList>::Start(TRHICmdList& InRHICmdList, const TCHAR* Fmt, ...)
@@ -24,7 +24,11 @@ void TDrawEvent<TRHICmdList>::Start(TRHICmdList& InRHICmdList, const TCHAR* Fmt,
 template<typename TRHICmdList>
 void TDrawEvent<TRHICmdList>::Stop()
 {
-	RHICmdList->PopEvent();
+	if (RHICmdList)
+	{
+		RHICmdList->PopEvent();
+		RHICmdList = NULL;
+	}
 }
 template struct TDrawEvent<FRHICommandList>;
 template struct TDrawEvent<FRHIAsyncComputeCommandList>;
@@ -49,4 +53,4 @@ void FDrawEventRHIExecute::Stop()
 	RHICommandContext->RHIPopEvent();
 }
 
-#endif // WANTS_DRAW_MESH_EVENTS && PLATFORM_SUPPORTS_DRAW_MESH_EVENTS
\ No newline at end of file
+#endif // WANTS_DRAW_MESH_EVENTS
\ No newline at end of file
diff --git a/Engine/Source/Runtime/Engine/Private/SceneView.cpp b/Engine/Source/Runtime/Engine/Private/SceneView.cpp
--- a/Engine/Source/Runtime/Engine/Private/SceneView.cpp
+++ b/Engine/Source/Runtime/Engine/Private/SceneView.cpp
@@ -538,6 +538,7 @@ FSceneView::FSceneView(const FSceneViewInitOptions& InitOptions)
 	}
 
 	bUseFieldOfViewForLOD = InitOptions.bUseFieldOfViewForLOD;
+	DrawDynamicFlags = EDrawDynamicFlags::None;
 
 #if WITH_EDITOR
 	EditorViewBitflag = InitOptions.EditorViewBitflag;
diff --git a/Engine/Source/Runtime/Engine/Private/StaticMeshRender.cpp b/Engine/Source/Runtime/Engine/Private/StaticMeshRender.cpp
--- a/Engine/Source/Runtime/Engine/Private/StaticMeshRender.cpp
+++ b/Engine/Source/Runtime/Engine/Private/StaticMeshRender.cpp
@@ -1387,6 +1387,10 @@ FLODMask FStaticMeshSceneProxy::GetLODMask(const FSceneView* View) const
 	{
 		Result.SetLOD(FMath::Clamp<int32>(CVarForcedLODLevel, 0, RenderData->LODResources.Num() - 1));
 	}
+	else if (View->DrawDynamicFlags & EDrawDynamicFlags::ForceLowestLOD)
+	{
+		Result.SetLOD(RenderData->LODResources.Num() - 1);
+	}
 	else if (ForcedLodModel > 0)
 	{
 		Result.SetLOD(FMath::Clamp(ForcedLodModel, 1, RenderData->LODResources.Num()) - 1);
diff --git a/Engine/Source/Runtime/Engine/Private/TimerManager.cpp b/Engine/Source/Runtime/Engine/Private/TimerManager.cpp
--- a/Engine/Source/Runtime/Engine/Private/TimerManager.cpp
+++ b/Engine/Source/Runtime/Engine/Private/TimerManager.cpp
@@ -6,6 +6,9 @@
 
 #include "EnginePrivate.h"
 
+DECLARE_CYCLE_STAT(TEXT("SetTimer"), STAT_SetTimer, STATGROUP_Engine);
+DECLARE_CYCLE_STAT(TEXT("ClearTimer"), STAT_ClearTimer, STATGROUP_Engine);
+
 void FTimerHandle::MakeValid()
 {
 	static int LastAssignedHandle = -1;
@@ -20,6 +23,36 @@ void FTimerHandle::MakeValid()
 }
 
 
+
+FString FTimerUnifiedDelegate::ToString() const
+{
+	const UObject* Object = nullptr;
+	FName FunctionName = NAME_None;
+	bool bDynDelegate = false;
+
+	if (FuncDelegate.IsBound())
+	{
+		Object = FuncDelegate.GetUObject();
+		if (IDelegateInstance* Ptr = FuncDelegate.GetDelegateInstance())
+		{
+			FunctionName = Ptr->GetFunctionName();
+		}
+	}
+	else if (FuncDynDelegate.IsBound())
+	{
+		Object = FuncDynDelegate.GetUObject();
+		FunctionName = FuncDynDelegate.GetFunctionName();
+		bDynDelegate = true;
+	}
+	else
+	{
+		static FName NotBoundName(TEXT("NotBound!"));
+		FunctionName = NotBoundName;
+	}
+
+	return FString::Printf(TEXT("%s,%s,%s"), bDynDelegate ? TEXT("DELEGATE") : TEXT("DYN DELEGATE"), Object == nullptr ? TEXT("NO OBJ") : *Object->GetPathName(), *FunctionName.ToString());
+}
+
 // ---------------------------------
 // Private members
 // ---------------------------------
@@ -120,6 +153,8 @@ FTimerHandle FTimerManager::K2_FindDynamicTimerHandle(FTimerDynamicDelegate InDy
 
 void FTimerManager::InternalSetTimer(FTimerHandle& InOutHandle, FTimerUnifiedDelegate const& InDelegate, float InRate, bool InbLoop, float InFirstDelay)
 {
+	SCOPE_CYCLE_COUNTER(STAT_SetTimer);
+
 	// not currently threadsafe
 	check(IsInGameThread());
 
@@ -186,6 +221,8 @@ void FTimerManager::InternalSetTimerForNextTick(FTimerUnifiedDelegate const& InD
 
 void FTimerManager::InternalClearTimer(FTimerHandle const& InHandle)
 {
+	SCOPE_CYCLE_COUNTER(STAT_ClearTimer);
+
 	// not currently threadsafe
 	check(IsInGameThread());
 
@@ -495,3 +532,45 @@ TStatId FTimerManager::GetStatId() const
 {
 	RETURN_QUICK_DECLARE_CYCLE_STAT(FTimerManager, STATGROUP_Tickables);
 }
+
+void FTimerManager::ListTimers() const
+{
+	UE_LOG(LogEngine, Log, TEXT("------- %d Active Timers -------"), ActiveTimerHeap.Num());
+	for(const FTimerData& Data : ActiveTimerHeap)
+	{
+		FString TimerString = Data.TimerDelegate.ToString();
+		UE_LOG(LogEngine, Log, TEXT("%s"), *TimerString);
+	}
+
+	UE_LOG(LogEngine, Log, TEXT("------- %d Paused Timers -------"), PausedTimerList.Num());
+	for (const FTimerData& Data : PausedTimerList)
+	{
+		FString TimerString = Data.TimerDelegate.ToString();
+		UE_LOG(LogEngine, Log, TEXT("%s"), *TimerString);
+	}
+
+	UE_LOG(LogEngine, Log, TEXT("------- %d Pending Timers -------"), PendingTimerList.Num());
+	for (const FTimerData& Data : PendingTimerList)
+	{
+		FString TimerString = Data.TimerDelegate.ToString();
+		UE_LOG(LogEngine, Log, TEXT("%s"), *TimerString);
+	}
+
+	UE_LOG(LogEngine, Log, TEXT("------- %d Total Timers -------"), PendingTimerList.Num() + PausedTimerList.Num() + ActiveTimerHeap.Num());
+}
+
+// Handler for ListTimers console command
+static void OnListTimers(UWorld* World)
+{
+	if(World != nullptr)
+	{
+		World->GetTimerManager().ListTimers();
+	}
+}
+
+// Register ListTimers console command, needs a World context
+FAutoConsoleCommandWithWorld ListTimersConsoleCommand(
+	TEXT("ListTimers"),
+	TEXT(""),
+	FConsoleCommandWithWorldDelegate::CreateStatic(OnListTimers)
+	);
\ No newline at end of file
diff --git a/Engine/Source/Runtime/Engine/Private/UnrealNetwork.cpp b/Engine/Source/Runtime/Engine/Private/UnrealNetwork.cpp
--- a/Engine/Source/Runtime/Engine/Private/UnrealNetwork.cpp
+++ b/Engine/Source/Runtime/Engine/Private/UnrealNetwork.cpp
@@ -8,6 +8,8 @@
 FNetworkVersion::FGetLocalNetworkVersionOverride FNetworkVersion::GetLocalNetworkVersionOverride;
 FNetworkVersion::FIsNetworkCompatibleOverride FNetworkVersion::IsNetworkCompatibleOverride;
 
+FPreReplayScrub FNetworkReplayDelegates::OnPreScrub;
+
 enum ENetworkVersionHistory
 {
 	HISTORY_INITIAL = 1,
diff --git a/Engine/Source/Runtime/Engine/Private/WorldSettings.cpp b/Engine/Source/Runtime/Engine/Private/WorldSettings.cpp
--- a/Engine/Source/Runtime/Engine/Private/WorldSettings.cpp
+++ b/Engine/Source/Runtime/Engine/Private/WorldSettings.cpp
@@ -144,6 +144,7 @@ void AWorldSettings::NotifyBeginPlay()
 			// Actors that have traveled seamlessly from other levels already had BeginPlay called in that level
 			if (!It->IsPendingKill() && !It->HasActorBegunPlay())
 			{
+				SCOPE_CYCLE_COUNTER(STAT_ActorBeginPlay);
 				It->BeginPlay();
 			}
 		}
diff --git a/Engine/Source/Runtime/Engine/Public/EngineStats.h b/Engine/Source/Runtime/Engine/Public/EngineStats.h
--- a/Engine/Source/Runtime/Engine/Public/EngineStats.h
+++ b/Engine/Source/Runtime/Engine/Public/EngineStats.h
@@ -40,7 +40,8 @@ DECLARE_FLOAT_COUNTER_STAT_EXTERN(TEXT("StatUnit GPUTime"), STAT_UnitGPU, STATGR
  * Game stats
  */
 DECLARE_CYCLE_STAT_EXTERN(TEXT("Async Physics Time"),STAT_PhysicsTime,STATGROUP_Game, );
-DECLARE_CYCLE_STAT_EXTERN(TEXT("Spawn Actor Time"),STAT_SpawnActorTime,STATGROUP_Game, );
+DECLARE_CYCLE_STAT_EXTERN(TEXT("SpawnActor"), STAT_SpawnActorTime, STATGROUP_Game, );
+DECLARE_CYCLE_STAT_EXTERN(TEXT("Actor BeginPlay"), STAT_ActorBeginPlay, STATGROUP_Game, );
 DECLARE_CYCLE_STAT_EXTERN(TEXT("MoveComponent(Primitive) Time"),STAT_MoveComponentTime,STATGROUP_Game, );
 DECLARE_CYCLE_STAT_EXTERN(TEXT("MoveComponent(SceneComp) Time"), STAT_MoveComponentSceneComponentTime, STATGROUP_Game, );
 DECLARE_CYCLE_STAT_EXTERN(TEXT("UpdateOverlaps Time"),STAT_UpdateOverlaps,STATGROUP_Game, );
diff --git a/Engine/Source/Runtime/Engine/Public/Net/UnrealNetwork.h b/Engine/Source/Runtime/Engine/Public/Net/UnrealNetwork.h
--- a/Engine/Source/Runtime/Engine/Public/Net/UnrealNetwork.h
+++ b/Engine/Source/Runtime/Engine/Public/Net/UnrealNetwork.h
@@ -82,6 +82,16 @@ struct ENGINE_API FNetworkVersion
 	static const uint32 InternalProtocolVersion;
 };
 
+DECLARE_MULTICAST_DELEGATE_OneParam(FPreActorDestroyReplayScrub, AActor*);
+
+DECLARE_MULTICAST_DELEGATE_OneParam(FPreReplayScrub, UWorld*);
+
+struct ENGINE_API FNetworkReplayDelegates
+{
+	/** global delegate called one time prior to scrubbing */
+	static FPreReplayScrub OnPreScrub;
+};
+
 /*-----------------------------------------------------------------------------
 	Replication.
 -----------------------------------------------------------------------------*/
diff --git a/Engine/Source/Runtime/Engine/Public/ParticleEmitterInstances.h b/Engine/Source/Runtime/Engine/Public/ParticleEmitterInstances.h
--- a/Engine/Source/Runtime/Engine/Public/ParticleEmitterInstances.h
+++ b/Engine/Source/Runtime/Engine/Public/ParticleEmitterInstances.h
@@ -305,6 +305,8 @@ public:
 	/** The PivotOffset applied to the vertex positions 			*/
 	FVector2D PivotOffset;
 
+	TArray<class UPointLightComponent*> HighQualityLights;
+
 	/** Constructor	*/
 	FParticleEmitterInstance();
 
@@ -670,6 +672,17 @@ public:
 	@return True if there were material overrides. Otherwise revert to default behaviour.
 	*/
 	virtual bool Tick_MaterialOverrides();
+
+	/**
+	* True if this emitter emits in local space
+	*/
+	bool UseLocalSpace();
+
+	/**
+	* returns the screen alignment and scale of the component.
+	*/
+	void GetScreenAlignmentAndScale(int32& OutScreenAlign, FVector& OutScale);
+
 protected:
 
 	/**
@@ -687,8 +700,8 @@ protected:
 	void UpdateTransforms();
 
 	/**
-	 * Retrieves the current LOD level and asserts that it is valid.
-	 */
+	* Retrieves the current LOD level and asserts that it is valid.
+	*/
 	class UParticleLODLevel* GetCurrentLODLevelChecked();
 
 	/**
diff --git a/Engine/Source/Runtime/Engine/Public/ParticleHelper.h b/Engine/Source/Runtime/Engine/Public/ParticleHelper.h
--- a/Engine/Source/Runtime/Engine/Public/ParticleHelper.h
+++ b/Engine/Source/Runtime/Engine/Public/ParticleHelper.h
@@ -524,10 +524,12 @@ struct FAttractorParticlePayload
 struct FLightParticlePayload
 {
 	FVector		ColorScale;
+	uint64		LightId;
 	float		RadiusScale;
 	float		LightExponent;
 	bool		bValid;
 	bool		bAffectsTranslucency;
+	bool		bHighQuality;
 };
 
 //
diff --git a/Engine/Source/Runtime/Engine/Public/PrimitiveSceneProxy.h b/Engine/Source/Runtime/Engine/Public/PrimitiveSceneProxy.h
--- a/Engine/Source/Runtime/Engine/Public/PrimitiveSceneProxy.h
+++ b/Engine/Source/Runtime/Engine/Public/PrimitiveSceneProxy.h
@@ -104,14 +104,6 @@ public:
 	{}
 };
 
-namespace EDrawDynamicFlags
-{
-	enum Type
-	{
-		ForceLowestLOD = 0x1
-	};
-}
-
 /**
  * Encapsulates the data which is mirrored to render a UPrimitiveComponent parallel to the game thread.
  * This is intended to be subclassed to support different primitive types.  
diff --git a/Engine/Source/Runtime/Engine/Public/SceneManagement.h b/Engine/Source/Runtime/Engine/Public/SceneManagement.h
--- a/Engine/Source/Runtime/Engine/Public/SceneManagement.h
+++ b/Engine/Source/Runtime/Engine/Public/SceneManagement.h
@@ -17,6 +17,7 @@
 #include "BatchedElements.h"
 #include "MeshBatch.h"
 #include "RendererInterface.h"
+#include "SceneUtils.h"
 
 // Forward declarations.
 class FLightSceneInfo;
@@ -2037,22 +2038,21 @@ extern ENGINE_API EVertexColorViewMode::Type GVertexColorViewMode;
  */
 extern ENGINE_API bool IsRichView(const FSceneViewFamily& ViewFamily);
 
-#if !(UE_BUILD_SHIPPING || UE_BUILD_TEST)
+#if WANTS_DRAW_MESH_EVENTS
 	/**
 	 * true if we debug material names with SCOPED_DRAW_EVENT.
 	 * Toggle with "ShowMaterialDrawEvents" console command.
 	 */
 	extern ENGINE_API bool GShowMaterialDrawEvents;
-	extern ENGINE_API void EmitMeshDrawEvents_Inner(FRHICommandList& RHICmdList, const class FPrimitiveSceneProxy* PrimitiveSceneProxy, const struct FMeshBatch& Mesh);
+	extern ENGINE_API void BeginMeshDrawEvent_Inner(FRHICommandList& RHICmdList, const class FPrimitiveSceneProxy* PrimitiveSceneProxy, const struct FMeshBatch& Mesh, struct TDrawEvent<FRHICommandList>& DrawEvent);
 #endif
 
-/** Emits draw events for a given FMeshBatch and the PrimitiveSceneProxy corresponding to that mesh element. */
-FORCEINLINE void EmitMeshDrawEvents(FRHICommandList& RHICmdList, const class FPrimitiveSceneProxy* PrimitiveSceneProxy, const struct FMeshBatch& Mesh)
+FORCEINLINE void BeginMeshDrawEvent(FRHICommandList& RHICmdList, const class FPrimitiveSceneProxy* PrimitiveSceneProxy, const struct FMeshBatch& Mesh, struct TDrawEvent<FRHICommandList>& DrawEvent)
 {
-#if !(UE_BUILD_SHIPPING || UE_BUILD_TEST)
-	if ( GShowMaterialDrawEvents )
+#if WANTS_DRAW_MESH_EVENTS
+	if (GShowMaterialDrawEvents)
 	{
-		EmitMeshDrawEvents_Inner(RHICmdList, PrimitiveSceneProxy, Mesh);
+		BeginMeshDrawEvent_Inner(RHICmdList, PrimitiveSceneProxy, Mesh, DrawEvent);
 	}
 #endif
 }
diff --git a/Engine/Source/Runtime/Engine/Public/SceneUtils.h b/Engine/Source/Runtime/Engine/Public/SceneUtils.h
--- a/Engine/Source/Runtime/Engine/Public/SceneUtils.h
+++ b/Engine/Source/Runtime/Engine/Public/SceneUtils.h
@@ -11,13 +11,13 @@
 #pragma once
 
 // Disable draw mesh events for final builds
-#if !(UE_BUILD_SHIPPING || UE_BUILD_TEST)
+#if !(UE_BUILD_SHIPPING || UE_BUILD_TEST) && PLATFORM_SUPPORTS_DRAW_MESH_EVENTS
 #define WANTS_DRAW_MESH_EVENTS 1
 #else
 #define WANTS_DRAW_MESH_EVENTS 0
 #endif
 
-#if WANTS_DRAW_MESH_EVENTS && PLATFORM_SUPPORTS_DRAW_MESH_EVENTS
+#if WANTS_DRAW_MESH_EVENTS
 
 	/**
 	 * Class that logs draw events based upon class scope. Draw events can be seen
@@ -85,6 +85,8 @@
 	#define SCOPED_DRAW_EVENTF(RHICmdList, Name, Format, ...) TDrawEvent<FRHICommandList> PREPROCESSOR_JOIN(Event_##Name,__LINE__); if(GEmitDrawEvents) PREPROCESSOR_JOIN(Event_##Name,__LINE__).Start(RHICmdList, Format, ##__VA_ARGS__);
 	#define SCOPED_CONDITIONAL_DRAW_EVENT(RHICmdList, Name, Condition) TDrawEvent<FRHICommandList> PREPROCESSOR_JOIN(Event_##Name,__LINE__); if(GEmitDrawEvents && (Condition)) PREPROCESSOR_JOIN(Event_##Name,__LINE__).Start(RHICmdList, TEXT(#Name));
 	#define SCOPED_CONDITIONAL_DRAW_EVENTF(RHICmdList, Name, Condition, Format, ...) TDrawEvent<FRHICommandList> PREPROCESSOR_JOIN(Event_##Name,__LINE__); if(GEmitDrawEvents && (Condition)) PREPROCESSOR_JOIN(Event_##Name,__LINE__).Start(RHICmdList, Format, ##__VA_ARGS__);
+	#define BEGIN_DRAW_EVENTF(RHICmdList, Name, Event, Format, ...) if(GEmitDrawEvents) Event.Start(RHICmdList, Format, ##__VA_ARGS__);
+	#define STOP_DRAW_EVENT(Event) Event.Stop();
 
 	#define SCOPED_COMPUTE_EVENT(RHICmdList, Name) TDrawEvent<FRHIAsyncComputeCommandList> PREPROCESSOR_JOIN(Event_##Name,__LINE__); if(GEmitDrawEvents) PREPROCESSOR_JOIN(Event_##Name,__LINE__).Start(RHICmdList, TEXT(#Name));
 	#define SCOPED_COMPUTE_EVENTF(RHICmdList, Name, Format, ...) TDrawEvent<FRHIAsyncComputeCommandList> PREPROCESSOR_JOIN(Event_##Name,__LINE__); if(GEmitDrawEvents) PREPROCESSOR_JOIN(Event_##Name,__LINE__).Start(RHICmdList, Format, ##__VA_ARGS__);
@@ -99,10 +101,17 @@
 
 #else
 
+	template<typename TRHICmdList>
+	struct ENGINE_API TDrawEvent
+	{
+	};
+
 	#define SCOPED_DRAW_EVENT(...)
 	#define SCOPED_DRAW_EVENTF(...)
 	#define SCOPED_CONDITIONAL_DRAW_EVENT(...)
 	#define SCOPED_CONDITIONAL_DRAW_EVENTF(...)
+	#define BEGIN_DRAW_EVENTF(...)
+	#define STOP_DRAW_EVENT(...)
 
 	#define SCOPED_COMPUTE_EVENT(RHICmdList, Name)
 	#define SCOPED_COMPUTE_EVENTF(RHICmdList, Name, Format, ...)
diff --git a/Engine/Source/Runtime/Engine/Public/SceneView.h b/Engine/Source/Runtime/Engine/Public/SceneView.h
--- a/Engine/Source/Runtime/Engine/Public/SceneView.h
+++ b/Engine/Source/Runtime/Engine/Public/SceneView.h
@@ -496,6 +496,15 @@ public:
 
 extern ENGINE_API TGlobalResource<FBuiltinSamplersUniformBuffer> GBuiltinSamplersUniformBuffer;
 
+namespace EDrawDynamicFlags
+{
+	enum Type
+	{
+		None = 0,
+		ForceLowestLOD = 0x1
+	};
+}
+
 /**
  * A projection from scene space into a 2D screen region.
  */
@@ -559,6 +568,8 @@ public:
 	// Whether to use FOV when computing mesh LOD.
 	bool bUseFieldOfViewForLOD;
 
+	EDrawDynamicFlags::Type DrawDynamicFlags;
+
 	/** Current buffer visualization mode */
 	FName CurrentBufferVisualizationMode;
 
diff --git a/Engine/Source/Runtime/Engine/Public/TimerManager.h b/Engine/Source/Runtime/Engine/Public/TimerManager.h
--- a/Engine/Source/Runtime/Engine/Public/TimerManager.h
+++ b/Engine/Source/Runtime/Engine/Public/TimerManager.h
@@ -67,6 +67,9 @@ struct FTimerUnifiedDelegate
 		FuncDelegate.Unbind();
 		FuncDynDelegate.Unbind();
 	}
+
+	/** Utility to output info about delegate as a string. */
+	FString ToString() const;
 };
 
 enum class ETimerStatus : uint8
@@ -357,6 +360,9 @@ public:
 	 */
 	FTimerHandle K2_FindDynamicTimerHandle(FTimerDynamicDelegate InDynamicDelegate) const;
 
+	/** Debug command to output info on all timers curently set to the log. */
+	void ListTimers() const;
+
 private:
 	void InternalSetTimer( FTimerHandle& InOutHandle, FTimerUnifiedDelegate const& InDelegate, float InRate, bool InbLoop, float InFirstDelay );
 	void InternalSetTimer( FTimerData& NewTimerData, float InRate, bool InbLoop, float InFirstDelay );
diff --git a/Engine/Source/Runtime/GameplayAbilities/Private/Abilities/Tasks/AbilityTask.cpp b/Engine/Source/Runtime/GameplayAbilities/Private/Abilities/Tasks/AbilityTask.cpp
--- a/Engine/Source/Runtime/GameplayAbilities/Private/Abilities/Tasks/AbilityTask.cpp
+++ b/Engine/Source/Runtime/GameplayAbilities/Private/Abilities/Tasks/AbilityTask.cpp
@@ -72,6 +72,11 @@ void UAbilityTask::SetWaitingOnRemotePlayerData()
 	}
 }
 
+void UAbilityTask::ClearWaitingOnRemotePlayerData()
+{
+	WaitState = EAbilityTaskWaitState::WaitingOnGame;
+}
+
 bool UAbilityTask::IsWaitingOnRemotePlayerdata() const
 {
 	return (WaitState == EAbilityTaskWaitState::WaitingOnUser);
diff --git a/Engine/Source/Runtime/GameplayAbilities/Private/AbilitySystemComponent_Abilities.cpp b/Engine/Source/Runtime/GameplayAbilities/Private/AbilitySystemComponent_Abilities.cpp
--- a/Engine/Source/Runtime/GameplayAbilities/Private/AbilitySystemComponent_Abilities.cpp
+++ b/Engine/Source/Runtime/GameplayAbilities/Private/AbilitySystemComponent_Abilities.cpp
@@ -2703,6 +2703,15 @@ void UAbilitySystemComponent::ServerSetReplicatedTargetData_Implementation(FGame
 	// Always adds to cache to store the new data
 	FAbilityReplicatedDataCache& ReplicatedData = AbilityTargetDataMap.FindOrAdd(FGameplayAbilitySpecHandleAndPredictionKey(AbilityHandle, AbilityOriginalPredictionKey));
 
+	if (ReplicatedData.TargetData.Num() > 0)
+	{
+		FGameplayAbilitySpec* Spec = FindAbilitySpecFromHandle(AbilityHandle);
+		if (Spec && Spec->Ability)
+		{
+			ABILITY_LOG(Warning, TEXT("Ability %s is overriding pending replicated target data."), *Spec->Ability->GetName());
+		}
+	}
+
 	ReplicatedData.TargetData = ReplicatedTargetDataHandle;
 	ReplicatedData.ApplicationTag = ApplicationTag;
 	ReplicatedData.bTargetConfirmed = true;
diff --git a/Engine/Source/Runtime/GameplayAbilities/Private/AttributeSet.cpp b/Engine/Source/Runtime/GameplayAbilities/Private/AttributeSet.cpp
--- a/Engine/Source/Runtime/GameplayAbilities/Private/AttributeSet.cpp
+++ b/Engine/Source/Runtime/GameplayAbilities/Private/AttributeSet.cpp
@@ -8,6 +8,7 @@
 #include "VisualLogger.h"
 
 #include "ComponentReregisterContext.h"
+#include "PropertyTag.h"
 
 FGameplayAttribute::FGameplayAttribute(UProperty *NewProperty)
 {
@@ -225,6 +226,21 @@ void FScalableFloat::OnCurveTablePostReimport(UObject* InObject, bool)
 }
 #endif // WITH_EDITOR
 
+bool FScalableFloat::SerializeFromMismatchedTag(const FPropertyTag& Tag, FArchive& Ar)
+{
+	if (Tag.Type == NAME_FloatProperty)
+	{
+		float OldValue;
+		Ar << OldValue;
+		*this = FScalableFloat(OldValue);
+
+		return true;
+	}
+
+	return false;
+}
+
+
 bool FGameplayAttribute::operator==(const FGameplayAttribute& Other) const
 {
 	return ((Other.Attribute == Attribute));
@@ -500,6 +516,7 @@ struct FBadScalableFloat
 
 static FBadScalableFloat GCurrentBadScalableFloat;
 static TArray<FBadScalableFloat> GCurrentBadScalableFloatList;
+static TArray<FBadScalableFloat> GCurrentNaughtyScalableFloatList;
 
 
 static bool CheckForBadScalableFloats_r(void* Data, UStruct* Struct, UClass* Class);
@@ -536,6 +553,16 @@ static bool CheckForBadScalableFloats_Prop_r(void* Data, UProperty* Prop, UClass
 					GCurrentBadScalableFloatList.Add(GCurrentBadScalableFloat);
 				}
 			}
+			else 
+			{
+				if (ThisScalableFloat->Curve.CurveTable != nullptr && ThisScalableFloat->Value != 1.f)
+				{
+					GCurrentBadScalableFloat.Property = Prop;
+					GCurrentBadScalableFloat.String = ThisScalableFloat->ToSimpleString();
+
+					GCurrentNaughtyScalableFloatList.Add(GCurrentBadScalableFloat);
+				}
+			}
 		}
 		else
 		{
@@ -613,7 +640,7 @@ static bool	FindClassesWithScalableFloat_r(const TArray<FString>& Args, UStruct*
 	return false;
 }
 
-static void	FindInvalidScalableFloats(const TArray<FString>& Args)
+static void	FindInvalidScalableFloats(const TArray<FString>& Args, bool ShowCoeffecients)
 {
 	GCurrentBadScalableFloatList.Empty();
 
@@ -670,20 +697,40 @@ static void	FindInvalidScalableFloats(const TArray<FString>& Args)
 	ABILITY_LOG( Error, TEXT(""));
 	ABILITY_LOG( Error, TEXT(""));
 
-	for ( FBadScalableFloat& BadFoo : GCurrentBadScalableFloatList)
+	if (ShowCoeffecients == false)
 	{
-		ABILITY_LOG( Error, TEXT("%s. %s. %s"), *BadFoo.Asset->GetFullName(), *BadFoo.Property->GetFullName(), *BadFoo.String );
 
+		for ( FBadScalableFloat& BadFoo : GCurrentBadScalableFloatList)
+		{
+			ABILITY_LOG( Error, TEXT(", %s, %s, %s,"), *BadFoo.Asset->GetFullName(), *BadFoo.Property->GetFullName(), *BadFoo.String );
+
+		}
+
+		ABILITY_LOG( Error, TEXT(""));
+		ABILITY_LOG( Error, TEXT("%d Errors total"), GCurrentBadScalableFloatList.Num() );
 	}
+	else
+	{
+		ABILITY_LOG( Error, TEXT("Non 1 coefficients: "));
 
-	ABILITY_LOG( Error, TEXT(""));
-	ABILITY_LOG( Error, TEXT("%d Errors total"), GCurrentBadScalableFloatList.Num() );
+		for ( FBadScalableFloat& BadFoo : GCurrentNaughtyScalableFloatList)
+		{
+			ABILITY_LOG( Error, TEXT(", %s, %s, %s"), *BadFoo.Asset->GetFullName(), *BadFoo.Property->GetFullName(), *BadFoo.String );
+
+		}
+	}
 }
 
 FAutoConsoleCommand FindInvalidScalableFloatsCommand(
 	TEXT("FindInvalidScalableFloats"), 
 	TEXT( "Searches for invalid scalable floats in all assets. Warning this is slow!" ), 
-	FConsoleCommandWithArgsDelegate::CreateStatic(FindInvalidScalableFloats)
+	FConsoleCommandWithArgsDelegate::CreateStatic(FindInvalidScalableFloats, false)
+);
+
+FAutoConsoleCommand FindCoefficientScalableFloatsCommand(
+	TEXT("FindCoefficientScalableFloats"), 
+	TEXT( "Searches for scalable floats with a non 1 coeffecient. Warning this is slow!" ), 
+	FConsoleCommandWithArgsDelegate::CreateStatic(FindInvalidScalableFloats, true)
 );
 
 #endif
\ No newline at end of file
diff --git a/Engine/Source/Runtime/GameplayAbilities/Private/GameplayCueManager.cpp b/Engine/Source/Runtime/GameplayAbilities/Private/GameplayCueManager.cpp
--- a/Engine/Source/Runtime/GameplayAbilities/Private/GameplayCueManager.cpp
+++ b/Engine/Source/Runtime/GameplayAbilities/Private/GameplayCueManager.cpp
@@ -8,6 +8,7 @@
 #include "GameplayTagsModule.h"
 #include "GameplayCueNotify_Static.h"
 #include "AbilitySystemComponent.h"
+#include "UnrealNetwork.h"
 
 #if WITH_EDITOR
 #include "UnrealEd.h"
@@ -52,6 +53,8 @@ void UGameplayCueManager::OnCreated()
 {
 	FWorldDelegates::OnPostWorldCreation.AddUObject(this, &UGameplayCueManager::OnWorldCreated);
 	FWorldDelegates::OnWorldCleanup.AddUObject(this, &UGameplayCueManager::OnWorldCleanup);
+
+	FNetworkReplayDelegates::OnPreScrub.AddUObject(this, &UGameplayCueManager::OnPreReplayScrub);
 }
 
 bool IsDedicatedServerForGameplayCue()
@@ -923,6 +926,12 @@ void UGameplayCueManager::DumpPreallocationStats(UWorld* World)
 	}
 }
 
+void UGameplayCueManager::OnPreReplayScrub(UWorld* World)
+{
+	FPreallocationInfo& Info = GetPreallocationInfo(World);
+	Info.PreallocatedInstances.Reset();
+}
+
 #if WITH_EDITOR
 #undef LOCTEXT_NAMESPACE
 #endif
diff --git a/Engine/Source/Runtime/GameplayAbilities/Private/GameplayCueNotify_Actor.cpp b/Engine/Source/Runtime/GameplayAbilities/Private/GameplayCueNotify_Actor.cpp
--- a/Engine/Source/Runtime/GameplayAbilities/Private/GameplayCueNotify_Actor.cpp
+++ b/Engine/Source/Runtime/GameplayAbilities/Private/GameplayCueNotify_Actor.cpp
@@ -11,12 +11,15 @@ AGameplayCueNotify_Actor::AGameplayCueNotify_Actor(const FObjectInitializer& Obj
 {
 	IsOverride = true;
 	PrimaryActorTick.bCanEverTick = true;
+	PrimaryActorTick.bStartWithTickEnabled = false;
 	bAutoDestroyOnRemove = false;
 	AutoDestroyDelay = 0.f;
 	bUniqueInstancePerSourceObject = false;
 	bUniqueInstancePerInstigator = false;
 
 	NumPreallocatedInstances = 0;
+
+	StackCount = 0;
 }
 
 #if WITH_EDITOR
@@ -76,12 +79,40 @@ bool AGameplayCueNotify_Actor::HandlesEvent(EGameplayCueEvent::Type EventType) c
 	return true;
 }
 
+int32 GameplayCueNotifyActorStacking = 0;
+static FAutoConsoleVariableRef CVarGameplayCueNotifyActorStacking(TEXT("AbilitySystem.GameplayCueNotifyActorStacking"), GameplayCueNotifyActorStacking, TEXT("Enable simple stacking rules for gameplaycue actors"), ECVF_Default );
+
 void AGameplayCueNotify_Actor::HandleGameplayCue(AActor* MyTarget, EGameplayCueEvent::Type EventType, FGameplayCueParameters Parameters)
 {
 	SCOPE_CYCLE_COUNTER(STAT_HandleGameplayCueNotifyActor);
 
 	if (MyTarget && !MyTarget->IsPendingKill())
 	{
+		if (GameplayCueNotifyActorStacking)
+		{
+			if (EventType == EGameplayCueEvent::WhileActive)
+			{
+				StackCount++;
+				if (StackCount > 1)
+				{
+					return;
+				}
+			}
+			else if (EventType == EGameplayCueEvent::Removed)
+			{
+				StackCount--;
+				if (StackCount > 0)
+				{
+					return;
+				}
+				if (!ensureMsgf(StackCount == 0, TEXT("GameplayCue %s has negative StackCount."), *GetName()))
+				{
+					StackCount = 0;
+				}
+			}
+		}
+
+
 		K2_HandleGameplayCue(MyTarget, EventType, Parameters);
 
 		// Clear any pending auto-destroy that may have occurred from a previous OnRemove
@@ -167,5 +198,6 @@ bool AGameplayCueNotify_Actor::GameplayCuePendingRemove()
 
 bool AGameplayCueNotify_Actor::Recycle()
 {
+	StackCount = 0;
 	return false;
 }
\ No newline at end of file
diff --git a/Engine/Source/Runtime/GameplayAbilities/Public/Abilities/Tasks/AbilityTask.h b/Engine/Source/Runtime/GameplayAbilities/Public/Abilities/Tasks/AbilityTask.h
--- a/Engine/Source/Runtime/GameplayAbilities/Public/Abilities/Tasks/AbilityTask.h
+++ b/Engine/Source/Runtime/GameplayAbilities/Public/Abilities/Tasks/AbilityTask.h
@@ -130,6 +130,7 @@ class GAMEPLAYABILITIES_API UAbilityTask : public UGameplayTask
 
 	/** Called when the ability task is waiting on remote player data. IF the remote player ends the ability prematurely, and a task with this set is still running, the ability is killed. */
 	void SetWaitingOnRemotePlayerData();
+	void ClearWaitingOnRemotePlayerData();
 
 	virtual bool IsWaitingOnRemotePlayerdata() const override;
 
diff --git a/Engine/Source/Runtime/GameplayAbilities/Public/AttributeSet.h b/Engine/Source/Runtime/GameplayAbilities/Public/AttributeSet.h
--- a/Engine/Source/Runtime/GameplayAbilities/Public/AttributeSet.h
+++ b/Engine/Source/Runtime/GameplayAbilities/Public/AttributeSet.h
@@ -229,6 +229,9 @@ public:
 	//copy operator to prevent duplicate handles
 	void operator=(const FScalableFloat& Src);
 
+	/* Used to upgrade a float property into an FScalableFloat */
+	bool SerializeFromMismatchedTag(const FPropertyTag& Tag, FArchive& Ar);
+
 private:
 
 	/** Conditionally register the OnCurveTablePostReimport function with the re-import manager */
@@ -249,6 +252,17 @@ private:
 	mutable FRichCurve* FinalCurve;
 };
 
+template<>
+struct TStructOpsTypeTraits<FScalableFloat>
+	: public TStructOpsTypeTraitsBase
+{
+	enum
+	{
+		WithSerializeFromMismatchedTag = true,
+	};
+};
+
+
 /**
  *	DataTable that allows us to define meta data about attributes. Still a work in progress.
  */
diff --git a/Engine/Source/Runtime/GameplayAbilities/Public/GameplayCueManager.h b/Engine/Source/Runtime/GameplayAbilities/Public/GameplayCueManager.h
--- a/Engine/Source/Runtime/GameplayAbilities/Public/GameplayCueManager.h
+++ b/Engine/Source/Runtime/GameplayAbilities/Public/GameplayCueManager.h
@@ -106,6 +106,8 @@ class GAMEPLAYABILITIES_API UGameplayCueManager : public UDataAsset
 
 	void OnWorldCleanup(UWorld* World, bool bSessionEnded, bool bCleanupResources);
 
+	void OnPreReplayScrub(UWorld* World);
+
 	/** Prints what classess exceeded their preallocation sizes during runtime */
 	void DumpPreallocationStats(UWorld* World);
 
diff --git a/Engine/Source/Runtime/GameplayAbilities/Public/GameplayCueNotify_Actor.h b/Engine/Source/Runtime/GameplayAbilities/Public/GameplayCueNotify_Actor.h
--- a/Engine/Source/Runtime/GameplayAbilities/Public/GameplayCueNotify_Actor.h
+++ b/Engine/Source/Runtime/GameplayAbilities/Public/GameplayCueNotify_Actor.h
@@ -104,9 +104,13 @@ class GAMEPLAYABILITIES_API AGameplayCueNotify_Actor : public AActor
 	int32 NumPreallocatedInstances;
 
 	FGCNotifyActorKey NotifyKey;
+	
+	int32 StackCount;
+
+protected:
+	FTimerHandle FinishTimerHandle;
 
 private:
 	virtual void DeriveGameplayCueTagFromAssetName();
 
-	FTimerHandle FinishTimerHandle;
 };
diff --git a/Engine/Source/Runtime/GameplayTags/Classes/GameplayTagContainer.h b/Engine/Source/Runtime/GameplayTags/Classes/GameplayTagContainer.h
--- a/Engine/Source/Runtime/GameplayTags/Classes/GameplayTagContainer.h
+++ b/Engine/Source/Runtime/GameplayTags/Classes/GameplayTagContainer.h
@@ -340,6 +340,11 @@ struct GAMEPLAYTAGS_API FGameplayTagContainer
 		RemoveAllTags(GameplayTags.Num());
 	}
 
+	void Reset()
+	{
+		GameplayTags.Reset();
+	}
+
 	/**
 	 * Serialize the tag container
 	 *
diff --git a/Engine/Source/Runtime/JsonUtilities/Private/JsonObjectConverter.cpp b/Engine/Source/Runtime/JsonUtilities/Private/JsonObjectConverter.cpp
--- a/Engine/Source/Runtime/JsonUtilities/Private/JsonObjectConverter.cpp
+++ b/Engine/Source/Runtime/JsonUtilities/Private/JsonObjectConverter.cpp
@@ -192,6 +192,8 @@ bool FJsonObjectConverter::UStructToJsonObjectString(const UStruct* StructDefini
 	return false;
 }
 
+//template bool FJsonObjectConverter::UStructToFormattedJsonObjectString<TCHAR, TCondensedJsonPrintPolicy>(const UStruct* StructDefinition, const void* Struct, FString& OutJsonString, int64 CheckFlags, int64 SkipFlags, int32 Indent, const CustomExportCallback* ExportCb);
+
 namespace
 {
 
diff --git a/Engine/Source/Runtime/JsonUtilities/Public/JsonObjectConverter.h b/Engine/Source/Runtime/JsonUtilities/Public/JsonObjectConverter.h
--- a/Engine/Source/Runtime/JsonUtilities/Public/JsonObjectConverter.h
+++ b/Engine/Source/Runtime/JsonUtilities/Public/JsonObjectConverter.h
@@ -70,6 +70,32 @@ public: // UStruct -> JSON
 	 */
 	static bool UStructToJsonObjectString(const UStruct* StructDefinition, const void* Struct, FString& OutJsonString, int64 CheckFlags, int64 SkipFlags, int32 Indent = 0, const CustomExportCallback* ExportCb = nullptr);
 
+	/**
+	 * Wrapper to UStructToJsonObjectString that allows a print policy to be specified.
+	 */
+	template<typename CharType, template<typename> class PrintPolicy>
+	static bool UStructToFormattedJsonObjectString(const UStruct* StructDefinition, const void* Struct, FString& OutJsonString, int64 CheckFlags, int64 SkipFlags, int32 Indent = 0, const CustomExportCallback* ExportCb = nullptr)
+	{
+		TSharedRef<FJsonObject> JsonObject = MakeShareable(new FJsonObject());
+		if (UStructToJsonObject(StructDefinition, Struct, JsonObject, CheckFlags, SkipFlags, ExportCb))
+		{
+			TSharedRef<TJsonWriter<CharType, PrintPolicy<CharType>>> JsonWriter = TJsonWriterFactory<CharType, PrintPolicy<CharType>>::Create(&OutJsonString, Indent);
+
+			if (FJsonSerializer::Serialize(JsonObject, JsonWriter))
+			{
+				JsonWriter->Close();
+				return true;
+			}
+			else
+			{
+				UE_LOG(LogJson, Warning, TEXT("UStructToFormattedObjectString - Unable to write out json"));
+				JsonWriter->Close();
+			}
+		}
+
+		return false;
+	}
+
 	/**
 	 * Converts from a UStruct to a set of json attributes (possibly from within a JsonObject)
 	 *
@@ -240,4 +266,3 @@ public: // JSON -> UStruct
 		return true;
 	}
 };
-
diff --git a/Engine/Source/Runtime/LevelSequence/Private/LevelSequence.cpp b/Engine/Source/Runtime/LevelSequence/Private/LevelSequence.cpp
--- a/Engine/Source/Runtime/LevelSequence/Private/LevelSequence.cpp
+++ b/Engine/Source/Runtime/LevelSequence/Private/LevelSequence.cpp
@@ -60,41 +60,13 @@ void ULevelSequence::ConvertPersistentBindingsToDefault(UObject* FixupContext)
 	PossessedObjects_DEPRECATED.Empty();
 }
 
-void ULevelSequence::BindToContext(UObject* NewContext)
+void ULevelSequence::BindPossessableObject(const FGuid& ObjectId, UObject& PossessedObject, UObject* Context)
 {
-	ResolutionContext = NewContext;
-	CachedObjectBindings.Reset();
-}
-
-void ULevelSequence::BindPossessableObject(const FGuid& ObjectId, UObject& PossessedObject)
-{
-	
-	UObject* Context = ResolutionContext.Get();
-
-	// Always use the parent as a context when there is a parent spawnable
-	UObject* ParentObject = GetParentObject(&PossessedObject);
-	FGuid ParentId = ParentObject ? FindObjectId(*ParentObject) : FGuid();
-
-	FMovieSceneSpawnable* ParentSpawnable = ParentId.IsValid() ? MovieScene->FindSpawnable(ParentId) : nullptr;
-	if (ParentSpawnable)
-	{
-		Context = ParentObject;
-	}
-
 	if (Context)
 	{
 		ObjectReferences.CreateBinding(ObjectId, &PossessedObject, Context);
 
 		MovieSceneHelpers::SetRuntimeObjectMobility(&PossessedObject);
-
-		// Ensure that we report possessables to their parent spawnable, if necessary
-		FMovieScenePossessable* ChildPossessable = MovieScene->FindPossessable(ObjectId);
-
-		if (ParentSpawnable && ensureMsgf(ChildPossessable, TEXT("Object being possessed without a corresponding possessable")))
-		{
-			ParentSpawnable->AddChildPossessable(ObjectId);
-			ChildPossessable->SetParentSpawnable(ParentId);
-		}
 	}
 }
 
@@ -103,76 +75,9 @@ bool ULevelSequence::CanPossessObject(UObject& Object) const
 	return Object.IsA<AActor>() || Object.IsA<UActorComponent>();
 }
 
-UObject* ULevelSequence::FindObject(const FGuid& ObjectId) const
+UObject* ULevelSequence::FindPossessableObject(const FGuid& ObjectId, UObject* Context) const
 {
-	// If it's already cached, we can just return that
-	if (auto* WeakCachedObject = CachedObjectBindings.Find(ObjectId))
-	{
-		if (auto* CachedObj = WeakCachedObject->Get())
-		{
-			return CachedObj;
-		}
-	}
-
-	// search spawned objects
-	if (auto* SpawnedObject = SpawnedObjects.FindRef(ObjectId).Get())
-	{
-		return SpawnedObject;
-	}
-
-	UObject* Context = ResolutionContext.Get();
-
-	// Always use the parent as a context when possible
-	FMovieScenePossessable* Possessable = MovieScene->FindPossessable(ObjectId);
-	if (Possessable && Possessable->GetParentSpawnable().IsValid())
-	{
-		Context = FindObject(Possessable->GetParentSpawnable());
-	}
-
-	// Attempt to resolve the object binding through the remapped bindings
-	UObject* FoundObject = Context ? ObjectReferences.ResolveBinding(ObjectId, Context) : nullptr;
-
-	if (FoundObject)
-	{
-		// Cache the object if we found one
-		CachedObjectBindings.Add(ObjectId, FoundObject);
-	}
-
-	return FoundObject;
-}
-
-FGuid ULevelSequence::FindObjectId(UObject& Object) const
-{
-	for (auto& Pair : CachedObjectBindings)
-	{
-		if (Pair.Value.Get() == &Object)
-		{
-			return Pair.Key;
-		}
-	}
-
-	// search spawned objects
-	for (auto& Pair : SpawnedObjects)
-	{
-		if (Pair.Value.Get() == &Object)
-		{
-			return Pair.Key;
-		}
-	}
-
-	// search possessed objects
-	if (UObject* Context = ResolutionContext.Get())
-	{
-		// Search instances
-		FGuid ObjectId = ObjectReferences.FindBindingId(&Object, Context);
-		if (ObjectId.IsValid())
-		{
-			return ObjectId;
-		}
-	}
-
-	// not found
-	return FGuid();
+	return Context ? ObjectReferences.ResolveBinding(ObjectId, Context) : nullptr;
 }
 
 UMovieScene* ULevelSequence::GetMovieScene() const
@@ -197,154 +102,7 @@ bool ULevelSequence::AllowsSpawnableObjects() const
 	return true;
 }
 
-UObject* ULevelSequence::SpawnObject(const FGuid& ObjectId)
-{
-	// Return an existing object relating to this object ID, if possible
-	UObject* ObjectPtr = SpawnedObjects.FindRef(ObjectId).Get();
-	if (ObjectPtr)
-	{
-		return ObjectPtr;
-	}
-
-	// Find the spawnable definition
-	FMovieSceneSpawnable* Spawnable = MovieScene->FindSpawnable(ObjectId);
-	if (!Spawnable)
-	{
-		return nullptr;
-	}
-
-	UClass* SpawnableClass = Spawnable->GetClass();
-	if (!SpawnableClass || !SpawnableClass->IsChildOf(AActor::StaticClass()))
-	{
-		return nullptr;
-	}
-
-	// @todo sequencer: We should probably spawn these in a specific sub-level!
-	// World->CurrentLevel = ???;
-
-	const FName ActorName = NAME_None;
-
-	// Override the object flags so that RF_Transactional is not set.  Puppet actors are never transactional
-	// @todo sequencer: These actors need to avoid any transaction history.  However, RF_Transactional can currently be set on objects on the fly!
-	// NOTE: We are omitting RF_Transactional intentionally
-	const EObjectFlags ObjectFlags = RF_Transient;
-
-	// @todo sequencer livecapture: Consider using SetPlayInEditorWorld() and RestoreEditorWorld() here instead
-	
-	// @todo sequencer actors: We need to make sure puppet objects aren't copied into PIE/SIE sessions!  They should be omitted from that duplication!
-
-	// Spawn the puppet actor
-	FActorSpawnParameters SpawnInfo;
-	{
-		SpawnInfo.Name = ActorName;
-		SpawnInfo.ObjectFlags = ObjectFlags;
-	}
-
-	FTransform SpawnTransform;
-
-	AActor* ActorCDO = CastChecked<AActor>(SpawnableClass->ClassDefaultObject);
-	if (USceneComponent* RootComponent = ActorCDO->GetRootComponent())
-	{
-		SpawnTransform.SetTranslation(RootComponent->RelativeLocation);
-		SpawnTransform.SetRotation(RootComponent->RelativeRotation.Quaternion());
-	}
-
-	//@todo: This is in place of SpawnActorAbsolute which doesn't exist yet in the Orion branch
-	FTransform NewTransform = SpawnTransform;
-	{
-		AActor* Template = SpawnInfo.Template;
-
-		if(!Template)
-		{
-			// Use class's default actor as a template.
-			Template = SpawnableClass->GetDefaultObject<AActor>();
-		}
-
-		USceneComponent* TemplateRootComponent = (Template)? Template->GetRootComponent() : NULL;
-		if(TemplateRootComponent)
-		{
-			TemplateRootComponent->UpdateComponentToWorld();
-			NewTransform = TemplateRootComponent->GetComponentToWorld().Inverse() * NewTransform;
-		}
-	}
-
-	UObject* SpawnedObject = GWorld->SpawnActor(SpawnableClass, &NewTransform, SpawnInfo);
-	if (!SpawnedObject)
-	{
-		return nullptr;
-	}
-	
-	SpawnedObjects.Add(ObjectId, SpawnedObject);
-	CachedObjectBindings.Add(ObjectId, SpawnedObject);
-
-	MovieSceneHelpers::SetRuntimeObjectMobility(SpawnedObject);
-
-	return SpawnedObject;
-}
-
-void ULevelSequence::DestroySpawnedObject(const FGuid& ObjectId)
-{
-	if (UObject* ObjectPtr = SpawnedObjects.FindRef(ObjectId).Get())
-	{
-		AActor* Actor = Cast<AActor>(ObjectPtr);
-		if (!ensure(Actor))
-		{
-			return;
-		}
-
-		UWorld* World = Actor->GetWorld();
-		if (ensure(World))
-		{
-			const bool bNetForce = false;
-			const bool bShouldModifyLevel = false;
-			World->DestroyActor(Actor, bNetForce, bShouldModifyLevel);
-		}
-	}
-
-	CachedObjectBindings.Remove(ObjectId);
-	SpawnedObjects.Remove(ObjectId);
-}
-
-void ULevelSequence::DestroyAllSpawnedObjects()
-{
-	// Use a key array to ensure we don't change the map while iterating it
-	TArray<FGuid> Keys;
-	Keys.Reserve(SpawnedObjects.Num());
-	SpawnedObjects.GenerateKeyArray(Keys);
-
-	for (auto& ObjectId : Keys)
-	{
-		DestroySpawnedObject(ObjectId);
-	}
-	SpawnedObjects.Reset();
-}
-
-#if WITH_EDITOR
-bool ULevelSequence::TryGetObjectDisplayName(const FGuid& ObjectId, FText& OutDisplayName) const
-{
-	UObject* Object = FindObject(ObjectId);
-
-	if (Object == nullptr)
-	{
-		return false;
-	}
-
-	AActor* Actor = Cast<AActor>(Object);
-
-	if (Actor == nullptr)
-	{
-		return false;
-	}
-
-	OutDisplayName = FText::FromString(Actor->GetActorLabel());
-
-	return true;
-}
-
-#endif
-
 void ULevelSequence::UnbindPossessableObjects(const FGuid& ObjectId)
 {
 	ObjectReferences.RemoveBinding(ObjectId);
-	CachedObjectBindings.Remove(ObjectId);
 }
diff --git a/Engine/Source/Runtime/LevelSequence/Private/LevelSequenceActor.cpp b/Engine/Source/Runtime/LevelSequence/Private/LevelSequenceActor.cpp
--- a/Engine/Source/Runtime/LevelSequence/Private/LevelSequenceActor.cpp
+++ b/Engine/Source/Runtime/LevelSequence/Private/LevelSequenceActor.cpp
@@ -82,11 +82,8 @@ void ALevelSequenceActor::InitializePlayer()
 	ULevelSequence* LevelSequenceAsset = Cast<ULevelSequence>(LevelSequence.TryLoad());
 	if (GetWorld()->IsGameWorld() && LevelSequenceAsset)
 	{
-		// Create a duplicate of the asset so we can bind it to the game world context
-		auto* SequenceInstance = DuplicateObject<ULevelSequence>(LevelSequenceAsset, this, TEXT("AnimationInstance"));
-
 		SequencePlayer = NewObject<ULevelSequencePlayer>(this, "AnimationPlayer");
-		SequencePlayer->Initialize(SequenceInstance, GetWorld(), PlaybackSettings);
+		SequencePlayer->Initialize(LevelSequenceAsset, GetWorld(), PlaybackSettings);
 
 		if (bAutoPlay)
 		{
diff --git a/Engine/Source/Runtime/LevelSequence/Private/LevelSequencePlayer.cpp b/Engine/Source/Runtime/LevelSequence/Private/LevelSequencePlayer.cpp
--- a/Engine/Source/Runtime/LevelSequence/Private/LevelSequencePlayer.cpp
+++ b/Engine/Source/Runtime/LevelSequence/Private/LevelSequencePlayer.cpp
@@ -5,6 +5,7 @@
 #include "MovieScene.h"
 #include "MovieSceneSequence.h"
 #include "MovieSceneSequenceInstance.h"
+#include "LevelSequenceSpawnRegister.h"
 
 
 struct FTickAnimationPlayers : public FTickableGameObject
@@ -68,7 +69,10 @@ ULevelSequencePlayer::ULevelSequencePlayer(const FObjectInitializer& ObjectIniti
 	, bIsPlaying(false)
 	, TimeCursorPosition(0.0f)
 	, CurrentNumLoops(0)
-{ }
+	, bHasCleanedUpSequence(false)
+{
+	SpawnRegister = MakeShareable(new FLevelSequenceSpawnRegister);
+}
 
 
 /* ULevelSequencePlayer interface
@@ -124,10 +128,15 @@ void ULevelSequencePlayer::Play()
 	}
 
 	// @todo Sequencer playback: Should we recreate the instance every time?
-	RootMovieSceneInstance = MakeShareable(new FMovieSceneSequenceInstance(*LevelSequence));
-	RootMovieSceneInstance->RefreshInstance(*this);
+	// We must not recreate the instance since it holds stateful information (such as which objects it has spawned). Recreating the instance would break any 
+	if (!RootMovieSceneInstance.IsValid())
+	{
+		RootMovieSceneInstance = MakeShareable(new FMovieSceneSequenceInstance(*LevelSequence));
+		RootMovieSceneInstance->RefreshInstance(*this);
+	}
 
 	bIsPlaying = true;
+	bHasCleanedUpSequence = false;
 
 	UpdateMovieSceneInstance(TimeCursorPosition, TimeCursorPosition);
 }
@@ -198,6 +207,8 @@ void ULevelSequencePlayer::OnCursorPositionChanged()
 			++CurrentNumLoops;
 			const float Overplay = FMath::Fmod(TimeCursorPosition, Length);
 			TimeCursorPosition = Overplay < 0 ? Length + Overplay : Overplay;
+
+			SpawnRegister->ForgetExternallyOwnedSpawnedObjects(*this);
 		}
 		else
 		{
@@ -215,7 +226,6 @@ void ULevelSequencePlayer::OnCursorPositionChanged()
 void ULevelSequencePlayer::Initialize(ULevelSequence* InLevelSequence, UWorld* InWorld, const FLevelSequencePlaybackSettings& Settings)
 {
 	LevelSequence = InLevelSequence;
-	LevelSequence->BindToContext(InWorld);
 
 	World = InWorld;
 	PlaybackSettings = Settings;
@@ -230,7 +240,8 @@ void ULevelSequencePlayer::Initialize(ULevelSequence* InLevelSequence, UWorld* I
 
 void ULevelSequencePlayer::GetRuntimeObjects(TSharedRef<FMovieSceneSequenceInstance> MovieSceneInstance, const FGuid& ObjectId, TArray<UObject*>& OutObjects) const
 {
-	if (UObject* FoundObject = LevelSequence->FindObject(ObjectId))
+	UObject* FoundObject = MovieSceneInstance->FindObject(ObjectId, *this);
+	if (FoundObject)
 	{
 		OutObjects.Add(FoundObject);
 	}
@@ -302,6 +313,11 @@ TSharedRef<FMovieSceneSequenceInstance> ULevelSequencePlayer::GetRootMovieSceneS
 	return RootMovieSceneInstance.ToSharedRef();
 }
 
+UObject* ULevelSequencePlayer::GetPlaybackContext() const
+{
+	return World.Get();
+}
+
 void ULevelSequencePlayer::Update(const float DeltaSeconds)
 {
 	float LastTimePosition = TimeCursorPosition;
@@ -310,9 +326,15 @@ void ULevelSequencePlayer::Update(const float DeltaSeconds)
 	{
 		TimeCursorPosition += DeltaSeconds * PlaybackSettings.PlayRate;
 		OnCursorPositionChanged();
+		UpdateMovieSceneInstance(TimeCursorPosition, LastTimePosition);
 	}
+	else if (!bHasCleanedUpSequence && TimeCursorPosition >= GetLength())
+	{
+		UpdateMovieSceneInstance(TimeCursorPosition, LastTimePosition);
 
-	UpdateMovieSceneInstance(TimeCursorPosition, LastTimePosition);
+		bHasCleanedUpSequence = true;
+		SpawnRegister->DestroyAllOwnedObjects(*this);
+	}
 }
 
 void ULevelSequencePlayer::UpdateMovieSceneInstance(float CurrentPosition, float PreviousPosition)
@@ -325,4 +347,9 @@ void ULevelSequencePlayer::UpdateMovieSceneInstance(float CurrentPosition, float
 		OnLevelSequencePlayerUpdate.Broadcast(*this, CurrentPosition, PreviousPosition);
 #endif
 	}
+}
+
+IMovieSceneSpawnRegister& ULevelSequencePlayer::GetSpawnRegister()
+{
+	return *SpawnRegister;
 }
\ No newline at end of file
diff --git a/Engine/Source/Runtime/LevelSequence/Private/LevelSequenceSpawnRegister.cpp b/Engine/Source/Runtime/LevelSequence/Private/LevelSequenceSpawnRegister.cpp
new file mode 100644
--- /dev/null
+++ b/Engine/Source/Runtime/LevelSequence/Private/LevelSequenceSpawnRegister.cpp
@@ -0,0 +1,231 @@
+// Copyright 1998-2015 Epic Games, Inc. All Rights Reserved.
+
+#include "LevelSequencePCH.h"
+#include "LevelSequenceSpawnRegister.h"
+#include "MovieSceneSequenceInstance.h"
+#include "MovieSceneSpawnable.h"
+#include "MovieSceneSequence.h"
+#include "MovieScene.h"
+#include "MovieSceneCommonHelpers.h"
+
+UObject* FLevelSequenceSpawnRegister::SpawnObject(const FGuid& BindingId, FMovieSceneSequenceInstance& SequenceInstance, IMovieScenePlayer& Player)
+{
+	FMovieSceneSpawnRegisterKey Key(BindingId, SequenceInstance);
+
+	FSpawnedObject* Existing = Register.Find(Key);
+	UObject* ObjectInstance = Existing ? Existing->Object.Get() : nullptr;
+	if (ObjectInstance)
+	{
+		return ObjectInstance;
+	}
+
+	// Find the spawnable definition
+	FMovieSceneSpawnable* Spawnable = SequenceInstance.GetSequence()->GetMovieScene()->FindSpawnable(BindingId);
+	if (!Spawnable)
+	{
+		return nullptr;
+	}
+
+	UClass* SpawnableClass = Spawnable->GetClass();
+	if (!SpawnableClass || !SpawnableClass->IsChildOf(AActor::StaticClass()))
+	{
+		return nullptr;
+	}
+
+	// @todo sequencer: We should probably spawn these in a specific sub-level!
+	// World->CurrentLevel = ???;
+
+	const FName ActorName = NAME_None;
+
+	// Override the object flags so that RF_Transactional is not set.  Puppet actors are never transactional
+	// @todo sequencer: These actors need to avoid any transaction history.  However, RF_Transactional can currently be set on objects on the fly!
+	// NOTE: We are omitting RF_Transactional intentionally
+	const EObjectFlags ObjectFlags = RF_Transient;
+
+	// @todo sequencer livecapture: Consider using SetPlayInEditorWorld() and RestoreEditorWorld() here instead
+	
+	// @todo sequencer actors: We need to make sure puppet objects aren't copied into PIE/SIE sessions!  They should be omitted from that duplication!
+
+	// Spawn the puppet actor
+	FActorSpawnParameters SpawnInfo;
+	{
+		SpawnInfo.Name = ActorName;
+		SpawnInfo.ObjectFlags = ObjectFlags;
+	}
+
+	FTransform SpawnTransform;
+
+	AActor* ActorCDO = CastChecked<AActor>(SpawnableClass->ClassDefaultObject);
+	if (USceneComponent* RootComponent = ActorCDO->GetRootComponent())
+	{
+		SpawnTransform.SetTranslation(RootComponent->RelativeLocation);
+		SpawnTransform.SetRotation(RootComponent->RelativeRotation.Quaternion());
+	}
+
+	//@todo: This is in place of SpawnActorAbsolute which doesn't exist yet in the Orion branch
+	FTransform NewTransform = SpawnTransform;
+	{
+		AActor* Template = SpawnInfo.Template;
+
+		if(!Template)
+		{
+			// Use class's default actor as a template.
+			Template = SpawnableClass->GetDefaultObject<AActor>();
+		}
+
+		USceneComponent* TemplateRootComponent = (Template)? Template->GetRootComponent() : NULL;
+		if(TemplateRootComponent)
+		{
+			TemplateRootComponent->UpdateComponentToWorld();
+			NewTransform = TemplateRootComponent->GetComponentToWorld().Inverse() * NewTransform;
+		}
+	}
+
+	UObject* SpawnedObject = GWorld->SpawnActor(SpawnableClass, &NewTransform, SpawnInfo);
+	if (!SpawnedObject)
+	{
+		return nullptr;
+	}
+
+	MovieSceneHelpers::SetRuntimeObjectMobility(SpawnedObject);
+
+	Register.Add(Key, FSpawnedObject(*SpawnedObject, Spawnable->GetSpawnOwnership()));
+
+	SequenceInstance.OnObjectSpawned(BindingId, *SpawnedObject, Player);
+	return SpawnedObject;
+}
+
+void FLevelSequenceSpawnRegister::DestroySpawnedObject(const FGuid& BindingId, FMovieSceneSequenceInstance& SequenceInstance, IMovieScenePlayer& Player)
+{
+	FMovieSceneSpawnRegisterKey Key(BindingId, SequenceInstance);
+	
+	FSpawnedObject* Existing = Register.Find(Key);
+	UObject* SpawnedObject = Existing ? Existing->Object.Get() : nullptr;
+	if (SpawnedObject)
+	{
+		PreDestroyObject(*SpawnedObject, BindingId, SequenceInstance);
+		DestroySpawnedObject(*SpawnedObject);
+	}
+
+	SequenceInstance.OnSpawnedObjectDestroyed(BindingId, Player);
+	Register.Remove(Key);
+}
+
+void FLevelSequenceSpawnRegister::PreDestroyObject(UObject& Object, const FGuid& BindingId, FMovieSceneSequenceInstance& SequenceInstance)
+{
+}
+
+void FLevelSequenceSpawnRegister::DestroySpawnedObject(UObject& Object)
+{
+	AActor* Actor = Cast<AActor>(&Object);
+	if (!ensure(Actor))
+	{
+		return;
+	}
+
+	UWorld* World = Actor->GetWorld();
+	if (ensure(World))
+	{
+		const bool bNetForce = false;
+		const bool bShouldModifyLevel = false;
+		World->DestroyActor(Actor, bNetForce, bShouldModifyLevel);
+	}
+}
+
+void FLevelSequenceSpawnRegister::ForgetExternallyOwnedSpawnedObjects(IMovieScenePlayer& Player)
+{
+	for (auto It = Register.CreateIterator(); It; ++It)
+	{
+		if (It.Value().Ownership == ESpawnOwnership::External)
+		{
+			TSharedPtr<FMovieSceneSequenceInstance> SequenceInstance = It.Key().SequenceInstance.Pin();
+			if (SequenceInstance.IsValid())
+			{
+				SequenceInstance->OnSpawnedObjectDestroyed(It.Key().BindingId, Player);
+			}
+			It.RemoveCurrent();
+		}
+	}
+}
+
+void FLevelSequenceSpawnRegister::DestroyObjects(IMovieScenePlayer& Player, TFunctionRef<bool(FMovieSceneSequenceInstance&, FSpawnedObject&)> Pred)
+{
+	for (auto It = Register.CreateIterator(); It; ++It)
+	{
+		FMovieSceneSequenceInstance* ThisInstance = It.Key().SequenceInstance.Pin().Get();
+		if (!ThisInstance || Pred(*ThisInstance, It.Value()))
+		{
+			UObject* SpawnedObject = It.Value().Object.Get();
+			if (SpawnedObject)
+			{
+				if (ThisInstance)
+				{
+					PreDestroyObject(*SpawnedObject, It.Key().BindingId, *ThisInstance);
+				}
+				DestroySpawnedObject(*SpawnedObject);
+			}
+
+			if (ThisInstance)
+			{
+				ThisInstance->OnSpawnedObjectDestroyed(It.Key().BindingId, Player);
+			}
+			It.RemoveCurrent();
+		}
+	}
+}
+
+void FLevelSequenceSpawnRegister::DestroyObjectsOwnedByInstance(FMovieSceneSequenceInstance& SequenceInstance, IMovieScenePlayer& Player)
+{
+	DestroyObjects(Player, [&](FMovieSceneSequenceInstance& ThisInstance, FSpawnedObject& SpawnedObject){
+		if (SpawnedObject.Ownership == ESpawnOwnership::InnerSequence)
+		{
+			return &ThisInstance == &SequenceInstance;
+		}
+		return false;
+	});
+}
+
+void FLevelSequenceSpawnRegister::DestroyObjectsSpawnedByInstance(FMovieSceneSequenceInstance& SequenceInstance, IMovieScenePlayer& Player)
+{
+	DestroyObjects(Player, [&](FMovieSceneSequenceInstance& ThisInstance, FSpawnedObject& SpawnedObject){
+		return &ThisInstance == &SequenceInstance;
+	});
+}
+
+void FLevelSequenceSpawnRegister::DestroyAllOwnedObjects(IMovieScenePlayer& Player)
+{
+	DestroyObjects(Player, [&](FMovieSceneSequenceInstance&, FSpawnedObject& SpawnedObject){
+		return SpawnedObject.Ownership != ESpawnOwnership::External;
+	});
+}
+
+void FLevelSequenceSpawnRegister::DestroyAllObjects(IMovieScenePlayer& Player)
+{
+	DestroyObjects(Player, [&](FMovieSceneSequenceInstance&, FSpawnedObject&){
+		return true;
+	});
+}
+
+void FLevelSequenceSpawnRegister::PreUpdateSequenceInstance(FMovieSceneSequenceInstance& Instance, IMovieScenePlayer& Player)
+{
+	++CurrentlyUpdatingSequenceCount;
+	ActiveInstances.Add(Instance.AsShared());
+}
+
+void FLevelSequenceSpawnRegister::PostUpdateSequenceInstance(FMovieSceneSequenceInstance& InInstance, IMovieScenePlayer& Player)
+{
+	if (--CurrentlyUpdatingSequenceCount == 0)
+	{
+		for (TWeakPtr<FMovieSceneSequenceInstance>& WeakInstance : PreviouslyActiveInstances)
+		{
+			TSharedPtr<FMovieSceneSequenceInstance> Instance = WeakInstance.Pin();
+			if (Instance.IsValid() && !ActiveInstances.Contains(Instance))
+			{
+				DestroyObjectsOwnedByInstance(*Instance, Player);
+			}
+		}
+
+		Swap(ActiveInstances, PreviouslyActiveInstances);
+		ActiveInstances.Reset();
+	}
+}
\ No newline at end of file
diff --git a/Engine/Source/Runtime/LevelSequence/Public/LevelSequence.h b/Engine/Source/Runtime/LevelSequence/Public/LevelSequence.h
--- a/Engine/Source/Runtime/LevelSequence/Public/LevelSequence.h
+++ b/Engine/Source/Runtime/LevelSequence/Public/LevelSequence.h
@@ -7,7 +7,6 @@
 #include "LevelSequenceObjectReference.h"
 #include "LevelSequence.generated.h"
 
-
 /**
  * Movie scene animation for Actors.
  */
@@ -31,34 +30,17 @@ public:
 	/** Convert old-style lazy object ptrs to new-style references using the specified context */
 	void ConvertPersistentBindingsToDefault(UObject* FixupContext);
 
-public:
-
-	/**
-	 * Bind this level sequence to an arbitrary context. For now we only support binding to a UWorld
-	 *
-	 * @param InContext The context to use for object binding resolution.
-	 */
-	void BindToContext(UObject* InContext);
-
 public:
 
 	// UMovieSceneSequence interface
 	
-	virtual void BindPossessableObject(const FGuid& ObjectId, UObject& PossessedObject) override;
+	virtual void BindPossessableObject(const FGuid& ObjectId, UObject& PossessedObject, UObject* Context) override;
 	virtual bool CanPossessObject(UObject& Object) const override;
-	virtual UObject* FindObject(const FGuid& ObjectId) const override;
-	virtual FGuid FindObjectId(UObject& Object) const override;
+	virtual UObject* FindPossessableObject(const FGuid& ObjectId, UObject* Context) const override;
 	virtual UMovieScene* GetMovieScene() const override;
 	virtual UObject* GetParentObject(UObject* Object) const override;
 	virtual void UnbindPossessableObjects(const FGuid& ObjectId) override;
 	virtual bool AllowsSpawnableObjects() const override;
-	virtual UObject* SpawnObject(const FGuid& ObjectId) override;
-	virtual void DestroySpawnedObject(const FGuid& ObjectId) override;
-	virtual void DestroyAllSpawnedObjects() override;
-
-#if WITH_EDITOR
-	virtual bool TryGetObjectDisplayName(const FGuid& ObjectId, FText& OutDisplayName) const override;
-#endif
 
 	virtual bool Rename(const TCHAR* NewName = nullptr, UObject* NewOuter = nullptr, ERenameFlags Flags = REN_None) override;
 	
@@ -71,15 +53,4 @@ private:
 	/** Deprecated property housing old possessed object bindings */
 	UPROPERTY()
 	TMap<FString, FLevelSequenceObject> PossessedObjects_DEPRECATED;
-
-private:
-
-	/** A context to use for resolving object bindings */
-	TWeakObjectPtr<UObject> ResolutionContext;
-
-	/** A transient map of cached object bindings */
-	mutable TMap<FGuid, TWeakObjectPtr<UObject>> CachedObjectBindings;
-
-	/** A map of object ID -> spawned object */
-	TMap<FGuid, TWeakObjectPtr<UObject>> SpawnedObjects;
 };
\ No newline at end of file
diff --git a/Engine/Source/Runtime/LevelSequence/Public/LevelSequencePlayer.h b/Engine/Source/Runtime/LevelSequence/Public/LevelSequencePlayer.h
--- a/Engine/Source/Runtime/LevelSequence/Public/LevelSequencePlayer.h
+++ b/Engine/Source/Runtime/LevelSequence/Public/LevelSequencePlayer.h
@@ -6,7 +6,7 @@
 #include "LevelSequence.h"
 #include "LevelSequencePlayer.generated.h"
 
-
+class FLevelSequenceSpawnRegister;
 class FMovieSceneSequenceInstance;
 class ULevel;
 class UMovieSceneBindings;
@@ -141,6 +141,8 @@ protected:
 	virtual void AddOrUpdateMovieSceneInstance(UMovieSceneSection& MovieSceneSection, TSharedRef<FMovieSceneSequenceInstance> InstanceToAdd) override;
 	virtual void RemoveMovieSceneInstance(UMovieSceneSection& MovieSceneSection, TSharedRef<FMovieSceneSequenceInstance> InstanceToRemove) override;
 	virtual TSharedRef<FMovieSceneSequenceInstance> GetRootMovieSceneSequenceInstance() const override;
+	virtual IMovieSceneSpawnRegister& GetSpawnRegister() override;
+	virtual UObject* GetPlaybackContext() const override;
 
 public:
 
@@ -178,6 +180,9 @@ private:
 	/** The number of times we have looped in the current playback */
 	int32 CurrentNumLoops;
 
+	/** Whether this player has cleaned up the level sequence after it has stopped playing or not */
+	bool bHasCleanedUpSequence;
+
 private:
 
 	/** The root movie scene instance to update when playing. */
@@ -186,6 +191,9 @@ private:
 	/** The world this player will spawn actors in, if needed */
 	TWeakObjectPtr<UWorld> World;
 
+	/** Register responsible for managing spawned objects */
+	TSharedPtr<FLevelSequenceSpawnRegister> SpawnRegister;
+
 #if WITH_EDITOR
 public:
 
diff --git a/Engine/Source/Runtime/LevelSequence/Public/LevelSequenceSpawnRegister.h b/Engine/Source/Runtime/LevelSequence/Public/LevelSequenceSpawnRegister.h
new file mode 100644
--- /dev/null
+++ b/Engine/Source/Runtime/LevelSequence/Public/LevelSequenceSpawnRegister.h
@@ -0,0 +1,65 @@
+// Copyright 1998-2015 Epic Games, Inc. All Rights Reserved.
+
+#pragma once
+
+#include "IMovieSceneSpawnRegister.h"
+#include "MovieSceneSequenceInstance.h"
+#include "MovieSceneSpawnable.h"
+
+/** Movie scene spawn register that knows how to handle spawning objects (actors) for a level sequence  */
+class LEVELSEQUENCE_API FLevelSequenceSpawnRegister : public IMovieSceneSpawnRegister
+{
+public:
+	FLevelSequenceSpawnRegister()
+	{
+		CurrentlyUpdatingSequenceCount = 0;
+	}
+
+	/** ~ IMovieSceneSpawnRegister interface */
+	virtual UObject* SpawnObject(const FGuid& BindingId, FMovieSceneSequenceInstance& SequenceInstance, IMovieScenePlayer&) override;
+	virtual void DestroySpawnedObject(const FGuid& BindingId, FMovieSceneSequenceInstance& SequenceInstance, IMovieScenePlayer&) override;
+	virtual void DestroyObjectsOwnedByInstance(FMovieSceneSequenceInstance&, IMovieScenePlayer&) override;
+	virtual void DestroyObjectsSpawnedByInstance(FMovieSceneSequenceInstance&, IMovieScenePlayer&) override;
+	virtual void DestroyAllOwnedObjects(IMovieScenePlayer&) override;
+	virtual void ForgetExternallyOwnedSpawnedObjects(IMovieScenePlayer& Player) override;
+	virtual void DestroyAllObjects(IMovieScenePlayer&) override;
+	virtual void PreUpdateSequenceInstance(FMovieSceneSequenceInstance& Instance, IMovieScenePlayer&) override;
+	virtual void PostUpdateSequenceInstance(FMovieSceneSequenceInstance& Instance, IMovieScenePlayer&) override;
+
+protected:
+
+	/** Structure holding information pertaining to a spawned object */
+	struct FSpawnedObject
+	{
+		FSpawnedObject(UObject& InObject, ESpawnOwnership InOwnership)
+			: Object(&InObject)
+			, Ownership(InOwnership)
+		{}
+
+		/** The object that has been spawned */
+		TWeakObjectPtr<UObject> Object;
+
+		/** What level of ownership this object was spawned with */
+		ESpawnOwnership Ownership;
+	};
+
+protected:
+	/** Called right before a spawned object with the specified ID and sequence instance is destroyed */
+	virtual void PreDestroyObject(UObject& Object, const FGuid& BindingId, FMovieSceneSequenceInstance& SequenceInstance);
+
+	/** Implementation function that destroys the specified object if it is an actor */
+	void DestroySpawnedObject(UObject& Object);
+
+	/** Destroy any objects that pass the specified predicate */
+	void DestroyObjects(IMovieScenePlayer& Player, TFunctionRef<bool(FMovieSceneSequenceInstance&, FSpawnedObject&)> Pred);
+	
+	
+	/** Sets of sequence instances that are active for the current update, and were active for the previous update */
+	TSet<TWeakPtr<FMovieSceneSequenceInstance>> ActiveInstances, PreviouslyActiveInstances;
+
+	/** The number of currently updating sequence instances. Used to keep track of the outermost sequence instance that triggered an update */
+	int32 CurrentlyUpdatingSequenceCount;
+
+	/** Register of spawned objects */
+	TMap<FMovieSceneSpawnRegisterKey, FSpawnedObject> Register;
+};
\ No newline at end of file
diff --git a/Engine/Source/Runtime/MovieScene/Private/MovieScene.cpp b/Engine/Source/Runtime/MovieScene/Private/MovieScene.cpp
--- a/Engine/Source/Runtime/MovieScene/Private/MovieScene.cpp
+++ b/Engine/Source/Runtime/MovieScene/Private/MovieScene.cpp
@@ -20,11 +20,6 @@ UMovieScene::UMovieScene(const FObjectInitializer& ObjectInitializer)
 #endif
 }
 
-void UMovieScene::PostLoad()
-{
-	UpgradeTimeRanges();
-	Super::PostLoad();
-}
 
 #if WITH_EDITOR
 
@@ -123,9 +118,9 @@ bool UMovieScene::RemovePossessable( const FGuid& PossessableGuid )
 			Modify();
 
 			// Remove the parent-child link for a parent spawnable/child possessable if necessary
-			if (CurPossesable.GetParentSpawnable().IsValid())
+			if (CurPossesable.GetParent().IsValid())
 			{
-				FMovieSceneSpawnable* ParentSpawnable = FindSpawnable(CurPossesable.GetParentSpawnable());
+				FMovieSceneSpawnable* ParentSpawnable = FindSpawnable(CurPossesable.GetParent());
 				if (ParentSpawnable)
 				{
 					ParentSpawnable->RemoveChildPossessable(PossessableGuid);
@@ -224,14 +219,44 @@ FText UMovieScene::GetObjectDisplayName(const FGuid& ObjectId)
 	}
 #endif
 	return FText::GetEmpty();
+}
+
+
+#if WITH_EDITORONLY_DATA
+int32 UMovieScene::GetAllLabels(TArray<FString>& OutLabels) const
+{
+	for (const auto& LabelsPair : ObjectsToLabels)
+	{
+		for (const auto& Label : LabelsPair.Value.Strings)
+		{
+			OutLabels.AddUnique(Label);
+		}
+	}
+
+	return OutLabels.Num();
+}
+
+
+bool UMovieScene::LabelExists(const FString& Label) const
+{
+	for (const auto& LabelsPair : ObjectsToLabels)
+	{
+		if (LabelsPair.Value.Strings.Contains(Label))
+		{
+			return true;
 		}
+	}
+
+	return false;
+}
+#endif
 
 
 TRange<float> UMovieScene::GetPlaybackRange() const
 {
 	check(PlaybackRange.HasLowerBound() && PlaybackRange.HasUpperBound());
 	return PlaybackRange;
-	}
+}
 
 
 void UMovieScene::SetPlaybackRange(float Start, float End)
@@ -496,6 +521,45 @@ void UMovieScene::UpgradeTimeRanges()
 }
 
 
+/* UObject interface
+ *****************************************************************************/
+
+void UMovieScene::PostLoad()
+{
+	UpgradeTimeRanges();
+	Super::PostLoad();
+}
+
+
+void UMovieScene::PreSave()
+{
+	Super::PreSave();
+
+#if WITH_EDITORONLY_DATA
+	// compress meta data mappings prior to saving
+	for (auto It = ObjectsToDisplayNames.CreateIterator(); It; ++It)
+	{
+		FGuid ObjectId;
+
+		if (!FGuid::Parse(It.Key(), ObjectId) || ((FindPossessable(ObjectId) == nullptr) && (FindSpawnable(ObjectId) == nullptr)))
+		{
+			It.RemoveCurrent();
+		}
+	}
+
+	for (auto It = ObjectsToLabels.CreateIterator(); It; ++It)
+	{
+		FGuid ObjectId;
+
+		if (!FGuid::Parse(It.Key(), ObjectId) || ((FindPossessable(ObjectId) == nullptr) && (FindSpawnable(ObjectId) == nullptr)))
+		{
+			It.RemoveCurrent();
+		}
+	}
+#endif
+}
+
+
 /* UMovieScene implementation
  *****************************************************************************/
 
diff --git a/Engine/Source/Runtime/MovieScene/Private/MovieSceneCommonHelpers.cpp b/Engine/Source/Runtime/MovieScene/Private/MovieSceneCommonHelpers.cpp
--- a/Engine/Source/Runtime/MovieScene/Private/MovieSceneCommonHelpers.cpp
+++ b/Engine/Source/Runtime/MovieScene/Private/MovieSceneCommonHelpers.cpp
@@ -5,7 +5,7 @@
 #include "MovieSceneCommonHelpers.h"
 
 
-TArray<UMovieSceneSection*> MovieSceneHelpers::GetTraversedSections( const TArray<UMovieSceneSection*>& Sections, float CurrentTime, float PreviousTime )
+TArray<UMovieSceneSection*> MovieSceneHelpers::GetAllTraversedSections( const TArray<UMovieSceneSection*>& Sections, float CurrentTime, float PreviousTime )
 {
 	TArray<UMovieSceneSection*> TraversedSections;
 
@@ -27,6 +27,37 @@ TArray<UMovieSceneSection*> MovieSceneHelpers::GetTraversedSections( const TArra
 	return TraversedSections;
 }
 
+TArray<UMovieSceneSection*> MovieSceneHelpers::GetTraversedSections( const TArray<UMovieSceneSection*>& Sections, float CurrentTime, float PreviousTime )
+{
+	TArray<UMovieSceneSection*> TraversedSections = GetAllTraversedSections(Sections, CurrentTime, PreviousTime);
+
+	// Remove any overlaps that are underneath another
+	for (int32 RemoveAt = 0; RemoveAt < TraversedSections.Num(); )
+	{
+		UMovieSceneSection* Section = TraversedSections[RemoveAt];
+		
+		const bool bShouldRemove = TraversedSections.ContainsByPredicate([=](UMovieSceneSection* OtherSection){
+			if (Section->GetRowIndex() == OtherSection->GetRowIndex() &&
+				Section->GetRange().Overlaps(OtherSection->GetRange()) &&
+				Section->GetOverlapPriority() < OtherSection->GetOverlapPriority())
+			{
+				return true;
+			}
+			return false;
+		});
+		
+		if (bShouldRemove)
+		{
+			TraversedSections.RemoveAt(RemoveAt, 1, false);
+		}
+		else
+		{
+			++RemoveAt;
+		}
+	}
+
+	return TraversedSections;
+}
 
 UMovieSceneSection* MovieSceneHelpers::FindSectionAtTime( const TArray<UMovieSceneSection*>& Sections, float Time )
 {
diff --git a/Engine/Source/Runtime/MovieScene/Private/MovieSceneSection.cpp b/Engine/Source/Runtime/MovieScene/Private/MovieSceneSection.cpp
--- a/Engine/Source/Runtime/MovieScene/Private/MovieSceneSection.cpp
+++ b/Engine/Source/Runtime/MovieScene/Private/MovieSceneSection.cpp
@@ -9,6 +9,7 @@ UMovieSceneSection::UMovieSceneSection( const FObjectInitializer& ObjectInitiali
 	, StartTime(0.0f)
 	, EndTime(0.0f)
 	, RowIndex(0)
+	, OverlapPriority(0)
 	, bIsActive(true)
 	, bIsLocked(false)
 	, bIsInfinite(false)
diff --git a/Engine/Source/Runtime/MovieScene/Private/MovieSceneSequenceInstance.cpp b/Engine/Source/Runtime/MovieScene/Private/MovieSceneSequenceInstance.cpp
--- a/Engine/Source/Runtime/MovieScene/Private/MovieSceneSequenceInstance.cpp
+++ b/Engine/Source/Runtime/MovieScene/Private/MovieSceneSequenceInstance.cpp
@@ -10,6 +10,62 @@ FMovieSceneSequenceInstance::FMovieSceneSequenceInstance(const UMovieSceneSequen
 	TimeRange = MovieSceneSequence->GetMovieScene()->GetPlaybackRange();
 }
 
+FMovieSceneSequenceInstance::~FMovieSceneSequenceInstance()
+{
+	for (auto& Pair : SpawnedObjects)
+	{
+		ensureMsgf(!Pair.Value.Get(), TEXT("Sequence instance is being destroyed but still has spawned objects"));
+	}
+}
+
+FGuid FMovieSceneSequenceInstance::FindObjectId(UObject& Object) const
+{
+	for (auto& Pair : ObjectBindingInstances)
+	{
+		if (Pair.Value.RuntimeObjects.Contains(&Object))
+		{
+			return Pair.Key;
+		}
+	}
+
+	return FGuid();
+}
+
+FGuid FMovieSceneSequenceInstance::FindParentObjectId(UObject& Object) const
+{
+	UObject* ParentObject = MovieSceneSequence->GetParentObject(&Object);
+	if (ParentObject)
+	{
+		return FindObjectId(*ParentObject);
+	}
+	return FGuid();
+}
+
+UObject* FMovieSceneSequenceInstance::FindObject(const FGuid& ObjectId, const IMovieScenePlayer& Player) const
+{
+	// Attempt to find a possessable first
+	FMovieScenePossessable* Possessable = MovieSceneSequence->GetMovieScene()->FindPossessable(ObjectId);
+	if (Possessable)
+	{
+		UObject* ParentObject = Player.GetPlaybackContext();
+		if (Possessable->GetParent().IsValid())
+		{
+			ParentObject = FindObject(Possessable->GetParent(), Player);
+		}
+
+		return MovieSceneSequence->FindPossessableObject(ObjectId, ParentObject);
+	}
+	else
+	{
+		return FindSpawnedObject(ObjectId);
+	}
+}
+
+UObject* FMovieSceneSequenceInstance::FindSpawnedObject(const FGuid& ObjectId) const
+{
+	const TWeakObjectPtr<>* SpawnedObject = SpawnedObjects.Find(ObjectId);
+	return SpawnedObject ? SpawnedObject->Get() : nullptr;
+}
 
 void FMovieSceneSequenceInstance::SaveState(class IMovieScenePlayer& Player)
 {
@@ -44,6 +100,8 @@ void FMovieSceneSequenceInstance::SaveState(class IMovieScenePlayer& Player)
 
 void FMovieSceneSequenceInstance::RestoreState(class IMovieScenePlayer& Player)
 {
+	Player.GetSpawnRegister().DestroyObjectsSpawnedByInstance(*this, Player);
+
 	TMap<FGuid, FMovieSceneObjectBindingInstance>::TIterator ObjectIt = ObjectBindingInstances.CreateIterator();
 	for (; ObjectIt; ++ObjectIt)
 	{
@@ -75,9 +133,32 @@ void FMovieSceneSequenceInstance::RestoreState(class IMovieScenePlayer& Player)
 
 void FMovieSceneSequenceInstance::Update( float Position, float LastPosition, class IMovieScenePlayer& Player )
 {
+	// Remove any stale runtime objects
+	TMap<FGuid, FMovieSceneObjectBindingInstance>::TIterator ObjectIt = ObjectBindingInstances.CreateIterator();
+	for(; ObjectIt; ++ObjectIt )
+	{
+		FMovieSceneObjectBindingInstance& ObjectBindingInstance = ObjectIt.Value();
+		for (int32 ObjectIndex = 0; ObjectIndex < ObjectBindingInstance.RuntimeObjects.Num(); )
+		{
+			UObject* RuntimeObject = ObjectBindingInstance.RuntimeObjects[ObjectIndex];
+			if (RuntimeObject == nullptr || RuntimeObject->IsPendingKill())
+			{
+				ObjectBindingInstance.RuntimeObjects.RemoveAt(ObjectIndex);
+			}
+			else
+			{
+				++ObjectIndex;
+			}
+		}
+	}
+
+	Player.GetSpawnRegister().PreUpdateSequenceInstance(*this, Player);
+
 	UpdateInternal(Position, LastPosition, Player, MSUP_PreUpdate);
 	UpdateInternal(Position, LastPosition, Player, MSUP_Update);
 	UpdateInternal(Position, LastPosition, Player, MSUP_PostUpdate);
+
+	Player.GetSpawnRegister().PostUpdateSequenceInstance(*this, Player);
 }
 
 void FMovieSceneSequenceInstance::UpdateInternal( float Position, float LastPosition, class IMovieScenePlayer& Player, EMovieSceneUpdatePass UpdatePass )
@@ -180,6 +261,7 @@ void FMovieSceneSequenceInstance::RefreshInstance( IMovieScenePlayer& Player )
 		RefreshInstanceMap( Tracks, BindingInstance.RuntimeObjects, BindingInstance.TrackInstances, Player );
 	}
 
+	IMovieSceneSpawnRegister& SpawnRegister = Player.GetSpawnRegister();
 
 	// Remove object binding instances which are no longer bound
 	TMap<FGuid, FMovieSceneObjectBindingInstance>::TIterator It = ObjectBindingInstances.CreateIterator();
@@ -187,6 +269,8 @@ void FMovieSceneSequenceInstance::RefreshInstance( IMovieScenePlayer& Player )
 	{
 		if( !FoundObjectBindings.Contains( It.Key() ) )
 		{
+			SpawnRegister.DestroySpawnedObject(It.Key(), *this, Player);
+
 			// The instance no longer is bound to an existing guid
 			It.RemoveCurrent();
 		}
@@ -265,7 +349,20 @@ void FMovieSceneSequenceInstance::UpdateObjectBinding(const FGuid& ObjectId, IMo
 
 	// Update the runtime objects
 	BindingInstance->RuntimeObjects.Reset();
-	Player.GetRuntimeObjects(SharedThis(this), BindingInstance->ObjectGuid, BindingInstance->RuntimeObjects);
+
+	TWeakObjectPtr<UObject>* WeakSpawnedObject = SpawnedObjects.Find(ObjectId);
+	if (WeakSpawnedObject)
+	{
+		UObject* SpawnedObject = WeakSpawnedObject->Get();
+		if (SpawnedObject)
+		{
+			BindingInstance->RuntimeObjects.Add(SpawnedObject);
+		}
+	}
+	else
+	{
+		Player.GetRuntimeObjects(SharedThis(this), BindingInstance->ObjectGuid, BindingInstance->RuntimeObjects);
+	}
 
 	const FMovieSceneBinding* ObjectBinding = Sequence->GetMovieScene()->GetBindings().FindByPredicate([&](const FMovieSceneBinding& In){
 		return In.GetObjectGuid() == ObjectId;
@@ -278,30 +375,23 @@ void FMovieSceneSequenceInstance::UpdateObjectBinding(const FGuid& ObjectId, IMo
 	}
 }
 
-void FMovieSceneSequenceInstance::SpawnObject(const FGuid& ObjectId, IMovieScenePlayer& Player)
+void FMovieSceneSequenceInstance::OnObjectSpawned(const FGuid& ObjectId, UObject& SpawnedObject, IMovieScenePlayer& Player)
 {
-	UMovieSceneSequence* Sequence = GetSequence();
 	auto* BindingInstance = ObjectBindingInstances.Find(ObjectId);
 
-	if (!Sequence || !BindingInstance)
+	if (!BindingInstance)
 	{
 		return;
 	}
 
-	// Reset the runtime objects for the binding instance
-	BindingInstance->RuntimeObjects.Reset();
-
-	// Attempt to spawn the object
-	UObject* SpawnedObject = Sequence->SpawnObject(ObjectId);
-	if (!SpawnedObject)
-	{
-		return;
-	}
+	SpawnedObjects.Add(ObjectId, &SpawnedObject);
 
 	// Add it to the instance's runtime objects array, and update any child possessable binding instances
-	BindingInstance->RuntimeObjects.Emplace(SpawnedObject);
+	BindingInstance->RuntimeObjects.Reset();
+	BindingInstance->RuntimeObjects.Emplace(&SpawnedObject);
 
-	FMovieSceneSpawnable* Spawnable = Sequence->GetMovieScene()->FindSpawnable(ObjectId);
+	UMovieSceneSequence* Sequence = GetSequence();
+	FMovieSceneSpawnable* Spawnable = Sequence ? Sequence->GetMovieScene()->FindSpawnable(ObjectId) : nullptr;
 	if (Spawnable)
 	{
 		for (const FGuid& Child : Spawnable->GetChildPossessables())
@@ -311,22 +401,22 @@ void FMovieSceneSequenceInstance::SpawnObject(const FGuid& ObjectId, IMovieScene
 	}
 }
 
-void FMovieSceneSequenceInstance::DestroySpawnedObject(const FGuid& ObjectId, IMovieScenePlayer& Player)
+void FMovieSceneSequenceInstance::OnSpawnedObjectDestroyed(const FGuid& ObjectId, IMovieScenePlayer& Player)
 {
-	UMovieSceneSequence* Sequence = GetSequence();
 	auto* BindingInstance = ObjectBindingInstances.Find(ObjectId);
-	
-	if (!Sequence || !BindingInstance)
+	if (!BindingInstance)
 	{
 		return;
 	}
 
+	SpawnedObjects.Remove(ObjectId);
+
 	// Destroy the object
-	Sequence->DestroySpawnedObject(ObjectId);
 	BindingInstance->RuntimeObjects.Reset();
 
 	// Update any child possessable object bindings
-	FMovieSceneSpawnable* Spawnable = Sequence->GetMovieScene()->FindSpawnable(ObjectId);
+	UMovieSceneSequence* Sequence = GetSequence();
+	FMovieSceneSpawnable* Spawnable = Sequence ? Sequence->GetMovieScene()->FindSpawnable(ObjectId) : nullptr;
 	if (Spawnable)
 	{
 		for (const FGuid& Child : Spawnable->GetChildPossessables())
diff --git a/Engine/Source/Runtime/MovieScene/Public/IMovieScenePlayer.h b/Engine/Source/Runtime/MovieScene/Public/IMovieScenePlayer.h
--- a/Engine/Source/Runtime/MovieScene/Public/IMovieScenePlayer.h
+++ b/Engine/Source/Runtime/MovieScene/Public/IMovieScenePlayer.h
@@ -2,6 +2,8 @@
 
 #pragma once
 
+#include "IMovieSceneSpawnRegister.h"
+
 namespace EMovieScenePlayerStatus
 {
 	enum Type
@@ -103,4 +105,18 @@ public:
 
 	/** @return whether the player is currently playing, scrubbing, etc. */
 	virtual EMovieScenePlayerStatus::Type GetPlaybackStatus() const = 0;
+
+	/**
+	 * Obtain an object responsible for managing movie scene spawnables
+	 */
+	virtual IMovieSceneSpawnRegister& GetSpawnRegister() { return NullRegister; }
+
+	/**
+	 * Access the playback context for this movie scene player
+	 */
+	virtual UObject* GetPlaybackContext() const { return nullptr; }
+
+private:
+	/** Null register that asserts on use */
+	FNullMovieSceneSpawnRegister NullRegister;
 };
