Copying //UE4/Dev-Blueprints to Dev-Main (//UE4/Dev-Main) @ 2781164

diff --git a/Engine/Source/Runtime/CoreUObject/Public/UObject/UnrealType.h b/Engine/Source/Runtime/CoreUObject/Public/UObject/UnrealType.h
--- a/Engine/Source/Runtime/CoreUObject/Public/UObject/UnrealType.h
+++ b/Engine/Source/Runtime/CoreUObject/Public/UObject/UnrealType.h
@@ -1762,7 +1762,7 @@ class COREUOBJECT_API UObjectPropertyBase : public UProperty
 	// UObjectPropertyBase interface
 public:
 
-	virtual FString GetCPPTypeCustom(FString* ExtendedTypeText, uint32 CPPExportFlags, UClass* ActualClass) const PURE_VIRTUAL(UObjectPropertyBase::GetCPPTypeCustom, return TEXT(""););
+	virtual FString GetCPPTypeCustom(FString* ExtendedTypeText, uint32 CPPExportFlags, const FString& InnerNativeTypeName) const PURE_VIRTUAL(UObjectPropertyBase::GetCPPTypeCustom, return TEXT(""););
 
 	/**
 	 * Parses a text buffer into an object reference.
@@ -1868,9 +1868,12 @@ public:
 	}
 	// End of UProperty interface
 
+	// TProperty::GetCPPType should not be used here
 	virtual FString GetCPPType(FString* ExtendedTypeText, uint32 CPPExportFlags) const override
 	{
-		return this->GetCPPTypeCustom(ExtendedTypeText, CPPExportFlags, UObjectPropertyBase::PropertyClass);
+		check(UObjectPropertyBase::PropertyClass);
+		return this->GetCPPTypeCustom(ExtendedTypeText, CPPExportFlags, 
+			FString::Printf(TEXT("%s%s"), UObjectPropertyBase::PropertyClass->GetPrefixCPP(), *UObjectPropertyBase::PropertyClass->GetName()));
 	}
 };
 
@@ -1916,7 +1919,7 @@ public:
 	}
 	virtual void SetObjectPropertyValue(void* PropertyValueAddress, UObject* Value) const override;
 
-	virtual FString GetCPPTypeCustom(FString* ExtendedTypeText, uint32 CPPExportFlags, UClass* ActualClass)  const override;
+	virtual FString GetCPPTypeCustom(FString* ExtendedTypeText, uint32 CPPExportFlags, const FString& InnerNativeTypeName)  const override;
 	// End of UObjectPropertyBase interface
 };
 
@@ -2050,7 +2053,7 @@ class COREUOBJECT_API UAssetObjectProperty : public TUObjectPropertyBase<FAssetP
 	{
 		return true;
 	}
-	virtual FString GetCPPTypeCustom(FString* ExtendedTypeText, uint32 CPPExportFlags, UClass* ActualClass)  const override;
+	virtual FString GetCPPTypeCustom(FString* ExtendedTypeText, uint32 CPPExportFlags, const FString& InnerNativeTypeName)  const override;
 
 private:
 	virtual uint32 GetValueTypeHashInternal(const void* Src) const override
@@ -2097,14 +2100,14 @@ class COREUOBJECT_API UClassProperty : public UObjectProperty
 	// Variables.
 	class UClass* MetaClass;
 public:
-	UClassProperty(ECppProperty, int32 InOffset, uint64 InFlags, UClass* InMetaClass)
-		: UObjectProperty(FObjectInitializer::Get(), EC_CppProperty, InOffset, InFlags, UClass::StaticClass())
+	UClassProperty(ECppProperty, int32 InOffset, uint64 InFlags, UClass* InMetaClass, UClass* InClassType)
+		: UObjectProperty(FObjectInitializer::Get(), EC_CppProperty, InOffset, InFlags, InClassType ? InClassType : UClass::StaticClass())
 		, MetaClass(InMetaClass)
 	{
 	}
 
-	UClassProperty( const FObjectInitializer& ObjectInitializer, ECppProperty, int32 InOffset, uint64 InFlags, UClass* InMetaClass )
-	:	UObjectProperty( ObjectInitializer, EC_CppProperty, InOffset, InFlags, UClass::StaticClass() )
+	UClassProperty(const FObjectInitializer& ObjectInitializer, ECppProperty, int32 InOffset, uint64 InFlags, UClass* InMetaClass, UClass* InClassType)
+		: UObjectProperty(ObjectInitializer, EC_CppProperty, InOffset, InFlags, InClassType ? InClassType : UClass::StaticClass())
 	,	MetaClass( InMetaClass )
 	{
 	}
@@ -2126,7 +2129,7 @@ public:
 	virtual bool SameType(const UProperty* Other) const override;
 	// End of UProperty interface
 
-	virtual FString GetCPPTypeCustom(FString* ExtendedTypeText, uint32 CPPExportFlags, UClass* ActualClass)  const override;
+	virtual FString GetCPPTypeCustom(FString* ExtendedTypeText, uint32 CPPExportFlags, const FString& InnerNativeTypeName)  const override;
 
 	/**
 	 * Setter function for this property's MetaClass member. Favor this function 
@@ -2180,13 +2183,29 @@ public:
 	// UObject interface
 	virtual void Serialize( FArchive& Ar ) override;
 	static void AddReferencedObjects(UObject* InThis, FReferenceCollector& Collector);
+	virtual void BeginDestroy() override;
 	// End of UObject interface
 
 	// UProperty interface
 	virtual bool SameType(const UProperty* Other) const override;
 	// End of UProperty interface
 
-	virtual FString GetCPPTypeCustom(FString* ExtendedTypeText, uint32 CPPExportFlags, UClass* ActualClass)  const override;
+	virtual FString GetCPPTypeCustom(FString* ExtendedTypeText, uint32 CPPExportFlags, const FString& InnerNativeTypeName)  const override;
+
+	/**
+	 * Setter function for this property's MetaClass member. Favor this function 
+	 * whilst loading (since, to handle circular dependencies, we defer some 
+	 * class loads and use a placeholder class instead). It properly handles 
+	 * deferred loading placeholder classes (so they can properly be replaced 
+	 * later).
+	 * 
+	 * @param  NewMetaClass    The MetaClass you want this property set with.
+	 */
+#if USE_CIRCULAR_DEPENDENCY_LOAD_DEFERRING
+	void SetMetaClass(UClass* NewMetaClass);
+#else
+	FORCEINLINE void SetMetaClass(UClass* NewMetaClass) { MetaClass = NewMetaClass; }
+#endif // USE_CIRCULAR_DEPENDENCY_LOAD_DEFERRING
 };
 
 /*-----------------------------------------------------------------------------
diff --git a/Engine/Source/Runtime/Engine/Classes/Animation/AnimBlueprintGeneratedClass.h b/Engine/Source/Runtime/Engine/Classes/Animation/AnimBlueprintGeneratedClass.h
--- a/Engine/Source/Runtime/Engine/Classes/Animation/AnimBlueprintGeneratedClass.h
+++ b/Engine/Source/Runtime/Engine/Classes/Animation/AnimBlueprintGeneratedClass.h
@@ -6,6 +6,7 @@
 #include "AnimSequenceBase.h"
 #include "Engine/BlueprintGeneratedClass.h"
 #include "AnimBlueprint.h"
+#include "AnimClassInterface.h"
 
 #include "AnimBlueprintGeneratedClass.generated.h"
 
@@ -160,7 +161,7 @@ namespace EPropertySearchMode
 #endif
 
 UCLASS()
-class ENGINE_API UAnimBlueprintGeneratedClass : public UBlueprintGeneratedClass
+class ENGINE_API UAnimBlueprintGeneratedClass : public UBlueprintGeneratedClass, public IAnimClassInterface
 {
 	GENERATED_UCLASS_BODY()
 
@@ -188,8 +189,25 @@ class ENGINE_API UAnimBlueprintGeneratedClass : public UBlueprintGeneratedClass
 	UPROPERTY()
 	TArray<FName> SyncGroupNames;
 
-	int32 GetSyncGroupIndex(FName SyncGroupName) const { return SyncGroupNames.IndexOfByKey(SyncGroupName); }
+public:
+
+	virtual const TArray<FBakedAnimationStateMachine>& GetBakedStateMachines() const override { return BakedStateMachines; }
+
+	virtual USkeleton* GetTargetSkeleton() const override { return TargetSkeleton; }
+
+	virtual const TArray<FAnimNotifyEvent>& GetAnimNotifies() const override { return AnimNotifies; }
+
+	virtual int32 GetRootAnimNodeIndex() const override { return RootAnimNodeIndex; }
+
+	virtual UStructProperty* GetRootAnimNodeProperty() const override { return RootAnimNodeProperty; }
+
+	virtual const TArray<UStructProperty*>& GetAnimNodeProperties() const override { return AnimNodeProperties; }
 
+	virtual const TArray<FName>& GetSyncGroupNames() const override { return SyncGroupNames; }
+
+	virtual int32 GetSyncGroupIndex(FName SyncGroupName) const override { return SyncGroupNames.IndexOfByKey(SyncGroupName); }
+
+public:
 #if WITH_EDITORONLY_DATA
 	FAnimBlueprintDebugData AnimBlueprintDebugData;
 
@@ -289,19 +307,14 @@ class ENGINE_API UAnimBlueprintGeneratedClass : public UBlueprintGeneratedClass
 	virtual void PurgeClass(bool bRecompilingOnLoad) override;
 	virtual uint8* GetPersistentUberGraphFrame(UObject* Obj, UFunction* FuncToCheck) const override;
 	// End of UClass interface
-
-#if WITH_EDITOR
-public:
-	class USkeleton* GetTargetSkeleton() const { return TargetSkeleton; }
-#endif
 };
 
 template<typename NodeType>
-NodeType* GetNodeFromPropertyIndex(UObject* AnimInstanceObject, const UAnimBlueprintGeneratedClass* AnimBlueprintClass, int32 PropertyIndex)
+NodeType* GetNodeFromPropertyIndex(UObject* AnimInstanceObject, const IAnimClassInterface* AnimBlueprintClass, int32 PropertyIndex)
 {
 	if (PropertyIndex != INDEX_NONE)
 	{
-		UStructProperty* NodeProperty = AnimBlueprintClass->AnimNodeProperties[AnimBlueprintClass->AnimNodeProperties.Num() - 1 - PropertyIndex]; //@TODO: Crazysauce
+		UStructProperty* NodeProperty = AnimBlueprintClass->GetAnimNodeProperties()[AnimBlueprintClass->GetAnimNodeProperties().Num() - 1 - PropertyIndex]; //@TODO: Crazysauce
 		check(NodeProperty->Struct == NodeType::StaticStruct());
 		return NodeProperty->ContainerPtrToValuePtr<NodeType>(AnimInstanceObject);
 	}
diff --git a/Engine/Source/Runtime/Engine/Classes/Animation/AnimClassData.h b/Engine/Source/Runtime/Engine/Classes/Animation/AnimClassData.h
new file mode 100644
--- /dev/null
+++ b/Engine/Source/Runtime/Engine/Classes/Animation/AnimClassData.h
@@ -0,0 +1,71 @@
+// Copyright 1998-2015 Epic Games, Inc. All Rights Reserved.
+
+#pragma once
+
+#include "AnimClassInterface.h"
+#include "AnimClassData.generated.h"
+
+UCLASS()
+class ENGINE_API UAnimClassData : public UObject, public IAnimClassInterface
+{
+	GENERATED_BODY()
+public:
+	// List of state machines present in this blueprint class
+	UPROPERTY()
+	TArray<FBakedAnimationStateMachine> BakedStateMachines;
+
+	/** Target skeleton for this blueprint class */
+	UPROPERTY()
+	class USkeleton* TargetSkeleton;
+
+	/** A list of anim notifies that state machines (or anything else) may reference */
+	UPROPERTY()
+	TArray<FAnimNotifyEvent> AnimNotifies;
+
+	// The index of the root node in the animation tree
+	UPROPERTY()
+	int32 RootAnimNodeIndex;
+	
+	UPROPERTY()
+	UStructProperty* RootAnimNodeProperty;
+
+	// The array of anim nodes
+	UPROPERTY()
+	TArray<UStructProperty*> AnimNodeProperties;
+
+	// Array of sync group names in the order that they are requested during compile
+	UPROPERTY()
+	TArray<FName> SyncGroupNames;
+
+public:
+
+	virtual const TArray<FBakedAnimationStateMachine>& GetBakedStateMachines() const override { return BakedStateMachines; }
+
+	virtual USkeleton* GetTargetSkeleton() const override { return TargetSkeleton; }
+
+	virtual const TArray<FAnimNotifyEvent>& GetAnimNotifies() const override { return AnimNotifies; }
+
+	virtual int32 GetRootAnimNodeIndex() const override { return RootAnimNodeIndex; }
+
+	virtual UStructProperty* GetRootAnimNodeProperty() const override { return RootAnimNodeProperty; }
+
+	virtual const TArray<UStructProperty*>& GetAnimNodeProperties() const override { return AnimNodeProperties; }
+
+	virtual const TArray<FName>& GetSyncGroupNames() const override { return SyncGroupNames; }
+
+	virtual int32 GetSyncGroupIndex(FName SyncGroupName) const override { return SyncGroupNames.IndexOfByKey(SyncGroupName); }
+
+#if WITH_EDITOR
+	void CopyFrom(IAnimClassInterface* AnimClass)
+	{
+		check(AnimClass);
+		BakedStateMachines = AnimClass->GetBakedStateMachines();
+		TargetSkeleton = AnimClass->GetTargetSkeleton();
+		AnimNotifies = AnimClass->GetAnimNotifies();
+		RootAnimNodeIndex = AnimClass->GetRootAnimNodeIndex();
+		RootAnimNodeProperty = AnimClass->GetRootAnimNodeProperty();
+		AnimNodeProperties = AnimClass->GetAnimNodeProperties();
+		SyncGroupNames = AnimClass->GetSyncGroupNames();
+	}
+#endif // WITH_EDITOR
+};
\ No newline at end of file
diff --git a/Engine/Source/Runtime/Engine/Classes/Animation/AnimClassInterface.h b/Engine/Source/Runtime/Engine/Classes/Animation/AnimClassInterface.h
new file mode 100644
--- /dev/null
+++ b/Engine/Source/Runtime/Engine/Classes/Animation/AnimClassInterface.h
@@ -0,0 +1,59 @@
+// Copyright 1998-2015 Epic Games, Inc. All Rights Reserved.
+
+#pragma once
+
+#include "AnimStateMachineTypes.h"
+#include "AnimSequenceBase.h"
+
+#include "AnimClassInterface.generated.h"
+
+struct FBakedAnimationStateMachine;
+class  USkeleton;
+struct FAnimNotifyEvent;
+
+UINTERFACE()
+class ENGINE_API UAnimClassInterface : public UInterface
+{
+	GENERATED_BODY()
+};
+
+class ENGINE_API IAnimClassInterface
+{
+	GENERATED_BODY()
+public:
+	virtual const TArray<FBakedAnimationStateMachine>& GetBakedStateMachines() const = 0;
+	virtual const TArray<FAnimNotifyEvent>& GetAnimNotifies() const = 0;
+	virtual const TArray<UStructProperty*>& GetAnimNodeProperties() const = 0;
+	virtual const TArray<FName>& GetSyncGroupNames() const = 0;
+
+	virtual USkeleton* GetTargetSkeleton() const = 0;
+	virtual int32 GetRootAnimNodeIndex() const = 0;
+	virtual UStructProperty* GetRootAnimNodeProperty() const = 0;
+	virtual int32 GetSyncGroupIndex(FName SyncGroupName) const = 0;
+
+	static IAnimClassInterface* GetFromClass(UClass* InClass)
+	{
+		if (auto AnimClassInterface = Cast<IAnimClassInterface>(InClass))
+		{
+			return AnimClassInterface;
+		}
+		if (auto DynamicClass = Cast<UDynamicClass>(InClass))
+		{
+			return CastChecked<IAnimClassInterface>(DynamicClass->AnimClassImplementation, ECastCheckedType::NullAllowed);
+		}
+		return nullptr;
+	}
+
+	static UClass* GetActualAnimClass(IAnimClassInterface* AnimClassInterface)
+	{
+		if (UClass* ActualAnimClass = Cast<UClass>(AnimClassInterface))
+		{
+			return ActualAnimClass;
+		}
+		if (UObject* AsObject = Cast<UObject>(AnimClassInterface))
+		{
+			return Cast<UClass>(AsObject->GetOuter());
+		}
+		return nullptr;
+	}
+};
diff --git a/Engine/Source/Runtime/Engine/Classes/Animation/AnimInstance.h b/Engine/Source/Runtime/Engine/Classes/Animation/AnimInstance.h
--- a/Engine/Source/Runtime/Engine/Classes/Animation/AnimInstance.h
+++ b/Engine/Source/Runtime/Engine/Classes/Animation/AnimInstance.h
@@ -9,6 +9,7 @@
 #include "BonePose.h"
 #include "Animation/AnimTypes.h"
 #include "Animation/AnimNotifyQueue.h"
+#include "Animation/AnimClassInterface.h"
 #include "AnimInstance.generated.h"
 
 struct FAnimMontageInstance;
@@ -860,7 +861,7 @@ public:
 	FAnimNode_StateMachine* GetStateMachineInstanceFromName(FName MachineName);
 
 	/** Get the machine description for the specified instance. Does not rely on PRIVATE_MachineDescription being initialized */
-	const FBakedAnimationStateMachine* GetMachineDescription(UAnimBlueprintGeneratedClass* AnimBlueprintClass, FAnimNode_StateMachine* MachineInstance);
+	const FBakedAnimationStateMachine* GetMachineDescription(IAnimClassInterface* AnimBlueprintClass, FAnimNode_StateMachine* MachineInstance);
 
 	/** Returns the baked sync group index from the compile step */
 	int32 GetSyncGroupIndexFromName(FName SyncGroupName) const;
diff --git a/Engine/Source/Runtime/Engine/Classes/Animation/AnimNodeBase.h b/Engine/Source/Runtime/Engine/Classes/Animation/AnimNodeBase.h
--- a/Engine/Source/Runtime/Engine/Classes/Animation/AnimNodeBase.h
+++ b/Engine/Source/Runtime/Engine/Classes/Animation/AnimNodeBase.h
@@ -33,8 +33,13 @@ public:
 public:
 	// Get the Blueprint Generated Class associated with this context, if there is one.
 	// Note: This can return NULL, so check the result.
+	DEPRECATED(4.11, "GetAnimBlueprintClass() is deprecated, UAnimBlueprintGeneratedClass should not be directly used at runtime. Please use GetAnimClass() instead.")
 	ENGINE_API UAnimBlueprintGeneratedClass* GetAnimBlueprintClass() const;
 
+	// Get the Blueprint IAnimClassInterface associated with this context, if there is one.
+	// Note: This can return NULL, so check the result.
+	ENGINE_API IAnimClassInterface* GetAnimClass() const;
+
 #if WITH_EDITORONLY_DATA
 	// Get the AnimBlueprint associated with this context, if there is one.
 	// Note: This can return NULL, so check the result.
diff --git a/Engine/Source/Runtime/Engine/Classes/Animation/AnimNode_StateMachine.h b/Engine/Source/Runtime/Engine/Classes/Animation/AnimNode_StateMachine.h
--- a/Engine/Source/Runtime/Engine/Classes/Animation/AnimNode_StateMachine.h
+++ b/Engine/Source/Runtime/Engine/Classes/Animation/AnimNode_StateMachine.h
@@ -202,7 +202,7 @@ public:
 	bool IsValidTransitionIndex(int32 TransitionIndex) const;
 
 	/** Cache the internal machine description */
-	void CacheMachineDescription(UAnimBlueprintGeneratedClass* AnimBlueprintClass);
+	void CacheMachineDescription(IAnimClassInterface* AnimBlueprintClass);
 
 protected:
 	// Tries to get the instance information for the state machine
diff --git a/Engine/Source/Runtime/Engine/Classes/Components/ActorComponent.h b/Engine/Source/Runtime/Engine/Classes/Components/ActorComponent.h
--- a/Engine/Source/Runtime/Engine/Classes/Components/ActorComponent.h
+++ b/Engine/Source/Runtime/Engine/Classes/Components/ActorComponent.h
@@ -127,7 +127,7 @@ public:
 	UPROPERTY()
 	uint32 bInstanceComponent_DEPRECATED:1;
 
-	/** Whether to the component is activated at creation or must be explicitly activated. */
+	/** Whether the component is activated at creation or must be explicitly activated. */
 	UPROPERTY(EditAnywhere, BlueprintReadOnly, Category=Activation)
 	uint32 bAutoActivate:1;
 
diff --git a/Engine/Source/Runtime/Engine/Classes/Components/SkeletalMeshComponent.h b/Engine/Source/Runtime/Engine/Classes/Components/SkeletalMeshComponent.h
--- a/Engine/Source/Runtime/Engine/Classes/Components/SkeletalMeshComponent.h
+++ b/Engine/Source/Runtime/Engine/Classes/Components/SkeletalMeshComponent.h
@@ -327,10 +327,14 @@ public:
 	class UAnimBlueprint* AnimationBlueprint_DEPRECATED;
 #endif
 
-	/* The AnimBlueprint class to use. Use 'SetAnimInstanceClass' to change at runtime. */
+	DEPRECATED(4.11, "This property is deprecated. Please use AnimClass instead")
 	UPROPERTY(EditAnywhere, BlueprintReadOnly, Category = Animation)
 	class UAnimBlueprintGeneratedClass* AnimBlueprintGeneratedClass;
 
+	/* The AnimBlueprint class to use. Use 'SetAnimInstanceClass' to change at runtime. */
+	UPROPERTY(EditAnywhere, BlueprintReadOnly, Category = Animation)
+	class TSubclassOf<UAnimInstance> AnimClass;
+
 	/** The active animation graph program instance. */
 	UPROPERTY(transient)
 	class UAnimInstance* AnimScriptInstance;
diff --git a/Engine/Source/Runtime/Engine/Classes/EdGraph/EdGraphPin.h b/Engine/Source/Runtime/Engine/Classes/EdGraph/EdGraphPin.h
--- a/Engine/Source/Runtime/Engine/Classes/EdGraph/EdGraphPin.h
+++ b/Engine/Source/Runtime/Engine/Classes/EdGraph/EdGraphPin.h
@@ -257,7 +257,7 @@ class UEdGraphPin : public UObject
 public:
 	// UObject interface
 	virtual bool IsSafeForRootSet() const override { return false; }
-	virtual void Serialize(FArchive& Ar) override;
+	virtual void PostLoad() override;
 	// End UObject interface
 
 	/** Create a link. Note, this does not check that schema allows it, and will not break any existing connections */
diff --git a/Engine/Source/Runtime/Engine/Classes/Engine/Blueprint.h b/Engine/Source/Runtime/Engine/Classes/Engine/Blueprint.h
--- a/Engine/Source/Runtime/Engine/Classes/Engine/Blueprint.h
+++ b/Engine/Source/Runtime/Engine/Classes/Engine/Blueprint.h
@@ -88,6 +88,9 @@ public:
 	/** Whether or not this compile is for a duplicated blueprint */
 	bool bIsDuplicationInstigated;
 
+	/** Whether or not this compile should emit instrumentation events */
+	bool bAddInstrumentation;
+
 	TSharedPtr<FString> OutHeaderSourceCode;
 	TSharedPtr<FString> OutCppSourceCode;
 
@@ -103,6 +106,11 @@ public:
 			|| (CompileType == EKismetCompileType::Cpp);
 	}
 
+	bool IsInstrumentationActive() const
+	{
+		return bAddInstrumentation && DoesRequireBytecodeGeneration();
+	}
+
 	/** Whether or not this compile type should operate on the generated class of the blueprint, as opposed to just the skeleton */
 	bool IsGeneratedClassCompileType() const
 	{
@@ -114,6 +122,7 @@ public:
 		, bSaveIntermediateProducts(false)
 		, bRegenerateSkelton(true)
 		, bIsDuplicationInstigated(false)
+		, bAddInstrumentation(false)
 	{
 	};
 };
diff --git a/Engine/Source/Runtime/Engine/Classes/Engine/BlueprintGeneratedClass.h b/Engine/Source/Runtime/Engine/Classes/Engine/BlueprintGeneratedClass.h
--- a/Engine/Source/Runtime/Engine/Classes/Engine/BlueprintGeneratedClass.h
+++ b/Engine/Source/Runtime/Engine/Classes/Engine/BlueprintGeneratedClass.h
@@ -431,6 +431,10 @@ public:
 	UPROPERTY()
 	TArray<FEventGraphFastCallPair> FastCallPairs_DEPRECATED;
 
+	// If this Generated Class has instrumentation
+	UPROPERTY()
+	bool bHasInstrumentation;
+
 #if WITH_EDITORONLY_DATA
 	UPROPERTY(Transient)
 	UObject* OverridenArchetypeForCDO;
@@ -475,6 +479,10 @@ public:
 	virtual void Bind() override;
 	virtual void GetRequiredPreloadDependencies(TArray<UObject*>& DependenciesOut) override;
 	virtual UObject* FindArchetype(UClass* ArchetypeClass, const FName ArchetypeName) const override;
+	virtual bool HasInstrumentation() const override 
+	{
+		return bHasInstrumentation; 
+	}
 	// End UClass interface
 
 	static void AddReferencedObjectsInUbergraphFrame(UObject* InThis, FReferenceCollector& Collector);
diff --git a/Engine/Source/Runtime/Engine/Classes/Matinee/InterpTrackFloatAnimBPParam.h b/Engine/Source/Runtime/Engine/Classes/Matinee/InterpTrackFloatAnimBPParam.h
--- a/Engine/Source/Runtime/Engine/Classes/Matinee/InterpTrackFloatAnimBPParam.h
+++ b/Engine/Source/Runtime/Engine/Classes/Matinee/InterpTrackFloatAnimBPParam.h
@@ -8,15 +8,20 @@ class UInterpTrackFloatAnimBPParam : public UInterpTrackFloatBase
 {
 	GENERATED_UCLASS_BODY()
 	
+	DEPRECATED(4.11, "This property is deprecated. Please use AnimClass instead")
+	UPROPERTY(EditAnywhere, Category = InterpTrackFloatAnimBPParam)
+	class UAnimBlueprintGeneratedClass* AnimBlueprintClass;
+
 	/** Materials whose parameters we want to change and the references to those materials. */
 	UPROPERTY(EditAnywhere, Category=InterpTrackFloatAnimBPParam)
-	class UAnimBlueprintGeneratedClass* AnimBlueprintClass;
+	TSubclassOf<UAnimInstance> AnimClass;
 
 	/** Name of parameter in the MaterialInstance which this track will modify over time. */
 	UPROPERTY(EditAnywhere, Category=InterpTrackFloatAnimBPParam)
 	FName ParamName;
 
 	//~ Begin UObject Interface
+	virtual void Serialize(FArchive& Ar) override;
 #if WITH_EDITOR
 	virtual void PostEditChangeProperty(FPropertyChangedEvent& PropertyChangedEvent) override;
 #endif // WITH_EDITOR
diff --git a/Engine/Source/Runtime/Engine/Private/ActorEditor.cpp b/Engine/Source/Runtime/Engine/Private/ActorEditor.cpp
--- a/Engine/Source/Runtime/Engine/Private/ActorEditor.cpp
+++ b/Engine/Source/Runtime/Engine/Private/ActorEditor.cpp
@@ -336,6 +336,15 @@ TSharedPtr<ITransactionObjectAnnotation> AActor::GetTransactionAnnotation() cons
 
 void AActor::PreEditUndo()
 {
+	// Check if this Actor needs to be re-instanced
+	UClass* OldClass = GetClass();
+	UClass* NewClass = OldClass->GetAuthoritativeClass();
+	if (NewClass != OldClass)
+	{
+		// Empty the OwnedComponents array, it's filled with invalid information
+		OwnedComponents.Empty();
+	}
+
 	// Since child actor components will rebuild themselves get rid of the Actor before we make changes
 	TInlineComponentArray<UChildActorComponent*> ChildActorComponents;
 	GetComponents(ChildActorComponents);
@@ -356,6 +365,20 @@ void AActor::PreEditUndo()
 
 void AActor::PostEditUndo()
 {
+	// Check if this Actor needs to be re-instanced
+	UClass* OldClass = GetClass();
+	if (OldClass->HasAnyClassFlags(CLASS_NewerVersionExists))
+	{
+		UClass* NewClass = OldClass->GetAuthoritativeClass();
+		if (!ensure(NewClass != OldClass))
+		{
+			UE_LOG(LogActor, Warning, TEXT("WARNING: %s is out of date and is the same as its AuthoritativeClass during PostEditUndo!"), *OldClass->GetName());
+		};
+
+		// Early exit, letting anything more occur would be invalid due to the REINST_ class
+		return;
+	}
+
 	// Notify LevelBounds actor that level bounding box might be changed
 	if (!IsTemplate())
 	{
@@ -381,6 +404,21 @@ void AActor::PostEditUndo(TSharedPtr<ITransactionObjectAnnotation> TransactionAn
 {
 	CurrentTransactionAnnotation = StaticCastSharedPtr<FActorTransactionAnnotation>(TransactionAnnotation);
 
+	// Check if this Actor needs to be re-instanced
+	UClass* OldClass = GetClass();
+	if (OldClass->HasAnyClassFlags(CLASS_NewerVersionExists))
+	{
+		UClass* NewClass = OldClass->GetAuthoritativeClass();
+		if (!ensure(NewClass != OldClass))
+		{
+			UE_LOG(LogActor, Warning, TEXT("WARNING: %s is out of date and is the same as its AuthoritativeClass during PostEditUndo!"), *OldClass->GetName());
+		};
+
+		// Early exit, letting anything more occur would be invalid due to the REINST_ class
+		return;
+	}
+
+
 	// Notify LevelBounds actor that level bounding box might be changed
 	if (!IsTemplate())
 	{
diff --git a/Engine/Source/Runtime/Engine/Private/Animation/AnimInstance.cpp b/Engine/Source/Runtime/Engine/Private/Animation/AnimInstance.cpp
--- a/Engine/Source/Runtime/Engine/Private/Animation/AnimInstance.cpp
+++ b/Engine/Source/Runtime/Engine/Private/Animation/AnimInstance.cpp
@@ -178,18 +178,18 @@ void UAnimInstance::InitializeAnimation()
 		CurrentSkeleton = NULL;
 	}
 
-	if (UAnimBlueprintGeneratedClass* AnimBlueprintClass = Cast<UAnimBlueprintGeneratedClass>(GetClass()))
+	if (IAnimClassInterface* AnimBlueprintClass = IAnimClassInterface::GetFromClass(GetClass()))
 	{
 #if WITH_EDITOR
 		LifeTimer = 0.0;
 		CurrentLifeTimerScrubPosition = 0.0;
 
-		if (UAnimBlueprint* Blueprint = Cast<UAnimBlueprint>(AnimBlueprintClass->ClassGeneratedBy))
+		if (UAnimBlueprint* Blueprint = Cast<UAnimBlueprint>(CastChecked<UAnimBlueprintGeneratedClass>(AnimBlueprintClass)->ClassGeneratedBy))
 		{
 			if (Blueprint->GetObjectBeingDebugged() == this)
 			{
 				// Reset the snapshot buffer
-				AnimBlueprintClass->GetAnimBlueprintDebugData().ResetSnapshotBuffer();
+				CastChecked<UAnimBlueprintGeneratedClass>(AnimBlueprintClass)->GetAnimBlueprintDebugData().ResetSnapshotBuffer();
 			}
 		}
 #endif
@@ -1315,13 +1315,14 @@ float UAnimInstance::GetAnimAssetPlayerTimeFromEndFraction(class UAnimationAsset
 
 float UAnimInstance::GetStateWeight(int32 MachineIndex, int32 StateIndex)
 {
-	if (UAnimBlueprintGeneratedClass* AnimBlueprintClass = Cast<UAnimBlueprintGeneratedClass>((UObject*)GetClass()))
+	if (IAnimClassInterface* AnimBlueprintClass = IAnimClassInterface::GetFromClass(GetClass()))
 	{
-		if ((MachineIndex >= 0) && (MachineIndex < AnimBlueprintClass->AnimNodeProperties.Num()))
+		const TArray<UStructProperty*>& AnimNodeProperties = AnimBlueprintClass->GetAnimNodeProperties();
+		if ((MachineIndex >= 0) && (MachineIndex < AnimNodeProperties.Num()))
 		{
-			const int32 InstancePropertyIndex = AnimBlueprintClass->AnimNodeProperties.Num() - 1 - MachineIndex; //@TODO: ANIMREFACTOR: Reverse indexing
+			const int32 InstancePropertyIndex = AnimNodeProperties.Num() - 1 - MachineIndex; //@TODO: ANIMREFACTOR: Reverse indexing
 
-			UStructProperty* MachineInstanceProperty = AnimBlueprintClass->AnimNodeProperties[InstancePropertyIndex];
+			UStructProperty* MachineInstanceProperty = AnimNodeProperties[InstancePropertyIndex];
 			checkSlow(MachineInstanceProperty->Struct->IsChildOf(FAnimNode_StateMachine::StaticStruct()));
 
 			FAnimNode_StateMachine* MachineInstance = MachineInstanceProperty->ContainerPtrToValuePtr<FAnimNode_StateMachine>(this);
@@ -1335,13 +1336,14 @@ float UAnimInstance::GetStateWeight(int32 MachineIndex, int32 StateIndex)
 
 float UAnimInstance::GetCurrentStateElapsedTime(int32 MachineIndex)
 {
-	if (UAnimBlueprintGeneratedClass* AnimBlueprintClass = Cast<UAnimBlueprintGeneratedClass>((UObject*)GetClass()))
+	if (IAnimClassInterface* AnimBlueprintClass = IAnimClassInterface::GetFromClass(GetClass()))
 	{
-		if ((MachineIndex >= 0) && (MachineIndex < AnimBlueprintClass->AnimNodeProperties.Num()))
+		const TArray<UStructProperty*>& AnimNodeProperties = AnimBlueprintClass->GetAnimNodeProperties();
+		if ((MachineIndex >= 0) && (MachineIndex < AnimNodeProperties.Num()))
 		{
-			const int32 InstancePropertyIndex = AnimBlueprintClass->AnimNodeProperties.Num() - 1 - MachineIndex; //@TODO: ANIMREFACTOR: Reverse indexing
+			const int32 InstancePropertyIndex = AnimNodeProperties.Num() - 1 - MachineIndex; //@TODO: ANIMREFACTOR: Reverse indexing
 
-			UStructProperty* MachineInstanceProperty = AnimBlueprintClass->AnimNodeProperties[InstancePropertyIndex];
+			UStructProperty* MachineInstanceProperty = AnimNodeProperties[InstancePropertyIndex];
 			checkSlow(MachineInstanceProperty->Struct->IsChildOf(FAnimNode_StateMachine::StaticStruct()));
 
 			FAnimNode_StateMachine* MachineInstance = MachineInstanceProperty->ContainerPtrToValuePtr<FAnimNode_StateMachine>(this);
@@ -1355,13 +1357,14 @@ float UAnimInstance::GetCurrentStateElapsedTime(int32 MachineIndex)
 
 FName UAnimInstance::GetCurrentStateName(int32 MachineIndex)
 {
-	if (UAnimBlueprintGeneratedClass* AnimBlueprintClass = Cast<UAnimBlueprintGeneratedClass>((UObject*)GetClass()))
+	if (IAnimClassInterface* AnimBlueprintClass = IAnimClassInterface::GetFromClass(GetClass()))
 	{
-		if ((MachineIndex >= 0) && (MachineIndex < AnimBlueprintClass->AnimNodeProperties.Num()))
+		const TArray<UStructProperty*>& AnimNodeProperties = AnimBlueprintClass->GetAnimNodeProperties();
+		if ((MachineIndex >= 0) && (MachineIndex < AnimNodeProperties.Num()))
 		{
-			const int32 InstancePropertyIndex = AnimBlueprintClass->AnimNodeProperties.Num() - 1 - MachineIndex; //@TODO: ANIMREFACTOR: Reverse indexing
+			const int32 InstancePropertyIndex = AnimNodeProperties.Num() - 1 - MachineIndex; //@TODO: ANIMREFACTOR: Reverse indexing
 
-			UStructProperty* MachineInstanceProperty = AnimBlueprintClass->AnimNodeProperties[InstancePropertyIndex];
+			UStructProperty* MachineInstanceProperty = AnimNodeProperties[InstancePropertyIndex];
 			checkSlow(MachineInstanceProperty->Struct->IsChildOf(FAnimNode_StateMachine::StaticStruct()));
 
 			FAnimNode_StateMachine* MachineInstance = MachineInstanceProperty->ContainerPtrToValuePtr<FAnimNode_StateMachine>(this);
@@ -2498,7 +2501,7 @@ FAnimNode_Base* UAnimInstance::GetNodeFromIndexUntyped(int32 NodeIdx, UScriptStr
 	return GetProxyOnGameThread<FAnimInstanceProxy>().GetNodeFromIndexUntyped(NodeIdx, RequiredStructType);
 }
 
-const FBakedAnimationStateMachine* UAnimInstance::GetMachineDescription(UAnimBlueprintGeneratedClass* AnimBlueprintClass, FAnimNode_StateMachine* MachineInstance)
+const FBakedAnimationStateMachine* UAnimInstance::GetMachineDescription(IAnimClassInterface* AnimBlueprintClass, FAnimNode_StateMachine* MachineInstance)
 {
 	return GetProxyOnGameThread<FAnimInstanceProxy>().GetMachineDescription(AnimBlueprintClass, MachineInstance);
 }
diff --git a/Engine/Source/Runtime/Engine/Private/Animation/AnimInstanceProxy.cpp b/Engine/Source/Runtime/Engine/Private/Animation/AnimInstanceProxy.cpp
--- a/Engine/Source/Runtime/Engine/Private/Animation/AnimInstanceProxy.cpp
+++ b/Engine/Source/Runtime/Engine/Private/Animation/AnimInstanceProxy.cpp
@@ -39,7 +39,7 @@ void FAnimInstanceProxy::Initialize(UAnimInstance* InAnimInstance)
 	// copy anim instance object if it has not already been set up
 	AnimInstanceObject = InAnimInstance;
 
-	AnimBlueprintClass = Cast<UAnimBlueprintGeneratedClass>(InAnimInstance->GetClass());
+	AnimClassInterface = IAnimClassInterface::GetFromClass(InAnimInstance->GetClass());
 
 	USkeletalMeshComponent* OwnerComponent = InAnimInstance->GetSkelMeshComponent();
 	if (OwnerComponent->SkeletalMesh != NULL)
@@ -51,12 +51,12 @@ void FAnimInstanceProxy::Initialize(UAnimInstance* InAnimInstance)
 		Skeleton = nullptr;
 	}
 
-	if(AnimBlueprintClass)
+	if (AnimClassInterface)
 	{
 		// Grab a pointer to the root node
-		if (AnimBlueprintClass->RootAnimNodeProperty != nullptr)
+		if (UStructProperty* RootAnimNodeProperty = AnimClassInterface->GetRootAnimNodeProperty())
 		{
-			RootNode = AnimBlueprintClass->RootAnimNodeProperty->ContainerPtrToValuePtr<FAnimNode_Base>(InAnimInstance);
+			RootNode = RootAnimNodeProperty->ContainerPtrToValuePtr<FAnimNode_Base>(InAnimInstance);
 		}
 		else
 		{
@@ -66,11 +66,11 @@ void FAnimInstanceProxy::Initialize(UAnimInstance* InAnimInstance)
 		// if no mesh, use Blueprint Skeleton
 		if (Skeleton == nullptr)
 		{
-			Skeleton = AnimBlueprintClass->TargetSkeleton;
+			Skeleton = AnimClassInterface->GetTargetSkeleton();
 		}
 
 #if WITH_EDITORONLY_DATA
-		if (UAnimBlueprint* Blueprint = Cast<UAnimBlueprint>(AnimBlueprintClass->ClassGeneratedBy))
+		if (UAnimBlueprint* Blueprint = Cast<UAnimBlueprint>(InAnimInstance->GetClass()->ClassGeneratedBy))
 		{
 			if (Blueprint->Status == BS_Error)
 			{
@@ -98,12 +98,12 @@ void FAnimInstanceProxy::InitializeRootNode()
 	if (RootNode != nullptr)
 	{
 		// cache any state machine descriptions we have
-		for(UStructProperty* Property : AnimBlueprintClass->AnimNodeProperties)
+		for(UStructProperty* Property : AnimClassInterface->GetAnimNodeProperties())
 		{
 			if(Property->Struct->IsChildOf(FAnimNode_StateMachine::StaticStruct()))
 			{
 				FAnimNode_StateMachine* StateMachine = Property->ContainerPtrToValuePtr<FAnimNode_StateMachine>(AnimInstanceObject);
-				StateMachine->CacheMachineDescription(AnimBlueprintClass);
+				StateMachine->CacheMachineDescription(AnimClassInterface);
 			}
 		}
 
@@ -143,12 +143,9 @@ void FAnimInstanceProxy::PreUpdate(const UAnimInstance* InAnimInstance, float De
 
 #if WITH_EDITORONLY_DATA
 	bIsBeingDebugged = false;
-	if(AnimBlueprintClass)
+	if (UAnimBlueprint* AnimBlueprint = GetAnimBlueprint())
 	{
-		if (UAnimBlueprint* AnimBlueprint = Cast<UAnimBlueprint>(AnimBlueprintClass->ClassGeneratedBy))
-		{
-			bIsBeingDebugged = (InAnimInstance && (AnimBlueprint->GetObjectBeingDebugged() == InAnimInstance));
-		}
+		bIsBeingDebugged = (InAnimInstance && (AnimBlueprint->GetObjectBeingDebugged() == InAnimInstance));
 	}
 #endif
 
@@ -300,7 +297,7 @@ void FAnimInstanceProxy::TickAssetPlayerInstances(float DeltaSeconds)
 			if (TickContext.CanUseMarkerPosition())
 			{
 				const FMarkerSyncAnimPosition& MarkerStart = TickContext.MarkerTickContext.GetMarkerSyncStartPosition();
-				FName SyncGroupName = AnimBlueprintClass->SyncGroupNames[GroupIndex];
+				FName SyncGroupName = GetAnimClassInterface()->GetSyncGroupNames()[GroupIndex];
 				FAnimTickRecord& GroupLeader = SyncGroup.ActivePlayers[SyncGroup.GroupLeaderIndex];
 				FString LeaderAnimName = GroupLeader.SourceAsset->GetName();
 
@@ -362,9 +359,9 @@ void FAnimInstanceProxy::AddAnimNotifies(const TArray<const FAnimNotifyEvent*>&
 
 int32 FAnimInstanceProxy::GetSyncGroupIndexFromName(FName SyncGroupName) const
 {
-	if (AnimBlueprintClass)
+	if (AnimClassInterface)
 	{
-		return AnimBlueprintClass->GetSyncGroupIndex(SyncGroupName);
+		return AnimClassInterface->GetSyncGroupIndex(SyncGroupName);
 	}
 	return INDEX_NONE;
 }
@@ -403,10 +400,10 @@ void FAnimInstanceProxy::AddAnimNotifyFromGeneratedClass(int32 NotifyIndex)
 		return;
 	}
 
-	if(AnimBlueprintClass)
+	if (AnimClassInterface)
 	{
-		check(AnimBlueprintClass->AnimNotifies.IsValidIndex(NotifyIndex));
-		const FAnimNotifyEvent* Notify =& AnimBlueprintClass->AnimNotifies[NotifyIndex];
+		check(AnimClassInterface->GetAnimNotifies().IsValidIndex(NotifyIndex));
+		const FAnimNotifyEvent* Notify = &AnimClassInterface->GetAnimNotifies()[NotifyIndex];
 		NotifyQueue.AnimNotifies.Add(Notify);
 	}
 }
@@ -480,7 +477,8 @@ void FAnimInstanceProxy::RegisterSlotNodeWithAnimInstance(FName SlotNodeName)
 	// then warn users, this is invalid
 	if (SlotWeightTracker.Contains(SlotNodeName))
 	{
-		FString ClassNameString = GetAnimBlueprintClass() ? GetAnimBlueprintClass()->GetName() : FString("Unavailable");
+		UClass* ActualAnimClass = IAnimClassInterface::GetActualAnimClass(GetAnimClassInterface());
+		FString ClassNameString = ActualAnimClass ? ActualAnimClass->GetName() : FString("Unavailable");
 
 		FMessageLog("AnimBlueprint").Warning(FText::Format(LOCTEXT("AnimInstance_SlotNode", "SLOTNODE: '{0}' in animation instance class {1} already exists. Remove duplicates from the animation graph for this class."), FText::FromString(SlotNodeName.ToString()), FText::FromString(ClassNameString)));
 		return;
@@ -549,13 +547,14 @@ float FAnimInstanceProxy::GetSlotRootMotionWeight(FName SlotNodeName) const
 FAnimNode_Base* FAnimInstanceProxy::GetCheckedNodeFromIndexUntyped(int32 NodeIdx, UScriptStruct* RequiredStructType)
 {
 	FAnimNode_Base* NodePtr = nullptr;
-	if (AnimBlueprintClass)
+	if (AnimClassInterface)
 	{
-		const int32 InstanceIdx = AnimBlueprintClass->AnimNodeProperties.Num() - 1 - NodeIdx;
+		const TArray<UStructProperty*>& AnimNodeProperties = AnimClassInterface->GetAnimNodeProperties();
+		const int32 InstanceIdx = AnimNodeProperties.Num() - 1 - NodeIdx;
 
-		if (AnimBlueprintClass->AnimNodeProperties.IsValidIndex(InstanceIdx))
+		if (AnimNodeProperties.IsValidIndex(InstanceIdx))
 		{
-			UStructProperty* NodeProperty = AnimBlueprintClass->AnimNodeProperties[InstanceIdx];
+			UStructProperty* NodeProperty = AnimNodeProperties[InstanceIdx];
 
 			if (NodeProperty->Struct->IsChildOf(RequiredStructType))
 			{
@@ -580,13 +579,14 @@ FAnimNode_Base* FAnimInstanceProxy::GetCheckedNodeFromIndexUntyped(int32 NodeIdx
 FAnimNode_Base* FAnimInstanceProxy::GetNodeFromIndexUntyped(int32 NodeIdx, UScriptStruct* RequiredStructType)
 {
 	FAnimNode_Base* NodePtr = nullptr;
-	if (AnimBlueprintClass)
+	if (AnimClassInterface)
 	{
-		const int32 InstanceIdx = AnimBlueprintClass->AnimNodeProperties.Num() - 1 - NodeIdx;
+		const TArray<UStructProperty*>& AnimNodeProperties = AnimClassInterface->GetAnimNodeProperties();
+		const int32 InstanceIdx = AnimNodeProperties.Num() - 1 - NodeIdx;
 
-		if (AnimBlueprintClass->AnimNodeProperties.IsValidIndex(InstanceIdx))
+		if (AnimNodeProperties.IsValidIndex(InstanceIdx))
 		{
-			UStructProperty* NodeProperty = AnimBlueprintClass->AnimNodeProperties[InstanceIdx];
+			UStructProperty* NodeProperty = AnimNodeProperties[InstanceIdx];
 
 			if (NodeProperty->Struct->IsChildOf(RequiredStructType))
 			{
@@ -1170,13 +1170,14 @@ FAnimNode_AssetPlayerBase* FAnimInstanceProxy::GetRelevantAssetPlayerFromState(i
 
 FAnimNode_StateMachine* FAnimInstanceProxy::GetStateMachineInstance(int32 MachineIndex)
 {
-	if(AnimBlueprintClass)
+	if (AnimClassInterface)
 	{
-		if((MachineIndex >= 0) && (MachineIndex < AnimBlueprintClass->AnimNodeProperties.Num()))
+		const TArray<UStructProperty*>& AnimNodeProperties = AnimClassInterface->GetAnimNodeProperties();
+		if ((MachineIndex >= 0) && (MachineIndex < AnimNodeProperties.Num()))
 		{
-			const int32 InstancePropertyIndex = AnimBlueprintClass->AnimNodeProperties.Num() - 1 - MachineIndex;
+			const int32 InstancePropertyIndex = AnimNodeProperties.Num() - 1 - MachineIndex;
 
-			UStructProperty* MachineInstanceProperty = AnimBlueprintClass->AnimNodeProperties[InstancePropertyIndex];
+			UStructProperty* MachineInstanceProperty = AnimNodeProperties[InstancePropertyIndex];
 			checkSlow(MachineInstanceProperty->Struct->IsChildOf(FAnimNode_StateMachine::StaticStruct()));
 
 			return MachineInstanceProperty->ContainerPtrToValuePtr<FAnimNode_StateMachine>(AnimInstanceObject);
@@ -1270,16 +1271,16 @@ void FAnimInstanceProxy::BindNativeDelegates()
 		return;
 	}
 
-	auto ForEachStateLambda = [&](UAnimBlueprintGeneratedClass* InAnimBlueprintClass, const FName& MachineName, const FName& StateName, TFunctionRef<void (FAnimNode_StateMachine*, const FBakedAnimationState&, int32)> Predicate)
+	auto ForEachStateLambda = [&](IAnimClassInterface* InAnimClassInterface, const FName& MachineName, const FName& StateName, TFunctionRef<void(FAnimNode_StateMachine*, const FBakedAnimationState&, int32)> Predicate)
 	{
-		for(UStructProperty* Property : InAnimBlueprintClass->AnimNodeProperties)
+		for (UStructProperty* Property : InAnimClassInterface->GetAnimNodeProperties())
 		{
 			if(Property && Property->Struct == FAnimNode_StateMachine::StaticStruct())
 			{
 				FAnimNode_StateMachine* StateMachine = Property->ContainerPtrToValuePtr<FAnimNode_StateMachine>(AnimInstanceObject);
 				if(StateMachine)
 				{
-					const FBakedAnimationStateMachine* MachineDescription = GetMachineDescription(InAnimBlueprintClass, StateMachine);
+					const FBakedAnimationStateMachine* MachineDescription = GetMachineDescription(InAnimClassInterface, StateMachine);
 					if(MachineDescription && MachineName == MachineDescription->MachineName)
 					{
 						// check each state transition for a match
@@ -1298,12 +1299,12 @@ void FAnimInstanceProxy::BindNativeDelegates()
 		}
 	};
 
-	if(AnimBlueprintClass)
+	if (AnimClassInterface)
 	{
 		// transition delegates
 		for(const auto& Binding : NativeTransitionBindings)
 		{
-			ForEachStateLambda(AnimBlueprintClass, Binding.MachineName, Binding.PreviousStateName, 
+			ForEachStateLambda(AnimClassInterface, Binding.MachineName, Binding.PreviousStateName,
 				[&](FAnimNode_StateMachine* StateMachine, const FBakedAnimationState& State, int32 StateIndex)
 				{
 					for(const FBakedStateExitTransition& TransitionExit : State.Transitions)
@@ -1311,13 +1312,13 @@ void FAnimInstanceProxy::BindNativeDelegates()
 						if(TransitionExit.CanTakeDelegateIndex != INDEX_NONE)
 						{
 							// In case the state machine hasn't been initilized, we need to re-get the desc
-							const FBakedAnimationStateMachine* MachineDesc = GetMachineDescription(AnimBlueprintClass, StateMachine);
+							const FBakedAnimationStateMachine* MachineDesc = GetMachineDescription(AnimClassInterface, StateMachine);
 							const FAnimationTransitionBetweenStates& Transition = MachineDesc->Transitions[TransitionExit.TransitionIndex];
 							const FBakedAnimationState& BakedState = MachineDesc->States[Transition.NextState];
 
 							if (BakedState.StateName == Binding.NextStateName)
 							{
-								FAnimNode_TransitionResult* ResultNode = GetNodeFromPropertyIndex<FAnimNode_TransitionResult>(AnimInstanceObject, AnimBlueprintClass, TransitionExit.CanTakeDelegateIndex);
+								FAnimNode_TransitionResult* ResultNode = GetNodeFromPropertyIndex<FAnimNode_TransitionResult>(AnimInstanceObject, AnimClassInterface, TransitionExit.CanTakeDelegateIndex);
 								if(ResultNode)
 								{
 									ResultNode->NativeTransitionDelegate = Binding.NativeTransitionDelegate;
@@ -1331,7 +1332,7 @@ void FAnimInstanceProxy::BindNativeDelegates()
 		// state entry delegates
 		for(const auto& Binding : NativeStateEntryBindings)
 		{
-			ForEachStateLambda(AnimBlueprintClass, Binding.MachineName, Binding.StateName, 
+			ForEachStateLambda(AnimClassInterface, Binding.MachineName, Binding.StateName,
 				[&](FAnimNode_StateMachine* StateMachine, const FBakedAnimationState& State, int32 StateIndex)
 				{
 					// allocate enough space for all our states we need so far
@@ -1343,7 +1344,7 @@ void FAnimInstanceProxy::BindNativeDelegates()
 		// state exit delegates
 		for(const auto& Binding : NativeStateExitBindings)
 		{
-			ForEachStateLambda(AnimBlueprintClass, Binding.MachineName, Binding.StateName, 
+			ForEachStateLambda(AnimClassInterface, Binding.MachineName, Binding.StateName,
 				[&](FAnimNode_StateMachine* StateMachine, const FBakedAnimationState& State, int32 StateIndex)
 				{
 					// allocate enough space for all our states we need so far
@@ -1354,24 +1355,26 @@ void FAnimInstanceProxy::BindNativeDelegates()
 	}
 }
 
-const FBakedAnimationStateMachine* FAnimInstanceProxy::GetMachineDescription(UAnimBlueprintGeneratedClass* AnimBlueprintClass, FAnimNode_StateMachine* MachineInstance)
+const FBakedAnimationStateMachine* FAnimInstanceProxy::GetMachineDescription(IAnimClassInterface* AnimBlueprintClass, FAnimNode_StateMachine* MachineInstance)
 {
-	return AnimBlueprintClass->BakedStateMachines.IsValidIndex(MachineInstance->StateMachineIndexInClass) ? &(AnimBlueprintClass->BakedStateMachines[MachineInstance->StateMachineIndexInClass]) : nullptr;
+	const TArray<FBakedAnimationStateMachine>& BakedStateMachines = AnimBlueprintClass->GetBakedStateMachines();
+	return BakedStateMachines.IsValidIndex(MachineInstance->StateMachineIndexInClass) ? &(BakedStateMachines[MachineInstance->StateMachineIndexInClass]) : nullptr;
 }
 
 FAnimNode_StateMachine* FAnimInstanceProxy::GetStateMachineInstanceFromName(FName MachineName)
 {
-	if (AnimBlueprintClass)
+	if (AnimClassInterface)
 	{
-		for (int32 MachineIndex = 0; MachineIndex < AnimBlueprintClass->AnimNodeProperties.Num(); MachineIndex++)
+		const TArray<UStructProperty*>& AnimNodeProperties = AnimClassInterface->GetAnimNodeProperties();
+		for (int32 MachineIndex = 0; MachineIndex < AnimNodeProperties.Num(); MachineIndex++)
 		{
-			UStructProperty* Property = AnimBlueprintClass->AnimNodeProperties[AnimBlueprintClass->AnimNodeProperties.Num() - 1 - MachineIndex];
+			UStructProperty* Property = AnimNodeProperties[AnimNodeProperties.Num() - 1 - MachineIndex];
 			if (Property && Property->Struct == FAnimNode_StateMachine::StaticStruct())
 			{
 				FAnimNode_StateMachine* StateMachine = Property->ContainerPtrToValuePtr<FAnimNode_StateMachine>(AnimInstanceObject);
 				if (StateMachine)
 				{
-					if (const FBakedAnimationStateMachine* MachineDescription = GetMachineDescription(AnimBlueprintClass, StateMachine))
+					if (const FBakedAnimationStateMachine* MachineDescription = GetMachineDescription(AnimClassInterface, StateMachine))
 					{
 						if (MachineDescription->MachineName == MachineName)
 						{
@@ -1388,17 +1391,18 @@ FAnimNode_StateMachine* FAnimInstanceProxy::GetStateMachineInstanceFromName(FNam
 
 const FBakedAnimationStateMachine* FAnimInstanceProxy::GetStateMachineInstanceDesc(FName MachineName)
 {
-	if(AnimBlueprintClass)
+	if (AnimClassInterface)
 	{
-		for(int32 MachineIndex = 0; MachineIndex < AnimBlueprintClass->AnimNodeProperties.Num(); MachineIndex++)
+		const TArray<UStructProperty*>& AnimNodeProperties = AnimClassInterface->GetAnimNodeProperties();
+		for (int32 MachineIndex = 0; MachineIndex < AnimNodeProperties.Num(); MachineIndex++)
 		{
-			UStructProperty* Property = AnimBlueprintClass->AnimNodeProperties[AnimBlueprintClass->AnimNodeProperties.Num() - 1 - MachineIndex];
+			UStructProperty* Property = AnimNodeProperties[AnimNodeProperties.Num() - 1 - MachineIndex];
 			if(Property && Property->Struct == FAnimNode_StateMachine::StaticStruct())
 			{
 				FAnimNode_StateMachine* StateMachine = Property->ContainerPtrToValuePtr<FAnimNode_StateMachine>(AnimInstanceObject);
 				if(StateMachine)
 				{
-					if(const FBakedAnimationStateMachine* MachineDescription = GetMachineDescription(AnimBlueprintClass, StateMachine))
+					if (const FBakedAnimationStateMachine* MachineDescription = GetMachineDescription(AnimClassInterface, StateMachine))
 					{
 						if(MachineDescription->MachineName == MachineName)
 						{
@@ -1415,17 +1419,18 @@ const FBakedAnimationStateMachine* FAnimInstanceProxy::GetStateMachineInstanceDe
 
 int32 FAnimInstanceProxy::GetStateMachineIndex(FName MachineName)
 {
-	if(AnimBlueprintClass)
+	if (AnimClassInterface)
 	{
-		for(int32 MachineIndex = 0; MachineIndex < AnimBlueprintClass->AnimNodeProperties.Num(); MachineIndex++)
+		const TArray<UStructProperty*>& AnimNodeProperties = AnimClassInterface->GetAnimNodeProperties();
+		for (int32 MachineIndex = 0; MachineIndex < AnimNodeProperties.Num(); MachineIndex++)
 		{
-			UStructProperty* Property = AnimBlueprintClass->AnimNodeProperties[AnimBlueprintClass->AnimNodeProperties.Num() - 1 - MachineIndex];
+			UStructProperty* Property = AnimNodeProperties[AnimNodeProperties.Num() - 1 - MachineIndex];
 			if(Property && Property->Struct == FAnimNode_StateMachine::StaticStruct())
 			{
 				FAnimNode_StateMachine* StateMachine = Property->ContainerPtrToValuePtr<FAnimNode_StateMachine>(AnimInstanceObject);
 				if(StateMachine)
 				{
-					if(const FBakedAnimationStateMachine* MachineDescription = GetMachineDescription(AnimBlueprintClass, StateMachine))
+					if (const FBakedAnimationStateMachine* MachineDescription = GetMachineDescription(AnimClassInterface, StateMachine))
 					{
 						if(MachineDescription->MachineName == MachineName)
 						{
@@ -1442,10 +1447,11 @@ int32 FAnimInstanceProxy::GetStateMachineIndex(FName MachineName)
 
 int32 FAnimInstanceProxy::GetInstanceAssetPlayerIndex(FName MachineName, FName StateName, FName AssetName)
 {
-	if(AnimBlueprintClass)
+	if (AnimClassInterface)
 	{
 		if(const FBakedAnimationStateMachine* MachineDescription = GetStateMachineInstanceDesc(MachineName))
 		{
+			const TArray<UStructProperty*>& AnimNodeProperties = AnimClassInterface->GetAnimNodeProperties();
 			for(int32 StateIndex = 0; StateIndex < MachineDescription->States.Num(); StateIndex++)
 			{
 				const FBakedAnimationState& State = MachineDescription->States[StateIndex];
@@ -1453,8 +1459,8 @@ int32 FAnimInstanceProxy::GetInstanceAssetPlayerIndex(FName MachineName, FName S
 				{
 					for(int32 PlayerIndex = 0; PlayerIndex < State.PlayerNodeIndices.Num(); PlayerIndex++)
 					{
-						checkSlow(State.PlayerNodeIndices[PlayerIndex] < AnimBlueprintClass->AnimNodeProperties.Num());
-						UStructProperty* AssetPlayerProperty = AnimBlueprintClass->AnimNodeProperties[AnimBlueprintClass->AnimNodeProperties.Num() - 1 - State.PlayerNodeIndices[PlayerIndex]];
+						checkSlow(State.PlayerNodeIndices[PlayerIndex] < AnimNodeProperties.Num());
+						UStructProperty* AssetPlayerProperty = AnimNodeProperties[AnimNodeProperties.Num() - 1 - State.PlayerNodeIndices[PlayerIndex]];
 						if(AssetPlayerProperty && AssetPlayerProperty->Struct->IsChildOf(FAnimNode_AssetPlayerBase::StaticStruct()))
 						{
 							FAnimNode_AssetPlayerBase* AssetPlayer = AssetPlayerProperty->ContainerPtrToValuePtr<FAnimNode_AssetPlayerBase>(AnimInstanceObject);
diff --git a/Engine/Source/Runtime/Engine/Private/Animation/AnimNodeBase.cpp b/Engine/Source/Runtime/Engine/Private/Animation/AnimNodeBase.cpp
--- a/Engine/Source/Runtime/Engine/Private/Animation/AnimNodeBase.cpp
+++ b/Engine/Source/Runtime/Engine/Private/Animation/AnimNodeBase.cpp
@@ -36,17 +36,20 @@ FAnimationBaseContext::FAnimationBaseContext(const FAnimationBaseContext& InCont
 
 UAnimBlueprintGeneratedClass* FAnimationBaseContext::GetAnimBlueprintClass() const
 {
+	PRAGMA_DISABLE_DEPRECATION_WARNINGS
 	return AnimInstanceProxy->GetAnimBlueprintClass();
+	PRAGMA_ENABLE_DEPRECATION_WARNINGS
+}
+
+IAnimClassInterface* FAnimationBaseContext::GetAnimClass() const
+{
+	return AnimInstanceProxy ? AnimInstanceProxy->GetAnimClassInterface() : nullptr;
 }
 
 #if WITH_EDITORONLY_DATA
 UAnimBlueprint* FAnimationBaseContext::GetAnimBlueprint() const
 {
-	if (UAnimBlueprintGeneratedClass* Class = AnimInstanceProxy->GetAnimBlueprintClass())
-	{
-		return Cast<UAnimBlueprint>(Class->ClassGeneratedBy);
-	}
-	return NULL;
+	return AnimInstanceProxy ? AnimInstanceProxy->GetAnimBlueprint() : nullptr;
 }
 #endif //WITH_EDITORONLY_DATA
 
@@ -91,13 +94,13 @@ void FPoseLinkBase::AttemptRelink(const FAnimationBaseContext& Context)
 	// Do the linkage
 	if ((LinkedNode == NULL) && (LinkID != INDEX_NONE))
 	{
-		UAnimBlueprintGeneratedClass* AnimBlueprintClass = Context.GetAnimBlueprintClass();
+		IAnimClassInterface* AnimBlueprintClass = Context.GetAnimClass();
 		check(AnimBlueprintClass);
 
 		// adding ensure. We had a crash on here
-		if ( ensure(AnimBlueprintClass->AnimNodeProperties.IsValidIndex(LinkID)) )
+		if ( ensure(AnimBlueprintClass->GetAnimNodeProperties().IsValidIndex(LinkID)) )
 		{
-			UProperty* LinkedProperty = AnimBlueprintClass->AnimNodeProperties[LinkID];
+			UProperty* LinkedProperty = AnimBlueprintClass->GetAnimNodeProperties()[LinkID];
 			void* LinkedNodePtr = LinkedProperty->ContainerPtrToValuePtr<void>(Context.AnimInstanceProxy->GetAnimInstanceObject());
 			LinkedNode = (FAnimNode_Base*)LinkedNodePtr;
 		}
@@ -107,8 +110,8 @@ void FPoseLinkBase::AttemptRelink(const FAnimationBaseContext& Context)
 void FPoseLinkBase::Initialize(const FAnimationInitializeContext& Context)
 {
 #if DO_CHECK
-	checkf( !bProcessed, TEXT( "Initialize already in progress, circular link for AnimInstance [%s] Blueprint [%s]" ), \
-		*Context.AnimInstanceProxy->GetAnimInstanceName(), Context.GetAnimBlueprintClass() ? *Context.GetAnimBlueprintClass()->GetFullName() : TEXT( "None" ) );
+	checkf(!bProcessed, TEXT("Initialize already in progress, circular link for AnimInstance [%s] Blueprint [%s]"), \
+		*Context.AnimInstanceProxy->GetAnimInstanceName(), *GetFullNameSafe(IAnimClassInterface::GetActualAnimClass(Context.AnimInstanceProxy->GetAnimClassInterface())));
 	TGuardValue<bool> CircularGuard(bProcessed, true);
 #endif
 
@@ -129,7 +132,7 @@ void FPoseLinkBase::CacheBones(const FAnimationCacheBonesContext& Context)
 {
 #if DO_CHECK
 	checkf( !bProcessed, TEXT( "CacheBones already in progress, circular link for AnimInstance [%s] Blueprint [%s]" ), \
-		*Context.AnimInstanceProxy->GetAnimInstanceName(), Context.GetAnimBlueprintClass() ? *Context.GetAnimBlueprintClass()->GetFullName() : TEXT( "None" ) );
+		*Context.AnimInstanceProxy->GetAnimInstanceName(), *GetFullNameSafe(IAnimClassInterface::GetActualAnimClass(Context.AnimInstanceProxy->GetAnimClassInterface())));
 	TGuardValue<bool> CircularGuard(bProcessed, true);
 #endif
 
@@ -148,7 +151,7 @@ void FPoseLinkBase::Update(const FAnimationUpdateContext& Context)
 	QUICK_SCOPE_CYCLE_COUNTER(STAT_FPoseLinkBase_Update);
 #if DO_CHECK
 	checkf( !bProcessed, TEXT( "Update already in progress, circular link for AnimInstance [%s] Blueprint [%s]" ), \
-		*Context.AnimInstanceProxy->GetAnimInstanceName(), Context.GetAnimBlueprintClass() ? *Context.GetAnimBlueprintClass()->GetFullName() : TEXT( "None" ) );
+		*Context.AnimInstanceProxy->GetAnimInstanceName(), *GetFullNameSafe(IAnimClassInterface::GetActualAnimClass(Context.AnimInstanceProxy->GetAnimClassInterface())));
 	TGuardValue<bool> CircularGuard(bProcessed, true);
 #endif
 
@@ -199,7 +202,7 @@ void FPoseLink::Evaluate(FPoseContext& Output)
 {
 #if DO_CHECK
 	checkf( !bProcessed, TEXT( "Evaluate already in progress, circular link for AnimInstance [%s] Blueprint [%s]" ), \
-		*Output.AnimInstanceProxy->GetAnimInstanceName(), Output.GetAnimBlueprintClass() ? *Output.GetAnimBlueprintClass()->GetFullName() : TEXT( "None" ) );
+		*Output.AnimInstanceProxy->GetAnimInstanceName(), *GetFullNameSafe(IAnimClassInterface::GetActualAnimClass(Output.AnimInstanceProxy->GetAnimClassInterface())));
 	TGuardValue<bool> CircularGuard(bProcessed, true);
 #endif
 
@@ -247,7 +250,7 @@ void FComponentSpacePoseLink::EvaluateComponentSpace(FComponentSpacePoseContext&
 {
 #if DO_CHECK
 	checkf( !bProcessed, TEXT( "EvaluateComponentSpace already in progress, circular link for AnimInstance [%s] Blueprint [%s]" ), \
-		*Output.AnimInstanceProxy->GetAnimInstanceName(), Output.GetAnimBlueprintClass() ? *Output.GetAnimBlueprintClass()->GetFullName() : TEXT( "None" ) );
+		*Output.AnimInstanceProxy->GetAnimInstanceName(), *GetFullNameSafe(IAnimClassInterface::GetActualAnimClass(Output.AnimInstanceProxy->GetAnimClassInterface())));
 	TGuardValue<bool> CircularGuard(bProcessed, true);
 #endif
 
diff --git a/Engine/Source/Runtime/Engine/Private/Animation/AnimNode_StateMachine.cpp b/Engine/Source/Runtime/Engine/Private/Animation/AnimNode_StateMachine.cpp
--- a/Engine/Source/Runtime/Engine/Private/Animation/AnimNode_StateMachine.cpp
+++ b/Engine/Source/Runtime/Engine/Private/Animation/AnimNode_StateMachine.cpp
@@ -77,8 +77,8 @@ void FAnimationActiveTransitionEntry::InitializeCustomGraphLinks(const FAnimatio
 {
 	if (TransitionRule.CustomResultNodeIndex != INDEX_NONE)
 	{
-		const UAnimBlueprintGeneratedClass* AnimBlueprintClass = Context.GetAnimBlueprintClass();
-		CustomTransitionGraph.LinkID = AnimBlueprintClass->AnimNodeProperties.Num() - 1 - TransitionRule.CustomResultNodeIndex; //@TODO: Crazysauce
+		const IAnimClassInterface* AnimBlueprintClass = Context.GetAnimClass();
+		CustomTransitionGraph.LinkID = AnimBlueprintClass->GetAnimNodeProperties().Num() - 1 - TransitionRule.CustomResultNodeIndex; //@TODO: Crazysauce
 		FAnimationInitializeContext InitContext(Context.AnimInstanceProxy);
 		CustomTransitionGraph.Initialize(InitContext);
 
@@ -210,7 +210,7 @@ void FAnimNode_StateMachine::Initialize(const FAnimationInitializeContext& Conte
 {
 	FAnimNode_Base::Initialize(Context);
 
-	const UAnimBlueprintGeneratedClass* AnimBlueprintClass = Context.GetAnimBlueprintClass();
+	IAnimClassInterface* AnimBlueprintClass = Context.GetAnimClass();
 
 	if (const FBakedAnimationStateMachine* Machine = GetMachineDescription())
 	{
@@ -231,7 +231,7 @@ void FAnimNode_StateMachine::Initialize(const FAnimationInitializeContext& Conte
 				// because conduits don't contain bound graphs, this link is no longer guaranteed to be valid
 				if (State.StateRootNodeIndex != INDEX_NONE)
 				{
-					StatePoseLink->LinkID = AnimBlueprintClass->AnimNodeProperties.Num() - 1 - State.StateRootNodeIndex; //@TODO: Crazysauce
+					StatePoseLink->LinkID = AnimBlueprintClass->GetAnimNodeProperties().Num() - 1 - State.StateRootNodeIndex; //@TODO: Crazysauce
 				}
 
 				// also initialize transitions
@@ -528,7 +528,7 @@ bool FAnimNode_StateMachine::FindValidTransition(const FAnimationUpdateContext&
 	}
 	OutVisitedStateIndices.Add(CheckingStateIndex);
 
-	const UAnimBlueprintGeneratedClass* AnimBlueprintClass = Context.GetAnimBlueprintClass();
+	const IAnimClassInterface* AnimBlueprintClass = Context.GetAnimClass();
 
 	// Conduit 'states' have an additional entry rule which must be true to consider taking any transitions via the conduit
 	//@TODO: It would add flexibility to be able to define this on normal state nodes as well, assuming the dual-graph editing is sorted out
@@ -1002,7 +1002,7 @@ FName FAnimNode_StateMachine::GetCurrentStateName() const
 	return NAME_None;
 }
 
-void FAnimNode_StateMachine::CacheMachineDescription(UAnimBlueprintGeneratedClass* AnimBlueprintClass)
+void FAnimNode_StateMachine::CacheMachineDescription(IAnimClassInterface* AnimBlueprintClass)
 {
-	PRIVATE_MachineDescription = AnimBlueprintClass->BakedStateMachines.IsValidIndex(StateMachineIndexInClass) ? &(AnimBlueprintClass->BakedStateMachines[StateMachineIndexInClass]) : nullptr;
+	PRIVATE_MachineDescription = AnimBlueprintClass->GetBakedStateMachines().IsValidIndex(StateMachineIndexInClass) ? &(AnimBlueprintClass->GetBakedStateMachines()[StateMachineIndexInClass]) : nullptr;
 }
diff --git a/Engine/Source/Runtime/Engine/Private/BlueprintGeneratedClass.cpp b/Engine/Source/Runtime/Engine/Private/BlueprintGeneratedClass.cpp
--- a/Engine/Source/Runtime/Engine/Private/BlueprintGeneratedClass.cpp
+++ b/Engine/Source/Runtime/Engine/Private/BlueprintGeneratedClass.cpp
@@ -21,6 +21,7 @@ UBlueprintGeneratedClass::UBlueprintGeneratedClass(const FObjectInitializer& Obj
 	: Super(ObjectInitializer)
 {
 	NumReplicatedProperties = 0;
+	bHasInstrumentation = false;
 }
 
 void UBlueprintGeneratedClass::PostInitProperties()
@@ -601,7 +602,7 @@ uint8* UBlueprintGeneratedClass::GetPersistentUberGraphFrame(UObject* Obj, UFunc
 
 void UBlueprintGeneratedClass::CreatePersistentUberGraphFrame(UObject* Obj, bool bCreateOnlyIfEmpty, bool bSkipSuperClass) const
 {
-	checkSlow(!UberGraphFramePointerProperty == !UberGraphFunction);
+	ensure(!UberGraphFramePointerProperty == !UberGraphFunction);
 	if (Obj && UsePersistentUberGraphFrame() && UberGraphFramePointerProperty && UberGraphFunction)
 	{
 		auto PointerToUberGraphFrame = UberGraphFramePointerProperty->ContainerPtrToValuePtr<FPointerToUberGraphFrame>(Obj);
@@ -642,7 +643,7 @@ void UBlueprintGeneratedClass::CreatePersistentUberGraphFrame(UObject* Obj, bool
 
 void UBlueprintGeneratedClass::DestroyPersistentUberGraphFrame(UObject* Obj, bool bSkipSuperClass) const
 {
-	checkSlow(!UberGraphFramePointerProperty == !UberGraphFunction);
+	ensure(!UberGraphFramePointerProperty == !UberGraphFunction);
 	if (Obj && UsePersistentUberGraphFrame() && UberGraphFramePointerProperty && UberGraphFunction)
 	{
 		auto PointerToUberGraphFrame = UberGraphFramePointerProperty->ContainerPtrToValuePtr<FPointerToUberGraphFrame>(Obj);
@@ -749,6 +750,18 @@ public:
 protected:
 	virtual FArchive& operator<<(UObject*& Object) override
 	{
+#if !(UE_BUILD_TEST || UE_BUILD_SHIPPING)
+		if (Object && !Object->IsValidLowLevelFast())
+		{
+			UProperty* SerializingProperty = GetSerializedProperty();
+
+			UE_LOG(LogBlueprint, Fatal, TEXT("Invalid object in PersistentFrame: 0x%016llx, Blueprint object: %s, ReferencingProperty: %s"),
+				(int64)(PTRINT)Object,
+				SerializingObject   ? *SerializingObject->GetFullName()   : TEXT("NULL"),
+				SerializingProperty ? *SerializingProperty->GetFullName() : TEXT("NULL"));
+		}
+#endif
+
 		const bool bWeakRef = Object ? !Object->HasAnyFlags(RF_StrongRefOnFrame) : false;
 		Collector.SetShouldHandleAsWeakRef(bWeakRef); 
 		return FSimpleObjectReferenceCollectorArchive::operator<<(Object);
@@ -764,11 +777,11 @@ void UBlueprintGeneratedClass::AddReferencedObjectsInUbergraphFrame(UObject* InT
 		{
 			if (BPGC->UberGraphFramePointerProperty)
 			{
-				checkSlow(BPGC->UberGraphFunction);
 				auto PointerToUberGraphFrame = BPGC->UberGraphFramePointerProperty->ContainerPtrToValuePtr<FPointerToUberGraphFrame>(InThis);
 				checkSlow(PointerToUberGraphFrame)
 				if (PointerToUberGraphFrame->RawPointer)
 				{
+					checkSlow(BPGC->UberGraphFunction);
 					FPersistentFrameCollectorArchive ObjectReferenceCollector(InThis, Collector);
 					BPGC->UberGraphFunction->SerializeBin(ObjectReferenceCollector, PointerToUberGraphFrame->RawPointer);
 
diff --git a/Engine/Source/Runtime/Engine/Private/Components/PrimitiveComponent.cpp b/Engine/Source/Runtime/Engine/Private/Components/PrimitiveComponent.cpp
--- a/Engine/Source/Runtime/Engine/Private/Components/PrimitiveComponent.cpp
+++ b/Engine/Source/Runtime/Engine/Private/Components/PrimitiveComponent.cpp
@@ -426,6 +426,11 @@ void FPrimitiveComponentInstanceData::ApplyToComponent(UActorComponent* Componen
 {
 	FSceneComponentInstanceData::ApplyToComponent(Component, CacheApplyPhase);
 
+#if WITH_EDITOR
+	// This is needed to restore transient collision profile data.
+	CastChecked<UPrimitiveComponent>(Component)->UpdateCollisionProfile();
+#endif // #if WITH_EDITOR
+
 	if (ContainsSavedProperties() && Component->IsRegistered())
 	{
 		Component->MarkRenderStateDirty();
diff --git a/Engine/Source/Runtime/Engine/Private/Components/SkeletalMeshComponent.cpp b/Engine/Source/Runtime/Engine/Private/Components/SkeletalMeshComponent.cpp
--- a/Engine/Source/Runtime/Engine/Private/Components/SkeletalMeshComponent.cpp
+++ b/Engine/Source/Runtime/Engine/Private/Components/SkeletalMeshComponent.cpp
@@ -278,10 +278,11 @@ void USkeletalMeshComponent::UpdateClothTickRegisteredState()
 
 bool USkeletalMeshComponent::NeedToSpawnAnimScriptInstance(bool bForceInit) const
 {
-	if (AnimationMode == EAnimationMode::AnimationBlueprint && (AnimBlueprintGeneratedClass != NULL) && 
-		(SkeletalMesh != NULL) && (SkeletalMesh->Skeleton->IsCompatible(AnimBlueprintGeneratedClass->TargetSkeleton)))
+	IAnimClassInterface* AnimClassInterface = IAnimClassInterface::GetFromClass(AnimClass);
+	if (AnimationMode == EAnimationMode::AnimationBlueprint && (AnimClassInterface != NULL) &&
+		(SkeletalMesh != NULL) && (SkeletalMesh->Skeleton->IsCompatible(AnimClassInterface->GetTargetSkeleton())))
 	{
-		if (bForceInit || (AnimScriptInstance == NULL) || (AnimScriptInstance->GetClass() != AnimBlueprintGeneratedClass) )
+		if (bForceInit || (AnimScriptInstance == NULL) || (AnimScriptInstance->GetClass() != AnimClass) )
 		{
 			return true;
 		}
@@ -292,7 +293,7 @@ bool USkeletalMeshComponent::NeedToSpawnAnimScriptInstance(bool bForceInit) cons
 
 bool USkeletalMeshComponent::IsAnimBlueprintInstanced() const
 {
-	return (AnimScriptInstance && AnimScriptInstance->GetClass() == AnimBlueprintGeneratedClass);
+	return (AnimScriptInstance && AnimScriptInstance->GetClass() == AnimClass);
 }
 
 void USkeletalMeshComponent::OnRegister()
@@ -339,8 +340,8 @@ void USkeletalMeshComponent::InitAnim(bool bForceReinit)
 		const bool bPerformPostAnimEvaluation = false; // Skip post evaluation, it would be wasted work
 		HandleExistingParallelEvaluationTask(bBlockOnTask, bPerformPostAnimEvaluation);
 
-		bool bBlueprintMismatch = (AnimBlueprintGeneratedClass != NULL) && 
-			(AnimScriptInstance != NULL) && (AnimScriptInstance->GetClass() != AnimBlueprintGeneratedClass);
+		bool bBlueprintMismatch = (AnimClass != NULL) &&
+			(AnimScriptInstance != NULL) && (AnimScriptInstance->GetClass() != AnimClass);
 
 		bool bSkeletonMismatch = AnimScriptInstance && AnimScriptInstance->CurrentSkeleton && (AnimScriptInstance->CurrentSkeleton!=SkeletalMesh->Skeleton);
 
@@ -369,7 +370,7 @@ void USkeletalMeshComponent::InitializeAnimScriptInstance(bool bForceReinit)
 		if (NeedToSpawnAnimScriptInstance(bForceReinit))
 		{
 			SCOPE_CYCLE_COUNTER(STAT_AnimSpawnTime);
-			AnimScriptInstance = NewObject<UAnimInstance>(this, AnimBlueprintGeneratedClass);
+			AnimScriptInstance = NewObject<UAnimInstance>(this, AnimClass);
 
 			if (AnimScriptInstance)
 			{
@@ -452,7 +453,7 @@ void USkeletalMeshComponent::PostEditChangeProperty(FPropertyChangedEvent& Prope
 		{
 			if (AnimationMode == EAnimationMode::AnimationBlueprint)
 			{
-				if (AnimBlueprintGeneratedClass == NULL)
+				if (AnimClass == NULL)
 				{
 					ClearAnimScriptInstance();
 				}
@@ -461,14 +462,14 @@ void USkeletalMeshComponent::PostEditChangeProperty(FPropertyChangedEvent& Prope
 					if (NeedToSpawnAnimScriptInstance(false))
 					{
 						SCOPE_CYCLE_COUNTER(STAT_AnimSpawnTime);
-						AnimScriptInstance = NewObject<UAnimInstance>(this, AnimBlueprintGeneratedClass);
+						AnimScriptInstance = NewObject<UAnimInstance>(this, AnimClass);
 						AnimScriptInstance->InitializeAnimation();
 					}
 				}
 			}
 		}
 
-		if ( PropertyThatChanged->GetFName() == GET_MEMBER_NAME_CHECKED( USkeletalMeshComponent, AnimBlueprintGeneratedClass ) )
+		if (PropertyThatChanged->GetFName() == GET_MEMBER_NAME_CHECKED(USkeletalMeshComponent, AnimClass))
 		{
 			InitAnim(false);
 		}
@@ -1420,17 +1421,16 @@ void USkeletalMeshComponent::SetForceRefPose(bool bNewForceRefPose)
 
 void USkeletalMeshComponent::SetAnimInstanceClass(class UClass* NewClass)
 {
-	if (NewClass != NULL)
+	if (NewClass != nullptr)
 	{
-		UAnimBlueprintGeneratedClass* NewGeneratedClass = Cast<UAnimBlueprintGeneratedClass>(NewClass);
-		ensure(NULL != NewGeneratedClass);
+		ensure(nullptr != IAnimClassInterface::GetFromClass(NewClass));
 		// set the animation mode
 		AnimationMode = EAnimationMode::Type::AnimationBlueprint;
 
-		if (NewGeneratedClass != AnimBlueprintGeneratedClass)
+		if (NewClass != AnimClass)
 		{
 			// Only need to initialize if it hasn't already been set.
-			AnimBlueprintGeneratedClass = NewGeneratedClass;
+			AnimClass = NewClass;
 			ClearAnimScriptInstance();
 			InitAnim(true);
 		}
@@ -1439,7 +1439,7 @@ void USkeletalMeshComponent::SetAnimInstanceClass(class UClass* NewClass)
 	{
 		// Need to clear the instance as well as the blueprint.
 		// @todo is this it?
-		AnimBlueprintGeneratedClass = NULL;
+		AnimClass = nullptr;
 		ClearAnimScriptInstance();
 	}
 }
@@ -2077,7 +2077,7 @@ void USkeletalMeshComponent::ValidateAnimation()
 		}
 		else
 		{
-			AnimBlueprintGeneratedClass = nullptr;
+			AnimClass = nullptr;
 		}
 		return;
 	}
@@ -2100,18 +2100,19 @@ void USkeletalMeshComponent::ValidateAnimation()
 	}
 	else
 	{
-		if(AnimBlueprintGeneratedClass && SkeletalMesh && AnimBlueprintGeneratedClass->TargetSkeleton != SkeletalMesh->Skeleton)
+		IAnimClassInterface* AnimClassInterface = IAnimClassInterface::GetFromClass(AnimClass);
+		if (AnimClassInterface && SkeletalMesh && AnimClassInterface->GetTargetSkeleton() != SkeletalMesh->Skeleton)
 		{
 			if(SkeletalMesh->Skeleton)
 			{
-				UE_LOG(LogAnimation, Warning, TEXT("AnimBP %s is incompatible with skeleton %s, removing AnimBP from actor."), *AnimBlueprintGeneratedClass->GetName(), *SkeletalMesh->Skeleton->GetName());
+				UE_LOG(LogAnimation, Warning, TEXT("AnimBP %s is incompatible with skeleton %s, removing AnimBP from actor."), *AnimClass->GetName(), *SkeletalMesh->Skeleton->GetName());
 			}
 			else
 			{
-				UE_LOG(LogAnimation, Warning, TEXT("AnimBP %s is incompatible because mesh %s has no skeleton, removing AnimBP from actor."), *AnimBlueprintGeneratedClass->GetName(), *SkeletalMesh->GetName());
+				UE_LOG(LogAnimation, Warning, TEXT("AnimBP %s is incompatible because mesh %s has no skeleton, removing AnimBP from actor."), *AnimClass->GetName(), *SkeletalMesh->GetName());
 			}
 
-			AnimBlueprintGeneratedClass = nullptr;
+			AnimClass = nullptr;
 		}
 	}
 }
diff --git a/Engine/Source/Runtime/Engine/Private/DataTable.cpp b/Engine/Source/Runtime/Engine/Private/DataTable.cpp
--- a/Engine/Source/Runtime/Engine/Private/DataTable.cpp
+++ b/Engine/Source/Runtime/Engine/Private/DataTable.cpp
@@ -77,12 +77,13 @@ void UDataTable::LoadStructData(FArchive& Ar)
 
 void UDataTable::SaveStructData(FArchive& Ar)
 {
+	// Always save number number of rows, so the DT could be safely loaded.
+	int32 NumRowsToSave = (nullptr != RowStruct) ? RowMap.Num() : 0;
+	Ar << NumRowsToSave;
+
 	// Don't even try to save rows if no RowStruct
-	if (RowStruct != NULL)
+	if (nullptr != RowStruct)
 	{
-		int32 NumRows = RowMap.Num();
-		Ar << NumRows;
-
 		// Now iterate over rows in the map
 		for (auto RowIt = RowMap.CreateIterator(); RowIt; ++RowIt)
 		{
diff --git a/Engine/Source/Runtime/Engine/Private/EdGraph/EdGraph.cpp b/Engine/Source/Runtime/Engine/Private/EdGraph/EdGraph.cpp
--- a/Engine/Source/Runtime/Engine/Private/EdGraph/EdGraph.cpp
+++ b/Engine/Source/Runtime/Engine/Private/EdGraph/EdGraph.cpp
@@ -132,7 +132,9 @@ void UEdGraph::PostLoad()
 		if (Nodes[i] == nullptr)
 		{
 			Nodes.RemoveAt(i);
-			UE_LOG(LogBlueprint, Display, TEXT("Found NULL Node in EdGraph Nodes array. A node type may have been deleted without creating an ActiveClassRedictor to K2Node_DeadClass."));
+			UE_LOG(LogBlueprint, Display, 
+				TEXT("Missing node found in the EdGraph. We've fixed up the issue, but couldn't determine the cause. If you are able to figure out the cause, please notify the engine team; otherwise consider resaving the asset: '%s'"),
+				*GetOutermost()->GetName());
 		}
 	}
 }
diff --git a/Engine/Source/Runtime/Engine/Private/EdGraph/EdGraphPin.cpp b/Engine/Source/Runtime/Engine/Private/EdGraph/EdGraphPin.cpp
--- a/Engine/Source/Runtime/Engine/Private/EdGraph/EdGraphPin.cpp
+++ b/Engine/Source/Runtime/Engine/Private/EdGraph/EdGraphPin.cpp
@@ -302,20 +302,34 @@ const FString UEdGraphPin::GetLinkInfoString( const FString& InFunctionName, con
 #endif
 }
 
-void UEdGraphPin::Serialize(FArchive& Ar)
+void UEdGraphPin::PostLoad()
 {
-	Super::Serialize(Ar);
+	Super::PostLoad();
 
-	if (Ar.IsLoading() && Ar.IsPersistent())
+	static FName GameplayTagName = TEXT("GameplayTag");
+	static FName GameplayTagContainerName = TEXT("GameplayTagContainer");
+	static FName GameplayTagsPathName = TEXT("/Script/GameplayTags");
+
+	if (PinType.PinSubCategoryObject.IsValid())
 	{
-		// Pins of type FGameplayTag were storing "()" for empty arrays and then importing that into ArrayProperty and expecting an empty array.
-		// That it was working was a bug and has been fixed, so let's fixup pins. A pin that wants an array size of 1 will always fill the parenthesis
-		// so there is no worry about breaking those cases.
-		if (PinType.PinSubCategoryObject.IsValid() && PinType.PinSubCategoryObject->GetPathName() == TEXT("/Script/GameplayTags.GameplayTag"))
+		UObject* PinSubCategoryObject = PinType.PinSubCategoryObject.Get();
+		if (PinSubCategoryObject->GetOuter()->GetFName() == GameplayTagsPathName)
 		{
-			if (DefaultValue == TEXT("()"))
+			if (PinSubCategoryObject->GetFName() == GameplayTagName)
+			{
+				// Pins of type FGameplayTag were storing "()" for empty arrays and then importing that into ArrayProperty and expecting an empty array.
+				// That it was working was a bug and has been fixed, so let's fixup pins. A pin that wants an array size of 1 will always fill the parenthesis
+				// so there is no worry about breaking those cases.
+				if (DefaultValue == TEXT("()"))
+				{
+					DefaultValue.Empty();
+				}
+			}
+			else if (PinSubCategoryObject->GetFName() == GameplayTagContainerName)
 			{
-				DefaultValue.Empty();
+				// Pins of type FGameplayTagContainer were storing "GameplayTags=()" for empty arrays, which equates to having a single item, default generated as detailed above for FGameplayTag.
+				// The solution is to replace occurances with an empty string, due to the item being a struct, we can't just empty the value and must replace only the section we need.
+				DefaultValue.ReplaceInline(TEXT("GameplayTags=()"), TEXT("GameplayTags="));
 			}
 		}
 	}
diff --git a/Engine/Source/Runtime/Engine/Private/EdGraph/EdGraphSchema.cpp b/Engine/Source/Runtime/Engine/Private/EdGraph/EdGraphSchema.cpp
--- a/Engine/Source/Runtime/Engine/Private/EdGraph/EdGraphSchema.cpp
+++ b/Engine/Source/Runtime/Engine/Private/EdGraph/EdGraphSchema.cpp
@@ -493,11 +493,38 @@ void UEdGraphSchema::TrySetDefaultText(UEdGraphPin& InPin, const FText& InNewDef
 void UEdGraphSchema::BreakNodeLinks(UEdGraphNode& TargetNode) const
 {
 #if WITH_EDITOR
+	TSet<UEdGraphNode*> NodeList;
+	NodeList.Add(&TargetNode);
+	
+	// Iterate over each pin and break all links
 	for (TArray<UEdGraphPin*>::TIterator PinIt(TargetNode.Pins); PinIt; ++PinIt)
 	{
-		BreakPinLinks(*(*PinIt), false);
+		UEdGraphPin* TargetPin = *PinIt;
+		if (TargetPin)
+		{
+			// Keep track of which node(s) the pin's connected to
+			for (auto OtherPin : TargetPin->LinkedTo)
+			{
+				if (OtherPin)
+				{
+					UEdGraphNode* OtherNode = OtherPin->GetOwningNode();
+					if (OtherNode)
+					{
+						NodeList.Add(OtherNode);
+					}
+				}
+			}
+
+			BreakPinLinks(*TargetPin, false);
+		}
+	}
+	
+	// Send all nodes that lost connections a notification
+	for (auto It = NodeList.CreateConstIterator(); It; ++It)
+	{
+		UEdGraphNode* Node = (*It);
+		Node->NodeConnectionListChanged();
 	}
-	TargetNode.NodeConnectionListChanged();
 #endif	//#if WITH_EDITOR
 }
 
diff --git a/Engine/Source/Runtime/Engine/Private/InputDelegateBinding.cpp b/Engine/Source/Runtime/Engine/Private/InputDelegateBinding.cpp
--- a/Engine/Source/Runtime/Engine/Private/InputDelegateBinding.cpp
+++ b/Engine/Source/Runtime/Engine/Private/InputDelegateBinding.cpp
@@ -17,11 +17,7 @@ UInputDelegateBinding::UInputDelegateBinding(const FObjectInitializer& ObjectIni
 
 bool UInputDelegateBinding::SupportsInputDelegate(const UClass* InClass)
 {
-	auto DynamicClass = Cast<UDynamicClass>(InClass);
-	return InClass && (
-		(DynamicClass && DynamicClass->DynamicBindingObjects.Num())
-		|| Cast<UBlueprintGeneratedClass>(InClass)
-		|| SupportsInputDelegate(InClass->GetSuperClass()));
+	return Cast<UDynamicClass>(InClass) || Cast<UBlueprintGeneratedClass>(InClass);
 }
 
 void UInputDelegateBinding::BindInputDelegates(const UClass* InClass, UInputComponent* InputComponent)
diff --git a/Engine/Source/Runtime/Engine/Private/Interpolation.cpp b/Engine/Source/Runtime/Engine/Private/Interpolation.cpp
--- a/Engine/Source/Runtime/Engine/Private/Interpolation.cpp
+++ b/Engine/Source/Runtime/Engine/Private/Interpolation.cpp
@@ -9705,6 +9705,29 @@ UInterpTrackFloatAnimBPParam::UInterpTrackFloatAnimBPParam(const FObjectInitiali
 	TrackTitle = TEXT("Float AnimBP Param");
 }
 
+void UInterpTrackFloatAnimBPParam::Serialize(FArchive& Ar)
+{
+	PRAGMA_DISABLE_DEPRECATION_WARNINGS	
+	if (Ar.IsSaving() && Ar.UE4Ver() < VER_UE4_NO_ANIM_BP_CLASS_IN_GAMEPLAY_CODE)
+	{
+		if ((nullptr != AnimBlueprintClass) && (nullptr == AnimClass))
+		{
+			AnimClass = AnimBlueprintClass;
+		}
+	}
+
+	Super::Serialize(Ar);
+
+	if (Ar.IsLoading() && Ar.UE4Ver() < VER_UE4_NO_ANIM_BP_CLASS_IN_GAMEPLAY_CODE)
+	{
+		if ((nullptr != AnimBlueprintClass) && (nullptr == AnimClass))
+		{
+			AnimClass = AnimBlueprintClass;
+		}
+	}
+	PRAGMA_ENABLE_DEPRECATION_WARNINGS
+}
+
 #if WITH_EDITOR
 
 void UInterpTrackFloatAnimBPParam::PostEditChangeProperty(FPropertyChangedEvent& PropertyChangedEvent)
@@ -9715,7 +9738,7 @@ void UInterpTrackFloatAnimBPParam::PostEditChangeProperty(FPropertyChangedEvent&
 	FName PropertyName = PropertyThatChanged != NULL ? PropertyThatChanged->GetFName() : NAME_None;
 
 	if (PropertyName == GET_MEMBER_NAME_CHECKED(UInterpTrackFloatAnimBPParam, ParamName) ||
-		PropertyName == GET_MEMBER_NAME_CHECKED(UInterpTrackFloatAnimBPParam, AnimBlueprintClass))
+		PropertyName == GET_MEMBER_NAME_CHECKED(UInterpTrackFloatAnimBPParam, AnimClass))
 	{
 		bRefreshParamter = true;
 	}
@@ -9819,7 +9842,7 @@ void UInterpTrackInstFloatAnimBPParam::RefreshParameter(UInterpTrack* Track)
 			{
 				UAnimInstance* NewAnimInstance = SkeletalMeshComponents[0]->GetAnimInstance();
 
-				if(NewAnimInstance && NewAnimInstance->GetClass() == ParamTrack->AnimBlueprintClass)
+				if (NewAnimInstance && NewAnimInstance->GetClass() == ParamTrack->AnimClass)
 				{
 					AnimScriptInstance = NewAnimInstance;
 					// make sure the class has the parameter
diff --git a/Engine/Source/Runtime/Engine/Private/KismetArrayLibrary.cpp b/Engine/Source/Runtime/Engine/Private/KismetArrayLibrary.cpp
--- a/Engine/Source/Runtime/Engine/Private/KismetArrayLibrary.cpp
+++ b/Engine/Source/Runtime/Engine/Private/KismetArrayLibrary.cpp
@@ -140,7 +140,7 @@ void UKismetArrayLibrary::GenericArray_Shuffle(void* TargetArray, const UArrayPr
 	{
 		FScriptArrayHelper ArrayHelper(ArrayProp, TargetArray);
 		int32 LastIndex = ArrayHelper.Num() - 1;
-		for (int32 i = 0; i < LastIndex; ++i)
+		for (int32 i = 0; i <= LastIndex; ++i)
 		{
 			int32 Index = FMath::RandRange(0, LastIndex);
 			if (i != Index)
diff --git a/Engine/Source/Runtime/Engine/Private/MemberReference.cpp b/Engine/Source/Runtime/Engine/Private/MemberReference.cpp
--- a/Engine/Source/Runtime/Engine/Private/MemberReference.cpp
+++ b/Engine/Source/Runtime/Engine/Private/MemberReference.cpp
@@ -125,7 +125,7 @@ FString FMemberReference::GetReferenceSearchString(UClass* InFieldOwner) const
 		{
 			if (MemberGuid.IsValid())
 			{
-				return FString::Printf(TEXT("Nodes(VariableReference(MemberName=+%s && MemberGuid(A=%i && B=%i && C=%i && D=%i) ))"), *MemberName.ToString(), MemberGuid.A, MemberGuid.B, MemberGuid.C, MemberGuid.D);
+				return FString::Printf(TEXT("Nodes(VariableReference(MemberName=+\"%s\" && MemberGuid(A=%i && B=%i && C=%i && D=%i) ))"), *MemberName.ToString(), MemberGuid.A, MemberGuid.B, MemberGuid.C, MemberGuid.D);
 			}
 			else
 			{
@@ -135,21 +135,21 @@ FString FMemberReference::GetReferenceSearchString(UClass* InFieldOwner) const
 				ExportMemberParentName += InFieldOwner->GetAuthoritativeClass()->GetPathName();
 				ExportMemberParentName.AppendChar('\'');
 
-				return FString::Printf(TEXT("Nodes(VariableReference(MemberName=+%s && (MemberParent=\"%s\" || bSelfContext=true) ))"), *MemberName.ToString(), *ExportMemberParentName);
+				return FString::Printf(TEXT("Nodes(VariableReference(MemberName=+\"%s\" && (MemberParent=\"%s\" || bSelfContext=true) ))"), *MemberName.ToString(), *ExportMemberParentName);
 			}
 		}
 		else if (MemberGuid.IsValid())
 		{
-			return FString::Printf(TEXT("Nodes(VariableReference(MemberName=+%s && MemberGuid(A=%i && B=%i && C=%i && D=%i)))"), *MemberName.ToString(), MemberGuid.A, MemberGuid.B, MemberGuid.C, MemberGuid.D);
+			return FString::Printf(TEXT("Nodes(VariableReference(MemberName=+\"%s\" && MemberGuid(A=%i && B=%i && C=%i && D=%i)))"), *MemberName.ToString(), MemberGuid.A, MemberGuid.B, MemberGuid.C, MemberGuid.D);
 		}
 		else
 		{
-			return FString::Printf(TEXT("Nodes(VariableReference(MemberName=+%s))"), *MemberName.ToString());
+			return FString::Printf(TEXT("Nodes(VariableReference(MemberName=+\"%s\"))"), *MemberName.ToString());
 		}
 	}
 	else
 	{
-		return FString::Printf(TEXT("Nodes(VariableReference(MemberName=+%s && MemberScope=+%s))"), *MemberName.ToString(), *GetMemberScopeName());
+		return FString::Printf(TEXT("Nodes(VariableReference(MemberName=+\"%s\" && MemberScope=+\"%s\"))"), *MemberName.ToString(), *GetMemberScopeName());
 	}
 }
 
diff --git a/Engine/Source/Runtime/Engine/Private/SimpleConstructionScript.cpp b/Engine/Source/Runtime/Engine/Private/SimpleConstructionScript.cpp
--- a/Engine/Source/Runtime/Engine/Private/SimpleConstructionScript.cpp
+++ b/Engine/Source/Runtime/Engine/Private/SimpleConstructionScript.cpp
@@ -730,7 +730,7 @@ void USimpleConstructionScript::RemoveNodeAndPromoteChildren(USCS_Node* Node)
 		if(PromoteIndex != INDEX_NONE)
 		{
 			ChildToPromote = Node->GetChildNodes()[PromoteIndex];
-			Node->RemoveChildNodeAt(PromoteIndex);
+			Node->RemoveChildNodeAt(PromoteIndex, false);
 		}
 
 		Modify();
@@ -981,7 +981,7 @@ void USimpleConstructionScript::ValidateSceneRootNodes()
 				if(RootNode != nullptr)
 				{
 					// We have an existing root node within the current BP class.
-					RootNode->AddChildNode(ChildNode);
+					RootNode->AddChildNode(ChildNode, false);
 				}
 				else
 				{
diff --git a/Engine/Source/Runtime/Engine/Private/TimelineTemplate.cpp b/Engine/Source/Runtime/Engine/Private/TimelineTemplate.cpp
--- a/Engine/Source/Runtime/Engine/Private/TimelineTemplate.cpp
+++ b/Engine/Source/Runtime/Engine/Private/TimelineTemplate.cpp
@@ -219,17 +219,6 @@ void UTimelineTemplate::PostDuplicate(bool bDuplicateForPIE)
 
 	UObject* NewCurveOuter = GetOuter();
 
-#if WITH_EDITOR
-	// The outer should always be a BlueprintGeneratedClass, however the Curves get put into the Blueprint, however if they belong to the Transient package, that's where the duplicated Blueprint should be.
-	if(GetOutermost() != GetTransientPackage())
-	{
-		if(UBlueprintGeneratedClass* BPClass = Cast<UBlueprintGeneratedClass>(GetOuter()))
-		{
-			NewCurveOuter = BPClass->ClassGeneratedBy;
-		}
-	}
-#endif // WITH_EDITOR
-
 	const bool bTransientPackage = GetOutermost() == GetTransientPackage();
 	// Prevent curves being duplicated during blueprint reinstancing
 	const bool bDuplicateCurves = !( bTransientPackage || GIsDuplicatingClassForReinstancing );
diff --git a/Engine/Source/Runtime/Engine/Public/Animation/AnimInstanceProxy.h b/Engine/Source/Runtime/Engine/Public/Animation/AnimInstanceProxy.h
--- a/Engine/Source/Runtime/Engine/Public/Animation/AnimInstanceProxy.h
+++ b/Engine/Source/Runtime/Engine/Public/Animation/AnimInstanceProxy.h
@@ -26,7 +26,7 @@ public:
 
 	FAnimInstanceProxy(UAnimInstance* Instance)
 		: AnimInstanceObject(Instance)
-		, AnimBlueprintClass(Cast<UAnimBlueprintGeneratedClass>(Instance->GetClass()))
+		, AnimClassInterface(IAnimClassInterface::GetFromClass(Instance->GetClass()))
 		, Skeleton(nullptr)
 		, SkeletalMeshComponent(nullptr)
 		, CurrentDeltaSeconds(0.0f)
@@ -43,11 +43,19 @@ public:
 
 	virtual ~FAnimInstanceProxy() {}
 
+	// Get the IAnimClassInterface associated with this context, if there is one.
+	// Note: This can return NULL, so check the result.
+	IAnimClassInterface* GetAnimClassInterface() const
+	{
+		return AnimClassInterface;
+	}
+
 	// Get the Blueprint Generated Class associated with this context, if there is one.
 	// Note: This can return NULL, so check the result.
+	DEPRECATED(4.11, "GetAnimBlueprintClass() is deprecated, UAnimBlueprintGeneratedClass should not be directly used at runtime. Please use GetAnimClassInterface() instead.")
 	UAnimBlueprintGeneratedClass* GetAnimBlueprintClass() const
 	{
-		return AnimBlueprintClass;
+		return Cast<UAnimBlueprintGeneratedClass>(IAnimClassInterface::GetActualAnimClass(AnimClassInterface));
 	}
 
 	/** Get the last DeltaSeconds passed into PreUpdate() */
@@ -68,6 +76,12 @@ public:
 	{
 		new (UpdatedNodesThisFrame) FAnimBlueprintDebugData::FNodeVisit(SourceNodeIndex, TargetNodeIndex, BlendWeight);
 	}
+
+	UAnimBlueprint* GetAnimBlueprint() const
+	{
+		UClass* ActualAnimClass = IAnimClassInterface::GetActualAnimClass(AnimClassInterface);
+		return ActualAnimClass ? Cast<UAnimBlueprint>(ActualAnimClass->ClassGeneratedBy) : nullptr;
+	}
 #endif
 
 	// flip sync group read/write indices
@@ -239,7 +253,7 @@ public:
 	FAnimNode_StateMachine* GetStateMachineInstanceFromName(FName MachineName);
 
 	/** Get the machine description for the specified instance. Does not rely on PRIVATE_MachineDescription being initialized */
-	static const FBakedAnimationStateMachine* GetMachineDescription(UAnimBlueprintGeneratedClass* AnimBlueprintClass, FAnimNode_StateMachine* MachineInstance);
+	static const FBakedAnimationStateMachine* GetMachineDescription(IAnimClassInterface* AnimBlueprintClass, FAnimNode_StateMachine* MachineInstance);
 
 	/** 
 	 * Get the index of the specified instance asset player. Useful to pass to GetInstanceAssetPlayerLength (etc.).
@@ -457,7 +471,7 @@ private:
 	mutable UObject* AnimInstanceObject;
 
 	/** Our anim blueprint generated class */
-	UAnimBlueprintGeneratedClass* AnimBlueprintClass;
+	IAnimClassInterface* AnimClassInterface;
 
 	/** Skeleton we are using, only used for comparison purposes */
 	USkeleton* Skeleton;
diff --git a/Engine/Source/Runtime/Engine/Public/GeneratedCodeHelpers.h b/Engine/Source/Runtime/Engine/Public/GeneratedCodeHelpers.h
--- a/Engine/Source/Runtime/Engine/Public/GeneratedCodeHelpers.h
+++ b/Engine/Source/Runtime/Engine/Public/GeneratedCodeHelpers.h
@@ -2,6 +2,9 @@
 
 #pragma once
 
+// DO NOT INCLUDE THIS HEADER!
+// THIS FILE SHOULD BE USED ONLY BY AUTOMATICALLY GENERATED CODE. 
+
 // Common includes
 #include "UObject/Stack.h"
 #include "Blueprint/BlueprintSupport.h"
@@ -17,6 +20,9 @@
 // Special libraries
 #include "Kismet/DataTableFunctionLibrary.h"
 
+//For DOREPLIFETIME macros
+#include "Net/UnrealNetwork.h"
+
 template<class NativeType>
 inline NativeType* NoNativeCast(UClass* NoNativeClass, UObject* Object)
 {
@@ -65,14 +71,14 @@ private:
 public:
 	//Replacements for CustomThunk functions from UKismetArrayLibrary
 
-	template<typename T>
-	static int32 Array_Add(TArray<T>& TargetArray, const T& NewItem)
+	template<typename T, typename U>
+	static int32 Array_Add(TArray<T>& TargetArray, const U& NewItem)
 	{
 		return TargetArray.Add(NewItem);
 	}
 
-	template<typename T>
-	static int32 Array_AddUnique(TArray<T>& TargetArray, const T& NewItem)
+	template<typename T, typename U>
+	static int32 Array_AddUnique(TArray<T>& TargetArray, const U& NewItem)
 	{
 		return TargetArray.AddUnique(NewItem);
 	}
@@ -91,14 +97,14 @@ public:
 		}
 	}
 
-	template<typename T>
-	static void Array_Append(TArray<T>& TargetArray, const TArray<T>& SourceArray)
+	template<typename T, typename U>
+	static void Array_Append(TArray<T>& TargetArray, const TArray<U>& SourceArray)
 	{
 		TargetArray.Append(SourceArray);
 	}
 
-	template<typename T>
-	static void Array_Insert(TArray<T>& TargetArray, const T& NewItem, int32 Index)
+	template<typename T, typename U>
+	static void Array_Insert(TArray<T>& TargetArray, const U& NewItem, int32 Index)
 	{
 		if ((Index >= 0) && (Index <= TargetArray.Num()))
 		{
@@ -123,14 +129,14 @@ public:
 		}
 	}
 
-	template<typename T>
-	static int32 Array_Find(const TArray<T>& TargetArray, const T& ItemToFind)
+	template<typename T, typename U>
+	static int32 Array_Find(const TArray<T>& TargetArray, const U& ItemToFind)
 	{
 		return TargetArray.Find(ItemToFind);
 	}
 
-	template<typename T>
-	static bool Array_RemoveItem(TArray<T>& TargetArray, const T& Item)
+	template<typename T, typename U>
+	static bool Array_RemoveItem(TArray<T>& TargetArray, const U& Item)
 	{
 		return TargetArray.Remove(Item) != 0;
 	}
@@ -166,8 +172,8 @@ public:
 		return TargetArray.Num() - 1;
 	}
 
-	template<typename T>
-	static void Array_Get(TArray<T>& TargetArray, int32 Index, T& Item)
+	template<typename T, typename U>
+	static void Array_Get(TArray<T>& TargetArray, int32 Index, U& Item)
 	{
 		if (TargetArray.IsValidIndex(Index))
 		{
@@ -176,12 +182,12 @@ public:
 		else
 		{
 			ExecutionMessage(*FString::Printf(TEXT("Attempted to get an item from array out of bounds [%d/%d]!"), Index, LastIndexForLog(TargetArray)), ELogVerbosity::Warning);
-			Item = T{};
+			Item = U{};
 		}
 	}
 
-	template<typename T>
-	static void Array_Set(TArray<T>& TargetArray, int32 Index, const T& Item, bool bSizeToFit)
+	template<typename T, typename U>
+	static void Array_Set(TArray<T>& TargetArray, int32 Index, const U& Item, bool bSizeToFit)
 	{
 		if (!TargetArray.IsValidIndex(Index) && bSizeToFit && (Index >= 0))
 		{
@@ -198,8 +204,8 @@ public:
 		}
 	}
 
-	template<typename T>
-	static bool Array_Contains(const TArray<T>& TargetArray, const T& ItemToFind)
+	template<typename T, typename U>
+	static bool Array_Contains(const TArray<T>& TargetArray, const U& ItemToFind)
 	{
 		return TargetArray.Contains(ItemToFind);
 	}
@@ -219,6 +225,11 @@ public:
 	}
 
 	//Replacements for CustomThunk functions from UKismetSystemLibrary
+	static void StackTrace()
+	{
+		ExecutionMessage(TEXT("Native code cannot generate a stack trace."), ELogVerbosity::Log);
+	}
+	
 	template<typename T>
 	static void SetStructurePropertyByName(UObject* Object, FName PropertyName, const T& Value)
 	{
@@ -278,6 +289,21 @@ struct TSwitchPair
 	{}
 };
 
+template<typename IndexType, typename ValueType>
+struct TSwitchPair<IndexType, ValueType*>
+{
+	const IndexType& IndexRef;
+	ValueType*& ValueRef;
+
+	template<typename DerivedType>
+	TSwitchPair(const IndexType& InIndexRef, DerivedType*& InValueRef)
+		: IndexRef(InIndexRef)
+		, ValueRef((ValueType*&)InValueRef)
+	{
+		static_assert(TPointerIsConvertibleFromTo<DerivedType, ValueType>::Value, "Incompatible pointers");
+	}
+};
+
 #if PLATFORM_COMPILER_HAS_VARIADIC_TEMPLATES
 template<typename IndexType, typename ValueType>
 ValueType& TSwitchValue(const IndexType& CurrentIndex, ValueType& DefaultValue, int OptionsNum)
@@ -334,4 +360,28 @@ struct FUnconvertedWrapper
 		check(__Object);
 		return __Object->GetClass();
 	}
+};
+
+template<typename T>
+struct TArrayCaster
+{
+	TArray<T*> Val;
+	TArray<T*>& ValRef;
+
+	TArrayCaster(TArray<T*>& InArr) 
+		: Val()
+		, ValRef(InArr)
+	{}
+
+	TArrayCaster(TArray<T*>&& InArr) 
+		: Val(MoveTemp(InArr))
+		, ValRef(Val)
+	{}
+
+	template<typename U>
+	TArray<U*>& Get()
+	{
+		static_assert(sizeof(T*) == sizeof(U*), "Incompatible pointers");
+		return *reinterpret_cast<TArray<U*>*>(&ValRef);
+	}
 };
\ No newline at end of file
diff --git a/Engine/Source/Runtime/Engine/Public/Net/UnrealNetwork.h b/Engine/Source/Runtime/Engine/Public/Net/UnrealNetwork.h
--- a/Engine/Source/Runtime/Engine/Public/Net/UnrealNetwork.h
+++ b/Engine/Source/Runtime/Engine/Public/Net/UnrealNetwork.h
@@ -124,6 +124,15 @@ static UProperty* GetReplicatedProperty(UClass* CallingClass, UClass* PropClass,
 	}																		\
 }
 
+#define DOREPLIFETIME_DIFFNAMES(c,v, n) \
+{ \
+	static UProperty* sp##v = GetReplicatedProperty(StaticClass(), c::StaticClass(), n); \
+	for ( int32 i = 0; i < sp##v->ArrayDim; i++ )							\
+	{																		\
+		OutLifetimeProps.AddUnique( FLifetimeProperty( sp##v->RepIndex + i ) );	\
+	}																		\
+}
+
 #define DOREPLIFETIME_CONDITION(c,v,cond) \
 { \
 	static UProperty* sp##v = GetReplicatedProperty(StaticClass(), c::StaticClass(),GET_MEMBER_NAME_CHECKED(c,v)); \
diff --git a/Engine/Source/Runtime/GameplayTags/Private/GameplayTagContainer.cpp b/Engine/Source/Runtime/GameplayTags/Private/GameplayTagContainer.cpp
--- a/Engine/Source/Runtime/GameplayTags/Private/GameplayTagContainer.cpp
+++ b/Engine/Source/Runtime/GameplayTags/Private/GameplayTagContainer.cpp
@@ -712,18 +712,25 @@ int32 FGameplayTagContainer::Num() const
 
 FString FGameplayTagContainer::ToString() const
 {
-	FString RetString = TEXT("(GameplayTags=(");
-	for (int i = 0; i < GameplayTags.Num(); ++i)
+	FString RetString = TEXT("(GameplayTags=");
+	if (GameplayTags.Num() > 0)
 	{
-		RetString += TEXT("(TagName=\"");
-		RetString += GameplayTags[i].ToString();
-		RetString += TEXT("\")");
-		if (i < GameplayTags.Num() - 1)
+		RetString += TEXT("(");
+
+		for (int i = 0; i < GameplayTags.Num(); ++i)
 		{
-			RetString += TEXT(",");
+			RetString += TEXT("(TagName=\"");
+			RetString += GameplayTags[i].ToString();
+			RetString += TEXT("\")");
+			if (i < GameplayTags.Num() - 1)
+			{
+				RetString += TEXT(",");
+			}
 		}
+		RetString += TEXT(")");
 	}
-	RetString += TEXT("))");
+	RetString += TEXT(")");
+
 	return RetString;
 }
 
diff --git a/Engine/Source/Runtime/Online/OnlineSubsystemUtils/OnlineSubsystemUtils.Build.cs b/Engine/Source/Runtime/Online/OnlineSubsystemUtils/OnlineSubsystemUtils.Build.cs
--- a/Engine/Source/Runtime/Online/OnlineSubsystemUtils/OnlineSubsystemUtils.Build.cs
+++ b/Engine/Source/Runtime/Online/OnlineSubsystemUtils/OnlineSubsystemUtils.Build.cs
@@ -17,12 +17,13 @@ public class OnlineSubsystemUtils : ModuleRules
 				"Engine", 
 				"EngineSettings",
                 "ImageCore",
-				"OnlineSubsystem",
 				"Sockets",
 				"Voice",
                 "PacketHandler",
 				"Json"
 			}
 		);
+
+        PublicDependencyModuleNames.Add("OnlineSubsystem");
 	}
 }
diff --git a/Engine/Source/Runtime/Projects/Private/ModuleDescriptor.cpp b/Engine/Source/Runtime/Projects/Private/ModuleDescriptor.cpp
--- a/Engine/Source/Runtime/Projects/Private/ModuleDescriptor.cpp
+++ b/Engine/Source/Runtime/Projects/Private/ModuleDescriptor.cpp
@@ -82,6 +82,9 @@ const TCHAR* EHostType::ToString( const EHostType::Type Value )
 		case Program:
 			return TEXT("Program");
 
+		case ServerOnly:
+			return TEXT("ServerOnly");
+
 		default:
 			ensureMsgf( false, TEXT( "Unrecognized EModuleType value: %i" ), Value );
 			return NULL;
@@ -293,6 +296,9 @@ bool FModuleDescriptor::IsCompiledInCurrentConfiguration() const
 			return true;
 		#endif
 		break;
+
+	case EHostType::ServerOnly:
+		return !FPlatformProperties::IsClientOnly();
 	}
 
 	return false;
@@ -344,6 +350,9 @@ bool FModuleDescriptor::IsLoadedInCurrentConfiguration() const
 			return true;
 		#endif
 		break;
+
+	case EHostType::ServerOnly:
+		return !FPlatformProperties::IsClientOnly();
 	}
 
 	return false;
diff --git a/Engine/Source/Runtime/Projects/Public/ModuleDescriptor.h b/Engine/Source/Runtime/Projects/Public/ModuleDescriptor.h
--- a/Engine/Source/Runtime/Projects/Public/ModuleDescriptor.h
+++ b/Engine/Source/Runtime/Projects/Public/ModuleDescriptor.h
@@ -65,7 +65,7 @@ namespace EHostType
 		Editor,
 		EditorNoCommandlet,
 		Program,		//!< Program-only plugin type
-
+		ServerOnly,
 		// NOTE: If you add a new value, make sure to update the ToString() method below!
 
 		Max
