xbox one cis fix

diff --git a/Engine/Source/Runtime/Engine/Private/UnrealEngine.cpp b/Engine/Source/Runtime/Engine/Private/UnrealEngine.cpp
--- a/Engine/Source/Runtime/Engine/Private/UnrealEngine.cpp
+++ b/Engine/Source/Runtime/Engine/Private/UnrealEngine.cpp
@@ -1807,6 +1807,9 @@ class FFakeStereoRenderingDevice : public IStereoRendering
 {
 public:
 	FFakeStereoRenderingDevice() 
+	: FOVInDegrees(100)
+	, Width(640)
+	, Height(480)
 	{
 		static TAutoConsoleVariable<float> CVarEmulateStereoFOV(TEXT("r.StereoEmulationFOV"), 0, TEXT("FOV in degrees, of the imaginable HMD for stereo emulation"));
 		static TAutoConsoleVariable<int32> CVarEmulateStereoWidth(TEXT("r.StereoEmulationWidth"), 0, TEXT("Width of the imaginable HMD for stereo emulation"));
@@ -1907,8 +1910,8 @@ public:
 		RHICmdList.Clear(true, FLinearColor::Black, false, 0, false, 0, FIntRect());
 	}
 
-	float FOVInDegrees = 100;		 // max(HFOV, VFOV) in degrees of imaginable HMD
-	int32 Width = 640, Height = 480; // resolution of imaginable HMD
+	float FOVInDegrees;		// max(HFOV, VFOV) in degrees of imaginable HMD
+	int32 Width, Height;	// resolution of imaginable HMD
 };
 
 bool UEngine::InitializeHMDDevice()
