Fix crash when calling RefreshBoneTransforms on a SkeletalMesh that has no Skeleton set

diff --git a/Engine/Source/Runtime/Engine/Private/Components/SkeletalMeshComponent.cpp b/Engine/Source/Runtime/Engine/Private/Components/SkeletalMeshComponent.cpp
--- a/Engine/Source/Runtime/Engine/Private/Components/SkeletalMeshComponent.cpp
+++ b/Engine/Source/Runtime/Engine/Private/Components/SkeletalMeshComponent.cpp
@@ -1064,11 +1064,13 @@ void USkeletalMeshComponent::PerformAnimationEvaluation(const USkeletalMesh* InS
 
 int32 GetCurveNumber(USkeleton* Skeleton)
 {
-	check (Skeleton);
 	// get all curve list
-	if(const FSmartNameMapping* Mapping = Skeleton->GetSmartNameContainer(USkeleton::AnimCurveMappingName))
+	if (Skeleton)
 	{
-		return Mapping->GetNumNames();
+		if (const FSmartNameMapping* Mapping = Skeleton->GetSmartNameContainer(USkeleton::AnimCurveMappingName))
+		{
+			return Mapping->GetNumNames();
+		}
 	}
 
 	return 0;
@@ -2358,8 +2360,10 @@ void USkeletalMeshComponent::FinalizeBoneTransform()
 
 void USkeletalMeshComponent::UpdateCachedAnimCurveMappingNameUids()
 {
-	check(SkeletalMesh->Skeleton != nullptr);
-	CachedAnimCurveMappingNameUids = SkeletalMesh->Skeleton->GetCachedAnimCurveMappingNameUids();
+	if (SkeletalMesh->Skeleton)
+	{
+		CachedAnimCurveMappingNameUids = SkeletalMesh->Skeleton->GetCachedAnimCurveMappingNameUids();
+	}
 }
 
 TArray<FSmartNameMapping::UID> const * USkeletalMeshComponent::GetCachedAnimCurveMappingNameUids()
