Copying //UE4/Dev-Networking to Dev-Main (//UE4/Dev-Main)

==========================
MAJOR FEATURES + CHANGES
==========================

Change 2775993 on 2015/11/20 by [REDACTED_PERSON][REDACTED_HANDLE]

	Add a shared ptr valid check before a log message that uses it

Change 2778442 on 2015/11/23 by [REDACTED_PERSON][REDACTED_HANDLE]

	Fix erroneous error message saying that the session interface init failed if rooms aren't enabled.

Change 2783404 on 2015/11/30 by [REDACTED_PERSON][REDACTED_HANDLE]

	Added Oodle dictionary generation commandlet. Rewrote and finalized dictionary serialization and file format. Added 'stat oodle' command, among other tweaks.

	Full list of changes:
	- Added "-run=OodleHandlerComponent.OodleTrainerCommandline AutoGenerateDictionaries" command, which automatically processes all Oodle packet captures, into dictionary files, placed in *Game*\\Content\\Oodle\\
	- Added "GenerateDictionary" command to same commandlet, for more fine-tuned inclusion of packet capture files
	- Both of the above have a lot of to-be-documented commandline switches, for tweaking the Oodle dictionary generation: "-HashTableSize=x -DictionarySize=x -DictionaryTrials=x -TrialRandomness=x -TrialGenerations=x -NoTrials"

	- Loaded dictionary data is now shared between net connections, instead of being loaded separately for each connection, saving memory/performance
	- Added 'stat oodle' command, which displays Oodle CPU/memory performance, and net/compression stats
	- Tweaked "-Oodle" and "-OodleCapturing" commandline parameters, to automatically setup Oodle .ini settings, if not present

Change 2785877 on 2015/12/01 by [REDACTED_PERSON][REDACTED_HANDLE]

	Network actor list is now stored on each UNetDriver instead of on UWorld - important foundation for future optimizations.
	Enables each driver to manage and sort the list separately, and to add metadata about actors without increasing the size of AActor.

Change 2787343 on 2015/12/02 by [REDACTED_PERSON][REDACTED_HANDLE]

	Amortize replay replication more evenly across frames to avoid spikes on servers, similar to how ServerReplicateActors handles NetUpdateFrequency.

Change 2788519 on 2015/12/03 by [REDACTED_PERSON][REDACTED_HANDLE]

	Added initial Oodle Dictionaries in OrionGameGame\\Content\\Oodle

Change 2788657 on 2015/12/03 by [REDACTED_PERSON][REDACTED_HANDLE]

	Initial implementation of adaptive net update frequency

	#tests Will be applying lots of tests in dev-networking

Change 2790209 on 2015/12/04 by [REDACTED_PERSON][REDACTED_HANDLE]

	Some net object list fixes:
	-AddInitialObjects now properly filters by NetDriverName.
	-DemoNetDriver no longer explicitly replicates UWorldSettings, since it's replicated through the normal path now.
	-Improved comments

Change 2795200 on 2015/12/08 by [REDACTED_PERSON][REDACTED_HANDLE]

	Remove UNetConnection::OwnedConsiderList + major cleanup of ServerReplicateActors

	#tests Shootergame, Paragon, split-screen

Change 2795301 on 2015/12/08 by [REDACTED_PERSON][REDACTED_HANDLE]

	Make UNetDriver::NetworkObjectList a pointer so that it can be forward declared and reduce compile times while iterating on NetworkObjectList.h.

Change 2796374 on 2015/12/09 by [REDACTED_PERSON][REDACTED_HANDLE]

	Use similar adaptive net update frequency logic from ServerReplicateActors for replay recording. Saves CPU time by processing actors that don't replicate often less frequently. Toggle with cvar demo.UseAdaptiveReplayUpdateFrequency.

Change 2799744 on 2015/12/11 by [REDACTED_PERSON][REDACTED_HANDLE]

	Fix issue where rejoining a steam lobby quickly after leaving would cause the player to get booted from the session

Change 2802422 on 2015/12/14 by [REDACTED_PERSON][REDACTED_HANDLE]

	Ability to have single pass, partial network profile streams (i.e. you don't need to change maps to force it to do the final save)

Change 2802572 on 2015/12/14 by [REDACTED_PERSON][REDACTED_HANDLE]

	Add new replication condition, COND_ReplayOrOwner, which will replicate properties to the owning connection and to the replay connection, if one exists. Useful for data that wants to be COND_OwnerOnly, but also visible during replay playback.
	Change replication condition of UAbilitySystemComponent::ActivatableAbilities to COND_ReplayOrOwner.

Change 2804273 on 2015/12/15 by [REDACTED_PERSON][REDACTED_HANDLE]

	Quick fix for actors not getting added to the proper net driver. It appears sometimes netdrivers liike beacons change the name afterward, so we'll need to re-think that. But for now, just add all actors, and check at runtime like we used to until we find a better fix.

diff --git a/Engine/Plugins/Runtime/PacketHandlers/CompressionComponents/Oodle/BuildDictionary.bat b/Engine/Plugins/Runtime/PacketHandlers/CompressionComponents/Oodle/BuildDictionary.bat
--- a/Engine/Plugins/Runtime/PacketHandlers/CompressionComponents/Oodle/BuildDictionary.bat
+++ b/Engine/Plugins/Runtime/PacketHandlers/CompressionComponents/Oodle/BuildDictionary.bat
@@ -1,7 +1,6 @@
 @echo off
 
-
-echo This batch file, takes .ucap packet capture files from Oodle 'training' mode play sessions, and processes them into a dictionary file.
+echo This batch file, goes through the fully automated process of building Oodle dictionaries from packet captures.
 echo.
 
 
@@ -17,144 +16,37 @@ goto End
 :SetUE4Editor
 set UE4EditorLoc="%BaseFolder:"=%\Engine\Binaries\Win64\UE4Editor.exe"
 
-if exist %UE4EditorLoc:"=% goto SetTrainer
+if exist %UE4EditorLoc:"=% goto GetGame
 
 echo Could not locate UE4Editor.exe
 goto End
 
 
-:SetTrainer
-set TrainerLoc="Source\ThirdParty\NotForLicensees\OodleTrainer\x64\Release\OodleTrainer.exe"
-
-if exist %TrainerLoc:"=% goto GetGame
-
-echo Could not locate OodleTrainer.exe, make sure it is built.
-goto End
-
-
 :GetGame
 set /p GameName=Type the name of the game you are working with: 
 echo.
 
-if exist %BaseFolder:"=%\%GameName:"=% goto GetCaptures
-if exist %BaseFolder:"=%\Samples\Games\%GameName:"=% goto GetCaptures
-
-set /p SkipGameCheck=Could not locate folder for game '%GameName:"=%'; continue anyway y/n?
-if not "%SkipGameCheck%"=="y" goto End
-
-
-:GetCaptures
-echo Place all .ucap files you want to process, into a folder, and then drag+drop the folder over this prompt.
-echo Output and Input capture files, should be processed separately.
-
-set /p CapturesFolder=Drag or paste captures folder here: 
-echo.
-
-if exist %CapturesFolder% goto GetOutput
-
-echo Could not locate folder: %CapturesFolder%
-goto End
-
-
-:GetOutput
-set BaseOutputPath="%BaseFolder:"=%\%GameName:"=%\Content\Oodle"
-
-if not exist "%BaseOutputPath%" mkdir "%BaseOutputPath%"
-
-echo Type the desired name of the dictionary file. This will be placed in %GameName:"=%\Content\Oodle\
-set /p OutputFile=Dictionary file name: 
-echo.
-
-if not exist %BaseOutputPath:"=%\%OutputFile:"=% goto ExecuteCommandlet
-
-set /p OverwriteOutput=Output file already exists. Overwrite y/n?
-if not "%OverwriteOutput%"=="y" goto End
-
-echo.
-
-
-:ExecuteCommandlet
-set MergedPacketsFile="%~dp0TempMerge.ucap"
-set CommandletParms=-run=OodleHandlerComponent.OodleTrainerCommandlet MergePackets "%MergedPacketsFile:"=%" all "%CapturesFolder:"=%
-set FinalUE4CmdLine=%GameName:"=% %CommandletParms%
-
-echo Executing Oodle MergePackets commandlet - commandline:
-echo %FinalUE4CmdLine%
-
-@echo on
-%UE4EditorLoc:"=% %FinalUE4CmdLine%
-@echo off
-echo.
-
-
-if %errorlevel%==0 goto ExecuteTrainer
 
-echo WARNING! Detected error! Review logfile for commandlet.
-goto CommandletCleanup
 
+:AutoGenDictionaries
+set AutoGenDictionariesParms=-run=OodleHandlerComponent.OodleTrainerCommandlet AutoGenerateDictionaries
+set FinalGenCmdLine=%GameName:"=% %AutoGenDictionariesParms% -forcelogflush
 
-:ExecuteTrainer
-if not exist %BaseFolder:"=%\%GameName:"=% goto AltExecuteTrainer
-
-set FinalOutputFile="%~dp0%BaseOutputPath:"=%\%OutputFile:"=%"
-set FinalTrainerCmdLine="%MergedPacketsFile:"=%" "%FinalOutputFile:"=%"
-
-echo Executing trainer for producing final dictionary file - commandline:
-echo %FinalTrainerCmdLine%
+echo Executing dictionary generation commandlet - commandline:
+echo %FinalGenCmdLine%
 
 @echo on
-%TrainerLoc:"=% %FinalTrainerCmdLine%
+%UE4EditorLoc:"=% %FinalGenCmdLine%
 @echo off
-
 echo.
 
-if not %errorlevel%==0 goto TrainerError
-goto TrainerSuccess
-
-
-:AltExecuteTrainer
-set AltOutputFile="%~dp0%OutputFile:"=%"
-set FinalAltTrainerCmdLine="%MergedPacketsFile:"=%" %AltOutputFile%
-
-echo Executing trainer for producing final dictionary file - commandline:
-echo %FinalAltTrainerCmdLine%
-
-@echo on
-%TrainerLoc:"=% %FinalAltTrainerCmdLine%
-@echo off
-
-if not %errorlevel%==0 goto TrainerError
 
+if %errorlevel%==0 goto End
 
-echo.
-echo As the game folder was not detected, you need to manually copy:
-echo %AltOutputFile%
-echo.
-echo To the folder:
-echo %GameName:"=%\Content\Oodle
-echo.
+echo WARNING! Detected error, dictionaries may not have been generated. Check output and logfile for errors.
 pause
 
 
-:TrainerSuccess
-echo.
-echo Dictionary successfully generated! Set '*Dictionary=Content\Oodle\%OutputFile:"=%' under [OodleHandlerComponent], in *Engine.ini.
-echo If you want to switch Oodle from 'Training' mode to 'Release' mode, set 'Mode=Release' in the same location.
-goto PostTrainer
-echo.
-
-
-:TrainerError
-echo Error generating dictionary file. Review trainer output.
-goto End
-
-
-:PostTrainer
-:CommandletCleanup
-del %MergedPacketsFile% 
-
-
-
 :End
 echo Execution complete.
 pause
@@ -162,4 +54,3 @@ pause
 
 REM Put nothing past here.
 
-
diff --git a/Engine/Plugins/Runtime/PacketHandlers/CompressionComponents/Oodle/Source/OodleHandlerComponent/Classes/OodleTrainerCommandlet.h b/Engine/Plugins/Runtime/PacketHandlers/CompressionComponents/Oodle/Source/OodleHandlerComponent/Classes/OodleTrainerCommandlet.h
--- a/Engine/Plugins/Runtime/PacketHandlers/CompressionComponents/Oodle/Source/OodleHandlerComponent/Classes/OodleTrainerCommandlet.h
+++ b/Engine/Plugins/Runtime/PacketHandlers/CompressionComponents/Oodle/Source/OodleHandlerComponent/Classes/OodleTrainerCommandlet.h
@@ -10,7 +10,21 @@
 /**
  * Commandlet for processing UE4 packet captures, through Oodle's training API, for generating compressed state dictionaries.
  *
- * Commands:
+ *
+ * Primary Commands:
+ *	- "AutoGenerateDictionaries":
+ *		- Iterates every directory within "*Game*\Saved\Oodle\Server", and uses all capture files within each directory,
+ *			to generate a dictionary stored in "*Game*\Content\Oodle", named "*Game**DirectoryName*.udic".
+ *
+ *		- For example, packet captures in "OrionGame\Saved\Oodle\Server\Input", will be generated into a dictionary stored in
+ *			"OrionGame\Content\Oodle\OrionGameInput.udic"
+ *
+ *		- Each folder within "*Game*\Content\Oodle", should contain at least 100mb of packet captures.
+ *
+ *		- @todo #JohnB: Document commandline switches
+ *
+ *
+ * Secondary/Testing Commands:
  *	- "Enable":
  *		- Inserts the Oodle PacketHandler into the games packet handler component list, and initializes Oodle *Engine.ini settings
  *
@@ -22,13 +36,14 @@
  *		- As above, but merges all capture files in the specified directory.
  *
  *
- * @todo #JohnB: Unimplemented commands:
- *	- "GenerateDictionary OutputFile PacketFile"
- *		- Takes the specified packet capture file, and processes it into an dictionary state file, used for compressing packet data.
- *		- The packet file (usually consisting of many merged packet files), should be at least 100MB in size (see Oodle docs)
- *		- @todo #JohnB: Implement
+ *	- "GenerateDictionary OutputFile PacketFile1,PacketFile2,PacketFileN":
+ *		- Takes the specified packet capture files, and uses them to generate a network compression dictionary
+ *
+ *	- "GenerateDictionary OutputFile All Directory":
+ *		- As above, but uses all capture files in the specified directory, to generate a network compression dictionary
  *
  *
+ * @todo #JohnB: Unimplemented commands:
  *	- "PacketInfo PacketFile":
  *		- Outputs information about the packet file, such as the MB amount of data recorded, per net connection channel, and data types
  *		- @todo #JohnB: Only implement, if deciding to actually capture/track this kind of data
@@ -40,6 +55,180 @@ class UOodleTrainerCommandlet : public UCommandlet
 
 public:
 	virtual int32 Main(const FString& Params) override;
+
+#if HAS_OODLE_SDK
+	/**
+	 * Handles the 'enable' command, which enables the Oodle packet handler component
+	 *
+	 * @return	Whether or not the command executed successfully
+	 */
+	static bool HandleEnable();
+
+	/**
+	 * Handles the 'MergePackets' command, which is used to merge multiple packet capture files
+	 *
+	 * @param OutputCapFile		The file which will contain the merged packets
+	 * @param MergeList			The list of packet capture files to merge, or if bMergeAll, the directory containing these files
+	 * @param bMergeAll			Whether or not MergeList refers to a directory, where all of the files it contains should be merged
+	 * @return					Whether or not the command executed successfully
+	 */
+	bool HandleMergePackets(FString OutputCapFile, FString MergeList, bool bMergeAll=false);
+
+
+	/**
+	 * Handles the 'AutoGenerateDictionaries' command, which is used to automatically detect packet captures and produce dictionaries
+	 *
+	 * This is done, by iterating every directory within "*Game*\Saved\Oodle\Server", and using all capture files within each directory,
+	 * to generate a dictionary name "*Game**DirectoryName*.udic", stored in "*Game*\Content\Oodle".
+	 *
+	 * For example, captures in "OrionGame\Saved\Oodle\Server\Input", generate "OrionGame\Content\Oodle\OrionGameInput.udic".
+	 *
+	 * @return	Whether or not the command executed successfully
+	 */
+	bool HandleAutoGenerateDictionaries();
+
+
+
+	/**
+	 * Converts a list of capture files to merge, into a map of file archives vs file names (doing all necessary verification etc.)
+	 * NOTE: OutMergeMap FArchive*'s must be deleted by the caller.
+	 *
+	 * @param MergeList			The list of capture files to merge (or the directory containing files to merge, if bMergeAll)
+	 * @param OutMergeMap		The output map of file archives vs file names
+	 * @param bMergeAll			If true, merges all capture files in the specified directory (with MergeList pointing to a directory)
+	 * @param bAllowSingleFile	Whether or not the function will error when only a single file is found
+	 * @return					Whether or not the list was successfully parsed, and merge map successfully assigned
+	 */
+	static bool GetMergeMapFromList(FString MergeList, TMap<FArchive*, FString>& OutMergeMap, bool bMergeAll=false,
+									bool bAllowSingleFile=false);
+
+	/**
+	 * Checks that the output file does not already exist, and prompts for an overwrite, if it does
+	 *
+	 * @param OutputFile	The output file string
+	 * @return				Whether or not OutputFile is a valid path
+	 */
+	static bool VerifyOutputFile(FString OutputFile);
+#endif
 };
 
 
+#if HAS_OODLE_SDK
+/**
+ * FOodleDictionaryGenerator
+ *
+ * This class encapsulates dictionary generation, and splits it into multiple stages for readability
+ */
+class FOodleDictionaryGenerator
+{
+	/** Input/Parameter variables */
+private:
+	/** The path for outputting the generated dictionary */
+	FString OutputDictionaryFile;
+
+	/** Whether or not compression testing should be performed after dictionary generation (uses up some of the packets) */
+	bool bCompressionTestPackets;
+
+	/** Size of the hash table to use for the dictionary */
+	int32 HashTableSize;
+
+	/** Size of the dictionary to be generated */
+	int32 DictionarySize;
+
+	/** The number of random packet-selection trials to run, when generating the dictionary, to try and optimize the dictionary */
+	int32 DictionaryTrials;
+
+	/** The randomness, in percent, of random packet-selection trials */
+	int32 TrialRandomness;
+
+	/** The number of generations of random packet-selection trials  */
+	// @todo #JohnB: Examine this more closely, and document better (see email from Charles on 20/11/15)
+	int32 TrialGenerations;
+
+	/** Whether or not random-trials have been disabled */
+	bool bNoTrials;
+
+
+	/** Runtime variables (Opaque) */
+private:
+	TMap<FArchive*, FString> MergeMap;
+
+
+	TArray<void*>	DictionaryPackets;
+	TArray<S32>		DictionaryPacketSizes;
+	uint32			DictionaryPacketBytes;
+	TArray<void*>	DictionaryTestPackets;
+	TArray<S32>		DictionaryTestPacketSizes;
+	uint32			DictionaryTestPacketBytes;
+	TArray<void*>	TrainerPackets;
+	TArray<S32>		TrainerPacketSizes;
+	uint32			TrainerPacketBytes;
+	TArray<void*>	CompressionTestPackets;
+	TArray<S32>		CompressionTestPacketSizes;
+	uint32			CompressionTestPacketBytes;
+
+	bool bDictionaryTestOverflow;
+
+public:
+	/**
+	 * Base constructor
+	 */
+	FOodleDictionaryGenerator()
+		: OutputDictionaryFile()
+		, bCompressionTestPackets(false)
+		, HashTableSize(19)
+		, DictionarySize(4 * 1024 * 1024)
+		, DictionaryTrials(3)
+		, TrialRandomness(200)
+		, TrialGenerations(1)
+		, bNoTrials(false)
+		, MergeMap()
+		, DictionaryPacketBytes(0)
+		, DictionaryTestPacketBytes(0)
+		, TrainerPacketBytes(0)
+		, CompressionTestPacketBytes(0)
+		, bDictionaryTestOverflow(false)
+	{
+	}
+
+	/**
+	 * Primary function handling generation of the dictionary
+	 *
+	 * @param InOutputDictionaryFile	The path for the final dictionary file
+	 * @param InputCaptureFiles			The capture file to process, or if bMergeAll, the directory where all captures files are located
+	 * @param bMergeAll					Whether or not to merge all capture files, in the folder specified by InputCaptureFiles
+	 * @return							Whether or not the dictionary generation was successful
+	 */
+	bool BeginGenerateDictionary(FString InOutputDictionaryFile, FString InputCaptureFiles, bool bMergeAll);
+
+
+private:
+	/**
+	 * Initialize the dictionary generate parameters and state
+	 *
+	 * @return	Whether or not initialization was successful
+	 */
+	bool InitGenerator();
+
+	/**
+	 * Reads the specified capture files, and loads/sorts them in memory, in preparation for processing
+	 *
+	 * @param InputCaptureFiles		As per BeginGenerateDictionary
+	 * @param bMergeAll				As per BeginGenerateDictionary
+	 * @return						Whether or not packet reading was successful
+	 */
+	bool ReadPackets(FString InputCaptureFiles, bool bMergeAll);
+
+	/**
+	 * Processes loaded packet data through the Oodle dictionary generation API, and then writes/compresses the result to the final file
+	 *
+	 * @return	Whether or not final dictionary generation was successful
+	 */
+	bool GenerateAndWriteDictionary();
+
+	/**
+	 * Cleans up any leftover allocated memory
+	 */
+	void Cleanup();
+};
+#endif
diff --git a/Engine/Plugins/Runtime/PacketHandlers/CompressionComponents/Oodle/Source/OodleHandlerComponent/Private/OodleArchives.cpp b/Engine/Plugins/Runtime/PacketHandlers/CompressionComponents/Oodle/Source/OodleHandlerComponent/Private/OodleArchives.cpp
--- a/Engine/Plugins/Runtime/PacketHandlers/CompressionComponents/Oodle/Source/OodleHandlerComponent/Private/OodleArchives.cpp
+++ b/Engine/Plugins/Runtime/PacketHandlers/CompressionComponents/Oodle/Source/OodleHandlerComponent/Private/OodleArchives.cpp
@@ -5,24 +5,155 @@
 #if HAS_OODLE_SDK
 #include "OodleArchives.h"
 
+// Maximum size of compress/decompress buffers (just under 2GB, due to max int32 value)
+#define MAX_COMPRESS_BUFFER (1024 * 1024 * 2047)
+
+
+// @todo #JohnB: Deprecate all pre-CAPTURE_VER_PACKETCOUNT file versions, after a short period
+
 
 /**
- * FPacketCaptureArchive
+ * FOodleCompressedData
  */
 
-void FPacketCaptureArchive::SerializeCaptureHeader()
+void FOodleArchiveBase::FOodleCompressedData::Serialize(FArchive& Ar)
 {
-	uint32 Magic = CAPTURE_HEADER_MAGIC;
+	Ar << Offset;
+	Ar << CompressedLength;
+	Ar << DecompressedLength;
+}
 
-	InnerArchive << Magic;
-	InnerArchive << CaptureVersion;
 
-	if (IsLoading())
+/**
+ * FOodleArchiveBase
+ */
+
+bool FOodleArchiveBase::SerializeOodleCompressData(FOodleCompressedData& OutDataInfo, void* Data, uint32 DataBytes)
+{
+	bool bSuccess = true;
+
+	bSuccess = ensure(Data != nullptr);
+	bSuccess = bSuccess && ensure(DataBytes > 0);
+	bSuccess = bSuccess && ensure(DataBytes <= MAX_COMPRESS_BUFFER);
+
+	if (bSuccess)
 	{
-		check(Magic == CAPTURE_HEADER_MAGIC);
-		check(CaptureVersion <= CAPTURE_FILE_VERSION);
+		OutDataInfo.DecompressedLength.Set(*this, DataBytes);
+	
+		uint32 CompressBufferLen = OodleLZ_GetCompressedBufferSizeNeeded(DataBytes);
+		void* CompressBuffer = new uint8[CompressBufferLen];
+
+		SINTa OodleLen = OodleLZ_Compress(OodleLZ_Compressor_LZNIB, Data, (SINTa)DataBytes, CompressBuffer,
+											OodleLZ_CompressionLevel_Optimal);
+
+		bSuccess = ensure(OodleLen != OODLELZ_FAILED);
+
+		if (bSuccess)
+		{
+			OutDataInfo.CompressedLength.Set(*this, OodleLen);
+
+			uint32 OffsetPos = InnerArchive.Tell();
+			OutDataInfo.Offset.Set(*this, OffsetPos);
+
+			InnerArchive.Serialize(CompressBuffer, OodleLen);
+		}
+
+
+		delete[] CompressBuffer;
 	}
-	else if (IsSaving() && bImmediateFlush)
+
+	return bSuccess;
+}
+
+bool FOodleArchiveBase::SerializeOodleDecompressData(FOodleCompressedData& DataInfo, void*& OutData, uint32& OutDataBytes)
+{
+	check(OutData == nullptr);
+
+	bool bSuccess = true;
+	uint32 DecompressedLength = DataInfo.DecompressedLength.Get();
+	uint32 CompressedLength = DataInfo.CompressedLength.Get();
+	uint32 DataOffset = DataInfo.Offset.Get();
+
+	bSuccess = ensure(CompressedLength <= (InnerArchive.TotalSize() - DataOffset));
+	bSuccess = bSuccess && ensure(DecompressedLength <= MAX_COMPRESS_BUFFER);
+	bSuccess = bSuccess && ensure(CompressedLength <= MAX_COMPRESS_BUFFER);
+
+	if (bSuccess)
+	{
+		SeekPush(DataOffset);
+
+		void* CompressedData = new uint8[CompressedLength];
+		void* DecompressedData = new uint8[DecompressedLength];
+
+		InnerArchive.Serialize(CompressedData, CompressedLength);
+
+		SINTa OodleLen = OodleLZ_Decompress(CompressedData, CompressedLength, DecompressedData, DecompressedLength);
+
+		bSuccess = OodleLen != 0;
+
+		if (bSuccess)
+		{
+			check(OodleLen == DecompressedLength);
+
+			OutData = DecompressedData;
+			OutDataBytes = DecompressedLength;
+		}
+		else
+		{
+			delete[] DecompressedData;
+		}
+
+		delete[] CompressedData;
+
+		SeekPop();
+	}
+
+	return bSuccess;
+}
+
+
+/**
+ * FCaptureHeader
+ */
+
+void FPacketCaptureArchive::FCaptureHeader::SerializeHeader(FPacketCaptureArchive& Ar)
+{
+	bool bSuccess = true;
+
+	Ar << Magic;
+	Ar << CaptureVersion;
+
+	if (Ar.IsLoading())
+	{
+		bSuccess = ensure(Magic == CAPTURE_HEADER_MAGIC);
+		bSuccess = bSuccess && ensure(CaptureVersion <= CAPTURE_FILE_VERSION);
+
+		if (!bSuccess)
+		{
+			Ar.SetError();
+		}
+	}
+
+	if (CaptureVersion >= CAPTURE_VER_PACKETCOUNT)
+	{
+		Ar << PacketCount;
+		Ar << PacketDataOffset;
+		Ar << PacketDataLength;
+	}
+
+	PacketDataOffset.Set(Ar, Ar.Tell());
+}
+
+
+/**
+ * FPacketCaptureArchive
+ */
+
+void FPacketCaptureArchive::SerializeCaptureHeader()
+{
+	Header.SerializeHeader(*this);
+
+	if (IsSaving() && bImmediateFlush)
 	{
 		Flush();
 	}
@@ -30,11 +161,21 @@ void FPacketCaptureArchive::SerializeCaptureHeader()
 
 void FPacketCaptureArchive::SerializePacket(void* PacketData, uint32& PacketSize)
 {
+	check(Header.PacketDataOffset.Get() != 0);
+	check(Tell() >= Header.PacketDataOffset.Get());
+
 	uint32 PacketBufferSize = (IsLoading() ? PacketSize : 0);
+	uint64 StartPos = Tell();
 	InnerArchive << PacketSize;
 
 	if (IsLoading())
 	{
+		if (Header.CaptureVersion >= CAPTURE_VER_PACKETCOUNT)
+		{
+			// Added PacketSize here, to deliberately overshoot, in case PacketDataLength is not updated in the file (possible)
+			check(Tell() < Header.PacketDataOffset.Get() + Header.PacketDataLength.Get() + (uint32)sizeof(PacketSize) + PacketSize);
+		}
+
 		// Max 128MB packet - excessive, but this is not meant to be a perfect security check
 		if (PacketBufferSize < PacketSize || !ensure(PacketSize <= 134217728))
 		{
@@ -58,9 +199,18 @@ void FPacketCaptureArchive::SerializePacket(void* PacketData, uint32& PacketSize
 
 	InnerArchive.Serialize(PacketData, PacketSize);
 
-	if (IsSaving() && bImmediateFlush)
+	if (IsSaving())
 	{
-		Flush();
+		uint32 NewPacketCount = Header.PacketCount.Get() + 1;
+		uint32 NewPacketDataLength = Header.PacketDataLength.Get() + (Tell() - StartPos);
+
+		Header.PacketCount.Set(*this, NewPacketCount);
+		Header.PacketDataLength.Set(*this, NewPacketDataLength);
+
+		if (bImmediateFlush)
+		{
+			Flush();
+		}
 	}
 }
 
@@ -74,10 +224,10 @@ void FPacketCaptureArchive::AppendPacketFile(FPacketCaptureArchive& InPacketFile
 	// Read past the header
 	InPacketFile.SerializeCaptureHeader();
 
-	check(CaptureVersion == InPacketFile.CaptureVersion);
+	check(Header.CaptureVersion == InPacketFile.Header.CaptureVersion);
 
 	// For appending, only support 1MB packets
-	const uint32 BufferSize = 1048576;
+	const uint32 BufferSize = 1024 * 1024;
 	uint8* ReadBuffer = new uint8[BufferSize];
 
 	// Iterate through all packets
@@ -106,5 +256,137 @@ void FPacketCaptureArchive::AppendPacketFile(FPacketCaptureArchive& InPacketFile
 		Flush();
 	}
 }
+
+
+uint32 FPacketCaptureArchive::GetPacketCount()
+{
+	uint32 ReturnVal = 0;
+
+	if (IsSaving())
+	{
+		ReturnVal = Header.PacketCount.Get();
+	}
+	else
+	{
+		if (Header.CaptureVersion >= CAPTURE_VER_PACKETCOUNT)
+		{
+			ReturnVal = Header.PacketCount.Get();
+		}
+		// Do it the hard way, by recomputing through stepping through all packets
+		// @todo #JohnB: Deprecate this, as part of removing pre-'CAPTURE_VER_PACKETCOUNT' file version support
+		else
+		{
+			int64 ArcTotal = InnerArchive.TotalSize();
+
+			check(Header.PacketDataOffset.Get() != 0);
+
+			SeekPush(Header.PacketDataOffset.Get());
+
+			while ((InnerArchive.Tell() + (int64)sizeof(uint32)) < ArcTotal)
+			{
+				uint32 PacketSize = 0;
+
+				InnerArchive << PacketSize;
+
+				int64 NewPos = InnerArchive.Tell() + PacketSize;
+
+				if (NewPos <= ArcTotal)
+				{
+					InnerArchive.Seek(NewPos);
+					ReturnVal++;
+				}
+			}
+
+			SeekPop();
+		}
+	}
+
+	return ReturnVal;
+}
+
+
+/**
+ * FDictionaryHeader
+ */
+
+void FOodleDictionaryArchive::FDictionaryHeader::SerializeHeader(FOodleDictionaryArchive& Ar)
+{
+	bool bSuccess = true;
+
+	Ar << Magic;
+	Ar << DictionaryVersion;
+	Ar << OodleMajorHeaderVersion;
+	Ar << HashTableSize;
+	Ar << DictionaryData;
+	Ar << CompressorData;
+
+	if (Ar.IsLoading())
+	{
+		if (Magic == 0x11235801)
+		{
+			LowLevelFatalError(TEXT("Tried to load dictionary from old format. Regenerate dictionary using the trainer commandlet."));
+			bSuccess = false;
+		}
+
+		bSuccess = bSuccess && ensure(Magic == DICTIONARY_HEADER_MAGIC);
+		bSuccess = bSuccess && ensure(DictionaryVersion <= DICTIONARY_FILE_VERSION);
+		bSuccess = bSuccess && ensure(OodleMajorHeaderVersion == OODLE2_VERSION_MAJOR);
+	}
+
+	if (!bSuccess)
+	{
+		Ar.SetError();
+	}
+}
+
+
+/**
+ * FOodleDictionaryArchive
+ */
+
+FOodleDictionaryArchive::FOodleDictionaryArchive(FArchive& InInnerArchive)
+	: FOodleArchiveBase(InInnerArchive)
+	, Header()
+{
+	if (IsSaving())
+	{
+		Header.DictionaryVersion = DICTIONARY_FILE_VERSION;
+		Header.OodleMajorHeaderVersion = OODLE2_VERSION_MAJOR;
+	}
+}
+
+void FOodleDictionaryArchive::SetDictionaryHeaderValues(int32 InHashTableSize)
+{
+	check(IsSaving());
+
+	Header.HashTableSize = InHashTableSize;
+}
+
+void FOodleDictionaryArchive::SerializeHeader()
+{
+	Header.SerializeHeader(*this);
+}
+
+void FOodleDictionaryArchive::SerializeDictionaryAndState(void*& DictionaryData, uint32& DictionaryBytes,
+															void*& CompactCompresorState, uint32& CompactCompressorStateBytes)
+{
+	if (IsLoading())
+	{
+		check(DictionaryData == nullptr);
+		check(CompactCompresorState == nullptr);
+
+		SerializeOodleDecompressData(Header.DictionaryData, DictionaryData, DictionaryBytes);
+		SerializeOodleDecompressData(Header.CompressorData, CompactCompresorState, CompactCompressorStateBytes);
+	}
+	else
+	{
+		check(DictionaryBytes > 0);
+		check(CompactCompressorStateBytes > 0);
+
+		SerializeOodleCompressData(Header.DictionaryData, DictionaryData, DictionaryBytes);
+		SerializeOodleCompressData(Header.CompressorData, CompactCompresorState, CompactCompressorStateBytes);
+	}
+}
+
 #endif
 
diff --git a/Engine/Plugins/Runtime/PacketHandlers/CompressionComponents/Oodle/Source/OodleHandlerComponent/Private/OodleHandlerComponent.cpp b/Engine/Plugins/Runtime/PacketHandlers/CompressionComponents/Oodle/Source/OodleHandlerComponent/Private/OodleHandlerComponent.cpp
--- a/Engine/Plugins/Runtime/PacketHandlers/CompressionComponents/Oodle/Source/OodleHandlerComponent/Private/OodleHandlerComponent.cpp
+++ b/Engine/Plugins/Runtime/PacketHandlers/CompressionComponents/Oodle/Source/OodleHandlerComponent/Private/OodleHandlerComponent.cpp
@@ -2,88 +2,187 @@
 
 #include "OodleHandlerComponentPCH.h"
 
+#include "OodleTrainerCommandlet.h"
+
 
 DEFINE_LOG_CATEGORY(OodleHandlerComponentLog);
 
 
-// @todo #JohnB: Deprecated the 'training' mode from this handler eventually, as it is misnamed/confusing
+// @todo #JohnB: Deprecate the 'training' mode name from this handler eventually, as it is misnamed/confusing
 
-// @todo #JohnB: Right now, there is no caching of loaded dictionary files, so this is very inefficient with multiple clients
+// @todo #JohnB: You need context-sensitive data capturing; e.g. you do NOT want to be capturing voice channel data
 
 #if HAS_OODLE_SDK
 #define OODLE_INI_SECTION TEXT("OodleHandlerComponent")
 
 
-// @todo #JohnB: Remove this when replacing the current dictionary serialization code
-#define STOPGAP_DICT_SERIALIZE 1
-
 // @todo #JohnB: Find a better solution than this, for handling the maximum packet size. This is far higher than it needs to be.
 #define MAX_OODLE_PACKET_SIZE 65535
 
 
-DEFINE_STAT(STAT_Oodle_BytesRaw);
-DEFINE_STAT(STAT_Oodle_BytesCompressed);
+#if STATS
+
+#if OODLE_STAT_NET
+DEFINE_STAT(STAT_Net_Oodle_OutRaw);
+DEFINE_STAT(STAT_Net_Oodle_OutCompressed);
+DEFINE_STAT(STAT_Net_Oodle_InRaw);
+DEFINE_STAT(STAT_Net_Oodle_InCompressed);
+#endif
+
+DEFINE_STAT(STAT_Oodle_OutRaw);
+DEFINE_STAT(STAT_Oodle_OutCompressed);
+DEFINE_STAT(STAT_Oodle_OutSavings);
+DEFINE_STAT(STAT_Oodle_OutTotalSavings);
+DEFINE_STAT(STAT_Oodle_InRaw);
+DEFINE_STAT(STAT_Oodle_InCompressed);
+DEFINE_STAT(STAT_Oodle_InSavings);
+DEFINE_STAT(STAT_Oodle_InTotalSavings);
+
+#if !UE_BUILD_SHIPPING
+DEFINE_STAT(STAT_Oodle_InDecompressTime);
+DEFINE_STAT(STAT_Oodle_OutCompressTime);
+#endif
+
+DEFINE_STAT(STAT_Oodle_DictionaryCount);
+DEFINE_STAT(STAT_Oodle_DictionaryBytes);
+DEFINE_STAT(STAT_Oodle_SharedBytes);
+DEFINE_STAT(STAT_Oodle_StateBytes);
+#endif
 
 
 FString GOodleSaveDir = TEXT("");
+FString GOodleContentDir = TEXT("");
 
 
-// @todo #JohnB: You need to integrate the dictionary generation into a commandlet, to support FArchive-format dictionaries
-//					(alternative - and do it this way first! - adjust OodleTrainer to read/write the FArchive version files)
+typedef TMap<FString, TSharedPtr<FOodleDictionary>> FDictionaryMap;
+
+/** Persistent map of loaded dictionaries */
+static FDictionaryMap DictionaryMap;
+
+
+#if STATS
+/** The global net stats tracker for Oodle */
+static FOodleNetStats GOodleNetStats;
 
-// @todo #JohnB: You need context-sensitive data capturing; e.g. you do NOT want to be capturing voice channel data
 
+/**
+ * FOodleNetStats
+ */
 
-// Saved compressor model file :
-// must match example_packet - or whatever you use to write the state file
-struct OodleNetwork1_SavedStates_Header
+void FOodleNetStats::UpdateStats(float DeltaTime)
 {
-#define ON1_MAGIC	0x11235801
-	U32 magic;
-	U32 compressor;
-	U32 ht_bits;
-	U32 dic_size;
-	U32 oodle_header_version;
-	U32 dic_complen;
-	U32 statecompacted_size;
-	U32 statecompacted_complen;
-
-	OodleNetwork1_SavedStates_Header()
-		: magic(0)
-		, compressor(0)
-		, ht_bits(0)
-		, dic_size(0)
-		, oodle_header_version(0)
-		, dic_complen(0)
-		, statecompacted_size(0)
-		, statecompacted_complen(0)
+	// Input
+	const uint32 InRaw = FMath::TruncToInt(InDecompressedLength / DeltaTime);
+	const uint32 InCompressed = FMath::TruncToInt(InCompressedLength / DeltaTime);
+
+#if OODLE_STAT_NET
+	SET_DWORD_STAT(STAT_Net_Oodle_InRaw, InRaw);
+	SET_DWORD_STAT(STAT_Net_Oodle_InCompressed, InCompressed);
+#endif
+
+	SET_DWORD_STAT(STAT_Oodle_InRaw, InRaw);
+	SET_DWORD_STAT(STAT_Oodle_InCompressed, InCompressed);
+
+	double InSavings = InCompressedLength > 0 ? ((1.0 - ((double)InCompressedLength/(double)InDecompressedLength)) * 100.0) : 0.0;
+
+	SET_FLOAT_STAT(STAT_Oodle_InSavings, InSavings);
+
+
+	// Output
+	uint32 OutRaw = FMath::TruncToInt(OutUncompressedLength / DeltaTime);
+	uint32 OutCompressed = FMath::TruncToInt(OutCompressedLength / DeltaTime);
+
+#if OODLE_STAT_NET
+	SET_DWORD_STAT(STAT_Net_Oodle_OutRaw, OutRaw);
+	SET_DWORD_STAT(STAT_Net_Oodle_OutCompressed, OutCompressed);
+#endif
+
+	SET_DWORD_STAT(STAT_Oodle_OutRaw, OutRaw);
+	SET_DWORD_STAT(STAT_Oodle_OutCompressed, OutCompressed);
+
+	double OutSavings = OutCompressedLength > 0 ? ((1.0 - ((double)OutCompressedLength/(double)OutUncompressedLength)) * 100.0) : 0.0;
+
+	SET_FLOAT_STAT(STAT_Oodle_OutSavings, OutSavings);
+
+
+	// Crude process-lifetime accumulation of all stat savings
+	if (TotalInCompressedLength > 0)
 	{
+		SET_FLOAT_STAT(STAT_Oodle_InTotalSavings, (1.0 - ((double)TotalInCompressedLength/(double)TotalInDecompressedLength)) * 100.0);
 	}
 
-	// Temporary serializer, prior to replacement of this code
-	void SerializeSavedStates(FArchive& InArc)
+	if (TotalOutCompressedLength > 0)
 	{
-		InArc << (uint32&)magic;
-		InArc << (uint32&)compressor;
-		InArc << (uint32&)ht_bits;
-		InArc << (uint32&)dic_size;
-		InArc << (uint32&)oodle_header_version;
-		InArc << (uint32&)dic_complen;
-		InArc << (uint32&)statecompacted_size;
-		InArc << (uint32&)statecompacted_complen;
+		SET_FLOAT_STAT(STAT_Oodle_OutTotalSavings,
+						(1.0 - ((double)TotalOutCompressedLength/(double)TotalOutUncompressedLength)) * 100.0);
 	}
-};
 
-// OODLE
-OodleHandlerComponent::OodleHandlerComponent()
-	: BytesSaved(0)
-#if EXAMPLE_CAPTURE_FORMAT
-	, PacketLogFile(NULL)
-	, bPacketLogAppend(false)
-#else
-	, InPacketLog(NULL)
-	, OutPacketLog(NULL)
+
+	// Reset stats accumulated since last update
+	InCompressedLength = 0;
+	InDecompressedLength = 0;
+	OutCompressedLength = 0;
+	OutUncompressedLength = 0;
+}
+#endif
+
+
+/**
+ * FOodleDictionary
+ */
+
+FOodleDictionary::FOodleDictionary(uint32 InHashTableSize, void* InDictionaryData, uint32 InDictionarySize,
+									OodleNetwork1_Shared* InSharedDictionary, uint32 InSharedDictionarySize,
+									OodleNetwork1UDP_State* InInitialCompressorState, uint32 InCompressorStateSize)
+	: HashTableSize(InHashTableSize)
+	, DictionaryData(InDictionaryData)
+	, DictionarySize(InDictionarySize)
+	, SharedDictionary(InSharedDictionary)
+	, SharedDictionarySize(InSharedDictionarySize)
+	, CompressorState(InInitialCompressorState)
+	, CompressorStateSize(InCompressorStateSize)
+{
+#if STATS
+	INC_DWORD_STAT(STAT_Oodle_DictionaryCount);
+	INC_DWORD_STAT_BY(STAT_Oodle_DictionaryBytes, DictionarySize);
+	INC_DWORD_STAT_BY(STAT_Oodle_SharedBytes, SharedDictionarySize);
+	INC_DWORD_STAT_BY(STAT_Oodle_StateBytes, CompressorStateSize);
 #endif
+}
+
+FOodleDictionary::~FOodleDictionary()
+{
+#if STATS
+	DEC_DWORD_STAT(STAT_Oodle_DictionaryCount);
+	DEC_DWORD_STAT_BY(STAT_Oodle_DictionaryBytes, DictionarySize);
+	DEC_DWORD_STAT_BY(STAT_Oodle_SharedBytes, SharedDictionarySize);
+	DEC_DWORD_STAT_BY(STAT_Oodle_StateBytes, CompressorStateSize);
+#endif
+
+	if (DictionaryData != nullptr)
+	{
+		delete DictionaryData;
+	}
+
+	if (SharedDictionary != nullptr)
+	{
+		FMemory::Free(SharedDictionary);
+	}
+
+	if (CompressorState != nullptr)
+	{
+		FMemory::Free(CompressorState);
+	}
+}
+
+
+/**
+ * OodleHandlerComponent
+ */
+
+OodleHandlerComponent::OodleHandlerComponent()
+	: InPacketLog(nullptr)
+	, OutPacketLog(nullptr)
 	, bEnableOodle(false)
 #if !UE_BUILD_SHIPPING
 	, bUseDictionaryIfPresent(false)
@@ -97,61 +196,32 @@ OodleHandlerComponent::OodleHandlerComponent()
 
 OodleHandlerComponent::~OodleHandlerComponent()
 {
-#if EXAMPLE_CAPTURE_FORMAT
-	// @todo #JohnB: PacketLogfile is not deleted. Doesn't matter though, as this is being deprecated.
-#else
-	if (OutPacketLog != NULL)
+	if (OutPacketLog != nullptr)
 	{
 		OutPacketLog->Close();
 		OutPacketLog->DeleteInnerArchive();
 
 		delete OutPacketLog;
 
-		OutPacketLog = NULL;
+		OutPacketLog = nullptr;
 	}
 
-	if (InPacketLog != NULL)
+	if (InPacketLog != nullptr)
 	{
 		InPacketLog->Close();
 		InPacketLog->DeleteInnerArchive();
 
 		delete InPacketLog;
 
-		InPacketLog = NULL;
-	}
-#endif
-
-	if (ServerDictionary.Dictionary != NULL)
-	{
-		free(ServerDictionary.Dictionary);
+		InPacketLog = nullptr;
 	}
 
-	if (ServerDictionary.SharedDictionaryData != NULL)
-	{
-		free(ServerDictionary.SharedDictionaryData);
-	}
 
-	if (ServerDictionary.CompressorState != NULL)
-	{
-		free(ServerDictionary.CompressorState);
-	}
+	FreeDictionary(ServerDictionary);
+	FreeDictionary(ClientDictionary);
 
 
-	if (ClientDictionary.Dictionary != NULL)
-	{
-		free(ClientDictionary.Dictionary);
-	}
-
-	if (ClientDictionary.SharedDictionaryData != NULL)
-	{
-		free(ClientDictionary.SharedDictionaryData);
-	}
-
-	if (ClientDictionary.CompressorState != NULL)
-	{
-		free(ClientDictionary.CompressorState);
-	}
-
+	// @todo #JohnB: Deprecate all of this code sooner rather than later - this is broken, as there will be multiple instances
 #if UE4_OODLE_VER < 200
 	Oodle_Shutdown_NoThreads();
 #endif
@@ -171,12 +241,10 @@ void OodleHandlerComponent::InitFirstRunConfig()
 		GConfig->SetString(OODLE_INI_SECTION, TEXT("Mode"), TEXT("Capturing"), GEngineIni);
 		GConfig->SetString(OODLE_INI_SECTION, TEXT("PacketLogFile"), TEXT("PacketDump"), GEngineIni);
 
-#if EXAMPLE_CAPTURE_FORMAT
-		GConfig->SetBool(OODLE_INI_SECTION, TEXT("bPacketLogAppend"), false, GEngineIni);
-#endif
-
 		GConfig->SetString(OODLE_INI_SECTION, TEXT("ServerDictionary"), TEXT(""), GEngineIni);
 		GConfig->SetString(OODLE_INI_SECTION, TEXT("ClientDictionary"), TEXT(""), GEngineIni);
+
+		GConfig->Flush(false);
 	}
 }
 
@@ -184,10 +252,12 @@ void OodleHandlerComponent::Initialize()
 {
 	InitFirstRunConfig();
 
+	const bool bOodleForceEnable = FParse::Param(FCommandLine::Get(), TEXT("Oodle"));
+
 	// Class config variables
 	GConfig->GetBool(OODLE_INI_SECTION, TEXT("bEnableOodle"), bEnableOodle, GEngineIni);
 
-	if (!bEnableOodle && FParse::Param(FCommandLine::Get(), TEXT("Oodle")))
+	if (!bEnableOodle && bOodleForceEnable)
 	{
 		UE_LOG(OodleHandlerComponentLog, Log, TEXT("Force-enabling Oodle from commandline."));
 		bEnableOodle = true;
@@ -195,12 +265,15 @@ void OodleHandlerComponent::Initialize()
 
 #if !UE_BUILD_SHIPPING
 	GConfig->GetBool(OODLE_INI_SECTION, TEXT("bUseDictionaryIfPresent"), bUseDictionaryIfPresent, GEngineIni);
-#endif
 
-#if EXAMPLE_CAPTURE_FORMAT
-	GConfig->GetBool(OODLE_INI_SECTION, TEXT("bPacketLogAppend"), bPacketLogAppend, GEngineIni);
+	if (!bUseDictionaryIfPresent && bOodleForceEnable)
+	{
+		UE_LOG(OodleHandlerComponentLog, Log, TEXT("Force-enabling 'bUseDictionaryIfPresent', due to -Oodle on commandline."));
+		bUseDictionaryIfPresent = true;
+	}
 #endif
 
+
 	// Mode
 	FString ReadMode;
 	GConfig->GetString(OODLE_INI_SECTION, TEXT("Mode"), ReadMode, GEngineIni);
@@ -250,6 +323,8 @@ void OodleHandlerComponent::Initialize()
 
 		switch(Mode)
 		{
+			// @todo #JohnB: Convert this code so that just one capture file is used for all connections, per session
+			//					(could set it up much like the dictionary sharing code)
 		case EOodleHandlerMode::Capturing:
 			{
 				if (Handler->Mode == Handler::Mode::Server)
@@ -265,56 +340,6 @@ void OodleHandlerComponent::Initialize()
 
 					BaseFilename = FPaths::GetBaseFilename(BaseFilename);
 
-
-#if EXAMPLE_CAPTURE_FORMAT
-					FString FinalFilePath;
-
-					// When bPacketLogAppend is used, a static log filename is used
-					if (bPacketLogAppend)
-					{
-						FString PreExtFilePath = FPaths::Combine(*ReadOutputLogDirectory, *BaseFilename);
-
-						// Try0: PacketDump.ucap, Try1: PacketDump_1.ucap, Try#:PacketDump_#.ucap
-						FinalFilePath = PreExtFilePath + CAPTURE_EXT;
-
-						for (int32 i=1; ; i++)
-						{
-							// If not appending to a file, make sure to disable this bool - otherwise the file header won't be written
-							bPacketLogAppend = PlatformFile.FileExists(*FinalFilePath);
-							PacketLogFile = PlatformFile.OpenWrite(*FinalFilePath, true);
-
-							if (PacketLogFile != NULL)
-							{
-								break;
-							}
-
-
-							FinalFilePath = PreExtFilePath + FString::Printf(TEXT("_%i"), i) + CAPTURE_EXT;
-						}
-					}
-					// When non-append logging is used, the log filename uses a timestamp
-					else
-					{
-						BaseFilename = BaseFilename + TEXT("_") + FDateTime::Now().ToString();
-
-						FString PreExtFilePath = FPaths::Combine(*ReadOutputLogDirectory, *BaseFilename);
-						FinalFilePath = PreExtFilePath + CAPTURE_EXT;
-
-						for (int32 i=1; PlatformFile.FileExists(*FinalFilePath); i++)
-						{
-							FinalFilePath = PreExtFilePath + FString::Printf(TEXT("_%i"), i) + CAPTURE_EXT;
-						}
-
-						// Try0: PacketDump_2015.09.23-20.12.54.ucap, Try1: PacketDump_2015.09.23-20.12.54_1.ucap, ...
-						PacketLogFile = PlatformFile.OpenWrite(*FinalFilePath);
-					}
-
-					if (PacketLogFile == NULL)
-					{
-						LowLevelFatalError(TEXT("Failed to create file %s"), *FinalFilePath);
-						return;
-					}
-#else
 					BaseFilename = BaseFilename + TEXT("_") + FDateTime::Now().ToString();
 
 					FString PreExtInFilePath = FPaths::Combine(*ReadOutputLogDirectory, TEXT("Input"), *(BaseFilename + TEXT("_Input")));
@@ -330,32 +355,21 @@ void OodleHandlerComponent::Initialize()
 					}
 
 					FArchive* InArc = IFileManager::Get().CreateFileWriter(*InPath);
-					FArchive* OutArc = (InArc != NULL ? IFileManager::Get().CreateFileWriter(*OutPath) : NULL);
+					FArchive* OutArc = (InArc != nullptr ? IFileManager::Get().CreateFileWriter(*OutPath) : nullptr);
 
-					InPacketLog = (InArc != NULL ? new FPacketCaptureArchive(*InArc) : NULL);
-					OutPacketLog = (OutArc != NULL ? new FPacketCaptureArchive(*OutArc) : NULL);
+					InPacketLog = (InArc != nullptr ? new FPacketCaptureArchive(*InArc) : nullptr);
+					OutPacketLog = (OutArc != nullptr ? new FPacketCaptureArchive(*OutArc) : nullptr);
 
 
-					if (InPacketLog == NULL || OutPacketLog == NULL)
+					if (InPacketLog == nullptr || OutPacketLog == nullptr)
 					{
 						LowLevelFatalError(TEXT("Failed to create files '%s' and '%s'"), *InPath, *OutPath);
 						return;
 					}
-#endif
 
 
-#if EXAMPLE_CAPTURE_FORMAT
-					// Only write the initial header, if not appending
-					if (!bPacketLogAppend)
-					{
-						int32 NumberOfChannels = 1;
-			
-						PacketLogFile->Write(reinterpret_cast<const uint8*>(&NumberOfChannels), sizeof(NumberOfChannels));
-					}
-#else
 					InPacketLog->SerializeCaptureHeader();
 					OutPacketLog->SerializeCaptureHeader();
-#endif
 				}
 
 				break;
@@ -400,196 +414,113 @@ void OodleHandlerComponent::Initialize()
 	Initialized();
 }
 
-void OodleHandlerComponent::InitializeDictionary(FString FilePath, FOodleDictionary& OutDictionary)
+void OodleHandlerComponent::InitializeDictionary(FString FilePath, TSharedPtr<FOodleDictionary>& OutDictionary)
 {
-	// @todo #JohnB: When rewriting, make sure to have proper error checking, for cases such as e.g. trying to load dictionary,
-	//					from past/incompatible Oodle version
-
-#if STOPGAP_DICT_SERIALIZE
-	TArray<uint8> FileInMemory;
-	FMemoryReader FileArc(FileInMemory);
+	TSharedPtr<FOodleDictionary>* DictionaryRef = DictionaryMap.Find(FilePath);
 
-	// @todo #JohnB: Need a way to verify the state file, is actually a state file (do this once converting to FArchive)
-	if (FFileHelper::LoadFileToArray(FileInMemory, *FilePath))
+	// Load the dictionary, if it's not yet loaded
+	if (DictionaryRef == nullptr)
 	{
-		// parse header :
-		OodleNetwork1_SavedStates_Header SavedDictionaryFileData;
+		FArchive* ReadArc = IFileManager::Get().CreateFileReader(*FilePath);
 
-		SavedDictionaryFileData.SerializeSavedStates(FileArc);
-
-		if (SavedDictionaryFileData.magic == ON1_MAGIC && SavedDictionaryFileData.oodle_header_version == OODLE_HEADER_VERSION)
+		if (ReadArc != nullptr)
 		{
-			// this code only loads compressed states currently
-			check(SavedDictionaryFileData.compressor != -1);
-
-
-			S32 StateHTBits = SavedDictionaryFileData.ht_bits;
-			SINTa DictionarySize = SavedDictionaryFileData.dic_size;
+			FOodleDictionaryArchive BoundArc(*ReadArc);
 
-			SINTa DictionaryCompLen = SavedDictionaryFileData.dic_complen;
-			SINTa CompressorStateCompactedSize = SavedDictionaryFileData.statecompacted_size;
-			SINTa CompressorStateCompactedLength = SavedDictionaryFileData.statecompacted_complen;
+			void* DictionaryData = nullptr;
+			uint32 DictionaryBytes = 0;
+			void* CompactCompressorState = nullptr;
+			uint32 CompactCompressorStateBytes = 0;
 
-			SINTa SharedSize = OodleNetwork1_Shared_Size(StateHTBits);
-			SINTa CompressorStateSize = OodleNetwork1UDP_State_Size();
+			BoundArc.SerializeHeader();
+			BoundArc.SerializeDictionaryAndState(DictionaryData, DictionaryBytes, CompactCompressorState, CompactCompressorStateBytes);
 
-			check(DictionarySize >= DictionaryCompLen);
-			check(CompressorStateCompactedSize >= CompressorStateCompactedLength);
 
-			check(CompressorStateCompactedSize > 0 && CompressorStateCompactedSize < OodleNetwork1UDP_StateCompacted_MaxSize());
-
-			check(FileArc.TotalSize() == (S64)sizeof(OodleNetwork1_SavedStates_Header) + DictionaryCompLen +
-					CompressorStateCompactedLength);
-
-			//OodleLog_Printf_v1("OodleNetwork1UDP Loading; dic comp %d , state %d->%d\n",DictionaryCompLen,
-			//						CompressorStateCompactedSize, CompressorStateCompactedLength);
+			if (!BoundArc.IsError())
+			{
+				UE_LOG(OodleHandlerComponentLog, Log, TEXT("Loading dictionary file: %s"), *FilePath);
 
-			//-------------------------------------------
-			// decompress Dictionary and CompressorStatecompacted
+				// Uncompact the compressor state
+				uint32 CompressorStateSize = OodleNetwork1UDP_State_Size();
+				OodleNetwork1UDP_State* CompressorState = (OodleNetwork1UDP_State*)FMemory::Malloc(CompressorStateSize);
 
-			// @todo #JohnB: Security issue
-			OutDictionary.Dictionary = malloc(DictionarySize);
+				OodleNetwork1UDP_State_Uncompact(CompressorState, (OodleNetwork1UDP_StateCompacted*)CompactCompressorState);
 
-			void * DictionaryCompressed = &FileInMemory[FileArc.Tell()];
 
-			SINTa DecompressedDictionarySize = OodleLZ_Decompress(DictionaryCompressed, DictionaryCompLen, OutDictionary.Dictionary,
-																	DictionarySize);
+				// Create the shared dictionary state
+				int32 HashTableSize = BoundArc.Header.HashTableSize.Get();
+				uint32 SharedDictionarySize = OodleNetwork1_Shared_Size(HashTableSize);
+				OodleNetwork1_Shared* SharedDictionary = (OodleNetwork1_Shared*)FMemory::Malloc(SharedDictionarySize);
 
-			check(DecompressedDictionarySize == DictionarySize);
+				OodleNetwork1_Shared_SetWindow(SharedDictionary, HashTableSize, DictionaryData, DictionaryBytes);
 
-			OodleNetwork1UDP_StateCompacted* UDPNewCompacted = (OodleNetwork1UDP_StateCompacted *)malloc(CompressorStateCompactedSize);
 
-			void * CompressorStatecompacted_comp_ptr = (U8 *)DictionaryCompressed + DictionaryCompLen;
+				// Now add the dictionary data to the map
+				FOodleDictionary* NewDictionary = new FOodleDictionary(HashTableSize, DictionaryData, DictionaryBytes, SharedDictionary,
+																		SharedDictionarySize, CompressorState, CompressorStateSize);
 
-			SINTa decomp_statecompacted_size = OodleLZ_Decompress(CompressorStatecompacted_comp_ptr, CompressorStateCompactedLength,
-																	UDPNewCompacted, CompressorStateCompactedSize);
+				DictionaryRef = &DictionaryMap.Add(FilePath, MakeShareable(NewDictionary));
 
-			check(decomp_statecompacted_size == CompressorStateCompactedSize);
 
-			//----------------------------------------------
-			// Uncompact the "Compacted" state into a usable state
+				if (CompactCompressorState != nullptr)
+				{
+					delete CompactCompressorState;
+				}
+			}
+			else
+			{
+				UE_LOG(OodleHandlerComponentLog, Warning, TEXT("Error loading dictionary file: %s"), *FilePath);
 
-			OutDictionary.CompressorState = (OodleNetwork1UDP_State *)malloc(CompressorStateSize);
-			check(OutDictionary.CompressorState != NULL);
+				if (DictionaryData != nullptr)
+				{
+					delete DictionaryData;
+				}
 
-			OodleNetwork1UDP_State_Uncompact(OutDictionary.CompressorState, UDPNewCompacted);
-			free(UDPNewCompacted);
+				if (CompactCompressorState != nullptr)
+				{
+					delete CompactCompressorState;
+				}
+			}
 
-			//----------------------------------------------
-			// fill out SharedDictionaryData from the dictionary
 
-			OutDictionary.SharedDictionaryData = (OodleNetwork1_Shared *)malloc(SharedSize);
-			check(OutDictionary.SharedDictionaryData != NULL);
-			OodleNetwork1_Shared_SetWindow(OutDictionary.SharedDictionaryData, StateHTBits, OutDictionary.Dictionary,
-											(S32)DictionarySize);
+			ReadArc->Close();
 
-			//-----------------------------------------------------
-	
-			//SET_DWORD_STAT(STAT_Oodle_DicSize, DictionarySize);
-		}
-		else if (SavedDictionaryFileData.magic != ON1_MAGIC)
-		{
-			UE_LOG(OodleHandlerComponentLog, Warning, TEXT("OodleNetCompressor: state_file ON1_MAGIC mismatch"));
+			delete ReadArc;
+			ReadArc = nullptr;
 		}
-		else //if (SavedDictionaryFileData->oodle_header_version != OODLE_HEADER_VERSION)
+		else
 		{
-			UE_LOG(OodleHandlerComponentLog, Warning,
-					TEXT("OodleNetCompressor: state_file OODLE_HEADER_VERSION mismatch. Was: %u, Expected: %u"),
-					(uint32)SavedDictionaryFileData.oodle_header_version, (uint32)OODLE_HEADER_VERSION);
+			LowLevelFatalError(TEXT("Incorrect DictionaryFile Provided"));
 		}
-#else
-	TArray<uint8> FileInMemory;
-
-	// @todo #JohnB: Need a way to verify the state file, is actually a state file (do this once converting to FArchive)
-	if (FFileHelper::LoadFileToArray(FileInMemory, *FilePath))
-	{
-		// parse header :
-		OodleNetwork1_SavedStates_Header* SavedDictionaryFileData = (OodleNetwork1_SavedStates_Header*)FileInMemory.GetData();
-
-		if (SavedDictionaryFileData->magic == ON1_MAGIC && SavedDictionaryFileData->oodle_header_version == OODLE_HEADER_VERSION)
-		{
-			// this code only loads compressed states currently
-			check(SavedDictionaryFileData->compressor != -1);
-
-
-			S32 StateHTBits = SavedDictionaryFileData->ht_bits;
-			SINTa DictionarySize = SavedDictionaryFileData->dic_size;
-
-			SINTa DictionaryCompLen = SavedDictionaryFileData->dic_complen;
-			SINTa CompressorStateCompactedSize = SavedDictionaryFileData->statecompacted_size;
-			SINTa CompressorStateCompactedLength = SavedDictionaryFileData->statecompacted_complen;
-
-			SINTa SharedSize = OodleNetwork1_Shared_Size(StateHTBits);
-			SINTa CompressorStateSize = OodleNetwork1UDP_State_Size();
-
-			check(DictionarySize >= DictionaryCompLen);
-			check(CompressorStateCompactedSize >= CompressorStateCompactedLength);
-
-			check(CompressorStateCompactedSize > 0 && CompressorStateCompactedSize < OodleNetwork1UDP_StateCompacted_MaxSize());
-
-			check(FileInMemory.Num() == (S64)sizeof(OodleNetwork1_SavedStates_Header) + DictionaryCompLen +
-					CompressorStateCompactedLength);
-
-			//OodleLog_Printf_v1("OodleNetwork1UDP Loading; dic comp %d , state %d->%d\n",DictionaryCompLen,
-			//						CompressorStateCompactedSize, CompressorStateCompactedLength);
-
-			//-------------------------------------------
-			// decompress Dictionary and CompressorStatecompacted
-
-			OutDictionary.Dictionary = malloc(DictionarySize);
-
-			void * DictionaryCompressed = SavedDictionaryFileData + 1;
-
-			SINTa DecompressedDictionarySize = OodleLZ_Decompress(DictionaryCompressed, DictionaryCompLen, OutDictionary.Dictionary,
-																	DictionarySize);
-
-			check(DecompressedDictionarySize == DictionarySize);
-
-			OodleNetwork1UDP_StateCompacted* UDPNewCompacted = (OodleNetwork1UDP_StateCompacted *)malloc(CompressorStateCompactedSize);
-
-			void * CompressorStatecompacted_comp_ptr = (U8 *)DictionaryCompressed + DictionaryCompLen;
-
-			SINTa decomp_statecompacted_size = OodleLZ_Decompress(CompressorStatecompacted_comp_ptr, CompressorStateCompactedLength,
-																	UDPNewCompacted, CompressorStateCompactedSize);
-
-			check(decomp_statecompacted_size == CompressorStateCompactedSize);
-
-			//----------------------------------------------
-			// Uncompact the "Compacted" state into a usable state
-
-			OutDictionary.CompressorState = (OodleNetwork1UDP_State *)malloc(CompressorStateSize);
-			check(OutDictionary.CompressorState != NULL);
+	}
 
-			OodleNetwork1UDP_State_Uncompact(OutDictionary.CompressorState, UDPNewCompacted);
-			free(UDPNewCompacted);
 
-			//----------------------------------------------
-			// fill out SharedDictionaryData from the dictionary
+	if (DictionaryRef != nullptr)
+	{
+		OutDictionary = *DictionaryRef;
+	}
+}
 
-			OutDictionary.SharedDictionaryData = (OodleNetwork1_Shared *)malloc(SharedSize);
-			check(OutDictionary.SharedDictionaryData != NULL);
-			OodleNetwork1_Shared_SetWindow(OutDictionary.SharedDictionaryData, StateHTBits, OutDictionary.Dictionary,
-											(S32)DictionarySize);
+void OodleHandlerComponent::FreeDictionary(TSharedPtr<FOodleDictionary>& InDictionary)
+{
+	if (InDictionary.IsValid())
+	{
+		// The dictionary is always referenced within DictionaryMap, so 2 represents last ref within an OodleHandlerComponent
+		bool bLastDictionaryRef = InDictionary.GetSharedReferenceCount() == 2;
 
-			//-----------------------------------------------------
-	
-			//SET_DWORD_STAT(STAT_Oodle_DicSize, DictionarySize);
-		}
-		else if (SavedDictionaryFileData->magic != ON1_MAGIC)
+		if (bLastDictionaryRef)
 		{
-			UE_LOG(OodleHandlerComponentLog, Warning, TEXT("OodleNetCompressor: state_file ON1_MAGIC mismatch"));
-		}
-		else //if (SavedDictionaryFileData->oodle_header_version != OODLE_HEADER_VERSION)
-		{
-			UE_LOG(OodleHandlerComponentLog, Warning,
-					TEXT("OodleNetCompressor: state_file OODLE_HEADER_VERSION mismatch. Was: %u, Expected: %u"),
-					(uint32)SavedDictionaryFileData->oodle_header_version, (uint32)OODLE_HEADER_VERSION);
+			for (FDictionaryMap::TIterator It(DictionaryMap); It; ++It)
+			{
+				if (It.Value() == InDictionary)
+				{
+					It.RemoveCurrent();
+					break;
+				}
+			}
 		}
-#endif
-	}
-	else
-	{
-		LowLevelFatalError(TEXT("Incorrect DictionaryFile Provided"));
+
+		InDictionary.Reset();
 	}
 }
 
@@ -765,20 +696,12 @@ void OodleHandlerComponent::Incoming(FBitReader& Packet)
 		{
 		case EOodleHandlerMode::Capturing:
 			{
-				if(Handler->Mode == Handler::Mode::Server)
+				if (Handler->Mode == Handler::Mode::Server)
 				{
 					uint32 SizeOfPacket = Packet.GetNumBytes();
-
-#if EXAMPLE_CAPTURE_FORMAT
-					int32 channel = 0;
-
-					PacketLogFile->Write(reinterpret_cast<const uint8*>(&channel), sizeof(channel));
-					PacketLogFile->Write(reinterpret_cast<const uint8*>(&SizeOfPacket), sizeof(SizeOfPacket));
-					PacketLogFile->Write(reinterpret_cast<const uint8*>(Packet.GetData()), SizeOfPacket);
-#else
 					void* PacketData = Packet.GetData();
+
 					InPacketLog->SerializePacket(PacketData, SizeOfPacket);
-#endif
 				}
 				break;
 			}
@@ -788,23 +711,73 @@ void OodleHandlerComponent::Incoming(FBitReader& Packet)
 				uint32 DecompressedLength;
 				Packet.SerializeIntPacked(DecompressedLength);
 
-				if (ensure(DecompressedLength < MAX_OODLE_PACKET_SIZE))
+				if (DecompressedLength < MAX_OODLE_PACKET_SIZE)
 				{
 					// @todo #JohnB: You want to tweak the size of 'MAX_OODLE_PACKET_SIZE'
 					static uint8 CompressedData[MAX_OODLE_PACKET_SIZE];
 					static uint8 DecompressedData[MAX_OODLE_PACKET_SIZE];
+
 					const int32 CompressedLength = Packet.GetBytesLeft();
-					FOodleDictionary* CurDict = ((Handler->Mode == Handler::Mode::Server) ? &ClientDictionary : &ServerDictionary);
+					FOodleDictionary* CurDict =
+						((Handler->Mode == Handler::Mode::Server) ? ClientDictionary.Get() : ServerDictionary.Get());
 
 					Packet.Serialize(CompressedData, CompressedLength);
 
-					check(OodleNetwork1UDP_Decode(CurDict->CompressorState, CurDict->SharedDictionaryData, CompressedData,
-													CompressedLength, DecompressedData, DecompressedLength));
+					bool bSuccess = !Packet.IsError();
 
+					if (bSuccess)
+					{
+						{
+#if STATS && !UE_BUILD_SHIPPING
+							SCOPE_CYCLE_COUNTER(STAT_Oodle_InDecompressTime);
+#endif
+							bSuccess = !!OodleNetwork1UDP_Decode(CurDict->CompressorState, CurDict->SharedDictionary, CompressedData,
+															CompressedLength, DecompressedData, DecompressedLength);
+						}
 
-					// @todo #JohnB: Decompress directly into FBitReader's buffer
-					FBitReader UnCompressedPacket(DecompressedData, DecompressedLength * 8);
-					Packet = UnCompressedPacket;
+						if (!bSuccess)
+						{
+							// @todo #JohnB: Set packet archive error and allow in shipping - see note below
+#if !UE_BUILD_SHIPPING
+							UE_LOG(OodleHandlerComponentLog, Error, TEXT("Error decoding Oodle network data."));
+#endif
+						}
+					}
+					else
+					{
+						// @todo #JohnB: Set packet archive error and allow in shipping - see note below
+#if !UE_BUILD_SHIPPING
+						UE_LOG(OodleHandlerComponentLog, Error, TEXT("Error serializing received packet data"));
+#endif
+					}
+
+
+					if (bSuccess)
+					{
+						// @todo #JohnB: Decompress directly into FBitReader's buffer
+						FBitReader UnCompressedPacket(DecompressedData, DecompressedLength * 8);
+						Packet = UnCompressedPacket;
+
+
+#if STATS
+						GOodleNetStats.IncomingStats(CompressedLength, DecompressedLength);
+#endif
+					}
+					else
+					{
+						// @todo #JohnB
+					}
+				}
+				else
+				{
+					// @todo #JohnB: Add ArIsError handling to the PacketHandler code, then allow the log message in shipping,
+					//					and set an archive error
+#if !UE_BUILD_SHIPPING
+					UE_LOG(OodleHandlerComponentLog, Error,
+							TEXT("Received packet with DecompressedLength (%i) >= MAX_OODLE_PACKET_SIZE"), DecompressedLength);
+
+					// @todo #JohnB: Set packet error
+#endif
 				}
 
 
@@ -822,19 +795,11 @@ void OodleHandlerComponent::Outgoing(FBitWriter& Packet)
 		{
 		case EOodleHandlerMode::Capturing:
 			{
-				if(Handler->Mode == Handler::Mode::Server)
+				if (Handler->Mode == Handler::Mode::Server)
 				{
 					uint32 SizeOfPacket = Packet.GetNumBytes();
 
-#if EXAMPLE_CAPTURE_FORMAT
-					int32 channel = 0;
-
-					PacketLogFile->Write(reinterpret_cast<const uint8*>(&channel), sizeof(channel));
-					PacketLogFile->Write(reinterpret_cast<const uint8*>(&SizeOfPacket), sizeof(SizeOfPacket));
-					PacketLogFile->Write(reinterpret_cast<const uint8*>(Packet.GetData()), SizeOfPacket);
-#else
 					OutPacketLog->SerializePacket((void*)Packet.GetData(), SizeOfPacket);
-#endif
 				}
 				break;
 			}
@@ -851,31 +816,35 @@ void OodleHandlerComponent::Outgoing(FBitWriter& Packet)
 					// @todo #JohnB: You should be able to share the same two buffers between Incoming and Outgoing, as an optimization
 					static uint8 UncompressedData[MAX_OODLE_PACKET_SIZE];
 					static uint8 CompressedData[MAX_OODLE_PACKET_SIZE];
-					FOodleDictionary* CurDict = ((Handler->Mode == Handler::Mode::Server) ? &ServerDictionary : &ClientDictionary);
+					FOodleDictionary* CurDict =
+										((Handler->Mode == Handler::Mode::Server) ? ServerDictionary.Get() : ClientDictionary.Get());
 
 					check(OodleLZ_GetCompressedBufferSizeNeeded(UncompressedLength) < MAX_OODLE_PACKET_SIZE);
 
 					memcpy(UncompressedData, Packet.GetData(), UncompressedLength);
 
-					SINTa CompressedLengthSINT = OodleNetwork1UDP_Encode(CurDict->CompressorState, CurDict->SharedDictionaryData,
+					SINTa CompressedLengthSINT = 0;
+					
+					{
+#if STATS && !UE_BUILD_SHIPPING
+						SCOPE_CYCLE_COUNTER(STAT_Oodle_OutCompressTime);
+#endif
+
+						CompressedLengthSINT = OodleNetwork1UDP_Encode(CurDict->CompressorState, CurDict->SharedDictionary,
 																			UncompressedData, UncompressedLength, CompressedData);
+					}
 
 					uint32 CompressedLength = static_cast<uint32>(CompressedLengthSINT);
 
 					Packet.Reset();
 
-					// @todo #JohnB: Compress directly into a (deliberately oversized) FBitWriter buffer, which you can shrink afterwards
+					// @todo #JohnB: Compress directly into a (deliberately oversized) FBitWriter buffer, which you can shrink after
 					Packet.SerializeIntPacked(UncompressedLength);
 					Packet.Serialize(CompressedData, CompressedLength);
 
-
-					BytesSaved += UncompressedLength - CompressedLength;
-
-					//UE_LOG(OodleHandlerComponentLog, Log, TEXT("Oodle Compressed: UnCompressed: %i Compressed: %i Total: %i"),
-					//			UncompressedLength, CompressedLength, BytesSaved);
-
-					INC_DWORD_STAT_BY(STAT_Oodle_BytesRaw, UncompressedLength);
-					INC_DWORD_STAT_BY(STAT_Oodle_BytesCompressed, CompressedLength);
+#if STATS
+					GOodleNetStats.OutgoingStats(CompressedLength, UncompressedLength);
+#endif
 				}
 				else
 				{
@@ -888,7 +857,11 @@ void OodleHandlerComponent::Outgoing(FBitWriter& Packet)
 	}
 }
 
-// MODULE INTERFACE
+
+/**
+ * Module Interface
+ */
+
 TSharedPtr<HandlerComponent> FOodleComponentModuleInterface::CreateComponentInstance(FString& Options)
 {
 	return MakeShareable(new OodleHandlerComponent);
@@ -896,9 +869,16 @@ TSharedPtr<HandlerComponent> FOodleComponentModuleInterface::CreateComponentInst
 
 void FOodleComponentModuleInterface::StartupModule()
 {
+	// If Oodle is force-enabled on the commandline, execute the commandlet-enable command, which also adds to the PacketHandler list
+	if (FParse::Param(FCommandLine::Get(), TEXT("Oodle")))
+	{
+		UOodleTrainerCommandlet::HandleEnable();
+	}
+
+
 	// Use an absolute path for this, as we want all relative paths, to be relative to this folder
 	GOodleSaveDir = FPaths::ConvertRelativePathToFull(FPaths::Combine(*FPaths::GameSavedDir(), TEXT("Oodle")));
-
+	GOodleContentDir = FPaths::ConvertRelativePathToFull(FPaths::Combine(*FPaths::GameContentDir(), TEXT("Oodle")));
 
 	// Load the Oodle library (NOTE: Path mirrored in Oodle.Build.cs)
 	FString OodleBinaryPath = FPaths::EngineDir() / TEXT("Binaries/ThirdParty/NotForLicensees/Oodle/");
@@ -930,7 +910,7 @@ void FOodleComponentModuleInterface::StartupModule()
 #endif
 
 #if !PLATFORM_LINUX
-	check(OodleDllHandle != NULL);
+	check(OodleDllHandle != nullptr);
 #endif
 
 
@@ -948,10 +928,10 @@ void FOodleComponentModuleInterface::ShutdownModule()
 	Oodle_Shutdown();
 #endif
 
-	if (OodleDllHandle != NULL)
+	if (OodleDllHandle != nullptr)
 	{
 		FPlatformProcess::FreeDllHandle(OodleDllHandle);
-		OodleDllHandle = NULL;
+		OodleDllHandle = nullptr;
 	}
 }
 #else
@@ -959,7 +939,7 @@ TSharedPtr<HandlerComponent> FOodleComponentModuleInterface::CreateComponentInst
 {
 	UE_LOG(OodleHandlerComponentLog, Error, TEXT("Can't create OodleHandlerComponent instance - HAS_OODLE_SDK is false."));
 
-	return TSharedPtr<HandlerComponent>(NULL);
+	return TSharedPtr<HandlerComponent>(nullptr);
 }
 
 void FOodleComponentModuleInterface::StartupModule()
diff --git a/Engine/Plugins/Runtime/PacketHandlers/CompressionComponents/Oodle/Source/OodleHandlerComponent/Private/OodleTrainerCommandlet.cpp b/Engine/Plugins/Runtime/PacketHandlers/CompressionComponents/Oodle/Source/OodleHandlerComponent/Private/OodleTrainerCommandlet.cpp
--- a/Engine/Plugins/Runtime/PacketHandlers/CompressionComponents/Oodle/Source/OodleHandlerComponent/Private/OodleTrainerCommandlet.cpp
+++ b/Engine/Plugins/Runtime/PacketHandlers/CompressionComponents/Oodle/Source/OodleHandlerComponent/Private/OodleTrainerCommandlet.cpp
@@ -4,7 +4,6 @@
 
 #include "OodleTrainerCommandlet.h"
 
-
 #define LOCTEXT_NAMESPACE "Oodle"
 
 
@@ -12,6 +11,7 @@
 
 // @todo #JohnB: Gather stats while merging, to give the commandlet some more informative output
 
+
 /**
  * UOodleTrainerCommandlet
  */
@@ -19,6 +19,7 @@
 UOodleTrainerCommandlet::UOodleTrainerCommandlet(const FObjectInitializer& ObjectInitializer)
 	: Super(ObjectInitializer)
 {
+	LogToConsole = true;
 }
 
 int32 UOodleTrainerCommandlet::Main(const FString& Params)
@@ -32,264 +33,916 @@ int32 UOodleTrainerCommandlet::Main(const FString& Params)
 	{
 		FString MainCmd = Tokens[0];
 
-		/**
-		 * Enable
-		 *
-		 * Command for quick-enabling Oodle
-		 */
 		if (MainCmd == TEXT("Enable"))
 		{
-			TArray<FString> Components;
-			bool bAlreadyEnabled = false;
+			HandleEnable();
+		}
+		else if (MainCmd == TEXT("MergePackets"))
+		{
+			if (Tokens.Num() > 2)
+			{
+				FString OutputCapFile = Tokens[1];
+				FString MergeList = Tokens[2];
+				bool bMergeAll = MergeList == TEXT("all");
 
-			GConfig->GetArray(TEXT("PacketHandlerComponents"), TEXT("Components"), Components, GEngineIni);
+				if (!bMergeAll)
+				{
+					HandleMergePackets(OutputCapFile, MergeList);
+				}
+				else if (Tokens.Num() > 3)
+				{
+					MergeList = Tokens[3];
 
-			for (FString CurComponent : Components)
+					HandleMergePackets(OutputCapFile, MergeList, true);
+				}
+				else
+				{
+					UE_LOG(OodleHandlerComponentLog, Error,
+							TEXT("Need to specify a target directory, e.g: MergePackets OutFile All Dir"));
+				}
+			}
+			else
+			{
+				UE_LOG(OodleHandlerComponentLog, Error,
+						TEXT("Error parsing 'MergePackets' parameters. Must specify an output file, and packet files to merge"));
+			}
+		}
+		else if (MainCmd == TEXT("GenerateDictionary"))
+		{
+			if (Tokens.Num() > 2)
 			{
-				if (CurComponent.StartsWith(TEXT("OodleHandlerComponent")))
+				FString OutputDictionaryFile = Tokens[1];
+				FString InputCaptureFiles = Tokens[2];
+				bool bMergeAll = InputCaptureFiles == TEXT("all");
+				FOodleDictionaryGenerator Generator;
+
+				if (!bMergeAll)
 				{
-					bAlreadyEnabled = true;
-					break;
+					Generator.BeginGenerateDictionary(OutputDictionaryFile, InputCaptureFiles, bMergeAll);
+				}
+				else if (Tokens.Num() > 3)
+				{
+					InputCaptureFiles = Tokens[3];
+
+					Generator.BeginGenerateDictionary(OutputDictionaryFile, InputCaptureFiles, bMergeAll);
+				}
+				else
+				{
+					UE_LOG(OodleHandlerComponentLog, Error,
+							TEXT("Need to specify a target directory, e.g: GenerateDictionary OutFile All Dir"));
 				}
 			}
+			else
+			{
+				UE_LOG(OodleHandlerComponentLog, Error,
+						TEXT("Error parsing 'GenerateDictionary' parameters. Must specify an output file, and packet files to merge"));
+			}
+		}
+		else if (MainCmd == TEXT("AutoGenerateDictionaries"))
+		{
+			HandleAutoGenerateDictionaries();
+		}
+		else
+		{
+			UE_LOG(OodleHandlerComponentLog, Error, TEXT("Unrecognized sub-command '%s'"), *MainCmd);
+		}
+	}
+	else
+	{
+		UE_LOG(OodleHandlerComponentLog, Error, TEXT("Error parsing main commandlet sub-command"));
+
+		// @todo #JohnB: Add a 'help' command, with documentation output
+	}
+#else
+	UE_LOG(OodleHandlerComponentLog, Error, TEXT("Oodle unavailable."));
+#endif
+
+
+	return (GWarn->Errors.Num() == 0 ? 0 : 1);
+}
+
+#if HAS_OODLE_SDK
+bool UOodleTrainerCommandlet::HandleEnable()
+{
+	bool bSuccess = true;
+	TArray<FString> Components;
+	bool bAlreadyEnabled = false;
+
+	GConfig->GetArray(TEXT("PacketHandlerComponents"), TEXT("Components"), Components, GEngineIni);
+
+	for (FString CurComponent : Components)
+	{
+		if (CurComponent.StartsWith(TEXT("OodleHandlerComponent")))
+		{
+			bAlreadyEnabled = true;
+			break;
+		}
+	}
+
+
+	if (!bAlreadyEnabled)
+	{
+		GConfig->SetBool(TEXT("PacketHandlerComponents"), TEXT("bEnabled"), true, GEngineIni);
+
+		Components.Add(TEXT("OodleHandlerComponent"));
+		GConfig->SetArray(TEXT("PacketHandlerComponents"), TEXT("Components"), Components, GEngineIni);
+
+		OodleHandlerComponent::InitFirstRunConfig();
+
+		GConfig->Flush(false);
+
+
+		UE_LOG(OodleHandlerComponentLog, Log, TEXT("Initialized first-run settings for Oodle (will start in Trainer mode)."));
+	}
+	else
+	{
+		UE_LOG(OodleHandlerComponentLog, Error, TEXT("The Oodle PacketHandler is already enabled."));
+	}
+
+	return bSuccess;
+}
+
+bool UOodleTrainerCommandlet::HandleMergePackets(FString OutputCapFile, FString MergeList, bool bMergeAll/*=false*/)
+{
+	bool bSuccess = true;
+	
+	// Adjust non-absolute directories, so they are relative to the Saved\Oodle\Server directory
+	if (FPaths::IsRelative(OutputCapFile))
+	{
+		OutputCapFile = FPaths::Combine(*GOodleSaveDir, TEXT("Server"), *OutputCapFile);
+	}
 
+	bSuccess = VerifyOutputFile(OutputCapFile);
 
-			if (!bAlreadyEnabled)
+
+	TMap<FArchive*, FString> MergeMap;
+
+	bSuccess = bSuccess && GetMergeMapFromList(MergeList, MergeMap, bMergeAll);
+
+	if (bSuccess)
+	{
+		FArchive* OutArc = IFileManager::Get().CreateFileWriter(*OutputCapFile);
+
+		if (OutArc != nullptr)
+		{
+			UE_LOG(OodleHandlerComponentLog, Log, TEXT("Merging files into output file: %s"), *OutputCapFile);
+
+			FPacketCaptureArchive OutputFile(*OutArc);
+			bool bErrorAppending = false;
+
+			OutputFile.SerializeCaptureHeader();
+
+			for (TMap<FArchive*, FString>::TConstIterator It(MergeMap); It; ++It)
 			{
-				GConfig->SetBool(TEXT("PacketHandlerComponents"), TEXT("bEnabled"), true, GEngineIni);
+				FArchive* CurReadArc = It.Key();
+				FPacketCaptureArchive ReadFile(*CurReadArc);
 
-				Components.Add(TEXT("OodleHandlerComponent"));
-				GConfig->SetArray(TEXT("PacketHandlerComponents"), TEXT("Components"), Components, GEngineIni);
+				UE_LOG(OodleHandlerComponentLog, Log, TEXT("    Merging: %s"), *It.Value());
 
-				OodleHandlerComponent::InitFirstRunConfig();
+				OutputFile.AppendPacketFile(ReadFile);
 
+				if (OutputFile.IsError())
+				{
+					bErrorAppending = true;
+					break;
+				}
 
-				UE_LOG(OodleHandlerComponentLog, Log, TEXT("Initialized first-run settings for Oodle (will start in Trainer mode)."));
+				delete CurReadArc;
+			}
+
+			MergeMap.Empty();
+
+
+			if (!bErrorAppending)
+			{
+				bSuccess =  true;
 			}
 			else
 			{
-				UE_LOG(OodleHandlerComponentLog, Error, TEXT("The Oodle PacketHandler is already enabled."));
+				UE_LOG(OodleHandlerComponentLog, Error, TEXT("Error appending packet file."));
 			}
+
+
+			if (bSuccess)
+			{
+				UE_LOG(OodleHandlerComponentLog, Log, TEXT("Merge packets success."));
+			}
+
+
+			OutArc->Close();
+
+			delete OutArc;
+
+			OutArc = nullptr;
 		}
-		/**
-		 * MergePackets
-		 *
-		 * Command for merging capture packets
-		 */
-		else if (MainCmd == TEXT("MergePackets"))
+		else
 		{
-			if (Tokens.Num() > 2)
+			UE_LOG(OodleHandlerComponentLog, Error, TEXT("Could not create output file."));
+		}
+	}
+
+	return bSuccess;
+}
+
+bool UOodleTrainerCommandlet::HandleAutoGenerateDictionaries()
+{
+	bool bSuccess = false;
+
+	class FLocalDirIterator : public IPlatformFile::FDirectoryVisitor
+	{
+		TArray<FString>& OutDirectories;
+
+		FString SearchDir;
+
+	public:
+		FLocalDirIterator(TArray<FString>& OutDirs, FString InSearchDir)
+			: OutDirectories(OutDirs)
+			, SearchDir(InSearchDir)
+		{
+		}
+
+		virtual bool Visit(const TCHAR* FilenameOrDirectory, bool bIsDirectory) override
+		{
+			if (bIsDirectory)
 			{
-				FString OutputCapFile = Tokens[1];
-				FString MergeList = Tokens[2];
-				IPlatformFile& PlatformFile = FPlatformFileManager::Get().GetPlatformFile();
-				TArray<FString> MergeFiles;
-				bool bValidFileParms = true;
+				// Add just the directory name, not full path
+				OutDirectories.Add(FString(FilenameOrDirectory).Mid(SearchDir.Len()+1));
+			}
 
+			return true;
+		}
+	};
 
-				if (MergeList == TEXT("all"))
-				{
-					bValidFileParms = false;
 
-					// @todo #JohnB: Document
+	// Get a list of all directories containing saved packet captures
+	FString OodleCaptureDir = FPaths::Combine(*GOodleSaveDir, TEXT("Server"));
+	TArray<FString> CaptureDirs;
+	FLocalDirIterator DirIterator(CaptureDirs, OodleCaptureDir);
 
-					if (Tokens.Num() > 3)
-					{
-						FString MergeDir = Tokens[3];
-
-						if (FPaths::IsRelative(MergeDir))
-						{
-							MergeDir = FPaths::Combine(*GOodleSaveDir, TEXT("Server"), *MergeDir);
-						}
-
-						if (FPaths::DirectoryExists(MergeDir))
-						{
-							IFileManager::Get().FindFiles(MergeFiles, *MergeDir, TEXT("ucap"));
-
-							if (MergeFiles.Num())
-							{
-								bValidFileParms = true;
-
-								for (int32 i=0; i<MergeFiles.Num(); i++)
-								{
-									MergeFiles[i] = FPaths::Combine(*MergeDir, *MergeFiles[i]);
-								}
-							}
-							else
-							{
-								UE_LOG(OodleHandlerComponentLog, Error, TEXT("Could not find any .ucap files in directory: %s"),
-										*MergeDir);
-							}
-						}
-						else
-						{
-							UE_LOG(OodleHandlerComponentLog, Error, TEXT("Could not find directory: %s"), *MergeDir);
-						}
-					}
-					else
-					{
-						UE_LOG(OodleHandlerComponentLog, Error,
-								TEXT("Need to specify a target directory, e.g: MergePackets OutFile All Dir"));
-					}
+	IFileManager::Get().IterateDirectory(*OodleCaptureDir, DirIterator);
+
+
+	if (CaptureDirs.Num() > 0)
+	{
+		UE_LOG(OodleHandlerComponentLog, Log, TEXT("Generating dictionaries for folders: %s"),
+				*FString::Join<FString>(CaptureDirs, TEXT(", ")));
+
+		// Now generate dictionaries for each directory (named after the directory)
+		for (FString CurDir : CaptureDirs)
+		{
+			FString OutDic = FPaths::Combine(*GOodleContentDir, *FString::Printf(TEXT("%s%s.udic"), FApp::GetGameName(), *CurDir));
+
+			UE_LOG(OodleHandlerComponentLog, Log, TEXT("Generating dictionary for folder '%s', outputting to: %s"), *CurDir, *OutDic);
+
+			FOodleDictionaryGenerator Generator;
+
+			bSuccess = Generator.BeginGenerateDictionary(OutDic, CurDir, true);
+
+			if (bSuccess)
+			{
+				UE_LOG(OodleHandlerComponentLog, Log, TEXT("Successfully generated dictionary."));
+			}
+			else
+			{
+				bool bTriggerFail = false;
+
+				if (FApp::IsUnattended())
+				{
+					bTriggerFail = true;
 				}
 				else
 				{
-					MergeList.ParseIntoArray(MergeFiles, TEXT(","));
-				}
+					FText Msg = FText::Format(
+						LOCTEXT("OodleGenerateContinue", "Failed to generate dictionary for folder '{0}', continue anyway?"),
+						FText::FromString(CurDir));
+
+					EAppReturnType::Type Result = FMessageDialog::Open(EAppMsgType::YesNo, Msg);
 
+					bTriggerFail = (Result != EAppReturnType::Yes);
+				}
 
-				// Adjust non-absolute directories, so they are relative to the Saved\Oodle\Server directory
-				if (bValidFileParms && FPaths::IsRelative(OutputCapFile))
+				if (bTriggerFail)
 				{
-					OutputCapFile = FPaths::Combine(*GOodleSaveDir, TEXT("Server"), *OutputCapFile);
+					UE_LOG(OodleHandlerComponentLog, Error, TEXT("Failed to generate dictionary for folder '%s', aborting."), *CurDir);
 
-					if (FPaths::FileExists(OutputCapFile))
-					{
-						if (FApp::IsUnattended())
-						{
-							// @todo #JohnB: Error message
-
-							bValidFileParms = false;
-						}
-						else
-						{
-							EAppReturnType::Type Result = FMessageDialog::Open(EAppMsgType::YesNo,
-																		LOCTEXT("MergeOverwrite", "Overwrite existing output file?"));
-
-							bValidFileParms = (Result == EAppReturnType::Yes);
-						}
-					}
+					break;
 				}
-
-				for (int32 i=0; i<MergeFiles.Num() && bValidFileParms; i++)
+				else
 				{
-					FString CurMergeFile = MergeFiles[i];
+					UE_LOG(OodleHandlerComponentLog, Warning,
+							TEXT("Failed to generate dictionary for folder '%s, continuing to next folder"), *CurDir);
+				}
+			}
+		}
 
-					if (FPaths::IsRelative(CurMergeFile))
-					{
-						MergeFiles[i] = FPaths::Combine(*GOodleSaveDir, TEXT("Server"), *CurMergeFile);
+		UE_LOG(OodleHandlerComponentLog, Log, TEXT("Finished generating all dictionaries."));
+	}
+	else
+	{
+		UE_LOG(OodleHandlerComponentLog, Log, TEXT("Could not find any packet captures to process."));
+	}
 
-						if (!FPaths::FileExists(CurMergeFile))
-						{
-							UE_LOG(OodleHandlerComponentLog, Error, TEXT("Packet capture file does not exist: %s"), *CurMergeFile);
+	return bSuccess;
+}
 
-							bValidFileParms = false;
-						}
-					}
-				}
+bool UOodleTrainerCommandlet::GetMergeMapFromList(FString MergeList, TMap<FArchive*, FString>& OutMergeMap, bool bMergeAll/*=false*/,
+													bool bAllowSingleFile/*=false*/)
+{
+	bool bSuccess = true;
+	TArray<FString> MergeFiles;
 
+	if (bMergeAll)
+	{
+		FString MergeDir = MergeList;
 
-				if (bValidFileParms && MergeFiles.Num() > 1)
+		bSuccess = false;
+
+		if (FPaths::IsRelative(MergeDir))
+		{
+			MergeDir = FPaths::Combine(*GOodleSaveDir, TEXT("Server"), *MergeDir);
+		}
+
+		if (FPaths::DirectoryExists(MergeDir))
+		{
+			IFileManager::Get().FindFiles(MergeFiles, *MergeDir, TEXT("ucap"));
+
+			if (MergeFiles.Num() > 0)
+			{
+				bSuccess = true;
+
+				for (int32 i=0; i<MergeFiles.Num(); i++)
 				{
-					bool bAllMergeHandlesValid = true;
-					TMap<FArchive*, FString> MergeMap;
+					MergeFiles[i] = FPaths::Combine(*MergeDir, *MergeFiles[i]);
+				}
+			}
+			else
+			{
+				UE_LOG(OodleHandlerComponentLog, Error, TEXT("Could not find any .ucap files in directory: %s"), *MergeDir);
+			}
+		}
+		else
+		{
+			UE_LOG(OodleHandlerComponentLog, Error, TEXT("Could not find directory: %s"), *MergeDir);
+		}
+	}
+	else
+	{
+		MergeList.ParseIntoArray(MergeFiles, TEXT(","));
+	}
 
-					for (FString CurMergeFile : MergeFiles)
-					{
-						FArchive* ReadArc = IFileManager::Get().CreateFileReader(*CurMergeFile);
-
-						if (ReadArc != NULL)
-						{
-							MergeMap.Add(ReadArc, CurMergeFile);
-						}
-						else
-						{
-							UE_LOG(OodleHandlerComponentLog, Error, TEXT("Could not read packet capture file: %s"), *CurMergeFile);
-
-							bAllMergeHandlesValid = false;
-							break;
-						}
-					}
 
-					FArchive* OutArc = (bAllMergeHandlesValid ? IFileManager::Get().CreateFileWriter(*OutputCapFile) : NULL);
+	for (int32 i=0; i<MergeFiles.Num() && bSuccess; i++)
+	{
+		FString CurMergeFile = MergeFiles[i];
 
-					if (bAllMergeHandlesValid && OutArc != NULL)
-					{
-						UE_LOG(OodleHandlerComponentLog, Log, TEXT("Merging files into output file: %s"), *OutputCapFile);
+		if (FPaths::IsRelative(CurMergeFile))
+		{
+			MergeFiles[i] = FPaths::Combine(*GOodleSaveDir, TEXT("Server"), *CurMergeFile);
 
-						FPacketCaptureArchive OutputFile(*OutArc);
-						bool bErrorAppending = false;
+			if (!FPaths::FileExists(CurMergeFile))
+			{
+				UE_LOG(OodleHandlerComponentLog, Error, TEXT("Packet capture file does not exist: %s"), *CurMergeFile);
 
-						OutputFile.SerializeCaptureHeader();
+				bSuccess = false;
+			}
+		}
+	}
 
-						for (TMap<FArchive*, FString>::TConstIterator It(MergeMap); It; ++It)
-						{
-							FArchive* CurReadArc = It.Key();
-							FPacketCaptureArchive ReadFile(*CurReadArc);
+	if (bSuccess && MergeFiles.Num() <= 1 && !bAllowSingleFile)
+	{
+		UE_LOG(OodleHandlerComponentLog, Error, TEXT("You need to specify 2 or more packet files for merging"));
+		bSuccess = false;
+	}
 
-							UE_LOG(OodleHandlerComponentLog, Log, TEXT("    Merging: %s"), *It.Value());
+	if (bSuccess)
+	{
+		for (FString CurMergeFile : MergeFiles)
+		{
+			FArchive* ReadArc = IFileManager::Get().CreateFileReader(*CurMergeFile);
 
-							OutputFile.AppendPacketFile(ReadFile);
+			if (ReadArc != NULL)
+			{
+				OutMergeMap.Add(ReadArc, CurMergeFile);
+			}
+			else
+			{
+				UE_LOG(OodleHandlerComponentLog, Error, TEXT("Could not read packet capture file: %s"), *CurMergeFile);
 
-							if (OutputFile.IsError())
-							{
-								bErrorAppending = true;
-								break;
-							}
+				bSuccess = false;
+				break;
+			}
+		}
+	}
 
-							delete CurReadArc;
-						}
+	return bSuccess;
+}
 
-						MergeMap.Empty();
+bool UOodleTrainerCommandlet::VerifyOutputFile(FString OutputFile)
+{
+	bool bSuccess = true;
 
+	if (FPaths::FileExists(OutputFile))
+	{
+		if (FApp::IsUnattended())
+		{
+			UE_LOG(OodleHandlerComponentLog, Log, TEXT("Output file '%s' already exists. Aborting (can't overwrite in unattended)."),
+					*OutputFile);
 
-						if (bErrorAppending)
-						{
-							UE_LOG(OodleHandlerComponentLog, Error, TEXT("Error appending packet file."));
-						}
+			bSuccess = false;
+		}
+		else
+		{
+			FText Msg = FText::Format(LOCTEXT("OutputFileOverwrite", "Overwrite existing output file '{0}'?"),
+										FText::FromString(OutputFile));
 
+			EAppReturnType::Type Result = FMessageDialog::Open(EAppMsgType::YesNo, Msg);
 
-						// @todo #JohnB: Log success
+			bSuccess = (Result == EAppReturnType::Yes);
+		}
+	}
 
-						OutArc->Close();
+	return bSuccess;
+}
 
-						delete OutArc;
 
-						OutArc = NULL;
-					}
-					else if (!bAllMergeHandlesValid)
-					{
-						UE_LOG(OodleHandlerComponentLog, Error, TEXT("Could not read all packet files."));
-					}
-					else // if (OutArc == NULL)
-					{
-						UE_LOG(OodleHandlerComponentLog, Error, TEXT("Could not create output file."));
-					}
-				}
-				else if (!bValidFileParms)
+/**
+ * FOodleDictionaryGenerator
+ */
+
+bool FOodleDictionaryGenerator::BeginGenerateDictionary(FString InOutputDictionaryFile, FString InputCaptureFiles, bool bMergeAll)
+{
+	bool bSuccess = false;
+
+	OutputDictionaryFile = InOutputDictionaryFile;
+
+	bSuccess = InitGenerator();
+	bSuccess = bSuccess && ReadPackets(InputCaptureFiles, bMergeAll);
+	bSuccess = bSuccess && GenerateAndWriteDictionary();
+
+	Cleanup();
+
+	return bSuccess;
+}
+
+bool FOodleDictionaryGenerator::InitGenerator()
+{
+	bool bSuccess = false;
+
+	// Parse and verify parameters
+	// @todo #JohnB: Set these defaults from the .ini file as well (commandline overriding) - mandatory for tailoring to specific games,
+	//					using the games .ini file
+
+	FParse::Value(FCommandLine::Get(), TEXT("-HashTableSize="), HashTableSize);
+	FParse::Value(FCommandLine::Get(), TEXT("-DictionarySize="), DictionarySize);
+	FParse::Value(FCommandLine::Get(), TEXT("-DictionaryTrials="), DictionaryTrials);
+	FParse::Value(FCommandLine::Get(), TEXT("-TrialRandomness="), TrialRandomness);
+	FParse::Value(FCommandLine::Get(), TEXT("-TrialGenerations="), TrialGenerations);
+	bNoTrials = FParse::Param(FCommandLine::Get(), TEXT("NoTrials"));
+
+	if (HashTableSize < 17 || HashTableSize > 21)
+	{
+		// @todo #JohnB: Don't know if ranges outside of this are supported or are useful
+		HashTableSize = FMath::Clamp(HashTableSize, 17, 21);
+
+		UE_LOG(OodleHandlerComponentLog, Warning, TEXT("Hash table size is usually between 17-21. Clamping to: %i"), HashTableSize);
+	}
+
+	if (DictionarySize < (1024 * 1024) || DictionarySize > OODLENETWORK1_MAX_DICTIONARY_SIZE)
+	{
+		 DictionarySize = FMath::Clamp(DictionarySize, 1024 * 1024, OODLENETWORK1_MAX_DICTIONARY_SIZE);
+
+		 UE_LOG(OodleHandlerComponentLog, Warning, TEXT("Dictionary size is usually between 1MB and %iMB. Clamping to: %iMB"),
+				(OODLENETWORK1_MAX_DICTIONARY_SIZE / (1024 * 1024)), (DictionarySize / (1024 * 1024)));
+	}
+
+	if (bNoTrials)
+	{
+		DictionaryTrials = 0;
+		TrialRandomness = 0;
+		TrialGenerations = 0;
+	}
+	else
+	{
+		if (DictionaryTrials < 2)
+		{
+			DictionaryTrials = 2;
+
+			UE_LOG(OodleHandlerComponentLog, Warning,
+					TEXT("DictionaryTrials must be at least two. Clamping to: 2. Specify -NoTrials to disable."));
+		}
+
+		if (TrialRandomness < 50 || TrialRandomness > 200)
+		{
+			TrialRandomness = FMath::Clamp(TrialRandomness, 50, 200);
+
+			UE_LOG(OodleHandlerComponentLog, Warning, TEXT("Dictionary trial randomness is usually between 50-200%. Clamping to: %i"),
+					TrialRandomness);
+		}
+
+		if (TrialGenerations < 1)
+		{
+			TrialGenerations = 1;
+
+			UE_LOG(OodleHandlerComponentLog, Warning, TEXT("Dictionary trial generations must be at least one. Clamping to: 1."));
+		}
+		else if (TrialGenerations > 1)
+		{
+			UE_LOG(OodleHandlerComponentLog, Log, TEXT("Dictionary trial generations higher than 1, can take a long time."));
+		}
+	}
+
+
+	UE_LOG(OodleHandlerComponentLog, Log, TEXT("Oodle trainer parameters:"));
+	UE_LOG(OodleHandlerComponentLog, Log, TEXT("- HashTableSize: %i"), HashTableSize);
+	UE_LOG(OodleHandlerComponentLog, Log, TEXT("- DictionarySize: %i (~%iMB)"), DictionarySize, (DictionarySize / (1024 * 1024)));
+	UE_LOG(OodleHandlerComponentLog, Log, TEXT("- Performing Trials: %s"), (bNoTrials ? TEXT("No") : TEXT("Yes")));
+
+	if (!bNoTrials)
+	{
+		UE_LOG(OodleHandlerComponentLog, Log, TEXT("- DictionaryTrials: %i"), DictionaryTrials);
+		UE_LOG(OodleHandlerComponentLog, Log, TEXT("- TrialRandomness: %i"), TrialRandomness);
+		UE_LOG(OodleHandlerComponentLog, Log, TEXT("- TrialGenerations: %i"), TrialGenerations);
+	}
+
+
+
+	// Adjust non-absolute directories, so they are relative to the game directory
+	if (FPaths::IsRelative(OutputDictionaryFile))
+	{
+		OutputDictionaryFile = FPaths::Combine(*FPaths::GameDir(), TEXT("Server"), *OutputDictionaryFile);
+	}
+
+	bSuccess = UOodleTrainerCommandlet::VerifyOutputFile(OutputDictionaryFile);
+
+	return bSuccess;
+}
+
+bool FOodleDictionaryGenerator::ReadPackets(FString InputCaptureFiles, bool bMergeAll)
+{
+	bool bSuccess = false;
+
+	bSuccess = UOodleTrainerCommandlet::GetMergeMapFromList(InputCaptureFiles, MergeMap, bMergeAll, true);
+
+	// Now begin training
+	if (bSuccess)
+	{
+		TArray<FArchive*> UnboundArchives;
+		TArray<FPacketCaptureArchive> BoundArchives;
+
+		MergeMap.GenerateKeyArray(UnboundArchives);
+
+		for (FArchive* CurArc : UnboundArchives)
+		{
+			new(BoundArchives) FPacketCaptureArchive(*CurArc);
+		}
+
+
+		// Serialize capture archive headers, and retrieve packet count
+		int32 PacketCount = 0;
+
+		for (FPacketCaptureArchive& CurArc : BoundArchives)
+		{
+			CurArc.SerializeCaptureHeader();
+
+			if (!CurArc.IsError())
+			{
+				PacketCount += CurArc.GetPacketCount();
+			}
+		}
+
+
+		// Generate a randomized index from 0-PacketCount (no repeated values), for randomly splitting packets to 3-4 different lists
+		// NOTE: There is a good reason the randomness is done in this particular way - it ensures good quality randomness.
+		TArray<uint32> RandIndex;
+
+		RandIndex.SetNum(PacketCount);
+
+		for (int32 i=0; i<PacketCount; i++)
+		{
+			RandIndex[i] = i;
+		}
+
+		// Knuth shuffle
+		int32 LastIdx = RandIndex.Num()-1;
+
+		for (int32 i=0; i<LastIdx; i++)
+		{
+			int32 SwapIdx = i + FMath::RandRange(0, LastIdx - i);
+
+			if (SwapIdx != i)
+			{
+				RandIndex.Swap(i, SwapIdx);
+			}
+		}
+
+		// Now translate the randomized index, back into a sequential index, which maps random packets to each list-type (from type 0-4)
+		bool bCompressionTestPackets = false;	// @todo #JohnB: Enable from commandline
+		int32 ListSplitCount = 3 + (bCompressionTestPackets ? 1 : 0);
+		TArray<uint8> PacketToListMap;
+
+		PacketToListMap.SetNum(PacketCount);
+
+		for (int32 i=0; i<RandIndex.Num(); i++)
+		{
+			PacketToListMap[RandIndex[i]] = (i % ListSplitCount);
+		}
+
+
+		// Now step through the packets, and randomly assign them to the appropriate list
+		int32 AvgSplitSize = (PacketCount / ListSplitCount) + 1;
+
+		DictionaryPackets.Empty(AvgSplitSize);
+		DictionaryPacketSizes.Empty(AvgSplitSize);
+		DictionaryTestPackets.Empty(AvgSplitSize);
+		DictionaryTestPacketSizes.Empty(AvgSplitSize);
+
+		// Dictionary test packets may overflow into the trainer packets list, so allocate two-lists worth
+		TrainerPackets.Empty(AvgSplitSize * 2);
+		TrainerPacketSizes.Empty(AvgSplitSize * 2);
+
+		if (bCompressionTestPackets)
+		{
+			CompressionTestPackets.Empty(AvgSplitSize);
+			CompressionTestPacketSizes.Empty(AvgSplitSize);
+		}
+
+
+		const uint32 DictionaryTestLimit = 1024 * 1024 * 64;
+		const uint32 BufferSize = 1048576;
+		uint8* ReadBuffer = new uint8[BufferSize];
+		uint32 PacketIdx = 0;
+
+		for (FPacketCaptureArchive& CurArc : BoundArchives)
+		{
+			while (CurArc.Tell() < CurArc.TotalSize())
+			{
+				uint32 PacketSize = BufferSize;
+
+				CurArc.SerializePacket(ReadBuffer, PacketSize);
+
+				if (CurArc.IsError())
 				{
-					UE_LOG(OodleHandlerComponentLog, Error,
-							TEXT("The packet files must exist, and the output file must not already exist"));
+					UE_LOG(OodleHandlerComponentLog, Warning, TEXT("Error serializing packet from packet capture archive. Skipping."));
+
+					break;
 				}
-				else // if (MergeFiles.Num() <= 1)
+
+				check(PacketSize > 0);
+
+
+				int32 ListIdx = PacketToListMap[PacketIdx];
+
+				check(ListIdx < ListSplitCount);
+
+				TArray<void*>& TargetPacketList =
+					(
+						ListIdx == 0 ?			DictionaryPackets :
+						ListIdx == 1 ?			(bDictionaryTestOverflow ? TrainerPackets : DictionaryTestPackets) :
+						ListIdx == 2 ?			TrainerPackets :
+						/* ListIdx == 3 ? */	CompressionTestPackets
+					);
+
+				TArray<S32>& TargetSizesList =
+					(
+						ListIdx == 0 ?			DictionaryPacketSizes :
+						ListIdx == 1 ?			(bDictionaryTestOverflow ? TrainerPacketSizes : DictionaryTestPacketSizes) :
+						ListIdx == 2 ?			TrainerPacketSizes :
+						/* ListIdx == 3 ? */	CompressionTestPacketSizes
+					);
+
+				uint32& TargetBytes =
+					(
+						ListIdx == 0 ?			DictionaryPacketBytes :
+						ListIdx == 1 ?			(bDictionaryTestOverflow ? TrainerPacketBytes : DictionaryTestPacketBytes) :
+						ListIdx == 2 ?			TrainerPacketBytes :
+						/* ListIdx == 3 ? */	CompressionTestPacketBytes
+					);
+
+
+				void* PacketDest = new uint8[PacketSize];
+
+				TargetPacketList.Add(PacketDest);
+
+
+				FMemory::Memcpy(PacketDest, ReadBuffer, PacketSize);
+
+				TargetSizesList.Add(PacketSize);
+				TargetBytes += PacketSize;
+
+
+				// Check if DictionaryTestPackets has overflowed (allow it to overflow, but just once)
+				if (!bDictionaryTestOverflow && ListIdx == 1)
 				{
-					UE_LOG(OodleHandlerComponentLog, Error, TEXT("You need to specify 2 or more packet files for merging"));
+					bDictionaryTestOverflow = DictionaryTestPacketBytes > DictionaryTestLimit;
+
+					if (bDictionaryTestOverflow)
+					{
+						UE_LOG(OodleHandlerComponentLog, Log, TEXT("DictionaryTestPackets overflowed into TrainerPackets (normal)."));
+					}
 				}
+
+
+				PacketIdx++;
 			}
-			else
-			{
-				UE_LOG(OodleHandlerComponentLog, Error,
-						TEXT("Error parsing 'MergePackets' parameters. Must specify an output file, and packet files to merge"));
-			}
+		}
+
+		delete[] ReadBuffer;
+		ReadBuffer = nullptr;
+	}
+
+	return bSuccess;
+}
+
+bool FOodleDictionaryGenerator::GenerateAndWriteDictionary()
+{
+	bool bSuccess = false;
+	void* NewDictionaryData = new uint8[DictionarySize];
+
+
+	UE_LOG(OodleHandlerComponentLog, Log, TEXT("Beginning dictionary generation..."));
+
+	UE_LOG(OodleHandlerComponentLog, Log, TEXT("- DictionaryPackets: %i, Size: %i (~%iMB)"),
+			DictionaryPackets.Num(), DictionaryPacketBytes, (DictionaryPacketBytes / (1024 * 1024)));
+
+	UE_LOG(OodleHandlerComponentLog, Log, TEXT("- DictionaryTestPackets: %i, Size: %i (~%iMB), Overflowed: %s"),
+			DictionaryTestPackets.Num(), DictionaryTestPacketBytes, (DictionaryTestPacketBytes / (1024 * 1024)),
+			(bDictionaryTestOverflow ? TEXT("Yes") : TEXT("No")));
+
+	UE_LOG(OodleHandlerComponentLog, Log, TEXT("- TrainerPackets: %i, Size: %i (~%iMB)"),
+			TrainerPackets.Num(), TrainerPacketBytes, (TrainerPacketBytes / (1024 * 1024)));
+
+
+	if (bNoTrials)
+	{
+		OodleNetwork1_SelectDictionaryFromPackets(NewDictionaryData, DictionarySize, HashTableSize,
+			(const void**)DictionaryPackets.GetData(), DictionaryPacketSizes.GetData(), DictionaryPackets.Num(),
+			(const void**)DictionaryTestPackets.GetData(), DictionaryTestPacketSizes.GetData(), DictionaryTestPackets.Num());
+	}
+	else
+	{
+		OodleNetwork1_SelectDictionaryFromPackets_Trials(NewDictionaryData, DictionarySize, HashTableSize,
+			(const void**)DictionaryPackets.GetData(), DictionaryPacketSizes.GetData(), DictionaryPackets.Num(),
+			(const void**)DictionaryTestPackets.GetData(), DictionaryTestPacketSizes.GetData(), DictionaryTestPackets.Num(),
+			DictionaryTrials, TrialRandomness, TrialGenerations);
+	}
+
+	UE_LOG(OodleHandlerComponentLog, Log, TEXT("Dictionary generation complete. Beginning Oodle state training..."));
+
+
+	// Begin training the initial Oodle state
+	const SINTa SharedStruct_Size = OodleNetwork1_Shared_Size(HashTableSize);
+	const SINTa StateStruct_Size = OodleNetwork1UDP_State_Size();
+	OodleNetwork1_Shared* SharedDictionary = (OodleNetwork1_Shared*)FMemory::Malloc(SharedStruct_Size);
+	OodleNetwork1UDP_State* CompressorState = (OodleNetwork1UDP_State*)FMemory::Malloc(StateStruct_Size);
+
+	OodleNetwork1_Shared_SetWindow(SharedDictionary, HashTableSize, NewDictionaryData, (S32)DictionarySize);
+
+	OodleNetwork1UDP_Train(CompressorState, SharedDictionary,
+		(const void**)TrainerPackets.GetData(), TrainerPacketSizes.GetData(), TrainerPackets.Num());
+
+
+
+	// Now compact the dictionary and initial state data, in preparation for writing to file
+	const SINTa CompactState_Size = OodleNetwork1UDP_StateCompacted_MaxSize();
+	OodleNetwork1UDP_StateCompacted* CompactCompressorState = (OodleNetwork1UDP_StateCompacted*)FMemory::Malloc(CompactState_Size);
+	SINTa CompactCompressorStateBytes = OodleNetwork1UDP_State_Compact(CompactCompressorState, CompressorState);
+
+
+	// Now write to the final file
+	FArchive* OutArc = IFileManager::Get().CreateFileWriter(*OutputDictionaryFile);
+
+	if (OutArc != nullptr)
+	{
+		FOodleDictionaryArchive OutputFile(*OutArc);
+
+		OutputFile.SetDictionaryHeaderValues(HashTableSize);
+		OutputFile.SerializeHeader();
+
+		bSuccess = !OutputFile.IsError();
+
+		bSuccess = bSuccess && OutputFile.SerializeOodleCompressData(OutputFile.Header.DictionaryData, NewDictionaryData,
+																		DictionarySize);
+
+		bSuccess = bSuccess && OutputFile.SerializeOodleCompressData(OutputFile.Header.CompressorData, CompactCompressorState,
+																		CompactCompressorStateBytes);
+
+
+		if (bSuccess)
+		{
+			UE_LOG(OodleHandlerComponentLog, Log, TEXT("Successfully processed packet captures, and wrote dictionary file."));
 		}
 		else
 		{
-			UE_LOG(OodleHandlerComponentLog, Error, TEXT("Unrecognized sub-command '%s'"), *MainCmd);
+			UE_LOG(OodleHandlerComponentLog, Error, TEXT("Error writing dictionary file."));
 		}
+
+
+		OutArc->Close();
+
+		delete OutArc;
 	}
 	else
 	{
-		UE_LOG(OodleHandlerComponentLog, Error, TEXT("Error parsing main commandlet sub-command"));
+		UE_LOG(OodleHandlerComponentLog, Error, TEXT("Failed to create output dictionary file '%s'"), *OutputDictionaryFile);
+	}
 
-		// @todo #JohnB: Add a 'help' command, with documentation output
+
+	// Cleanup allocated data
+	delete[] NewDictionaryData;
+	NewDictionaryData = nullptr;
+
+	if (CompactCompressorState != nullptr)
+	{
+		FMemory::Free(CompactCompressorState);
 	}
-#else
-	UE_LOG(OodleHandlerComponentLog, Error, TEXT("Oodle unavailable."));
-#endif
 
+	if (CompressorState != nullptr)
+	{
+		FMemory::Free(CompressorState);
+	}
 
-	return (GWarn->Errors.Num() == 0 ? 0 : 1);
+	if (SharedDictionary != nullptr)
+	{
+		FMemory::Free(SharedDictionary);
+	}
+
+	CompressorState = nullptr;
+	SharedDictionary = nullptr;
+
+
+
+	// @todo #JohnB: Move to more appropriate place
+	if (bCompressionTestPackets)
+	{
+		UE_LOG(OodleHandlerComponentLog, Log, TEXT("- CompressionTestPackets: %i, Size: %i (~%iMB)"),
+				CompressionTestPackets.Num(), CompressionTestPacketBytes, (CompressionTestPacketBytes / (1024 * 1024)));
+	}
+
+
+	// @todo #JohnB: IMPORTANT: Add a disabled-by-default fourth 'compression test' list,
+	//							which also is randomly selected like above, for testing the dictionary after generation
+	//							(keep it disabled normally though, to maximize packet capture usage)
+
+
+	return bSuccess;
 }
 
+void FOodleDictionaryGenerator::Cleanup()
+{
+	for (void* CurPacket : DictionaryPackets)
+	{
+		if (CurPacket != nullptr)
+		{
+			delete[] CurPacket;
+		}
+	}
+
+	for (void* CurPacket : DictionaryTestPackets)
+	{
+		if (CurPacket != nullptr)
+		{
+			delete[] CurPacket;
+		}
+	}
+
+	for (void* CurPacket : TrainerPackets)
+	{
+		if (CurPacket != nullptr)
+		{
+			delete[] CurPacket;
+		}
+	}
+
+	for (void* CurPacket : CompressionTestPackets)
+	{
+		if (CurPacket != nullptr)
+		{
+			delete[] CurPacket;
+		}
+	}
+
+	DictionaryPackets.Empty();
+	DictionaryTestPackets.Empty();
+	TrainerPackets.Empty();
+	CompressionTestPackets.Empty();
+
+
+	if (MergeMap.Num() > 0)
+	{
+		for (TMap<FArchive*, FString>::TConstIterator It(MergeMap); It; ++It)
+		{
+			FArchive* CurKey = It.Key();
+
+			if (CurKey != nullptr)
+			{
+				delete CurKey;
+			}
+		}
+
+		MergeMap.Empty();
+	}
+}
+#endif
+
 #undef LOCTEXT_NAMESPACE
 
diff --git a/Engine/Plugins/Runtime/PacketHandlers/CompressionComponents/Oodle/Source/OodleHandlerComponent/Public/OodleArchives.h b/Engine/Plugins/Runtime/PacketHandlers/CompressionComponents/Oodle/Source/OodleHandlerComponent/Public/OodleArchives.h
--- a/Engine/Plugins/Runtime/PacketHandlers/CompressionComponents/Oodle/Source/OodleHandlerComponent/Public/OodleArchives.h
+++ b/Engine/Plugins/Runtime/PacketHandlers/CompressionComponents/Oodle/Source/OodleHandlerComponent/Public/OodleArchives.h
@@ -5,27 +5,150 @@
 #if HAS_OODLE_SDK
 #include "ArchiveBase.h"
 
-#define CAPTURE_HEADER_MAGIC	0x41091CC4
-#define CAPTURE_FILE_VERSION	0x00000001
+#define CAPTURE_HEADER_MAGIC		0x41091CC4
+#define CAPTURE_FILE_VERSION		0x00000002
 
-#define DICTIONARY_HEADER_MAGIC	0x1B1BACD4
-#define DICTIONARY_FILE_VERSION	0x00000001
+#define DICTIONARY_HEADER_MAGIC		0x1B1BACD4
+#define DICTIONARY_FILE_VERSION		0x00000001
 
 
-// @todo #JohnB: Add to the header, whether the capture file is an input or output capture type,
-//					and use that to give a warning when merging mismatched input/output capture files
-//					(actually, make that error by default, requiring a special -blah commandline parm, to override the error)
+// Markers for capture file version updates
+#define CAPTURE_VER_PACKETCOUNT		0x00000002	// Added packet count to file header
 
 
 /**
- * @todo #JohnB
+ * Base file archive for the Oodle plugin.
+ *
+ * Contains some common code, such as Oodle file compression, and helper structs/types, for archive writing/navigation
  */
 class FOodleArchiveBase : public FArchiveProxy
 {
+protected:
+	/**
+	 * Encapsulates a value written to an archive, which can be seamlessly rewritten at any time, without disturbing seek position
+	 * NOTE: Don't use this with types of variable size, e.g. FString, only with types of well defined size
+	 */
+	template<class Type>
+	struct TRewritable
+	{
+		/** Written variables */
+	private:
+		/** The property value */
+		Type Value;
+
+
+		/** Runtime variables */
+	private:
+		/** The offset of the value, into the archive */
+		uint32 Offset;
+
+
+		/**
+		 * Disable default constructor
+		 */
+		TRewritable() {}
+
+	public:
+		/**
+		 * Initializing constructor
+		 */
+		TRewritable(Type DefaultVal)
+			: Value(DefaultVal)
+			, Offset(0)
+		{
+		}
+
+		/**
+		 * Retrieves the property value
+		 *
+		 * @return	Returns the property value
+		 */
+		FORCEINLINE Type Get()
+		{
+			return Value;
+		}
+
+		/**
+		 * Sets the property value
+		 *
+		 * @param Ar		The archive being serialized to
+		 * @param InValue	The new value
+		 */
+		FORCEINLINE void Set(FOodleArchiveBase& Ar, Type InValue)
+		{
+			Value = InValue;
+
+			if (Offset != 0)
+			{
+				Ar.SeekPush(Offset);
+				Ar << InValue;
+				Ar.SeekPop();
+			}
+		}
+
+		/**
+		 * Serializes the value, to/from an archive
+		 *
+		 * @param Ar	The archive being serialized to/from
+		 */
+		FORCEINLINE void Serialize(FArchive& Ar)
+		{
+			if (Offset == 0)
+			{
+				Offset = Ar.Tell();
+			}
+			else
+			{
+				check(Offset == Ar.Tell());
+			}
+
+			Ar << Value;
+		}
+	};
+
+	/**
+	 * Struct for handling compressed data within the archive
+	 */
+	struct FOodleCompressedData
+	{
+		/** Written values */
+	public:
+		/** The offset of the compressed data, within the archive */
+		TRewritable<uint32> Offset;
+
+		/** The compressed length of the data */
+		TRewritable<uint32> CompressedLength;
+
+		/** The decompressed length of the data */
+		TRewritable<uint32> DecompressedLength;
+
+	public:
+		/**
+		 * Base constructor
+		 */
+		FOodleCompressedData()
+			: Offset(0)
+			, CompressedLength(0)
+			, DecompressedLength(0)
+		{
+		}
+
+		/**
+		 * Serialize this struct to/from an archive
+		 *
+		 * @param Ar	The archive to serialize to/from
+		 */
+		void Serialize(FArchive& Ar);
+	};
+
 protected:
 	/** Whether or not to flush immediately */
 	bool bImmediateFlush;
 
+private:
+	/** Stack of previous seek positions */
+	TArray<int64> SeekStack;
+
 public:
 	/**
 	 * Base constructor
@@ -43,6 +166,83 @@ public:
 	{
 		delete (FArchive*)&InnerArchive;
 	}
+
+
+	/**
+	 * Pushes the current archive position onto a stack
+	 */
+	void SeekPush()
+	{
+		SeekStack.Push(InnerArchive.Tell());
+	}
+
+	/**
+	 * Pushes the current archive position onto a stack, and seeks to a new position
+	 *
+	 * @param SeekPos	The position to seek to
+	 */
+	void SeekPush(int64 SeekPos)
+	{
+		SeekStack.Push(InnerArchive.Tell());
+		InnerArchive.Seek(SeekPos);
+	}
+
+	/**
+	 * Pops the most recent archive position from the stack, and seeks to that position
+	 */
+	void SeekPop()
+	{
+		check(SeekStack.Num() > 0);
+		InnerArchive.Seek(SeekStack.Pop());
+	}
+
+
+	/**
+	 * Compresses data and serializes it into the archive, and writes/serializes info to the specified FOodleCompressedData struct
+	 *
+	 * @param OutDataInfo	The FOodleCompressedData struct in the file header, which will store info needed for decompression
+	 * @param Data			The data to be compressed
+	 * @param DataBytes		The length of Data
+	 * @return				Whether or not compression was successful
+	 */
+	bool SerializeOodleCompressData(FOodleCompressedData& OutDataInfo, void* Data, uint32 DataBytes);
+
+	/**
+	 * Decompresses the data referenced by the input FOodleCompressedData struct
+	 *
+	 * @param DataInfo		The FOodleCompressedData struct, describing the location of the data, and info necessary for decompressing
+	 * @param OutData		Outputs the newly decompressed/allocated data, to a pointer (NOTE: Initialize the pointer to nullptr)
+	 * @param OutDataBytes	Outputs the size of OutData
+	 * @return				Whether or not decompression was successful
+	 */
+	bool SerializeOodleDecompressData(FOodleCompressedData& DataInfo, void*& OutData, uint32& OutDataBytes);
+
+
+	/**
+	 * Serializes an FOodleCompressedData struct to/from an archive
+	 *
+	 * @param Ar		The archive to serialize to/from
+	 * @param Value		The struct value to serialize
+	 * @return			Returns a chained-reference to the archive
+	 */
+	FORCEINLINE friend FArchive& operator<<(FArchive& Ar, FOodleCompressedData& Value)
+	{
+		Value.Serialize(Ar);
+		return Ar;
+	}
+
+	/**
+	 * Serializes a rewritable (e.g. seeks to value, writes, then returns to previous position) value to/from an archive
+	 *
+	 * @param Ar		The archive to serialize to/from
+	 * @param Value		The rewritable value to serialize
+	 * @return			Returns a chained-reference to the archive
+	 */
+	template<class Type> FORCEINLINE friend FArchive& operator<<(FArchive& Ar, TRewritable<Type>& Value)
+	{
+		Value.Serialize(Ar);
+		return Ar;
+	}
 };
 
 
@@ -52,8 +252,53 @@ public:
 class FPacketCaptureArchive : public FOodleArchiveBase
 {
 protected:
-	/** Capture file format version */
-	uint32 CaptureVersion;
+	/**
+	 * Capture file header
+	 */
+	struct FCaptureHeader
+	{
+		/** Written header values */
+	public:
+		/** Unique value indicating this file type */
+		uint32 Magic;
+
+		/** Capture file format version */
+		uint32 CaptureVersion;
+
+		/** Total number of captured packets */
+		TRewritable<uint32> PacketCount;
+
+		/** Position in the archive, where the packet data starts */
+		TRewritable<uint32> PacketDataOffset;
+
+		/** Total length of all packet data */
+		TRewritable<uint32> PacketDataLength;
+
+
+	public:
+		/**
+		 * Base constructor
+		 */
+		FCaptureHeader()
+			: Magic(CAPTURE_HEADER_MAGIC)
+			, CaptureVersion(0)
+			, PacketCount(0)
+			, PacketDataOffset(0)
+			, PacketDataLength(0)
+		{
+		}
+
+		/**
+		 * Serialize the capture file header
+		 *
+		 * @param Ar	The capture archive being serialized to/from
+		 */
+		void SerializeHeader(FPacketCaptureArchive& Ar);
+	};
+
+
+	/** The capture file header */
+	FCaptureHeader Header;
 
 public:
 	/**
@@ -61,14 +306,15 @@ public:
 	 */
 	FPacketCaptureArchive(FArchive& InInnerArchive)
 		: FOodleArchiveBase(InInnerArchive)
+		, Header()
 	{
 		if (IsSaving())
 		{
-			CaptureVersion = CAPTURE_FILE_VERSION;
+			Header.CaptureVersion = CAPTURE_FILE_VERSION;
 		}
 		else
 		{
-			CaptureVersion = 0;
+			Header.CaptureVersion = 0;
 		}
 	}
 
@@ -94,33 +340,102 @@ public:
 	 * @param InPacketFile	The packet file archive, to append to this one.
 	 */
 	void AppendPacketFile(FPacketCaptureArchive& InPacketFile);
-};
 
+	/**
+	 * Returns PacketCount, the total number of packets in the file
+	 */
+	uint32 GetPacketCount();
+};
 
-// @todo #JohnB: Remove once complete
-/*
-	Values to include in header:
-		- Dictionary size (as relating to Oodle API methods)
-		- Hash table size (also relating to Oodle API)
-*/
 
 /**
- * @todo #JohnB
+ * Archive for handling UE4 Oodle dictionary (.udic) files
  */
 class FOodleDictionaryArchive : public FOodleArchiveBase
 {
-protected:
-	/** Dictionary file format version */
-	uint32 DictionaryVersion;
+public:
+	/**
+	 * Dictionary file header
+	 */
+	class FDictionaryHeader
+	{
+		/** Written header values */
+	public:
+		/** Unique value indicating this file type */
+		uint32 Magic;
+
+		/** Dictionary file format version */
+		uint32 DictionaryVersion;
+
+		/** Oodle header version - noting changes in Oodle data format (only the major-version reflects file format changes) */
+		uint32 OodleMajorHeaderVersion;
+
+		/** Size of the hash table used for the dictionary */
+		TRewritable<int32> HashTableSize;
+
+		/** Compressed dictionary data, within the archive */
+		FOodleCompressedData DictionaryData;
+
+		/** Compressed Oodle compressor state data, within the archive */
+		FOodleCompressedData CompressorData;
+
+
+		/**
+		 * Base constructor
+		 */
+		FDictionaryHeader()
+			: Magic(DICTIONARY_HEADER_MAGIC)
+			, DictionaryVersion(0)
+			, OodleMajorHeaderVersion(0)
+			, HashTableSize(INDEX_NONE)
+			, DictionaryData()
+			, CompressorData()
+		{
+		}
+
+
+		/**
+		 * Serializes the dictionary file header
+		 *
+		 * @param Ar	The dictionary archive to serialize to/from
+		 */
+		void SerializeHeader(FOodleDictionaryArchive& Ar);
+	};
+
+
+	/** The dictionary file header */
+	FDictionaryHeader Header;
 
 public:
 	/**
 	 * Base constructor
 	 */
-	FOodleDictionaryArchive(FArchive& InInnerArchive)
-		: FOodleArchiveBase(InInnerArchive)
-	{
-	}
+	FOodleDictionaryArchive(FArchive& InInnerArchive);
+
+	/**
+	 * Sets dictionary header values, that should be set prior to serializing the header to file
+	 *
+	 * @param InHashTableSize	The HashTableSize to set in the header
+	 */
+	void SetDictionaryHeaderValues(int32 InHashTableSize);
+
+	/**
+	 * Serializes initial basic header values
+	 */
+	void SerializeHeader();
+
+	/**
+	 * Serializes raw dictionary data and compressor state to/from file, compressing/decompressing the data as needed
+	 *
+	 * IMPORTANT: When serializing from disk, the caller is responsible for freeing DictionaryData and CompactCompressorState pointers
+	 *
+	 * @param DictionaryData				Pointer to dictionary data, when saving, and reference to nullptr void*, when loading
+	 * @param DictionaryBytes				Inputs/Outputs the size of the dictionary
+	 * @param CompactCompressorState		Pointer to compressor data, when saving, and reference to nullptr void*, when loading
+	 * @param CompactCompressorStateBytes	Inputs/Outputs the size of the compressor data
+	 */
+	void SerializeDictionaryAndState(void*& DictionaryData, uint32& DictionaryBytes,
+										void*& CompactCompresorState, uint32& CompactCompressorStateBytes);
 };
 
 #endif
diff --git a/Engine/Plugins/Runtime/PacketHandlers/CompressionComponents/Oodle/Source/OodleHandlerComponent/Public/OodleHandlerComponent.h b/Engine/Plugins/Runtime/PacketHandlers/CompressionComponents/Oodle/Source/OodleHandlerComponent/Public/OodleHandlerComponent.h
--- a/Engine/Plugins/Runtime/PacketHandlers/CompressionComponents/Oodle/Source/OodleHandlerComponent/Public/OodleHandlerComponent.h
+++ b/Engine/Plugins/Runtime/PacketHandlers/CompressionComponents/Oodle/Source/OodleHandlerComponent/Public/OodleHandlerComponent.h
@@ -20,19 +20,172 @@ DECLARE_LOG_CATEGORY_EXTERN(OodleHandlerComponentLog, Log, All);
 #endif
 
 
-// Whether or not to utilize the Oodle-example-code based training and packet capture code (in the process of deprecation)
-#define EXAMPLE_CAPTURE_FORMAT 0
+#define CAPTURE_EXT TEXT(".ucap")
 
-#if EXAMPLE_CAPTURE_FORMAT
-	#define CAPTURE_EXT TEXT(".bin")
-#else
-	#define CAPTURE_EXT TEXT(".ucap")
+
+/** Stats */
+
+#if STATS
+
+// @todo #JohnB: Deprecate Oodle in 'stat net' stats, and rely only on 'stat oodle'?
+#define OODLE_STAT_NET 0
+
+#if OODLE_STAT_NET
+// NOTE: These stats can not be compared to normal 'stat net' stats - these stats do not count packet overhead, 'stat net' does
+DECLARE_DWORD_ACCUMULATOR_STAT_EXTERN(TEXT("Oodle Out Rate Raw (bytes)"), STAT_Net_Oodle_OutRaw, STATGROUP_Net, );
+DECLARE_DWORD_ACCUMULATOR_STAT_EXTERN(TEXT("Oodle Out Rate Compressed (bytes)"), STAT_Net_Oodle_OutCompressed, STATGROUP_Net, );
+DECLARE_DWORD_ACCUMULATOR_STAT_EXTERN(TEXT("Oodle In Rate Raw (bytes)"), STAT_Net_Oodle_InRaw, STATGROUP_Net, );
+DECLARE_DWORD_ACCUMULATOR_STAT_EXTERN(TEXT("Oodle In Rate Compressed (bytes)"), STAT_Net_Oodle_InCompressed, STATGROUP_Net, );
+#endif
+
+
+DECLARE_STATS_GROUP(TEXT("Oodle"), STATGROUP_Oodle, STATCAT_Advanced)
+
+DECLARE_DWORD_ACCUMULATOR_STAT_EXTERN(TEXT("Oodle Out Rate Raw (bytes)"), STAT_Oodle_OutRaw, STATGROUP_Oodle, );
+DECLARE_DWORD_ACCUMULATOR_STAT_EXTERN(TEXT("Oodle Out Rate Compressed (bytes)"), STAT_Oodle_OutCompressed, STATGROUP_Oodle, );
+DECLARE_FLOAT_ACCUMULATOR_STAT_EXTERN(TEXT("Oodle Out Rate Savings %"), STAT_Oodle_OutSavings, STATGROUP_Oodle, );
+DECLARE_FLOAT_ACCUMULATOR_STAT_EXTERN(TEXT("Oodle Out Total Savings %"), STAT_Oodle_OutTotalSavings, STATGROUP_Oodle, );
+
+DECLARE_DWORD_ACCUMULATOR_STAT_EXTERN(TEXT("Oodle In Rate Raw (bytes)"), STAT_Oodle_InRaw, STATGROUP_Oodle, );
+DECLARE_DWORD_ACCUMULATOR_STAT_EXTERN(TEXT("Oodle In Rate Compressed (bytes)"), STAT_Oodle_InCompressed, STATGROUP_Oodle, );
+DECLARE_FLOAT_ACCUMULATOR_STAT_EXTERN(TEXT("Oodle In Rate Savings %"), STAT_Oodle_InSavings, STATGROUP_Oodle, );
+DECLARE_FLOAT_ACCUMULATOR_STAT_EXTERN(TEXT("Oodle In Total Savings %"), STAT_Oodle_InTotalSavings, STATGROUP_Oodle, );
+
+#if !UE_BUILD_SHIPPING
+DECLARE_CYCLE_STAT_EXTERN(TEXT("Oodle Out Compress Time"), STAT_Oodle_OutCompressTime, STATGROUP_Oodle, );
+DECLARE_CYCLE_STAT_EXTERN(TEXT("Oodle In Decompress Time"), STAT_Oodle_InDecompressTime, STATGROUP_Oodle, );
+#endif
+
+DECLARE_DWORD_ACCUMULATOR_STAT_EXTERN(TEXT("Oodle Dictionary Count"), STAT_Oodle_DictionaryCount, STATGROUP_Oodle, );
+DECLARE_MEMORY_STAT_EXTERN(TEXT("Oodle Dictionary Bytes"), STAT_Oodle_DictionaryBytes, STATGROUP_Oodle, );
+DECLARE_MEMORY_STAT_EXTERN(TEXT("Oodle Shared Bytes"), STAT_Oodle_SharedBytes, STATGROUP_Oodle, );
+DECLARE_MEMORY_STAT_EXTERN(TEXT("Oodle State Bytes"), STAT_Oodle_StateBytes, STATGROUP_Oodle, );
 #endif
 
 
 /** Globals */
+
+/** The directory Oodle packet captures are saved to */
 extern FString GOodleSaveDir;
 
+/** The directory Oodle dictionaries are saved/loaded to/from */
+extern FString GOodleContentDir;
+
+
+#if STATS
+// @todo #JohnB: The stats collecting is a bit crude, and should probably be giving per-dictionary stats
+
+/**
+ * Stores Oodle net traffic stats, accumulated over the past second, before passing it to the stats system
+ */
+class FOodleNetStats
+{
+private:
+	/** Accumulated stats since last update */
+
+	/** Input traffic compressed packet length */
+	uint32 InCompressedLength;
+
+	/** Input traffic decompressed packet length */
+	uint32 InDecompressedLength;
+
+	/** Output traffic compressed packet length */
+	uint32 OutCompressedLength;
+
+	/** Output traffic uncompressed packet length */
+	uint32 OutUncompressedLength;
+
+
+	/** Time of the last stats update */
+	double LastStatsUpdate;
+
+
+	/** Process lifetime stats */
+
+	/** Total compressed input packets length */
+	uint64 TotalInCompressedLength;
+
+	/** Total decompressed input packets length */
+	uint64 TotalInDecompressedLength;
+
+	/** Total compressed output packets length */
+	uint64 TotalOutCompressedLength;
+
+	/** Total uncompressed output packets length */
+	uint64 TotalOutUncompressedLength;
+
+public:
+	/**
+	 * Base constructor
+	 */
+	FOodleNetStats()
+		: InCompressedLength(0)
+		, InDecompressedLength(0)
+		, OutCompressedLength(0)
+		, OutUncompressedLength(0)
+		, LastStatsUpdate(0.0)
+		, TotalInCompressedLength(0)
+		, TotalInDecompressedLength(0)
+		, TotalOutCompressedLength(0)
+		, TotalOutUncompressedLength(0)
+	{
+	}
+
+	/**
+	 * Process incoming packet stats
+	 *
+	 * @param CompressedLength		The compressed size of the input packet
+	 * @param DecompressedLength	The decompressed size of the input packet
+	 */
+	FORCEINLINE void IncomingStats(uint32 CompressedLength, uint32 DecompressedLength)
+	{
+		InCompressedLength += CompressedLength;
+		TotalInCompressedLength += CompressedLength;
+		InDecompressedLength += DecompressedLength;
+		TotalInDecompressedLength += DecompressedLength;
+
+		CheckForUpdate();
+	}
+
+	/**
+	 * Process outgoing packet stats
+	 *
+	 * @param CompressedLength		The compressed size of the output packet
+	 * @param UncompressedLength	The uncompressed size of the output packets
+	 */
+	FORCEINLINE void OutgoingStats(uint32 CompressedLength, uint32 UncompressedLength)
+	{
+		OutCompressedLength += CompressedLength;
+		TotalOutCompressedLength += CompressedLength;
+		OutUncompressedLength += UncompressedLength;
+		TotalOutUncompressedLength += UncompressedLength;
+
+		CheckForUpdate();
+	}
+
+	/**
+	 * Checks to see if the main stats are due an update, and triggers an update if so
+	 */
+	FORCEINLINE void CheckForUpdate()
+	{
+		float DeltaTime = FPlatformTime::Seconds() - LastStatsUpdate;
+
+		if (DeltaTime > 1.f)
+		{
+			UpdateStats(DeltaTime);
+			LastStatsUpdate = FPlatformTime::Seconds();
+		}
+	}
+
+	/**
+	 * Passes up the accumulated stats, to the main engine stats tracking
+	 *
+	 * @param DeltaTime		The exact time since last stats update
+	 */
+	void UpdateStats(float DeltaTime);
+};
+#endif
+
 
 /**
  * The mode that the Oodle packet handler should operate in
@@ -43,30 +196,53 @@ enum EOodleHandlerMode
 	Release		// Compresses packet data, based on the dictionary file
 };
 
-
 /**
- * Stores dictionary data from an individual file (typically one dictionary each, for server/client)
- *
- * @todo #JohnB: In its current form, this is temporary/stopgap, based on the example code
+ * Encapsulates Oodle dictionary data loaded from file, to be wrapped in a shared pointer (auto-deleting when no longer in use)
  */
 struct FOodleDictionary
 {
-	/** Base constructor */
+	/** Size of the hash table used for the dictionary */
+	uint32 HashTableSize;
+
+	/** The raw dictionary data */
+	void* DictionaryData;
+
+	/** The size of the dictionary */
+	uint32 DictionarySize;
+
+	/** Shared dictionary state */
+	OodleNetwork1_Shared* SharedDictionary;
+
+	/** The size of the shared dictionary data (stored only for memory accounting) */
+	uint32 SharedDictionarySize;
+
+	/** The uncompacted compressor state */
+	OodleNetwork1UDP_State* CompressorState;
+
+	/** The size of CompressorState */
+	uint32 CompressorStateSize;
+
+
+private:
 	FOodleDictionary()
-		: Dictionary(NULL)
-		, SharedDictionaryData(NULL)
-		, CompressorState(NULL)
 	{
 	}
 
-	/* Dictionary */
-	void* Dictionary; 
+	FOodleDictionary(const FOodleDictionary&) = delete;
+	FOodleDictionary& operator=(const FOodleDictionary&) = delete;
 
-	/* Shared Compressor */
-	OodleNetwork1_Shared* SharedDictionaryData;
+public:
 
-	/* State of Oodle */
-	OodleNetwork1UDP_State* CompressorState;
+	/**
+	 * Base constructor
+	 */
+	FOodleDictionary(uint32 InHashTableSize, void* InDictionaryData, uint32 InDictionarySize, OodleNetwork1_Shared* InSharedDictionary,
+						uint32 InSharedDictionarySize, OodleNetwork1UDP_State* InInitialCompressorState, uint32 InCompressorStateSize);
+
+	/**
+	 * Base destructor
+	 */
+	~FOodleDictionary();
 };
 
 
@@ -93,10 +269,17 @@ public:
 	/**
 	 * Initializes FOodleDictionary data, from the specified dictionary file
 	 *
-	 * @param FilePath		The dictionary file path
-	 * @param OutDictionary	The FOodleDictionary struct to write to
+	 * @param FilePath			The dictionary file path
+	 * @param OutDictionary		The FOodleDictionary shared pointer to write to
+	 */
+	void InitializeDictionary(FString FilePath, TSharedPtr<FOodleDictionary>& OutDictionary);
+
+	/**
+	 * Frees the local reference to FOodleDictionary data, and removes it from memory if it was the last reference
+	 *
+	 * @param InDictionary		The FOodleDictionary shared pointer being freed
 	 */
-	void InitializeDictionary(FString FilePath, FOodleDictionary& OutDictionary);
+	void FreeDictionary(TSharedPtr<FOodleDictionary>& InDictionary);
 
 	/**
 	 * Resolves and returns the default dictionary file paths.
@@ -134,22 +317,11 @@ public:
 	virtual void Outgoing(FBitWriter& Packet) override;
 
 protected:
-	/* How many bytes have been saved from compressing the data */
-	uint32 BytesSaved;
-
-#if EXAMPLE_CAPTURE_FORMAT
-	/* File to log packets to */
-	IFileHandle* PacketLogFile;
-
-	/** Whether or not to append data to a statically named log file */
-	bool bPacketLogAppend;
-#else
 	/* File to log input packets to */
 	FPacketCaptureArchive* InPacketLog;
 
 	/* File to log output packets to */
 	FPacketCaptureArchive* OutPacketLog;
-#endif
 
 
 	/** Whether or not Oodle is enabled */
@@ -168,19 +340,18 @@ protected:
 	*/
 	EOodleHandlerMode Mode;
 
-
 	/** Server (Outgoing) dictionary data */
-	FOodleDictionary ServerDictionary;
+	TSharedPtr<FOodleDictionary> ServerDictionary;
 
 	/** Client (Incoming - relative to server) dictionary data */
-	FOodleDictionary ClientDictionary;
+	TSharedPtr<FOodleDictionary> ClientDictionary;
 };
-
-DECLARE_DWORD_COUNTER_STAT_EXTERN(TEXT("Oodle Bytes Raw"), STAT_Oodle_BytesRaw, STATGROUP_Net, );
-DECLARE_DWORD_COUNTER_STAT_EXTERN(TEXT("Oodle Bytes Compressed"), STAT_Oodle_BytesCompressed, STATGROUP_Net, );
 #endif
 
-/* Reliability Module Interface */
+
+/**
+ * Oodle Module Interface
+ */
 class FOodleComponentModuleInterface : public FPacketHandlerComponentModuleInterface
 {
 private:
diff --git a/Engine/Source/Programs/NetworkProfiler/NetworkProfiler/StreamHeader.cs b/Engine/Source/Programs/NetworkProfiler/NetworkProfiler/StreamHeader.cs
--- a/Engine/Source/Programs/NetworkProfiler/NetworkProfiler/StreamHeader.cs
+++ b/Engine/Source/Programs/NetworkProfiler/NetworkProfiler/StreamHeader.cs
@@ -17,23 +17,13 @@ namespace NetworkProfiler
 		public const UInt32 ExpectedMagic = 0x1DBF348C;
 
 		/** We expect this version, or we can't proceed.			*/
-		public const UInt32 ExpectedVersion = 9;
+		public const UInt32 ExpectedVersion = 10;
 
 		/** Magic to ensure we're opening the right file.			*/
 		public UInt32 Magic;
 		/** Version number to detect version mismatches.			*/
 		public UInt32 Version;
 
-		/** Offset in file for name table.							*/
-		public UInt32 NameTableOffset;
-		/** Number of name table entries.							*/
-		public UInt32 NameTableEntries;
-
-		/** Offset in file for address table.							*/
-		public UInt32 AddressTableOffset;
-		/** Number of address table entries.							*/
-		public UInt32 AddressTableEntries;
-
 		/** Tag, set via -networkprofiler=TAG						*/
 		public string Tag;
 		/** Game name, e.g. Example									*/
@@ -142,13 +132,6 @@ namespace NetworkProfiler
 				return;
 			}
 
-			// Name table offset in file and number of entries.
-			NameTableOffset = BinaryStream.ReadUInt32();
-			NameTableEntries = BinaryStream.ReadUInt32();
-
-			AddressTableOffset = BinaryStream.ReadUInt32();
-			AddressTableEntries = BinaryStream.ReadUInt32();
-
 			// Serialize various dynamically-sized strings.
 			Tag = SerializeAnsiString( BinaryStream );
 			GameName = SerializeAnsiString( BinaryStream );
diff --git a/Engine/Source/Programs/NetworkProfiler/NetworkProfiler/StreamParser.cs b/Engine/Source/Programs/NetworkProfiler/NetworkProfiler/StreamParser.cs
--- a/Engine/Source/Programs/NetworkProfiler/NetworkProfiler/StreamParser.cs
+++ b/Engine/Source/Programs/NetworkProfiler/NetworkProfiler/StreamParser.cs
@@ -69,33 +69,6 @@ namespace NetworkProfiler
 			BinaryReader BinaryStream = null; 
 			var Header = StreamHeader.ReadHeader( ParserStream, out BinaryStream );
 
-			// Keep track of token stream offset as name table is at end of file.
-			long TokenStreamOffset = ParserStream.Position;
-
-			// Seek to name table and serialize it.
-			ParserStream.Seek(Header.NameTableOffset,SeekOrigin.Begin);
-			for(int NameIndex = 0;NameIndex < Header.NameTableEntries;NameIndex++)
-			{
-				UInt32 Length = BinaryStream.ReadUInt32();
-				NetworkStream.NameArray.Add(new string(BinaryStream.ReadChars((int)Length)));
-
-				// Find "Unreal" name index used for misc socket parsing optimizations.
-				if( NetworkStream.NameArray[NameIndex] == "Unreal" )
-				{
-					NetworkStream.NameIndexUnreal = NameIndex;
-				}
-			}
-
-			// Seek to address table and serialize it.
-			ParserStream.Seek( Header.AddressTableOffset, SeekOrigin.Begin );
-			for ( int AddressIndex = 0; AddressIndex < Header.AddressTableEntries; AddressIndex++ )
-			{
-				NetworkStream.AddressArray.Add( BinaryStream.ReadUInt64() );
-			}
-
-			// Seek to beginning of token stream.
-			ParserStream.Seek(TokenStreamOffset,SeekOrigin.Begin);
-
 			// Scratch variables used for building stream. Required as we emit information in reverse
 			// order needed for parsing.
 			var CurrentFrameTokens = new List<TokenBase>();
@@ -123,7 +96,41 @@ namespace NetworkProfiler
 					InMainWindow.UpdateProgress( ( int )( Percent * 100 ) );
 				}
 
-				TokenBase Token = TokenBase.ReadNextToken( BinaryStream, NetworkStream );
+				if ( ParserStream.Position == ParserStream.Length )
+				{
+					// We reached stream early (must not have been finalized properly, but we can still read it)
+					break;
+				}
+
+				TokenBase Token = null;
+
+				try
+				{
+					Token = TokenBase.ReadNextToken( BinaryStream, NetworkStream );
+				}
+				catch ( System.IO.EndOfStreamException )
+				{
+					// We reached stream early (must not have been finalized properly, but we can still read it)
+					break;
+				}
+
+				if ( Token.TokenType == ETokenTypes.NameReference )
+				{
+					NetworkStream.NameArray.Add( ( Token as TokenNameReference ).Name );
+
+					// Find "Unreal" name index used for misc socket parsing optimizations.
+					if ( NetworkStream.NameArray[NetworkStream.NameArray.Count - 1] == "Unreal" )
+					{
+						NetworkStream.NameIndexUnreal = NetworkStream.NameArray.Count - 1;
+					}
+					continue;
+				}
+
+				if ( Token.TokenType == ETokenTypes.ConnectionReference )
+				{
+					NetworkStream.AddressArray.Add( ( Token as TokenConnectionReference ).Address );
+					continue;
+				}
 
 				if ( Token.TokenType == ETokenTypes.ConnectionChange )
 				{
diff --git a/Engine/Source/Programs/NetworkProfiler/NetworkProfiler/Tokens.cs b/Engine/Source/Programs/NetworkProfiler/NetworkProfiler/Tokens.cs
--- a/Engine/Source/Programs/NetworkProfiler/NetworkProfiler/Tokens.cs
+++ b/Engine/Source/Programs/NetworkProfiler/NetworkProfiler/Tokens.cs
@@ -30,6 +30,8 @@ namespace NetworkProfiler
 		EndContentBlock,			// Content block footers
 		WritePropertyHandle,		// Property handles
 		ConnectionChange,			// Connection changed
+		NameReference,				// Reference to name
+		ConnectionReference,		// Reference to connection
 		MaxAndInvalid,				// Invalid token, also used as the max token index
 	}
 
@@ -124,6 +126,12 @@ namespace NetworkProfiler
 				case ETokenTypes.WritePropertyHandle:
 					SerializedToken = new TokenWritePropertyHandle( BinaryStream );
 					break;
+				case ETokenTypes.NameReference:
+					SerializedToken = new TokenNameReference( BinaryStream );
+					break;
+				case ETokenTypes.ConnectionReference:
+					SerializedToken = new TokenConnectionReference( BinaryStream );
+					break;
 				case ETokenTypes.ConnectionChange:
 					SerializedToken = new TokenConnectionChanged( BinaryStream );
 					break;
@@ -693,6 +701,37 @@ namespace NetworkProfiler
 		}
 	}
 
+	/**
+	 * Token for connection reference event
+	 */
+	class TokenNameReference : TokenBase
+	{
+		/** Address of connection */
+		public string Name = null;
+
+		/** Constructor, serializing members from passed in stream. */
+		public TokenNameReference( BinaryReader BinaryStream )
+		{
+			UInt32 Length = BinaryStream.ReadUInt32();
+			Name = new string(BinaryStream.ReadChars((int)Length));
+		}
+	}
+
+	/**
+	 * Token for connection reference event
+	 */
+	class TokenConnectionReference : TokenBase
+	{
+		/** Address of connection */
+		public UInt64 Address;
+
+		/** Constructor, serializing members from passed in stream. */
+		public TokenConnectionReference( BinaryReader BinaryStream )
+		{
+			Address = BinaryStream.ReadUInt64();
+		}
+	}
+
 	/**
 	 * Token for events.
 	 */
diff --git a/Engine/Source/Runtime/CoreUObject/Public/UObject/CoreNet.h b/Engine/Source/Runtime/CoreUObject/Public/UObject/CoreNet.h
--- a/Engine/Source/Runtime/CoreUObject/Public/UObject/CoreNet.h
+++ b/Engine/Source/Runtime/CoreUObject/Public/UObject/CoreNet.h
@@ -217,7 +217,8 @@ enum ELifetimeCondition
 	COND_SimulatedOrPhysics	= 6,		// This property will send to simulated OR bRepPhysics actors
 	COND_InitialOrOwner		= 7,		// This property will send on the initial packet, or to the actors owner
 	COND_Custom				= 8,		// This property has no particular condition, but wants the ability to toggle on/off via SetCustomIsActiveOverride
-	COND_Max				= 9,
+	COND_ReplayOrOwner		= 9,		// This property will only send to the replay connection, or to the actors owner
+	COND_Max				= 10,
 };
 
 
diff --git a/Engine/Source/Runtime/Engine/Classes/Engine/DemoNetDriver.h b/Engine/Source/Runtime/Engine/Classes/Engine/DemoNetDriver.h
--- a/Engine/Source/Runtime/Engine/Classes/Engine/DemoNetDriver.h
+++ b/Engine/Source/Runtime/Engine/Classes/Engine/DemoNetDriver.h
@@ -45,9 +45,6 @@ class ENGINE_API UDemoNetDriver : public UNetDriver
 	/** Current record/playback frame number */
 	int32 DemoFrameNum;
 
-	/** Last time (in real seconds) that we recorded a frame */
-	double LastRecordTime;
-
 	/** Total time of demo in seconds */
 	float DemoTotalTime;
 
diff --git a/Engine/Source/Runtime/Engine/Classes/Engine/EngineTypes.h b/Engine/Source/Runtime/Engine/Classes/Engine/EngineTypes.h
--- a/Engine/Source/Runtime/Engine/Classes/Engine/EngineTypes.h
+++ b/Engine/Source/Runtime/Engine/Classes/Engine/EngineTypes.h
@@ -3459,6 +3459,8 @@ struct FReplicationFlags
 			uint32 bNetSimulated:1;
 			/** True if this is actor's ReplicatedMovement.bRepPhysics flag is true. */
 			uint32 bRepPhysics:1;
+			/** True if this actor is replicating on a replay connection. */
+			uint32 bReplay:1;
 		};
 
 		uint32	Value;
diff --git a/Engine/Source/Runtime/Engine/Classes/Engine/NetConnection.h b/Engine/Source/Runtime/Engine/Classes/Engine/NetConnection.h
--- a/Engine/Source/Runtime/Engine/Classes/Engine/NetConnection.h
+++ b/Engine/Source/Runtime/Engine/Classes/Engine/NetConnection.h
@@ -16,6 +16,7 @@
 
 class FObjectReplicator;
 class FUniqueNetId;
+struct FNetworkObjectInfo;
 
 /*-----------------------------------------------------------------------------
 	Types.
@@ -329,9 +330,6 @@ public:
 	 */
 	TArray<FName> ClientVisibleLevelNames;
 
-	/** @todo document */
-	class TArray<AActor*> OwnedConsiderList;
-
 #if DO_ENABLE_NET_TEST
 	// For development.
 	/** Packet settings for testing lag, net errors, etc */
diff --git a/Engine/Source/Runtime/Engine/Classes/Engine/NetDriver.h b/Engine/Source/Runtime/Engine/Classes/Engine/NetDriver.h
--- a/Engine/Source/Runtime/Engine/Classes/Engine/NetDriver.h
+++ b/Engine/Source/Runtime/Engine/Classes/Engine/NetDriver.h
@@ -5,12 +5,13 @@
 #pragma once
 
 #include "EngineTypes.h"
-
 #include "NetDriver.generated.h"
 
 class FRepChangedPropertyTracker;
 class FRepLayout;
 class FObjectReplicator;
+class FNetworkObjectList;
+struct FNetworkObjectInfo;
 
 //
 // Whether to support net lag and packet loss testing.
@@ -77,18 +78,18 @@ struct ENGINE_API FPacketSimulationSettings
 //
 struct FActorPriority
 {
-	int32					Priority;	// Update priority, higher = more important.
+	int32						Priority;	// Update priority, higher = more important.
 	
-	class AActor*			Actor;		// Actor.	
-	class UActorChannel*	Channel;	// Actor channel.
+	FNetworkObjectInfo*			ActorInfo;	// Actor info.
+	class UActorChannel*		Channel;	// Actor channel.
 
 	struct FActorDestructionInfo *	DestructionInfo;	// Destroy an actor
 
 	FActorPriority() : 
-		Priority(0), Actor(NULL), Channel(NULL), DestructionInfo(NULL)
+		Priority(0), ActorInfo(NULL), Channel(NULL), DestructionInfo(NULL)
 	{}
 
-	FActorPriority(class UNetConnection* InConnection, class UActorChannel* InChannel, class AActor* InActor, const TArray<struct FNetViewer>& Viewers, bool bLowBandwidth);
+	FActorPriority(class UNetConnection* InConnection, class UActorChannel* InChannel, FNetworkObjectInfo* InActorInfo, const TArray<struct FNetViewer>& Viewers, bool bLowBandwidth);
 	FActorPriority(class UNetConnection* InConnection, struct FActorDestructionInfo * DestructInfo, const TArray<struct FNetViewer>& Viewers );
 };
 
@@ -576,6 +577,12 @@ public:
 	/** Returns the actor that corresponds to InGUID, if one can be found. */
 	virtual AActor* GetActorForGUID(FNetworkGUID InGUID) const { return nullptr; }
 
+	/** Returns the object that manages the list of replicated UObjects. */
+	ENGINE_API virtual FNetworkObjectList& GetNetworkObjectList() { return *NetworkObjects; }
+
+	/** Returns the object that manages the list of replicated UObjects. */
+	ENGINE_API virtual const FNetworkObjectList& GetNetworkObjectList() const { return *NetworkObjects; }
+
 protected:
 
 	/** Adds (fully initialized, ready to go) client connection to the ClientConnections list + any other game related setup */
@@ -587,4 +594,18 @@ protected:
 	ENGINE_API void UnregisterTickEvents(class UWorld* InWorld);
 	/** Returns true if this actor is considered to be in a loaded level */
 	ENGINE_API virtual bool IsLevelInitializedForActor( const AActor* InActor, const UNetConnection* InConnection ) const;
+
+#if WITH_SERVER_CODE
+	/**
+	* Helper functions for ServerReplicateActors
+	*/
+	int32 ServerReplicateActors_PrepConnections( const float DeltaSeconds );
+	void ServerReplicateActors_BuildConsiderList( TArray<FNetworkObjectInfo*>& OutConsiderList, const float ServerTickTime );
+	int32 ServerReplicateActors_PrioritizeActors( UNetConnection* Connection, const TArray<FNetViewer>& ConnectionViewers, const TArray<FNetworkObjectInfo*> ConsiderList, const bool bCPUSaturated, FActorPriority*& OutPriorityList, FActorPriority**& OutPriorityActors );
+	int32 ServerReplicateActors_ProcessPrioritizedActors( UNetConnection* Connection, const TArray<FNetViewer>& ConnectionViewers, FActorPriority** PriorityActors, const int32 FinalSortedCount, int32& OutUpdated );
+#endif
+
+private:
+	/** Stores the list of objects to replicate into the replay stream. This should be a TUniquePtr, but it appears the generated.cpp file needs the full definition of the pointed-to type. */
+	TSharedPtr<FNetworkObjectList> NetworkObjects;
 };
diff --git a/Engine/Source/Runtime/Engine/Classes/Engine/NetworkObjectList.h b/Engine/Source/Runtime/Engine/Classes/Engine/NetworkObjectList.h
new file mode 100644
--- /dev/null
+++ b/Engine/Source/Runtime/Engine/Classes/Engine/NetworkObjectList.h
@@ -0,0 +1,95 @@
+// Copyright 1998-2015 Epic Games, Inc. All Rights Reserved.
+
+#pragma once
+
+#include "Set.h"
+
+class AActor;
+class UWorld;
+
+/**
+ * Struct to store an actor pointer and any internal metadata for that actor used
+ * internally by a UNetDriver.
+ */
+struct FNetworkObjectInfo
+{
+	/** Pointer to the replicated actor. */
+	AActor* Actor;
+
+	/** Next time to consider replicating the actor. Based on FPlatformTime::Seconds(). */
+	double NextUpdateTime;
+
+	/** Last absolute time in seconds since actor actually sent something during replication */
+	double LastNetReplicateTime;
+
+	/** Optimal delta between replication updates based on how frequently actor properties are actually changing */
+	float OptimalNetUpdateDelta;
+
+	FNetworkObjectInfo()
+		: Actor(nullptr)
+		, NextUpdateTime(0.0)
+		, LastNetReplicateTime(0.0)
+		, OptimalNetUpdateDelta(0.0f) {}
+
+	FNetworkObjectInfo(AActor* InActor)
+		: Actor(InActor)
+		, NextUpdateTime(0.0)
+		, LastNetReplicateTime(0.0)
+		, OptimalNetUpdateDelta(0.0f) {}
+};
+
+/**
+ * KeyFuncs to allow using the actor pointer as the comparison key in a set.
+ */
+struct FNetworkObjectKeyFuncs : BaseKeyFuncs<TSharedPtr<FNetworkObjectInfo>, AActor*, false>
+{
+	/**
+	 * @return The key used to index the given element.
+	 */
+	static KeyInitType GetSetKey(ElementInitType Element)
+	{
+		return Element.Get()->Actor;
+	}
+
+	/**
+	 * @return True if the keys match.
+	 */
+	static bool Matches(KeyInitType A,KeyInitType B)
+	{
+		return A == B;
+	}
+
+	/** Calculates a hash index for a key. */
+	static uint32 GetKeyHash(KeyInitType Key)
+	{
+		return GetTypeHash(Key);
+	}
+};
+
+/**
+ * Stores the list of replicated actors for a given UNetDriver.
+ */
+class FNetworkObjectList
+{
+public:
+	typedef TSet<TSharedPtr<FNetworkObjectInfo>, FNetworkObjectKeyFuncs> FNetworkObjectSet;
+
+	/**
+	 * Adds replicated actors in World to the internal set of replicated actors.
+	 * Used when a net driver is initialized after some actors may have already
+	 * been added to the world.
+	 *
+	 * @param World The world from which actors are added.
+	 * @param NetDriverName The name of the net driver to which this object list belongs.
+	 */
+	void AddInitialObjects(UWorld* const World, const FName NetDriverName);
+
+	/** Adds Actor to the internal list if its NetDriverName matches, or if we're adding to the demo net driver. */
+	void Add(AActor* const Actor, const FName NetDriverName);
+
+	/** Returns a reference to the set of tracked actors. */
+	FNetworkObjectSet& GetObjects() { return NetworkObjects; }
+
+private:
+	FNetworkObjectSet NetworkObjects;
+};
diff --git a/Engine/Source/Runtime/Engine/Classes/Engine/World.h b/Engine/Source/Runtime/Engine/Classes/Engine/World.h
--- a/Engine/Source/Runtime/Engine/Classes/Engine/World.h
+++ b/Engine/Source/Runtime/Engine/Classes/Engine/World.h
@@ -637,8 +637,6 @@ private:
 	TArray<class ULevel*>						Levels;
 
 public:
-	/** Array of actors that are candidates for sending over the network */
-	TSet<class AActor*>						NetworkActors;
 
 #if WITH_EDITOR
 	/** Hierarchical LOD System. Used when WorldSetting.bEnableHierarchicalLODSystem is true */
diff --git a/Engine/Source/Runtime/Engine/Classes/GameFramework/Actor.h b/Engine/Source/Runtime/Engine/Classes/GameFramework/Actor.h
--- a/Engine/Source/Runtime/Engine/Classes/GameFramework/Actor.h
+++ b/Engine/Source/Runtime/Engine/Classes/GameFramework/Actor.h
@@ -311,13 +311,17 @@ public:
 	int32 NetTag;
 
 	/** Next time this actor will be considered for replication, set by SetNetUpdateTime() */
-	UPROPERTY()
+	UPROPERTY(transient)
 	float NetUpdateTime;
 
 	/** How often (per second) this actor will be considered for replication, used to determine NetUpdateTime */
 	UPROPERTY(Category=Replication, EditDefaultsOnly, BlueprintReadWrite)
 	float NetUpdateFrequency;
 
+	/** Used to determine what rate to throttle down to when replicated properties are changing infrequently */
+	UPROPERTY( Category=Replication, EditDefaultsOnly, BlueprintReadWrite )
+	float MinNetUpdateFrequency;
+
 	/** Priority for this actor when checking for replication in a low bandwidth or saturated situation, higher priority means it is more likely to replicate */
 	UPROPERTY(Category=Replication, EditDefaultsOnly, BlueprintReadWrite)
 	float NetPriority;
diff --git a/Engine/Source/Runtime/Engine/Classes/GameFramework/WorldSettings.h b/Engine/Source/Runtime/Engine/Classes/GameFramework/WorldSettings.h
--- a/Engine/Source/Runtime/Engine/Classes/GameFramework/WorldSettings.h
+++ b/Engine/Source/Runtime/Engine/Classes/GameFramework/WorldSettings.h
@@ -5,6 +5,8 @@
 #include "Sound/AudioVolume.h"
 #include "WorldSettings.generated.h"
 
+class UNetConnection;
+
 UENUM()
 enum EVisibilityAggressiveness
 {
@@ -173,6 +175,9 @@ struct ENGINE_API FNetViewer
 {
 	GENERATED_USTRUCT_BODY()
 
+	UPROPERTY()
+	UNetConnection* Connection;
+
 	/** The "controlling net object" associated with this view (typically player controller) */
 	UPROPERTY()
 	class AActor* InViewer;
@@ -190,7 +195,8 @@ struct ENGINE_API FNetViewer
 	FVector ViewDir;
 
 	FNetViewer()
-		: InViewer(NULL)
+		: Connection(NULL)
+		, InViewer(NULL)
 		, ViewTarget(NULL)
 		, ViewLocation(ForceInit)
 		, ViewDir(ForceInit)
diff --git a/Engine/Source/Runtime/Engine/Private/Actor.cpp b/Engine/Source/Runtime/Engine/Private/Actor.cpp
--- a/Engine/Source/Runtime/Engine/Private/Actor.cpp
+++ b/Engine/Source/Runtime/Engine/Private/Actor.cpp
@@ -72,6 +72,7 @@ void AActor::InitializeDefaults()
 	bReplicates = false;
 	NetPriority = 1.0f;
 	NetUpdateFrequency = 100.0f;
+	MinNetUpdateFrequency = 2.0f;
 	bNetLoadOnClient = true;
 #if WITH_EDITORONLY_DATA
 	bEditable = true;
diff --git a/Engine/Source/Runtime/Engine/Private/DataChannel.cpp b/Engine/Source/Runtime/Engine/Private/DataChannel.cpp
--- a/Engine/Source/Runtime/Engine/Private/DataChannel.cpp
+++ b/Engine/Source/Runtime/Engine/Private/DataChannel.cpp
@@ -12,6 +12,7 @@
 #include "Engine/ActorChannel.h"
 #include "Engine/ControlChannel.h"
 #include "Engine/PackageMapClient.h"
+#include "Engine/DemoNetDriver.h"
 
 DEFINE_LOG_CATEGORY(LogNet);
 DEFINE_LOG_CATEGORY(LogRep);
@@ -2204,6 +2205,9 @@ bool UActorChannel::ReplicateActor()
 	RepFlags.bNetSimulated	= ( Actor->GetRemoteRole() == ROLE_SimulatedProxy );
 	RepFlags.bRepPhysics	= Actor->ReplicatedMovement.bRepPhysics;
 
+	const UWorld* const ActorWorld	= Actor->GetWorld();
+	RepFlags.bReplay				= ActorWorld && (ActorWorld->DemoNetDriver == Connection->GetDriver());
+
 	RepFlags.bNetInitial = RepFlags.bNetInitial || bActorStillInitial; // for replication purposes, bNetInitial stays true until all properties sent
 	bActorMustStayDirty = false;
 	bActorStillInitial = false;
diff --git a/Engine/Source/Runtime/Engine/Private/DataReplication.cpp b/Engine/Source/Runtime/Engine/Private/DataReplication.cpp
--- a/Engine/Source/Runtime/Engine/Private/DataReplication.cpp
+++ b/Engine/Source/Runtime/Engine/Private/DataReplication.cpp
@@ -822,6 +822,7 @@ void FObjectReplicator::ReplicateCustomDeltaProperties( FOutBunch & Bunch, FRepl
 	const bool bIsOwner = RepFlags.bNetOwner ? true : false;
 	const bool bIsSimulated = RepFlags.bNetSimulated ? true : false;
 	const bool bIsPhysics = RepFlags.bRepPhysics ? true : false;
+	const bool bIsReplay = RepFlags.bReplay ? true : false;
 
 	ConditionMap[COND_None] = true;
 	ConditionMap[COND_InitialOnly] = bIsInitial;
@@ -832,6 +833,7 @@ void FObjectReplicator::ReplicateCustomDeltaProperties( FOutBunch & Bunch, FRepl
 	ConditionMap[COND_SimulatedOrPhysics] = bIsSimulated || bIsPhysics;
 	ConditionMap[COND_InitialOrOwner] = bIsInitial || bIsOwner;
 	ConditionMap[COND_Custom] = true;
+	ConditionMap[COND_ReplayOrOwner] = bIsReplay || bIsOwner;
 
 	// Replicate those properties.
 	for ( int32 i = 0; i < LifetimeCustomDeltaProperties.Num(); i++ )
diff --git a/Engine/Source/Runtime/Engine/Private/DemoNetDriver.cpp b/Engine/Source/Runtime/Engine/Private/DemoNetDriver.cpp
--- a/Engine/Source/Runtime/Engine/Private/DemoNetDriver.cpp
+++ b/Engine/Source/Runtime/Engine/Private/DemoNetDriver.cpp
@@ -10,6 +10,7 @@
 #include "Engine/DemoPendingNetGame.h"
 #include "Engine/ActorChannel.h"
 #include "Engine/PackageMapClient.h"
+#include "Engine/NetworkObjectList.h"
 #include "RepLayout.h"
 #include "GameFramework/SpectatorPawn.h"
 #include "Engine/LevelStreamingKismet.h"
@@ -34,6 +35,7 @@ static TAutoConsoleVariable<float> CVarGotoTimeInSeconds( TEXT( "demo.GotoTimeIn
 static TAutoConsoleVariable<int32> CVarDemoFastForwardDestroyTearOffActors( TEXT( "demo.FastForwardDestroyTearOffActors" ), 1, TEXT( "If true, the driver will destroy any torn-off actors immediately while fast-forwarding a replay." ) );
 static TAutoConsoleVariable<int32> CVarDemoFastForwardSkipRepNotifies( TEXT( "demo.FastForwardSkipRepNotifies" ), 1, TEXT( "If true, the driver will optimize fast-forwarding by deferring calls to RepNotify functions until the fast-forward is complete. " ) );
 static TAutoConsoleVariable<int32> CVarDemoQueueCheckpointChannels( TEXT( "demo.QueueCheckpointChannels" ), 1, TEXT( "If true, the driver will put all channels created during checkpoint loading into queuing mode, to amortize the cost of spawning new actors across multiple frames." ) );
+static TAutoConsoleVariable<int32> CVarUseAdaptiveReplayUpdateFrequency( TEXT( "demo.UseAdaptiveReplayUpdateFrequency" ), 0, TEXT( "If 1, NetUpdateFrequency will be calculated based on how often actors actually write something when recording to a replay" ) );
 
 #if !(UE_BUILD_SHIPPING || UE_BUILD_TEST)
 static TAutoConsoleVariable<int32> CVarDemoForceFailure( TEXT( "demo.ForceFailure" ), 0, TEXT( "" ) );
@@ -434,7 +436,6 @@ struct FNetworkDemoMetadataHeader
 void UDemoNetDriver::ResetDemoState()
 {
 	DemoFrameNum		= 0;
-	LastRecordTime		= FPlatformTime::Seconds();
 	LastCheckpointTime	= FPlatformTime::Seconds();
 	DemoTotalTime		= 0;
 	DemoCurrentTime		= 0;
@@ -767,6 +768,8 @@ void UDemoNetDriver::TickDispatch( float DeltaSeconds )
 
 	if ( IsRecording() )
 	{
+		DECLARE_SCOPE_CYCLE_COUNTER(TEXT("Net replay record time"), STAT_ReplayRecordTime, STATGROUP_Net);
+
 		const double StartTime = FPlatformTime::Seconds();
 
 		TickDemoRecord( DeltaSeconds );
@@ -792,7 +795,7 @@ void UDemoNetDriver::TickDispatch( float DeltaSeconds )
 
 			if ( AvgTimeMS > 8.0f )//|| MaxRecordTimeMS > 6.0f )
 			{
-				UE_LOG( LogDemo, Warning, TEXT( "UDemoNetDriver::TickFlush: SLOW FRAME. Avg: %2.2f, Max: %2.2f, Actors: %i" ), AvgTimeMS, MaxRecordTimeMS, World->NetworkActors.Num() );
+				UE_LOG( LogDemo, Warning, TEXT( "UDemoNetDriver::TickFlush: SLOW FRAME. Avg: %2.2f, Max: %2.2f, Actors: %i" ), AvgTimeMS, MaxRecordTimeMS, GetNetworkObjectList().GetObjects().Num() );
 			}
 
 			LastRecordAvgFlush		= EndTime;
@@ -989,7 +992,7 @@ void UDemoNetDriver::StopDemo()
 Demo Recording tick.
 -----------------------------------------------------------------------------*/
 
-static void DemoReplicateActor( AActor* Actor, UNetConnection* Connection, APlayerController* SpectatorController, bool bMustReplicate )
+static bool DemoReplicateActor( AActor* Actor, UNetConnection* Connection, APlayerController* SpectatorController, bool bMustReplicate )
 {
 	// RAII object to swap the Role and RemoteRole of an actor within a scope. Used for recording replays on a client.
 	class FScopedActorRoleSwap
@@ -1018,6 +1021,8 @@ static void DemoReplicateActor( AActor* Actor, UNetConnection* Connection, APlay
 
 	const int32 OriginalOutBunches = Connection->Driver->OutBunches;
 	
+	bool bDidReplicateActor = false;
+
 	if ( Actor != NULL )
 	{
 		// We need to swap roles if:
@@ -1057,7 +1062,7 @@ static void DemoReplicateActor( AActor* Actor, UNetConnection* Connection, APlay
 			if (Channel != NULL && !Channel->Closing)
 			{
 				// Send it out!
-				Channel->ReplicateActor();
+				bDidReplicateActor = Channel->ReplicateActor();
 
 				// Close the channel if this actor shouldn't have one
 				if (!bShouldHaveChannel)
@@ -1072,6 +1077,8 @@ static void DemoReplicateActor( AActor* Actor, UNetConnection* Connection, APlay
 	{
 		UE_LOG( LogDemo, Error, TEXT( "DemoReplicateActor: bMustReplicate is true but nothing was sent: %s" ), Actor ? *Actor->GetName() : TEXT( "NULL" ) );
 	}
+
+	return bDidReplicateActor;
 }
 
 static void SerializeGuidCache( TSharedPtr< class FNetGUIDCache > GuidCache, FArchive* CheckpointArchive )
@@ -1179,13 +1186,11 @@ void UDemoNetDriver::SaveCheckpoint()
 
 	// Replicate *only* the actors that were in the previous frame, we want to be able to re-create up to that point with this single checkpoint
 	// It's important that we don't catch any new actors that the next frame will also catch, that will cause conflict with bOpen (the open will occur twice on the same channel)
-	if ( CheckpointConnection->ActorChannels.Contains( World->GetWorldSettings() ) )
-	{
-		DemoReplicateActor( World->GetWorldSettings(), CheckpointConnection, SpectatorController, true );
-	}
 
-	for ( AActor* Actor : World->NetworkActors )
+	for ( TSharedPtr<FNetworkObjectInfo>& ObjectInfo : GetNetworkObjectList().GetObjects() )
 	{
+		AActor* Actor = ObjectInfo.Get()->Actor;
+
 		if ( CheckpointConnection->ActorChannels.Contains( Actor ) )
 		{
 			Actor->CallPreReplication( this );
@@ -1275,20 +1280,7 @@ void UDemoNetDriver::TickDemoRecord( float DeltaSeconds )
 
 	const double CurrentSeconds = FPlatformTime::Seconds();
 
-	const double RECORD_HZ		= CVarDemoRecordHz.GetValueOnGameThread();
-	const double RECORD_DELAY	= 1.0 / RECORD_HZ;
-
-	if ( CurrentSeconds - LastRecordTime < RECORD_DELAY )
-	{
-		return;		// Not enough real-time has passed to record another frame
-	}
-
-	// Advance by the delay amount to take into account we could be fractionally into the next frame
-	// But don't be more than a fraction of a frame behind either (we don't want to do catch-up frames when there is a long delay)
-	if ( CurrentSeconds - LastRecordTime >= RECORD_DELAY )
-	{
-		LastRecordTime += RECORD_DELAY;
-	}
+	const float RECORD_HZ		= CVarDemoRecordHz.GetValueOnGameThread();
 
 	// Save out a frame
 	DemoFrameNum++;
@@ -1321,28 +1313,89 @@ void UDemoNetDriver::TickDemoRecord( float DeltaSeconds )
 
 	ClientDemoConnection->QueuedDemoPackets.Empty();
 
-	DemoReplicateActor( World->GetWorldSettings(), ClientConnections[0], SpectatorController, false );
+	float ServerTickTime = GEngine->GetMaxTickRate( DeltaSeconds );
+	if ( ServerTickTime == 0.0 )
+	{
+		ServerTickTime = DeltaSeconds;
+	}
+	else
+	{
+		ServerTickTime	= 1.0 / ServerTickTime;
+	}
+
+	const bool bUseAdapativeNetFrequency = CVarUseAdaptiveReplayUpdateFrequency.GetValueOnGameThread() > 0;
 
-	for ( TSet<AActor*>::TIterator ActorIt = World->NetworkActors.CreateIterator(); ActorIt; ++ActorIt)
+	for ( auto ActorIt = GetNetworkObjectList().GetObjects().CreateIterator(); ActorIt; ++ActorIt)
 	{
-		AActor* Actor = *ActorIt;
+		FNetworkObjectInfo* ActorInfo = (*ActorIt).Get();
+		AActor* Actor = ActorInfo->Actor;
 		
-		if ( Actor->IsPendingKill() )
+		// Check NetUpdateFrequency for this actor, but clamp it to RECORD_HZ.
+		if (CurrentSeconds > ActorInfo->NextUpdateTime)
 		{
-			ActorIt.RemoveCurrent();
-			continue;
-		}
+			if ( Actor->IsPendingKill() )
+			{
+				ActorIt.RemoveCurrent();
+				continue;
+			}
 
-		// During client recording, a torn-off actor will already have its remote role set to None, but
-		// we still need to replicate it one more time so that the recorded replay knows it's been torn-off as well.
-		if ( Actor->GetRemoteRole() == ROLE_None && !Actor->bTearOff)
-		{
-			ActorIt.RemoveCurrent();
-			continue;
-		}
+			// During client recording, a torn-off actor will already have its remote role set to None, but
+			// we still need to replicate it one more time so that the recorded replay knows it's been torn-off as well.
+			if ( Actor->GetRemoteRole() == ROLE_None && !Actor->bTearOff)
+			{
+				ActorIt.RemoveCurrent();
+				continue;
+			}
+
+			const float ClampedNetUpdateFrequency = FMath::Clamp(Actor->NetUpdateFrequency, 0.0f, RECORD_HZ);
+			const double NetUpdateDelay = 1.0 / ClampedNetUpdateFrequency;
+
+			// Set defaults if this actor is replicating for first time
+			if ( ActorInfo->LastNetReplicateTime == 0 )
+			{
+				ActorInfo->LastNetReplicateTime		= CurrentSeconds;
+				ActorInfo->OptimalNetUpdateDelta	= NetUpdateDelay;
+			}
+
+			const float LastReplicateDelta = static_cast<float>(CurrentSeconds - ActorInfo->LastNetReplicateTime);
+
+			if ( Actor->MinNetUpdateFrequency == 0.0f )
+			{
+				Actor->MinNetUpdateFrequency = 2.0f;
+			}
 
-		Actor->CallPreReplication( this );
-		DemoReplicateActor( Actor, ClientConnections[0], SpectatorController, false );
+			// Calculate min delta (max rate actor will update), and max delta (slowest rate actor will update)
+			const float MinOptimalDelta = NetUpdateDelay;										// Don't go faster than NetUpdateFrequency
+			const float MaxOptimalDelta = FMath::Max( 1.0f / Actor->MinNetUpdateFrequency, MinOptimalDelta );	// Don't go slower than MinNetUpdateFrequency (or NetUpdateFrequency if it's slower)
+
+			const float ScaleDownStartTime = 2.0f;
+			const float ScaleDownTimeRange = 5.0f;				
+
+			if ( LastReplicateDelta > ScaleDownStartTime )
+			{
+				// Interpolate between MinOptimalDelta/MaxOptimalDelta based on how long it's been since this actor actually sent anything
+				const float Alpha = FMath::Clamp( ( LastReplicateDelta - ScaleDownStartTime ) / ScaleDownTimeRange, 0.0f, 1.0f );
+				ActorInfo->OptimalNetUpdateDelta = FMath::Lerp( MinOptimalDelta, MaxOptimalDelta, Alpha );
+			}
+
+			const double NextUpdateDelta = bUseAdapativeNetFrequency ? ActorInfo->OptimalNetUpdateDelta : NetUpdateDelay;
+
+			// Account for being fractionally into the next frame
+			// But don't be more than a fraction of a frame behind either (we don't want to do catch-up frames when there is a long delay)
+			const double ExtraTime = CurrentSeconds - ActorInfo->NextUpdateTime;
+			const double ClampedExtraTime = FMath::Clamp(ExtraTime, 0.0, NetUpdateDelay);
+
+			// Try to spread the updates across multiple frames to smooth out spikes.
+			ActorInfo->NextUpdateTime = (CurrentSeconds + NextUpdateDelta - ClampedExtraTime + ((FMath::SRand() - 0.5) * ServerTickTime));
+
+			Actor->CallPreReplication( this );
+			if ( DemoReplicateActor( Actor, ClientConnections[0], SpectatorController, false ) )
+			{
+				// Choose an optimal time, we choose 70% of the actual rate to allow frequency to go up if needed
+				ActorInfo->OptimalNetUpdateDelta = FMath::Clamp( LastReplicateDelta * 0.7f, MinOptimalDelta, MaxOptimalDelta );
+				ActorInfo->LastNetReplicateTime = CurrentSeconds;
+			}
+		}
 	}
 
 	// Make sure nothing is left over
diff --git a/Engine/Source/Runtime/Engine/Private/GameInstance.cpp b/Engine/Source/Runtime/Engine/Private/GameInstance.cpp
--- a/Engine/Source/Runtime/Engine/Private/GameInstance.cpp
+++ b/Engine/Source/Runtime/Engine/Private/GameInstance.cpp
@@ -9,6 +9,7 @@
 #include "Engine/GameInstance.h"
 #include "Engine/Engine.h"
 #include "Engine/DemoNetDriver.h"
+#include "Engine/NetworkObjectList.h"
 #include "OnlineSubsystem.h"
 #include "OnlineSessionInterface.h"
 #include "GameFramework/OnlineSession.h"
@@ -789,7 +790,7 @@ void UGameInstance::StartRecordingReplay(const FString& Name, const FString& Fri
 	}
 	else
 	{
-		UE_LOG(LogDemo, Log, TEXT( "Num Network Actors: %i" ), CurrentWorld->NetworkActors.Num() );
+		UE_LOG(LogDemo, Log, TEXT( "Num Network Actors: %i" ), CurrentWorld->DemoNetDriver->GetNetworkObjectList().GetObjects().Num() );
 	}
 }
 
diff --git a/Engine/Source/Runtime/Engine/Private/NetConnection.cpp b/Engine/Source/Runtime/Engine/Private/NetConnection.cpp
--- a/Engine/Source/Runtime/Engine/Private/NetConnection.cpp
+++ b/Engine/Source/Runtime/Engine/Private/NetConnection.cpp
@@ -9,6 +9,7 @@
 #include "Net/NetworkProfiler.h"
 #include "Net/DataReplication.h"
 #include "Engine/ActorChannel.h"
+#include "Engine/NetworkObjectList.h"
 #include "DataChannel.h"
 #include "Engine/PackageMapClient.h"
 #include "GameFramework/GameMode.h"
@@ -377,11 +378,6 @@ void UNetConnection::FinishDestroy()
 void UNetConnection::AddReferencedObjects(UObject* InThis, FReferenceCollector& Collector)
 {
 	UNetConnection* This = CastChecked<UNetConnection>(InThis);
-	// Let GC know that we're potentially referencing some Actor objects
-	for (int32 ActorIndex = 0; ActorIndex < This->OwnedConsiderList.Num(); ++ActorIndex)
-	{
-		Collector.AddReferencedObject( This->OwnedConsiderList[ActorIndex], This );
-	}
 
 	// Let GC know that we're referencing some UChannel objects
 	for( int32 ChIndex=0; ChIndex < MAX_CHANNELS; ++ChIndex )
diff --git a/Engine/Source/Runtime/Engine/Private/NetworkDriver.cpp b/Engine/Source/Runtime/Engine/Private/NetworkDriver.cpp
--- a/Engine/Source/Runtime/Engine/Private/NetworkDriver.cpp
+++ b/Engine/Source/Runtime/Engine/Private/NetworkDriver.cpp
@@ -11,6 +11,7 @@
 #include "Net/RepLayout.h"
 #include "Engine/ActorChannel.h"
 #include "Engine/VoiceChannel.h"
+#include "Engine/NetworkObjectList.h"
 #include "GameFramework/GameNetworkManager.h"
 #include "OnlineSubsystemUtils.h"
 #include "NetworkingDistanceConstants.h"
@@ -110,6 +111,11 @@ static TAutoConsoleVariable<int32> CVarNetDormancyValidate(
 	TEXT("0: Dont validate. 1: Validate on wake up. 2: Validate on each net update"),
 	ECVF_Default);
 
+static TAutoConsoleVariable<int32> CVarUseAdaptiveNetUpdateFrequency(
+	TEXT( "net.UseAdaptiveNetUpdateFrequency" ), 
+	0, 
+	TEXT( "If 1, NetUpdateFrequency will be calculated based on how often actors actually send something when replicating" ) );
+
 /*-----------------------------------------------------------------------------
 	UNetDriver implementation.
 -----------------------------------------------------------------------------*/
@@ -143,6 +149,7 @@ UNetDriver::UNetDriver(const FObjectInitializer& ObjectInitializer)
 ,	NetTag(0)
 ,	DebugRelevantActors(false)
 ,	ProcessQueuedBunchesCurrentFrameMilliseconds(0.0f)
+,	NetworkObjects(new FNetworkObjectList)
 {
 }
 
@@ -1601,17 +1608,6 @@ void UNetDriver::NotifyActorDestroyed( AActor* ThisActor, bool IsSeamlessTravel
 				Connection->DestroyedStartupOrDormantActors.Add(DestructionInfo->NetGUID);
 			}
 
-			// This should be rare, but remove from OwnedConsiderList
-			for (int32 ActorIdx=0; ActorIdx < Connection->OwnedConsiderList.Num(); ++ActorIdx)
-			{
-				if (Connection->OwnedConsiderList[ActorIdx] == ThisActor)
-				{
-					// Swap with last element, shrink list size by 1
-					Connection->OwnedConsiderList.RemoveAtSwap(ActorIdx, 1, false);
-					break;
-				}
-			}
-
 			// Remove it from any dormancy lists				
 			Connection->RecentlyDormantActors.Remove( ThisActor );
 			Connection->DormantReplicatorMap.Remove( ThisActor );
@@ -1871,6 +1867,7 @@ bool FPacketSimulationSettings::ParseSettings(const TCHAR* Cmd)
 #endif
 
 FNetViewer::FNetViewer(UNetConnection* InConnection, float DeltaSeconds) :
+	Connection(InConnection),
 	InViewer(InConnection->PlayerController ? InConnection->PlayerController : InConnection->OwningActor),
 	ViewTarget(InConnection->ViewTarget),
 	ViewLocation(ForceInit),
@@ -1931,20 +1928,20 @@ FNetViewer::FNetViewer(UNetConnection* InConnection, float DeltaSeconds) :
 	}
 }
 
-FActorPriority::FActorPriority(UNetConnection* InConnection, UActorChannel* InChannel, AActor* InActor, const TArray<struct FNetViewer>& Viewers, bool bLowBandwidth)
-	: Actor(InActor), Channel(InChannel), DestructionInfo(NULL)
+FActorPriority::FActorPriority(UNetConnection* InConnection, UActorChannel* InChannel, FNetworkObjectInfo* InActorInfo, const TArray<struct FNetViewer>& Viewers, bool bLowBandwidth)
+	: ActorInfo(InActorInfo), Channel(InChannel), DestructionInfo(NULL)
 {	
 	float Time  = Channel ? (InConnection->Driver->Time - Channel->LastUpdateTime) : InConnection->Driver->SpawnPrioritySeconds;
 	// take the highest priority of the viewers on this connection
 	Priority = 0;
 	for (int32 i = 0; i < Viewers.Num(); i++)
 	{
-		Priority = FMath::Max<int32>(Priority, FMath::RoundToInt(65536.0f * Actor->GetNetPriority(Viewers[i].ViewLocation, Viewers[i].ViewDir, Viewers[i].InViewer, Viewers[i].ViewTarget, InChannel, Time, bLowBandwidth)));
+		Priority = FMath::Max<int32>(Priority, FMath::RoundToInt(65536.0f * ActorInfo->Actor->GetNetPriority(Viewers[i].ViewLocation, Viewers[i].ViewDir, Viewers[i].InViewer, Viewers[i].ViewTarget, InChannel, Time, bLowBandwidth)));
 	}
 }
 
 FActorPriority::FActorPriority(class UNetConnection* InConnection, struct FActorDestructionInfo * Info, const TArray<struct FNetViewer>& Viewers )
-	: Actor(NULL), Channel(NULL), DestructionInfo(Info)
+	: ActorInfo(NULL), Channel(NULL), DestructionInfo(Info)
 {
 	
 	Priority = 0;
@@ -1971,37 +1968,26 @@ FActorPriority::FActorPriority(class UNetConnection* InConnection, struct FActor
 	}
 }
 
-int32 UNetDriver::ServerReplicateActors(float DeltaSeconds)
-{
-	SCOPE_CYCLE_COUNTER(STAT_NetServerRepActorsTime);
-
 #if WITH_SERVER_CODE
-	if ( ClientConnections.Num() == 0 )
-	{
-		return 0;
-	}
-
-	int32 Updated = 0;
-
-	// Bump the ReplicationFrame value to invalidate any properties marked as "unchanged" for this frame.
-	ReplicationFrame++;
-
+int32 UNetDriver::ServerReplicateActors_PrepConnections( const float DeltaSeconds )
+{
 	int32 NumClientsToTick = ClientConnections.Num();
 
 	// by default only throttle update for listen servers unless specified on the commandline
-	static bool bForceClientTickingThrottle = FParse::Param(FCommandLine::Get(),TEXT("limitclientticks"));
-	if (bForceClientTickingThrottle || GetNetMode() == NM_ListenServer)
+	static bool bForceClientTickingThrottle = FParse::Param( FCommandLine::Get(), TEXT( "limitclientticks" ) );
+	if ( bForceClientTickingThrottle || GetNetMode() == NM_ListenServer )
 	{
 		// determine how many clients to tick this frame based on GEngine->NetTickRate (always tick at least one client), double for lan play
+		// FIXME: DeltaTimeOverflow is a static, and will conflict with other running net drivers, we investigate storing it on the driver itself!
 		static float DeltaTimeOverflow = 0.f;
 		// updates are doubled for lan play
-		static bool LanPlay = FParse::Param(FCommandLine::Get(),TEXT("lanplay"));
+		static bool LanPlay = FParse::Param( FCommandLine::Get(), TEXT( "lanplay" ) );
 		//@todo - ideally we wouldn't want to tick more clients with a higher deltatime as that's not going to be good for performance and probably saturate bandwidth in hitchy situations, maybe 
 		// come up with a solution that is greedier with higher framerates, but still won't risk saturating server upstream bandwidth
-		float ClientUpdatesThisFrame = GEngine->NetClientTicksPerSecond * (DeltaSeconds + DeltaTimeOverflow) * (LanPlay ? 2.f : 1.f);
-		NumClientsToTick = FMath::Min<int32>(NumClientsToTick,FMath::TruncToInt(ClientUpdatesThisFrame));
+		float ClientUpdatesThisFrame = GEngine->NetClientTicksPerSecond * ( DeltaSeconds + DeltaTimeOverflow ) * ( LanPlay ? 2.f : 1.f );
+		NumClientsToTick = FMath::Min<int32>( NumClientsToTick, FMath::TruncToInt( ClientUpdatesThisFrame ) );
 		//UE_LOG(LogNet, Log, TEXT("%2.3f: Ticking %d clients this frame, %2.3f/%2.4f"),GetWorld()->GetTimeSeconds(),NumClientsToTick,DeltaSeconds,ClientUpdatesThisFrame);
-		if (NumClientsToTick == 0)
+		if ( NumClientsToTick == 0 )
 		{
 			// if no clients are ticked this frame accumulate the time elapsed for the next frame
 			DeltaTimeOverflow += DeltaSeconds;
@@ -2010,48 +1996,35 @@ int32 UNetDriver::ServerReplicateActors(float DeltaSeconds)
 		DeltaTimeOverflow = 0.f;
 	}
 
-	bool bNetRelevantActorCount = false;
-	int32 NetRelevantActorCount = 0;
-
-	check( World );
+	bool bFoundReadyConnection = false;
 
-	FMemMark Mark(FMemStack::Get());
-	// initialize connections
-	bool bFoundReadyConnection = false; // whether we have at least one connection ready for property replication
-	for( int32 ConnIdx = 0; ConnIdx < ClientConnections.Num(); ConnIdx++ )
+	for ( int32 ConnIdx = 0; ConnIdx < ClientConnections.Num(); ConnIdx++ )
 	{
 		UNetConnection* Connection = ClientConnections[ConnIdx];
-		check(Connection);
-		check(Connection->State == USOCK_Pending || Connection->State == USOCK_Open || Connection->State == USOCK_Closed);
-		checkSlow(Connection->GetUChildConnection() == NULL);
+		check( Connection );
+		check( Connection->State == USOCK_Pending || Connection->State == USOCK_Open || Connection->State == USOCK_Closed );
+		checkSlow( Connection->GetUChildConnection() == NULL );
 
 		// Handle not ready channels.
 		//@note: we cannot add Connection->IsNetReady(0) here to check for saturation, as if that's the case we still want to figure out the list of relevant actors
 		//			to reset their NetUpdateTime so that they will get sent as soon as the connection is no longer saturated
 		AActor* OwningActor = Connection->OwningActor;
-		if (OwningActor != NULL && Connection->State == USOCK_Open && (Connection->Driver->Time - Connection->LastReceiveTime < 1.5f))
+		if ( OwningActor != NULL && Connection->State == USOCK_Open && ( Connection->Driver->Time - Connection->LastReceiveTime < 1.5f ) )
 		{
 			check( World == OwningActor->GetWorld() );
-			if( !bNetRelevantActorCount )
-			{
-				NetRelevantActorCount = World->GetNetRelevantActorCount() + 2;
-				bNetRelevantActorCount = true;
-			}
+
 			bFoundReadyConnection = true;
-			
+
 			// the view target is what the player controller is looking at OR the owning actor itself when using beacons
 			Connection->ViewTarget = Connection->PlayerController ? Connection->PlayerController->GetViewTarget() : OwningActor;
-			//@todo - eliminate this mallocs if the connection isn't going to actually be updated this frame (currently needed to verify owner relevancy below)
-			Connection->OwnedConsiderList.Empty(NetRelevantActorCount);
 
-			for (int32 ChildIdx = 0; ChildIdx < Connection->Children.Num(); ChildIdx++)
+			for ( int32 ChildIdx = 0; ChildIdx < Connection->Children.Num(); ChildIdx++ )
 			{
 				UNetConnection *Child = Connection->Children[ChildIdx];
 				APlayerController* ChildPlayerController = Child->PlayerController;
-				if (ChildPlayerController != NULL)
+				if ( ChildPlayerController != NULL )
 				{
 					Child->ViewTarget = ChildPlayerController->GetViewTarget();
-					Child->OwnedConsiderList.Empty(NetRelevantActorCount);
 				}
 				else
 				{
@@ -2062,648 +2035,692 @@ int32 UNetDriver::ServerReplicateActors(float DeltaSeconds)
 		else
 		{
 			Connection->ViewTarget = NULL;
-			for (int32 ChildIdx = 0; ChildIdx < Connection->Children.Num(); ChildIdx++)
+			for ( int32 ChildIdx = 0; ChildIdx < Connection->Children.Num(); ChildIdx++ )
 			{
 				Connection->Children[ChildIdx]->ViewTarget = NULL;
 			}
 		}
 	}
 
-	// early out if no connections are ready to receive data
-	if (!bFoundReadyConnection)
-	{
-		return 0;
-	}
+	return bFoundReadyConnection ? NumClientsToTick : 0;
+}
 
-	// The world should be set if bFoundReadyConnection is true
-	check(World);
+void UNetDriver::ServerReplicateActors_BuildConsiderList( TArray<FNetworkObjectInfo*>& OutConsiderList, const float ServerTickTime )
+{
+	SCOPE_CYCLE_COUNTER( STAT_NetConsiderActorsTime );
+
+	UE_LOG( LogNetTraffic, Log, TEXT( "ServerReplicateActors_BuildConsiderList, Building ConsiderList %4.2f" ), World->GetTimeSeconds() );
 
-	// make list of actors to consider to relevancy checking and replication
-	TArray<AActor*> ConsiderList;
-	ConsiderList.Reserve(NetRelevantActorCount);
+	SET_DWORD_STAT( STAT_NumNetActors, GetNetworkObjectList().GetObjects().Num() );
 
 	int32 NumInitiallyDormant = 0;
 
-	// Add WorldSettings to consider list if we have one
-	AWorldSettings* WorldSettings = World->GetWorldSettings();
-	if( WorldSettings )
+	const bool bUseAdapativeNetFrequency = CVarUseAdaptiveNetUpdateFrequency.GetValueOnGameThread() > 0;
+
+	for ( auto ActorIt = GetNetworkObjectList().GetObjects().CreateIterator(); ActorIt; ++ActorIt )
 	{
-		if (WorldSettings->GetRemoteRole() != ROLE_None && WorldSettings->NetDriverName == NetDriverName)
+		FNetworkObjectInfo* ActorInfo = ( *ActorIt ).Get();
+		AActor* Actor = ActorInfo->Actor;
+
+		if ( Actor->IsPendingKill() )
+		{
+			ActorIt.RemoveCurrent();
+			continue;
+		}
+
+		if ( Actor->GetRemoteRole() == ROLE_None )
+		{
+			ActorIt.RemoveCurrent();
+			continue;
+		}
+
+		// This actor may belong to a different net driver, make sure this is the correct one
+		// (this can happen when using beacon net drivers for example)
+		if ( Actor->NetDriverName != NetDriverName )
+		{
+			continue;
+		}
+
+		// Don't send actors that may still be streaming in
+		ULevel* Level = Actor->GetLevel();
+		if ( Level->HasVisibilityRequestPending() || Level->bIsAssociatingLevel )
+		{
+			continue;
+		}
+
+		if ( Actor->NetDormancy == DORM_Initial && Actor->IsNetStartupActor() )
+		{
+			// This stat isn't that useful in its current form when using NetworkActors list
+			// We'll want to track initially dormant actors some other way to track them with stats
+			SCOPE_CYCLE_COUNTER( STAT_NetInitialDormantCheckTime );
+			NumInitiallyDormant++;
+			ActorIt.RemoveCurrent();
+			//UE_LOG(LogNetTraffic, Log, TEXT("Skipping Actor %s - its initially dormant!"), *Actor->GetName() );
+			continue;
+		}
+
+		check( Actor->NeedsLoadForClient() );			// We have no business sending this unless the client can load
+		check( World == Actor->GetWorld() );
+
+		// See if it's time for this actor to update replicated properties
+		if ( ( Actor->bPendingNetUpdate || World->TimeSeconds > Actor->NetUpdateTime ) )
 		{
+			// Set defaults if this actor is replicating for first time
+			if ( ActorInfo->LastNetReplicateTime == 0 )
+			{
+				ActorInfo->LastNetReplicateTime		= World->TimeSeconds;
+				ActorInfo->OptimalNetUpdateDelta	= 1.0f / Actor->NetUpdateFrequency;
+			}
+
+			const float ScaleDownStartTime = 2.0f;
+			const float ScaleDownTimeRange = 5.0f;
+
+			const float LastReplicateDelta = World->TimeSeconds - ActorInfo->LastNetReplicateTime;
+
+			if ( LastReplicateDelta > ScaleDownStartTime )
+			{
+				if ( Actor->MinNetUpdateFrequency == 0.0f )
+				{
+					Actor->MinNetUpdateFrequency = 2.0f;
+				}
+
+				// Calculate min delta (max rate actor will upate), and max delta (slowest rate actor will update)
+				const float MinOptimalDelta = 1.0f / Actor->NetUpdateFrequency;										// Don't go faster than NetUpdateFrequency
+				const float MaxOptimalDelta = FMath::Max( 1.0f / Actor->MinNetUpdateFrequency, MinOptimalDelta );	// Don't go slower than MinNetUpdateFrequency (or NetUpdateFrequency if it's slower)
+
+				// Interpolate between MinOptimalDelta/MaxOptimalDelta based on how long it's been since this actor actually sent anything
+				const float Alpha = FMath::Clamp( ( LastReplicateDelta - ScaleDownStartTime ) / ScaleDownTimeRange, 0.0f, 1.0f );
+				ActorInfo->OptimalNetUpdateDelta = FMath::Lerp( MinOptimalDelta, MaxOptimalDelta, Alpha );
+			}
+
+			// Setup Actor->NetUpdateTime, which will be the next time this actor will replicate properties to connections
+			// NOTE - We don't do this if bPendingNetUpdate is true, since this means we're forcing an update due to at least one connection
+			//	that wasn't to replicate previously (due to saturation, etc)
+			// NOTE - This also means all other connections will force an update (even if they just updated, we should look into this)
+			if ( !Actor->bPendingNetUpdate )
+			{
+				UE_LOG( LogNetTraffic, Log, TEXT( "actor %s requesting new net update, time: %2.3f" ), *Actor->GetName(), World->TimeSeconds );
+
+				const float NextUpdateDelta = bUseAdapativeNetFrequency ? ActorInfo->OptimalNetUpdateDelta : 1.0f / Actor->NetUpdateFrequency;
+
+				// then set the next update time
+				Actor->NetUpdateTime = World->TimeSeconds + FMath::SRand() * ServerTickTime + NextUpdateDelta;
+
+				// and mark when the actor first requested an update
+				//@note: using Time because it's compared against UActorChannel.LastUpdateTime which also uses that value
+				Actor->LastNetUpdateTime = Time;
+			}
+
+			// and clear the pending update flag assuming all clients will be able to consider it
+			Actor->bPendingNetUpdate = false;
+
+			// add it to the list to consider below
 			// For performance reasons, make sure we don't resize the array. It should already be appropriately sized above!
-			ensure(ConsiderList.Num() < ConsiderList.Max());
-			ConsiderList.Add(WorldSettings);
+			ensure( OutConsiderList.Num() < OutConsiderList.Max() );
+			OutConsiderList.Add( ActorInfo );
+
+			// Call PreReplication on all actors that will be considered
+			Actor->CallPreReplication( this );
 		}
 	}
 
-	bool bCPUSaturated		= false;
-	float ServerTickTime	= GEngine->GetMaxTickRate( DeltaSeconds );
-	if ( ServerTickTime == 0.f )
+	// Update stats
+	SET_DWORD_STAT( STAT_NumInitiallyDormantActors, NumInitiallyDormant );
+	SET_DWORD_STAT( STAT_NumConsideredActors, OutConsiderList.Num() );
+}
+
+// Returns true if this actor should replicate to *any* of the passed in connections
+static FORCEINLINE_DEBUGGABLE bool IsActorRelevantToConnection( const AActor* Actor, const TArray<FNetViewer>& ConnectionViewers )
+{
+	for ( int32 viewerIdx = 0; viewerIdx < ConnectionViewers.Num(); viewerIdx++ )
 	{
-		ServerTickTime = DeltaSeconds;
+		if ( Actor->IsNetRelevantFor( ConnectionViewers[viewerIdx].InViewer, ConnectionViewers[viewerIdx].ViewTarget, ConnectionViewers[viewerIdx].ViewLocation ) )
+		{
+			return true;
+		}
 	}
-	else
+
+	return false;
+}
+
+// Returns true if this actor is owned by, and should replicate to *any* of the passed in connections
+static FORCEINLINE_DEBUGGABLE UNetConnection* IsActorOwnedByAndRelevantToConnection( const AActor* Actor, const TArray<FNetViewer>& ConnectionViewers, bool& bOutHasNullViewTarget )
+{
+	const AActor* ActorOwner = Actor->GetNetOwner();
+
+	bOutHasNullViewTarget = false;
+
+	for ( int i = 0; i < ConnectionViewers.Num(); i++ )
 	{
-		ServerTickTime	= 1.f/ServerTickTime;
-		bCPUSaturated	= DeltaSeconds > 1.2f * ServerTickTime;
+		UNetConnection* ViewerConnection = ConnectionViewers[i].Connection;
+
+		if ( ViewerConnection->ViewTarget == nullptr )
+		{
+			bOutHasNullViewTarget = true;
+		}
+
+		if ( ActorOwner == ViewerConnection->PlayerController ||
+			 ( ViewerConnection->PlayerController && ActorOwner == ViewerConnection->PlayerController->GetPawn() ) ||
+			 ViewerConnection->ViewTarget->IsRelevancyOwnerFor( Actor, ActorOwner, ViewerConnection->OwningActor ) )
+		{
+			return ViewerConnection;
+		}
+	}
+
+	return nullptr;
+}
+
+// Returns true if this actor is considered dormant (and all properties caught up) to the current connection
+static FORCEINLINE_DEBUGGABLE bool IsActorDormant( const AActor* Actor, const UNetConnection* Connection )
+{
+	// If actor is already dormant on this channel, then skip replication entirely
+	if ( Connection->DormantActors.Contains( Actor ) )
+	{
+		// net.DormancyValidate can be set to 2 to validate dormant actor properties on every replicate
+		// (this could be moved to be done every tick instead of every net update if necessary, but seems excessive)
+		if ( CVarNetDormancyValidate.GetValueOnGameThread() == 2 )
+		{
+			const TSharedRef< FObjectReplicator > * Replicator = Connection->DormantReplicatorMap.Find( Actor );
+
+			if ( Replicator != NULL )
+			{
+				Replicator->Get().ValidateAgainstState( Actor );
+			}
+		}
+
+		return true;
+	}
+
+	return false;
+}
+
+// Returns true if this actor wants to go dormant for a particular connection
+static FORCEINLINE_DEBUGGABLE bool ShouldActorGoDormant( AActor* Actor, const TArray<FNetViewer>& ConnectionViewers, UActorChannel* Channel, const float Time, const bool bLowNetBandwidth )
+{
+	if ( Actor->NetDormancy <= DORM_Awake || !Channel || Channel->bPendingDormancy || Channel->Dormant )
+	{
+		// Either shouldn't go dormant, or is already dormant
+		return false;
+	}
+
+	if ( Actor->NetDormancy == DORM_DormantPartial )
+	{
+		for ( int32 viewerIdx = 0; viewerIdx < ConnectionViewers.Num(); viewerIdx++ )
+		{
+			if ( !Actor->GetNetDormancy( ConnectionViewers[viewerIdx].ViewLocation, ConnectionViewers[viewerIdx].ViewDir, ConnectionViewers[viewerIdx].InViewer, ConnectionViewers[viewerIdx].ViewTarget, Channel, Time, bLowNetBandwidth ) )
+			{
+				return false;
+			}
+		}
 	}
 
+	return true;
+}
+
+int32 UNetDriver::ServerReplicateActors_PrioritizeActors( UNetConnection* Connection, const TArray<FNetViewer>& ConnectionViewers, const TArray<FNetworkObjectInfo*> ConsiderList, const bool bCPUSaturated, FActorPriority*& OutPriorityList, FActorPriority**& OutPriorityActors )
+{
+	SCOPE_CYCLE_COUNTER( STAT_NetPrioritizeActorsTime );
+
+	// Get list of visible/relevant actors.
+
+	NetTag++;
+	Connection->TickCount++;
+
+	// Set up to skip all sent temporary actors
+	for ( int32 j=0; j < Connection->SentTemporaries.Num(); j++ )
+	{
+		Connection->SentTemporaries[j]->NetTag = NetTag;
+	}
+
+	// Make list of all actors to consider.
+	check( World == Connection->OwningActor->GetWorld() );
+
+	const int32 MaxSortedActors = ConsiderList.Num() + DestroyedStartupOrDormantActors.Num();
+	OutPriorityList = new( FMemStack::Get(), MaxSortedActors )FActorPriority;
+	OutPriorityActors = new( FMemStack::Get(), MaxSortedActors )FActorPriority*;
+
+	// determine whether we should priority sort the list of relevant actors based on the saturation/bandwidth of the current connection
+	//@note - if the server is currently CPU saturated then do not sort until framerate improves
+	check( World == Connection->ViewTarget->GetWorld() );
+	AGameMode const* const GameMode = World->GetAuthGameMode();
+	const bool bLowNetBandwidth = !bCPUSaturated && ( Connection->CurrentNetSpeed / float( GameMode->NumPlayers + GameMode->NumBots ) < 500.f );
+
+	int32 FinalSortedCount = 0;
+
+	for ( FNetworkObjectInfo* ActorInfo : ConsiderList )
 	{
-		SCOPE_CYCLE_COUNTER(STAT_NetConsiderActorsTime);
-		UE_LOG(LogNetTraffic, Log, TEXT("UWorld::ServerTickClients, Building ConsiderList %4.2f"), World->GetTimeSeconds());
+		AActor* Actor = ActorInfo->Actor;
+
+		UActorChannel* Channel = Connection->ActorChannels.FindRef( Actor );
 
-		SET_DWORD_STAT( STAT_NumNetActors, World->NetworkActors.Num() );
+		UNetConnection* PriorityConnection = Connection;
 
-		for ( TSet<AActor*>::TIterator ActorIt = World->NetworkActors.CreateIterator(); ActorIt; ++ActorIt)
+		if ( Actor->bOnlyRelevantToOwner )
 		{
-			AActor* Actor = *ActorIt;
+			// This actor should be owned by a particular connection, see if that connection is the one passed in
+			bool bHasNullViewTarget = false;
 
-			if (Actor->IsPendingKill() )
+			PriorityConnection = IsActorOwnedByAndRelevantToConnection( Actor, ConnectionViewers, bHasNullViewTarget );
+
+			if ( PriorityConnection == nullptr )
 			{
-				ActorIt.RemoveCurrent();
+				// Not owned by this connection, if we have a channel, close it, and continue
+				// NOTE - We won't close the channel if any connection has a NULL view target.
+				//	This is to give all connections a chance to own it
+				if ( !bHasNullViewTarget && Channel != NULL && Time - Channel->RelevantTime >= RelevantTimeout )
+				{
+					Channel->Close();
+				}
+
+				// This connection doesn't own this actor
 				continue;
 			}
-
-			if (Actor->GetRemoteRole()==ROLE_None)
+		}
+		else if ( CVarSetNetDormancyEnabled.GetValueOnGameThread() != 0 )
+		{
+			// Skip Actor if dormant
+			if ( IsActorDormant( Actor, Connection ) )
 			{
-				ActorIt.RemoveCurrent();
 				continue;
 			}
 
-			// This actor may belong to a different net driver, make sure this is the correct one
-			// (this can happen when using beacon net drivers for example)
-			if ( Actor->NetDriverName != NetDriverName )
+			// See of actor wants to try and go dormant
+			if ( ShouldActorGoDormant( Actor, ConnectionViewers, Channel, Time, bLowNetBandwidth ) )
 			{
-				continue;
+				// Channel is marked to go dormant now once all properties have been replicated (but is not dormant yet)
+				Channel->StartBecomingDormant();
 			}
+		}
 
-			// Don't send actors that may still be streaming in
-			ULevel* Level = Actor->GetLevel();
-			if ( Level->HasVisibilityRequestPending() || Level->bIsAssociatingLevel )
+		// Skip actor if not relevant and theres no channel already.
+		// Historically Relevancy checks were deferred until after prioritization because they were expensive (line traces).
+		// Relevancy is now cheap and we are dealing with larger lists of considered actors, so we want to keep the list of
+		// prioritized actors low.
+		if ( !Channel )
+		{
+			if ( !IsLevelInitializedForActor( Actor, Connection ) )
 			{
+				// If the level this actor belongs to isn't loaded on client, don't bother sending
 				continue;
 			}
 
-			if ( Actor->NetDormancy == DORM_Initial && Actor->IsNetStartupActor() )
+			if ( !IsActorRelevantToConnection( Actor, ConnectionViewers ) )
 			{
-				// This stat isn't that useful in its current form when using NetworkActors list
-				// We'll want to track initially dormant actors some other way to track them with stats
-				SCOPE_CYCLE_COUNTER(STAT_NetInitialDormantCheckTime);		
-				NumInitiallyDormant++;
-				ActorIt.RemoveCurrent();
-				//UE_LOG(LogNetTraffic, Log, TEXT("Skipping Actor %s - its initially dormant!"), *Actor->GetName() );
+				// If not relevant (and we don't have a channel), skip
 				continue;
 			}
+		}
 
-			check( Actor->NeedsLoadForClient() );			// We have no business sending this unless the client can load
-
-			check (World == Actor->GetWorld());
-			if( (Actor->bPendingNetUpdate || World->TimeSeconds > Actor->NetUpdateTime) ) 
-			{
-				// if this actor isn't being considered due to a previous ServerTickClients() call where not all clients were able to replicate the actor
-				if (!Actor->bPendingNetUpdate)
-				{
-					UE_LOG(LogNetTraffic, Log, TEXT("actor %s requesting new net update, time: %2.3f"), *Actor->GetName(), World->TimeSeconds);
-					// then set the next update time
-					Actor->NetUpdateTime = World->TimeSeconds + FMath::SRand() * ServerTickTime + 1.f/Actor->NetUpdateFrequency; // FIXME - cache 1/netupdatefreq
-					// and mark when the actor first requested an update
-					//@note: using Time because it's compared against UActorChannel.LastUpdateTime which also uses that value
-					Actor->LastNetUpdateTime = Time;
-				}
-				/*
-				else
-				{
-					UE_LOG(LogNetTraffic, Log, TEXT("actor %s still pending update, time since update request: %2.3f"),*Actor->GetName(),World->TimeSeconds-Actor->LastNetUpdateTime);
-				}
-				*/
-				// and clear the pending update flag assuming all clients will be able to consider it
-				Actor->bPendingNetUpdate = false;
+		// Actor is relevant to this connection, add it to the list
+		// NOTE - We use NetTag to make sure SentTemporaries didn't already mark this actor to be skipped
+		if ( Actor->NetTag != NetTag )
+		{
+			UE_LOG( LogNetTraffic, Log, TEXT( "Consider %s alwaysrelevant %d frequency %f " ), *Actor->GetName(), Actor->bAlwaysRelevant, Actor->NetUpdateFrequency );
 
-				bool bWasConsidered = false;
+			Actor->NetTag = NetTag;
 
-				// if this actor relevant to any client
-				if ( !Actor->bOnlyRelevantToOwner ) 
-				{
-					// add it to the list to consider below
-					// For performance reasons, make sure we don't resize the array. It should already be appropriately sized above!
-					ensure(ConsiderList.Num() < ConsiderList.Max());
-					ConsiderList.Add(Actor);
+			OutPriorityList[FinalSortedCount]	= FActorPriority( PriorityConnection, Channel, ActorInfo, ConnectionViewers, bLowNetBandwidth );
+			OutPriorityActors[FinalSortedCount] = OutPriorityList + FinalSortedCount;
 
-					bWasConsidered = true;
-				}
-				else
-				{
-					const AActor* ActorOwner = Actor->GetNetOwner();
-					if ( ActorOwner )
-					{
-						// iterate through each connection (and child connections) looking for an owner for this actor
-						for ( int32 ConnIdx = 0; ConnIdx < ClientConnections.Num(); ConnIdx++ )
-						{
-							UNetConnection* ClientConnection = ClientConnections[ConnIdx];
-							UNetConnection* Connection = ClientConnection;
-							int32 ChildIndex = 0;
-							bool bCloseChannel = true;
-							while (Connection != NULL)
-							{
-								if (Connection->ViewTarget != NULL)
-								{
-									if (ActorOwner == Connection->PlayerController || 
-										(Connection->PlayerController && ActorOwner == Connection->PlayerController->GetPawn()) ||
-										Connection->ViewTarget->IsRelevancyOwnerFor(Actor, ActorOwner, Connection->OwningActor))
-									{
-										// For performance reasons, make sure we don't resize the array. It should already be appropriately sized above!
-										ensure(Connection->OwnedConsiderList.Num() < Connection->OwnedConsiderList.Max());
-										Connection->OwnedConsiderList.Add(Actor);
-										bCloseChannel = false;
-										
-										bWasConsidered = true;
-									}
-								}
-								else
-								{
-									// don't ever close the channel if one or more child connections don't have a Viewer to check relevancy with
-									bCloseChannel = false;
-								}
-								// iterate to the next child connection if available
-								Connection = (ChildIndex < ClientConnection->Children.Num()) ? ClientConnection->Children[ChildIndex++] : NULL;
-							}
-							// if it's not being considered, but there is an open channel for this actor already, close it
-							if (bCloseChannel)
-							{
-								UActorChannel* Channel = ClientConnection->ActorChannels.FindRef(Actor);
-								if (Channel != NULL && Time - Channel->RelevantTime >= RelevantTimeout)
-								{
-									Channel->Close();
-								}
-							}
-						}
-					}
-				}
+			FinalSortedCount++;
 
-				if ( bWasConsidered )
-				{
-					Actor->CallPreReplication( this );
-				}
-			}
-			/*
-			else
+			if ( DebugRelevantActors )
 			{
-				if( Actor->GetAPawn() && (Actor->RemoteRole!=ROLE_None) && (World->GetTimeSeconds() <= Actor->NetUpdateTime) ) 
-				{
-					UE_LOG(LogNetTraffic, Log, TEXT("%s skipped in considerlist because of NetUpdateTime %f"), *Actor->GetName(), (World->GetTimeSeconds() - Actor->NetUpdateTime) );
-				}
+				LastPrioritizedActors.Add( Actor );
 			}
-			*/
 		}
 	}
 
-	SET_DWORD_STAT(STAT_NumInitiallyDormantActors,NumInitiallyDormant);
-	SET_DWORD_STAT(STAT_NumConsideredActors,ConsiderList.Num());
+	int32 DeletedCount = 0;
 
-	for( int32 i=0; i < ClientConnections.Num(); i++ )
+	// Add in deleted actors
+	for ( auto It = Connection->DestroyedStartupOrDormantActors.CreateIterator(); It; ++It )
 	{
-		UNetConnection* Connection = ClientConnections[i];
-		check(Connection);
-		int32 ActorUpdatesThisConnection = 0;
-		int32 ActorUpdatesThisConnectionSent = 0;
+		FActorDestructionInfo &DInfo = DestroyedStartupOrDormantActors.FindChecked( *It );
+		OutPriorityList[FinalSortedCount]	= FActorPriority( Connection, &DInfo, ConnectionViewers );
+		OutPriorityActors[FinalSortedCount] = OutPriorityList + FinalSortedCount;
+		FinalSortedCount++;
+		DeletedCount++;
+	}
 
-		// if this client shouldn't be ticked this frame
-		if (i >= NumClientsToTick)
+	// Sort by priority
+	struct FCompareFActorPriority
+	{
+		FORCEINLINE bool operator()( const FActorPriority& A, const FActorPriority& B ) const
 		{
-			//UE_LOG(LogNet, Log, TEXT("skipping update to %s"),*Connection->GetName());
-			// then mark each considered actor as bPendingNetUpdate so that they will be considered again the next frame when the connection is actually ticked
-			for (int32 ConsiderIdx = 0; ConsiderIdx < ConsiderList.Num(); ConsiderIdx++)
+			return B.Priority < A.Priority;
+		}
+	};
+
+	Sort( OutPriorityActors, FinalSortedCount, FCompareFActorPriority() );
+
+	UE_LOG( LogNetTraffic, Log, TEXT( "ServerReplicateActors_PrioritizeActors: Potential %04i ConsiderList %03i FinalSortedCount %03i" ), MaxSortedActors, ConsiderList.Num(), FinalSortedCount );
+	
+	// Setup stats
+	SET_DWORD_STAT( STAT_PrioritizedActors, FinalSortedCount );
+	SET_DWORD_STAT( STAT_NumRelevantDeletedActors, DeletedCount );
+
+	return FinalSortedCount;
+}
+
+int32 UNetDriver::ServerReplicateActors_ProcessPrioritizedActors( UNetConnection* Connection, const TArray<FNetViewer>& ConnectionViewers, FActorPriority** PriorityActors, const int32 FinalSortedCount, int32& OutUpdated )
+{
+	if ( !Connection->IsNetReady( 0 ) )
+	{
+		// Connection saturated, don't process any actors
+		return 0;
+	}
+
+	int32 ActorUpdatesThisConnection		= 0;
+	int32 ActorUpdatesThisConnectionSent	= 0;
+	int32 FinalRelevantCount				= 0;
+
+	for ( int32 j = 0; j < FinalSortedCount; j++ )
+	{
+		// Deletion entry
+		if ( PriorityActors[j]->ActorInfo == NULL && PriorityActors[j]->DestructionInfo )
+		{
+			// Make sure client has streaming level loaded
+			if ( PriorityActors[j]->DestructionInfo->StreamingLevelName != NAME_None && !Connection->ClientVisibleLevelNames.Contains( PriorityActors[j]->DestructionInfo->StreamingLevelName ) )
 			{
-				AActor *Actor = ConsiderList[ConsiderIdx];
-				// if the actor hasn't already been flagged by another connection,
-				if (Actor != NULL && !Actor->bPendingNetUpdate)
-				{
-					// find the channel
-					UActorChannel *Channel = Connection->ActorChannels.FindRef(Actor);
-					// and if the channel last update time doesn't match the last net update time for the actor
-					if (Channel != NULL && Channel->LastUpdateTime < Actor->LastNetUpdateTime)
-					{
-						//UE_LOG(LogNet, Log, TEXT("flagging %s for a future update"),*Actor->GetName());
-						// flag it for a pending update
-						Actor->bPendingNetUpdate = true;
-					}
-				}
+				// This deletion entry is for an actor in a streaming level the connection doesn't have loaded, so skip it
+				continue;
 			}
-			// clear the time sensitive flag to avoid sending an extra packet to this connection
-			Connection->TimeSensitive = false;
 
-			Connection->OwnedConsiderList.Empty();
-			for (int32 ChildIdx = 0; ChildIdx < Connection->Children.Num(); ChildIdx++)
+			UActorChannel* Channel = ( UActorChannel* )Connection->CreateChannel( CHTYPE_Actor, 1 );
+			if ( Channel )
 			{
-				if (Connection->Children[ChildIdx])
-				{
-					Connection->Children[ChildIdx]->OwnedConsiderList.Empty();
-				}
+				FinalRelevantCount++;
+				UE_LOG( LogNetTraffic, Log, TEXT( "Server replicate actor creating destroy channel for NetGUID <%s,%s> Priority: %d" ), *PriorityActors[j]->DestructionInfo->NetGUID.ToString(), *PriorityActors[j]->DestructionInfo->PathName, PriorityActors[j]->Priority );
+
+				Channel->SetChannelActorForDestroy( PriorityActors[j]->DestructionInfo ); // Send a close bunch on the new channel
+				Connection->DestroyedStartupOrDormantActors.Remove( PriorityActors[j]->DestructionInfo->NetGUID ); // Remove from connections to-be-destroyed list (close bunch of reliable, so it will make it there)
 			}
+			continue;
 		}
-		else if (Connection->ViewTarget)
+
+#if !(UE_BUILD_SHIPPING || UE_BUILD_TEST)
+		static IConsoleVariable* DebugObjectCvar = IConsoleManager::Get().FindConsoleVariable( TEXT( "net.PackageMap.DebugObject" ) );
+		static IConsoleVariable* DebugAllObjectsCvar = IConsoleManager::Get().FindConsoleVariable( TEXT( "net.PackageMap.DebugAll" ) );
+		if ( PriorityActors[j]->ActorInfo &&
+				( ( DebugObjectCvar && !DebugObjectCvar->GetString().IsEmpty() && PriorityActors[j]->ActorInfo->Actor->GetName().Contains( DebugObjectCvar->GetString() ) ) ||
+				( DebugAllObjectsCvar && DebugAllObjectsCvar->GetInt() != 0 ) ) )
 		{
-			int32 j;
-			int32 ConsiderCount	= 0;
-			int32 DeletedCount = 0;
-			
-			int32 NetRelevantCount = 0;
-			FActorPriority* PriorityList = NULL;
-			FActorPriority** PriorityActors = NULL;
+			UE_LOG( LogNetPackageMap, Log, TEXT( "Evaluating actor for replication %s" ), *PriorityActors[j]->ActorInfo->Actor->GetName() );
+		}
+#endif
 
-			TArray<FNetViewer>& ConnectionViewers = WorldSettings->ReplicationViewers;
+		// Normal actor replication
+		UActorChannel* Channel = PriorityActors[j]->Channel;
+		UE_LOG( LogNetTraffic, Log, TEXT( " Maybe Replicate %s" ), *PriorityActors[j]->ActorInfo->Actor->GetName() );
+		if ( !Channel || Channel->Actor ) //make sure didn't just close this channel
+		{
+			AActor*		Actor       = PriorityActors[j]->ActorInfo->Actor;
+			bool		bIsRelevant = false;
 
-			float PruneActors = 0.f;
-			CLOCK_CYCLES(PruneActors);
-			FMemMark RelevantActorMark(FMemStack::Get());
+			const bool bLevelInitializedForActor = IsLevelInitializedForActor( Actor, Connection );
 
-			// Prioritize actors for this connection
+			// only check visibility on already visible actors every 1.0 + 0.5R seconds
+			// bTearOff actors should never be checked
+			if ( bLevelInitializedForActor )
 			{
-				SCOPE_CYCLE_COUNTER(STAT_NetPrioritizeActorsTime);
-
-				// send ClientAdjustment if necessary
-				// we do this here so that we send a maximum of one per packet to that client; there is no value in stacking additional corrections
-				if (Connection->PlayerController)
+				if ( !Actor->bTearOff && ( !Channel || Time - Channel->RelevantTime > 1.f ) )
 				{
-					Connection->PlayerController->SendClientAdjustment();
-				}
-				
-				for (int32 ChildIdx = 0; ChildIdx < Connection->Children.Num(); ChildIdx++)
-				{
-					if (Connection->Children[ChildIdx]->PlayerController != NULL)
+					if ( IsActorRelevantToConnection( Actor, ConnectionViewers ) )
+					{
+						bIsRelevant = true;
+					}
+					else if ( DebugRelevantActors )
 					{
-						Connection->Children[ChildIdx]->PlayerController->SendClientAdjustment();
+						LastNonRelevantActors.Add( Actor );
 					}
 				}
+			}
+			else
+			{
+				// Actor is no longer relevant because the world it is/was in is not loaded by client
+				// exception: player controllers should never show up here
+				UE_LOG( LogNetTraffic, Log, TEXT( "- Level not initialized for actor %s" ), *Actor->GetName() );
+			}
 
-				// Get list of visible/relevant actors.
-				
-				NetTag++;
-				Connection->TickCount++;
+			// if the actor is now relevant or was recently relevant
+			const bool bIsRecentlyRelevant = bIsRelevant || ( Channel && Time - Channel->RelevantTime < RelevantTimeout );
 
-				// Set up to skip all sent temporary actors
-				for( j=0; j<Connection->SentTemporaries.Num(); j++ )
-				{
-					Connection->SentTemporaries[j]->NetTag = NetTag;
-				}
+			if ( bIsRecentlyRelevant )
+			{
+				FinalRelevantCount++;
 
-				// set the replication viewers to the current connection (and children) so that actors can determine who is currently being considered for relevancy checks
-				ConnectionViewers.Reset();
-				new(ConnectionViewers) FNetViewer(Connection, DeltaSeconds);
-				for (j = 0; j < Connection->Children.Num(); j++)
+				// Find or create the channel for this actor.
+				// we can't create the channel if the client is in a different world than we are
+				// or the package map doesn't support the actor's class/archetype (or the actor itself in the case of serializable actors)
+				// or it's an editor placed actor and the client hasn't initialized the level it's in
+				if ( Channel == NULL && GuidCache->SupportsObject( Actor->GetClass() ) && GuidCache->SupportsObject( Actor->IsNetStartupActor() ? Actor : Actor->GetArchetype() ) )
 				{
-					if (Connection->Children[j]->ViewTarget != NULL)
+					if ( bLevelInitializedForActor )
 					{
-						new(ConnectionViewers) FNetViewer(Connection->Children[j], DeltaSeconds);
+						// Create a new channel for this actor.
+						Channel = ( UActorChannel* )Connection->CreateChannel( CHTYPE_Actor, 1 );
+						if ( Channel )
+						{
+							Channel->SetChannelActor( Actor );
+						}
+					}
+					// if we couldn't replicate it for a reason that should be temporary, and this Actor is updated very infrequently, make sure we update it again soon
+					else if ( Actor->NetUpdateFrequency < 1.0f )
+					{
+						UE_LOG( LogNetTraffic, Log, TEXT( "Unable to replicate %s" ), *Actor->GetName() );
+						Actor->NetUpdateTime = Actor->GetWorld()->TimeSeconds + 0.2f * FMath::FRand();
 					}
 				}
 
-				// Make list of all actors to consider.
-				check(World == Connection->OwningActor->GetWorld());
-				
-				NetRelevantCount = World->GetNetRelevantActorCount() + DestroyedStartupOrDormantActors.Num();
-				PriorityList = new(FMemStack::Get(),NetRelevantCount+2)FActorPriority;
-				PriorityActors = new(FMemStack::Get(),NetRelevantCount+2)FActorPriority*;
-
-				// determine whether we should priority sort the list of relevant actors based on the saturation/bandwidth of the current connection
-				//@note - if the server is currently CPU saturated then do not sort until framerate improves
-				check(World == Connection->ViewTarget->GetWorld());
-				AGameMode const* const GameMode = World->GetAuthGameMode();
-				bool bLowNetBandwidth = !bCPUSaturated && (Connection->CurrentNetSpeed / float(GameMode->NumPlayers + GameMode->NumBots) < 500.f );
-
-				for( AActor* Actor : ConsiderList )
+				if ( Channel )
 				{
-					UActorChannel* Channel = Connection->ActorChannels.FindRef(Actor);
-
-					// Skip Actor if dormant
-					if ( CVarSetNetDormancyEnabled.GetValueOnGameThread() == 1 )
+					// if it is relevant then mark the channel as relevant for a short amount of time
+					if ( bIsRelevant )
 					{
-						// If actor is already dormant on this channel, then skip replication entirely
-						if ( Connection->DormantActors.Contains( Actor ) )
+						Channel->RelevantTime = Time + 0.5f * FMath::SRand();
+					}
+					// if the channel isn't saturated
+					if ( Channel->IsNetReady( 0 ) )
+					{
+						// replicate the actor
+						UE_LOG( LogNetTraffic, Log, TEXT( "- Replicate %s. %d" ), *Actor->GetName(), PriorityActors[j]->Priority );
+						if ( DebugRelevantActors )
 						{
-							// net.DormancyValidate can be set to 2 to validate dormant actor properties on every replicate
-							// (this could be moved to be done every tick instead of every net update if necessary, but seems excessive)
-							if ( CVarNetDormancyValidate.GetValueOnGameThread() == 2 )
-							{
-								TSharedRef< FObjectReplicator > * Replicator = Connection->DormantReplicatorMap.Find( Actor );
-
-								if ( Replicator != NULL )
-								{
-									Replicator->Get().ValidateAgainstState( Actor );
-								}
-							}
-
-							continue;
+							LastRelevantActors.Add( Actor );
 						}
 
-						// If actor might need to go dormant on this channel, then check
-						if (Actor->NetDormancy > DORM_Awake && Channel && !Channel->bPendingDormancy && !Channel->Dormant )
+						if ( Channel->ReplicateActor() )
 						{
-							bool ShouldGoDormant = true;
-							if (Actor->NetDormancy == DORM_DormantPartial)
+							ActorUpdatesThisConnectionSent++;
+							if ( DebugRelevantActors )
 							{
-								float LastReplicationTime  = Channel ? (Connection->Driver->Time - Channel->LastUpdateTime) : Connection->Driver->SpawnPrioritySeconds;
-								for (int32 viewerIdx = 0; viewerIdx < ConnectionViewers.Num(); viewerIdx++)
-								{
-									if (!Actor->GetNetDormancy(ConnectionViewers[viewerIdx].ViewLocation, ConnectionViewers[viewerIdx].ViewDir, ConnectionViewers[viewerIdx].InViewer, ConnectionViewers[viewerIdx].ViewTarget, Channel, Time, bLowNetBandwidth))
-									{
-										ShouldGoDormant = false;
-										break;
-									}
-								}
+								LastSentActors.Add( Actor );
 							}
 
-							if (ShouldGoDormant)
-							{
-								// Channel is marked to go dormant now once all properties have been replicated (but is not dormant yet)
-								Channel->StartBecomingDormant();
-							}
+							// Calculate min delta (max rate actor will upate), and max delta (slowest rate actor will update)
+							const float MinOptimalDelta				= 1.0f / Actor->NetUpdateFrequency;
+							const float MaxOptimalDelta				= FMath::Max( 1.0f / Actor->MinNetUpdateFrequency, MinOptimalDelta );
+							const float DeltaBetweenReplications	= ( World->TimeSeconds - PriorityActors[j]->ActorInfo->LastNetReplicateTime );
+
+							// Choose an optimal time, we choose 70% of the actual rate to allow frequency to go up if needed
+							PriorityActors[j]->ActorInfo->OptimalNetUpdateDelta = FMath::Clamp( DeltaBetweenReplications * 0.7f, MinOptimalDelta, MaxOptimalDelta );
+							PriorityActors[j]->ActorInfo->LastNetReplicateTime = World->TimeSeconds;
 						}
+						ActorUpdatesThisConnection++;
+						OutUpdated++;
 					}
-
-
-					// Skip actor if not relevant and theres no channel already.
-					// Historically Relevancy checks were deferred until after prioritization because they were expensive (line traces).
-					// Relevancy is now cheap and we are dealing with larger lists of considered actors, so we want to keep the list of
-					// prioritized actors low.
-					if (!Channel)
+					else
 					{
-						if ( !IsLevelInitializedForActor(Actor, Connection) )
-						{
-							// If the level this actor belongs to isn't loaded on client, don't bother sending
-							continue;
-						}
-						bool Relevant = false;
-						for (int32 viewerIdx = 0; viewerIdx < ConnectionViewers.Num(); viewerIdx++)
-						{
-							if(Actor->IsNetRelevantFor(ConnectionViewers[viewerIdx].InViewer, ConnectionViewers[viewerIdx].ViewTarget, ConnectionViewers[viewerIdx].ViewLocation))
-							{
-								Relevant = true;
-								break;
-							}
-						}
-						if (!Relevant)
-						{
-							continue;
-						}
+						UE_LOG( LogNetTraffic, Log, TEXT( "- Channel saturated, forcing pending update for %s" ), *Actor->GetName() );
+						// otherwise force this actor to be considered in the next tick again
+						Actor->ForceNetUpdate();
 					}
-
-					if( Actor->NetTag!=NetTag ) // Do not consider actor for this connection if this connection has it marked dormant
+					// second check for channel saturation
+					if ( !Connection->IsNetReady( 0 ) )
 					{
-						UE_LOG(LogNetTraffic, Log, TEXT("Consider %s alwaysrelevant %d frequency %f "),*Actor->GetName(), Actor->bAlwaysRelevant, Actor->NetUpdateFrequency);
-						Actor->NetTag                 = NetTag;
-						PriorityList  [ConsiderCount] = FActorPriority(Connection, Channel, Actor, ConnectionViewers, bLowNetBandwidth);
-						PriorityActors[ConsiderCount] = PriorityList + ConsiderCount;
-						ConsiderCount++;
-
-						if (DebugRelevantActors)
-						{
-							LastPrioritizedActors.Add(Actor);
-						}
+						// We can bail out now since this connection is saturated, we'll return how far we got though
+						SET_DWORD_STAT( STAT_NumReplicatedActorAttempts, ActorUpdatesThisConnection );
+						SET_DWORD_STAT( STAT_NumReplicatedActors, ActorUpdatesThisConnectionSent );
+						SET_DWORD_STAT( STAT_NumRelevantActors, FinalRelevantCount );
+						return j;
 					}
 				}
+			}
+
+			// If the actor wasn't recently relevant, or if it was torn off, close the actor channel if it exists for this connection
+			if ( ( !bIsRecentlyRelevant || Actor->bTearOff ) && Channel != NULL )
+			{
+				// Non startup (map) actors have their channels closed immediately, which destroys them.
+				// Startup actors get to keep their channels open.
 
-				// Add in deleted actors
-				for (auto It = Connection->DestroyedStartupOrDormantActors.CreateIterator(); It; ++It)
+				// Fixme: this should be a setting
+				if ( !bLevelInitializedForActor || !Actor->IsNetStartupActor() )
 				{
-					FActorDestructionInfo &DInfo = DestroyedStartupOrDormantActors.FindChecked(*It);
-					PriorityList  [ConsiderCount] = FActorPriority(Connection, &DInfo, ConnectionViewers);
-					PriorityActors[ConsiderCount] = PriorityList + ConsiderCount;
-					ConsiderCount++;
-					DeletedCount++;
+					UE_LOG( LogNetTraffic, Log, TEXT( "- Closing channel for no longer relevant actor %s" ), *Actor->GetName() );
+					Channel->Close();
 				}
+			}
+		}
+	}
 
-				UNetConnection* NextConnection = Connection;
-				int32 ChildIndex = 0;
-				while (NextConnection != NULL)
-				{
-					for (AActor* Actor : NextConnection->OwnedConsiderList)
-					{
-						UE_LOG(LogNetTraffic, Log, TEXT("Consider owned %s always relevant %d frequency %f  "),*Actor->GetName(), Actor->bAlwaysRelevant,Actor->NetUpdateFrequency);
-						if (Actor->NetTag != NetTag)
-						{
-							UActorChannel* Channel = Connection->ActorChannels.FindRef(Actor);
-							Actor->NetTag                 = NetTag;
-							PriorityList  [ConsiderCount] = FActorPriority(NextConnection, Channel, Actor, ConnectionViewers, bLowNetBandwidth);
-							PriorityActors[ConsiderCount] = PriorityList + ConsiderCount;
-							ConsiderCount++;
+	SET_DWORD_STAT( STAT_NumReplicatedActorAttempts, ActorUpdatesThisConnection );
+	SET_DWORD_STAT( STAT_NumReplicatedActors, ActorUpdatesThisConnectionSent );
+	SET_DWORD_STAT( STAT_NumRelevantActors, FinalRelevantCount );
 
-							if (DebugRelevantActors)
-							{
-								LastPrioritizedActors.Add(Actor);
-							}
-						}
-					}
-					NextConnection->OwnedConsiderList.Empty();
+	return FinalSortedCount;
+}
+#endif
 
-					NextConnection = (ChildIndex < Connection->Children.Num()) ? Connection->Children[ChildIndex++] : NULL;
-				}
+int32 UNetDriver::ServerReplicateActors(float DeltaSeconds)
+{
+	SCOPE_CYCLE_COUNTER(STAT_NetServerRepActorsTime);
 
-				SET_DWORD_STAT(STAT_PrioritizedActors,ConsiderCount);
-				SET_DWORD_STAT(STAT_NumRelevantDeletedActors,DeletedCount);
+#if WITH_SERVER_CODE
+	if ( ClientConnections.Num() == 0 )
+	{
+		return 0;
+	}
 
-				// Sort by priority
-				struct FCompareFActorPriority
-				{
-					FORCEINLINE bool operator()( const FActorPriority& A, const FActorPriority& B ) const
-					{
-						return B.Priority < A.Priority;
-					}
-				};
-				Sort( PriorityActors, ConsiderCount, FCompareFActorPriority() );
+	check( World );
 
-			} // END PRIORITIZE
+	int32 Updated = 0;
 
-			// Update all relevant actors in sorted order.
-			bool bNewSaturated = !Connection->IsNetReady(0);
-			if (bNewSaturated)
-			{
-				j = 0;
-			}
-			else
-			{
-				UE_LOG(LogNetTraffic, Log, TEXT("START"));
-				int32 FinalRelevantCount = 0;
-				for (j = 0; j < ConsiderCount; j++)
-				{
-					// Deletion entry
-					if (PriorityActors[j]->Actor == NULL && PriorityActors[j]->DestructionInfo)
-					{
-						// Make sure client has streaming level loaded
-						if (PriorityActors[j]->DestructionInfo->StreamingLevelName != NAME_None && !Connection->ClientVisibleLevelNames.Contains(PriorityActors[j]->DestructionInfo->StreamingLevelName))
-						{
-							// This deletion entry is for an actor in a streaming level the connection doesn't have loaded, so skip it
-							continue;
-						}
+	// Bump the ReplicationFrame value to invalidate any properties marked as "unchanged" for this frame.
+	ReplicationFrame++;
 
-						UActorChannel* Channel = (UActorChannel*)Connection->CreateChannel( CHTYPE_Actor, 1 );
-						if (Channel)
-						{
-							FinalRelevantCount++;
-							UE_LOG(LogNetTraffic, Log, TEXT("Server replicate actor creating destroy channel for NetGUID <%s,%s> Priority: %d"), *PriorityActors[j]->DestructionInfo->NetGUID.ToString(), *PriorityActors[j]->DestructionInfo->PathName, PriorityActors[j]->Priority );
+	const int32 NumClientsToTick = ServerReplicateActors_PrepConnections( DeltaSeconds );
 
-							Channel->SetChannelActorForDestroy( PriorityActors[j]->DestructionInfo ); // Send a close bunch on the new channel
-							Connection->DestroyedStartupOrDormantActors.Remove( PriorityActors[j]->DestructionInfo->NetGUID ); // Remove from connections to-be-destroyed list (close bunch of reliable, so it will make it there)
-						}
-						continue;
-					}
+	if ( NumClientsToTick == 0 )
+	{
+		// No connections are ready this frame
+		return 0;
+	}
 
-#if !(UE_BUILD_SHIPPING || UE_BUILD_TEST)
-					static IConsoleVariable* DebugObjectCvar = IConsoleManager::Get().FindConsoleVariable(TEXT("net.PackageMap.DebugObject"));
-					static IConsoleVariable* DebugAllObjectsCvar = IConsoleManager::Get().FindConsoleVariable(TEXT("net.PackageMap.DebugAll"));
-					if (PriorityActors[j]->Actor && 
-						((DebugObjectCvar && !DebugObjectCvar->GetString().IsEmpty() && PriorityActors[j]->Actor->GetName().Contains(DebugObjectCvar->GetString())) ||
-						(DebugAllObjectsCvar && DebugAllObjectsCvar->GetInt() != 0)))
-					{
-						UE_LOG(LogNetPackageMap, Log, TEXT("Evaluating actor for replication %s"), *PriorityActors[j]->Actor->GetName());
-					}
-#endif
+	AWorldSettings* WorldSettings = World->GetWorldSettings();
+
+	bool bCPUSaturated		= false;
+	float ServerTickTime	= GEngine->GetMaxTickRate( DeltaSeconds );
+	if ( ServerTickTime == 0.f )
+	{
+		ServerTickTime = DeltaSeconds;
+	}
+	else
+	{
+		ServerTickTime	= 1.f/ServerTickTime;
+		bCPUSaturated	= DeltaSeconds > 1.2f * ServerTickTime;
+	}
 
-					// Normal actor replication
-					UActorChannel* Channel     = PriorityActors[j]->Channel;
-					UE_LOG(LogNetTraffic, Log, TEXT(" Maybe Replicate %s"),*PriorityActors[j]->Actor->GetName());
-					if ( !Channel || Channel->Actor ) //make sure didn't just close this channel
-					{ 
-						AActor*		Actor       = PriorityActors[j]->Actor;
-						bool		bIsRelevant = false;
+	TArray<FNetworkObjectInfo*> ConsiderList;
+	ConsiderList.Reserve( GetNetworkObjectList().GetObjects().Num() );
 
-						const bool bLevelInitializedForActor = IsLevelInitializedForActor(Actor, Connection);
+	// Build the consider list (actors that are ready to replicate)
+	ServerReplicateActors_BuildConsiderList( ConsiderList, ServerTickTime );
 
-						// only check visibility on already visible actors every 1.0 + 0.5R seconds
-						// bTearOff actors should never be checked
-						if ( bLevelInitializedForActor )
-						{
-							if (!Actor->bTearOff && (!Channel || Time - Channel->RelevantTime > 1.f))
-							{
-								for (int32 k = 0; k < ConnectionViewers.Num(); k++)
-								{
-									if (Actor->IsNetRelevantFor(ConnectionViewers[k].InViewer, ConnectionViewers[k].ViewTarget, ConnectionViewers[k].ViewLocation))
-									{
-										bIsRelevant = true;
-										break;
-									}
-									else
-									{
-										//UE_LOG(LogNetPackageMap, Warning, TEXT("Actor NonRelevant: %s"), *Actor->GetName() );
-										if (DebugRelevantActors)
-										{
-											LastNonRelevantActors.Add(Actor);
-										}
-									}
-								}
-							}
-						}
-						else
-						{
-							// Actor is no longer relevant because the world it is/was in is not loaded by client
-							// exception: player controllers should never show up here
-							UE_LOG(LogNetTraffic, Log, TEXT("- Level not initialized for actor %s"), *Actor->GetName());
-						}
-						
-						// if the actor is now relevant or was recently relevant
-						const bool bIsRecentlyRelevant = bIsRelevant || (Channel && Time - Channel->RelevantTime < RelevantTimeout);
-
-						if( bIsRecentlyRelevant )
-						{	
-							FinalRelevantCount++;
-
-							// Find or create the channel for this actor.
-							// we can't create the channel if the client is in a different world than we are
-							// or the package map doesn't support the actor's class/archetype (or the actor itself in the case of serializable actors)
-							// or it's an editor placed actor and the client hasn't initialized the level it's in
-							if ( Channel == NULL && GuidCache->SupportsObject(Actor->GetClass()) &&
-									GuidCache->SupportsObject(Actor->IsNetStartupActor() ? Actor : Actor->GetArchetype()) )
-							{
-								if (bLevelInitializedForActor)
-								{
-									// Create a new channel for this actor.
-									Channel = (UActorChannel*)Connection->CreateChannel( CHTYPE_Actor, 1 );
-									if( Channel )
-									{
-										Channel->SetChannelActor( Actor );
-									}
-								}
-								// if we couldn't replicate it for a reason that should be temporary, and this Actor is updated very infrequently, make sure we update it again soon
-								else if (Actor->NetUpdateFrequency < 1.0f)
-								{
-									UE_LOG(LogNetTraffic, Log, TEXT("Unable to replicate %s"),*Actor->GetName());
-									Actor->NetUpdateTime = Actor->GetWorld()->TimeSeconds + 0.2f * FMath::FRand();
-								}
-							}
+	FMemMark Mark( FMemStack::Get() );
 
-							if( Channel )
-							{
-								// if it is relevant then mark the channel as relevant for a short amount of time
-								if( bIsRelevant )
-								{
-									Channel->RelevantTime = Time + 0.5f * FMath::SRand();
-								}
-								// if the channel isn't saturated
-								if( Channel->IsNetReady(0) )
-								{
-									// replicate the actor
-									UE_LOG(LogNetTraffic, Log, TEXT("- Replicate %s. %d"),*Actor->GetName(), PriorityActors[j]->Priority);
-									if (DebugRelevantActors)
-									{
-										LastRelevantActors.Add( Actor );
-									}
-
-									if (Channel->ReplicateActor())
-									{
-										ActorUpdatesThisConnectionSent++;
-										if (DebugRelevantActors)
-										{
-											LastSentActors.Add( Actor );
-										}
-									}
-									ActorUpdatesThisConnection++;
-									Updated++;
-								}
-								else
-								{							
-									UE_LOG(LogNetTraffic, Log, TEXT("- Channel saturated, forcing pending update for %s"),*Actor->GetName());
-									// otherwise force this actor to be considered in the next tick again
-									Actor->ForceNetUpdate();
-								}
-								// second check for channel saturation
-								if (!Connection->IsNetReady(0))
-								{
-									bNewSaturated = true;
-									break;
-								}
-							}
-						}
-						
-						// If the actor wasn't recently relevant, or if it was torn off,
-						// close the actor channel if it exists for this connection
-						if ((!bIsRecentlyRelevant || Actor->bTearOff) && Channel != NULL)
-						{
-							// Non startup (map) actors have their channels closed immediately, which destroys them.
-							// Startup actors get to keep their channels open.
+	for ( int32 i=0; i < ClientConnections.Num(); i++ )
+	{
+		UNetConnection* Connection = ClientConnections[i];
+		check(Connection);
 
-							// Fixme: this should be a setting
-							if ( !bLevelInitializedForActor || !Actor->IsNetStartupActor() )
-							{
-								UE_LOG(LogNetTraffic, Log, TEXT("- Closing channel for no longer relevant actor %s"),*Actor->GetName());
-								Channel->Close();
-							}
-						}
+		// if this client shouldn't be ticked this frame
+		if (i >= NumClientsToTick)
+		{
+			//UE_LOG(LogNet, Log, TEXT("skipping update to %s"),*Connection->GetName());
+			// then mark each considered actor as bPendingNetUpdate so that they will be considered again the next frame when the connection is actually ticked
+			for (int32 ConsiderIdx = 0; ConsiderIdx < ConsiderList.Num(); ConsiderIdx++)
+			{
+				AActor *Actor = ConsiderList[ConsiderIdx]->Actor;
+				// if the actor hasn't already been flagged by another connection,
+				if (Actor != NULL && !Actor->bPendingNetUpdate)
+				{
+					// find the channel
+					UActorChannel *Channel = Connection->ActorChannels.FindRef(Actor);
+					// and if the channel last update time doesn't match the last net update time for the actor
+					if (Channel != NULL && Channel->LastUpdateTime < Actor->LastNetUpdateTime)
+					{
+						//UE_LOG(LogNet, Log, TEXT("flagging %s for a future update"),*Actor->GetName());
+						// flag it for a pending update
+						Actor->bPendingNetUpdate = true;
 					}
 				}
+			}
+			// clear the time sensitive flag to avoid sending an extra packet to this connection
+			Connection->TimeSensitive = false;
+		}
+		else if (Connection->ViewTarget)
+		{		
+			// Make a list of viewers this connection should consider (this connection and children of this connection)
+			TArray<FNetViewer>& ConnectionViewers = WorldSettings->ReplicationViewers;
 
-				SET_DWORD_STAT(STAT_NumRelevantActors,FinalRelevantCount);
+			ConnectionViewers.Reset();
+			new( ConnectionViewers )FNetViewer( Connection, DeltaSeconds );
+			for ( int32 ViewerIndex = 0; ViewerIndex < Connection->Children.Num(); ViewerIndex++ )
+			{
+				if ( Connection->Children[ViewerIndex]->ViewTarget != NULL )
+				{
+					new( ConnectionViewers )FNetViewer( Connection->Children[ViewerIndex], DeltaSeconds );
+				}
+			}
+
+			// send ClientAdjustment if necessary
+			// we do this here so that we send a maximum of one per packet to that client; there is no value in stacking additional corrections
+			if ( Connection->PlayerController )
+			{
+				Connection->PlayerController->SendClientAdjustment();
 			}
 
+			for ( int32 ChildIdx = 0; ChildIdx < Connection->Children.Num(); ChildIdx++ )
+			{
+				if ( Connection->Children[ChildIdx]->PlayerController != NULL )
+				{
+					Connection->Children[ChildIdx]->PlayerController->SendClientAdjustment();
+				}
+			}
+
+			FMemMark RelevantActorMark(FMemStack::Get());
+
+			FActorPriority* PriorityList	= NULL;
+			FActorPriority** PriorityActors = NULL;
+
+			// Get a sorted list of actors for this connection
+			const int32 FinalSortedCount = ServerReplicateActors_PrioritizeActors( Connection, ConnectionViewers, ConsiderList, bCPUSaturated, PriorityList, PriorityActors );
+
+			// Process the sorted list of actors for this connection
+			const int32 LastProcessedActor = ServerReplicateActors_ProcessPrioritizedActors( Connection, ConnectionViewers, PriorityActors, FinalSortedCount, Updated );
+
 			// relevant actors that could not be processed this frame are marked to be considered for next frame
-			for ( int32 k=j; k<ConsiderCount; k++ )
+			for ( int32 k=LastProcessedActor; k<FinalSortedCount; k++ )
 			{
-				AActor* Actor = PriorityActors[k]->Actor;
-				if (!Actor)
+				if (!PriorityActors[k]->ActorInfo)
 				{
 					// A deletion entry, skip it because we dont have anywhere to store a 'better give higher priority next time'
 					continue;
 				}
 
+				AActor* Actor = PriorityActors[k]->ActorInfo->Actor;
+
 				UActorChannel* Channel = PriorityActors[k]->Channel;
 				
 				UE_LOG(LogNetTraffic, Verbose, TEXT("Saturated. %s"), *Actor->GetName());
@@ -2712,29 +2729,18 @@ int32 UNetDriver::ServerReplicateActors(float DeltaSeconds)
 					UE_LOG(LogNetTraffic, Log, TEXT(" Saturated. Mark %s NetUpdateTime to be checked for next tick"), *Actor->GetName());
 					Actor->bPendingNetUpdate = true;
 				}
-				else
+				else if ( IsActorRelevantToConnection( Actor, ConnectionViewers ) )
 				{
-					for (int32 h = 0; h < ConnectionViewers.Num(); h++)
+					// If this actor was relevant but didn't get processed, force another update for next frame
+					UE_LOG( LogNetTraffic, Log, TEXT( " Saturated. Mark %s NetUpdateTime to be checked for next tick" ), *Actor->GetName() );
+					Actor->bPendingNetUpdate = true;
+					if ( Channel != NULL )
 					{
-						if (Actor->IsNetRelevantFor(ConnectionViewers[h].InViewer, ConnectionViewers[h].ViewTarget, ConnectionViewers[h].ViewLocation))
-						{
-							UE_LOG(LogNetTraffic, Log, TEXT(" Saturated. Mark %s NetUpdateTime to be checked for next tick"), *Actor->GetName());
-							Actor->bPendingNetUpdate = true;
-							if (Channel != NULL)
-							{
-								Channel->RelevantTime = Time + 0.5f * FMath::SRand();
-							}
-							break;
-						}
+						Channel->RelevantTime = Time + 0.5f * FMath::SRand();
 					}
 				}
 			}
 			RelevantActorMark.Pop();
-			UE_LOG(LogNetTraffic, Log, TEXT("Potential %04i ConsiderList %03i ConsiderCount %03i Prune=%01.4f "),NetRelevantCount, 
-						ConsiderList.Num(), ConsiderCount, FPlatformTime::ToMilliseconds(PruneActors) );
-
-			SET_DWORD_STAT(STAT_NumReplicatedActorAttempts,ActorUpdatesThisConnection);
-			SET_DWORD_STAT(STAT_NumReplicatedActors,ActorUpdatesThisConnectionSent);
 		}
 	}
 
@@ -2933,6 +2939,8 @@ void UNetDriver::SetWorld(class UWorld* InWorld)
 		UnregisterTickEvents(World);
 		World = NULL;
 		Notify = NULL;
+
+		GetNetworkObjectList().GetObjects().Reset();
 	}
 
 	if (InWorld)
@@ -2941,6 +2949,8 @@ void UNetDriver::SetWorld(class UWorld* InWorld)
 		World = InWorld;
 		Notify = InWorld;
 		RegisterTickEvents(InWorld);
+
+		GetNetworkObjectList().AddInitialObjects(InWorld, NetDriverName);
 	}
 }
 
diff --git a/Engine/Source/Runtime/Engine/Private/NetworkObjectList.cpp b/Engine/Source/Runtime/Engine/Private/NetworkObjectList.cpp
new file mode 100644
--- /dev/null
+++ b/Engine/Source/Runtime/Engine/Private/NetworkObjectList.cpp
@@ -0,0 +1,46 @@
+// Copyright 1998-2015 Epic Games, Inc. All Rights Reserved.
+
+#include "EnginePrivate.h"
+#include "Engine/NetworkObjectList.h"
+#include "Engine/Level.h"
+#include "GameFramework/Actor.h"
+
+void FNetworkObjectList::AddInitialObjects(UWorld* const World, const FName NetDriverName)
+{
+	if (World == nullptr)
+	{
+		return;
+	}
+
+	for (FActorIterator Iter(World); Iter; ++Iter)
+	{
+		AActor* Actor = *Iter;
+		if (Actor != nullptr && !Actor->IsPendingKill() && ULevel::IsNetActor(Actor))
+		{
+			Add(Actor, NetDriverName);
+		}
+	}
+}
+
+void FNetworkObjectList::Add(AActor* const Actor, const FName NetDriverName)
+{
+	if (Actor == nullptr)
+	{
+		return;
+	}
+
+	if (!NetworkObjects.Contains(Actor))
+	{
+		// Special case the demo net driver, since actors currently only have one associated NetDriverName.
+		
+		// FIXME: For now, add all actors, regardless of NetDriverName.
+		// There is an issue where NetDriverName isn't correct at this time
+		// The real fix is to register actors with netdrivers when the name changes on the netdriver or actor
+		// But for now, this fix is safe.
+
+		//if (Actor->NetDriverName == NetDriverName || NetDriverName == NAME_DemoNetDriver)
+		{
+			NetworkObjects.Emplace(new FNetworkObjectInfo(Actor));
+		}
+	}
+}
diff --git a/Engine/Source/Runtime/Engine/Private/NetworkProfiler.cpp b/Engine/Source/Runtime/Engine/Private/NetworkProfiler.cpp
--- a/Engine/Source/Runtime/Engine/Private/NetworkProfiler.cpp
+++ b/Engine/Source/Runtime/Engine/Private/NetworkProfiler.cpp
@@ -26,7 +26,7 @@ FNetworkProfiler GNetworkProfiler;
 /** Magic value, determining that file is a network profiler file.				*/
 #define NETWORK_PROFILER_MAGIC						0x1DBF348C
 /** Version of memory profiler. Incremented on serialization changes.			*/
-#define NETWORK_PROFILER_VERSION					9
+#define NETWORK_PROFILER_VERSION					10
 
 static const FString UnknownName("UnknownName");
 
@@ -50,7 +50,9 @@ enum ENetworkProfilingPayloadType
 	NPTYPE_BeginContentBlock,			// Content block headers
 	NPTYPE_EndContentBlock,				// Content block footers
 	NPTYPE_WritePropertyHandle,			// Property handles
-	NPTYPE_ConnectionChanged			// Connection changed
+	NPTYPE_ConnectionChanged,			// Connection changed
+	NPTYPE_NameReference,				// New reference to name
+	NPTYPE_ConnectionReference			// New reference to connection
 };
 
 /*=============================================================================
@@ -60,21 +62,11 @@ enum ENetworkProfilingPayloadType
 FNetworkProfilerHeader::FNetworkProfilerHeader()
 	: Magic(NETWORK_PROFILER_MAGIC)
 	, Version(NETWORK_PROFILER_VERSION)
-	, NameTableOffset(0)
-	, NameTableEntries(0)
-	, AddressTableOffset( 0 )
-	, AddressTableEntries( 0 )
 {
 }
 
 void FNetworkProfilerHeader::Reset(const FURL& InURL)
 {
-	NameTableOffset = 0;
-	NameTableEntries = 0;
-
-	AddressTableOffset = 0;
-	AddressTableEntries = 0;
-		
 	FParse::Value( FCommandLine::Get(), TEXT( "NETWORKPROFILER=" ), Tag );
 	GameName = FApp::GetGameName();
 	URL = InURL.ToString();
@@ -84,11 +76,7 @@ FArchive& operator << ( FArchive& Ar, FNetworkProfilerHeader& Header )
 {
 	check( Ar.IsSaving() );
 	Ar	<< Header.Magic
-		<< Header.Version
-		<< Header.NameTableOffset
-		<< Header.NameTableEntries
-		<< Header.AddressTableOffset
-		<< Header.AddressTableEntries;
+		<< Header.Version;
 	Header.Tag.SerializeAsANSICharArray( Ar );
 	Header.GameName.SerializeAsANSICharArray( Ar );
 	Header.URL.SerializeAsANSICharArray( Ar );
@@ -133,6 +121,11 @@ int32 FNetworkProfiler::GetNameTableIndex( const FString& Name )
 		Index = NameArray.Num();
 		new(NameArray)FString(Name);
 		NameToNameTableIndexMap.Add(*Name,Index);
+
+		// Write out the name reference token
+		uint8 Type = NPTYPE_NameReference;
+		( *FileWriter ) << Type;
+		Name.SerializeAsANSICharArray( *FileWriter );
 	}
 
 	check(Index!=INDEX_NONE);
@@ -162,6 +155,11 @@ int32 FNetworkProfiler::GetAddressTableIndex( uint64 Address )
 		Index = AddressArray.Num();
 		AddressArray.Add( Address );
 		AddressTableIndexMap.Add( Address, Index );
+
+		// Write out the name reference token
+		uint8 Type = NPTYPE_ConnectionReference;
+		( *FileWriter ) << Type;
+		( *FileWriter ) << Address;
 	}
 
 	check( Index != INDEX_NONE );
@@ -247,11 +245,12 @@ void FNetworkProfiler::TrackSendRPC( const AActor* Actor, const UFunction* Funct
 
 		SetCurrentConnection( Connection );
 
+		uint32 ActorNameTableIndex = GetNameTableIndex( Actor->GetName() );
+		uint32 FunctionNameTableIndex = GetNameTableIndex( Function->GetName() );
+
 		uint8 Type = NPTYPE_SendRPC;
 		(*FileWriter) << Type;
-		uint32 ActorNameTableIndex = GetNameTableIndex( Actor->GetName() );
 		( *FileWriter ).SerializeIntPacked( ActorNameTableIndex );
-		uint32 FunctionNameTableIndex = GetNameTableIndex( Function->GetName() );
 		( *FileWriter ).SerializeIntPacked( FunctionNameTableIndex );
 		( *FileWriter ) << NumHeaderBits;
 		(*FileWriter) << NumParameterBits;
@@ -266,10 +265,12 @@ void FNetworkProfiler::TrackQueuedRPC( UNetConnection* Connection, UObject* Targ
 		SCOPE_LOCK_REF(CriticalSection);
 		
 		FQueuedRPCInfo Info;
-		Info.Connection = Connection;
-		Info.TargetObject = TargetObject;
+
 		Info.ActorNameIndex = GetNameTableIndex( Actor->GetName() );
 		Info.FunctionNameIndex = GetNameTableIndex( Function->GetName() );
+
+		Info.Connection = Connection;
+		Info.TargetObject = TargetObject;
 		Info.NumHeaderBits = NumHeaderBits;
 		Info.NumParameterBits = NumParameterBits;
 		Info.NumFooterBits = NumFooterBits;
@@ -370,9 +371,10 @@ void FNetworkProfiler::TrackSocketSendToCore(
 
 		SetCurrentConnection( Connection );
 
+		uint32 NameTableIndex = GetNameTableIndex( SocketDesc );
+
 		uint8 Type = NPTYPE_SocketSendTo;
 		(*FileWriter) << Type;
-		uint32 NameTableIndex = GetNameTableIndex( SocketDesc );
 		(*FileWriter).SerializeIntPacked( NameTableIndex );
 		(*FileWriter) << BytesSent;
 		(*FileWriter) << NumPacketIdBits;
@@ -475,11 +477,12 @@ void FNetworkProfiler::TrackReplicateActor( const AActor* Actor, FReplicationFla
 
 		SetCurrentConnection( Connection );
 
+		uint32 NameTableIndex = GetNameTableIndex( Actor->GetName() );
+
 		uint8 Type = NPTYPE_ReplicateActor;
 		(*FileWriter) << Type;
 		uint8 NetFlags = (RepFlags.bNetInitial << 1) | (RepFlags.bNetOwner << 2);
 		(*FileWriter) << NetFlags;
-		uint32 NameTableIndex = GetNameTableIndex( Actor->GetName() );
 		(*FileWriter).SerializeIntPacked( NameTableIndex );
 		float TimeInMS = FPlatformTime::ToMilliseconds(Cycles);	// FIXME: We may want to just pass in cycles to profiler to we don't lose precision
 		(*FileWriter) << TimeInMS;
@@ -502,9 +505,10 @@ void FNetworkProfiler::TrackReplicateProperty( const UProperty* Property, uint16
 
 		SetCurrentConnection( Connection );
 
+		uint32 NameTableIndex = GetNameTableIndex( Property->GetName() );
+
 		uint8 Type = NPTYPE_ReplicateProperty;
 		(*FileWriter) << Type;
-		uint32 NameTableIndex = GetNameTableIndex( Property->GetName() );
 		(*FileWriter).SerializeIntPacked( NameTableIndex );
 		(*FileWriter) << NumBits;
 	}
@@ -518,9 +522,10 @@ void FNetworkProfiler::TrackWritePropertyHeader( const UProperty* Property, uint
 
 		SetCurrentConnection( Connection );
 
+		uint32 NameTableIndex = GetNameTableIndex( Property->GetName() );
+
 		uint8 Type = NPTYPE_WritePropertyHeader;
 		(*FileWriter) << Type;
-		uint32 NameTableIndex = GetNameTableIndex( Property->GetName() );
 		(*FileWriter).SerializeIntPacked( NameTableIndex );
 		(*FileWriter) << NumBits;
 	}
@@ -540,11 +545,12 @@ void FNetworkProfiler::TrackEvent( const FString& EventName, const FString& Even
 
 		SetCurrentConnection( Connection );
 
+		uint32 EventNameNameTableIndex = GetNameTableIndex( EventName );
+		uint32 EventDescriptionNameTableIndex = GetNameTableIndex( EventDescription );
+
 		uint8 Type = NPTYPE_Event;
 		(*FileWriter) << Type;
-		uint32 EventNameNameTableIndex = GetNameTableIndex( EventName );
 		(*FileWriter).SerializeIntPacked( EventNameNameTableIndex );
-		uint32 EventDescriptionNameTableIndex = GetNameTableIndex( EventDescription );
 		(*FileWriter).SerializeIntPacked( EventDescriptionNameTableIndex );
 	}
 }
@@ -570,63 +576,20 @@ void FNetworkProfiler::TrackSessionChange( bool bShouldContinueTracking, const F
 		// End existing tracking session.
 		if( FileWriter )
 		{	
-			if( bHasNoticeableNetworkTrafficOccured )
-			{
-				UE_LOG(LogNet, Log, TEXT("Network Profiler: Writing out session file for '%s'"), *CurrentHeader.GetURL());
-
-				// Write end of stream marker.
-				uint8 Type = NPTYPE_EndOfStreamMarker;
-				(*FileWriter) << Type;
+			UE_LOG( LogNet, Log, TEXT( "Network Profiler: Closing session file for '%s'" ), *CurrentHeader.GetURL() );
 
-				// Write out name table and update header with offset and count.
-				CurrentHeader.SetNameTableValues(FileWriter->Tell(), NameArray.Num());
-
-				for( int32 NameIndex=0; NameIndex<NameArray.Num(); NameIndex++ )
-				{
-					NameArray[NameIndex].SerializeAsANSICharArray( *FileWriter );
-				}
-
-				CurrentHeader.SetAddressTableValues( FileWriter->Tell(), AddressArray.Num() );
-
-				for ( int32 AddressIndex=0; AddressIndex < AddressArray.Num(); AddressIndex++ )
-				{
-					( *FileWriter ) << AddressArray[AddressIndex];
-				}
-
-				// Seek to the beginning of the file and write out proper header.
-				FileWriter->Seek( 0 );
-				(*FileWriter) << CurrentHeader;
-
-				// Close file writer so we can rename the file to its final destination.
-				FileWriter->Close();
-			
-				// Rename/ move file.
-				static int32 Salt = 0;
-				Salt++;		// Use a salt to solve the issue where this function is called so fast it produces the same time (seems to happen during seamless travel)
-				const FString FinalFileName = FPaths::ProfilingDir() + FApp::GetGameName() + TEXT("-") + FDateTime::Now().ToString() + FString::Printf(TEXT("[%i]"), Salt) + TEXT(".nprof");
-				bool bWasMovedSuccessfully = IFileManager::Get().Move( *FinalFileName, *TempFileName );
-
-				// Send data to UnrealConsole to upload to DB.
-				if( bWasMovedSuccessfully )
-				{
-					UE_LOG( LogNet, Log, TEXT( "Network Profiler: Saved SUCCESS: %s" ), *FinalFileName );
-
-					SendDataToPCViaUnrealConsole( TEXT( "UE_PROFILER!NETWORK:" ), *FinalFileName );
-				}
-				else
-				{
-					UE_LOG( LogNet, Error, TEXT( "Network Profiler: Saved FAILED: %s" ), *FinalFileName );
-				}
-			}
-			else
+			if ( !bHasNoticeableNetworkTrafficOccured )
 			{
-				UE_LOG(LogNet, Warning, TEXT("Network Profiler: Nothing important happened"));
-				FileWriter->Close();
-
-				// Delete the temporary file.
-				IFileManager::Get().Delete(*TempFileName);
+				UE_LOG( LogNet, Warning, TEXT( "Network Profiler: Nothing important happened" ) );
 			}
 
+			// Write end of stream marker.
+			uint8 Type = NPTYPE_EndOfStreamMarker;
+			( *FileWriter ) << Type;
+
+			// Close file writer so we can rename the file to its final destination.
+			FileWriter->Close();
+
 			// Clean up.
 			delete FileWriter;
 			FileWriter = NULL;
@@ -635,28 +598,15 @@ void FNetworkProfiler::TrackSessionChange( bool bShouldContinueTracking, const F
 
 		if( bShouldContinueTracking )
 		{
-			// Delete any stale .tmp files.
-			TArray<FString> FoundTempFiles;
-			IFileManager::Get().FindFiles(FoundTempFiles, *(FPaths::ProfilingDir() / TEXT("*.tmp")), true, false);
-			
-			for ( const FString& FoundFile : FoundTempFiles)
-			{
-				const FString FullFilename = FPaths::ProfilingDir() + FoundFile;
-				if ( IFileManager::Get().GetFileAgeSeconds(*FullFilename) > MaxTempFileAgeSeconds)
-				{
-					IFileManager::Get().Delete(*FullFilename);
-				}
-			}
-
 			// Start a new tracking session.
 			check( FileWriter == NULL );
 
-			// Use a dummy name for sessions in progress that is renamed at end.
-			TempFileName = FPaths::CreateTempFilename(*FPaths::ProfilingDir(), TEXT("NetworkProfiling-"));
+			static int32 Salt = 0;
+			Salt++;		// Use a salt to solve the issue where this function is called so fast it produces the same time (seems to happen during seamless travel)
+			const FString FinalFileName = FPaths::ProfilingDir() + FApp::GetGameName() + TEXT( "-" ) + FDateTime::Now().ToString() + FString::Printf( TEXT( "[%i]" ), Salt ) + TEXT( ".nprof" );
 
-			// Create folder and file writer.
-			IFileManager::Get().MakeDirectory( *FPaths::GetPath(TempFileName) );
-			FileWriter = IFileManager::Get().CreateFileWriter( *TempFileName, FILEWRITE_EvenIfReadOnly );
+			IFileManager::Get().MakeDirectory( *FPaths::GetPath( FinalFileName ) );
+			FileWriter = IFileManager::Get().CreateFileWriter( *FinalFileName, FILEWRITE_EvenIfReadOnly );
 			check( FileWriter );
 			
 			CurrentHeader.Reset(InURL);
@@ -720,9 +670,10 @@ void FNetworkProfiler::TrackBeginContentBlock( UObject* Object, uint16 NumBits,
 
 		SetCurrentConnection( Connection );
 
+		uint32 NameTableIndex = GetNameTableIndex( Object != nullptr ? Object->GetName() : UnknownName );
+
 		uint8 Type = NPTYPE_BeginContentBlock;
 		(*FileWriter) << Type;
-		uint32 NameTableIndex = GetNameTableIndex( Object != nullptr ? Object->GetName() : UnknownName );
 		(*FileWriter).SerializeIntPacked( NameTableIndex );
 		(*FileWriter) << NumBits;
 	 }
@@ -736,9 +687,10 @@ void FNetworkProfiler::TrackEndContentBlock( UObject* Object, uint16 NumBits, UN
 
 		SetCurrentConnection( Connection );
 
+		uint32 NameTableIndex = GetNameTableIndex( Object != nullptr ? Object->GetName() : UnknownName );
+
 		uint8 Type = NPTYPE_EndContentBlock;
 		(*FileWriter) << Type;
-		uint32 NameTableIndex = GetNameTableIndex( Object != nullptr ? Object->GetName() : UnknownName );
 		(*FileWriter).SerializeIntPacked( NameTableIndex );
 		(*FileWriter) << NumBits;
 	}
diff --git a/Engine/Source/Runtime/Engine/Private/RepLayout.cpp b/Engine/Source/Runtime/Engine/Private/RepLayout.cpp
--- a/Engine/Source/Runtime/Engine/Private/RepLayout.cpp
+++ b/Engine/Source/Runtime/Engine/Private/RepLayout.cpp
@@ -2545,6 +2545,7 @@ void FRepLayout::RebuildConditionalProperties( FRepState * RESTRICT	RepState, co
 	const bool bIsOwner		= RepFlags.bNetOwner		? true : false;
 	const bool bIsSimulated	= RepFlags.bNetSimulated	? true : false;
 	const bool bIsPhysics	= RepFlags.bRepPhysics		? true : false;
+	const bool bIsReplay	= RepFlags.bReplay			? true : false;
 
 	ConditionMap[COND_None]					= true;
 	ConditionMap[COND_InitialOnly]			= bIsInitial;
@@ -2558,6 +2559,7 @@ void FRepLayout::RebuildConditionalProperties( FRepState * RESTRICT	RepState, co
 	ConditionMap[COND_SimulatedOrPhysics]	= bIsSimulated || bIsPhysics;
 
 	ConditionMap[COND_InitialOrOwner]		= bIsInitial || bIsOwner;
+	ConditionMap[COND_ReplayOrOwner]		= bIsReplay || bIsOwner;
 
 	ConditionMap[COND_Custom]				= true;
 
diff --git a/Engine/Source/Runtime/Engine/Private/World.cpp b/Engine/Source/Runtime/Engine/Private/World.cpp
--- a/Engine/Source/Runtime/Engine/Private/World.cpp
+++ b/Engine/Source/Runtime/Engine/Private/World.cpp
@@ -36,6 +36,7 @@
 #include "SceneInterface.h"
 #include "Camera/CameraActor.h"
 #include "Engine/DemoNetDriver.h"
+#include "Engine/NetworkObjectList.h"
 #include "Layers/Layer.h"
 #include "GameFramework/GameMode.h"
 #include "GameFramework/GameState.h"
@@ -90,6 +91,28 @@ DEFINE_LOG_CATEGORY(LogSpawn);
 
 #define LOCTEXT_NAMESPACE "World"
 
+template<class Function>
+static void ForEachNetDriver(UWorld* const World, const Function InFunction)
+{
+	if (World == nullptr)
+	{
+		return;
+	}
+
+	const UGameInstance* const GameInstance = World->GetGameInstance();
+	if(GameInstance != nullptr)
+	{
+		FWorldContext* const Context = GameInstance->GetWorldContext();
+		if (Context != nullptr)
+		{
+			for (FNamedNetDriver& Driver : Context->ActiveNetDrivers)
+			{
+				InFunction(Driver.NetDriver);
+			}
+		}
+	}
+}
+
 // Deprecation warnings disabled to initialize bNoCollisionFail
 PRAGMA_DISABLE_DEPRECATION_WARNINGS
 
@@ -1111,7 +1134,13 @@ void UWorld::DestroyWorld( bool bInformEngineOfWorld, UWorld* NewWorld )
 	FlushLevelStreaming();
 	CleanupWorld(true, true, NewWorld);
 
-	check( NetworkActors.Num() == 0 );
+	ForEachNetDriver(this, [](UNetDriver* const Driver)
+	{
+		if (Driver != nullptr)
+		{
+			check(Driver->GetNetworkObjectList().GetObjects().Num() == 0);
+		}
+	});
 
 	// Tell the engine we are destroying the world.(unless we are asked not to)
 	if( ( GEngine ) && ( bInformEngineOfWorld == true ) )
@@ -3193,7 +3222,13 @@ void UWorld::CleanupWorld(bool bSessionEnded, bool bCleanupResources, UWorld* Ne
 		}
 	}
 
-	NetworkActors.Empty();
+	ForEachNetDriver(this, [](UNetDriver* const Driver)
+	{
+		if (Driver != nullptr)
+		{
+			Driver->GetNetworkObjectList().GetObjects().Reset();
+		}
+	});
 
 #if WITH_EDITOR
 	// If we're server traveling, we need to break the reference dependency here (caused by levelscript)
@@ -3388,7 +3423,14 @@ void UWorld::AddNetworkActor( AActor* Actor )
 		return;
 	}
 
-	NetworkActors.Add( Actor );
+	ForEachNetDriver(this, [Actor](UNetDriver* const Driver)
+	{
+		if (Driver != nullptr)
+		{
+			// Special case the demo net driver, since actors currently only have one associated NetDriverName.
+			Driver->GetNetworkObjectList().Add(Actor, Driver->NetDriverName);
+		}
+	});
 }
 
 void UWorld::RemoveNetworkActor( AActor* Actor )
@@ -3398,7 +3440,13 @@ void UWorld::RemoveNetworkActor( AActor* Actor )
 		return;
 	}
 
-	NetworkActors.Remove( Actor );
+	ForEachNetDriver(this, [Actor](UNetDriver* const Driver)
+	{
+		if (Driver != nullptr)
+		{
+			Driver->GetNetworkObjectList().GetObjects().Remove(Actor);
+		}
+	});
 }
 
 FDelegateHandle UWorld::AddOnActorSpawnedHandler( const FOnActorSpawned::FDelegate& InHandler )
diff --git a/Engine/Source/Runtime/Engine/Public/Net/NetworkProfiler.h b/Engine/Source/Runtime/Engine/Public/Net/NetworkProfiler.h
--- a/Engine/Source/Runtime/Engine/Public/Net/NetworkProfiler.h
+++ b/Engine/Source/Runtime/Engine/Public/Net/NetworkProfiler.h
@@ -27,16 +27,6 @@ private:
 	/** Version number to detect version mismatches.	*/
 	uint32	Version;
 
-	/** Offset in file for name table.					*/
-	uint32	NameTableOffset;
-	/** Number of name table entries.					*/
-	uint32	NameTableEntries;
-
-	/** Offset in file for address table.				*/
-	uint32	AddressTableOffset;
-	/** Number of address table entries.				*/
-	uint32	AddressTableEntries;
-
 	/** Tag, set via -networkprofiler=TAG				*/
 	FString Tag;
 	/** Game name, e.g. Example							*/
@@ -51,20 +41,6 @@ public:
 	/** Resets the header info for a new session.		*/
 	void Reset(const FURL& InURL);
 
-	/** Sets the final name table values.				*/
-	void SetNameTableValues(uint32 Offset, uint32 Entries)
-	{
-		NameTableOffset = Offset;
-		NameTableEntries = Entries;
-	}
-
-	/** Sets the final address table values.				*/
-	void SetAddressTableValues( uint32 Offset, uint32 Entries )
-	{
-		AddressTableOffset = Offset;
-		AddressTableEntries = Entries;
-	}
-
 	/** Returns the URL stored in the header.			*/
 	FString GetURL() const { return URL; }
 
@@ -106,9 +82,6 @@ private:
 	/** Array of unique addresses																	*/
 	TArray<uint64>							AddressArray;
 
-	/**	Temp file used for sessions in progress.													*/
-	FString									TempFileName;
-
 	/** Whether noticeable network traffic has occured in this session. Used to discard it.			*/
 	bool									bHasNoticeableNetworkTrafficOccured;
 	/** Whether tracking is enabled.																*/
diff --git a/Engine/Source/Runtime/GameplayAbilities/Private/AbilitySystemComponent.cpp b/Engine/Source/Runtime/GameplayAbilities/Private/AbilitySystemComponent.cpp
--- a/Engine/Source/Runtime/GameplayAbilities/Private/AbilitySystemComponent.cpp
+++ b/Engine/Source/Runtime/GameplayAbilities/Private/AbilitySystemComponent.cpp
@@ -1273,7 +1273,7 @@ void UAbilitySystemComponent::GetLifetimeReplicatedProps(TArray< FLifetimeProper
 	DOREPLIFETIME(UAbilitySystemComponent, ActiveGameplayEffects);
 	DOREPLIFETIME(UAbilitySystemComponent, ActiveGameplayCues);
 	
-	DOREPLIFETIME_CONDITION(UAbilitySystemComponent, ActivatableAbilities, COND_OwnerOnly);
+	DOREPLIFETIME_CONDITION(UAbilitySystemComponent, ActivatableAbilities, COND_ReplayOrOwner);
 	DOREPLIFETIME_CONDITION(UAbilitySystemComponent, BlockedAbilityBindings, COND_OwnerOnly)
 
 	DOREPLIFETIME(UAbilitySystemComponent, OwnerActor);
diff --git a/Engine/Source/Runtime/Online/OnlineFramework/Party/Private/Party.cpp b/Engine/Source/Runtime/Online/OnlineFramework/Party/Private/Party.cpp
--- a/Engine/Source/Runtime/Online/OnlineFramework/Party/Private/Party.cpp
+++ b/Engine/Source/Runtime/Online/OnlineFramework/Party/Private/Party.cpp
@@ -70,7 +70,14 @@ void UParty::OnLogoutComplete(int32 LocalUserNum, bool bWasSuccessful)
 			UPartyGameState* Party = PartyKeyValue.Value;
 			if (Party && Party->GetPartyId().IsValid())
 			{
-				UE_LOG(LogParty, Log, TEXT("[%s] Removed"), *(Party->GetPartyId()->ToDebugString()));
+				if (Party->GetPartyId().IsValid())
+				{
+					UE_LOG(LogParty, Log, TEXT("[%s] Removed"), *(Party->GetPartyId()->ToDebugString()));
+				}
+				else
+				{
+					UE_LOG(LogParty, Log, TEXT("[%s] Removed - Invalid party Id"));
+				}
 				Party->HandleRemovedFromParty(EMemberExitedReason::Left);
 			}
 		}
diff --git a/Engine/Source/Runtime/Online/OnlineSubsystemSteam/Private/SocketSubsystemSteam.cpp b/Engine/Source/Runtime/Online/OnlineSubsystemSteam/Private/SocketSubsystemSteam.cpp
--- a/Engine/Source/Runtime/Online/OnlineSubsystemSteam/Private/SocketSubsystemSteam.cpp
+++ b/Engine/Source/Runtime/Online/OnlineSubsystemSteam/Private/SocketSubsystemSteam.cpp
@@ -272,9 +272,14 @@ void FSocketSubsystemSteam::UnregisterConnection(USteamNetConnection* Connection
 	check(!Connection->bIsPassthrough);
 
 	FWeakObjectPtr ObjectPtr = Connection;
-	SteamConnections.RemoveSingleSwap(ObjectPtr);
-
-	if (Connection->RemoteAddr.IsValid())
+	int32 NumRemoved = SteamConnections.RemoveSingleSwap(ObjectPtr);
+
+	// Don't call P2PRemove again if we didn't actually remove a connection. This 
+	// will get called twice - once the connection is closed and when the connection
+	// is garbage collected. It's possible that the player who left rejoined before garbage
+	// collection runs (their connection object will be different), so P2PRemove would kick
+	// them from the session when it shouldn't.
+	if (NumRemoved > 0 && Connection->RemoteAddr.IsValid())
 	{
 		FInternetAddrSteam& SteamAddr = (FInternetAddrSteam&)(*Connection->RemoteAddr);
 		P2PRemove(SteamAddr.SteamId, SteamAddr.SteamChannel);
