Copying //UE4/Dev-Core to //UE4/Main

diff --git a/Engine/Source/Runtime/Core/Public/Stats/StatsData.h b/Engine/Source/Runtime/Core/Public/Stats/StatsData.h
--- a/Engine/Source/Runtime/Core/Public/Stats/StatsData.h
+++ b/Engine/Source/Runtime/Core/Public/Stats/StatsData.h
@@ -464,19 +464,17 @@ struct IItemFiler
 };
 
 
-
-// #YRX_STATS 2014-12-03 Separate stats thread state vs raw stats thread state?
 // #YRX_STATS 2014-03-21 Move metadata functionality into a separate class
 // 
 /**
  * Tracks stat state and history
  * GetLocalState() is a singleton to the state for stats being collected in this executable.
- *  Other instances can be used to load stats for visualization.
  */
 class CORE_API FStatsThreadState
 {
 	friend class FStatsThread;
 	friend struct FStatPacketArray;
+	friend struct FStatsReadFile;
 
 	/** Delegate that FStatsThreadState calls on the stats thread whenever we have a new frame. */
 	DECLARE_MULTICAST_DELEGATE_OneParam(FOnNewFrameHistory, int64);
@@ -526,6 +524,7 @@ private:
 	/** Generates a list of most memory expensive stats and dump to the log. */
 	void FindAndDumpMemoryExtensiveStats( FStatPacketArray &Frame );
 
+protected:
 	/** Called in response to SetLongName messages to update ShortNameToLongName and NotClearedEveryFrame **/
 	void FindOrAddMetaData(FStatMessage const& Item);
 
@@ -662,50 +661,6 @@ public:
 
 	/** Singleton to get the stats being collected by this executable. Can be only accessed from the stats thread. **/
 	static FStatsThreadState& GetLocalState();
-
-	/*-----------------------------------------------------------------------------
-		Stats file related functionality
-
-		#YRX_Stats: 2015-07-07 Maybe move to FStatsLoadedState
-	-----------------------------------------------------------------------------*/
-public:
-	friend struct FStatsReadFile;
-
-	/** Constructor to load stats from a file **/
-	FStatsThreadState( FString const& Filename );
-
-	/** Adds a frame worth of messages */
-	void AddCondensedMessages( TArray<FStatMessage>& CondensedMessages );
-
-	/** Marks this stats state as loaded. */
-	void MarkAsLoaded()
-	{
-		bWasLoaded = true;
-	}
-
-protected:
-	/** Internal method to update the internal metadata. **/
-	void ProcessMetaDataForLoad( TArray<FStatMessage>& CondensedMessages );
-
-	/** 
-	  * Internal method to place the data into the history, 
-	  * maintains the history based on the requested number of frames to keep in the history. 
-	  * The condensed messages are emplaced in the condensed history.
-	  */
-	void AddFrameFromCondensedMessages( TArray<FStatMessage>& CondensedMessages );
-
-	/** Internal method to scan the messages to find the current game/render thread frame. */
-	void AdvanceFrameForLoad( TArray<FStatMessage>& CondensedMessages );
-
-	/** Largest frame seen. Loaded stats only. **/
-	int64 MaxFrameSeen;
-
-	/** First frame seen. Loaded stats only. **/
-	int64 MinFrameSeen;
-
-	/** Valid frame computation is different if we just loaded these stats **/
-	bool bWasLoaded;
-
 };
 
 //@todo split header
diff --git a/Engine/Source/Runtime/Core/Public/Stats/StatsFile.h b/Engine/Source/Runtime/Core/Public/Stats/StatsFile.h
--- a/Engine/Source/Runtime/Core/Public/Stats/StatsFile.h
+++ b/Engine/Source/Runtime/Core/Public/Stats/StatsFile.h
@@ -6,6 +6,7 @@
 #if	STATS
 
 class FAsyncStatsWrite;
+struct FStatsReadFile;
 
 /**
 * Magic numbers for stats streams, this is for the first version.
@@ -537,6 +538,45 @@ protected:
 	Stats file reading functionality
 -----------------------------------------------------------------------------*/
 
+/** Tracks stat state and history for loaded stats. */
+class CORE_API FStatsLoadedState : public FStatsThreadState
+{
+	friend struct FStatsReadFile;
+
+public:
+	/** Default constructor. */
+	FStatsLoadedState()
+		: MaxFrameSeen( 0 )
+		, MinFrameSeen( -1 )
+	{}
+
+	/** Method to update the internal metadata. Metadata messages are removed from the array. */
+	void ProcessMetaDataAndLeaveDataOnly( TArray<FStatMessage>& CondensedMessages );
+
+	/**
+	* Method to place the data into the history,
+	* maintains the history based on the requested number of frames to keep in the history.
+	* The condensed messages are emplaced in the condensed history.
+	*/
+	void AddFrameFromCondensedMessages( TArray<FStatMessage>& CondensedMessages );
+
+	/** Return the oldest frame of data we have. **/
+	int64 GetOldestValidFrame() const;
+
+	/** Return the newest frame of data we have. **/
+	int64 GetLatestValidFrame() const;
+
+protected:
+	/** Internal method to scan the messages to find the current game/render thread frame. */
+	void AdvanceFrameForLoad( TArray<FStatMessage>& CondensedMessages );
+
+	/** Largest frame seen. **/
+	int64 MaxFrameSeen;
+
+	/** First frame seen. **/
+	int64 MinFrameSeen;
+};
+
 /**
  * Class for maintaining state of receiving a stream of stat messages
  */
@@ -637,7 +677,7 @@ public:
 			{
 				if( FNamesIndexMap.Contains( Index ) )
 				{
-					int32 MyIndex = FNamesIndexMap.FindChecked( Index );
+					const int32 MyIndex = FNamesIndexMap.FindChecked( Index );
 					TheFName = FName( MyIndex, MyIndex, 0 );
 				}
 				else
@@ -658,7 +698,7 @@ public:
 			}
 			if( FNamesIndexMap.Contains( Index ) )
 			{
-				int32 MyIndex = FNamesIndexMap.FindChecked( Index );
+				const int32 MyIndex = FNamesIndexMap.FindChecked( Index );
 				TheFName = FName( MyIndex, MyIndex, 0 );
 			}
 			else
@@ -802,8 +842,6 @@ enum class EStatsProcessingStage : int32
 
 };
 
-struct FStatsReadFile;
-
 /**
  *	Helper class used to read and process stats file on the async task.
  */
@@ -1091,7 +1129,7 @@ protected:
 
 protected:
 	/** Current state of the stats. Mostly for metadata. */
-	FStatsThreadState State;
+	FStatsLoadedState State;
 
 	/** Reading stream. */
 	FStatsReadStream Stream;
@@ -1164,6 +1202,9 @@ struct FCommandStatsFile
 	/** Stat StopFile. */
 	void Stop();
 
+	/** Test file by loading the last capture. */
+	void TestLastSaved();
+
 	/**
 	 * @return true, if any stats write command is active.
 	 */
diff --git a/Engine/Source/Runtime/Core/Public/UObject/DevObjectVersion.h b/Engine/Source/Runtime/Core/Public/UObject/DevObjectVersion.h
new file mode 100644
--- /dev/null
+++ b/Engine/Source/Runtime/Core/Public/UObject/DevObjectVersion.h
@@ -0,0 +1,26 @@
+// Copyright 1998-2015 Epic Games, Inc. All Rights Reserved.
+#pragma once
+
+#include "Serialization/CustomVersion.h"
+#include "BlueprintsObjectVersion.h"
+#include "BuildObjectVersion.h"
+#include "CoreObjectVersion.h"
+#include "EditorObjectVersion.h"
+#include "FrameworkObjectVersion.h"
+#include "MobileObjectVersion.h"
+#include "NetworkingObjectVersion.h"
+#include "OnlineObjectVersion.h"
+#include "PhysicsObjectVersion.h"
+#include "PlatformObjectVersion.h"
+#include "RenderingObjectVersion.h"
+#include "SequencerObjectVersion.h"
+#include "VRObjectVersion.h"
+
+class CORE_API FDevVersionRegistration :  public FCustomVersionRegistration
+{
+public:
+	FDevVersionRegistration(FGuid InKey, int32 Version, FName InFriendlyName);
+
+	/** Dumps all registered versions to log */
+	static void DumpVersionsToLog();
+};
\ No newline at end of file
diff --git a/Engine/Source/Runtime/Core/Public/UObject/PhysicsObjectVersion.h b/Engine/Source/Runtime/Core/Public/UObject/PhysicsObjectVersion.h
new file mode 100644
--- /dev/null
+++ b/Engine/Source/Runtime/Core/Public/UObject/PhysicsObjectVersion.h
@@ -0,0 +1,23 @@
+// Copyright 1998-2015 Epic Games, Inc. All Rights Reserved.
+#pragma once
+
+// Custom serialization version for changes made in Dev-Physics stream
+struct CORE_API FPhysicsObjectVersion
+{
+	enum Type
+	{
+		// Before any version changes were made
+		BeforeCustomVersionWasAdded = 0,
+
+
+		// -----<new versions can be added above this line>-------------------------------------------------
+		VersionPlusOne,
+		LatestVersion = VersionPlusOne - 1
+	};
+
+	// The GUID for this custom version number
+	const static FGuid GUID;
+
+private:
+	FPhysicsObjectVersion() {}
+};
diff --git a/Engine/Source/Runtime/CoreUObject/Private/Misc/RedirectCollector.cpp b/Engine/Source/Runtime/CoreUObject/Private/Misc/RedirectCollector.cpp
--- a/Engine/Source/Runtime/CoreUObject/Private/Misc/RedirectCollector.cpp
+++ b/Engine/Source/Runtime/CoreUObject/Private/Misc/RedirectCollector.cpp
@@ -64,6 +64,9 @@ void FRedirectCollector::OnStringAssetReferenceLoaded(const FString& InString)
 			{
 				ContainingPackageAndProperty.Property = FString::Printf(TEXT("%s:%s"), *ThreadContext.SerializedObject->GetPathName(), *Linker->GetSerializedProperty()->GetName());
 			}
+#if WITH_EDITORONLY_DATA
+			ContainingPackageAndProperty.bReferencedByEditorOnlyProperty = Linker->IsEditorOnlyPropertyOnTheStack();
+#endif
 		}
 	}
 	StringAssetReferences.AddUnique(InString, ContainingPackageAndProperty);
@@ -112,7 +115,7 @@ void FRedirectCollector::ResolveStringAssetReference(FString FilterPackage)
 
 			StringAssetRefFilenameStack.Push(RefFilenameAndProperty.Package);
 
-			UObject *Loaded = LoadObject<UObject>(NULL, *ToLoad, NULL, LOAD_None, NULL);
+			UObject *Loaded = LoadObject<UObject>(NULL, *ToLoad, NULL, RefFilenameAndProperty.bReferencedByEditorOnlyProperty ? LOAD_EditorOnly : LOAD_None, NULL);
 			StringAssetRefFilenameStack.Pop();
 
 			UObjectRedirector* Redirector = dynamic_cast<UObjectRedirector*>(Loaded);
diff --git a/Engine/Source/Runtime/CoreUObject/Private/Serialization/AsyncLoading.cpp b/Engine/Source/Runtime/CoreUObject/Private/Serialization/AsyncLoading.cpp
--- a/Engine/Source/Runtime/CoreUObject/Private/Serialization/AsyncLoading.cpp
+++ b/Engine/Source/Runtime/CoreUObject/Private/Serialization/AsyncLoading.cpp
@@ -328,6 +328,11 @@ void FAsyncLoadingThread::InitializeAsyncThread()
 
 void FAsyncLoadingThread::CancelAsyncLoadingInternal()
 {
+	{
+		FGCScopeGuard GCGuard;
+		FAsyncObjectsReferencer::Get().EmptyReferencedObjectsAndCancelLoading();
+	}
+
 	{
 		// Packages we haven't yet started processing.
 #if THREADSAFE_UOBJECTS
@@ -381,10 +386,6 @@ void FAsyncLoadingThread::CancelAsyncLoadingInternal()
 	QueuedPackagesCounter.Reset();
 
 	FUObjectThreadContext::Get().ObjLoaded.Empty();
-	{
-		FGCScopeGuard GCGuard;
-		FAsyncObjectsReferencer::Get().EmptyReferencedObjectsAndCancelLoading();
-	}
 
 	// Notify everyone streaming is canceled.
 	CancelLoadingEvent->Trigger();
@@ -2034,6 +2035,7 @@ void FAsyncPackage::Cancel()
 	{
 		CompletionCallbacks[CallbackIndex].Callback.ExecuteIfBound(Desc.Name, nullptr, Result);
 	}
+	bLoadHasFailed = true;
 	if (LinkerRoot)
 	{
 		if (Linker)
diff --git a/Engine/Source/Runtime/CoreUObject/Private/UObject/Class.cpp b/Engine/Source/Runtime/CoreUObject/Private/UObject/Class.cpp
--- a/Engine/Source/Runtime/CoreUObject/Private/UObject/Class.cpp
+++ b/Engine/Source/Runtime/CoreUObject/Private/UObject/Class.cpp
@@ -3477,6 +3477,9 @@ void UClass::Serialize( FArchive& Ar )
 			};
 			Ar.Seek(InterfacesStart + sizeof(NumInterfaces) + NumInterfaces * sizeof(FSerializedInterfaceReference));
 		}
+
+		// At this point token stream is most likely dirty too
+		SetTokenStreamMaybeDirty(true);
 	}
 
 	if (!Ar.IsIgnoringClassGeneratedByRef())
diff --git a/Engine/Source/Runtime/CoreUObject/Private/UObject/GarbageCollection.cpp b/Engine/Source/Runtime/CoreUObject/Private/UObject/GarbageCollection.cpp
--- a/Engine/Source/Runtime/CoreUObject/Private/UObject/GarbageCollection.cpp
+++ b/Engine/Source/Runtime/CoreUObject/Private/UObject/GarbageCollection.cpp
@@ -1094,6 +1094,8 @@ static FAutoConsoleVariableRef CVarFlushStreamingOnGC(
 	ECVF_Default
 	);
 
+bool VerifyClusterAssumptions(UObject* ClusterRootObject);
+
 /** 
  * Deletes all unreferenced objects, keeping objects that have any of the passed in KeepFlags set
  *
@@ -1152,7 +1154,7 @@ void CollectGarbageInternal(EObjectFlags KeepFlags, bool bPerformFullPurge)
 			// Don't require UGCObjectReferencer's references to adhere to the assumptions.
 			// Although we want the referencer itself to sit in the disregard for gc set, most of the objects
 			// it's referencing will not be in the root set.
-			if ((ObjectItem->HasAnyFlags(EInternalObjectFlags::ClusterRoot) || UObjectArray.IsDisregardForGC(Object)) && !Object->IsA(UGCObjectReferencer::StaticClass()))
+			if (UObjectArray.IsDisregardForGC(Object) && !Object->IsA(UGCObjectReferencer::StaticClass()))
 			{
 				// Serialize object with reference collector.
 				TArray<UObject*> CollectedReferences;
@@ -1176,6 +1178,13 @@ void CollectGarbageInternal(EObjectFlags KeepFlags, bool bPerformFullPurge)
 					}
 				}
 			}
+			else if (ObjectItem->HasAnyFlags(EInternalObjectFlags::ClusterRoot))
+			{
+				if (!VerifyClusterAssumptions(Object))
+				{
+					bShouldAssert = true;
+				}
+			}
 			// Assert if we encountered any objects breaking implicit assumptions.
 			if( bShouldAssert )
 			{
diff --git a/Engine/Source/Runtime/CoreUObject/Private/UObject/LinkerLoad.cpp b/Engine/Source/Runtime/CoreUObject/Private/UObject/LinkerLoad.cpp
--- a/Engine/Source/Runtime/CoreUObject/Private/UObject/LinkerLoad.cpp
+++ b/Engine/Source/Runtime/CoreUObject/Private/UObject/LinkerLoad.cpp
@@ -3723,7 +3723,7 @@ UObject* FLinkerLoad::CreateExport( int32 Index )
 					// Do this for all subobjects created in the native constructor.
 					FUObjectThreadContext::Get().ObjLoaded.AddUnique(Export.Object);
 					if (!Export.Object->HasAnyFlags(RF_LoadCompleted) &&
-						(Export.Object->HasAnyFlags(RF_DefaultSubObject) || (ThisParent && ThisParent->HasAnyFlags(RF_ClassDefaultObject))))
+						(Export.Object->HasAnyFlags(RF_DefaultSubObject) || (ThisParent && ThisParent->IsTemplate(RF_ClassDefaultObject))))
 					{
 						Export.Object->SetFlags(RF_NeedLoad | RF_NeedPostLoad | RF_NeedPostLoadSubobjects | RF_WasLoaded);
 					}
diff --git a/Engine/Source/Runtime/CoreUObject/Private/UObject/SavePackage.cpp b/Engine/Source/Runtime/CoreUObject/Private/UObject/SavePackage.cpp
--- a/Engine/Source/Runtime/CoreUObject/Private/UObject/SavePackage.cpp
+++ b/Engine/Source/Runtime/CoreUObject/Private/UObject/SavePackage.cpp
@@ -706,10 +706,9 @@ void FArchiveSaveTagExports::ProcessTaggedObjects()
 			{
 				Obj->GetClass()->SerializeDefaultObject(Obj, *this);
 			}
-			else
-			{
-				Obj->Serialize(*this);
-			}
+			// In the CDO case the above would serialize most of the references, including transient properties
+			// but we still want to serialize the object using the normal path to collect all custom versions it might be using.
+			Obj->Serialize(*this);
 		}
 
 		CurrentlyTaggedObjects.Empty(ArrayPreSize);
diff --git a/Engine/Source/Runtime/CoreUObject/Private/UObject/UObjectClusters.cpp b/Engine/Source/Runtime/CoreUObject/Private/UObject/UObjectClusters.cpp
--- a/Engine/Source/Runtime/CoreUObject/Private/UObject/UObjectClusters.cpp
+++ b/Engine/Source/Runtime/CoreUObject/Private/UObject/UObjectClusters.cpp
@@ -344,13 +344,14 @@ public:
 /**
  * Specialized FReferenceCollector that uses FClusterReferenceProcessor to construct the cluster
  */
-class FClusterCollector : public FReferenceCollector
+template <class TProcessor>
+class TClusterCollector : public FReferenceCollector
 {
-	FClusterReferenceProcessor& Processor;
+	TProcessor& Processor;
 	TArray<UObject*>& ObjectArray;
 
 public:
-	FClusterCollector(FClusterReferenceProcessor& InProcessor, TArray<UObject*>& InObjectArray)
+	TClusterCollector(TProcessor& InProcessor, TArray<UObject*>& InObjectArray)
 		: Processor(InProcessor)
 		, ObjectArray(InObjectArray)
 	{
@@ -391,7 +392,7 @@ void UObjectBaseUtility::CreateCluster()
 
 	// Collect all objects referenced by cluster root and by all objects it's referencing
 	FClusterReferenceProcessor Processor(InternalIndex, *Cluster);
-	TFastReferenceCollector<FClusterReferenceProcessor, FClusterCollector, FClusterArrayPool> ReferenceCollector(Processor, FClusterArrayPool::Get());
+	TFastReferenceCollector<FClusterReferenceProcessor, TClusterCollector<FClusterReferenceProcessor>, FClusterArrayPool> ReferenceCollector(Processor, FClusterArrayPool::Get());
 	TArray<UObject*> ObjectsToProcess;
 	ObjectsToProcess.Add(static_cast<UObject*>(this));
 	ReferenceCollector.CollectReferences(ObjectsToProcess, true);
@@ -409,4 +410,102 @@ void UObjectBaseUtility::CreateCluster()
 	{
 		delete Cluster;
 	}
-}
\ No newline at end of file
+}
+
+
+/**
+* Handles UObject references found by TFastReferenceCollector
+*/
+class FClusterVerifyReferenceProcessor
+{	
+	const UObject* const ClusterRootObject;
+	const int32 ClusterRootIndex;
+	volatile bool bIsRunningMultithreaded;
+	bool bFailed;
+	TSet<UObject*> ProcessedObjects;
+
+public:
+
+	FClusterVerifyReferenceProcessor(UObject* InClusterRootObject)		
+		: ClusterRootObject(InClusterRootObject)
+		, ClusterRootIndex(GUObjectArray.ObjectToIndex(InClusterRootObject))
+		, bIsRunningMultithreaded(false)
+		, bFailed(false)
+	{
+	}
+
+	bool NoExternalReferencesFound() const
+	{
+		return !bFailed;
+	}
+
+	FORCEINLINE int32 GetMinDesiredObjectsPerSubTask() const
+	{
+		// We're not running the processor in parallel when creating clusters
+		return 0;
+	}
+
+	FORCEINLINE volatile bool IsRunningMultithreaded() const
+	{
+		// This should always be false
+		return bIsRunningMultithreaded;
+	}
+
+	FORCEINLINE void SetIsRunningMultithreaded(bool bIsParallel)
+	{
+		check(!bIsParallel);
+		bIsRunningMultithreaded = bIsParallel;
+	}
+
+	void UpdateDetailedStats(UObject* CurrentObject, uint32 DeltaCycles)
+	{
+	}
+
+	void LogDetailedStatsSummary()
+	{
+	}
+
+	/**
+	* Handles UObject reference from the token stream. Performance is critical here so we're FORCEINLINING this function.
+	*
+	* @param ObjectsToSerialize An array of remaining objects to serialize (Obj must be added to it if Obj can be added to cluster)
+	* @param ReferencingObject Object referencing the object to process.
+	* @param TokenIndex Index to the token stream where the reference was found.
+	* @param bAllowReferenceElimination True if reference elimination is allowed (ignored when constructing clusters).
+	*/
+	FORCEINLINE void HandleTokenStreamObjectReference(TArray<UObject*>& ObjectsToSerialize, UObject* ReferencingObject, UObject*& Object, const int32 TokenIndex, bool bAllowReferenceElimination)
+	{
+		if (Object && !ProcessedObjects.Contains(Object))
+		{
+			ProcessedObjects.Add(Object);
+			FUObjectItem* ObjectItem = GUObjectArray.ObjectToObjectItem(Object);
+			if (ObjectItem->GetOwnerIndex() == 0)
+			{
+				if (!ObjectItem->HasAnyFlags(EInternalObjectFlags::ClusterRoot|EInternalObjectFlags::RootSet) && !GUObjectArray.IsDisregardForGC(Object))
+				{
+					UE_LOG(LogObj, Warning, TEXT("Object %s from cluster %s is referencing %s which is not part of root set or cluster."),
+						*ReferencingObject->GetFullName(),
+						*ClusterRootObject->GetFullName(),
+						*Object->GetFullName());
+					bFailed = true;
+				}
+			}
+			else if (ObjectItem->GetOwnerIndex() == ClusterRootIndex)
+			{
+				// If this object belongs to the current cluster, keep processing its references. Otherwise ignore it as it will be processed by its cluster
+				ObjectsToSerialize.Add(Object);
+			}
+		}
+	}
+};
+
+bool VerifyClusterAssumptions(UObject* ClusterRootObject)
+{
+	// Collect all objects referenced by cluster root and by all objects it's referencing
+	FClusterVerifyReferenceProcessor Processor(ClusterRootObject);
+	TFastReferenceCollector<FClusterVerifyReferenceProcessor, TClusterCollector<FClusterVerifyReferenceProcessor>, FClusterArrayPool> ReferenceCollector(Processor, FClusterArrayPool::Get());
+	TArray<UObject*> ObjectsToProcess;
+	ObjectsToProcess.Add(ClusterRootObject);
+	ReferenceCollector.CollectReferences(ObjectsToProcess, true);
+	return Processor.NoExternalReferencesFound();
+}
diff --git a/Engine/Source/Runtime/CoreUObject/Private/UObject/UObjectGlobals.cpp b/Engine/Source/Runtime/CoreUObject/Private/UObject/UObjectGlobals.cpp
--- a/Engine/Source/Runtime/CoreUObject/Private/UObject/UObjectGlobals.cpp
+++ b/Engine/Source/Runtime/CoreUObject/Private/UObject/UObjectGlobals.cpp
@@ -625,12 +625,12 @@ const FString* GetIniFilenameFromObjectsReference(const FString& Name)
 //
 // Resolve a package and name.
 //
-bool ResolveName( UObject*& InPackage, FString& InOutName, bool Create, bool Throw )
-	{
+bool ResolveName(UObject*& InPackage, FString& InOutName, bool Create, bool Throw, uint32 LoadFlags /*= LOAD_None*/)
+{
 	const FString* IniFilename = GetIniFilenameFromObjectsReference(InOutName);
 
 	if (IniFilename && InOutName.Contains(TEXT("."), ESearchCase::CaseSensitive))
-		{
+	{
 		InOutName = ResolveIniObjectsReference(InOutName, IniFilename, Throw);
 	}
 
@@ -692,7 +692,7 @@ bool ResolveName( UObject*& InPackage, FString& InOutName, bool Create, bool Thr
 			InPackage = StaticFindObjectFast(UPackage::StaticClass(), InPackage, *PartialName);
 			if (!ScriptPackageName && !InPackage)
 			{
-				InPackage = LoadPackage(dynamic_cast<UPackage*>(InPackage), *PartialName, 0);
+				InPackage = LoadPackage(dynamic_cast<UPackage*>(InPackage), *PartialName, LoadFlags);
 			}
 			if (!InPackage)
 			{
@@ -778,7 +778,7 @@ UObject* StaticLoadObjectInternal(UClass* ObjectClass, UObject* InOuter, const T
 	const bool bContainsObjectName = !!FCString::Strstr(InName, TEXT("."));
 
 	// break up the name into packages, returning the innermost name and its outer
-	ResolveName(InOuter, StrName, true, true);
+	ResolveName(InOuter, StrName, true, true, LoadFlags & LOAD_EditorOnly);
 	if (InOuter)
 	{
 		// If we have a full UObject name then attempt to find the object in memory first,
@@ -836,7 +836,7 @@ UObject* StaticLoadObjectInternal(UClass* ObjectClass, UObject* InOuter, const T
 		Result = StaticLoadObjectInternal(ObjectClass, InOuter, *StrName, Filename, LoadFlags, Sandbox, bAllowObjectReconciliation);
 	}
 #if WITH_EDITORONLY_DATA
-	else if (Result)
+	else if (Result && !(LoadFlags & LOAD_EditorOnly))
 	{
 		Result->GetOutermost()->SetLoadedByEditorPropertiesOnly(false);
 	}
@@ -853,7 +853,7 @@ UObject* StaticLoadObject(UClass* ObjectClass, UObject* InOuter, const TCHAR* In
 	if (!Result)
 	{
 		FString ObjectName = InName;
-		ResolveName(InOuter, ObjectName, true, true);
+		ResolveName(InOuter, ObjectName, true, true, LoadFlags & LOAD_EditorOnly);
 
 		// we haven't created or found the object, error
 		FFormatNamedArguments Arguments;
@@ -1076,7 +1076,7 @@ UPackage* LoadPackageInternal(UPackage* InOuter, const TCHAR* InLongPackageName,
 		};
 
 #if WITH_EDITORONLY_DATA
-		if (!(LoadFlags & LOAD_IsVerifying) &&
+		if (!(LoadFlags & (LOAD_IsVerifying|LOAD_EditorOnly)) &&
 			(!ImportLinker || !ImportLinker->GetSerializedProperty() || !ImportLinker->GetSerializedProperty()->IsEditorOnlyProperty()))
 		{
 			// If this package hasn't been loaded as part of import verification and there's no import linker or the
diff --git a/Engine/Source/Runtime/CoreUObject/Public/Misc/RedirectCollector.h b/Engine/Source/Runtime/CoreUObject/Public/Misc/RedirectCollector.h
--- a/Engine/Source/Runtime/CoreUObject/Public/Misc/RedirectCollector.h
+++ b/Engine/Source/Runtime/CoreUObject/Public/Misc/RedirectCollector.h
@@ -48,10 +48,11 @@ public:
 	/** Helper struct for string asset reference tracking */
 	struct FPackagePropertyPair
 	{
-		FPackagePropertyPair() {}
+		FPackagePropertyPair() : bReferencedByEditorOnlyProperty(false) {}
 		FPackagePropertyPair(const FString& InPackage, const FString& InProperty)
 		: Package(InPackage)
 		, Property(InProperty)
+		, bReferencedByEditorOnlyProperty(false)
 		{}
 
 		bool operator==(const FPackagePropertyPair& Other) const
@@ -62,6 +63,7 @@ public:
 
 		FString Package;
 		FString Property;
+		bool bReferencedByEditorOnlyProperty;
 	};
 
 	/**
diff --git a/Engine/Source/Runtime/CoreUObject/Public/UObject/ObjectBase.h b/Engine/Source/Runtime/CoreUObject/Public/UObject/ObjectBase.h
--- a/Engine/Source/Runtime/CoreUObject/Public/UObject/ObjectBase.h
+++ b/Engine/Source/Runtime/CoreUObject/Public/UObject/ObjectBase.h
@@ -44,7 +44,7 @@ enum ELoadFlags
 	LOAD_None						= 0x00000000,	// No flags.
 	LOAD_SeekFree					= 0x00000001,	// Loads the package via the seek free loading path/ reader
 	LOAD_NoWarn						= 0x00000002,	// Don't display warning if load fails.
-//	LOAD_Unused						= 0x00000004,
+	LOAD_EditorOnly					= 0x00000004, // Load for editor-only purposes and by editor-only code
 //	LOAD_Unused						= 0x00000008,
 	LOAD_Verify						= 0x00000010,	// Only verify existance; don't actually load.
 	LOAD_AllowDll					= 0x00000020,	// Allow plain DLLs.
diff --git a/Engine/Source/Runtime/CoreUObject/Public/UObject/UObjectGlobals.h b/Engine/Source/Runtime/CoreUObject/Public/UObject/UObjectGlobals.h
--- a/Engine/Source/Runtime/CoreUObject/Public/UObject/UObjectGlobals.h
+++ b/Engine/Source/Runtime/CoreUObject/Public/UObject/UObjectGlobals.h
@@ -152,7 +152,7 @@ COREUOBJECT_API const FString* GetIniFilenameFromObjectsReference(const FString&
  */
 COREUOBJECT_API FString ResolveIniObjectsReference(const FString& ObjectReference, const FString* IniFilename, bool bThrow = false);
 
-COREUOBJECT_API bool ResolveName( UObject*& Outer, FString& ObjectsReferenceString, bool Create, bool Throw );
+COREUOBJECT_API bool ResolveName(UObject*& Outer, FString& ObjectsReferenceString, bool Create, bool Throw, uint32 LoadFlags = LOAD_None);
 COREUOBJECT_API void SafeLoadError( UObject* Outer, uint32 LoadFlags, const TCHAR* ErrorMessage);
 
 /**
diff --git a/Engine/Source/Runtime/Engine/Classes/Engine/Engine.h b/Engine/Source/Runtime/Engine/Classes/Engine/Engine.h
--- a/Engine/Source/Runtime/Engine/Classes/Engine/Engine.h
+++ b/Engine/Source/Runtime/Engine/Classes/Engine/Engine.h
@@ -1670,6 +1670,7 @@ public:
 
 	bool HandleRecordAnimationCommand(UWorld* InWorld, const TCHAR* InStr, FOutputDevice& Ar);
 	bool HandleStopRecordAnimationCommand(UWorld* InWorld, const TCHAR* InStr, FOutputDevice& Ar);
+	bool HandleShowLogCommand( const TCHAR* Cmd, FOutputDevice& Ar );
 
 	// Only compile in when STATS is set
 #if STATS
@@ -1677,10 +1678,12 @@ public:
 #endif
 
 	// Compile in Debug or Development
-#if !(UE_BUILD_SHIPPING || UE_BUILD_TEST)
-#if WITH_HOT_RELOAD
+#if !(UE_BUILD_SHIPPING || UE_BUILD_TEST) && WITH_HOT_RELOAD
 	bool HandleHotReloadCommand( const TCHAR* Cmd, FOutputDevice& Ar );
-#endif
+#endif // !(UE_BUILD_SHIPPING || UE_BUILD_TEST) && WITH_HOT_RELOAD
+
+	// Compile in Debug, Development, and Test
+#if !UE_BUILD_SHIPPING
 	bool HandleDumpConsoleCommandsCommand( const TCHAR* Cmd, FOutputDevice& Ar, UWorld* InWorld );
 	bool HandleShowMaterialDrawEventsCommand( const TCHAR* Cmd, FOutputDevice& Ar );
 	bool HandleDumpAvailableResolutionsCommand( const TCHAR* Cmd, FOutputDevice& Ar );
@@ -1692,7 +1695,7 @@ public:
 	bool HandleFlushIOManagerCommand( const TCHAR* Cmd, FOutputDevice& Ar );						// Smedis
 	bool HandleToggleRenderingThreadCommand( const TCHAR* Cmd, FOutputDevice& Ar );
 	bool HandleToggleRHIThreadCommand( const TCHAR* Cmd, FOutputDevice& Ar );
-	bool HandleToggleAsyncComputeCommand(const TCHAR* Cmd, FOutputDevice& Ar);
+	bool HandleToggleAsyncComputeCommand( const TCHAR* Cmd, FOutputDevice& Ar );
 	bool HandleRecompileShadersCommand( const TCHAR* Cmd, FOutputDevice& Ar );
 	bool HandleRecompileGlobalShadersCommand( const TCHAR* Cmd, FOutputDevice& Ar );
 	bool HandleDumpShaderStatsCommand( const TCHAR* Cmd, FOutputDevice& Ar );
@@ -1703,11 +1706,6 @@ public:
 	bool HandleShaderComplexityCommand( const TCHAR* Cmd, FOutputDevice& Ar );
 	bool HandleFreezeRenderingCommand( const TCHAR* Cmd, FOutputDevice& Ar, UWorld* InWorld );
 	bool HandleShowSelectedLightmapCommand( const TCHAR* Cmd, FOutputDevice& Ar );
-#endif // !(UE_BUILD_SHIPPING || UE_BUILD_TEST)
-
-	// Compile in Debug, Development, and Test
-#if !UE_BUILD_SHIPPING
-	bool HandleShowLogCommand( const TCHAR* Cmd, FOutputDevice& Ar );
 	bool HandleStartFPSChartCommand( const TCHAR* Cmd, FOutputDevice& Ar );
 	bool HandleStopFPSChartCommand( const TCHAR* Cmd, FOutputDevice& Ar, UWorld* InWorld );
 	bool HandleDumpLevelScriptActorsCommand( UWorld* InWorld, const TCHAR* Cmd, FOutputDevice& Ar );
@@ -2906,9 +2904,9 @@ private:
 	 * @param ViewLocation The world space view location.
 	 * @param ViewRotation The world space view rotation.
 	 */
-#if !(UE_BUILD_SHIPPING || UE_BUILD_TEST)
+#if !UE_BUILD_SHIPPING
 	int32 RenderStatVersion(UWorld* World, FViewport* Viewport, FCanvas* Canvas, int32 X, int32 Y, const FVector* ViewLocation = nullptr, const FRotator* ViewRotation = nullptr);
-#endif
+#endif // !UE_BUILD_SHIPPING
 	int32 RenderStatFPS(UWorld* World, FViewport* Viewport, FCanvas* Canvas, int32 X, int32 Y, const FVector* ViewLocation = nullptr, const FRotator* ViewRotation = nullptr);
 	int32 RenderStatHitches(UWorld* World, FViewport* Viewport, FCanvas* Canvas, int32 X, int32 Y, const FVector* ViewLocation = nullptr, const FRotator* ViewRotation = nullptr);
 	int32 RenderStatSummary(UWorld* World, FViewport* Viewport, FCanvas* Canvas, int32 X, int32 Y, const FVector* ViewLocation = nullptr, const FRotator* ViewRotation = nullptr);
@@ -2917,12 +2915,12 @@ private:
 	int32 RenderStatLevels(UWorld* World, FViewport* Viewport, FCanvas* Canvas, int32 X, int32 Y, const FVector* ViewLocation = nullptr, const FRotator* ViewRotation = nullptr);
 	int32 RenderStatLevelMap(UWorld* World, FViewport* Viewport, FCanvas* Canvas, int32 X, int32 Y, const FVector* ViewLocation = nullptr, const FRotator* ViewRotation = nullptr);
 	int32 RenderStatUnit(UWorld* World, FViewport* Viewport, FCanvas* Canvas, int32 X, int32 Y, const FVector* ViewLocation = nullptr, const FRotator* ViewRotation = nullptr);
-#if !(UE_BUILD_SHIPPING || UE_BUILD_TEST)
+#if !UE_BUILD_SHIPPING
 	int32 RenderStatReverb(UWorld* World, FViewport* Viewport, FCanvas* Canvas, int32 X, int32 Y, const FVector* ViewLocation = nullptr, const FRotator* ViewRotation = nullptr);
 	int32 RenderStatSoundMixes(UWorld* World, FViewport* Viewport, FCanvas* Canvas, int32 X, int32 Y, const FVector* ViewLocation = nullptr, const FRotator* ViewRotation = nullptr);
 	int32 RenderStatSoundWaves(UWorld* World, FViewport* Viewport, FCanvas* Canvas, int32 X, int32 Y, const FVector* ViewLocation = nullptr, const FRotator* ViewRotation = nullptr);
 	int32 RenderStatSoundCues(UWorld* World, FViewport* Viewport, FCanvas* Canvas, int32 X, int32 Y, const FVector* ViewLocation = nullptr, const FRotator* ViewRotation = nullptr);
-#endif
+#endif // !UE_BUILD_SHIPPING
 	int32 RenderStatSounds(UWorld* World, FViewport* Viewport, FCanvas* Canvas, int32 X, int32 Y, const FVector* ViewLocation = nullptr, const FRotator* ViewRotation = nullptr);
 	int32 RenderStatAI(UWorld* World, FViewport* Viewport, FCanvas* Canvas, int32 X, int32 Y, const FVector* ViewLocation = nullptr, const FRotator* ViewRotation = nullptr);
 #if STATS
diff --git a/Engine/Source/Runtime/Engine/Private/Animation/AnimSequence.cpp b/Engine/Source/Runtime/Engine/Private/Animation/AnimSequence.cpp
--- a/Engine/Source/Runtime/Engine/Private/Animation/AnimSequence.cpp
+++ b/Engine/Source/Runtime/Engine/Private/Animation/AnimSequence.cpp
@@ -4700,11 +4700,7 @@ float UAnimSequence::GetNextMatchingPosFromMarkerSyncPos(const FMarkerSyncAnimPo
 	AnimNotify& subclasses
 -----------------------------------------------------------------------------*/
 
-
-
-
-////////////
-#if !(UE_BUILD_SHIPPING || UE_BUILD_TEST)
+#if !UE_BUILD_SHIPPING
 
 void GatherAnimSequenceStats(FOutputDevice& Ar)
 {
@@ -4794,5 +4790,7 @@ void GatherAnimSequenceStats(FOutputDevice& Ar)
 	Ar.Logf( TEXT("Approx Compressed Memory: %i bytes"), GlobalApproxCompressedSize);
 	Ar.Logf( TEXT("Approx Key Data Memory: %i bytes"), GlobalApproxKeyDataSize);
 }
-#endif
+
+#endif // !UE_BUILD_SHIPPING
+
 #undef LOCTEXT_NAMESPACE 
\ No newline at end of file
diff --git a/Engine/Source/Runtime/Engine/Private/Collision/Collision.cpp b/Engine/Source/Runtime/Engine/Private/Collision/Collision.cpp
--- a/Engine/Source/Runtime/Engine/Private/Collision/Collision.cpp
+++ b/Engine/Source/Runtime/Engine/Private/Collision/Collision.cpp
@@ -518,7 +518,7 @@ bool FSeparatingAxisPointCheck::FindSeparatingAxisGeneric()
 
 
 
-#if !(UE_BUILD_SHIPPING || UE_BUILD_TEST)
+#if !UE_BUILD_SHIPPING
 
 DECLARE_LOG_CATEGORY_EXTERN(LogCollisionCommands, Log, All);
 DEFINE_LOG_CATEGORY(LogCollisionCommands);
@@ -952,7 +952,7 @@ namespace CollisionResponseConsoleCommands
 
 }
 
-#endif //!(UE_BUILD_SHIPPING || UE_BUILD_TEST)
+#endif //!UE_BUILD_SHIPPING
 
 
 
diff --git a/Engine/Source/Runtime/Engine/Private/ContentStreaming.cpp b/Engine/Source/Runtime/Engine/Private/ContentStreaming.cpp
--- a/Engine/Source/Runtime/Engine/Private/ContentStreaming.cpp
+++ b/Engine/Source/Runtime/Engine/Private/ContentStreaming.cpp
@@ -4422,7 +4422,7 @@ bool FStreamingManagerTexture::HandleListStreamingTexturesReportCommand( const T
 
 #endif
 
-#if !(UE_BUILD_SHIPPING || UE_BUILD_TEST)
+#if !UE_BUILD_SHIPPING
 bool FStreamingManagerTexture::HandleLightmapStreamingFactorCommand( const TCHAR* Cmd, FOutputDevice& Ar )
 {
 	FString FactorString(FParse::Token(Cmd, 0));
@@ -4696,7 +4696,7 @@ bool FStreamingManagerTexture::HandleInvestigateTextureCommand( const TCHAR* Cmd
 	}
 	return true;
 }
-#endif // !(UE_BUILD_SHIPPING || UE_BUILD_TEST)
+#endif // !UE_BUILD_SHIPPING
 /**
  * Allows the streaming manager to process exec commands.
  * @param InWorld	world contexxt
@@ -4730,7 +4730,7 @@ bool FStreamingManagerTexture::Exec( UWorld* InWorld, const TCHAR* Cmd, FOutputD
 		return HandleListStreamingTexturesReportCommand( Cmd, Ar );
 	}
 #endif
-#if !(UE_BUILD_SHIPPING || UE_BUILD_TEST)
+#if !UE_BUILD_SHIPPING
 	if (FParse::Command(&Cmd,TEXT("LightmapStreamingFactor")))
 	{
 		return HandleLightmapStreamingFactorCommand( Cmd, Ar );
@@ -4783,7 +4783,7 @@ bool FStreamingManagerTexture::Exec( UWorld* InWorld, const TCHAR* Cmd, FOutputD
 	{
 		return HandleInvestigateTextureCommand( Cmd, Ar );
 	}
-#endif // !(UE_BUILD_SHIPPING || UE_BUILD_TEST)
+#endif // !UE_BUILD_SHIPPING
 
 	return false;
 }
@@ -4791,7 +4791,7 @@ bool FStreamingManagerTexture::Exec( UWorld* InWorld, const TCHAR* Cmd, FOutputD
 void FStreamingManagerTexture::DumpTextureGroupStats( bool bDetailedStats )
 {
 	bTriggerDumpTextureGroupStats = false;
-#if !(UE_BUILD_SHIPPING || UE_BUILD_TEST)
+#if !UE_BUILD_SHIPPING
 	struct FTextureGroupStats
 	{
 		FTextureGroupStats()
@@ -4987,7 +4987,7 @@ void FStreamingManagerTexture::DumpTextureGroupStats( bool bDetailedStats )
 	//@TODO: Calculate memory usage for non-pool textures properly!
 //	UE_LOG(LogContentStreaming, Log,  TEXT("%34s: NumTextures=%4d, Current=%7.1f KB"), TEXT("Non-streaming pool textures"), NumNonStreamingPoolTextures, NonStreamingPoolSize/1024.0f );
 //	UE_LOG(LogContentStreaming, Log,  TEXT("%34s: NumTextures=%4d, Current=%7.1f KB"), TEXT("Non-streaming non-pool textures"), NumNonStreamingTextures, NonStreamingSize/1024.0f );
-#endif
+#endif // !UE_BUILD_SHIPPING
 }
 
 /**
@@ -4999,7 +4999,7 @@ void FStreamingManagerTexture::DumpTextureGroupStats( bool bDetailedStats )
 void FStreamingManagerTexture::InvestigateTexture( const FString& InInvestigateTextureName )
 {
 	bTriggerInvestigateTexture = false;
-#if !(UE_BUILD_SHIPPING || UE_BUILD_TEST)
+#if !UE_BUILD_SHIPPING
 	for ( int32 TextureIndex=0; TextureIndex < StreamingTextures.Num(); ++TextureIndex )
 	{
 		FStreamingTexture& StreamingTexture = StreamingTextures[TextureIndex];
@@ -5191,12 +5191,12 @@ void FStreamingManagerTexture::InvestigateTexture( const FString& InInvestigateT
 			}
 		}
 	}
-#endif
+#endif // !UE_BUILD_SHIPPING
 }
 
 void FStreamingManagerTexture::DumpTextureInstances( const UPrimitiveComponent* Primitive, FSpawnedPrimitiveData& PrimitiveData, UTexture2D* Texture2D )
 {
-#if !(UE_BUILD_SHIPPING || UE_BUILD_TEST)
+#if !UE_BUILD_SHIPPING
 	const TCHAR* TypeString = (PrimitiveData.DynamicType == DPT_Level) ? TEXT("Dynamic (level)") : TEXT("Dynamic (spawned)");
 	for ( int32 InstanceIndex=0; InstanceIndex < PrimitiveData.TextureInstances.Num(); ++InstanceIndex )
 	{
@@ -5256,7 +5256,7 @@ void FStreamingManagerTexture::DumpTextureInstances( const UPrimitiveComponent*
 				*Primitive->GetReadableName() );
 		}
 	}
-#endif
+#endif // !UE_BUILD_SHIPPING
 }
 
 
diff --git a/Engine/Source/Runtime/Engine/Private/LocalPlayer.cpp b/Engine/Source/Runtime/Engine/Private/LocalPlayer.cpp
--- a/Engine/Source/Runtime/Engine/Private/LocalPlayer.cpp
+++ b/Engine/Source/Runtime/Engine/Private/LocalPlayer.cpp
@@ -27,7 +27,7 @@
 
 DEFINE_LOG_CATEGORY(LogPlayerManagement);
 
-#if !(UE_BUILD_SHIPPING || UE_BUILD_TEST)
+#if !UE_BUILD_SHIPPING
 
 static TAutoConsoleVariable<int32> CVarViewportTest(
 	TEXT("r.ViewportTest"),
@@ -37,7 +37,7 @@ static TAutoConsoleVariable<int32> CVarViewportTest(
 	TEXT("1..7: Various Configuations"),
 	ECVF_RenderThreadSafe);
 
-#endif
+#endif // !UE_BUILD_SHIPPING
 
 DECLARE_CYCLE_STAT(TEXT("CalcSceneView"), STAT_CalcSceneView, STATGROUP_Engine);
 
diff --git a/Engine/Source/Runtime/Engine/Private/UnrealEngine.cpp b/Engine/Source/Runtime/Engine/Private/UnrealEngine.cpp
--- a/Engine/Source/Runtime/Engine/Private/UnrealEngine.cpp
+++ b/Engine/Source/Runtime/Engine/Private/UnrealEngine.cpp
@@ -166,13 +166,13 @@
  */
 ENGINE_API bool GShowDebugSelectedLightmap = false;
 
-#if !(UE_BUILD_SHIPPING || UE_BUILD_TEST)
+#if !UE_BUILD_SHIPPING
 	/**
 	 * true if we debug material names with SCOPED_DRAW_EVENT.
 	 * Toggle with "ShowMaterialDrawEvents" console command.
 	 */
 	bool GShowMaterialDrawEvents = false;
-#endif
+#endif // !UE_BUILD_SHIPPING
 
 ENGINE_API uint32 GGPUFrameTime = 0;
 
@@ -195,14 +195,14 @@
 	TEXT("     1920x1080wm for windowed mirror")
 	);
 
-#if !(UE_BUILD_SHIPPING || UE_BUILD_TEST)
+#if !UE_BUILD_SHIPPING
 static TAutoConsoleVariable<float> CVarSetOverrideFPS(
 	TEXT("t.OverrideFPS"),
 	0.0f,
 	TEXT("This allows to override the frame time measurement with a fixed fps number (game can run faster or slower).\n")
 	TEXT("<=0:off, in frames per second, e.g. 60"),
 	ECVF_Cheat);
-#endif // !(UE_BUILD_SHIPPING || UE_BUILD_TEST)
+#endif // !UE_BUILD_SHIPPING
 
 /** Enum entries represent index to global object referencer stored in UGameEngine */
 enum EGametypeContentReferencerTypes
@@ -977,9 +977,9 @@
 	}
 
 	// Add the stats to the list, note this is also the order that they get rendered in if active.
-#if !(UE_BUILD_SHIPPING || UE_BUILD_TEST)
+#if !UE_BUILD_SHIPPING
 	EngineStats.Add(FEngineStatFuncs(TEXT("STAT_Version"), TEXT("STATCAT_Engine"), FText::GetEmpty(), &UEngine::RenderStatVersion, NULL, bIsRHS));
-#endif
+#endif // !UE_BUILD_SHIPPING
 	EngineStats.Add(FEngineStatFuncs(TEXT("STAT_NamedEvents"), TEXT("STATCAT_Engine"), FText::GetEmpty(), &UEngine::RenderStatNamedEvents, &UEngine::ToggleStatNamedEvents, bIsRHS));
 	EngineStats.Add(FEngineStatFuncs(TEXT("STAT_FPS"), TEXT("STATCAT_Engine"), FText::GetEmpty(), &UEngine::RenderStatFPS, &UEngine::ToggleStatFPS, bIsRHS));
 	EngineStats.Add(FEngineStatFuncs(TEXT("STAT_Summary"), TEXT("STATCAT_Engine"), FText::GetEmpty(), &UEngine::RenderStatSummary, NULL, bIsRHS));
@@ -994,12 +994,12 @@
 
 	EngineStats.Add(FEngineStatFuncs(TEXT("STAT_ColorList"), TEXT("STATCAT_Engine"), FText::GetEmpty(), &UEngine::RenderStatColorList, NULL));
 	EngineStats.Add(FEngineStatFuncs(TEXT("STAT_Levels"), TEXT("STATCAT_Engine"), FText::GetEmpty(), &UEngine::RenderStatLevels, NULL));
-#if !(UE_BUILD_SHIPPING || UE_BUILD_TEST)
+#if !UE_BUILD_SHIPPING
 	EngineStats.Add(FEngineStatFuncs(TEXT("STAT_SoundMixes"), TEXT("STATCAT_Engine"), FText::GetEmpty(), &UEngine::RenderStatSoundMixes, NULL));
 	EngineStats.Add(FEngineStatFuncs(TEXT("STAT_Reverb"), TEXT("STATCAT_Engine"), FText::GetEmpty(), &UEngine::RenderStatReverb, NULL));
 	EngineStats.Add(FEngineStatFuncs(TEXT("STAT_SoundWaves"), TEXT("STATCAT_Engine"), FText::GetEmpty(), &UEngine::RenderStatSoundWaves, NULL));
 	EngineStats.Add(FEngineStatFuncs(TEXT("STAT_SoundCues"), TEXT("STATCAT_Engine"), FText::GetEmpty(), &UEngine::RenderStatSoundCues, NULL));
-#endif
+#endif // !UE_BUILD_SHIPPING
 	EngineStats.Add(FEngineStatFuncs(TEXT("STAT_Sounds"), TEXT("STATCAT_Engine"), FText::GetEmpty(), &UEngine::RenderStatSounds, &UEngine::ToggleStatSounds));
 /* @todo UE4 physx fix this once we have convexelem drawing again
 	EngineStats.Add(FEngineStatFuncs(TEXT("STAT_LevelMap"), TEXT("STATCAT_Engine"), FText::GetEmpty(), &UEngine::RenderStatLevelMap, NULL));
@@ -1051,7 +1051,7 @@
 
 void UEngine::PreExit()
 {
-	if (IMovieSceneCaptureModule* Module = FModuleManager::GetModulePtr<IMovieSceneCaptureModule>("MovieSceneCapture"))
+	if (IMovieSceneCaptureModule* Module = FModuleManager::GetModulePtr<IMovieSceneCaptureModule>( "MovieSceneCapture" ))
 	{
 		Module->DestroyAllActiveCaptures();
 	}
@@ -1262,7 +1262,7 @@
 		}
 	}
 
-#if !(UE_BUILD_SHIPPING || UE_BUILD_TEST)
+#if !UE_BUILD_SHIPPING
 	{
 		float OverrideFPS = CVarSetOverrideFPS.GetValueOnGameThread();
 		if(OverrideFPS >= 0.001f)
@@ -1274,7 +1274,7 @@
 			bTimeWasManipulated = true;
 		}
 	}
-#endif // !(UE_BUILD_SHIPPING || UE_BUILD_TEST)
+#endif // !UE_BUILD_SHIPPING
 }
 
 
@@ -2453,13 +2453,16 @@
 	}
 #endif
 
-#if !(UE_BUILD_SHIPPING || UE_BUILD_TEST)
-#if WITH_HOT_RELOAD
+
+#if	!(UE_BUILD_SHIPPING || UE_BUILD_TEST) && WITH_HOT_RELOAD
 	else if( FParse::Command(&Cmd,TEXT("HotReload")) )
 	{
 		return HandleHotReloadCommand( Cmd, Ar );
 	}
-#endif // WITH_HOT_RELOAD
+#endif // !(UE_BUILD_SHIPPING || UE_BUILD_TEST) && WITH_HOT_RELOAD
+
+
+#if !UE_BUILD_SHIPPING
 	else if (FParse::Command(&Cmd, TEXT("DumpConsoleCommands")))
 	{
 		return HandleDumpConsoleCommandsCommand( Cmd, Ar, InWorld );
@@ -2553,9 +2556,6 @@
 	{
 		return HandleShowSelectedLightmapCommand( Cmd, Ar );
 	}
-#endif // !(UE_BUILD_SHIPPING || UE_BUILD_TEST)
-
-#if !UE_BUILD_SHIPPING
 	else if( FParse::Command(&Cmd,TEXT("SHOWLOG")) )
 	{
 		return HandleShowLogCommand( Cmd, Ar );
@@ -2747,20 +2747,17 @@
 	return true;
 }
 
-bool UEngine::HandleStopMovieCaptureCommand( const TCHAR* Cmd, FOutputDevice& Ar )
+bool UEngine::HandleFlushLogCommand( const TCHAR* Cmd, FOutputDevice& Ar )
 {
-	if (IMovieSceneCaptureInterface* CaptureInterface = IMovieSceneCaptureModule::Get().GetFirstActiveMovieSceneCapture())
-	{
-		CaptureInterface->Close();
-		return true;
-	}
-	return false;
+	GLog->FlushThreadedLogs();
+	GLog->Flush();
+	return true;
 }
 
 bool UEngine::HandleGameVerCommand( const TCHAR* Cmd, FOutputDevice& Ar )
 {
-	FString VersionString = FString::Printf(TEXT("GameVersion Branch: %s, Configuration: %s, Version: %s, CommandLine: %s"), 
-		*FApp::GetBranchName(), EBuildConfigurations::ToString(FApp::GetBuildConfiguration()), *FEngineVersion::Current().ToString(), FCommandLine::Get());
+	FString VersionString = FString::Printf( TEXT( "GameVersion Branch: %s, Configuration: %s, Version: %s, CommandLine: %s" ),
+											 *FApp::GetBranchName(), EBuildConfigurations::ToString( FApp::GetBuildConfiguration() ), *FEngineVersion::Current().ToString(), FCommandLine::Get() );
 
 	Ar.Logf( *VersionString );
 	FPlatformMisc::ClipboardCopy( *VersionString );
@@ -2768,6 +2765,34 @@
 	return 1;
 }
 
+bool UEngine::HandleStatCommand( UWorld* World, FCommonViewportClient* ViewportClient, const TCHAR* Cmd, FOutputDevice& Ar )
+{
+	const TCHAR* Temp = Cmd;
+	for (int32 StatIdx = 0; StatIdx < EngineStats.Num(); StatIdx++)
+	{
+		const FEngineStatFuncs& EngineStat = EngineStats[StatIdx];
+		if (FParse::Command( &Temp, *EngineStat.CommandNameString ))
+		{
+			if (EngineStat.ToggleFunc)
+			{
+				return ViewportClient ? (this->*(EngineStat.ToggleFunc))(World, ViewportClient, Temp) : false;
+			}
+			return true;
+		}
+	}
+	return false;
+}
+
+bool UEngine::HandleStopMovieCaptureCommand( const TCHAR* Cmd, FOutputDevice& Ar )
+{
+	if (IMovieSceneCaptureInterface* CaptureInterface = IMovieSceneCaptureModule::Get().GetFirstActiveMovieSceneCapture())
+	{
+		CaptureInterface->Close();
+		return true;
+	}
+	return false;
+}
+
 bool UEngine::HandleCrackURLCommand( const TCHAR* Cmd, FOutputDevice& Ar )
 {
 	FURL URL(NULL,Cmd,TRAVEL_Absolute);
@@ -2793,129 +2818,214 @@
 	return 1;
 }
 
-#if !UE_BUILD_SHIPPING
-bool UEngine::HandleMergeMeshCommand( const TCHAR* Cmd, FOutputDevice& Ar, UWorld* InWorld )
+bool UEngine::HandleCeCommand( UWorld* InWorld, const TCHAR* Cmd, FOutputDevice& Ar )
 {
-	FString CmdCopy = Cmd;
-	TArray<FString> Tokens;
-	while (CmdCopy.Len() > 0)
+	const TCHAR* ErrorMessage = TEXT( "No level found for CE processing" );
+	bool bResult = false;
+
+	// Try to execute the command on all level script actors
+	for (TArray<ULevel*>::TConstIterator it = InWorld->GetLevels().CreateConstIterator(); it; ++it)
 	{
-		const TCHAR* LocalCmd = *CmdCopy;
-		FString Token = FParse::Token(LocalCmd, true);
-		Tokens.Add(Token);
-		CmdCopy = CmdCopy.Right(CmdCopy.Len() - Token.Len() - 1);
-	}
+		ULevel* CurrentLevel = *it;
+		if (CurrentLevel)
+		{
+			ErrorMessage = TEXT( "No LevelScriptActor found for CE processing" );
 
-	// array of source meshes that will be merged
-	TArray<USkeletalMesh*> SourceMeshList;
+			if (CurrentLevel->GetLevelScriptActor())
+			{
+				ErrorMessage = 0;
 
-	if ( Tokens.Num() >= 2 )
-	{
-		for (int32 I=0; I<Tokens.Num(); ++I)
-		{
-			USkeletalMesh * SrcMesh = LoadObject<USkeletalMesh>(NULL, *Tokens[I], NULL, LOAD_None, NULL);
-			if (SrcMesh)
-			{
-				SourceMeshList.Add(SrcMesh);
+				// return true if at least one level handles the command
+				bResult |= CurrentLevel->GetLevelScriptActor()->CallFunctionByNameWithArguments( Cmd, Ar, NULL, true );
 			}
 		}
 	}
 
-	// find player controller skeletalmesh
-	APawn * PlayerPawn = NULL;
-	USkeletalMesh * PlayerMesh = NULL;
-	for( FConstPlayerControllerIterator Iterator = InWorld->GetPlayerControllerIterator(); Iterator; ++Iterator )
+	if (!bResult)
 	{
-		APlayerController* PlayerController = *Iterator;
-		if (PlayerController->GetCharacter() != NULL && PlayerController->GetCharacter()->GetMesh())
+		ErrorMessage = TEXT( "CE command wasn't processed" );
+	}
+
+	if (ErrorMessage)
+	{
+		UE_LOG( LogEngine, Error, TEXT( "%s" ), ErrorMessage );
+	}
+
+	// the command was processed (resulted in executing the command or an error message) - no other spot handles "CE"
+	return true;
+}
+
+
+bool UEngine::HandleDumpTicksCommand( UWorld* InWorld, const TCHAR* Cmd, FOutputDevice& Ar )
+{
+	// Handle optional parameters, will dump all tick functions by default.
+	bool bShowEnabled = true;
+	bool bShowDisabled = true;
+	if (FParse::Command( &Cmd, TEXT( "ENABLED" ) ))
+	{
+		bShowDisabled = false;
+	}
+	else if (FParse::Command( &Cmd, TEXT( "DISABLED" ) ))
+	{
+		bShowEnabled = false;
+	}
+	FTickTaskManagerInterface::Get().DumpAllTickFunctions( Ar, InWorld, bShowEnabled, bShowDisabled );
+	return true;
+}
+
+bool UEngine::HandleGammaCommand( const TCHAR* Cmd, FOutputDevice& Ar )
+{
+	DisplayGamma = (*Cmd != 0) ? FMath::Clamp<float>( FCString::Atof( *FParse::Token( Cmd, false ) ), 0.5f, 5.0f ) : 2.2f;
+	return true;
+}
+
+bool UEngine::HandleRecordAnimationCommand( UWorld* InWorld, const TCHAR* InStr, FOutputDevice& Ar )
+{
+#if WITH_EDITOR
+	const TCHAR* Str = InStr;
+	// parse actor name
+	TCHAR ActorName[128];
+	AActor* FoundActor = nullptr;
+	if (FParse::Token( Str, ActorName, ARRAY_COUNT( ActorName ), 0 ))
+	{
+		FString const ActorNameStr = FString( ActorName );
+		for (ULevel const* Level : InWorld->GetLevels())
 		{
-			PlayerPawn = PlayerController->GetCharacter();
-			PlayerMesh = PlayerController->GetCharacter()->GetMesh()->SkeletalMesh;
-			break;
+			if (Level)
+			{
+				for (AActor* Actor : Level->Actors)
+				{
+					if (Actor)
+					{
+						if (Actor->GetName() == ActorNameStr)
+						{
+							FoundActor = Actor;
+							break;
+						}
+					}
+				}
+			}
 		}
 	}
 
-	if (PlayerMesh)
+	if (FoundActor)
 	{
-		if (SourceMeshList.Num() ==  0)
+		USkeletalMeshComponent* const SkelComp = FoundActor->FindComponentByClass<USkeletalMeshComponent>();
+		if (SkelComp)
 		{
-			SourceMeshList.Add(PlayerMesh);
-			SourceMeshList.Add(PlayerMesh);
+			TCHAR AssetPath[256];
+			FParse::Token( Str, AssetPath, ARRAY_COUNT( AssetPath ), 0 );
+			FString const AssetName = FPackageName::GetLongPackageAssetName( AssetPath );
+			return FAnimationRecorderManager::Get().RecordAnimation( FoundActor, SkelComp, AssetPath, AssetName );
 		}
 	}
-	else
+#endif		// WITH_EDITOR
+
+	return false;
+}
+
+bool UEngine::HandleStopRecordAnimationCommand( UWorld* InWorld, const TCHAR* InStr, FOutputDevice& Ar )
+{
+#if WITH_EDITOR
+	const TCHAR* Str = InStr;
+
+	// parse actor name
+	TCHAR ActorName[128];
+	AActor* FoundActor = nullptr;
+	bool bStopAll = false;
+	if (FParse::Token( Str, ActorName, ARRAY_COUNT( ActorName ), 0 ))
 	{
-		// we don't have a pawn (because we couldn't find a mesh), use any pawn as a spawn point
-		for( FConstPlayerControllerIterator Iterator = InWorld->GetPlayerControllerIterator(); Iterator; ++Iterator )
+		FString const ActorNameStr = FString( ActorName );
+
+		if (ActorNameStr.ToLower() == TEXT( "all" ))
 		{
-			APlayerController* PlayerController = *Iterator;
-			if (PlayerController->GetPawn() != NULL)
+			bStopAll = true;
+		}
+		else if (InWorld)
+		{
+			// search for the actor by name
+			for (ULevel* Level : InWorld->GetLevels())
 			{
-				PlayerPawn = PlayerController->GetPawn();
-				break;
+				if (Level)
+				{
+					for (AActor* Actor : Level->Actors)
+					{
+						if (Actor)
+						{
+							if (Actor->GetName() == ActorNameStr)
+							{
+								FoundActor = Actor;
+								break;
+							}
+						}
+					}
+				}
 			}
-		}		
+		}
 	}
 
-	if (PlayerPawn && SourceMeshList.Num() >= 2)
+	if (bStopAll)
 	{
-		// create the composite mesh
-		auto CompositeMesh = NewObject<USkeletalMesh>(GetTransientPackage(), NAME_None, RF_Transient);
-
-		TArray<FSkelMeshMergeSectionMapping> InForceSectionMapping;
-		// create an instance of the FSkeletalMeshMerge utility
-		FSkeletalMeshMerge MeshMergeUtil( CompositeMesh, SourceMeshList, InForceSectionMapping,  0);
-
-		// merge the source meshes into the composite mesh
-		if( !MeshMergeUtil.DoMerge() )
+		FAnimationRecorderManager::Get().StopRecordingAllAnimations();
+		return true;
+	}
+	else if (FoundActor)
+	{
+		USkeletalMeshComponent* const SkelComp = FoundActor->FindComponentByClass<USkeletalMeshComponent>();
+		if (SkelComp)
 		{
-			// handle errors
-			// ...
-			UE_LOG(LogEngine, Log, TEXT("DoMerge Error: Merge Mesh Test Failed"));
+			FAnimationRecorderManager::Get().StopRecordingAnimation( FoundActor, SkelComp );
 			return true;
 		}
+	}
+#endif		// WITH_EDITOR
 
-		FVector SpawnLocation = PlayerPawn->GetActorLocation() + PlayerPawn->GetActorForwardVector()*50.f;
+	return false;
+}
 
-		// set the new composite mesh in the existing skeletal mesh component
-		ASkeletalMeshActor* const SMA = PlayerPawn->GetWorld()->SpawnActor<ASkeletalMeshActor>( SpawnLocation, PlayerPawn->GetActorRotation()*-1);
-		if (SMA)
-		{
-			SMA->GetSkeletalMeshComponent()->SetSkeletalMesh(CompositeMesh);
-		}
+bool UEngine::HandleShowLogCommand( const TCHAR* Cmd, FOutputDevice& Ar )
+{
+	// Toggle display of console log window.
+	if (GLogConsole)
+	{
+		GLogConsole->Show( !GLogConsole->IsShown() );
 	}
+	return 1;
+}
 
+#if STATS
+bool UEngine::HandleDumpParticleMemCommand( const TCHAR* Cmd, FOutputDevice& Ar )
+{
+	FParticleMemoryStatManager::DumpParticleMemoryStats( Ar );
 	return true;
 }
-#endif // !UE_BUILD_SHIPPING
+#endif
 
-
-#if !(UE_BUILD_SHIPPING || UE_BUILD_TEST)
-#if WITH_HOT_RELOAD
+#if !(UE_BUILD_SHIPPING || UE_BUILD_TEST) && WITH_HOT_RELOAD
 bool UEngine::HandleHotReloadCommand( const TCHAR* Cmd, FOutputDevice& Ar )
 {
-	FString Module(FParse::Token(Cmd, 0));
+	FString Module( FParse::Token( Cmd, 0 ) );
 	FString PackagePath( FString( TEXT( "/Script/" ) ) + Module );
-	UPackage *Package = FindPackage(NULL,*PackagePath );
+	UPackage *Package = FindPackage( NULL, *PackagePath );
 	if (!Package)
 	{
-		Ar.Logf( TEXT("Could not HotReload '%s', package not found in memory"),*Module);
+		Ar.Logf( TEXT( "Could not HotReload '%s', package not found in memory" ), *Module );
 	}
 	else
 	{
-		Ar.Logf( TEXT("HotReloading %s..."),*Module);
+		Ar.Logf( TEXT( "HotReloading %s..." ), *Module );
 		TArray< UPackage*> PackagesToRebind;
 		PackagesToRebind.Add( Package );
 		const bool bWaitForCompletion = true;	// Always wait when hotreload is initiated from the console
-		IHotReloadInterface& HotReloadSupport = FModuleManager::LoadModuleChecked<IHotReloadInterface>("HotReload");
-		const ECompilationResult::Type CompilationResult = HotReloadSupport.RebindPackages(PackagesToRebind, TArray<FName>(), bWaitForCompletion, Ar);
+		IHotReloadInterface& HotReloadSupport = FModuleManager::LoadModuleChecked<IHotReloadInterface>( "HotReload" );
+		const ECompilationResult::Type CompilationResult = HotReloadSupport.RebindPackages( PackagesToRebind, TArray<FName>(), bWaitForCompletion, Ar );
 	}
 	return true;
 }
-#endif // WITH_HOT_RELOAD
+#endif // !(UE_BUILD_SHIPPING || UE_BUILD_TEST) && WITH_HOT_RELOAD
 
+#if !UE_BUILD_SHIPPING
 
-#if !(UE_BUILD_SHIPPING || UE_BUILD_TEST)
 static void DumpHelp(UWorld* InWorld)
 {
 	UE_LOG(LogEngine, Display, TEXT("Console Help:"));
@@ -2970,8 +3080,6 @@
 	TEXT("Outputs some helptext to the console and the log"),
 	FConsoleCommandWithWorldDelegate::CreateStatic(DumpHelp)
 	);
-#endif // !(UE_BUILD_SHIPPING || UE_BUILD_TEST)
 
-
 bool UEngine::HandleDumpConsoleCommandsCommand( const TCHAR* Cmd, FOutputDevice& Ar, UWorld* InWorld )
 {
@@ -3359,19 +3467,6 @@
 	return false;
 }
 
-#endif // !(UE_BUILD_SHIPPING || UE_BUILD_TEST)
-
-#if !UE_BUILD_SHIPPING
-bool UEngine::HandleShowLogCommand( const TCHAR* Cmd, FOutputDevice& Ar )
-{
-	// Toggle display of console log window.
-	if( GLogConsole )
-	{
-		GLogConsole->Show( !GLogConsole->IsShown() );
-	}
-	return 1;
-}
-
 bool UEngine::HandleStartFPSChartCommand( const TCHAR* Cmd, FOutputDevice& Ar )
 {
 	// start the chart data capture
@@ -4623,6 +4718,100 @@
 	return false;
 }
 
+bool UEngine::HandleMergeMeshCommand( const TCHAR* Cmd, FOutputDevice& Ar, UWorld* InWorld )
+{
+	FString CmdCopy = Cmd;
+	TArray<FString> Tokens;
+	while (CmdCopy.Len() > 0)
+	{
+		const TCHAR* LocalCmd = *CmdCopy;
+		FString Token = FParse::Token( LocalCmd, true );
+		Tokens.Add( Token );
+		CmdCopy = CmdCopy.Right( CmdCopy.Len() - Token.Len() - 1 );
+	}
+
+	// array of source meshes that will be merged
+	TArray<USkeletalMesh*> SourceMeshList;
+
+	if (Tokens.Num() >= 2)
+	{
+		for (int32 I = 0; I<Tokens.Num(); ++I)
+		{
+			USkeletalMesh * SrcMesh = LoadObject<USkeletalMesh>( NULL, *Tokens[I], NULL, LOAD_None, NULL );
+			if (SrcMesh)
+			{
+				SourceMeshList.Add( SrcMesh );
+			}
+		}
+	}
+
+	// find player controller skeletalmesh
+	APawn * PlayerPawn = NULL;
+	USkeletalMesh * PlayerMesh = NULL;
+	for (FConstPlayerControllerIterator Iterator = InWorld->GetPlayerControllerIterator(); Iterator; ++Iterator)
+	{
+		APlayerController* PlayerController = *Iterator;
+		if (PlayerController->GetCharacter() != NULL && PlayerController->GetCharacter()->GetMesh())
+		{
+			PlayerPawn = PlayerController->GetCharacter();
+			PlayerMesh = PlayerController->GetCharacter()->GetMesh()->SkeletalMesh;
+			break;
+		}
+	}
+
+	if (PlayerMesh)
+	{
+		if (SourceMeshList.Num() == 0)
+		{
+			SourceMeshList.Add( PlayerMesh );
+			SourceMeshList.Add( PlayerMesh );
+		}
+	}
+	else
+	{
+		// we don't have a pawn (because we couldn't find a mesh), use any pawn as a spawn point
+		for (FConstPlayerControllerIterator Iterator = InWorld->GetPlayerControllerIterator(); Iterator; ++Iterator)
+		{
+			APlayerController* PlayerController = *Iterator;
+			if (PlayerController->GetPawn() != NULL)
+			{
+				PlayerPawn = PlayerController->GetPawn();
+				break;
+			}
+		}
+	}
+
+	if (PlayerPawn && SourceMeshList.Num() >= 2)
+	{
+		// create the composite mesh
+		auto CompositeMesh = NewObject<USkeletalMesh>( GetTransientPackage(), NAME_None, RF_Transient );
+
+		TArray<FSkelMeshMergeSectionMapping> InForceSectionMapping;
+		// create an instance of the FSkeletalMeshMerge utility
+		FSkeletalMeshMerge MeshMergeUtil( CompositeMesh, SourceMeshList, InForceSectionMapping, 0 );
+
+		// merge the source meshes into the composite mesh
+		if (!MeshMergeUtil.DoMerge())
+		{
+			// handle errors
+			// ...
+			UE_LOG( LogEngine, Log, TEXT( "DoMerge Error: Merge Mesh Test Failed" ) );
+			return true;
+		}
+
+		FVector SpawnLocation = PlayerPawn->GetActorLocation() + PlayerPawn->GetActorForwardVector()*50.f;
+
+		// set the new composite mesh in the existing skeletal mesh component
+		ASkeletalMeshActor* const SMA = PlayerPawn->GetWorld()->SpawnActor<ASkeletalMeshActor>( SpawnLocation, PlayerPawn->GetActorRotation()*-1 );
+		if (SMA)
+		{
+			SMA->GetSkeletalMeshComponent()->SetSkeletalMesh( CompositeMesh );
+		}
+	}
+
+	return true;
+}
+
 bool UEngine::HandleContentComparisonCommand( const TCHAR* Cmd, FOutputDevice& Ar )
 {
 	TArray<FString> Tokens, Switches;
@@ -4897,7 +5086,7 @@
 	}
 };
 
-// #YRX_UObject 2014-09-15 Move to ObjectCommads.cpp or ObjectExec.cpp
+// #TODO Move to ObjectCommads.cpp or ObjectExec.cpp
 bool UEngine::HandleObjCommand( const TCHAR* Cmd, FOutputDevice& Ar )
 {
 	if( FParse::Command(&Cmd,TEXT("GARBAGE")) || FParse::Command(&Cmd,TEXT("GC")) )
@@ -5671,72 +5860,8 @@
 	return true;
 }
 
-#endif // !UE_BUILD_SHIPPING
 
-bool UEngine::HandleCeCommand( UWorld* InWorld, const TCHAR* Cmd, FOutputDevice& Ar )
-{
-	const TCHAR* ErrorMessage = TEXT("No level found for CE processing");
-	bool bResult = false;
 
-	// Try to execute the command on all level script actors
-	for( TArray<ULevel*>::TConstIterator it = InWorld->GetLevels().CreateConstIterator(); it; ++it )
-	{
-		ULevel* CurrentLevel = *it;
-		if( CurrentLevel )
-		{
-			ErrorMessage = TEXT("No LevelScriptActor found for CE processing");
-
-			if( CurrentLevel->GetLevelScriptActor() )
-			{
-				ErrorMessage = 0;
-
-				// return true if at least one level handles the command
-				bResult |= CurrentLevel->GetLevelScriptActor()->CallFunctionByNameWithArguments(Cmd, Ar, NULL, true);
-			}
-		}
-	}
-
-	if(!bResult)
-	{
-		ErrorMessage = TEXT("CE command wasn't processed");
-	}
-
-	if(ErrorMessage)
-	{
-		UE_LOG(LogEngine, Error, TEXT("%s"), ErrorMessage);
-	}
-
-	// the command was processed (resulted in executing the command or an error message) - no other spot handles "CE"
-	return true;
-}
-
-#if STATS
-bool UEngine::HandleDumpParticleMemCommand( const TCHAR* Cmd, FOutputDevice& Ar )
-{
-	FParticleMemoryStatManager::DumpParticleMemoryStats(Ar);
-	return true;
-}
-#endif
-
-bool UEngine::HandleStatCommand( UWorld* World, FCommonViewportClient* ViewportClient, const TCHAR* Cmd, FOutputDevice& Ar )
-{
-	const TCHAR* Temp = Cmd;
-	for (int32 StatIdx = 0; StatIdx < EngineStats.Num(); StatIdx++)
-	{
-		const FEngineStatFuncs& EngineStat = EngineStats[StatIdx];
-		if (FParse::Command(&Temp, *EngineStat.CommandNameString))
-		{
-			if (EngineStat.ToggleFunc)
-			{
-				return ViewportClient ? ( this->*(EngineStat.ToggleFunc) )(World, ViewportClient, Temp) : false;
-			}
-			return true;
-		}
-	}
-	return false;
-}
-
-#if !UE_BUILD_SHIPPING
 bool UEngine::HandleTestslateGameUICommand( const TCHAR* Cmd, FOutputDevice& Ar )
 {
 	TSharedRef<SWidget> GameUI = 
@@ -5896,140 +6021,6 @@
 }
 #endif // !UE_BUILD_SHIPPING
 
-bool UEngine::HandleFlushLogCommand( const TCHAR* Cmd, FOutputDevice& Ar )
-{
-	GLog->FlushThreadedLogs();
-	GLog->Flush();
-	return true;
-}
 
-bool UEngine::HandleDumpTicksCommand( UWorld* InWorld, const TCHAR* Cmd, FOutputDevice& Ar )
-{
-	// Handle optional parameters, will dump all tick functions by default.
-	bool bShowEnabled = true;
-	bool bShowDisabled = true;
-	if (FParse::Command(&Cmd, TEXT("ENABLED")))
-	{
-		bShowDisabled = false;
-	}
-	else if (FParse::Command(&Cmd, TEXT("DISABLED")))
-	{
-		bShowEnabled = false;
-	}
-	FTickTaskManagerInterface::Get().DumpAllTickFunctions(Ar, InWorld, bShowEnabled, bShowDisabled);
-	return true;
-}
-
-bool UEngine::HandleGammaCommand( const TCHAR* Cmd, FOutputDevice& Ar )
-{
-	DisplayGamma = (*Cmd != 0) ? FMath::Clamp<float>(FCString::Atof(*FParse::Token(Cmd, false)), 0.5f, 5.0f) : 2.2f;
-	return true;
-}
-
-bool UEngine::HandleRecordAnimationCommand(UWorld* InWorld, const TCHAR* InStr, FOutputDevice& Ar)
-{
-#if WITH_EDITOR
-	const TCHAR* Str = InStr;
-	// parse actor name
-	TCHAR ActorName[128];
-	AActor* FoundActor = nullptr;
-	if (FParse::Token(Str, ActorName, ARRAY_COUNT(ActorName), 0))
-	{
-		FString const ActorNameStr = FString(ActorName);
-		for (ULevel const* Level : InWorld->GetLevels())
-		{
-			if (Level)
-			{
-				for (AActor* Actor : Level->Actors)
-				{
-					if (Actor)
-					{
-						if (Actor->GetName() == ActorNameStr)
-						{
-							FoundActor = Actor;
-							break;
-						}
-					}
-				}
-			}
-		}
-	}
-
-	if (FoundActor)
-	{
-		USkeletalMeshComponent* const SkelComp = FoundActor->FindComponentByClass<USkeletalMeshComponent>();
-		if (SkelComp)
-		{
-			TCHAR AssetPath[256];
-			FParse::Token(Str, AssetPath, ARRAY_COUNT(AssetPath), 0);
-			FString const AssetName = FPackageName::GetLongPackageAssetName(AssetPath);
-			return FAnimationRecorderManager::Get().RecordAnimation(FoundActor, SkelComp, AssetPath, AssetName);
-		}
-	}
-#endif		// WITH_EDITOR
-
-	return false;
-}
-
-bool UEngine::HandleStopRecordAnimationCommand(UWorld* InWorld, const TCHAR* InStr, FOutputDevice& Ar)
-{
-#if WITH_EDITOR
-	const TCHAR* Str = InStr;
-
-	// parse actor name
-	TCHAR ActorName[128];
-	AActor* FoundActor = nullptr;
-	bool bStopAll = false;
-	if (FParse::Token(Str, ActorName, ARRAY_COUNT(ActorName), 0))
-	{
-		FString const ActorNameStr = FString(ActorName);
-
-		if (ActorNameStr.ToLower() == TEXT("all"))
-		{
-			bStopAll = true;
-		}
-		else if (InWorld)
-		{
-			// search for the actor by name
-			for (ULevel* Level : InWorld->GetLevels())
-			{
-				if (Level)
-				{
-					for (AActor* Actor : Level->Actors)
-					{
-						if (Actor)
-						{
-							if (Actor->GetName() == ActorNameStr)
-							{
-								FoundActor = Actor;
-								break;
-							}
-						}
-					}
-				}
-			}
-		}
-	}
-
-	if (bStopAll)
-	{
-		FAnimationRecorderManager::Get().StopRecordingAllAnimations();
-		return true;
-	}
-	else if (FoundActor)
-	{
-		USkeletalMeshComponent* const SkelComp = FoundActor->FindComponentByClass<USkeletalMeshComponent>();
-		if (SkelComp)
-		{
-			FAnimationRecorderManager::Get().StopRecordingAnimation(FoundActor, SkelComp);
-			return true;
-		}
-	}
-#endif		// WITH_EDITOR
-
-	return false;
-}
-
-
 /**
  * Computes a color to use for property coloration for the given object.
@@ -7131,7 +7122,7 @@
 /** draws a property of the given object on the screen similarly to stats */
 static void DrawProperty(UCanvas* CanvasObject, UObject* Obj, const FDebugDisplayProperty& PropData, UProperty* Prop, int32 X, int32& Y)
 {
-#if !(UE_BUILD_SHIPPING || UE_BUILD_TEST)
+#if UE_BUILD_SHIPPING
 	checkSlow(PropData.bSpecialProperty || Prop != NULL);
 	checkSlow(Prop == NULL || Obj->GetClass()->IsChildOf(Prop->GetOwnerClass()));
 
@@ -7279,7 +7270,7 @@
 			Y += YL;
 		}
 	} while( CommaIdx >= 0 );
-#endif
+#endif // !UE_BUILD_SHIPPING
 }
 
 /** Basic timing collation - cannot use stats as these are not enabled in Win32 shipping */
@@ -11246,7 +11237,7 @@
 }
 
 // VERSION
-#if !(UE_BUILD_SHIPPING || UE_BUILD_TEST)
+#if !UE_BUILD_SHIPPING
 int32 UEngine::RenderStatVersion(UWorld* World, FViewport* Viewport, FCanvas* Canvas, int32 X, int32 Y, const FVector* ViewLocation, const FRotator* ViewRotation)
 {
 	if (!GIsHighResScreenshot && !GIsDumpingMovie && GAreScreenMessagesEnabled)
@@ -11261,7 +11252,7 @@
 	}
 	return Y;
 }
-#endif
+#endif // !UE_BUILD_SHIPPING
 
 // DETAILED
 bool UEngine::ToggleStatDetailed(UWorld* World, FCommonViewportClient* ViewportClient, const TCHAR* Stream)
@@ -11888,7 +11879,7 @@
 	);
 
 // REVERB
-#if !(UE_BUILD_SHIPPING || UE_BUILD_TEST)
+#if !UE_BUILD_SHIPPING
 int32 UEngine::RenderStatReverb(UWorld* World, FViewport* Viewport, FCanvas* Canvas, int32 X, int32 Y, const FVector* ViewLocation, const FRotator* ViewRotation)
 {
 	if (FAudioDevice* AudioDevice = World->GetAudioDevice())
@@ -12083,7 +12074,7 @@
 	Y += 12;
 	return Y;
 }
-#endif
+#endif // !UE_BUILD_SHIPPING
 
 // SOUNDS
 bool UEngine::ToggleStatSounds(UWorld* World, FCommonViewportClient* ViewportClient, const TCHAR* Stream)
@@ -12158,7 +12149,8 @@
 
 int32 UEngine::RenderStatSounds(UWorld* World, FViewport* Viewport, FCanvas* Canvas, int32 X, int32 Y, const FVector* ViewLocation, const FRotator* ViewRotation)
 {
-#if !(UE_BUILD_SHIPPING || UE_BUILD_TEST)
+#if UE_BUILD_SHIPPING
+
 #if UE_BUILD_DEBUG
 
 	typedef TMap< const FActiveSound*, FSoundInfo* > TMapSounds;
@@ -12386,7 +12378,7 @@
 			delete It.Value();
 		}
 	}
-#endif
+#endif // !UE_BUILD_SHIPPING
 	return Y;
 }
 
diff --git a/Engine/Source/Runtime/Launch/Private/LaunchEngineLoop.cpp b/Engine/Source/Runtime/Launch/Private/LaunchEngineLoop.cpp
--- a/Engine/Source/Runtime/Launch/Private/LaunchEngineLoop.cpp
+++ b/Engine/Source/Runtime/Launch/Private/LaunchEngineLoop.cpp
@@ -16,6 +16,7 @@
 #include "ModuleManager.h"
 #include "../Resources/Version.h"
 #include "VersionManifest.h"
+#include "UObject/DevObjectVersion.h"
 
 #if WITH_EDITOR
 	#include "EditorStyle.h"
@@ -228,7 +229,7 @@ bool ParseGameProjectFromCommandLine(const TCHAR* InCmdLine, FString& OutProject
 			OutGameName = FPaths::GetBaseFilename(OutProjectFilePath);
 			return true;
 		}
-		else if (FPlatformProperties::IsMonolithicBuild() == false)
+		else if (FPaths::IsRelative(FirstCommandLineToken) && FPlatformProperties::IsMonolithicBuild() == false)
 		{
 			// Full game name is assumed to be the first token
 			OutGameName = MoveTemp(FirstCommandLineToken);
@@ -1078,11 +1079,17 @@ int32 FEngineLoop::PreInit( const TCHAR* CmdLine )
 		const FString ProjectFilePath = FPaths::Combine(*FPaths::GameDir(), *FString::Printf(TEXT("%s.%s"), FApp::GetGameName(), *FProjectDescriptor::GetExtension()));
 		FPaths::SetProjectFilePath(ProjectFilePath);
 	}
+#endif
 
 	// Now verify the project file if we have one
-	if ( FPaths::IsProjectFilePathSet() )
+	if (FPaths::IsProjectFilePathSet()
+#if IS_PROGRAM
+		// Programs don't need uproject files to exist, but some do specify them and if they exist we should load them
+		&& FPaths::FileExists(FPaths::GetProjectFilePath())
+#endif
+		)
 	{
-		if ( !IProjectManager::Get().LoadProjectFile(FPaths::GetProjectFilePath()) )
+		if (!IProjectManager::Get().LoadProjectFile(FPaths::GetProjectFilePath()))
 		{
 			// The project file was invalid or saved with a newer version of the engine. Exit.
 			UE_LOG(LogInit, Warning, TEXT("Could not find a valid project file, the engine will exit now."));
@@ -1090,6 +1097,7 @@ int32 FEngineLoop::PreInit( const TCHAR* CmdLine )
 		}
 	}
 
+#if !IS_PROGRAM
 	if( FApp::HasGameName() )
 	{
 		// Tell the module manager what the game binaries folder is
@@ -2924,6 +2932,9 @@ bool FEngineLoop::AppInit( )
 	// init config system
 	FConfigCacheIni::InitializeConfigSystem();
 
+	// Now that configs have been initialized, setup stack walking options
+	FPlatformStackWalk::Init();
+
 	CheckForPrintTimesOverride();
 
 	// Check whether the project or any of its plugins are missing or are out of date
@@ -3046,17 +3057,18 @@ bool FEngineLoop::AppInit( )
 		GIsSilent = true;
 	}
 
+#endif // !UE_BUILD_SHIPPING
+
 	// Show log if wanted.
 	if (GLogConsole && FParse::Param(FCommandLine::Get(), TEXT("LOG")))
 	{
 		GLogConsole->Show(true);
 	}
 
-#endif // !UE_BUILD_SHIPPING
-
 	//// Command line.
 	UE_LOG(LogInit, Log, TEXT("Version: %s"), *FEngineVersion::Current().ToString());
 	UE_LOG(LogInit, Log, TEXT("API Version: %u"), FEngineVersion::CompatibleWith().GetChangelist());
+	FDevVersionRegistration::DumpVersionsToLog();
 
 #if PLATFORM_64BITS
 	UE_LOG(LogInit, Log, TEXT("Compiled (64-bit): %s %s"), ANSI_TO_TCHAR(__DATE__), ANSI_TO_TCHAR(__TIME__));
diff --git a/Engine/Source/Runtime/Projects/Private/ModuleDescriptor.cpp b/Engine/Source/Runtime/Projects/Private/ModuleDescriptor.cpp
--- a/Engine/Source/Runtime/Projects/Private/ModuleDescriptor.cpp
+++ b/Engine/Source/Runtime/Projects/Private/ModuleDescriptor.cpp
@@ -70,6 +70,9 @@ const TCHAR* EHostType::ToString( const EHostType::Type Value )
 		case RuntimeNoCommandlet:
 			return TEXT( "RuntimeNoCommandlet" );
 
+		case RuntimeNoProgram:
+			return TEXT("RuntimeNoProgram");
+
 		case Developer:
 			return TEXT( "Developer" );
 
@@ -278,6 +281,13 @@ bool FModuleDescriptor::IsCompiledInCurrentConfiguration() const
 	case EHostType::RuntimeNoCommandlet:
 		return true;
 
+	case EHostType::RuntimeNoProgram:
+#if IS_PROGRAM
+		return false;
+#else
+		return true;
+#endif
+
 	case EHostType::Developer:
 		#if WITH_UNREAL_DEVELOPER_TOOLS
 			return true;
@@ -315,12 +325,16 @@ bool FModuleDescriptor::IsLoadedInCurrentConfiguration() const
 	// Check that the runtime environment allows it to be loaded
 	switch (Type)
 	{
+	case EHostType::RuntimeNoProgram:
+#if IS_PROGRAM
+		return false;
+#endif // else fall through to EHostType::Runtime below
 	case EHostType::Runtime:
-		#if WITH_ENGINE || WITH_PLUGIN_SUPPORT
+		#if (WITH_ENGINE || WITH_PLUGIN_SUPPORT)
 			return true;
 		#endif
 		break;
-
+	
 	case EHostType::RuntimeNoCommandlet:
 		#if WITH_ENGINE || WITH_PLUGIN_SUPPORT
 			if(!IsRunningCommandlet()) return true;
diff --git a/Engine/Source/Runtime/Projects/Private/PluginManager.cpp b/Engine/Source/Runtime/Projects/Private/PluginManager.cpp
--- a/Engine/Source/Runtime/Projects/Private/PluginManager.cpp
+++ b/Engine/Source/Runtime/Projects/Private/PluginManager.cpp
@@ -48,6 +48,14 @@ namespace PluginSystemDefs
 			}
 		} while (SearchStr != nullptr);
 
+#if IS_PROGRAM
+		// For programs that have the project dir set, look for plugins under the project directory
+		const FProjectDescriptor *Project = IProjectManager::Get().GetCurrentProject();
+		if (Project != nullptr)
+		{
+			PluginPathsOut.Add(FPaths::GetPath(FPaths::GetProjectFilePath()) / TEXT("Plugins"));
+		}
+#endif
 		return PluginCount;
 	}
 }
@@ -292,6 +300,35 @@ public:
 	}
 };
 
+bool FPluginManager::IsPluginSupportedByCurrentTarget(TSharedRef<FPlugin> Plugin) const
+{
+	bool bSupported = false;
+	if (Plugin->GetDescriptor().Modules.Num())
+	{
+		for (const FModuleDescriptor& Module : Plugin->GetDescriptor().Modules)
+		{
+			// Programs support only program type plugins
+			// Non-program targets don't support from plugins
+#if IS_PROGRAM
+			if (Module.Type == EHostType::Program)
+			{
+				bSupported = true;
+			}
+#else
+			if (Module.Type != EHostType::Program)
+			{
+				bSupported = true;
+			}
+#endif
+		}
+	}
+	else
+	{
+		bSupported = true;
+	}
+	return bSupported;
+}
+
 bool FPluginManager::ConfigureEnabledPlugins()
 {
 	if(!bHaveConfiguredEnabledPlugins)
@@ -343,13 +380,22 @@ bool FPluginManager::ConfigureEnabledPlugins()
 		// If we made it here, we have all the required plugins
 		bHaveAllRequiredPlugins = true;
 
+		const bool bHasProjectFile = IProjectManager::Get().GetCurrentProject() != nullptr;
 		// Get all the enabled plugin names
 		TArray< FString > EnabledPluginNames;
-		#if IS_PROGRAM
-			GConfig->GetArray(TEXT("Plugins"), TEXT("ProgramEnabledPlugins"), EnabledPluginNames, GEngineIni);
-		#else
+	#if IS_PROGRAM
+		// Programs with a project file specified take the list of enabled plugins from the project file
+		if (bHasProjectFile)
+		{
 			FProjectManager::Get().GetEnabledPlugins(EnabledPluginNames);
-		#endif
+		}
+		else
+		{
+			GConfig->GetArray(TEXT("Plugins"), TEXT("ProgramEnabledPlugins"), EnabledPluginNames, GEngineIni);
+		}
+	#else
+		FProjectManager::Get().GetEnabledPlugins(EnabledPluginNames);
+	#endif
 
 		// Build a set from the array
 		TSet< FString > AllEnabledPlugins;
@@ -359,8 +405,8 @@ bool FPluginManager::ConfigureEnabledPlugins()
 		for( const TSharedRef< FPlugin > Plugin : AllPlugins )
 		{
 			if ( AllEnabledPlugins.Contains(Plugin->Name) )
-			{
-				Plugin->bEnabled = true;
+			{				
+				Plugin->bEnabled = (!IS_PROGRAM || !bHasProjectFile) || IsPluginSupportedByCurrentTarget(Plugin);
 			}
 		}
 
diff --git a/Engine/Source/Runtime/Projects/Private/PluginManager.h b/Engine/Source/Runtime/Projects/Private/PluginManager.h
--- a/Engine/Source/Runtime/Projects/Private/PluginManager.h
+++ b/Engine/Source/Runtime/Projects/Private/PluginManager.h
@@ -93,6 +93,9 @@ private:
 	/** Gets the instance of a given plugin */
 	TSharedPtr<FPlugin> FindPluginInstance(const FString& Name);
 
+	/** Returns true if the plugin is supported by the current target (program/game) */
+	bool IsPluginSupportedByCurrentTarget(TSharedRef<FPlugin> Plugin) const;
+
 private:
 	/** All of the plugins that we know about */
 	TArray< TSharedRef< FPlugin > > AllPlugins;
diff --git a/Engine/Source/Runtime/Projects/Public/ModuleDescriptor.h b/Engine/Source/Runtime/Projects/Public/ModuleDescriptor.h
--- a/Engine/Source/Runtime/Projects/Public/ModuleDescriptor.h
+++ b/Engine/Source/Runtime/Projects/Public/ModuleDescriptor.h
@@ -61,6 +61,7 @@ namespace EHostType
 	{
 		Runtime,
 		RuntimeNoCommandlet,
+		RuntimeNoProgram,
 		Developer,
 		Editor,
 		EditorNoCommandlet,
diff --git a/Engine/Source/Runtime/Renderer/Private/ShadowSetup.cpp b/Engine/Source/Runtime/Renderer/Private/ShadowSetup.cpp
--- a/Engine/Source/Runtime/Renderer/Private/ShadowSetup.cpp
+++ b/Engine/Source/Runtime/Renderer/Private/ShadowSetup.cpp
@@ -140,7 +140,7 @@ static TAutoConsoleVariable<int32> CVarUseConservativeShadowBounds(
 	TEXT("Whether to use safe and conservative shadow frustum creation that wastes some shadowmap space"),
 	ECVF_RenderThreadSafe);
 
-#if !(UE_BUILD_SHIPPING || UE_BUILD_TEST)
+#if !UE_BUILD_SHIPPING
 // read and written on the render thread
 bool GDumpShadowSetup = false;
 void DumpShadowDumpSetup()
@@ -157,7 +157,7 @@ FAutoConsoleCommand CmdDumpShadowDumpSetup(
 	TEXT("Dump shadow setup (for developer only, only for non shiping build)"),
 	FConsoleCommandDelegate::CreateStatic(DumpShadowDumpSetup)
 	);
-#endif // !(UE_BUILD_SHIPPING || UE_BUILD_TEST)
+#endif // !UE_BUILD_SHIPPING
 
 /**
  * Helper function to determine fade alpha value for shadows based on resolution. In the below ASCII art (1) is
@@ -1942,7 +1942,7 @@ void FSceneRenderer::InitProjectedShadowVisibility(FRHICommandListImmediate& RHI
 		}
 	}
 
-#if !(UE_BUILD_SHIPPING || UE_BUILD_TEST)
+#if !UE_BUILD_SHIPPING
 	if(GDumpShadowSetup)
 	{
 		GDumpShadowSetup = false;
@@ -1992,7 +1992,7 @@ void FSceneRenderer::InitProjectedShadowVisibility(FRHICommandListImmediate& RHI
 			}
 		}
 	}
-#endif // !(UE_BUILD_SHIPPING || UE_BUILD_TEST)
+#endif // !UE_BUILD_SHIPPING
 }
 
 void FSceneRenderer::GatherShadowDynamicMeshElements()
