Fix compilation on a case-sensitive fs.

diff --git a/Engine/Source/Runtime/Engine/Private/Particles/SubUVAnimation.cpp b/Engine/Source/Runtime/Engine/Private/Particles/SubUVAnimation.cpp
--- a/Engine/Source/Runtime/Engine/Private/Particles/SubUVAnimation.cpp
+++ b/Engine/Source/Runtime/Engine/Private/Particles/SubUVAnimation.cpp
@@ -2,7 +2,7 @@
 
 #include "EnginePrivate.h"
 #include "Particles/SubUVAnimation.h"
-#include "ConvexHull2D.h"
+#include "ConvexHull2d.h"
 #include "Particles/ParticleSystemComponent.h"
 #include "ParticleHelper.h"
 #include "DerivedDataCacheInterface.h"
