Copying //UE4/Dev-Framework to //UE4/Main

diff --git a/Engine/Source/Runtime/Engine/Private/AI/Navigation/RecastNavMesh.cpp b/Engine/Source/Runtime/Engine/Private/AI/Navigation/RecastNavMesh.cpp
--- a/Engine/Source/Runtime/Engine/Private/AI/Navigation/RecastNavMesh.cpp
+++ b/Engine/Source/Runtime/Engine/Private/AI/Navigation/RecastNavMesh.cpp
@@ -1818,22 +1818,35 @@ FPathFindingResult ARecastNavMesh::FindPath(const FNavAgentProperties& AgentProp
 		return ENavigationQueryResult::Error;
 	}
 		
-	FPathFindingResult Result;
-	Result.Path = Query.PathInstanceToFill.IsValid() ? Query.PathInstanceToFill : Self->CreatePathInstance<FNavMeshPath>(Query.Owner.Get());
+	FPathFindingResult Result(ENavigationQueryResult::Error);
 
-	FNavMeshPath* NavMeshPath = (FNavMeshPath*)Result.Path.Get();
-	NavMeshPath->SetFilter(Query.QueryFilter);
-	NavMeshPath->ApplyFlags(Query.NavDataFlags);
+	FNavigationPath* NavPath = Query.PathInstanceToFill.Get();
+	FNavMeshPath* NavMeshPath = NavPath ? NavPath->CastPath<FNavMeshPath>() : nullptr;
 
-	if ((Query.StartLocation - Query.EndLocation).IsNearlyZero() == true)
+	if (NavMeshPath)
 	{
-		Result.Path->GetPathPoints().Reset();
-		Result.Path->GetPathPoints().Add(FNavPathPoint(Query.EndLocation));
-		Result.Result = ENavigationQueryResult::Success;
+		Result.Path = Query.PathInstanceToFill;
+		NavMeshPath->ResetForRepath();
 	}
 	else
 	{
-		if(Query.QueryFilter.IsValid())
+		Result.Path = Self->CreatePathInstance<FNavMeshPath>(Query.Owner.Get());
+		NavPath = Result.Path.Get();
+		NavMeshPath = NavPath ? NavPath->CastPath<FNavMeshPath>() : nullptr;
+	}
+
+	if (NavMeshPath)
+	{
+		NavMeshPath->SetFilter(Query.QueryFilter);
+		NavMeshPath->ApplyFlags(Query.NavDataFlags);
+
+		if ((Query.StartLocation - Query.EndLocation).IsNearlyZero() == true)
+		{
+			Result.Path->GetPathPoints().Reset();
+			Result.Path->GetPathPoints().Add(FNavPathPoint(Query.EndLocation));
+			Result.Result = ENavigationQueryResult::Success;
+		}
+		else if (Query.QueryFilter.IsValid())
 		{
 			Result.Result = RecastNavMesh->RecastNavMeshImpl->FindPath(Query.StartLocation, Query.EndLocation, *NavMeshPath,
 				*(Query.QueryFilter.Get()), Query.Owner.Get());
@@ -1844,10 +1857,6 @@ FPathFindingResult ARecastNavMesh::FindPath(const FNavAgentProperties& AgentProp
 				Result.Result = Query.bAllowPartialPaths ? ENavigationQueryResult::Success : ENavigationQueryResult::Fail;
 			}
 		}
-		else
-		{
-			Result.Result = ENavigationQueryResult::Error;
-		}
 	}
 
 	return Result;
diff --git a/Engine/Source/Runtime/Engine/Private/ActiveSound.cpp b/Engine/Source/Runtime/Engine/Private/ActiveSound.cpp
--- a/Engine/Source/Runtime/Engine/Private/ActiveSound.cpp
+++ b/Engine/Source/Runtime/Engine/Private/ActiveSound.cpp
@@ -276,7 +276,7 @@ void FActiveSound::UpdateWaveInstances( TArray<FWaveInstance*> &InWaveInstances,
 	{
 		// If this active sound is told to limit concurrency by the quietest sound
 		const FSoundConcurrencySettings* ConcurrencySettingsToApply = GetSoundConcurrencySettingsToApply();
-		if (ConcurrencySettingsToApply && ConcurrencySettingsToApply->ResolutionRule == EMaxConcurrentResolutionRule::StopQuietist)
+		if (ConcurrencySettingsToApply && ConcurrencySettingsToApply->ResolutionRule == EMaxConcurrentResolutionRule::StopQuietest)
 		{
 			check(ConcurrencyGroupID != 0);
 			// Now that we have this sound's active wave instances, lets find the loudest active wave instance to represent the "volume" of this active sound
@@ -902,4 +902,4 @@ void FActiveSound::ApplyAttenuation(FSoundParseParameters& ParseParams, const FL
 	{
 		ParseParams.SpatializationAlgorithm = Settings->SpatializationAlgorithm;
 	}
-}
\ No newline at end of file
+}
diff --git a/Engine/Source/Runtime/Engine/Private/Actor.cpp b/Engine/Source/Runtime/Engine/Private/Actor.cpp
--- a/Engine/Source/Runtime/Engine/Private/Actor.cpp
+++ b/Engine/Source/Runtime/Engine/Private/Actor.cpp
@@ -1075,22 +1075,16 @@ void AActor::ClearComponentOverlaps()
 	// Remove owned components from overlap tracking
 	// We don't traverse the RootComponent attachment tree since that might contain
 	// components owned by other actors.
-	TArray<UPrimitiveComponent*> OverlappingComponentsForCurrentComponent;
+	TArray<FOverlapInfo, TInlineAllocator<3>> OverlapsForCurrentComponent;
 	for (UPrimitiveComponent* const PrimComp : PrimitiveComponents)
 	{
-		PrimComp->GetOverlappingComponents(OverlappingComponentsForCurrentComponent);
-
-		for (UPrimitiveComponent* const OverlapComp : OverlappingComponentsForCurrentComponent)
+		OverlapsForCurrentComponent.Reset();
+		OverlapsForCurrentComponent.Append(PrimComp->GetOverlapInfos());
+		for (const FOverlapInfo& CurrentOverlap : OverlapsForCurrentComponent)
 		{
-			if (OverlapComp)
-			{
-				PrimComp->EndComponentOverlap(OverlapComp, true, true);
-
-				if (IsPendingKill())
-				{
-					return;
-				}
-			}
+			const bool bDoNotifies = true;
+			const bool bSkipNotifySelf = false;
+			PrimComp->EndComponentOverlap(CurrentOverlap, bDoNotifies, bSkipNotifySelf);
 		}
 	}
 }
diff --git a/Engine/Source/Runtime/Engine/Private/Animation/AnimCompress.cpp b/Engine/Source/Runtime/Engine/Private/Animation/AnimCompress.cpp
--- a/Engine/Source/Runtime/Engine/Private/Animation/AnimCompress.cpp
+++ b/Engine/Source/Runtime/Engine/Private/Animation/AnimCompress.cpp
@@ -639,7 +639,7 @@ bool UAnimCompress::Reduce(UAnimSequence* AnimSeq, bool bOutput)
 	return bResult;
 }
 
-bool UAnimCompress::Reduce(TArray<UAnimSequence*> AnimSequences, bool bOutput)
+bool UAnimCompress::Reduce(const TArray<UAnimSequence*>& AnimSequences, bool bOutput)
 {
 	bool bResult = false;
 
diff --git a/Engine/Source/Runtime/Engine/Private/Animation/AnimNode_StateMachine.cpp b/Engine/Source/Runtime/Engine/Private/Animation/AnimNode_StateMachine.cpp
--- a/Engine/Source/Runtime/Engine/Private/Animation/AnimNode_StateMachine.cpp
+++ b/Engine/Source/Runtime/Engine/Private/Animation/AnimNode_StateMachine.cpp
@@ -336,6 +336,14 @@ const FAnimationTransitionBetweenStates& FAnimNode_StateMachine::GetTransitionIn
 
 void FAnimNode_StateMachine::Update(const FAnimationUpdateContext& Context)
 {
+	// If we just became relevant and haven't been initialized yet, then reinitialize state machine.
+	if (!bFirstUpdate && (UpdateCounter.Get() != INDEX_NONE) && !UpdateCounter.WasSynchronizedInTheLastFrame(Context.AnimInstanceProxy->GetUpdateCounter()))
+	{
+		FAnimationInitializeContext InitializationContext(Context.AnimInstanceProxy);
+		this->Initialize(InitializationContext);
+	}
+	UpdateCounter.SynchronizeWith(Context.AnimInstanceProxy->GetUpdateCounter());
+
 	const FBakedAnimationStateMachine* Machine = GetMachineDescription();
 	if (Machine != nullptr)
 	{
diff --git a/Engine/Source/Runtime/Engine/Private/Animation/Skeleton.cpp b/Engine/Source/Runtime/Engine/Private/Animation/Skeleton.cpp
--- a/Engine/Source/Runtime/Engine/Private/Animation/Skeleton.cpp
+++ b/Engine/Source/Runtime/Engine/Private/Animation/Skeleton.cpp
@@ -76,9 +76,6 @@ void USkeleton::PostLoad()
 		ConvertToFReferenceSkeleton();
 	}
 
-	// Build look up table between Slot nodes and their Group.
-	BuildSlotToGroupMap();
-
 	// catch any case if guid isn't valid
 	check(Guid.IsValid());
 
@@ -163,6 +160,17 @@ void USkeleton::Serialize( FArchive& Ar )
 		SmartNames.Serialize(Ar);
 	}
 
+	// Build look up table between Slot nodes and their Group.
+	if(Ar.UE4Ver() < VER_UE4_FIX_SLOT_NAME_DUPLICATION)
+	{
+		// In older assets we may have duplicates, remove these while building the map.
+		BuildSlotToGroupMap(true);
+	}
+	else
+	{
+		BuildSlotToGroupMap();
+	}
+
 #if WITH_EDITORONLY_DATA
 	if (Ar.UE4Ver() < VER_UE4_SKELETON_ASSET_PROPERTY_TYPE_CHANGE)
 	{
@@ -904,17 +912,35 @@ const TArray<FAnimSlotGroup>& USkeleton::GetSlotGroups() const
 	return SlotGroups;
 }
 
-void USkeleton::BuildSlotToGroupMap()
+void USkeleton::BuildSlotToGroupMap(bool bInRemoveDuplicates)
 {
 	SlotToGroupNameMap.Empty();
 
-	for (auto SlotGroup : SlotGroups)
+	for (FAnimSlotGroup& SlotGroup : SlotGroups)
 	{
-		for (auto SlotName : SlotGroup.SlotNames)
+		for (const FName& SlotName : SlotGroup.SlotNames)
 		{
 			SlotToGroupNameMap.Add(SlotName, SlotGroup.GroupName);
 		}
 	}
+
+	// Use the map we've just build to rebuild the slot groups
+	if(bInRemoveDuplicates)
+	{
+		for(FAnimSlotGroup& SlotGroup : SlotGroups)
+		{
+			SlotGroup.SlotNames.Empty(SlotGroup.SlotNames.Num());
+
+			for(TPair<FName, FName>& SlotToGroupPair : SlotToGroupNameMap)
+			{
+				if(SlotToGroupPair.Value == SlotGroup.GroupName)
+				{
+					SlotGroup.SlotNames.Add(SlotToGroupPair.Key);
+				}
+			}
+		}
+
+	}
 }
 
 FAnimSlotGroup* USkeleton::FindAnimSlotGroup(const FName& InGroupName)
diff --git a/Engine/Source/Runtime/Engine/Private/Collision/CollisionConversions.cpp b/Engine/Source/Runtime/Engine/Private/Collision/CollisionConversions.cpp
--- a/Engine/Source/Runtime/Engine/Private/Collision/CollisionConversions.cpp
+++ b/Engine/Source/Runtime/Engine/Private/Collision/CollisionConversions.cpp
@@ -345,7 +345,6 @@ static FVector FindGeomOpposingNormal(PxGeometryType::Enum QueryGeomType, const
 	return InNormal;
 }
 
-
 /** Set info in the HitResult (Actor, Component, PhysMaterial, BoneName, Item) based on the supplied shape and face index */
 static void SetHitResultFromShapeAndFaceIndex(const PxShape* PShape,  const PxRigidActor* PActor, const uint32 FaceIndex, FHitResult& OutResult, bool bReturnPhysMat)
 {
@@ -411,7 +410,7 @@ static void SetHitResultFromShapeAndFaceIndex(const PxShape* PShape,  const PxRi
 	{
 		OutResult.Item = BodyInst->InstanceBodyIndex;
 
-		const UBodySetup* BodySetup = BodyInst->BodySetup.Get();
+		const UBodySetup* BodySetup = BodyInst->BodySetup.Get();	//this data should be immutable at runtime so ok to check from worker thread.
 		if (BodySetup)
 		{
 			OutResult.BoneName = BodySetup->BoneName;
@@ -496,7 +495,7 @@ EConvertQueryResult ConvertQueryImpactHit(const UWorld* World, const PxLocationH
 
 	const PxGeometryType::Enum SweptGeometryType = Geom ? Geom->getType() : PxGeometryType::eINVALID;
 	OutResult.ImpactNormal = FindGeomOpposingNormal(SweptGeometryType, PHit, TraceStartToEnd, Normal);
-	
+
 	// Fill in Actor, Component, material, etc.
 	SetHitResultFromShapeAndFaceIndex(PHit.shape, PHit.actor, PHit.faceIndex, OutResult, bReturnPhysMat);
 
@@ -969,6 +968,7 @@ static bool ConvertOverlappedShapeToImpactHit(const UWorld* World, const PxLocat
 	}
 
 	OutResult.Normal = OutResult.ImpactNormal;
+	
 	SetHitResultFromShapeAndFaceIndex(PShape, PActor, FaceIdx, OutResult, bReturnPhysMat);
 
 	return bBlockingHit;
diff --git a/Engine/Source/Runtime/Engine/Private/Collision/KAggregateGeom.cpp b/Engine/Source/Runtime/Engine/Private/Collision/KAggregateGeom.cpp
--- a/Engine/Source/Runtime/Engine/Private/Collision/KAggregateGeom.cpp
+++ b/Engine/Source/Runtime/Engine/Private/Collision/KAggregateGeom.cpp
@@ -367,7 +367,7 @@ static bool TriHasEdge(int32 T0, int32 T1, int32 T2, int32 Edge0, int32 Edge1)
 	return false;
 }
 
-static void GetTriIndicesUsingEdge(const int32 Edge0, const int32 Edge1, const TArray<int32> TriData, int32& Tri0Index, int32& Tri1Index)
+static void GetTriIndicesUsingEdge(const int32 Edge0, const int32 Edge1, const TArray<int32>& TriData, int32& Tri0Index, int32& Tri1Index)
 {
 	Tri0Index = INDEX_NONE;
 	Tri1Index = INDEX_NONE;
diff --git a/Engine/Source/Runtime/Engine/Private/Collision/PhysXCollision.cpp b/Engine/Source/Runtime/Engine/Private/Collision/PhysXCollision.cpp
--- a/Engine/Source/Runtime/Engine/Private/Collision/PhysXCollision.cpp
+++ b/Engine/Source/Runtime/Engine/Private/Collision/PhysXCollision.cpp
@@ -843,9 +843,9 @@ void CaptureOverlap(const UWorld* World, const PxGeometry& PGeom, const PxTransf
 }
 
 #define STARTQUERYTIMER() double StartTime = FPlatformTime::Seconds()
-#define CAPTUREGEOMSWEEP(World, Start, End, Rot, QueryMode, PGeom, TraceChannel, Params, ResponseParam, ObjectParam, Results) if (GCollisionAnalyzerIsRecording) { CaptureGeomSweep(World, Start, End, Rot, QueryMode, PGeom, TraceChannel, Params, ResponseParam, ObjectParam, Results, FPlatformTime::Seconds() - StartTime); }
-#define CAPTURERAYCAST(World, Start, End, QueryMode, TraceChannel, Params, ResponseParam, ObjectParam, Results)	if (GCollisionAnalyzerIsRecording) { CaptureRaycast(World, Start, End, QueryMode, TraceChannel, Params, ResponseParam, ObjectParam, Results, FPlatformTime::Seconds() - StartTime); }
-#define CAPTUREGEOMOVERLAP(World, PGeom, PGeomPose, QueryMode, TraceChannel, Params, ResponseParams, ObjectParams, Results)	if (GCollisionAnalyzerIsRecording) { CaptureOverlap(World, PGeom, PGeomPose, QueryMode, TraceChannel, Params, ResponseParams, ObjectParams, Results, FPlatformTime::Seconds() - StartTime); }
+#define CAPTUREGEOMSWEEP(World, Start, End, Rot, QueryMode, PGeom, TraceChannel, Params, ResponseParam, ObjectParam, Results) if (GCollisionAnalyzerIsRecording && IsInGameThread()) { CaptureGeomSweep(World, Start, End, Rot, QueryMode, PGeom, TraceChannel, Params, ResponseParam, ObjectParam, Results, FPlatformTime::Seconds() - StartTime); }
+#define CAPTURERAYCAST(World, Start, End, QueryMode, TraceChannel, Params, ResponseParam, ObjectParam, Results)	if (GCollisionAnalyzerIsRecording && IsInGameThread()) { CaptureRaycast(World, Start, End, QueryMode, TraceChannel, Params, ResponseParam, ObjectParam, Results, FPlatformTime::Seconds() - StartTime); }
+#define CAPTUREGEOMOVERLAP(World, PGeom, PGeomPose, QueryMode, TraceChannel, Params, ResponseParams, ObjectParams, Results)	if (GCollisionAnalyzerIsRecording && IsInGameThread()) { CaptureOverlap(World, PGeom, PGeomPose, QueryMode, TraceChannel, Params, ResponseParams, ObjectParams, Results, FPlatformTime::Seconds() - StartTime); }
 
 
 #else
@@ -922,7 +922,7 @@ bool RaycastTest(const UWorld* World, const FVector Start, const FVector End, EC
 
 	TArray<FHitResult> Hits;
 #if !(UE_BUILD_SHIPPING || UE_BUILD_TEST)
-	if ((World->DebugDrawTraceTag != NAME_None) && (World->DebugDrawTraceTag == Params.TraceTag))
+	if ((World->DebugDrawTraceTag != NAME_None) && (World->DebugDrawTraceTag == Params.TraceTag) && IsInGameThread())
 	{
 		DrawLineTraces(World, Start, End, Hits, DebugLineLifetime);
 	}
@@ -1035,7 +1035,7 @@ bool RaycastSingle(const UWorld* World, struct FHitResult& OutHit, const FVector
 
 
 #if !(UE_BUILD_SHIPPING || UE_BUILD_TEST)
-	if ((World->DebugDrawTraceTag != NAME_None) && (World->DebugDrawTraceTag == Params.TraceTag))
+	if ((World->DebugDrawTraceTag != NAME_None) && (World->DebugDrawTraceTag == Params.TraceTag) && IsInGameThread())
 	{
 		TArray<FHitResult> Hits;
 		if (bHaveBlockingHit)
@@ -1047,7 +1047,7 @@ bool RaycastSingle(const UWorld* World, struct FHitResult& OutHit, const FVector
 #endif //!(UE_BUILD_SHIPPING || UE_BUILD_TEST)
 
 #if ENABLE_COLLISION_ANALYZER
-	if (GCollisionAnalyzerIsRecording)
+	if (GCollisionAnalyzerIsRecording && IsInGameThread())
 	{
 		TArray<FHitResult> Hits;
 		if (bHaveBlockingHit)
@@ -1549,7 +1549,7 @@ bool GeomSweepMulti_PhysX(const UWorld* World, const PxGeometry& PGeom, const Px
 	}
 
 #if !(UE_BUILD_SHIPPING || UE_BUILD_TEST)
-	if ((World->DebugDrawTraceTag != NAME_None) && (World->DebugDrawTraceTag == Params.TraceTag))
+	if ((World->DebugDrawTraceTag != NAME_None) && (World->DebugDrawTraceTag == Params.TraceTag) && IsInGameThread())
 	{
 		TArray<FHitResult> OnlyMyHits(OutHits);
 		OnlyMyHits.RemoveAt(0, InitialHitCount, false); // Remove whatever was there initially.
@@ -1725,7 +1725,7 @@ bool GeomOverlapMultiImp_PhysX(const UWorld* World, const PxGeometry& PGeom, con
 			}
 
 #if !(UE_BUILD_SHIPPING || UE_BUILD_TEST)
-			if ((World->DebugDrawTraceTag != NAME_None) && (World->DebugDrawTraceTag == Params.TraceTag))
+			if ((World->DebugDrawTraceTag != NAME_None) && (World->DebugDrawTraceTag == Params.TraceTag) && IsInGameThread())
 			{
 				DrawGeomOverlaps(World, PGeom, PGeomPose, OutOverlaps, DebugLineLifetime);
 			}
diff --git a/Engine/Source/Runtime/Engine/Private/Components/BoxComponent.cpp b/Engine/Source/Runtime/Engine/Private/Components/BoxComponent.cpp
--- a/Engine/Source/Runtime/Engine/Private/Components/BoxComponent.cpp
+++ b/Engine/Source/Runtime/Engine/Private/Components/BoxComponent.cpp
@@ -33,41 +33,65 @@ void UBoxComponent::SetBoxExtent(FVector NewBoxExtent, bool bUpdateOverlaps)
 	}
 }
 
-void UBoxComponent::UpdateBodySetup()
-{
-	if(ShapeBodySetup == NULL || ShapeBodySetup->IsPendingKill())
-	{
-		ShapeBodySetup = NewObject<UBodySetup>(this);
-		ShapeBodySetup->CollisionTraceFlag = CTF_UseSimpleAsComplex;
-		ShapeBodySetup->AggGeom.BoxElems.Add(FKBoxElem());
-		ShapeBodySetup->bNeverNeedsCookedCollisionData = true;
-	}
 
+template <EShapeBodySetupHelper UpdateBodySetupAction>
+bool InvalidateOrUpdateBoxBodySetup(UBodySetup*& ShapeBodySetup, bool bUseArchetypeBodySetup, FVector BoxExtent)
+{
+	check((bUseArchetypeBodySetup && UpdateBodySetupAction == EShapeBodySetupHelper::InvalidateSharingIfStale) || (!bUseArchetypeBodySetup && UpdateBodySetupAction == EShapeBodySetupHelper::UpdateBodySetup));
 	check(ShapeBodySetup->AggGeom.BoxElems.Num() == 1);
 	FKBoxElem* se = ShapeBodySetup->AggGeom.BoxElems.GetData();
 
 	// @todo UE4 do we allow this now?
 	// check for malformed values
-	if( BoxExtent.X < KINDA_SMALL_NUMBER )
+	if (BoxExtent.X < KINDA_SMALL_NUMBER)
 	{
 		BoxExtent.X = 1.0f;
 	}
 
-	if( BoxExtent.Y < KINDA_SMALL_NUMBER )
+	if (BoxExtent.Y < KINDA_SMALL_NUMBER)
 	{
 		BoxExtent.Y = 1.0f;
 	}
 
-	if( BoxExtent.Z < KINDA_SMALL_NUMBER )
+	if (BoxExtent.Z < KINDA_SMALL_NUMBER)
 	{
 		BoxExtent.Z = 1.0f;
 	}
 
-	// now set the PhysX data values
-	se->SetTransform( FTransform::Identity );
-	se->X = BoxExtent.X*2;
-	se->Y = BoxExtent.Y*2;
-	se->Z = BoxExtent.Z*2;
+	float XExtent = BoxExtent.X * 2.f;
+	float YExtent = BoxExtent.Y * 2.f;
+	float ZExtent = BoxExtent.Z * 2.f;
+
+	if (UpdateBodySetupAction == EShapeBodySetupHelper::UpdateBodySetup)
+	{
+		// now set the PhysX data values
+		se->SetTransform(FTransform::Identity);
+		se->X = XExtent;
+		se->Y = YExtent;
+		se->Z = ZExtent;
+	}
+	else if(se->X != XExtent || se->Y != YExtent || se->Z != ZExtent)
+	{
+		ShapeBodySetup = nullptr;
+		bUseArchetypeBodySetup = false;
+	}
+
+	return bUseArchetypeBodySetup;
+}
+
+void UBoxComponent::UpdateBodySetup()
+{
+	if (PrepareSharedBodySetup<UBoxComponent>())
+	{
+		bUseArchetypeBodySetup = InvalidateOrUpdateBoxBodySetup<EShapeBodySetupHelper::InvalidateSharingIfStale>(ShapeBodySetup, bUseArchetypeBodySetup, BoxExtent);
+	}
+
+	CreateShapeBodySetupIfNeeded<FKBoxElem>();
+
+	if (!bUseArchetypeBodySetup)
+	{
+		InvalidateOrUpdateBoxBodySetup<EShapeBodySetupHelper::UpdateBodySetup>(ShapeBodySetup, bUseArchetypeBodySetup, BoxExtent);
+	}
 }
 
 bool UBoxComponent::IsZeroExtent() const
diff --git a/Engine/Source/Runtime/Engine/Private/Components/CapsuleComponent.cpp b/Engine/Source/Runtime/Engine/Private/Components/CapsuleComponent.cpp
--- a/Engine/Source/Runtime/Engine/Private/Components/CapsuleComponent.cpp
+++ b/Engine/Source/Runtime/Engine/Private/Components/CapsuleComponent.cpp
@@ -163,22 +163,47 @@ void UCapsuleComponent::SetCapsuleSize(float NewRadius, float NewHalfHeight, boo
 	}
 }
 
+
+template <EShapeBodySetupHelper UpdateBodySetupAction>
+bool InvalidateOrUpdateCapsuleBodySetup(UBodySetup*& ShapeBodySetup, bool bUseArchetypeBodySetup, float CapsuleRadius, float CapsuleHalfHeight)
+{
+	check((bUseArchetypeBodySetup && UpdateBodySetupAction == EShapeBodySetupHelper::InvalidateSharingIfStale) || (!bUseArchetypeBodySetup && UpdateBodySetupAction == EShapeBodySetupHelper::UpdateBodySetup));
+	check(ShapeBodySetup->AggGeom.SphylElems.Num() == 1);
+	FKSphylElem* SE = ShapeBodySetup->AggGeom.SphylElems.GetData();
+
+	const float Length = 2 * FMath::Max(CapsuleHalfHeight - CapsuleRadius, 0.f);	//SphylElem uses height from center of capsule spheres, but UCapsuleComponent uses halfHeight from end of the sphere
+
+	if (UpdateBodySetupAction == EShapeBodySetupHelper::UpdateBodySetup)
+	{
+		SE->SetTransform(FTransform::Identity);
+		SE->Radius = CapsuleRadius;
+		SE->Length = Length;
+	}
+	else
+	{
+		if(SE->Radius != CapsuleRadius || SE->Length != Length)
+		{
+			ShapeBodySetup = nullptr;
+			bUseArchetypeBodySetup = false;
+		}
+	}
+	
+	return bUseArchetypeBodySetup;
+}
+
 void UCapsuleComponent::UpdateBodySetup()
 {
-	if (ShapeBodySetup == NULL || ShapeBodySetup->IsPendingKill())
+	if (PrepareSharedBodySetup<UCapsuleComponent>())
 	{
-		ShapeBodySetup = NewObject<UBodySetup>(this);
-		ShapeBodySetup->CollisionTraceFlag = CTF_UseSimpleAsComplex;
-		ShapeBodySetup->AggGeom.SphylElems.Add(FKSphylElem());
-		ShapeBodySetup->bNeverNeedsCookedCollisionData = true;
+		bUseArchetypeBodySetup = InvalidateOrUpdateCapsuleBodySetup<EShapeBodySetupHelper::InvalidateSharingIfStale>(ShapeBodySetup, bUseArchetypeBodySetup, CapsuleRadius, CapsuleHalfHeight);
 	}
 
-	check(ShapeBodySetup->AggGeom.SphylElems.Num() == 1);
-	FKSphylElem* SE = ShapeBodySetup->AggGeom.SphylElems.GetData();
-	
-	SE->SetTransform(FTransform::Identity);
-	SE->Radius = CapsuleRadius;
-	SE->Length = 2 * FMath::Max(CapsuleHalfHeight - CapsuleRadius, 0.f);	//SphylElem uses height from center of capsule spheres, but UCapsuleComponent uses halfHeight from end of the sphere
+	CreateShapeBodySetupIfNeeded<FKSphylElem>();
+
+	if (!bUseArchetypeBodySetup)
+	{
+		InvalidateOrUpdateCapsuleBodySetup<EShapeBodySetupHelper::UpdateBodySetup>(ShapeBodySetup, bUseArchetypeBodySetup, CapsuleRadius, CapsuleHalfHeight);
+	}
 }
 
 bool UCapsuleComponent::IsZeroExtent() const
diff --git a/Engine/Source/Runtime/Engine/Private/Components/CharacterMovementComponent.cpp b/Engine/Source/Runtime/Engine/Private/Components/CharacterMovementComponent.cpp
--- a/Engine/Source/Runtime/Engine/Private/Components/CharacterMovementComponent.cpp
+++ b/Engine/Source/Runtime/Engine/Private/Components/CharacterMovementComponent.cpp
@@ -1216,9 +1216,9 @@ void UCharacterMovementComponent::SimulatedTick(float DeltaSeconds)
 			CharacterOwner->OnRep_ReplicatedBasedMovement();
 		}
 
-		// Avoid moving the mesh during movement, SmoothClientPosition will take care of it.
+		// Avoid moving the mesh during movement if SmoothClientPosition will take care of it.
 		// TODO: also for root motion stuff above?
-		const FScopedPreventAttachedComponentMove PreventMeshMovement(CharacterOwner->GetMesh());
+		const FScopedPreventAttachedComponentMove PreventMeshMovement(bNetworkSmoothingComplete ? nullptr : CharacterOwner->GetMesh());
 
 		if (CharacterOwner->bReplicateMovement)
 		{
@@ -2634,6 +2634,20 @@ bool UCharacterMovementComponent::CanStopPathFollowing() const
 	return !IsFalling();
 }
 
+float UCharacterMovementComponent::GetPathFollowingBrakingDistance(float MaxSpeed) const
+{
+	if (bUseFixedBrakingDistanceForPaths)
+	{
+		return FixedPathBrakingDistance;
+	}
+
+	const float BrakingDeceleration = BrakingDecelerationWalking * (bUseSeparateBrakingFriction ? BrakingFriction : GroundFriction);
+
+	// character won't be able to stop with negative or nearly zero deceleration, use MaxSpeed for path length calculations
+	const float BrakingDistance = (BrakingDeceleration < SMALL_NUMBER) ? MaxSpeed : (FMath::Square(MaxSpeed) / BrakingDeceleration);
+	return BrakingDistance;
+}
+
 void UCharacterMovementComponent::CalcAvoidanceVelocity(float DeltaTime)
 {
 	SCOPE_CYCLE_COUNTER(STAT_AI_ObstacleAvoidance);
diff --git a/Engine/Source/Runtime/Engine/Private/Components/DestructibleComponent.cpp b/Engine/Source/Runtime/Engine/Private/Components/DestructibleComponent.cpp
--- a/Engine/Source/Runtime/Engine/Private/Components/DestructibleComponent.cpp
+++ b/Engine/Source/Runtime/Engine/Private/Components/DestructibleComponent.cpp
@@ -793,7 +793,7 @@ class UDestructibleMesh * UDestructibleComponent::GetDestructibleMesh()
 	return Cast<UDestructibleMesh>(SkeletalMesh);
 }
 
-void UDestructibleComponent::SetSkeletalMesh( USkeletalMesh* InSkelMesh )
+void UDestructibleComponent::SetSkeletalMesh(USkeletalMesh* InSkelMesh, bool bReinitPose)
 {
 	if(InSkelMesh != NULL && !InSkelMesh->IsA(UDestructibleMesh::StaticClass()))
 	{
diff --git a/Engine/Source/Runtime/Engine/Private/Components/MovementComponent.cpp b/Engine/Source/Runtime/Engine/Private/Components/MovementComponent.cpp
--- a/Engine/Source/Runtime/Engine/Private/Components/MovementComponent.cpp
+++ b/Engine/Source/Runtime/Engine/Private/Components/MovementComponent.cpp
@@ -236,7 +236,8 @@ APhysicsVolume* UMovementComponent::GetPhysicsVolume() const
 
 bool UMovementComponent::IsInWater() const
 {
-	return GetPhysicsVolume() && GetPhysicsVolume()->bWaterVolume;
+	const APhysicsVolume* PhysVolume = GetPhysicsVolume();
+	return PhysVolume && PhysVolume->bWaterVolume;
 }
 
 bool UMovementComponent::ShouldSkipUpdate(float DeltaTime) const
diff --git a/Engine/Source/Runtime/Engine/Private/Components/NavMovementComponent.cpp b/Engine/Source/Runtime/Engine/Private/Components/NavMovementComponent.cpp
--- a/Engine/Source/Runtime/Engine/Private/Components/NavMovementComponent.cpp
+++ b/Engine/Source/Runtime/Engine/Private/Components/NavMovementComponent.cpp
@@ -27,6 +27,8 @@ bool FNavAgentProperties::IsNavDataMatching(const FNavAgentProperties& Other) co
 UNavMovementComponent::UNavMovementComponent(const FObjectInitializer& ObjectInitializer)
 	: Super(ObjectInitializer)
 	, bUpdateNavAgentWithOwnersCollision(true)
+	, bUseAccelerationForPaths(false)
+	, bUseFixedBrakingDistanceForPaths(false)
 	, bStopMovementAbortPaths(true)
 {
 }
@@ -41,11 +43,21 @@ void UNavMovementComponent::RequestDirectMove(const FVector& MoveVelocity, bool
 	Velocity = MoveVelocity;
 }
 
+void UNavMovementComponent::RequestPathMove(const FVector& MoveInput)
+{
+	// empty in base class, requires at least PawnMovementComponent for input related operations
+}
+
 bool UNavMovementComponent::CanStopPathFollowing() const
 {
 	return true;
 }
 
+float UNavMovementComponent::GetPathFollowingBrakingDistance(float MaxSpeed) const
+{
+	return bUseFixedBrakingDistanceForPaths ? FixedPathBrakingDistance : MaxSpeed;
+}
+
 void UNavMovementComponent::StopActiveMovement()
 {
 	if (PathFollowingComp.IsValid() && bStopMovementAbortPaths)
diff --git a/Engine/Source/Runtime/Engine/Private/Components/PawnMovementComponent.cpp b/Engine/Source/Runtime/Engine/Private/Components/PawnMovementComponent.cpp
--- a/Engine/Source/Runtime/Engine/Private/Components/PawnMovementComponent.cpp
+++ b/Engine/Source/Runtime/Engine/Private/Components/PawnMovementComponent.cpp
@@ -66,6 +66,14 @@ FVector UPawnMovementComponent::ConsumeInputVector()
 	return PawnOwner ? PawnOwner->Internal_ConsumeMovementInputVector() : FVector::ZeroVector;
 }
 
+void UPawnMovementComponent::RequestPathMove(const FVector& MoveInput)
+{
+	if (PawnOwner)
+	{
+		PawnOwner->Internal_AddMovementInput(MoveInput);
+	}
+}
+
 // TODO: deprecated, remove
 FVector UPawnMovementComponent::GetInputVector() const
 {
diff --git a/Engine/Source/Runtime/Engine/Private/Components/PrimitiveComponent.cpp b/Engine/Source/Runtime/Engine/Private/Components/PrimitiveComponent.cpp
--- a/Engine/Source/Runtime/Engine/Private/Components/PrimitiveComponent.cpp
+++ b/Engine/Source/Runtime/Engine/Private/Components/PrimitiveComponent.cpp
@@ -41,6 +41,8 @@ namespace PrimitiveComponentStatics
 	static const FName UpdateOverlapsName(TEXT("UpdateOverlaps"));
 }
 
+typedef TArray<FOverlapInfo, TInlineAllocator<3>> TInlineOverlapInfoArray;
+
 DEFINE_LOG_CATEGORY_STATIC(LogPrimitiveComponent, Log, All);
 
 static FAutoConsoleVariable CVarAllowCachedOverlaps(
@@ -419,6 +421,7 @@ void UPrimitiveComponent::OnUnregister()
 
 FPrimitiveComponentInstanceData::FPrimitiveComponentInstanceData(const UPrimitiveComponent* SourceComponent)
 	: FSceneComponentInstanceData(SourceComponent)
+	, LODParent(SourceComponent->GetLODParentPrimitive())
 {
 }
 
@@ -426,11 +429,14 @@ void FPrimitiveComponentInstanceData::ApplyToComponent(UActorComponent* Componen
 {
 	FSceneComponentInstanceData::ApplyToComponent(Component, CacheApplyPhase);
 
+	UPrimitiveComponent* NewComponent = CastChecked<UPrimitiveComponent>(Component);
+
 #if WITH_EDITOR
 	// This is needed to restore transient collision profile data.
-	CastChecked<UPrimitiveComponent>(Component)->UpdateCollisionProfile();
+	NewComponent->UpdateCollisionProfile();
 #endif // #if WITH_EDITOR
-
+	NewComponent->SetLODParentPrimitive(LODParent);
+	
 	if (ContainsSavedProperties() && Component->IsRegistered())
 	{
 		Component->MarkRenderStateDirty();
@@ -439,7 +445,32 @@ void FPrimitiveComponentInstanceData::ApplyToComponent(UActorComponent* Componen
 
 bool FPrimitiveComponentInstanceData::ContainsData() const
 {
-	return (ContainsSavedProperties() || AttachedInstanceComponents.Num() > 0);
+	return (ContainsSavedProperties() || AttachedInstanceComponents.Num() > 0 || LODParent);
+}
+
+void FPrimitiveComponentInstanceData::AddReferencedObjects(FReferenceCollector& Collector)
+{
+	FSceneComponentInstanceData::AddReferencedObjects(Collector);
+
+	// if LOD Parent
+	if (LODParent)
+	{
+		Collector.AddReferencedObject(LODParent);
+	}
+}
+
+void FPrimitiveComponentInstanceData::FindAndReplaceInstances(const TMap<UObject*, UObject*>& OldToNewInstanceMap)
+{
+	FSceneComponentInstanceData::FindAndReplaceInstances(OldToNewInstanceMap);
+
+	// if LOD Parent 
+	if (LODParent)
+	{
+		if (UObject* const* NewLODParent = OldToNewInstanceMap.Find(LODParent))
+		{
+			LODParent = CastChecked<UPrimitiveComponent>(*NewLODParent, ECastCheckedType::NullAllowed);
+		}
+	}
 }
 
 FActorComponentInstanceData* UPrimitiveComponent::GetComponentInstanceData() const
@@ -877,6 +908,22 @@ void UPrimitiveComponent::BeginDestroy()
 	}
 }
 
+void UPrimitiveComponent::OnComponentDestroyed()
+{
+	// Prevent future overlap events. Any later calls to UpdateOverlaps will only allow this to end overlaps.
+	bGenerateOverlapEvents = false;
+
+	// End all current overlaps
+	if (OverlappingComponents.Num() > 0)
+	{
+		const bool bDoNotifies = true;
+		const bool bSkipNotifySelf = false;
+		ClearComponentOverlaps(bDoNotifies, bSkipNotifySelf);
+	}
+
+	Super::OnComponentDestroyed();
+}
+
 bool UPrimitiveComponent::IsReadyForFinishDestroy()
 {
 	// Don't allow the primitive component to the purged until its pending scene detachments have completed.
@@ -2067,7 +2114,7 @@ void UPrimitiveComponent::BeginComponentOverlap(const FOverlapInfo& OtherOverlap
 }
 
 
-void UPrimitiveComponent::EndComponentOverlap(const FOverlapInfo& OtherOverlap, bool bDoNotifies, bool bNoNotifySelf)
+void UPrimitiveComponent::EndComponentOverlap(const FOverlapInfo& OtherOverlap, bool bDoNotifies, bool bSkipNotifySelf)
 {
 	UPrimitiveComponent* OtherComp = OtherOverlap.OverlapInfo.Component.Get();
 	if (OtherComp == nullptr)
@@ -2094,7 +2141,7 @@ void UPrimitiveComponent::EndComponentOverlap(const FOverlapInfo& OtherOverlap,
 			AActor* const MyActor = GetOwner();
 			if (OtherActor)
 			{
-				if (!bNoNotifySelf && IsPrimCompValidAndAlive(this))
+				if (!bSkipNotifySelf && IsPrimCompValidAndAlive(this))
 				{
 					OnComponentEndOverlap.Broadcast(OtherActor, OtherComp, OtherOverlap.GetBodyIndex());
 				}
@@ -2376,6 +2423,7 @@ void UPrimitiveComponent::UpdateOverlaps(const TArray<FOverlapInfo>* NewPendingO
 		AActor* const MyActor = GetOwner();
 		if ( MyActor && MyActor->IsActorInitialized() )
 		{
+			const FTransform PrevTransform = GetComponentTransform();
 			// If we are the root component we ignore child components. Those children will update their overlaps when we descend into the child tree.
 			// This aids an optimization in MoveComponent.
 			const bool bIgnoreChildren = (MyActor->GetRootComponent() == this);
@@ -2390,14 +2438,13 @@ void UPrimitiveComponent::UpdateOverlaps(const TArray<FOverlapInfo>* NewPendingO
 
 			// now generate full list of new touches, so we can compare to existing list and
 			// determine what changed
-			typedef TArray<FOverlapInfo, TInlineAllocator<3>> TInlineOverlapInfoArray;
 			TInlineOverlapInfoArray NewOverlappingComponents;
 
 			// If pending kill, we should not generate any new overlaps
 			if (!IsPendingKill())
 			{
 				// Might be able to avoid testing for new overlaps at the end location.
-				if (OverlapsAtEndLocation != NULL && CVarAllowCachedOverlaps->GetInt())
+				if (OverlapsAtEndLocation != NULL && CVarAllowCachedOverlaps->GetInt() && PrevTransform.Equals(GetComponentTransform()))
 				{
 					UE_LOG(LogPrimitiveComponent, VeryVerbose, TEXT("%s->%s Skipping overlap test!"), *GetNameSafe(GetOwner()), *GetName());
 					NewOverlappingComponents = *OverlapsAtEndLocation;
@@ -2492,16 +2539,11 @@ void UPrimitiveComponent::UpdateOverlaps(const TArray<FOverlapInfo>* NewPendingO
 	else
 	{
 		// bGenerateOverlapEvents is false or collision is disabled
-
+		// End all overlaps that exist, in case bGenerateOverlapEvents was true last tick (i.e. was just turned off)
 		if (OverlappingComponents.Num() > 0)
 		{
-			// End all overlaps that exist, in case bGenerateOverlapEvents was true last tick (i.e. was just turned off)
-			// Make a copy since EndComponentOverlap will remove items from OverlappingComponents.
-			auto OverlapsCopy = OverlappingComponents;
-			for (const FOverlapInfo& OtherOverlap : OverlapsCopy)
-			{
-				EndComponentOverlap(OtherOverlap, bDoNotifies, false);
-			}
+			const bool bSkipNotifySelf = false;
+			ClearComponentOverlaps(bDoNotifies, bSkipNotifySelf);
 		}
 	}
 
@@ -2523,6 +2565,19 @@ void UPrimitiveComponent::UpdateOverlaps(const TArray<FOverlapInfo>* NewPendingO
 	}
 }
 
+void UPrimitiveComponent::ClearComponentOverlaps(bool bDoNotifies, bool bSkipNotifySelf)
+{
+	if (OverlappingComponents.Num() > 0)
+	{
+		// Make a copy since EndComponentOverlap will remove items from OverlappingComponents.
+		const TInlineOverlapInfoArray OverlapsCopy(OverlappingComponents);
+		for (const FOverlapInfo& OtherOverlap : OverlapsCopy)
+		{
+			EndComponentOverlap(OtherOverlap, bDoNotifies, bSkipNotifySelf);
+		}
+	}
+}
+
 bool UPrimitiveComponent::ComponentOverlapMultiImpl(TArray<struct FOverlapResult>& OutOverlaps, const UWorld* World, const FVector& Pos, const FQuat& Quat, ECollisionChannel TestChannel, const struct FComponentQueryParams& Params, const struct FCollisionObjectQueryParams& ObjectQueryParams) const
 {
 	FComponentQueryParams ParamsWithSelf = Params;
@@ -2533,38 +2588,40 @@ bool UPrimitiveComponent::ComponentOverlapMultiImpl(TArray<struct FOverlapResult
 
 void UPrimitiveComponent::UpdatePhysicsVolume( bool bTriggerNotifiers )
 {
-	if (bShouldUpdatePhysicsVolume && !IsPendingKill() && GetWorld())
+	if (bShouldUpdatePhysicsVolume && !IsPendingKill())
 	{
 		SCOPE_CYCLE_COUNTER(STAT_UpdatePhysicsVolume);
-
-		if (bGenerateOverlapEvents && IsQueryCollisionEnabled())
+		if (UWorld* MyWorld = GetWorld())
 		{
-			APhysicsVolume* BestVolume = GetWorld()->GetDefaultPhysicsVolume();
-			int32 BestPriority = BestVolume->Priority;
-
-			for (auto CompIt = OverlappingComponents.CreateIterator(); CompIt; ++CompIt)
+			if (bGenerateOverlapEvents && IsQueryCollisionEnabled())
 			{
-				const FOverlapInfo& Overlap = *CompIt;
-				UPrimitiveComponent* OtherComponent = Overlap.OverlapInfo.Component.Get();
-				if (OtherComponent)
+				APhysicsVolume* BestVolume = MyWorld->GetDefaultPhysicsVolume();
+				int32 BestPriority = BestVolume->Priority;
+
+				for (auto CompIt = OverlappingComponents.CreateIterator(); CompIt; ++CompIt)
 				{
-					APhysicsVolume* V = Cast<APhysicsVolume>(OtherComponent->GetOwner());
-					if (V && V->Priority > BestPriority)
+					const FOverlapInfo& Overlap = *CompIt;
+					UPrimitiveComponent* OtherComponent = Overlap.OverlapInfo.Component.Get();
+					if (OtherComponent && OtherComponent->bGenerateOverlapEvents)
 					{
-						if (V->IsOverlapInVolume(*this))
+						APhysicsVolume* V = Cast<APhysicsVolume>(OtherComponent->GetOwner());
+						if (V && V->Priority > BestPriority)
 						{
-							BestPriority = V->Priority;
-							BestVolume = V;
+							if (V->IsOverlapInVolume(*this))
+							{
+								BestPriority = V->Priority;
+								BestVolume = V;
+							}
 						}
 					}
 				}
-			}
 
-			SetPhysicsVolume(BestVolume, bTriggerNotifiers);
-		}
-		else
-		{
-			Super::UpdatePhysicsVolume(bTriggerNotifiers);
+				SetPhysicsVolume(BestVolume, bTriggerNotifiers);
+			}
+			else
+			{
+				Super::UpdatePhysicsVolume(bTriggerNotifiers);
+			}
 		}
 	}
 }
@@ -2829,7 +2886,7 @@ void UPrimitiveComponent::SetLODParentPrimitive(UPrimitiveComponent * InLODParen
 	MarkRenderStateDirty();
 }
 
-UPrimitiveComponent* UPrimitiveComponent::GetLODParentPrimitive()
+UPrimitiveComponent* UPrimitiveComponent::GetLODParentPrimitive() const
 {
 	return LODParentPrimitive;
 }
diff --git a/Engine/Source/Runtime/Engine/Private/Components/SceneComponent.cpp b/Engine/Source/Runtime/Engine/Private/Components/SceneComponent.cpp
--- a/Engine/Source/Runtime/Engine/Private/Components/SceneComponent.cpp
+++ b/Engine/Source/Runtime/Engine/Private/Components/SceneComponent.cpp
@@ -1022,32 +1022,33 @@ void USceneComponent::SetWorldScale3D(FVector NewScale)
 void USceneComponent::SetWorldTransform(const FTransform& NewTransform, bool bSweep, FHitResult* OutSweepHitResult, ETeleportType Teleport)
 {
 	// If attached to something, transform into local space
-	FQuat NewRotation = NewTransform.GetRotation();
-	FVector NewLocation = NewTransform.GetTranslation();
-	FVector NewScale = NewTransform.GetScale3D();
-
 	if (AttachParent != nullptr)
 	{
-		FTransform ParentToWorld = AttachParent->GetSocketTransform(AttachSocketName);
+		const FTransform ParentToWorld = AttachParent->GetSocketTransform(AttachSocketName);
 		FTransform RelativeTM = NewTransform.GetRelativeTransform(ParentToWorld);
-		if (!bAbsoluteLocation)
+
+		// Absolute location, rotation, and scale use the world transform directly.
+		if (bAbsoluteLocation)
 		{
-			NewLocation = RelativeTM.GetTranslation();
+			RelativeTM.CopyTranslation(NewTransform);
 		}
 
-		if (!bAbsoluteRotation)
+		if (bAbsoluteRotation)
 		{
-			// Quat multiplication works reverse way, make sure you do Parent(-1) * World = Local, not World*Parent(-) = Local (the way matrix does)
-			NewRotation = RelativeTM.GetRotation();
+			RelativeTM.CopyRotation(NewTransform);
 		}
 
-		if (!bAbsoluteScale)
+		if (bAbsoluteScale)
 		{
-			NewScale = RelativeTM.GetScale3D();
+			RelativeTM.CopyScale3D(NewTransform);
 		}
-	}
 
-	SetRelativeTransform(FTransform(NewRotation, NewLocation, NewScale), bSweep, OutSweepHitResult, Teleport);
+		SetRelativeTransform(RelativeTM, bSweep, OutSweepHitResult, Teleport);
+	}
+	else
+	{
+		SetRelativeTransform(NewTransform, bSweep, OutSweepHitResult, Teleport);
+	}
 }
 
 void USceneComponent::SetWorldLocationAndRotation(FVector NewLocation, FRotator NewRotation, bool bSweep, FHitResult* OutSweepHitResult, ETeleportType Teleport)
diff --git a/Engine/Source/Runtime/Engine/Private/Components/ShapeComponent.cpp b/Engine/Source/Runtime/Engine/Private/Components/ShapeComponent.cpp
--- a/Engine/Source/Runtime/Engine/Private/Components/ShapeComponent.cpp
+++ b/Engine/Source/Runtime/Engine/Private/Components/ShapeComponent.cpp
@@ -4,6 +4,12 @@
 #include "EnginePrivate.h"
 #include "AI/Navigation/NavAreas/NavArea_Obstacle.h"
 
+#if WITH_PHYSX
+	#include "PhysicsEngine/PhysXSupport.h"
+#endif // WITH_PHYSX
+
+
+
 UShapeComponent::UShapeComponent(const FObjectInitializer& ObjectInitializer)
 	: Super(ObjectInitializer)
 {
@@ -20,6 +26,8 @@ UShapeComponent::UShapeComponent(const FObjectInitializer& ObjectInitializer)
 	bCastDynamicShadow = false;
 	ShapeColor = FColor(223, 149, 157, 255);
 	bShouldCollideWhenPlacing = false;
+
+	bUseArchetypeBodySetup = !IsTemplate();
 	
 	bHasCustomNavigableGeometry = EHasCustomNavigableGeometry::Yes;
 	bCanEverAffectNavigation = true;
@@ -81,4 +89,53 @@ bool UShapeComponent::IsNavigationRelevant() const
 	// dynamic obstacle overrides collision check
 
 	return (bDynamicObstacle && CanEverAffectNavigation()) || Super::IsNavigationRelevant();
-}
\ No newline at end of file
+}
+
+template <> void UShapeComponent::AddShapeToGeomArray<FKBoxElem>() { ShapeBodySetup->AggGeom.BoxElems.Add(FKBoxElem()); }
+template <> void UShapeComponent::AddShapeToGeomArray<FKSphereElem>() { ShapeBodySetup->AggGeom.SphereElems.Add(FKSphereElem()); }
+template <> void UShapeComponent::AddShapeToGeomArray<FKSphylElem>() { ShapeBodySetup->AggGeom.SphylElems.Add(FKSphylElem()); }
+
+#if WITH_PHYSX
+template <> void UShapeComponent::SetShapeToNewGeom<FKBoxElem>(PxShape* PShape) { PShape->userData = (void*)ShapeBodySetup->AggGeom.BoxElems[0].GetUserData(); }
+template <> void UShapeComponent::SetShapeToNewGeom<FKSphereElem>(PxShape* PShape) { PShape->userData = (void*)ShapeBodySetup->AggGeom.SphereElems[0].GetUserData(); }
+template <> void UShapeComponent::SetShapeToNewGeom<FKSphylElem>(PxShape* PShape) { PShape->userData = (void*)ShapeBodySetup->AggGeom.SphylElems[0].GetUserData(); }
+#endif
+
+template <typename ShapeElemType>
+void UShapeComponent::CreateShapeBodySetupIfNeeded()
+{
+	if (ShapeBodySetup == nullptr || ShapeBodySetup->IsPendingKill())
+	{
+		ShapeBodySetup = NewObject<UBodySetup>(this);
+		if (IsTemplate())
+		{
+			ShapeBodySetup->AddToRoot();
+		}
+
+		ShapeBodySetup->CollisionTraceFlag = CTF_UseSimpleAsComplex;
+		AddShapeToGeomArray<ShapeElemType>();
+		ShapeBodySetup->bNeverNeedsCookedCollisionData = true;
+		bUseArchetypeBodySetup = false;	//We're making our own body setup, so don't use the archetype's.
+
+		//Update bodyinstance and shapes
+		BodyInstance.BodySetup = ShapeBodySetup;
+		{
+			if(BodyInstance.IsValidBodyInstance())
+			{
+#if WITH_PHYSX
+				SCOPED_SCENE_READ_LOCK(GetPhysXSceneFromIndex(BodyInstance.GetSceneIndex()));
+				TArray<PxShape *> PShapes;
+				BodyInstance.GetAllShapes_AssumesLocked(PShapes);
+
+				check(PShapes.Num() == 1);	//Shape component should only have 1 shape
+				SetShapeToNewGeom<ShapeElemType>(PShapes[0]);
+#endif
+			}
+		}
+	}
+}
+
+//Explicit instantiation of the different shape components
+template void UShapeComponent::CreateShapeBodySetupIfNeeded<FKSphylElem>();
+template void UShapeComponent::CreateShapeBodySetupIfNeeded<FKBoxElem>();
+template void UShapeComponent::CreateShapeBodySetupIfNeeded<FKSphereElem>();
\ No newline at end of file
diff --git a/Engine/Source/Runtime/Engine/Private/Components/SkeletalMeshComponent.cpp b/Engine/Source/Runtime/Engine/Private/Components/SkeletalMeshComponent.cpp
--- a/Engine/Source/Runtime/Engine/Private/Components/SkeletalMeshComponent.cpp
+++ b/Engine/Source/Runtime/Engine/Private/Components/SkeletalMeshComponent.cpp
@@ -276,13 +276,13 @@ void USkeletalMeshComponent::UpdateClothTickRegisteredState()
 	RegisterClothTick(PrimaryComponentTick.IsTickFunctionRegistered() && ShouldRunClothTick());
 }
 
-bool USkeletalMeshComponent::NeedToSpawnAnimScriptInstance(bool bForceInit) const
+bool USkeletalMeshComponent::NeedToSpawnAnimScriptInstance() const
 {
 	IAnimClassInterface* AnimClassInterface = IAnimClassInterface::GetFromClass(AnimClass);
 	if (AnimationMode == EAnimationMode::AnimationBlueprint && (AnimClassInterface != NULL) &&
 		(SkeletalMesh != NULL) && (SkeletalMesh->Skeleton->IsCompatible(AnimClassInterface->GetTargetSkeleton())))
 	{
-		if (bForceInit || (AnimScriptInstance == NULL) || (AnimScriptInstance->GetClass() != AnimClass) )
+		if ( (AnimScriptInstance == NULL) || (AnimScriptInstance->GetClass() != AnimClass) )
 		{
 			return true;
 		}
@@ -319,7 +319,7 @@ void USkeletalMeshComponent::OnUnregister()
 	ReleaseAllClothingResources();
 #endif// #if WITH_APEX_CLOTHING
 
-	if (AnimScriptInstance && bReInitAnimationOnSetSkeletalMeshCalls)
+	if (AnimScriptInstance)
 	{
 		AnimScriptInstance->UninitializeAnimation();
 	}
@@ -353,21 +353,39 @@ void USkeletalMeshComponent::InitAnim(bool bForceReinit)
 		// this has to be called before Initialize Animation because it will required RequiredBones list when InitializeAnimScript
 		RecalcRequiredBones(0);
 
-		InitializeAnimScriptInstance(bForceReinit);
-
-		//Make sure we have a valid pose		
-		TickAnimation(0.f, false); 
+		bool bDoRefreshBoneTransform = true;
+		if (InitializeAnimScriptInstance(bForceReinit))
+		{
+			//Make sure we have a valid pose
+			if (bUseRefPoseOnInitAnim)
+			{
+				LocalAtoms = SkeletalMesh->RefSkeleton.GetRefBonePose();
+				//Mini RefreshBoneTransforms (the bit we actually care about)
+				FillSpaceBases(SkeletalMesh, LocalAtoms, GetEditableSpaceBases());
+				bNeedToFlipSpaceBaseBuffers = true; // Have updated space bases so need to flip
+				FlipEditableSpaceBases();
+				bDoRefreshBoneTransform = false;
+			}
+			else
+			{
+				TickAnimation(0.f, false);
+			}
+		}
 
-		RefreshBoneTransforms();
+		if (bDoRefreshBoneTransform)
+		{
+			RefreshBoneTransforms();
+		}
 		UpdateComponentToWorld();
 	}
 }
 
-void USkeletalMeshComponent::InitializeAnimScriptInstance(bool bForceReinit)
+bool USkeletalMeshComponent::InitializeAnimScriptInstance(bool bForceReinit)
 {
+	bool bCalledInitialize = false;
 	if (IsRegistered())
 	{
-		if (NeedToSpawnAnimScriptInstance(bForceReinit))
+		if (NeedToSpawnAnimScriptInstance())
 		{
 			SCOPE_CYCLE_COUNTER(STAT_AnimSpawnTime);
 			AnimScriptInstance = NewObject<UAnimInstance>(this, AnimClass);
@@ -375,6 +393,7 @@ void USkeletalMeshComponent::InitializeAnimScriptInstance(bool bForceReinit)
 			if (AnimScriptInstance)
 			{
 				AnimScriptInstance->InitializeAnimation();
+				bCalledInitialize = true;
 			}
 		}
 		else if (AnimationMode == EAnimationMode::AnimationSingleNode)
@@ -392,6 +411,7 @@ void USkeletalMeshComponent::InitializeAnimScriptInstance(bool bForceReinit)
 			if (AnimScriptInstance)
 			{
 				AnimScriptInstance->InitializeAnimation();
+				bCalledInitialize = true;
 			}
 
 			if (OldInstance && AnimScriptInstance)
@@ -402,11 +422,14 @@ void USkeletalMeshComponent::InitializeAnimScriptInstance(bool bForceReinit)
 				CachedData.Initialize(Cast<UAnimSingleNodeInstance>(AnimScriptInstance));
 			}
 		}
-		else if (AnimScriptInstance && bReInitAnimationOnSetSkeletalMeshCalls)
+
+		if (AnimScriptInstance && !bCalledInitialize && bForceReinit)
 		{
 			AnimScriptInstance->InitializeAnimation();
-		}		
+			bCalledInitialize = true;
+		}
 	}
+	return bCalledInitialize;
 }
 
 bool USkeletalMeshComponent::IsWindEnabled() const
@@ -459,7 +482,7 @@ void USkeletalMeshComponent::PostEditChangeProperty(FPropertyChangedEvent& Prope
 				}
 				else
 				{
-					if (NeedToSpawnAnimScriptInstance(false))
+					if (NeedToSpawnAnimScriptInstance())
 					{
 						SCOPE_CYCLE_COUNTER(STAT_AnimSpawnTime);
 						AnimScriptInstance = NewObject<UAnimInstance>(this, AnimClass);
@@ -1359,7 +1382,7 @@ FBoxSphereBounds USkeletalMeshComponent::CalcBounds(const FTransform& LocalToWor
 	}
 }
 
-void USkeletalMeshComponent::SetSkeletalMesh(USkeletalMesh* InSkelMesh)
+void USkeletalMeshComponent::SetSkeletalMesh(USkeletalMesh* InSkelMesh, bool bReinitPose)
 {
 	if (InSkelMesh == SkeletalMesh)
 	{
@@ -1369,7 +1392,7 @@ void USkeletalMeshComponent::SetSkeletalMesh(USkeletalMesh* InSkelMesh)
 
 	UPhysicsAsset* OldPhysAsset = GetPhysicsAsset();
 
-	Super::SetSkeletalMesh(InSkelMesh);
+	Super::SetSkeletalMesh(InSkelMesh, bReinitPose);
 
 #if WITH_EDITOR
 	ValidateAnimation();
@@ -1382,7 +1405,7 @@ void USkeletalMeshComponent::SetSkeletalMesh(USkeletalMesh* InSkelMesh)
 
 	UpdateHasValidBodies();
 
-	InitAnim(false);
+	InitAnim(bReinitPose);
 
 #if WITH_APEX_CLOTHING
 	RecreateClothingActors();
@@ -1394,9 +1417,7 @@ void USkeletalMeshComponent::SetSkeletalMesh(USkeletalMesh* InSkelMesh)
 
 void USkeletalMeshComponent::SetSkeletalMeshWithoutResettingAnimation(USkeletalMesh* InSkelMesh)
 {
-	bReInitAnimationOnSetSkeletalMeshCalls = false;
-	SetSkeletalMesh(InSkelMesh);
-	bReInitAnimationOnSetSkeletalMeshCalls = true;
+	SetSkeletalMesh(InSkelMesh,false);
 }
 
 bool USkeletalMeshComponent::AllocateTransformData()
diff --git a/Engine/Source/Runtime/Engine/Private/Components/SkinnedMeshComponent.cpp b/Engine/Source/Runtime/Engine/Private/Components/SkinnedMeshComponent.cpp
--- a/Engine/Source/Runtime/Engine/Private/Components/SkinnedMeshComponent.cpp
+++ b/Engine/Source/Runtime/Engine/Private/Components/SkinnedMeshComponent.cpp
@@ -297,6 +297,7 @@ void USkinnedMeshComponent::OnRegister()
 	Super::OnRegister();
 
 	UpdateLODStatus();
+	InvalidateCachedBounds();
 }
 
 void USkinnedMeshComponent::OnUnregister()
@@ -951,7 +952,7 @@ FVector USkinnedMeshComponent::GetRefPosePosition(int32 BoneIndex)
 }
 
 
-void USkinnedMeshComponent::SetSkeletalMesh(USkeletalMesh* InSkelMesh)
+void USkinnedMeshComponent::SetSkeletalMesh(USkeletalMesh* InSkelMesh, bool bReinitPose)
 {
 	// NOTE: InSkelMesh may be NULL (useful in the editor for removing the skeletal mesh associated with
 	//   this component on-the-fly)
@@ -962,13 +963,28 @@ void USkinnedMeshComponent::SetSkeletalMesh(USkeletalMesh* InSkelMesh)
 		return;
 	}
 
-	// Unregister the component, swap the skeletal mesh, and reregister.
+	// Destroy render state, swap skeletal mesh and recreate
+	const bool bWasRenderStateCreated = bRenderStateCreated;
+	if (bWasRenderStateCreated)
 	{
-		FComponentReregisterContext ReregisterContext(this);
-		check(MeshObject == NULL);
-		SkeletalMesh = InSkelMesh;
+		DestroyRenderState_Concurrent();
+	}
+
+	SkeletalMesh = InSkelMesh;
+
+	// Don't init anim state if not registered
+	if (IsRegistered())
+	{
+		AllocateTransformData();
+		UpdateMasterBoneMap();
+		UpdateLODStatus();
 		InvalidateCachedBounds();
 	}
+
+	if (bWasRenderStateCreated)
+	{
+		CreateRenderState_Concurrent();
+	}
 	
 	// TODO: (LH) find better way to call this 
 	// Right now SkeletalMeshComponent(child) needs to call this after animation update
diff --git a/Engine/Source/Runtime/Engine/Private/Components/SphereComponent.cpp b/Engine/Source/Runtime/Engine/Private/Components/SphereComponent.cpp
--- a/Engine/Source/Runtime/Engine/Private/Components/SphereComponent.cpp
+++ b/Engine/Source/Runtime/Engine/Private/Components/SphereComponent.cpp
@@ -24,29 +24,51 @@ void USphereComponent::CalcBoundingCylinder(float& CylinderRadius, float& Cylind
 	CylinderHalfHeight = CylinderRadius;
 }
 
-void USphereComponent::UpdateBodySetup()
+template <EShapeBodySetupHelper UpdateBodySetupAction>
+bool InvalidateOrUpdateSphereBodySetup(UBodySetup*& ShapeBodySetup, bool bUseArchetypeBodySetup, float SphereRadius)
 {
-	if (ShapeBodySetup == NULL || ShapeBodySetup->IsPendingKill())
-	{
-		ShapeBodySetup = NewObject<UBodySetup>(this);
-		ShapeBodySetup->CollisionTraceFlag = CTF_UseSimpleAsComplex;
-		ShapeBodySetup->AggGeom.SphereElems.Add(FKSphereElem());
-		ShapeBodySetup->bNeverNeedsCookedCollisionData = true;
-	}
-
-	check (ShapeBodySetup->AggGeom.SphereElems.Num() == 1);
-	FKSphereElem* se = ShapeBodySetup->AggGeom.SphereElems.GetData();
+	check((bUseArchetypeBodySetup && UpdateBodySetupAction == EShapeBodySetupHelper::InvalidateSharingIfStale) || (!bUseArchetypeBodySetup && UpdateBodySetupAction == EShapeBodySetupHelper::UpdateBodySetup) );
+	check(ShapeBodySetup->AggGeom.SphereElems.Num() == 1);
+	FKSphereElem* SphereElem = ShapeBodySetup->AggGeom.SphereElems.GetData();
 
 	// check for mal formed values
 	float Radius = SphereRadius;
-	if( Radius < KINDA_SMALL_NUMBER )
+	if (Radius < KINDA_SMALL_NUMBER)
 	{
 		Radius = 0.1f;
 	}
 
-	// now set the PhysX data values
-	se->Center = FVector::ZeroVector;
-	se->Radius = Radius;
+	if(UpdateBodySetupAction == EShapeBodySetupHelper::UpdateBodySetup)
+	{
+		// now set the PhysX data values
+		SphereElem->Center = FVector::ZeroVector;
+		SphereElem->Radius = Radius;
+	}
+	else
+	{
+		if(SphereElem->Radius != Radius)
+		{
+			ShapeBodySetup = nullptr;
+			bUseArchetypeBodySetup = false;
+		}
+	}
+	
+	return bUseArchetypeBodySetup;
+}
+
+void USphereComponent::UpdateBodySetup()
+{
+	if(PrepareSharedBodySetup<USphereComponent>())
+	{
+		bUseArchetypeBodySetup = InvalidateOrUpdateSphereBodySetup<EShapeBodySetupHelper::InvalidateSharingIfStale>(ShapeBodySetup, bUseArchetypeBodySetup, SphereRadius);
+	}
+
+	CreateShapeBodySetupIfNeeded<FKSphereElem>();
+
+	if(!bUseArchetypeBodySetup)
+	{
+		InvalidateOrUpdateSphereBodySetup<EShapeBodySetupHelper::UpdateBodySetup>(ShapeBodySetup, bUseArchetypeBodySetup, SphereRadius);
+	}
 }
 
 void USphereComponent::SetSphereRadius( float InSphereRadius, bool bUpdateOverlaps )
diff --git a/Engine/Source/Runtime/Engine/Private/HUD.cpp b/Engine/Source/Runtime/Engine/Private/HUD.cpp
--- a/Engine/Source/Runtime/Engine/Private/HUD.cpp
+++ b/Engine/Source/Runtime/Engine/Private/HUD.cpp
@@ -178,7 +178,7 @@ void AHUD::PostRender()
 						ContactPoint += ContactPointOffset;
 					}
 				}
-				UpdateHitBoxCandidates( ContactPoints );
+				UpdateHitBoxCandidates( MoveTemp(ContactPoints) );
 			}
 		}
 		else if (HitBoxesOver.Num() > 0)
diff --git a/Engine/Source/Runtime/Engine/Private/PackageMapClient.cpp b/Engine/Source/Runtime/Engine/Private/PackageMapClient.cpp
--- a/Engine/Source/Runtime/Engine/Private/PackageMapClient.cpp
+++ b/Engine/Source/Runtime/Engine/Private/PackageMapClient.cpp
@@ -192,7 +192,7 @@ bool UPackageMapClient::WriteObject( FArchive& Ar, UObject* ObjOuter, FNetworkGU
  *		For static actors, this will just be a single call to SerializeObject, since they can be referenced by their path name.
  *		For dynamic actors, first the actor's reference is serialized but will not resolve on clients since they haven't spawned the actor yet.
  *		The actor archetype is then serialized along with the starting location, rotation, and velocity.
- *		After reading this information, the client spawns this actor via GWorld and assigns it the NetGUID it read at the top of the function.
+ *		After reading this information, the client spawns this actor in the NetDriver's World and assigns it the NetGUID it read at the top of the function.
  *
  *		returns true if a new actor was spawned. false means an existing actor was found for the netguid.
  */
diff --git a/Engine/Source/Runtime/Engine/Private/Particles/ParticleComponents.cpp b/Engine/Source/Runtime/Engine/Private/Particles/ParticleComponents.cpp
--- a/Engine/Source/Runtime/Engine/Private/Particles/ParticleComponents.cpp
+++ b/Engine/Source/Runtime/Engine/Private/Particles/ParticleComponents.cpp
@@ -3104,7 +3104,6 @@ SIZE_T UParticleSystemComponent::GetResourceSize(EResourceSizeMode::Type Mode)
 
 bool UParticleSystemComponent::ParticleLineCheck(FHitResult& Hit, AActor* SourceActor, const FVector& End, const FVector& Start, const FVector& HalfExtent, const FCollisionObjectQueryParams& ObjectParams)
 {
-	ForceAsyncWorkCompletion(ENSURE_AND_STALL);
 	check(GetWorld());
 	static FName NAME_ParticleCollision = FName(TEXT("ParticleCollision"));
 
@@ -6071,6 +6070,10 @@ void UParticleSystemComponent::GetUsedMaterials( TArray<UMaterialInterface*>& Ou
 	}
 }
 
+FBodyInstance* UParticleSystemComponent::GetBodyInstance(FName BoneName /*= NAME_None*/, bool bGetWelded /*= true*/) const
+{
+	return nullptr;
+}
 
 void UParticleSystemComponent::ReportEventSpawn(const FName InEventName, const float InEmitterTime, 
 	const FVector InLocation, const FVector InVelocity, const TArray<UParticleModuleEventSendToGame*>& InEventData)
diff --git a/Engine/Source/Runtime/Engine/Private/Particles/ParticleModules_Collision.cpp b/Engine/Source/Runtime/Engine/Private/Particles/ParticleModules_Collision.cpp
--- a/Engine/Source/Runtime/Engine/Private/Particles/ParticleModules_Collision.cpp
+++ b/Engine/Source/Runtime/Engine/Private/Particles/ParticleModules_Collision.cpp
@@ -50,6 +50,7 @@ UParticleModuleCollision::UParticleModuleCollision(const FObjectInitializer& Obj
 	bCollideOnlyIfVisible = true;
 	MaxCollisionDistance = 1000.0f;
 	bIgnoreSourceActor = true;
+	bIgnoreTriggerVolumes = true;
 	CollisionTypes.Add(UEngineTypes::ConvertToObjectType(ECC_WorldStatic));
 }
 
@@ -389,7 +390,7 @@ void UParticleModuleCollision::Update(FParticleEmitterInstance* Owner, int32 Off
 			if (Hit.GetActor())
 			{
 				bDecrementMaxCount = !bPawnsDoNotDecrementCount || !Cast<APawn>(Hit.GetActor());
-				bIgnoreCollision = Hit.GetActor()->IsA(ATriggerBase::StaticClass());
+				bIgnoreCollision = bIgnoreTriggerVolumes && Hit.GetActor()->IsA(ATriggerBase::StaticClass());
 				[REDACTED_EMAIL]. Allow for PSys to say what it wants to collide w/?
 			}
 
@@ -436,6 +437,7 @@ void UParticleModuleCollision::Update(FParticleEmitterInstance* Owner, int32 Off
 
 						if (bApplyPhysics)
 						{
+							check(IsInGameThread());
 							UPrimitiveComponent* PrimitiveComponent = Hit.Component.Get();
 							if(PrimitiveComponent && PrimitiveComponent->IsAnySimulatingPhysics())
 							{
@@ -465,6 +467,7 @@ void UParticleModuleCollision::Update(FParticleEmitterInstance* Owner, int32 Off
 
 						if (bApplyPhysics)
 						{
+							check(IsInGameThread());
 							UPrimitiveComponent* PrimitiveComponent = Hit.Component.Get();
 							if(PrimitiveComponent && PrimitiveComponent->IsAnySimulatingPhysics())
 							{
@@ -544,6 +547,10 @@ void UParticleModuleCollision::Update(FParticleEmitterInstance* Owner, int32 Off
 	END_UPDATE_LOOP;
 }
 
+bool UParticleModuleCollision::CanTickInAnyThread()
+{
+	return false; //TODO: make all of the update function thread safe before turning this one /*!bIgnoreTriggerVolumes && !bApplyPhysics && !bPawnsDoNotDecrementCount;*/
+}
 
 void UParticleModuleCollision::SetToSensibleDefaults(UParticleEmitter* Owner)
 {
diff --git a/Engine/Source/Runtime/Engine/Private/PhysicsEngine/BodyInstance.cpp b/Engine/Source/Runtime/Engine/Private/PhysicsEngine/BodyInstance.cpp
--- a/Engine/Source/Runtime/Engine/Private/PhysicsEngine/BodyInstance.cpp
+++ b/Engine/Source/Runtime/Engine/Private/PhysicsEngine/BodyInstance.cpp
@@ -513,13 +513,14 @@ void FBodyInstance::SetCollisionEnabled(ECollisionEnabled::Type NewType, bool bU
 			UpdatePhysicsFilterData();
 		}
 
-		// If we used to be QueryOnly we have to set our dynamic properties since they were skipped previously
-		if (OldType == ECollisionEnabled::QueryOnly)
+		//If we are going from QueryOnly to simulation we need to recreate the physics state. This is because physx doesn't have the actors in its simulation scene
+		if (OldType == ECollisionEnabled::QueryOnly && CollisionEnabled != ECollisionEnabled::NoCollision)
 		{
-			ExecuteOnPhysicsReadWrite([&]
+			if(UPrimitiveComponent* PrimComponent = OwnerComponent.Get())
 			{
-				InitDynamicProperties_AssumesLocked();
-			});
+				PrimComponent->RecreatePhysicsState();
+			}
+
 		}
 
 	}
@@ -1121,12 +1122,26 @@ struct FInitBodiesHelper
 	{
 		physx::PxRigidDynamic* PNewDynamic = nullptr;
 
+		const ECollisionEnabled::Type CollisionType = Instance->GetCollisionEnabled();
+		const bool bDisableSim = CollisionType == ECollisionEnabled::QueryOnly || CollisionType == ECollisionEnabled::NoCollision;
+
 		if (bCompileStatic || bStatic)
 		{
 			Instance->RigidActorSync = GPhysXSDK->createRigidStatic(PTransform);
+
+			if(bDisableSim)
+			{
+				Instance->RigidActorSync->setActorFlag(PxActorFlag::eDISABLE_SIMULATION, true);
+			}
+
 			if (PAsyncScene)
 			{
 				Instance->RigidActorAsync = GPhysXSDK->createRigidStatic(PTransform);
+
+				if (bDisableSim)
+				{
+					Instance->RigidActorAsync->setActorFlag(PxActorFlag::eDISABLE_SIMULATION, true);
+				}
 			}
 		}
 		else
@@ -1146,10 +1161,19 @@ struct FInitBodiesHelper
 				PNewDynamic->setRigidDynamicFlag(PxRigidDynamicFlag::eKINEMATIC, true);
 			}
 
+			PxActorFlags ActorFlags = PNewDynamic->getActorFlags();
+			
 			if(Instance->bGenerateWakeEvents)
 			{
-				PNewDynamic->setActorFlag(PxActorFlag::eSEND_SLEEP_NOTIFIES, true);
+				ActorFlags |= PxActorFlag::eSEND_SLEEP_NOTIFIES;
+			}
+
+			if(bDisableSim)
+			{
+				ActorFlags |= PxActorFlag::eDISABLE_SIMULATION;
 			}
+
+			PNewDynamic->setActorFlags(ActorFlags);
 		}
 
 		return PNewDynamic;
@@ -1765,6 +1789,7 @@ const FBodyInstance* FBodyInstance::GetOriginalBodyInstance(const PxShape* PShap
 
 const FTransform& FBodyInstance::GetRelativeBodyTransform(const physx::PxShape* PShape) const
 {
+	check(IsInGameThread());
 	const FBodyInstance* BI = WeldParent ? WeldParent : this;
 	const FWeldInfo* Result = BI->ShapeToBodiesMap.IsValid() ? BI->ShapeToBodiesMap->Find(PShape) : nullptr;
 	return Result ? Result->RelativeTM : FTransform::Identity;
@@ -1919,6 +1944,7 @@ void FBodyInstance::TermBody()
 
 bool FBodyInstance::Weld(FBodyInstance* TheirBody, const FTransform& TheirTM)
 {
+	check(IsInGameThread());
 	check(TheirBody);
 	if (TheirBody->BodySetup.IsValid() == false)	//attach actor can be called before body has been initialized. In this case just return false
 	{
@@ -2001,6 +2027,8 @@ bool FBodyInstance::Weld(FBodyInstance* TheirBody, const FTransform& TheirTM)
 
 void FBodyInstance::UnWeld(FBodyInstance* TheirBI)
 {
+	check(IsInGameThread());
+
 	//@TODO: BOX2D: Implement Weld
 
 #if WITH_PHYSX
@@ -2669,8 +2697,9 @@ void FBodyInstance::SetBodyTransform(const FTransform& NewTransform, ETeleportTy
 			// SIMULATED & KINEMATIC
 			if (PxRigidDynamic* PRigidDynamic = GetPxRigidDynamic_AssumesLocked())
 			{
+				const bool bIsRigidBodyKinematic = IsRigidBodyKinematicAndInSimulationScene_AssumesLocked(PRigidDynamic);
 				// If kinematic and not teleporting, set kinematic target
-				if (!IsRigidBodyNonKinematic_AssumesLocked(PRigidDynamic) && Teleport == ETeleportType::None)
+				if (bIsRigidBodyKinematic && Teleport == ETeleportType::None)
 				{
 					if(FPhysScene* PhysScene = GetPhysicsScene(this))
 					{
@@ -2680,7 +2709,7 @@ void FBodyInstance::SetBodyTransform(const FTransform& NewTransform, ETeleportTy
 				// Otherwise, set global pose
 				else
 				{
-					if (!IsRigidBodyNonKinematic_AssumesLocked(PRigidDynamic))  // check if kinematic  (checks the physx bit for this)
+					if (bIsRigidBodyKinematic)  // check if kinematic  (checks the physx bit for this)
 					{
 						PRigidDynamic->setKinematicTarget(PNewPose);  // physx doesn't clear target on setGlobalPose, so overwrite any previous attempt to set this that wasn't yet resolved
 					}
@@ -3356,10 +3385,10 @@ void FBodyInstance::WakeInstance()
 #if WITH_PHYSX
 	ExecuteOnPxRigidDynamicReadWrite(this, [&](PxRigidDynamic* PRigidDynamic)
 	{
-		if (IsRigidBodyNonKinematic_AssumesLocked(PRigidDynamic))
-	{
-		PRigidDynamic->wakeUp();
-	}
+		if (!IsRigidBodyKinematic_AssumesLocked(PRigidDynamic))
+		{
+			PRigidDynamic->wakeUp();
+		}
 	});
 #endif
 
@@ -3376,10 +3405,10 @@ void FBodyInstance::PutInstanceToSleep()
 #if WITH_PHYSX
 	ExecuteOnPxRigidDynamicReadWrite(this, [&](PxRigidDynamic* PRigidDynamic)
 	{
-		if (IsRigidBodyNonKinematic_AssumesLocked(PRigidDynamic))
-	{
-		PRigidDynamic->putToSleep();
-	}
+		if (!IsRigidBodyKinematic_AssumesLocked(PRigidDynamic))
+		{
+			PRigidDynamic->putToSleep();
+		}
 	});
 #endif //WITH_PHYSX
 
@@ -3535,7 +3564,7 @@ void FBodyInstance::AddCustomPhysics(FCalculateCustomPhysics& CalculateCustomPhy
 #if WITH_PHYSX
 	ExecuteOnPxRigidBodyReadOnly(this, [&](const PxRigidBody* PRigidBody)
 	{
-		if (IsRigidBodyNonKinematic_AssumesLocked(PRigidBody))
+		if (!IsRigidBodyKinematic_AssumesLocked(PRigidBody))
 		{
 			if(FPhysScene* PhysScene = GetPhysicsScene(this))
 			{
@@ -3560,7 +3589,7 @@ void FBodyInstance::AddForce(const FVector& Force, bool bAllowSubstepping, bool
 #if WITH_PHYSX
 	ExecuteOnPxRigidBodyReadWrite(this, [&](PxRigidBody* PRigidBody)
 	{
-		if (IsRigidBodyNonKinematic_AssumesLocked(PRigidBody))
+		if (!IsRigidBodyKinematic_AssumesLocked(PRigidBody))
 		{
 			if(FPhysScene* PhysScene = GetPhysicsScene(this))
 			{
@@ -3588,7 +3617,7 @@ void FBodyInstance::AddForceAtPosition(const FVector& Force, const FVector& Posi
 #if WITH_PHYSX
 	ExecuteOnPxRigidBodyReadWrite(this, [&](PxRigidBody* PRigidBody)
 	{
-		if (IsRigidBodyNonKinematic_AssumesLocked(PRigidBody))
+		if (!IsRigidBodyKinematic_AssumesLocked(PRigidBody))
 		{
 			if(FPhysScene* PhysScene = GetPhysicsScene(this))
 			{
@@ -3617,7 +3646,7 @@ void FBodyInstance::AddTorque(const FVector& Torque, bool bAllowSubstepping, boo
 #if WITH_PHYSX
 	ExecuteOnPxRigidBodyReadWrite(this, [&](PxRigidBody* PRigidBody)
 	{
-		if (IsRigidBodyNonKinematic_AssumesLocked(PRigidBody))
+		if (!IsRigidBodyKinematic_AssumesLocked(PRigidBody))
 		{
 			if(FPhysScene* PhysScene = GetPhysicsScene(this))
 			{
@@ -3644,11 +3673,11 @@ void FBodyInstance::AddAngularImpulse(const FVector& AngularImpulse, bool bVelCh
 #if WITH_PHYSX
 	ExecuteOnPxRigidBodyReadWrite(this, [&](PxRigidBody* PRigidBody)
 	{
-		if (IsRigidBodyNonKinematic_AssumesLocked(PRigidBody))
-	{
-		PxForceMode::Enum Mode = bVelChange ? PxForceMode::eVELOCITY_CHANGE : PxForceMode::eIMPULSE;
-		PRigidBody->addTorque(U2PVector(AngularImpulse), Mode, true);
-	}
+		if (!IsRigidBodyKinematic_AssumesLocked(PRigidBody))
+		{
+			PxForceMode::Enum Mode = bVelChange ? PxForceMode::eVELOCITY_CHANGE : PxForceMode::eIMPULSE;
+			PRigidBody->addTorque(U2PVector(AngularImpulse), Mode, true);
+		}
 	});
 	
 #endif // WITH_PHYSX
@@ -3670,11 +3699,11 @@ void FBodyInstance::AddImpulse(const FVector& Impulse, bool bVelChange)
 #if WITH_PHYSX
 	ExecuteOnPxRigidBodyReadWrite(this, [&](PxRigidBody* PRigidBody)
 	{
-		if (IsRigidBodyNonKinematic_AssumesLocked(PRigidBody))
-	{
-		PxForceMode::Enum Mode = bVelChange ? PxForceMode::eVELOCITY_CHANGE : PxForceMode::eIMPULSE;
-		PRigidBody->addForce(U2PVector(Impulse), Mode, true);
-	}
+		if (!IsRigidBodyKinematic_AssumesLocked(PRigidBody))
+		{
+			PxForceMode::Enum Mode = bVelChange ? PxForceMode::eVELOCITY_CHANGE : PxForceMode::eIMPULSE;
+			PRigidBody->addForce(U2PVector(Impulse), Mode, true);
+		}
 	});
 	
 #endif // WITH_PHYSX
@@ -3696,11 +3725,11 @@ void FBodyInstance::AddImpulseAtPosition(const FVector& Impulse, const FVector&
 #if WITH_PHYSX
 	ExecuteOnPxRigidBodyReadWrite(this, [&](PxRigidBody* PRigidBody)
 	{
-		if (IsRigidBodyNonKinematic_AssumesLocked(PRigidBody))
-	{
-		PxForceMode::Enum Mode = PxForceMode::eIMPULSE; // does not support eVELOCITY_CHANGE
-		PxRigidBodyExt::addForceAtPos(*PRigidBody, U2PVector(Impulse), U2PVector(Position), Mode, true);
-	}
+		if (!IsRigidBodyKinematic_AssumesLocked(PRigidBody))
+		{
+			PxForceMode::Enum Mode = PxForceMode::eIMPULSE; // does not support eVELOCITY_CHANGE
+			PxRigidBodyExt::addForceAtPos(*PRigidBody, U2PVector(Impulse), U2PVector(Position), Mode, true);
+		}
 	});
 	
 #endif // WITH_PHYSX
@@ -3799,10 +3828,10 @@ void FBodyInstance::AddRadialImpulseToBody(const FVector& Origin, float Radius,
 #if WITH_PHYSX
 	ExecuteOnPxRigidBodyReadWrite(this, [&](PxRigidBody* PRigidBody)
 	{
-		if (IsRigidBodyNonKinematic_AssumesLocked(PRigidBody))
-	{
+		if (!IsRigidBodyKinematic_AssumesLocked(PRigidBody))
+		{
 			AddRadialImpulseToPxRigidBody_AssumesLocked(*PRigidBody, Origin, Radius, Strength, Falloff, bVelChange);
-	}
+		}
 	});
 	
 #endif // WITH_PHYSX
@@ -3820,7 +3849,7 @@ void FBodyInstance::AddRadialForceToBody(const FVector& Origin, float Radius, fl
 #if WITH_PHYSX
 	ExecuteOnPxRigidBodyReadWrite(this, [&](PxRigidBody* PRigidBody)
 	{
-		if (IsRigidBodyNonKinematic_AssumesLocked(PRigidBody))
+		if (!IsRigidBodyKinematic_AssumesLocked(PRigidBody))
 		{
 			if(FPhysScene* PhysScene = GetPhysicsScene(this))
 			{
@@ -4628,7 +4657,7 @@ void FBodyInstance::InitDynamicProperties_AssumesLocked()
 		if(RigidActor->getScene())
 		{
 			CreateDOFLock();
-			if(IsRigidBodyNonKinematic_AssumesLocked(RigidActor))
+			if(!IsRigidBodyKinematic_AssumesLocked(RigidActor))
 			{
 				if(bStartAwake || bWokenExternally)
 				{
@@ -4934,10 +4963,12 @@ void FBodyInstance::GetShapeFlags_AssumesLocked(FShapeData& ShapeData, TEnumAsBy
 		}
 
 		// enable swept bounds for CCD for this shape
-		PxRigidBody* PBody = GetPxRigidActor_AssumesLocked()->is<PxRigidBody>();
-		if(bSimCollision && !bPhysicsStatic && bUseCCD && PBody)
+		if(bSimCollision && !bPhysicsStatic && bUseCCD )
 		{
-			ShapeData.SyncBodyFlags |= PxRigidBodyFlag::eENABLE_CCD;
+			if(GetPxRigidActor_AssumesLocked()->is<PxRigidBody>())
+			{
+				ShapeData.SyncBodyFlags |= PxRigidBodyFlag::eENABLE_CCD;
+			}
 		}
 	}
 	// No collision enabled
diff --git a/Engine/Source/Runtime/Engine/Private/PhysicsEngine/BodySetup.cpp b/Engine/Source/Runtime/Engine/Private/PhysicsEngine/BodySetup.cpp
--- a/Engine/Source/Runtime/Engine/Private/PhysicsEngine/BodySetup.cpp
+++ b/Engine/Source/Runtime/Engine/Private/PhysicsEngine/BodySetup.cpp
@@ -276,7 +276,7 @@ bool CalcMeshNegScaleCompensation(const FVector& InScale3D, PxTransform& POutTra
 	return (InScale3D.X * InScale3D.Y * InScale3D.Z) < 0.f;
 }
 
-void SetupNonUniformHelper(FVector& Scale3D, float& MinScale, float& MinScaleAbs, FVector& Scale3DAbs)
+void SetupNonUniformHelper(FVector Scale3D, float& MinScale, float& MinScaleAbs, FVector& Scale3DAbs)
 {
 	// if almost zero, set min scale
 	// @todo fixme
@@ -366,57 +366,52 @@ struct FAddShapesHelper
 	float MaxContactOffset;
 
 public:
-	FORCEINLINE_DEBUGGABLE void AddSpheresToRigidActor_AssumesLocked() const
+	void AddSpheresToRigidActor_AssumesLocked() const
 	{
 		for (int32 i = 0; i < BodySetup->AggGeom.SphereElems.Num(); i++)
 		{
 			const FKSphereElem& SphereElem = BodySetup->AggGeom.SphereElems[i];
+			const FKSphereElem ScaledSphereElem = SphereElem.GetFinalScaled(Scale3D, RelativeTM);
 
 			PxSphereGeometry PSphereGeom;
-			PSphereGeom.radius = (SphereElem.Radius * MinScaleAbs);
+			PSphereGeom.radius = ScaledSphereElem.Radius;
 
 			if (ensure(PSphereGeom.isValid()))
 			{
-				FVector LocalOrigin = RelativeTM.TransformPosition(SphereElem.Center);
-				PxTransform PLocalPose(U2PVector(LocalOrigin));
-				PLocalPose.p *= MinScale;
-
+				PxTransform PLocalPose(U2PVector(ScaledSphereElem.Center));
 				ensure(PLocalPose.isValid());
 				{
 					const float ContactOffset = FMath::Min(MaxContactOffset, ContactOffsetFactor * PSphereGeom.radius);
-					PxShape* PShape = AttachShape_AssumesLocked(PSphereGeom, PLocalPose, ContactOffset, &SphereElem);
+					PxShape* PShape = AttachShape_AssumesLocked(PSphereGeom, PLocalPose, ContactOffset, SphereElem.GetUserData());
 				}
 			}
 			else
 			{
-				UE_LOG(LogPhysics, Warning, TEXT("AddSpheresToRigidActor: [%s] SphereElem[%d] invalid"), *GetPathNameSafe(BodySetup->GetOuter()), i);
+				UE_LOG(LogPhysics, Warning, TEXT("AddSpheresToRigidActor: [%s] ScaledSphereElem[%d] invalid"), *GetPathNameSafe(BodySetup->GetOuter()), i);
 			}
 		}
 	}
 
-	FORCEINLINE_DEBUGGABLE void AddBoxesToRigidActor_AssumesLocked() const
+	void AddBoxesToRigidActor_AssumesLocked() const
 	{
 		for (int32 i = 0; i < BodySetup->AggGeom.BoxElems.Num(); i++)
 		{
 			const FKBoxElem& BoxElem = BodySetup->AggGeom.BoxElems[i];
-
+			const FKBoxElem ScaledBoxElem = BoxElem.GetFinalScaled(Scale3D, RelativeTM);
+			const FTransform& BoxTransform = ScaledBoxElem.GetTransform();
+			
 			PxBoxGeometry PBoxGeom;
-			PBoxGeom.halfExtents.x = (0.5f * BoxElem.X * Scale3DAbs.X);
-			PBoxGeom.halfExtents.y = (0.5f * BoxElem.Y * Scale3DAbs.Y);
-			PBoxGeom.halfExtents.z = (0.5f * BoxElem.Z * Scale3DAbs.Z);
+			PBoxGeom.halfExtents.x = ScaledBoxElem.X * 0.5f;
+			PBoxGeom.halfExtents.y = ScaledBoxElem.Y * 0.5f;
+			PBoxGeom.halfExtents.z = ScaledBoxElem.Z * 0.5f;
 
-			FTransform BoxTransform = BoxElem.GetTransform() * RelativeTM;
 			if (PBoxGeom.isValid() && BoxTransform.IsValid())
 			{
 				PxTransform PLocalPose(U2PTransform(BoxTransform));
-				PLocalPose.p.x *= Scale3D.X;
-				PLocalPose.p.y *= Scale3D.Y;
-				PLocalPose.p.z *= Scale3D.Z;
-
 				ensure(PLocalPose.isValid());
 				{
 					const float ContactOffset = FMath::Min(MaxContactOffset, ContactOffsetFactor * PBoxGeom.halfExtents.minElement());
-					AttachShape_AssumesLocked(PBoxGeom, PLocalPose, ContactOffset, &BoxElem);
+					AttachShape_AssumesLocked(PBoxGeom, PLocalPose, ContactOffset, BoxElem.GetUserData());
 				}
 			}
 			else
@@ -426,7 +421,7 @@ public:
 		}
 	}
 
-	FORCEINLINE_DEBUGGABLE void AddSphylsToRigidActor_AssumesLocked() const
+	void AddSphylsToRigidActor_AssumesLocked() const
 	{
 		float ScaleRadius = FMath::Max(Scale3DAbs.X, Scale3DAbs.Y);
 		float ScaleLength = Scale3DAbs.Z;
@@ -434,32 +429,21 @@ public:
 		for (int32 i = 0; i < BodySetup->AggGeom.SphylElems.Num(); i++)
 		{
 			const FKSphylElem& SphylElem = BodySetup->AggGeom.SphylElems[i];
-
-			// this is a bit confusing since radius and height is scaled
-			// first apply the scale first 
-			float Radius = FMath::Max(SphylElem.Radius * ScaleRadius, 0.1f);
-			float Length = SphylElem.Length + SphylElem.Radius * 2.f;
-			float HalfLength = FMath::Max(Length * ScaleLength * 0.5f, 0.1f);
-			Radius = FMath::Clamp(Radius, 0.1f, HalfLength);	//radius is capped by half length
-			float HalfHeight = HalfLength - Radius;
-			HalfHeight = FMath::Max(0.1f, HalfHeight);
+			const FKSphylElem ScaledSphylElem = SphylElem.GetFinalScaled(Scale3D, RelativeTM);
 
 			PxCapsuleGeometry PCapsuleGeom;
-			PCapsuleGeom.halfHeight = HalfHeight;
-			PCapsuleGeom.radius = Radius;
+			PCapsuleGeom.halfHeight = ScaledSphylElem.Length * 0.5f;
+			PCapsuleGeom.radius = ScaledSphylElem.Radius;
 
 			if (PCapsuleGeom.isValid())
 			{
 				// The stored sphyl transform assumes the sphyl axis is down Z. In PhysX, it points down X, so we twiddle the matrix a bit here (swap X and Z and negate Y).
-				PxTransform PLocalPose(U2PVector(RelativeTM.TransformPosition(SphylElem.Center)), U2PQuat(SphylElem.Orientation) * U2PSphylBasis);
-				PLocalPose.p.x *= Scale3D.X;
-				PLocalPose.p.y *= Scale3D.Y;
-				PLocalPose.p.z *= Scale3D.Z;
+				PxTransform PLocalPose(U2PVector(ScaledSphylElem.Center), U2PQuat(ScaledSphylElem.Orientation) * U2PSphylBasis);
 
 				ensure(PLocalPose.isValid());
 				{
 					const float ContactOffset = FMath::Min(MaxContactOffset, ContactOffsetFactor * PCapsuleGeom.radius);
-					AttachShape_AssumesLocked(PCapsuleGeom, PLocalPose, ContactOffset, &SphylElem);
+					AttachShape_AssumesLocked(PCapsuleGeom, PLocalPose, ContactOffset, SphylElem.GetUserData());
 				}
 			}
 			else
@@ -469,7 +453,7 @@ public:
 		}
 	}
 
-	FORCEINLINE_DEBUGGABLE void AddConvexElemsToRigidActor_AssumesLocked() const
+	void AddConvexElemsToRigidActor_AssumesLocked() const
 	{
 		for (int32 i = 0; i < BodySetup->AggGeom.ConvexElems.Num(); i++)
 		{
@@ -504,7 +488,7 @@ public:
 						ensure(PLocalPose.isValid());
 						{
 							const float ContactOffset = FMath::Min(MaxContactOffset, ContactOffsetFactor * PBoundsExtents.minElement());
-							AttachShape_AssumesLocked(PConvexGeom, PLocalPose, ContactOffset, &ConvexElem);
+							AttachShape_AssumesLocked(PConvexGeom, PLocalPose, ContactOffset, ConvexElem.GetUserData());
 						}
 					}
 					else
@@ -524,7 +508,7 @@ public:
 		}
 	}
 
-	FORCEINLINE_DEBUGGABLE void AddTriMeshToRigidActor_AssumesLocked() const
+	void AddTriMeshToRigidActor_AssumesLocked() const
 	{
 		for(PxTriangleMesh* TriMesh : BodySetup->TriMeshes)
 		{
@@ -559,14 +543,14 @@ public:
 
 private:
 
-	PxShape* AttachShape_AssumesLocked(const PxGeometry& PGeom, const PxTransform& PLocalPose, const float ContactOffset, const FKShapeElem* ShapeElem, PxShapeFlags PShapeFlags = PxShapeFlag::eVISUALIZATION | PxShapeFlag::eSCENE_QUERY_SHAPE | PxShapeFlag::eSIMULATION_SHAPE) const
+	PxShape* AttachShape_AssumesLocked(const PxGeometry& PGeom, const PxTransform& PLocalPose, const float ContactOffset, const FPhysxUserData* ShapeElemUserData, PxShapeFlags PShapeFlags = PxShapeFlag::eVISUALIZATION | PxShapeFlag::eSCENE_QUERY_SHAPE | PxShapeFlag::eSIMULATION_SHAPE) const
 	{
 		const PxMaterial* PMaterial = GetDefaultPhysMaterial(); 
 		PxShape* PNewShape = bShapeSharing ? GPhysXSDK->createShape(PGeom, *PMaterial, /*isExclusive =*/ false, PShapeFlags) : PDestActor->createShape(PGeom, *PMaterial, PShapeFlags);
 
 		if (PNewShape)
 		{
-			PNewShape->userData = ShapeElem ? ((void*) ShapeElem->GetUserData()) : nullptr;
+			PNewShape->userData = (void*) ShapeElemUserData;
 			PNewShape->setLocalPose(PLocalPose);
 
 			if (NewShapes)
@@ -607,6 +591,14 @@ void UBodySetup::AddShapesToRigidActor_AssumesLocked(FBodyInstance* OwningInstan
 	CreatePhysicsMeshes();
 #endif
 
+	// if almost zero, set min scale
+	// @todo fixme
+	if (Scale3D.IsNearlyZero())
+	{
+		// set min scale
+		Scale3D = FVector(0.1f);
+	}
+
 	FAddShapesHelper AddShapesHelper(this, OwningInstance, PDestActor, SceneType, Scale3D, SimpleMaterial, ComplexMaterials, ShapeData, RelativeTM, NewShapes, bShapeSharing);
 
 	// Create shapes for simple collision if we do not want to use the complex collision mesh 
@@ -882,6 +874,92 @@ void UBodySetup::UpdateTriMeshVertices(const TArray<FVector> & NewPositions)
 #endif
 }
 
+template <bool bPositionAndNormal>
+float GetClosestPointAndNormalImpl(const UBodySetup* BodySetup, const FVector& WorldPosition, const FTransform& LocalToWorld, FVector* ClosestWorldPosition, FVector* FeatureNormal)
+{
+	float ClosestDist = FLT_MAX;
+	FVector TmpPosition, TmpNormal;
+
+	//Note that this function is optimized for BodySetup with few elements. This is more common. If we want to optimize the case with many elements we should really return the element during the distance check to avoid pointless iteration
+	for (const FKSphereElem& SphereElem : BodySetup->AggGeom.SphereElems)
+	{
+		
+		if(bPositionAndNormal)
+		{
+			const float Dist = SphereElem.GetClosestPointAndNormal(WorldPosition, LocalToWorld, TmpPosition, TmpNormal);
+
+			if(Dist < ClosestDist)
+			{
+				*ClosestWorldPosition = TmpPosition;
+				*FeatureNormal = TmpNormal;
+				ClosestDist = Dist;
+			}
+		}
+		else
+		{
+			const float Dist = SphereElem.GetShortestDistanceToPoint(WorldPosition, LocalToWorld);
+			ClosestDist = Dist < ClosestDist ? Dist : ClosestDist;
+		}
+	}
+
+	for (const FKSphylElem& SphylElem : BodySetup->AggGeom.SphylElems)
+	{
+		if (bPositionAndNormal)
+		{
+			const float Dist = SphylElem.GetClosestPointAndNormal(WorldPosition, LocalToWorld, TmpPosition, TmpNormal);
+
+			if (Dist < ClosestDist)
+			{
+				*ClosestWorldPosition = TmpPosition;
+				*FeatureNormal = TmpNormal;
+				ClosestDist = Dist;
+			}
+		}
+		else
+		{
+			const float Dist = SphylElem.GetShortestDistanceToPoint(WorldPosition, LocalToWorld);
+			ClosestDist = Dist < ClosestDist ? Dist : ClosestDist;
+		}
+	}
+
+	for (const FKBoxElem& BoxElem : BodySetup->AggGeom.BoxElems)
+	{
+		if (bPositionAndNormal)
+		{
+			const float Dist = BoxElem.GetClosestPointAndNormal(WorldPosition, LocalToWorld, TmpPosition, TmpNormal);
+
+			if (Dist < ClosestDist)
+			{
+				*ClosestWorldPosition = TmpPosition;
+				*FeatureNormal = TmpNormal;
+				ClosestDist = Dist;
+			}
+		}
+		else
+		{
+			const float Dist =  BoxElem.GetShortestDistanceToPoint(WorldPosition, LocalToWorld);
+			ClosestDist = Dist < ClosestDist ? Dist : ClosestDist;
+		}
+	}
+
+	if (ClosestDist == FLT_MAX)
+	{
+		UE_LOG(LogPhysics, Warning, TEXT("GetClosestPointAndNormalImpl ClosestDist for BodySetup %s is coming back as FLT_MAX. WorldPosition = %s, LocalToWorld = %s"), *BodySetup->GetFullName(), *WorldPosition.ToString(), *LocalToWorld.ToHumanReadableString());
+	}
+
+	return ClosestDist;
+}
+
+float UBodySetup::GetShortestDistanceToPoint(const FVector& WorldPosition, const FTransform& LocalToWorld) const
+{
+	return GetClosestPointAndNormalImpl<false>(this, WorldPosition, LocalToWorld, nullptr, nullptr);
+}
+
+float UBodySetup::GetClosestPointAndNormal(const FVector& WorldPosition, const FTransform& LocalToWorld, FVector& ClosestWorldPosition, FVector& FeatureNormal) const
+{
+	return GetClosestPointAndNormalImpl<true>(this, WorldPosition, LocalToWorld, &ClosestWorldPosition, &FeatureNormal);
+}
+
 FByteBulkData* UBodySetup::GetCookedData(FName Format)
 {
 	if (IsTemplate())
@@ -1133,6 +1211,39 @@ void FKSphereElem::Serialize( const FArchive& Ar )
 	}
 }
 
+float FKSphereElem::GetShortestDistanceToPoint(const FVector& WorldPosition, const FTransform& LocalToWorldTM) const
+{
+	FKSphereElem ScaledSphere = GetFinalScaled(LocalToWorldTM.GetScale3D(), FTransform::Identity);
+
+	const FVector Dir = LocalToWorldTM.TransformPositionNoScale(ScaledSphere.Center) - WorldPosition;
+	const float DistToCenter = Dir.Size();
+	const float DistToEdge = DistToCenter - ScaledSphere.Radius;
+	
+	return DistToEdge > SMALL_NUMBER ? DistToEdge : 0.f;
+}
+
+float FKSphereElem::GetClosestPointAndNormal(const FVector& WorldPosition, const FTransform& LocalToWorldTM, FVector& ClosestWorldPosition, FVector& Normal) const
+{
+	FKSphereElem ScaledSphere = GetFinalScaled(LocalToWorldTM.GetScale3D(), FTransform::Identity);
+
+	const FVector Dir = LocalToWorldTM.TransformPositionNoScale(ScaledSphere.Center) - WorldPosition;
+	const float DistToCenter = Dir.Size();
+	const float DistToEdge = FMath::Max(DistToCenter - ScaledSphere.Radius, 0.f);
+
+	if(DistToCenter > SMALL_NUMBER)
+	{
+		Normal = -Dir.GetUnsafeNormal();
+	}
+	else
+	{
+		Normal = FVector::ZeroVector;
+	}
+	
+	ClosestWorldPosition = WorldPosition - Normal*DistToEdge;
+
+	return DistToEdge;
+}
+
 void FKSphereElem::ScaleElem(FVector DeltaSize, float MinSize)
 {
 	// Find element with largest magnitude, btu preserve sign.
@@ -1145,6 +1256,23 @@ void FKSphereElem::ScaleElem(FVector DeltaSize, float MinSize)
 	Radius = FMath::Max(Radius + DeltaRadius, MinSize);
 }
 
+FKSphereElem FKSphereElem::GetFinalScaled(const FVector& Scale3D, const FTransform& RelativeTM) const
+{
+	float MinScale, MinScaleAbs;
+	FVector Scale3DAbs;
+
+	SetupNonUniformHelper(Scale3D, MinScale, MinScaleAbs, Scale3DAbs);
+
+	FKSphereElem ScaledSphere = *this;
+	ScaledSphere.Radius *= MinScaleAbs;
+
+	ScaledSphere.Center = RelativeTM.TransformPosition(Center);
+	ScaledSphere.Center *= MinScale;
+
+
+	return ScaledSphere;
+}
+
 void FKBoxElem::Serialize( const FArchive& Ar )
 {
 	if ( Ar.IsLoading() && Ar.UE4Ver() < VER_UE4_REFACTOR_PHYSICS_TRANSFORMS )
@@ -1162,6 +1290,67 @@ void FKBoxElem::ScaleElem(FVector DeltaSize, float MinSize)
 	Z = FMath::Max(Z + 2 * DeltaSize.Z, MinSize);
 }
 
+
+FKBoxElem FKBoxElem::GetFinalScaled(const FVector& Scale3D, const FTransform& RelativeTM) const
+{
+	float MinScale, MinScaleAbs;
+	FVector Scale3DAbs;
+
+	SetupNonUniformHelper(Scale3D, MinScale, MinScaleAbs, Scale3DAbs);
+
+	FKBoxElem ScaledBox = *this;
+	ScaledBox.X *= Scale3DAbs.X;
+	ScaledBox.Y *= Scale3DAbs.Y;
+	ScaledBox.Z *= Scale3DAbs.Z;
+
+	FTransform BoxTransform = GetTransform() * RelativeTM;
+	BoxTransform.ScaleTranslation(Scale3D);
+	ScaledBox.SetTransform(BoxTransform);
+
+	return ScaledBox;
+}
+
+float FKBoxElem::GetShortestDistanceToPoint(const FVector& WorldPosition, const FTransform& BoneToWorldTM) const
+{
+	const FKBoxElem& ScaledBox = GetFinalScaled(BoneToWorldTM.GetScale3D(), FTransform::Identity);
+	const FTransform LocalToWorldTM = GetTransform() * BoneToWorldTM;
+	const FVector LocalPosition = LocalToWorldTM.InverseTransformPositionNoScale(WorldPosition);
+	const FVector LocalPositionAbs = LocalPosition.GetAbs();
+
+	const FVector HalfPoint(ScaledBox.X*0.5f, ScaledBox.Y*0.5f, ScaledBox.Z*0.5f);
+	const FVector Delta = LocalPositionAbs - HalfPoint;
+	const FVector Errors = FVector(FMath::Max(Delta.X, 0.f), FMath::Max(Delta.Y, 0.f), FMath::Max(Delta.Z, 0.f));
+	const float Error = Errors.Size();
+
+	return Error > SMALL_NUMBER ? Error : 0.f;
+}
+
+float FKBoxElem::GetClosestPointAndNormal(const FVector& WorldPosition, const FTransform& BoneToWorldTM, FVector& ClosestWorldPosition, FVector& Normal) const
+{
+	const FKBoxElem& ScaledBox = GetFinalScaled(BoneToWorldTM.GetScale3D(), FTransform::Identity);
+	const FTransform LocalToWorldTM = GetTransform() * BoneToWorldTM;
+	const FVector LocalPosition = LocalToWorldTM.InverseTransformPositionNoScale(WorldPosition);
+
+	const float HalfX = ScaledBox.X * 0.5f;
+	const float HalfY = ScaledBox.Y * 0.5f;
+	const float HalfZ = ScaledBox.Z * 0.5f;
+	
+	const FVector ClosestLocalPosition(FMath::Clamp(LocalPosition.X, -HalfX, HalfX), FMath::Clamp(LocalPosition.Y, -HalfY, HalfY), FMath::Clamp(LocalPosition.Z, -HalfZ, HalfZ));
+	ClosestWorldPosition = LocalToWorldTM.TransformPositionNoScale(ClosestLocalPosition);
+
+	const FVector LocalDelta = LocalPosition - ClosestLocalPosition;
+	float Error = LocalDelta.Size();
+	
+	bool bIsOutside = Error > SMALL_NUMBER;
+	
+	const FVector LocalNormal = bIsOutside ? LocalDelta.GetUnsafeNormal() : FVector::ZeroVector;
+
+	ClosestWorldPosition = LocalToWorldTM.TransformPositionNoScale(ClosestLocalPosition);
+	Normal = LocalToWorldTM.TransformVectorNoScale(LocalNormal);
+	
+	return bIsOutside ? Error : 0.f;
+}
+
 void FKSphylElem::Serialize( const FArchive& Ar )
 {
 	if ( Ar.IsLoading() && Ar.UE4Ver() < VER_UE4_REFACTOR_PHYSICS_TRANSFORMS )
@@ -1190,6 +1379,80 @@ void FKSphylElem::ScaleElem(FVector DeltaSize, float MinSize)
 	Length = length;
 }
 
+FKSphylElem FKSphylElem::GetFinalScaled(const FVector& Scale3D, const FTransform& RelativeTM) const
+{
+	FKSphylElem ScaledSphylElem = *this;
+
+	float MinScale, MinScaleAbs;
+	FVector Scale3DAbs;
+
+	SetupNonUniformHelper(Scale3D, MinScale, MinScaleAbs, Scale3DAbs);
+
+	float ScaleRadius = FMath::Max(Scale3DAbs.X, Scale3DAbs.Y);
+	float ScaleLength = Scale3DAbs.Z;
+	
+	// this is a bit confusing since radius and height is scaled
+	// first apply the scale first 
+	ScaledSphylElem.Radius = FMath::Max(Radius * ScaleRadius, 0.1f);
+	ScaledSphylElem.Length = Length + Radius * 2.f;
+	float HalfLength = FMath::Max(ScaledSphylElem.Length * ScaleLength * 0.5f, 0.1f);
+	ScaledSphylElem.Radius = FMath::Clamp(ScaledSphylElem.Radius, 0.1f, HalfLength);	//radius is capped by half length
+	ScaledSphylElem.Length = FMath::Max(0.1f, (HalfLength - ScaledSphylElem.Radius) * 2.f);
+
+	FVector LocalOrigin = RelativeTM.TransformPosition(Center) * Scale3D;
+	ScaledSphylElem.Center = LocalOrigin;
+	
+	return ScaledSphylElem;
+}
+
+float FKSphylElem::GetShortestDistanceToPoint(const FVector& WorldPosition, const FTransform& BoneToWorldTM) const
+{
+	const FKSphylElem ScaledSphyl = GetFinalScaled(BoneToWorldTM.GetScale3D(), FTransform::Identity);
+
+	const FTransform LocalToWorldTM = GetTransform() * BoneToWorldTM;
+	const FVector ErrorScale = LocalToWorldTM.GetScale3D();
+	const FVector LocalPosition = LocalToWorldTM.InverseTransformPositionNoScale(WorldPosition);
+	const FVector LocalPositionAbs = LocalPosition.GetAbs();
+	
+	
+	const FVector Target(LocalPositionAbs.X, LocalPositionAbs.Y, FMath::Max(LocalPositionAbs.Z - ScaledSphyl.Length * 0.5f, 0.f));	//If we are above half length find closest point to cap, otherwise to cylinder
+	const float Error = FMath::Max(Target.Size() - ScaledSphyl.Radius, 0.f);
+
+	return Error > SMALL_NUMBER ? Error : 0.f;
+}
+
+float FKSphylElem::GetClosestPointAndNormal(const FVector& WorldPosition, const FTransform& BoneToWorldTM, FVector& ClosestWorldPosition, FVector& Normal) const
+{
+	const FKSphylElem ScaledSphyl = GetFinalScaled(BoneToWorldTM.GetScale3D(), FTransform::Identity);
+
+	const FTransform LocalToWorldTM = GetTransform() * BoneToWorldTM;
+	const FVector ErrorScale = LocalToWorldTM.GetScale3D();
+	const FVector LocalPosition = LocalToWorldTM.InverseTransformPositionNoScale(WorldPosition);
+	
+	const float HalfLength = 0.5f * ScaledSphyl.Length;
+	const float TargetZ = FMath::Clamp(LocalPosition.Z, -HalfLength, HalfLength);	//We want to move to a sphere somewhere along the capsule axis
+
+	const FVector WorldSphere = LocalToWorldTM.TransformPositionNoScale(FVector(0.f, 0.f, TargetZ));
+	const FVector Dir = WorldSphere - WorldPosition;
+	const float DistToCenter = Dir.Size();
+	const float DistToEdge = FMath::Max(DistToCenter - ScaledSphyl.Radius, 0.f);
+
+	bool bIsOutside = DistToCenter > SMALL_NUMBER;
+	if (bIsOutside)
+	{
+		Normal = -Dir.GetUnsafeNormal();
+	}
+	else
+	{
+		Normal = FVector::ZeroVector;
+	}
+
+	ClosestWorldPosition = WorldPosition - Normal*DistToEdge;
+
+	return bIsOutside ? DistToEdge : 0.f;
+}
+
+
 class UPhysicalMaterial* UBodySetup::GetPhysMaterial() const
 {
 	UPhysicalMaterial* PhysMat = PhysMaterial;
diff --git a/Engine/Source/Runtime/Engine/Private/PhysicsEngine/ConstraintInstance.cpp b/Engine/Source/Runtime/Engine/Private/PhysicsEngine/ConstraintInstance.cpp
--- a/Engine/Source/Runtime/Engine/Private/PhysicsEngine/ConstraintInstance.cpp
+++ b/Engine/Source/Runtime/Engine/Private/PhysicsEngine/ConstraintInstance.cpp
@@ -470,12 +470,12 @@ void EnsureSleepingActorsStaySleeping_AssumesLocked(PxRigidActor* PActor1, PxRig
 	// creation of joints wakes up rigid bodies, so we put them to sleep again if both were initially asleep
 	if (bActor1Asleep && bActor2Asleep)
 	{
-		if (PActor1 && IsRigidBodyNonKinematic_AssumesLocked(PActor1->isRigidDynamic()))
+		if (PActor1 && !IsRigidBodyKinematic_AssumesLocked(PActor1->isRigidDynamic()))
 		{
 			PActor1->isRigidDynamic()->putToSleep();
 		}
 
-		if (PActor2 && IsRigidBodyNonKinematic_AssumesLocked(PActor2->isRigidDynamic()))
+		if (PActor2 && !IsRigidBodyKinematic_AssumesLocked(PActor2->isRigidDynamic()))
 		{
 			PActor2->isRigidDynamic()->putToSleep();
 		}
diff --git a/Engine/Source/Runtime/Engine/Private/PhysicsEngine/PhysScene.cpp b/Engine/Source/Runtime/Engine/Private/PhysicsEngine/PhysScene.cpp
--- a/Engine/Source/Runtime/Engine/Private/PhysicsEngine/PhysScene.cpp
+++ b/Engine/Source/Runtime/Engine/Private/PhysicsEngine/PhysScene.cpp
@@ -281,6 +281,8 @@ FPhysScene::FPhysScene()
 	GApexModuleDestructible->setWorldSupportPhysXScene(*ApexScene, SyncPhysXScene);
 	GApexModuleDestructible->setDamageApplicationRaycastFlags(NxDestructibleActorRaycastFlags::AllChunks, *ApexScene);
 #endif
+
+	PreGarbageCollectDelegateHandle = FCoreUObjectDelegates::PreGarbageCollect.AddRaw(this, &FPhysScene::WaitPhysScenes);
 }
 
 void FPhysScene::SetOwningWorld(UWorld* InOwningWorld)
@@ -291,6 +293,7 @@ void FPhysScene::SetOwningWorld(UWorld* InOwningWorld)
 /** Exposes destruction of physics-engine scene outside Engine. */
 FPhysScene::~FPhysScene()
 {
+	FCoreUObjectDelegates::PreGarbageCollect.Remove(PreGarbageCollectDelegateHandle);
 	// Make sure no scenes are left simulating (no-ops if not simulating)
 	WaitPhysScenes();
 	// Loop through scene types to get all scenes
@@ -352,19 +355,29 @@ void FPhysScene::SetKinematicTarget_AssumesLocked(FBodyInstance* BodyInstance, c
 #if WITH_PHYSX
 	if (PxRigidDynamic * PRigidDynamic = BodyInstance->GetPxRigidDynamic_AssumesLocked())
 	{
-#if WITH_SUBSTEPPING
-		uint32 BodySceneType = SceneType_AssumesLocked(BodyInstance);
-		if (bAllowSubstepping && IsSubstepping(BodySceneType))
+		const bool bIsKinematicTarget = IsRigidBodyKinematicAndInSimulationScene_AssumesLocked(PRigidDynamic);
+		if(bIsKinematicTarget)
 		{
-			FPhysSubstepTask * PhysSubStepper = PhysSubSteppers[BodySceneType];
-			PhysSubStepper->SetKinematicTarget_AssumesLocked(BodyInstance, TargetTransform);
+#if WITH_SUBSTEPPING
+			uint32 BodySceneType = SceneType_AssumesLocked(BodyInstance);
+			if (bAllowSubstepping && IsSubstepping(BodySceneType))
+			{
+				FPhysSubstepTask * PhysSubStepper = PhysSubSteppers[BodySceneType];
+				PhysSubStepper->SetKinematicTarget_AssumesLocked(BodyInstance, TargetTransform);
+			}
+			else
+#endif
+			{
+				const PxTransform PNewPose = U2PTransform(TargetTransform);
+				PRigidDynamic->setKinematicTarget(PNewPose);
+			}
 		}
 		else
-#endif
 		{
 			const PxTransform PNewPose = U2PTransform(TargetTransform);
-			PRigidDynamic->setKinematicTarget(PNewPose);
+			PRigidDynamic->setGlobalPose(PNewPose);
 		}
+
 	}
 #endif
 }
diff --git a/Engine/Source/Runtime/Engine/Private/PhysicsEngine/PhysXSupport.cpp b/Engine/Source/Runtime/Engine/Private/PhysicsEngine/PhysXSupport.cpp
--- a/Engine/Source/Runtime/Engine/Private/PhysicsEngine/PhysXSupport.cpp
+++ b/Engine/Source/Runtime/Engine/Private/PhysicsEngine/PhysXSupport.cpp
@@ -252,11 +252,23 @@ void AddRadialForceToPxRigidBody_AssumesLocked(PxRigidBody& PRigidBody, const FV
 #endif // WITH_PHYSX
 }
 
-bool IsRigidBodyNonKinematic_AssumesLocked(const PxRigidBody* PRigidBody)
+bool IsRigidBodyKinematic_AssumesLocked(const PxRigidBody* PRigidBody)
 {
 	if (PRigidBody)
 	{
-		return !(PRigidBody->getRigidBodyFlags() & PxRigidBodyFlag::eKINEMATIC);
+		//For some cases we only consider an actor kinematic if it's in the simulation scene. This is in cases where we set a kinematic target
+		return (PRigidBody->getRigidBodyFlags() & PxRigidBodyFlag::eKINEMATIC);
+	}
+
+	return false;
+}
+
+bool IsRigidBodyKinematicAndInSimulationScene_AssumesLocked(const PxRigidBody* PRigidBody)
+{
+	if (PRigidBody)
+	{
+		//For some cases we only consider an actor kinematic if it's in the simulation scene. This is in cases where we set a kinematic target
+		return (PRigidBody->getRigidBodyFlags() & PxRigidBodyFlag::eKINEMATIC) && !(PRigidBody->getActorFlags() & PxActorFlag::eDISABLE_SIMULATION);
 	}
 
 	return false;
diff --git a/Engine/Source/Runtime/Engine/Private/PhysicsEngine/PhysXSupport.h b/Engine/Source/Runtime/Engine/Private/PhysicsEngine/PhysXSupport.h
--- a/Engine/Source/Runtime/Engine/Private/PhysicsEngine/PhysXSupport.h
+++ b/Engine/Source/Runtime/Engine/Private/PhysicsEngine/PhysXSupport.h
@@ -81,13 +81,21 @@ inline void AddRadialForceToPxRigidBody(PxRigidBody& PRigidBody, const FVector&
 	AddRadialForceToPxRigidBody_AssumesLocked(PRigidBody, Origin, Radius, Strength, Falloff, bAccelChange);
 }
 
-bool IsRigidBodyNonKinematic_AssumesLocked(const PxRigidBody* PRigidBody);
+bool IsRigidBodyKinematic_AssumesLocked(const PxRigidBody* PRigidBody);
+
+bool IsRigidBodyKinematicAndInSimulationScene_AssumesLocked(const PxRigidBody* PRigidBody);
+
+DEPRECATED(4.12, "Please call IsRigidBodyKinematic_AssumesLocked")
+inline bool IsRigidBodyNonKinematic_AssumesLocked(const PxRigidBody* PRigidBody)
+{
+	return !IsRigidBodyKinematic_AssumesLocked(PRigidBody);
+}
 
 /** Util to see if a PxRigidActor is non-kinematic */
-DEPRECATED(4.8, "Please call IsRigidBodyNonKinematic_AssumesLocked and make sure you obtain the appropriate PhysX scene locks")
+DEPRECATED(4.8, "Please call IsRigidBodyKinematic_AssumesLocked and make sure you obtain the appropriate PhysX scene locks")
 inline bool IsRigidBodyNonKinematic(PxRigidBody* PRigidBody)
 {
-	return IsRigidBodyNonKinematic_AssumesLocked(PRigidBody);
+	return !IsRigidBodyKinematic_AssumesLocked(PRigidBody);
 }
 
 
diff --git a/Engine/Source/Runtime/Engine/Private/PhysicsVolume.cpp b/Engine/Source/Runtime/Engine/Private/PhysicsVolume.cpp
--- a/Engine/Source/Runtime/Engine/Private/PhysicsVolume.cpp
+++ b/Engine/Source/Runtime/Engine/Private/PhysicsVolume.cpp
@@ -57,6 +57,30 @@ void APhysicsVolume::EndPlay(const EEndPlayReason::Type EndPlayReason)
 	if (MyWorld)
 	{
 		MyWorld->RemovePhysicsVolume(this);
+
+		if ((EndPlayReason == EEndPlayReason::RemovedFromWorld) || (EndPlayReason == EEndPlayReason::Destroyed))
+		{
+			// Prevent UpdatePhysicsVolume() calls below from returning this component.
+			UPrimitiveComponent* Brush = GetBrushComponent();
+			const bool bSavedGenerateOverlapEvents = Brush->bGenerateOverlapEvents;
+			Brush->bGenerateOverlapEvents = false;
+
+			// Refresh physics volume on any components touching this volume.
+			// TODO: Physics volume tracking code needs a cleanup, ideally it just uses normal begin/end overlap events,
+			// but there are subtleties with the stacking and priority rules.
+			const TArray<FOverlapInfo>& Overlaps = Brush->GetOverlapInfos();
+			for (const FOverlapInfo& Info : Overlaps)
+			{
+				UPrimitiveComponent* OtherPrim = Info.OverlapInfo.GetComponent();
+				if (OtherPrim && OtherPrim->bShouldUpdatePhysicsVolume)
+				{
+					OtherPrim->UpdatePhysicsVolume(true);
+				}
+			}
+
+			// Restore saved flag, since we may stream back in.
+			Brush->bGenerateOverlapEvents = bSavedGenerateOverlapEvents;
+		}
 	}
 	Super::EndPlay(EndPlayReason);
 }
@@ -78,7 +102,8 @@ bool APhysicsVolume::IsOverlapInVolume(const class USceneComponent& TestComponen
 
 float APhysicsVolume::GetGravityZ() const
 {
-	return GetWorld()->GetGravityZ();
+	const UWorld* MyWorld = GetWorld();
+	return MyWorld ? MyWorld->GetGravityZ() : UPhysicsSettings::Get()->DefaultGravityZ;
 }
 
 void APhysicsVolume::ActorEnteredVolume(AActor* Other) {}
diff --git a/Engine/Source/Runtime/Engine/Private/PrimitiveComponentPhysics.cpp b/Engine/Source/Runtime/Engine/Private/PrimitiveComponentPhysics.cpp
--- a/Engine/Source/Runtime/Engine/Private/PrimitiveComponentPhysics.cpp
+++ b/Engine/Source/Runtime/Engine/Private/PrimitiveComponentPhysics.cpp
@@ -755,7 +755,7 @@ bool UPrimitiveComponent::WeldToImplementation(USceneComponent * InParent, FName
 			//if root is kinematic simply set child to be kinematic and we're done
 			if (RootComponent->IsSimulatingPhysics(SocketName) == false)
 			{
-				BI->WeldParent = NULL;
+				FPlatformAtomics::InterlockedExchangePtr((void**)&BI->WeldParent, nullptr);
 				SetSimulatePhysics(false);
 				return false;	//return false because we need to continue with regular body initialization
 			}
@@ -809,7 +809,7 @@ void UPrimitiveComponent::UnWeldFromParent()
 
 			NewRootBI->bWelded = false;
 			const FBodyInstance* PrevWeldParent = NewRootBI->WeldParent;
-			NewRootBI->WeldParent = nullptr;
+			FPlatformAtomics::InterlockedExchangePtr((void**)&NewRootBI->WeldParent, nullptr);
 
 			bool bHasBodySetup = GetBodySetup() != nullptr;
 
@@ -843,7 +843,7 @@ void UPrimitiveComponent::UnWeldFromParent()
 					}
 
 					//At this point, NewRootBI must be kinematic because it's being unwelded. It's up to the code that simulates to call Weld on the children as needed
-					ChildBI->WeldParent = nullptr;	//null because we are currently kinematic
+					FPlatformAtomics::InterlockedExchangePtr((void**)&ChildBI->WeldParent, nullptr); //null because we are currently kinematic
 				}
 			}
 		}
