Copying //UE4/Dev-Framework to //UE4/Main

diff --git a/Engine/Source/Runtime/Engine/Private/SkeletalMesh.cpp b/Engine/Source/Runtime/Engine/Private/SkeletalMesh.cpp
--- a/Engine/Source/Runtime/Engine/Private/SkeletalMesh.cpp
+++ b/Engine/Source/Runtime/Engine/Private/SkeletalMesh.cpp
@@ -47,6 +47,8 @@
 
 DEFINE_LOG_CATEGORY(LogSkeletalMesh);
 
+DECLARE_CYCLE_STAT(TEXT("GetShadowShapes"), STAT_GetShadowShapes, STATGROUP_Anim);
+
 #if WITH_APEX_CLOTHING
 /*-----------------------------------------------------------------------------
 	utility functions for apex clothing 
@@ -2514,7 +2516,7 @@ void USkeletalMesh::Serialize( FArchive& Ar )
 			FMultiSizeIndexContainerData AdjacencyIndexData;
 			IMeshUtilities& MeshUtilities = FModuleManager::Get().LoadModuleChecked<IMeshUtilities>("MeshUtilities");
 
-			UE_LOG(LogSkeletalMesh, Warning, TEXT("Building adjacency information for skeletal mesh '%s'."), *GetPathName() );
+			UE_LOG(LogSkeletalMesh, Warning, TEXT("Building adjacency information for skeletal mesh '%s'. Please resave the asset."), *GetPathName() );
 			LODModel.GetVertices( Vertices );
 			LODModel.MultiSizeIndexContainer.GetIndexBufferData( IndexData );
 			AdjacencyIndexData.DataTypeSize = IndexData.DataTypeSize;
@@ -2690,26 +2692,26 @@ void USkeletalMesh::CalculateInvRefMatrices()
 		RefBasesInvMatrix.Empty(RefSkeleton.GetNum());
 		RefBasesInvMatrix.AddUninitialized(RefSkeleton.GetNum());
 
-		// Temporary storage for calculating mesh-space ref pose
-		TArray<FMatrix> RefBases;
-		RefBases.AddUninitialized( RefSkeleton.GetNum() );
+		// Reset cached mesh-space ref pose
+		CachedComposedRefPoseMatrices.Empty( RefSkeleton.GetNum() );
+		CachedComposedRefPoseMatrices.AddUninitialized( RefSkeleton.GetNum() );
 
 		// Precompute the Mesh.RefBasesInverse.
 		for( int32 b=0; b<RefSkeleton.GetNum(); b++)
 		{
 			// Render the default pose.
-			RefBases[b] = GetRefPoseMatrix(b);
+			CachedComposedRefPoseMatrices[b] = GetRefPoseMatrix(b);
 
 			// Construct mesh-space skeletal hierarchy.
 			if( b>0 )
 			{
 				int32 Parent = RefSkeleton.GetParentIndex(b);
-				RefBases[b] = RefBases[b] * RefBases[Parent];
+				CachedComposedRefPoseMatrices[b] = CachedComposedRefPoseMatrices[b] * CachedComposedRefPoseMatrices[Parent];
 			}
 
 			FVector XAxis, YAxis, ZAxis;
 
-			RefBases[b].GetScaledAxes(XAxis, YAxis, ZAxis);
+			CachedComposedRefPoseMatrices[b].GetScaledAxes(XAxis, YAxis, ZAxis);
 			if(	XAxis.IsNearlyZero(SMALL_NUMBER) &&
 				YAxis.IsNearlyZero(SMALL_NUMBER) &&
 				ZAxis.IsNearlyZero(SMALL_NUMBER))
@@ -2719,7 +2721,7 @@ void USkeletalMesh::CalculateInvRefMatrices()
 			}
 
 			// Precompute inverse so we can use from-refpose-skin vertices.
-			RefBasesInvMatrix[b] = RefBases[b].Inverse(); 
+			RefBasesInvMatrix[b] = CachedComposedRefPoseMatrices[b].Inverse(); 
 		}
 
 #if WITH_EDITORONLY_DATA
@@ -3167,19 +3169,7 @@ FMatrix USkeletalMesh::GetComposedRefPoseMatrix( FName InBoneName ) const
 
 FMatrix USkeletalMesh::GetComposedRefPoseMatrix(int32 InBoneIndex) const
 {
-	FMatrix LocalPose(FMatrix::Identity);
-	int32 BoneIndex = InBoneIndex;
-
-	if(BoneIndex != INDEX_NONE)
-	{
-		while(BoneIndex != INDEX_NONE)
-		{
-			LocalPose = LocalPose * GetRefPoseMatrix(BoneIndex);
-			BoneIndex = RefSkeleton.GetParentIndex(BoneIndex);
-		}
-	}
-
-	return LocalPose;
+	return CachedComposedRefPoseMatrices[InBoneIndex];
 }
 
 TArray<USkeletalMeshSocket*>& USkeletalMesh::GetMeshOnlySocketList()
@@ -3934,6 +3924,32 @@ bool USkeletalMeshSocket::AttachActor(AActor* Actor, class USkeletalMeshComponen
 	return bAttached;
 }
 
+#if WITH_EDITOR
+void USkeletalMeshSocket::PostEditChangeProperty(FPropertyChangedEvent& PropertyChangedEvent)
+{
+	Super::PostEditChangeProperty(PropertyChangedEvent);
+
+	if (PropertyChangedEvent.Property)
+	{
+		ChangedEvent.Broadcast(this, PropertyChangedEvent.MemberProperty);
+	}
+}
+
+void USkeletalMeshSocket::CopyFrom(const class USkeletalMeshSocket* OtherSocket)
+{
+	if (OtherSocket)
+	{
+		SocketName = OtherSocket->SocketName;
+		BoneName = OtherSocket->BoneName;
+		RelativeLocation = OtherSocket->RelativeLocation;
+		RelativeRotation = OtherSocket->RelativeRotation;
+		RelativeScale = OtherSocket->RelativeScale;
+		bForceAlwaysAnimated = OtherSocket->bForceAlwaysAnimated;
+	}
+}
+
+#endif
+
 /*-----------------------------------------------------------------------------
 	ASkeletalMeshActor
 -----------------------------------------------------------------------------*/
@@ -4227,6 +4243,8 @@ FSkeletalMeshSceneProxy
 #include "LevelUtils.h"
 #include "SkeletalRender.h"
 
+const FQuat SphylBasis(FVector(1.0f / FMath::Sqrt(2.0f), 0.0f, 1.0f / FMath::Sqrt(2.0f)), PI);
+
 /** 
  * Constructor. 
  * @param	Component - skeletal mesh primitive being added
@@ -4365,8 +4383,45 @@ FSkeletalMeshSceneProxy::FSkeletalMeshSceneProxy(const USkinnedMeshComponent* Co
 	FColor NewPropertyColor;
 	GEngine->GetPropertyColorationColor( (UObject*)Component, NewPropertyColor );
 	PropertyColor = NewPropertyColor;
+
+	// Copy out shadow physics asset data
+	if(SkeletalMeshComponent)
+	{
+		UPhysicsAsset* ShadowPhysicsAsset = SkeletalMeshComponent->SkeletalMesh->ShadowPhysicsAsset;
+
+		if (ShadowPhysicsAsset
+			&& SkeletalMeshComponent->CastShadow
+			&& (SkeletalMeshComponent->bCastCapsuleDirectShadow || SkeletalMeshComponent->bCastCapsuleIndirectShadow))
+		{
+			for (int32 BodyIndex = 0; BodyIndex < ShadowPhysicsAsset->BodySetup.Num(); BodyIndex++)
+			{
+				UBodySetup* BodySetup = ShadowPhysicsAsset->BodySetup[BodyIndex];
+				int32 BoneIndex = SkeletalMeshComponent->GetBoneIndex(BodySetup->BoneName);
+
+				if (BoneIndex != INDEX_NONE)
+				{
+					const FMatrix& RefBoneMatrix = SkeletalMeshComponent->SkeletalMesh->GetComposedRefPoseMatrix(BoneIndex);
+
+					const int32 NumSpheres = BodySetup->AggGeom.SphereElems.Num();
+					for (int32 SphereIndex = 0; SphereIndex < NumSpheres; SphereIndex++)
+					{
+						const FKSphereElem& SphereShape = BodySetup->AggGeom.SphereElems[SphereIndex];
+						ShadowCapsuleData.Add(TPairInitializer<int32, FCapsuleShape>(BoneIndex, FCapsuleShape(RefBoneMatrix.TransformPosition(SphereShape.Center), SphereShape.Radius, FVector(0.0f, 0.0f, 1.0f), 0.0f)));
+					}
+
+					const int32 NumCapsules = BodySetup->AggGeom.SphylElems.Num();
+					for (int32 CapsuleIndex = 0; CapsuleIndex < NumCapsules; CapsuleIndex++)
+					{
+						const FKSphylElem& SphylShape = BodySetup->AggGeom.SphylElems[CapsuleIndex];
+						ShadowCapsuleData.Add(TPairInitializer<int32, FCapsuleShape>(BoneIndex, FCapsuleShape(RefBoneMatrix.TransformPosition(SphylShape.Center), SphylShape.Radius, RefBoneMatrix.TransformVector((SphylShape.Orientation * SphylBasis).Vector()), SphylShape.Length)));
+					}
+				}
+			}
+		}
+	}
 }
 
+
 // FPrimitiveSceneProxy interface.
 
 /** 
@@ -4727,7 +4782,26 @@ bool FSkeletalMeshSceneProxy::HasDistanceFieldRepresentation() const
 
 void FSkeletalMeshSceneProxy::GetShadowShapes(TArray<FCapsuleShape>& CapsuleShapes) const 
 {
-	CapsuleShapes.Append(MeshObject->ShadowCapsuleShapes);
+	SCOPE_CYCLE_COUNTER(STAT_GetShadowShapes);
+
+	const TArray<FMatrix>& ReferenceToLocalMatrices = MeshObject->GetReferenceToLocalMatrices();
+	const FMatrix& LocalToWorld = GetLocalToWorld();
+
+	int32 CapsuleIndex = CapsuleShapes.Num();
+	CapsuleShapes.SetNum(CapsuleShapes.Num() + ShadowCapsuleData.Num(), false);
+
+	for(const TPair<int32, FCapsuleShape>& CapsuleData : ShadowCapsuleData)
+	{
+		FMatrix ReferenceToWorld = ReferenceToLocalMatrices[CapsuleData.Key] * LocalToWorld;
+		const float MaxScale = ReferenceToWorld.GetScaleVector().GetMax();
+
+		FCapsuleShape& NewCapsule = CapsuleShapes[CapsuleIndex++];
+
+		NewCapsule.Center = ReferenceToWorld.TransformPosition(CapsuleData.Value.Center);
+		NewCapsule.Radius = CapsuleData.Value.Radius * MaxScale;
+		NewCapsule.Orientation = ReferenceToWorld.TransformVector(CapsuleData.Value.Orientation);
+		NewCapsule.Length = CapsuleData.Value.Length * MaxScale;
+	}
 }
 
 /**
@@ -4902,7 +4976,6 @@ USkinnedMeshComponent::USkinnedMeshComponent(const FObjectInitializer& ObjectIni
 	bCastCapsuleIndirectShadow = false;
 
 	bDoubleBufferedBlendSpaces = true;
-	bReInitAnimationOnSetSkeletalMeshCalls = true;
 	CurrentEditableSpaceBases = 0;
 	CurrentReadSpaceBases = 1;
 	bNeedToFlipSpaceBaseBuffers = false;
@@ -5052,15 +5125,4 @@ FPrimitiveSceneProxy* USkinnedMeshComponent::CreateSceneProxy()
 USkeletalMeshComponent* ASkeletalMeshActor::GetSkeletalMeshComponent() { return SkeletalMeshComponent; }
 
 
-#if WITH_EDITOR
-void USkeletalMeshSocket::PostEditChangeProperty(FPropertyChangedEvent& PropertyChangedEvent)
-{
-	Super::PostEditChangeProperty(PropertyChangedEvent);
-
-	if (PropertyChangedEvent.Property)
-	{
-		ChangedEvent.Broadcast(this, PropertyChangedEvent.MemberProperty);
-	}
-}
-#endif
 
diff --git a/Engine/Source/Runtime/Engine/Private/SkeletalMeshComponentPhysics.cpp b/Engine/Source/Runtime/Engine/Private/SkeletalMeshComponentPhysics.cpp
--- a/Engine/Source/Runtime/Engine/Private/SkeletalMeshComponentPhysics.cpp
+++ b/Engine/Source/Runtime/Engine/Private/SkeletalMeshComponentPhysics.cpp
@@ -1442,11 +1442,6 @@ void USkeletalMeshComponent::OnUpdateTransform(bool bSkipPhysicsMove, ETeleportT
 #endif
 	}
 
-	if (MeshObject)
-	{
-		MeshObject->UpdateShadowShapes(this);
-	}
-
 #if WITH_APEX_CLOTHING
 	if(ClothingActors.Num() > 0)
 	{
@@ -1965,6 +1960,81 @@ float USkeletalMeshComponent::GetDistanceToCollision(const FVector& Point, FVect
 	return ClosestPointDistance;
 }
 
+DECLARE_CYCLE_STAT(TEXT("GetClosestPointOnPhysicsAsset"), STAT_GetClosestPointOnPhysicsAsset, STATGROUP_Physics);
+
+bool USkeletalMeshComponent::GetClosestPointOnPhysicsAsset(const FVector& WorldPosition, FClosestPointOnPhysicsAsset& ClosestPointOnPhysicsAsset, bool bApproximate) const
+{
+	SCOPE_CYCLE_COUNTER(STAT_GetClosestPointOnPhysicsAsset);
+
+	bool bSuccess = false;
+	const UPhysicsAsset* PhysicsAsset = GetPhysicsAsset();
+	const FReferenceSkeleton* RefSkeleton = SkeletalMesh ? &SkeletalMesh->RefSkeleton : nullptr;
+	if(PhysicsAsset && RefSkeleton)
+	{
+		const TArray<FTransform>& SpaceBases = GetSpaceBases();
+		const bool bHasMasterPoseComponent = MasterPoseComponent.IsValid();
+		const FVector ComponentPosition = ComponentToWorld.InverseTransformPosition(WorldPosition);
+	
+		float CurrentClosestDistance = FLT_MAX;
+		int32 CurrentClosestBoneIndex = INDEX_NONE;
+		const UBodySetup* CurrentClosestBodySetup = nullptr;
+
+		for(const UBodySetup* BodySetupInstance : PhysicsAsset->BodySetup)
+		{
+			ClosestPointOnPhysicsAsset.Distance = FLT_MAX;
+			const FName BoneName = BodySetupInstance->BoneName;
+			const int32 BoneIndex = RefSkeleton->FindBoneIndex(BoneName);
+			if(BoneIndex != INDEX_NONE)
+			{
+				const FTransform BoneTM = bHasMasterPoseComponent ? GetBoneTransform(BoneIndex) : SpaceBases[BoneIndex];
+				const float Dist = bApproximate ? (BoneTM.GetLocation() - ComponentPosition).SizeSquared() : BodySetupInstance->GetShortestDistanceToPoint(ComponentPosition, BoneTM);
+
+				if (Dist < CurrentClosestDistance)
+				{
+					CurrentClosestDistance = Dist;
+					CurrentClosestBoneIndex = BoneIndex;
+					CurrentClosestBodySetup = BodySetupInstance;
+
+					if(Dist <= 0.f) { break; }
+				}
+			}
+		}
+
+		if(CurrentClosestBoneIndex >= 0)
+		{
+			bSuccess = true;
+
+			const FTransform BoneTM = bHasMasterPoseComponent ? GetBoneTransform(CurrentClosestBoneIndex) : (SpaceBases[CurrentClosestBoneIndex] * ComponentToWorld);
+			ClosestPointOnPhysicsAsset.Distance = CurrentClosestBodySetup->GetClosestPointAndNormal(WorldPosition, BoneTM, ClosestPointOnPhysicsAsset.ClosestWorldPosition, ClosestPointOnPhysicsAsset.Normal);
+			ClosestPointOnPhysicsAsset.BoneName = CurrentClosestBodySetup->BoneName;
+		}
+	}
+
+	return bSuccess;
+}
+
+bool USkeletalMeshComponent::K2_GetClosestPointOnPhysicsAsset(const FVector& WorldPosition, FVector& ClosestWorldPosition, FVector& Normal, FName& BoneName, float& Distance) const
+{
+	FClosestPointOnPhysicsAsset ClosestPointOnPhysicsAsset;
+	bool bSuccess = GetClosestPointOnPhysicsAsset(WorldPosition, ClosestPointOnPhysicsAsset, /*bApproximate =*/ false);
+	if(bSuccess)
+	{
+		ClosestWorldPosition = ClosestPointOnPhysicsAsset.ClosestWorldPosition;
+		Normal = ClosestPointOnPhysicsAsset.Normal;
+		BoneName = ClosestPointOnPhysicsAsset.BoneName;
+		Distance = ClosestPointOnPhysicsAsset.Distance;
+	}
+	else
+	{
+		ClosestWorldPosition = FVector::ZeroVector;
+		Normal = FVector::ZeroVector;
+		BoneName = NAME_None;
+		Distance = -1;
+	}
+
+	return bSuccess;
+}
+
 bool USkeletalMeshComponent::LineTraceComponent(struct FHitResult& OutHit, const FVector Start, const FVector End, const struct FCollisionQueryParams& Params)
 {
 	UWorld* const World = GetWorld();
diff --git a/Engine/Source/Runtime/Engine/Private/SkeletalRender.cpp b/Engine/Source/Runtime/Engine/Private/SkeletalRender.cpp
--- a/Engine/Source/Runtime/Engine/Private/SkeletalRender.cpp
+++ b/Engine/Source/Runtime/Engine/Private/SkeletalRender.cpp
@@ -181,63 +181,6 @@ void FSkeletalMeshObject::InitLODInfos(const USkinnedMeshComponent* SkelComponen
 	}
 }
 
-const FQuat SphylBasis(FVector(1.0f / FMath::Sqrt(2.0f), 0.0f, 1.0f / FMath::Sqrt(2.0f)), PI);
-
-void FSkeletalMeshObject::UpdateShadowShapes(USkinnedMeshComponent* InMeshComponent)
-{
-	UPhysicsAsset* ShadowPhysicsAsset = InMeshComponent->SkeletalMesh->ShadowPhysicsAsset;
-
-	if (ShadowPhysicsAsset 
-		&& InMeshComponent->CastShadow
-		&& (InMeshComponent->bCastCapsuleDirectShadow || InMeshComponent->bCastCapsuleIndirectShadow))
-	{
-		//@todo - double buffer or temporary allocator
-		TArray<FCapsuleShape>* NewShadowCapsuleShapes = new TArray<FCapsuleShape>();
-
-		for (int32 BodyIndex = 0; BodyIndex < ShadowPhysicsAsset->BodySetup.Num(); BodyIndex++)
-		{
-			UBodySetup* BodySetup = ShadowPhysicsAsset->BodySetup[BodyIndex];
-			int32 BoneIndex = InMeshComponent->GetBoneIndex(BodySetup->BoneName);
-
-			if (BoneIndex != INDEX_NONE)
-			{
-				FTransform WorldBoneTransform = InMeshComponent->GetBoneTransform(BoneIndex, InMeshComponent->GetComponentTransform());
-				const float MaxScale = WorldBoneTransform.GetScale3D().GetMax();
-
-				for (int32 ShapeIndex = 0; ShapeIndex < BodySetup->AggGeom.SphereElems.Num(); ShapeIndex++)
-				{
-					const FKSphereElem& SphereShape = BodySetup->AggGeom.SphereElems[ShapeIndex];
-					NewShadowCapsuleShapes->Add(FCapsuleShape(
-						WorldBoneTransform.TransformPosition(SphereShape.Center), 
-						SphereShape.Radius * MaxScale,
-						FVector(0, 0, 1),
-						0));
-				}
-
-				for (int32 ShapeIndex = 0; ShapeIndex < BodySetup->AggGeom.SphylElems.Num(); ShapeIndex++)
-				{
-					const FKSphylElem& SphylShape = BodySetup->AggGeom.SphylElems[ShapeIndex];
-					NewShadowCapsuleShapes->Add(FCapsuleShape(
-						WorldBoneTransform.TransformPosition(SphylShape.Center), 
-						SphylShape.Radius * MaxScale,
-						WorldBoneTransform.TransformVector((SphylShape.Orientation * SphylBasis).Vector()),
-						SphylShape.Length * MaxScale));
-				}
-			}
-		}
-
-		ENQUEUE_UNIQUE_RENDER_COMMAND_TWOPARAMETER(
-			ShadowShapesUpdateCommand,
-			FSkeletalMeshObject*, MeshObject, this,
-			TArray<FCapsuleShape>*, NewShadowCapsuleShapes, NewShadowCapsuleShapes,
-			{
-				MeshObject->ShadowCapsuleShapes = *NewShadowCapsuleShapes;
-				delete NewShadowCapsuleShapes;
-			}
-		);
-	}
-}
-
 /*-----------------------------------------------------------------------------
 Global functions
 -----------------------------------------------------------------------------*/
diff --git a/Engine/Source/Runtime/Engine/Private/SkeletalRenderCPUSkin.cpp b/Engine/Source/Runtime/Engine/Private/SkeletalRenderCPUSkin.cpp
--- a/Engine/Source/Runtime/Engine/Private/SkeletalRenderCPUSkin.cpp
+++ b/Engine/Source/Runtime/Engine/Private/SkeletalRenderCPUSkin.cpp
@@ -170,8 +170,6 @@ void FSkeletalMeshObjectCPUSkin::Update(int32 LODIndex,USkinnedMeshComponent* In
 	// this data is only deleted when another update is sent
 	FDynamicSkelMeshObjectDataCPUSkin* NewDynamicData = new FDynamicSkelMeshObjectDataCPUSkin(InMeshComponent,SkeletalMeshResource,LODIndex,ActiveVertexAnims);
 
-	UpdateShadowShapes(InMeshComponent);
-
 	// queue a call to update this data
 	ENQUEUE_UNIQUE_RENDER_COMMAND_THREEPARAMETER(
 		SkelMeshObjectUpdateDataCommand,
@@ -353,6 +351,11 @@ TArray<FTransform>* FSkeletalMeshObjectCPUSkin::GetSpaceBases() const
 	}
 }
 
+const TArray<FMatrix>& FSkeletalMeshObjectCPUSkin::GetReferenceToLocalMatrices() const
+{
+	return DynamicData->ReferenceToLocal;
+}
+
 /**
  * Get the origin and direction vectors for TRISORT_CustomLeftRight sections
  */
diff --git a/Engine/Source/Runtime/Engine/Private/SkeletalRenderCPUSkin.h b/Engine/Source/Runtime/Engine/Private/SkeletalRenderCPUSkin.h
--- a/Engine/Source/Runtime/Engine/Private/SkeletalRenderCPUSkin.h
+++ b/Engine/Source/Runtime/Engine/Private/SkeletalRenderCPUSkin.h
@@ -152,6 +152,7 @@ public:
 	virtual bool IsCPUSkinned() const override { return true; }
 	virtual const FVertexFactory* GetVertexFactory(int32 LODIndex,int32 ChunkIdx) const override;
 	virtual TArray<FTransform>* GetSpaceBases() const override;
+	virtual const TArray<FMatrix>& GetReferenceToLocalMatrices() const override;
 	virtual int32 GetLOD() const override
 	{
 		if(DynamicData)
diff --git a/Engine/Source/Runtime/Engine/Private/SkeletalRenderGPUSkin.cpp b/Engine/Source/Runtime/Engine/Private/SkeletalRenderGPUSkin.cpp
--- a/Engine/Source/Runtime/Engine/Private/SkeletalRenderGPUSkin.cpp
+++ b/Engine/Source/Runtime/Engine/Private/SkeletalRenderGPUSkin.cpp
@@ -195,8 +195,6 @@ void FSkeletalMeshObjectGPUSkin::Update(int32 LODIndex,USkinnedMeshComponent* In
 	FDynamicSkelMeshObjectDataGPUSkin* NewDynamicData = FDynamicSkelMeshObjectDataGPUSkin::AllocDynamicSkelMeshObjectDataGPUSkin();		
 	NewDynamicData->InitDynamicSkelMeshObjectDataGPUSkin(InMeshComponent,SkeletalMeshResource,LODIndex,ActiveVertexAnims);
 
-	UpdateShadowShapes(InMeshComponent);
-
 	{
 		// Handle the case of skin caching shaders not done compiling before updates are finished/editor is loading
 		static bool bNeedToWait = GEnableGPUSkinCache != 0;
@@ -988,6 +986,11 @@ TArray<FTransform>* FSkeletalMeshObjectGPUSkin::GetSpaceBases() const
 	}
 }
 
+const TArray<FMatrix>& FSkeletalMeshObjectGPUSkin::GetReferenceToLocalMatrices() const
+{
+	return DynamicData->ReferenceToLocal;
+}
+
 const FTwoVectors& FSkeletalMeshObjectGPUSkin::GetCustomLeftRightVectors(int32 SectionIndex) const
 {
 	if( DynamicData && DynamicData->CustomLeftRightVectors.IsValidIndex(SectionIndex) )
diff --git a/Engine/Source/Runtime/Engine/Private/SkeletalRenderGPUSkin.h b/Engine/Source/Runtime/Engine/Private/SkeletalRenderGPUSkin.h
--- a/Engine/Source/Runtime/Engine/Private/SkeletalRenderGPUSkin.h
+++ b/Engine/Source/Runtime/Engine/Private/SkeletalRenderGPUSkin.h
@@ -200,6 +200,7 @@ public:
 	virtual void CacheVertices(int32 LODIndex, bool bForce) const override {}
 	virtual bool IsCPUSkinned() const override { return false; }
 	virtual TArray<FTransform>* GetSpaceBases() const override;
+	virtual const TArray<FMatrix>& GetReferenceToLocalMatrices() const override;
 
 	virtual int32 GetLOD() const override
 	{
diff --git a/Engine/Source/Runtime/Engine/Private/SoundBase.cpp b/Engine/Source/Runtime/Engine/Private/SoundBase.cpp
--- a/Engine/Source/Runtime/Engine/Private/SoundBase.cpp
+++ b/Engine/Source/Runtime/Engine/Private/SoundBase.cpp
@@ -20,7 +20,13 @@ void USoundBase::PostInitProperties()
 	const FStringAssetReference DefaultSoundClassName = GetDefault<UAudioSettings>()->DefaultSoundClassName;
 	if (DefaultSoundClassName.IsValid())
 	{
-		SoundClassObject = LoadObject<USoundClass>(NULL, *DefaultSoundClassName.ToString());
+		SoundClassObject = LoadObject<USoundClass>(nullptr, *DefaultSoundClassName.ToString());
+	}
+
+	const FStringAssetReference DefaultSoundConcurrencyName = GetDefault<UAudioSettings>()->DefaultSoundConcurrencyName;
+	if (DefaultSoundConcurrencyName.IsValid())
+	{
+		SoundConcurrencySettings = LoadObject<USoundConcurrency>(nullptr, *DefaultSoundConcurrencyName.ToString());
 	}
 }
 
diff --git a/Engine/Source/Runtime/Engine/Private/SoundConcurrency.cpp b/Engine/Source/Runtime/Engine/Private/SoundConcurrency.cpp
--- a/Engine/Source/Runtime/Engine/Private/SoundConcurrency.cpp
+++ b/Engine/Source/Runtime/Engine/Private/SoundConcurrency.cpp
@@ -98,7 +98,7 @@ void FConcurrencyGroup::AddActiveSound(FActiveSound* ActiveSound)
 void FConcurrencyGroup::StopQuietSoundsDueToMaxConcurrency()
 {
 	// Nothing to do if our active sound count is less than or equal to our max active sounds
-	if (ResolutionRule != EMaxConcurrentResolutionRule::StopQuietist || ActiveSounds.Num() <= MaxActiveSounds)
+	if (ResolutionRule != EMaxConcurrentResolutionRule::StopQuietest || ActiveSounds.Num() <= MaxActiveSounds)
 	{
 		return;
 	}
@@ -183,7 +183,7 @@ FActiveSound* FSoundConcurrencyManager::ResolveConcurrency(const FActiveSound& N
 	bool bCanPlay = true;
 	FActiveSound* SoundToStop = nullptr;
 
-	if (ActiveSounds.Num() >= Concurrency->MaxCount && Concurrency->ResolutionRule != EMaxConcurrentResolutionRule::StopQuietist)
+	if (ActiveSounds.Num() >= Concurrency->MaxCount && Concurrency->ResolutionRule != EMaxConcurrentResolutionRule::StopQuietest)
 	{
 		switch (Concurrency->ResolutionRule)
 		{
@@ -250,8 +250,8 @@ FActiveSound* FSoundConcurrencyManager::ResolveConcurrency(const FActiveSound& N
 			}
 			break;
 
-			case EMaxConcurrentResolutionRule::StopQuietist:
-			// We won't do anything upfront for StopQuietist resolution rule, but instead resolve it after volumes have been computed
+			case EMaxConcurrentResolutionRule::StopQuietest:
+			// We won't do anything upfront for StopQuietest resolution rule, but instead resolve it after volumes have been computed
 			// and before we stopping sounds that play past the max overall voice count. This is because it is nigh impossible (currently)
 			// for UE4 to evaluate sound volumes *before* they play. In addition to the fact that gain-stages aren't evaluated independently, 
 			// this is primarily because of the way SoundCue nodes work.
diff --git a/Engine/Source/Runtime/Engine/Private/SoundNode.cpp b/Engine/Source/Runtime/Engine/Private/SoundNode.cpp
--- a/Engine/Source/Runtime/Engine/Private/SoundNode.cpp
+++ b/Engine/Source/Runtime/Engine/Private/SoundNode.cpp
@@ -152,6 +152,22 @@ float USoundNode::GetDuration()
 	return MaxDuration;
 }
 
+int32 USoundNode::GetNumSounds(const UPTRINT NodeWaveInstanceHash, FActiveSound& ActiveSound) const
+{
+	// Default implementation loops through all child nodes and sums the number of sounds.
+	// For most nodes this will result in 1, for node mixers, this will result in multiple sounds.
+	int32 NumSounds = 0;
+	for (USoundNode* ChildNode : ChildNodes)
+	{
+		if (ChildNode)
+		{
+			NumSounds += ChildNode->GetNumSounds(NodeWaveInstanceHash, ActiveSound);
+		}
+	}
+
+	return NumSounds;
+}
+
 #if WITH_EDITOR
 void USoundNode::PostEditChangeProperty(FPropertyChangedEvent& PropertyChangedEvent)
 {
diff --git a/Engine/Source/Runtime/Engine/Private/SoundNodeConcatenator.cpp b/Engine/Source/Runtime/Engine/Private/SoundNodeConcatenator.cpp
--- a/Engine/Source/Runtime/Engine/Private/SoundNodeConcatenator.cpp
+++ b/Engine/Source/Runtime/Engine/Private/SoundNodeConcatenator.cpp
@@ -4,6 +4,21 @@
 #include "EnginePrivate.h"
 #include "SoundDefinitions.h"
 #include "Sound/SoundNodeConcatenator.h"
+#include "Sound/SoundWave.h"
+
+// Payload used for instances of a concatenator node
+struct FSoundNodeConcatenatorPayload
+{
+	/** Which child node we're currently playing back in the concatenation sequence */
+	int32 NodeIndex;
+
+	/** Which child node sound index we're currently playing in the current child node */
+	int32 CurrentChildNodeSoundIndex;
+
+	/** The number of sounds playing in current child node */
+	int32 CurrentChildNodeNumSound;
+};
+
 
 /*-----------------------------------------------------------------------------
          USoundNodeConcatenator implementation.
@@ -13,23 +28,43 @@ USoundNodeConcatenator::USoundNodeConcatenator(const FObjectInitializer& ObjectI
 {
 }
 
-
 bool USoundNodeConcatenator::NotifyWaveInstanceFinished( FWaveInstance* WaveInstance )
 {
 	FActiveSound& ActiveSound = *WaveInstance->ActiveSound;
 	const UPTRINT NodeWaveInstanceHash = WaveInstance->NotifyBufferFinishedHooks.GetHashForNode(this);
-	RETRIEVE_SOUNDNODE_PAYLOAD( sizeof( int32 ) );
-	DECLARE_SOUNDNODE_ELEMENT( int32, NodeIndex );
+	RETRIEVE_SOUNDNODE_PAYLOAD(sizeof(FSoundNodeConcatenatorPayload));
+	DECLARE_SOUNDNODE_ELEMENT(FSoundNodeConcatenatorPayload, ConcatenatorPayload);
 	check( *RequiresInitialization == 0 );
 
-	// Allow wave instance to be played again the next iteration.
-	WaveInstance->bIsStarted = false;
-	WaveInstance->bIsFinished = false;
+	// Only bump up the node index if all the current child node's sounds have finished playing. 
+	// Otherwise, this will end up cutting short sounds played through a mixer.
 
-	// Advance index.
-	NodeIndex++;
+	ConcatenatorPayload.CurrentChildNodeSoundIndex++;
+	if (ConcatenatorPayload.NodeIndex < ChildNodes.Num())
+	{
+		USoundNode* CurrentChildNode = ChildNodes[ConcatenatorPayload.NodeIndex];
+		check(CurrentChildNode);
 
-	return (NodeIndex < ChildNodes.Num());
+		if (ConcatenatorPayload.CurrentChildNodeSoundIndex == ConcatenatorPayload.CurrentChildNodeNumSound)
+		{
+			// Find the next non-null index
+			for (++ConcatenatorPayload.NodeIndex; ConcatenatorPayload.NodeIndex < ChildNodes.Num(); ++ConcatenatorPayload.NodeIndex)
+			{
+				if (ChildNodes[ConcatenatorPayload.NodeIndex])
+				{
+					break;
+				}
+			}
+
+			ConcatenatorPayload.CurrentChildNodeSoundIndex = 0;
+
+			// Allow wave instance to be played again the next iteration.
+			WaveInstance->bIsStarted = false;
+			WaveInstance->bIsFinished = false;
+		}
+	}
+
+	return (ConcatenatorPayload.NodeIndex < ChildNodes.Num());
 }
 
 
@@ -48,6 +83,12 @@ float USoundNodeConcatenator::GetDuration()
 	return Duration;
 }
 
+int32 USoundNodeConcatenator::GetNumSounds(const UPTRINT NodeWaveInstanceHash, FActiveSound& ActiveSound) const
+{
+	// Counter-intuitively, a node concatenator will always play 1 sound from the perspective of other node-concatenators
+	return 1;
+}
+
 void USoundNodeConcatenator::CreateStartingConnectors()
 {
 	// Concatenators default to two two connectors.
@@ -72,28 +113,54 @@ void USoundNodeConcatenator::RemoveChildNode( int32 Index )
 
 void USoundNodeConcatenator::ParseNodes( FAudioDevice* AudioDevice, const UPTRINT NodeWaveInstanceHash, FActiveSound& ActiveSound, const FSoundParseParameters& ParseParams, TArray<FWaveInstance*>& WaveInstances )
 {
-	RETRIEVE_SOUNDNODE_PAYLOAD( sizeof( int32 ) );
-	DECLARE_SOUNDNODE_ELEMENT( int32, NodeIndex );
+	// Local stack copy of the payload. We need to do this because ParseNodes may resize the payload map
+	FSoundNodeConcatenatorPayload LocalPayload;
 
-	// Start from the beginning.
-	if( *RequiresInitialization )
 	{
-		NodeIndex = 0;
-		*RequiresInitialization = false;
+		RETRIEVE_SOUNDNODE_PAYLOAD(sizeof(FSoundNodeConcatenatorPayload));
+		DECLARE_SOUNDNODE_ELEMENT(FSoundNodeConcatenatorPayload, ConcatenatorPayload);
+
+		// Start from the beginning.
+		if (*RequiresInitialization)
+		{
+			ConcatenatorPayload.NodeIndex = 0;
+			ConcatenatorPayload.CurrentChildNodeSoundIndex = 0;
+			ConcatenatorPayload.CurrentChildNodeNumSound = 0;
+			*RequiresInitialization = false;
+		}
+
+		ConcatenatorPayload.CurrentChildNodeNumSound = 0;
+
+		LocalPayload = ConcatenatorPayload;
 	}
 
 	// Play the current node.
-	if( NodeIndex < ChildNodes.Num() )
+	while (LocalPayload.NodeIndex < ChildNodes.Num())
 	{
-		FSoundParseParameters UpdatedParams = ParseParams;
-		UpdatedParams.NotifyBufferFinishedHooks.AddNotify(this, NodeWaveInstanceHash);
+		USoundNode* ChildNode = ChildNodes[LocalPayload.NodeIndex];
+		if (ChildNode)
+		{
+			// Play currently active node.
+			FSoundParseParameters UpdatedParams = ParseParams;
+			UpdatedParams.NotifyBufferFinishedHooks.AddNotify(this, NodeWaveInstanceHash);
+			UpdatedParams.VolumeMultiplier *= InputVolume[LocalPayload.NodeIndex];
 
-		// Play currently active node.
-		USoundNode* ChildNode = ChildNodes[ NodeIndex ];
-		if( ChildNode )
+			const UPTRINT ChildNodeWaveInstanceHash = GetNodeWaveInstanceHash(NodeWaveInstanceHash, ChildNode, LocalPayload.NodeIndex);
+			ChildNode->ParseNodes(AudioDevice, ChildNodeWaveInstanceHash, ActiveSound, UpdatedParams, WaveInstances);
+
+			// Update the payload with the number of sounds played and update our local copy
+			RETRIEVE_SOUNDNODE_PAYLOAD(sizeof(FSoundNodeConcatenatorPayload));
+			DECLARE_SOUNDNODE_ELEMENT(FSoundNodeConcatenatorPayload, ConcatenatorPayload);
+
+			ConcatenatorPayload = LocalPayload;
+			ConcatenatorPayload.CurrentChildNodeNumSound = ChildNode->GetNumSounds(ChildNodeWaveInstanceHash, ActiveSound);
+
+			break;
+		}
+		else
 		{
-			UpdatedParams.VolumeMultiplier *= InputVolume[NodeIndex];
-			ChildNode->ParseNodes( AudioDevice, GetNodeWaveInstanceHash(NodeWaveInstanceHash, ChildNode, NodeIndex), ActiveSound, UpdatedParams, WaveInstances);
+			LocalPayload.CurrentChildNodeSoundIndex = 0;
+			++LocalPayload.NodeIndex;
 		}
 	}
 }
diff --git a/Engine/Source/Runtime/Engine/Private/SoundNodeDistanceCrossFade.cpp b/Engine/Source/Runtime/Engine/Private/SoundNodeDistanceCrossFade.cpp
--- a/Engine/Source/Runtime/Engine/Private/SoundNodeDistanceCrossFade.cpp
+++ b/Engine/Source/Runtime/Engine/Private/SoundNodeDistanceCrossFade.cpp
@@ -39,6 +39,8 @@ float USoundNodeDistanceCrossFade::MaxAudibleDistance( float CurrentMaxDistance
 
 void USoundNodeDistanceCrossFade::ParseNodes( FAudioDevice* AudioDevice, const UPTRINT NodeWaveInstanceHash, FActiveSound& ActiveSound, const FSoundParseParameters& ParseParams, TArray<FWaveInstance*>& WaveInstances )
 {
+	int32 NumSoundsActive = 0;
+
 	FSoundParseParameters UpdatedParams = ParseParams;
 	for( int32 ChildNodeIndex = 0; ChildNodeIndex < ChildNodes.Num(); ChildNodeIndex++ )
 	{
@@ -88,12 +90,31 @@ void USoundNodeDistanceCrossFade::ParseNodes( FAudioDevice* AudioDevice, const U
 				VolumeToSet = 0.f; //CrossFadeInput( ChildNodeIndex ).Volume;
 			}
 
+			if (VolumeToSet > 0.0f)
+			{
+				++NumSoundsActive;
+			}
+
 			UpdatedParams.Volume = ParseParams.Volume * VolumeToSet;
 
 			// "play" the rest of the tree
 			ChildNodes[ ChildNodeIndex ]->ParseNodes( AudioDevice, GetNodeWaveInstanceHash(NodeWaveInstanceHash, ChildNodes[ChildNodeIndex], ChildNodeIndex), ActiveSound, UpdatedParams, WaveInstances );
 		}
 	}
+
+	// Write out the results of the NumSounds count. We can't do this in the loop since the table storing this data will potentially reallocate
+	RETRIEVE_SOUNDNODE_PAYLOAD(sizeof(int32));
+	DECLARE_SOUNDNODE_ELEMENT(int32, NumSoundsUsedInCrossFade);
+
+	NumSoundsUsedInCrossFade = NumSoundsActive;
+}
+
+int32 USoundNodeDistanceCrossFade::GetNumSounds(const UPTRINT NodeWaveInstanceHash, FActiveSound& ActiveSound) const
+{
+	RETRIEVE_SOUNDNODE_PAYLOAD(sizeof(int32));
+	DECLARE_SOUNDNODE_ELEMENT(int32, NumSoundsUsedInCrossFade);
+
+	return NumSoundsUsedInCrossFade;
 }
 
 
diff --git a/Engine/Source/Runtime/Engine/Private/SoundNodeRandom.cpp b/Engine/Source/Runtime/Engine/Private/SoundNodeRandom.cpp
--- a/Engine/Source/Runtime/Engine/Private/SoundNodeRandom.cpp
+++ b/Engine/Source/Runtime/Engine/Private/SoundNodeRandom.cpp
@@ -161,6 +161,20 @@ void USoundNodeRandom::ParseNodes( FAudioDevice* AudioDevice, const UPTRINT Node
 	}
 }
 
+int32 USoundNodeRandom::GetNumSounds(const UPTRINT NodeWaveInstanceHash, FActiveSound& ActiveSound) const
+{
+	RETRIEVE_SOUNDNODE_PAYLOAD(sizeof(int32));
+	DECLARE_SOUNDNODE_ELEMENT(int32, NodeIndex);
+
+	// Only return the number of soundso the chosen node. If the random node hasn't been initialized yet
+	// then there is no chosen node so return 0.
+	if (!*RequiresInitialization && NodeIndex < ChildNodes.Num() && ChildNodes[NodeIndex])
+	{
+		return ChildNodes[NodeIndex]->GetNumSounds(NodeWaveInstanceHash, ActiveSound);
+	}
+	return 0;
+}
+
 void USoundNodeRandom::CreateStartingConnectors()
 {
 	// Random Sound Nodes default with two connectors.
diff --git a/Engine/Source/Runtime/Engine/Private/SoundWave.cpp b/Engine/Source/Runtime/Engine/Private/SoundWave.cpp
--- a/Engine/Source/Runtime/Engine/Private/SoundWave.cpp
+++ b/Engine/Source/Runtime/Engine/Private/SoundWave.cpp
@@ -66,8 +66,15 @@ USoundWave::USoundWave(const FObjectInitializer& ObjectInitializer)
 SIZE_T USoundWave::GetResourceSize(EResourceSizeMode::Type Mode)
 {
 	int32 CalculatedResourceSize = 0;
+	
+	FAudioDevice* AudioDevice = nullptr;
+	
+	if (GEngine)
+	{
+		AudioDevice = GEngine->GetMainAudioDevice();
+	}
 
-	if (DecompressionType == DTYPE_Native)
+	if (AudioDevice && AudioDevice->HasCompressedAudioInfoClass(this) && DecompressionType == DTYPE_Native)
 	{
 		// If we've been decompressed, need to account for decompressed and also compressed
 		CalculatedResourceSize += RawPCMDataSize;
@@ -80,13 +87,10 @@ SIZE_T USoundWave::GetResourceSize(EResourceSizeMode::Type Mode)
 		}
 	}
 
-	if (GEngine && GEngine->GetMainAudioDevice())
+	// Don't add compressed data to size of streaming sounds
+	if (AudioDevice && (!FPlatformProperties::SupportsAudioStreaming() || !IsStreaming()))
 	{
-		// Don't add compressed data to size of streaming sounds
-		if (!FPlatformProperties::SupportsAudioStreaming() || !IsStreaming())
-		{
-			CalculatedResourceSize += GetCompressedDataSize(GEngine->GetMainAudioDevice()->GetRuntimeFormat(this));
-		}
+		CalculatedResourceSize += GetCompressedDataSize(AudioDevice->GetRuntimeFormat(this));
 	}
 
 	return CalculatedResourceSize;
diff --git a/Engine/Source/Runtime/Engine/Private/TickTaskManager.cpp b/Engine/Source/Runtime/Engine/Private/TickTaskManager.cpp
--- a/Engine/Source/Runtime/Engine/Private/TickTaskManager.cpp
+++ b/Engine/Source/Runtime/Engine/Private/TickTaskManager.cpp
@@ -22,6 +22,7 @@ DECLARE_DWORD_COUNTER_STAT(TEXT("Ticks Queued"),STAT_TicksQueued,STATGROUP_Game)
 DECLARE_CYCLE_STAT(TEXT("TG_NewlySpawned"), STAT_TG_NewlySpawned, STATGROUP_TickGroups);
 DECLARE_CYCLE_STAT(TEXT("ReleaseTickGroup"), STAT_ReleaseTickGroup, STATGROUP_TickGroups);
 DECLARE_CYCLE_STAT(TEXT("ReleaseTickGroup Block"), STAT_ReleaseTickGroup_Block, STATGROUP_TickGroups);
+DECLARE_CYCLE_STAT(TEXT("CleanupTasksWait"), STAT_CleanupTasksWait, STATGROUP_TickGroups);
 
 
 
@@ -525,6 +526,9 @@ public:
 		{
 			bAllowConcurrentTicks = !!CVarAllowAsyncComponentTicks.GetValueOnGameThread();
 		}
+
+		WaitForCleanup();
+
 		for (int32 Index = 0; Index < TG_MAX; Index++)
 		{
 			check(!TickCompletionEvents[Index].Num());  // we should not be adding to these outside of a ticking proper and they were already cleared after they were ticked
@@ -547,16 +551,6 @@ public:
 		{
 			UE_LOG(LogTick, Log, TEXT("tick %6d ---------------------------------------- End Frame"),GFrameCounter);
 		}
-		FTaskGraphInterface::Get().WaitUntilTasksComplete(CleanupTasks, ENamedThreads::GameThread);
-		CleanupTasks.Reset();
-		for (int32 Index = 0; Index < TG_MAX; Index++)
-		{
-			check(!TickCompletionEvents[Index].Num());  // we should not be adding to these outside of a ticking proper and they were already cleared after they were ticked
-			for (int32 IndexInner = 0; IndexInner < TG_MAX; IndexInner++)
-			{
-				check(!TickTasks[Index][IndexInner].Num() && !HiPriTickTasks[Index][IndexInner].Num());  // we should not be adding to these outside of a ticking proper and they were already cleared after they were ticked
-			}
-		}
 	}
 private:
 
@@ -565,6 +559,25 @@ private:
 		, bLogTicks(false)
 		, bLogTicksShowPrerequistes(false)
 	{
+		TFunction<void()> ShutdownCallback([this](){WaitForCleanup();});
+		FTaskGraphInterface::Get().AddShutdownCallback(ShutdownCallback);
+	}
+
+	~FTickTaskSequencer()
+	{
+		// Need to clean up oustanding tasks before we destroy this data structure.
+		// Typically it is already gone because the task graph shutdown first.
+		WaitForCleanup();
+	}
+
+	void WaitForCleanup()
+	{
+		if (CleanupTasks.Num() > 0)
+		{
+			SCOPE_CYCLE_COUNTER(STAT_CleanupTasksWait);
+			FTaskGraphInterface::Get().WaitUntilTasksComplete(CleanupTasks, ENamedThreads::GameThread);
+			CleanupTasks.Reset();
+		}
 	}
 
 	void ResetTickGroup(ETickingGroup WorldTickGroup)
@@ -1098,7 +1111,7 @@ public:
 					FTickFunction* PrevComparisionFunction = nullptr;
 					FTickFunction* ComparisonFunction = AllCoolingDownTickFunctions.Head;
 					bool bFound = false;
-					while (ComparisonFunction)
+					while (ComparisonFunction && !bFound)
 					{
 						if (ComparisonFunction == TickFunction)
 						{
diff --git a/Engine/Source/Runtime/Engine/Private/UnrealEngine.cpp b/Engine/Source/Runtime/Engine/Private/UnrealEngine.cpp
--- a/Engine/Source/Runtime/Engine/Private/UnrealEngine.cpp
+++ b/Engine/Source/Runtime/Engine/Private/UnrealEngine.cpp
@@ -10312,7 +10312,7 @@ void UEngine::VerifyLoadMapWorldCleanup()
 		const bool bIsPersistantWorldType = (World->WorldType == EWorldType::Inactive) || (World->WorldType == EWorldType::Preview);
 		if (!bIsPersistantWorldType && !WorldHasValidContext(World))
 		{
-			if (World->PersistentLevel != nullptr && !WorldHasValidContext(World->PersistentLevel->OwningWorld))
+			if (World->PersistentLevel == nullptr || !WorldHasValidContext(World->PersistentLevel->OwningWorld))
 			{
 				// Print some debug information...
 				UE_LOG(LogLoad, Log, TEXT("%s not cleaned up by garbage collection! "), *World->GetFullName());
diff --git a/Engine/Source/Runtime/Engine/Private/UserInterface/Canvas.cpp b/Engine/Source/Runtime/Engine/Private/UserInterface/Canvas.cpp
--- a/Engine/Source/Runtime/Engine/Private/UserInterface/Canvas.cpp
+++ b/Engine/Source/Runtime/Engine/Private/UserInterface/Canvas.cpp
@@ -1948,7 +1948,7 @@ void UCanvas::K2_DrawTriangle(UTexture* RenderTexture, TArray<FCanvasUVTri> Tria
 	if (Triangles.Num() > 0)
 	{
 		FCanvasTriangleItem TriangleItem(FVector2D::ZeroVector, FVector2D::ZeroVector, FVector2D::ZeroVector, (RenderTexture) ? RenderTexture->Resource : GWhiteTexture);
-		TriangleItem.TriangleList = Triangles;
+		TriangleItem.TriangleList = MoveTemp(Triangles);
 		DrawItem(TriangleItem);
 	}
 }
@@ -1959,7 +1959,7 @@ void UCanvas::K2_DrawMaterialTriangle(UMaterialInterface* RenderMaterial, TArray
 	{
 		FCanvasTriangleItem TriangleItem(FVector2D::ZeroVector, FVector2D::ZeroVector, FVector2D::ZeroVector, NULL);
 		TriangleItem.MaterialRenderProxy = RenderMaterial->GetRenderProxy(0);
-		TriangleItem.TriangleList = Triangles;
+		TriangleItem.TriangleList = MoveTemp(Triangles);
 		DrawItem(TriangleItem);
 	}
 }
diff --git a/Engine/Source/Runtime/Engine/Private/World.cpp b/Engine/Source/Runtime/Engine/Private/World.cpp
--- a/Engine/Source/Runtime/Engine/Private/World.cpp
+++ b/Engine/Source/Runtime/Engine/Private/World.cpp
@@ -3384,7 +3384,7 @@ void UWorld::RemovePawn( APawn* Pawn )
 	check( Pawn );
 	
 	AController* Controller = Pawn->GetController();
-	if (Controller)
+	if (Controller && (Controller->GetPawn() == Pawn))
 	{
 		Controller->UnPossess();
 	}
diff --git a/Engine/Source/Runtime/Engine/Public/AudioDevice.h b/Engine/Source/Runtime/Engine/Public/AudioDevice.h
--- a/Engine/Source/Runtime/Engine/Public/AudioDevice.h
+++ b/Engine/Source/Runtime/Engine/Public/AudioDevice.h
@@ -225,6 +225,11 @@ public:
 	{
 	}
 
+	/** Returns an array of available audio devices names for the platform */
+	virtual void GetAudioDeviceList(TArray<FString>& OutAudioDeviceNames) const
+	{
+	}
+
 	/**
 	 * Basic initialisation of the platform agnostic layer of the audio system
 	 */
diff --git a/Engine/Source/Runtime/Engine/Public/ComponentReregisterContext.h b/Engine/Source/Runtime/Engine/Public/ComponentReregisterContext.h
--- a/Engine/Source/Runtime/Engine/Public/ComponentReregisterContext.h
+++ b/Engine/Source/Runtime/Engine/Public/ComponentReregisterContext.h
@@ -14,7 +14,7 @@ ENGINE_API DECLARE_LOG_CATEGORY_EXTERN(LogActorComponent, Log, All);
 class FComponentReregisterContextBase
 {
 protected:
-	//Unregisters the Componet and returns the world it was registered to.
+	//Unregisters the Component and returns the world it was registered to.
 	UWorld* UnRegister(UActorComponent* InComponent)
 	{
 		UWorld* World = NULL;
@@ -143,12 +143,11 @@ private:
 	TArray<FMultiComponentReregisterPair> ComponentsPair;
 	
 public:
-	FMultiComponentReregisterContext(TArray<UActorComponent*> InComponents)
+	FMultiComponentReregisterContext(const TArray<UActorComponent*>& InComponents)
 	{
 		// Unregister each component and cache resulting scene
-		for(auto Iter = InComponents.CreateIterator(); Iter; ++Iter)
+		for (UActorComponent* Component : InComponents)
 		{
-			UActorComponent* Component = (*Iter);
 			UWorld* World = UnRegister(Component);
 			if(World)
 			{
diff --git a/Engine/Source/Runtime/Engine/Public/PhysicsPublic.h b/Engine/Source/Runtime/Engine/Public/PhysicsPublic.h
--- a/Engine/Source/Runtime/Engine/Public/PhysicsPublic.h
+++ b/Engine/Source/Runtime/Engine/Public/PhysicsPublic.h
@@ -624,6 +624,8 @@ private:
 
 	/** Map of SkeletalMeshComponents that need their bone transforms sent to the physics engine before simulation. */
 	TMap<USkeletalMeshComponent*, FDeferredKinematicUpdateInfo>	DeferredKinematicUpdateSkelMeshes;
+
+	FDelegateHandle PreGarbageCollectDelegateHandle;
 };
 
 /**
diff --git a/Engine/Source/Runtime/Engine/Public/SkeletalMeshTypes.h b/Engine/Source/Runtime/Engine/Public/SkeletalMeshTypes.h
--- a/Engine/Source/Runtime/Engine/Public/SkeletalMeshTypes.h
+++ b/Engine/Source/Runtime/Engine/Public/SkeletalMeshTypes.h
@@ -1839,6 +1839,12 @@ protected:
 	/** Array of section elements for each LOD */
 	TArray<FLODSectionElements> LODSections;
 
+	/** 
+	 * BoneIndex->capsule pairs used for rendering sphere/capsule shadows 
+	 * Note that these are in refpose component space, NOT bone space.
+	 */
+	TArray<TPair<int32, FCapsuleShape>> ShadowCapsuleData;
+
 	/** Set of materials used by this scene proxy, safe to access from the game thread. */
 	TSet<UMaterialInterface*> MaterialsInUse_GameThread;
 	bool bMaterialsNeedMorphUsage_GameThread;
diff --git a/Engine/Source/Runtime/Engine/Public/SkeletalRenderPublic.h b/Engine/Source/Runtime/Engine/Public/SkeletalRenderPublic.h
--- a/Engine/Source/Runtime/Engine/Public/SkeletalRenderPublic.h
+++ b/Engine/Source/Runtime/Engine/Public/SkeletalRenderPublic.h
@@ -74,6 +74,12 @@ public:
 	 */
 	virtual TArray<FTransform>* GetSpaceBases() const = 0;
 
+	/** 
+	 *	Get the array of refpose->local matrices
+	 *	Not safe to hold this reference between frames, because it exists in dynamic data passed from main thread.
+	 */
+	virtual const TArray<FMatrix>& GetReferenceToLocalMatrices() const = 0;
+
 	/** Get the LOD to render this mesh at. */
 	virtual int32 GetLOD() const = 0;
 	/** 
@@ -156,8 +162,6 @@ public:
 	 */
 	void InitLODInfos(const USkinnedMeshComponent* InMeshComponent);
 
-	void UpdateShadowShapes(USkinnedMeshComponent* InMeshComponent);
-
 	FORCEINLINE TStatId GetStatId() const 
 	{ 
 		return StatId; 
diff --git a/Engine/Source/Runtime/GameplayAbilities/Private/Abilities/GameplayAbility.cpp b/Engine/Source/Runtime/GameplayAbilities/Private/Abilities/GameplayAbility.cpp
--- a/Engine/Source/Runtime/GameplayAbilities/Private/Abilities/GameplayAbility.cpp
+++ b/Engine/Source/Runtime/GameplayAbilities/Private/Abilities/GameplayAbility.cpp
@@ -1487,7 +1487,7 @@ TArray<FActiveGameplayEffectHandle> UGameplayAbility::K2_ApplyGameplayEffectToTa
 	return BP_ApplyGameplayEffectToTarget(Target, GameplayEffect->GetClass(), GameplayEffectLevel, Stacks);
 }
 
-TArray<FActiveGameplayEffectHandle> UGameplayAbility::ApplyGameplayEffectToTarget(const FGameplayAbilitySpecHandle Handle, const FGameplayAbilityActorInfo* ActorInfo, const FGameplayAbilityActivationInfo ActivationInfo, FGameplayAbilityTargetDataHandle Target, TSubclassOf<UGameplayEffect> GameplayEffectClass, float GameplayEffectLevel, int32 Stacks) const
+TArray<FActiveGameplayEffectHandle> UGameplayAbility::ApplyGameplayEffectToTarget(const FGameplayAbilitySpecHandle Handle, const FGameplayAbilityActorInfo* ActorInfo, const FGameplayAbilityActivationInfo ActivationInfo, const FGameplayAbilityTargetDataHandle& Target, TSubclassOf<UGameplayEffect> GameplayEffectClass, float GameplayEffectLevel, int32 Stacks) const
 {
 	SCOPE_CYCLE_COUNTER(STAT_ApplyGameplayEffectToTarget);
 
@@ -1521,7 +1521,7 @@ TArray<FActiveGameplayEffectHandle> UGameplayAbility::K2_ApplyGameplayEffectSpec
 	return ApplyGameplayEffectSpecToTarget(CurrentSpecHandle, CurrentActorInfo, CurrentActivationInfo, SpecHandle, TargetData);
 }
 
-TArray<FActiveGameplayEffectHandle> UGameplayAbility::ApplyGameplayEffectSpecToTarget(const FGameplayAbilitySpecHandle AbilityHandle, const FGameplayAbilityActorInfo* ActorInfo, const FGameplayAbilityActivationInfo ActivationInfo, const FGameplayEffectSpecHandle SpecHandle, FGameplayAbilityTargetDataHandle TargetData) const
+TArray<FActiveGameplayEffectHandle> UGameplayAbility::ApplyGameplayEffectSpecToTarget(const FGameplayAbilitySpecHandle AbilityHandle, const FGameplayAbilityActorInfo* ActorInfo, const FGameplayAbilityActivationInfo ActivationInfo, const FGameplayEffectSpecHandle SpecHandle, const FGameplayAbilityTargetDataHandle& TargetData) const
 {
 	TArray<FActiveGameplayEffectHandle> EffectHandles;
 	
diff --git a/Engine/Source/Runtime/GameplayAbilities/Private/Abilities/GameplayAbilityTargetActor_Radius.cpp b/Engine/Source/Runtime/GameplayAbilities/Private/Abilities/GameplayAbilityTargetActor_Radius.cpp
--- a/Engine/Source/Runtime/GameplayAbilities/Private/Abilities/GameplayAbilityTargetActor_Radius.cpp
+++ b/Engine/Source/Runtime/GameplayAbilities/Private/Abilities/GameplayAbilityTargetActor_Radius.cpp
@@ -38,7 +38,7 @@ void AGameplayAbilityTargetActor_Radius::ConfirmTargetingAndContinue()
 	}
 }
 
-FGameplayAbilityTargetDataHandle AGameplayAbilityTargetActor_Radius::MakeTargetData(const TArray<TWeakObjectPtr<AActor>> Actors, const FVector& Origin) const
+FGameplayAbilityTargetDataHandle AGameplayAbilityTargetActor_Radius::MakeTargetData(const TArray<TWeakObjectPtr<AActor>>& Actors, const FVector& Origin) const
 {
 	if (OwningAbility)
 	{
diff --git a/Engine/Source/Runtime/GameplayAbilities/Private/Abilities/GameplayAbilityTargetActor_Trace.cpp b/Engine/Source/Runtime/GameplayAbilities/Private/Abilities/GameplayAbilityTargetActor_Trace.cpp
--- a/Engine/Source/Runtime/GameplayAbilities/Private/Abilities/GameplayAbilityTargetActor_Trace.cpp
+++ b/Engine/Source/Runtime/GameplayAbilities/Private/Abilities/GameplayAbilityTargetActor_Trace.cpp
@@ -203,7 +203,7 @@ void AGameplayAbilityTargetActor_Trace::ConfirmTargetingAndContinue()
 	}
 }
 
-FGameplayAbilityTargetDataHandle AGameplayAbilityTargetActor_Trace::MakeTargetData(FHitResult HitResult) const
+FGameplayAbilityTargetDataHandle AGameplayAbilityTargetActor_Trace::MakeTargetData(const FHitResult& HitResult) const
 {
 	/** Note: This will be cleaned up by the FGameplayAbilityTargetDataHandle (via an internal TSharedPtr) */
 	return StartLocation.MakeTargetDataHandleFromHitResult(OwningAbility, HitResult);
diff --git a/Engine/Source/Runtime/GameplayAbilities/Private/Abilities/Tasks/AbilityTask_SpawnActor.cpp b/Engine/Source/Runtime/GameplayAbilities/Private/Abilities/Tasks/AbilityTask_SpawnActor.cpp
--- a/Engine/Source/Runtime/GameplayAbilities/Private/Abilities/Tasks/AbilityTask_SpawnActor.cpp
+++ b/Engine/Source/Runtime/GameplayAbilities/Private/Abilities/Tasks/AbilityTask_SpawnActor.cpp
@@ -14,8 +14,8 @@ UAbilityTask_SpawnActor::UAbilityTask_SpawnActor(const FObjectInitializer& Objec
 
 UAbilityTask_SpawnActor* UAbilityTask_SpawnActor::SpawnActor(UObject* WorldContextObject, FGameplayAbilityTargetDataHandle TargetData, TSubclassOf<AActor> InClass)
 {
-	auto MyObj = NewAbilityTask<UAbilityTask_SpawnActor>(WorldContextObject);
-	MyObj->CachedTargetDataHandle = TargetData;
+	UAbilityTask_SpawnActor* MyObj = NewAbilityTask<UAbilityTask_SpawnActor>(WorldContextObject);
+	MyObj->CachedTargetDataHandle = MoveTemp(TargetData);
 	return MyObj;
 }
 
@@ -45,19 +45,25 @@ void UAbilityTask_SpawnActor::FinishSpawningActor(UObject* WorldContextObject, F
 {
 	if (SpawnedActor)
 	{
-		FTransform SpawnTransform = AbilitySystemComponent->GetOwner()->GetTransform();
-
+		bool bTransformSet = false;
+		FTransform SpawnTransform;
 		if (FGameplayAbilityTargetData* LocationData = CachedTargetDataHandle.Get(0))		//Hardcode to use data 0. It's OK if data isn't useful/valid.
 		{
 			//Set location. Rotation is unaffected.
 			if (LocationData->HasHitResult())
 			{
 				SpawnTransform.SetLocation(LocationData->GetHitResult()->Location);
+				bTransformSet = true;
 			}
 			else if (LocationData->HasEndPoint())
 			{
 				SpawnTransform = LocationData->GetEndPointTransform();
+				bTransformSet = true;
+			}
 			}
+		if (!bTransformSet)
+		{
+			SpawnTransform = AbilitySystemComponent->GetOwner()->GetTransform();
 		}
 
 		SpawnedActor->FinishSpawning(SpawnTransform);
diff --git a/Engine/Source/Runtime/GameplayAbilities/Private/Abilities/Tasks/AbilityTask_WaitTargetData.cpp b/Engine/Source/Runtime/GameplayAbilities/Private/Abilities/Tasks/AbilityTask_WaitTargetData.cpp
--- a/Engine/Source/Runtime/GameplayAbilities/Private/Abilities/Tasks/AbilityTask_WaitTargetData.cpp
+++ b/Engine/Source/Runtime/GameplayAbilities/Private/Abilities/Tasks/AbilityTask_WaitTargetData.cpp
@@ -13,17 +13,17 @@ UAbilityTask_WaitTargetData::UAbilityTask_WaitTargetData(const FObjectInitialize
 
 UAbilityTask_WaitTargetData* UAbilityTask_WaitTargetData::WaitTargetData(UObject* WorldContextObject, FName TaskInstanceName, TEnumAsByte<EGameplayTargetingConfirmation::Type> ConfirmationType, TSubclassOf<AGameplayAbilityTargetActor> InTargetClass)
 {
-	auto MyObj = NewAbilityTask<UAbilityTask_WaitTargetData>(WorldContextObject, TaskInstanceName);		//Register for task list here, providing a given FName as a key
+	UAbilityTask_WaitTargetData* MyObj = NewAbilityTask<UAbilityTask_WaitTargetData>(WorldContextObject, TaskInstanceName);		//Register for task list here, providing a given FName as a key
 	MyObj->TargetClass = InTargetClass;
-	MyObj->TargetActor = NULL;
+	MyObj->TargetActor = nullptr;
 	MyObj->ConfirmationType = ConfirmationType;
 	return MyObj;
 }
 
 UAbilityTask_WaitTargetData* UAbilityTask_WaitTargetData::WaitTargetDataUsingActor(UObject* WorldContextObject, FName TaskInstanceName, TEnumAsByte<EGameplayTargetingConfirmation::Type> ConfirmationType, AGameplayAbilityTargetActor* InTargetActor)
 {
-	auto MyObj = NewAbilityTask<UAbilityTask_WaitTargetData>(WorldContextObject, TaskInstanceName);		//Register for task list here, providing a given FName as a key
-	MyObj->TargetClass = NULL;
+	UAbilityTask_WaitTargetData* MyObj = NewAbilityTask<UAbilityTask_WaitTargetData>(WorldContextObject, TaskInstanceName);		//Register for task list here, providing a given FName as a key
+	MyObj->TargetClass = nullptr;
 	MyObj->TargetActor = InTargetActor;
 	MyObj->ConfirmationType = ConfirmationType;
 	return MyObj;
@@ -56,7 +56,7 @@ void UAbilityTask_WaitTargetData::Activate()
 			}
 			else
 			{
-				TargetActor = NULL;
+				TargetActor = nullptr;
 
 				// We may need a better solution here.  We don't know the target actor isn't needed till after it's already been spawned.
 				SpawnedActor->Destroy();
@@ -79,7 +79,7 @@ bool UAbilityTask_WaitTargetData::BeginSpawningActor(UObject* WorldContextObject
 		if (ShouldSpawnTargetActor())
 		{
 			UClass* Class = *InTargetClass;
-			if (Class != NULL)
+			if (Class != nullptr)
 			{
 				UWorld* World = GEngine->GetWorldFromContextObject(WorldContextObject);
 				SpawnedActor = World->SpawnActorDeferred<AGameplayAbilityTargetActor>(Class, FTransform::Identity, nullptr, nullptr, ESpawnActorCollisionHandlingMethod::AlwaysSpawn);
@@ -205,7 +205,7 @@ void UAbilityTask_WaitTargetData::RegisterTargetDataCallbacks()
 }
 
 /** Valid TargetData was replicated to use (we are server, was sent from client) */
-void UAbilityTask_WaitTargetData::OnTargetDataReplicatedCallback(FGameplayAbilityTargetDataHandle Data, FGameplayTag ActivationTag)
+void UAbilityTask_WaitTargetData::OnTargetDataReplicatedCallback(const FGameplayAbilityTargetDataHandle& Data, FGameplayTag ActivationTag)
 {
 	check(AbilitySystemComponent);
 
@@ -220,13 +220,14 @@ void UAbilityTask_WaitTargetData::OnTargetDataReplicatedCallback(FGameplayAbilit
 	 *	explicitly, the client is basically just sending a 'confirm' and the server is now going to do the work
 	 *	in OnReplicatedTargetDataReceived.
 	 */
-	if (TargetActor && !TargetActor->OnReplicatedTargetDataReceived(Data))
+	FGameplayAbilityTargetDataHandle MutableData = Data;
+	if (TargetActor && !TargetActor->OnReplicatedTargetDataReceived(MutableData))
 	{
-		Cancelled.Broadcast(Data);
+		Cancelled.Broadcast(MutableData);
 	}
 	else
 	{
-		ValidData.Broadcast(Data);
+		ValidData.Broadcast(MutableData);
 	}
 
 	if (ConfirmationType != EGameplayTargetingConfirmation::CustomMulti)
@@ -244,7 +245,7 @@ void UAbilityTask_WaitTargetData::OnTargetDataReplicatedCancelledCallback()
 }
 
 /** The TargetActor we spawned locally has called back with valid target data */
-void UAbilityTask_WaitTargetData::OnTargetDataReadyCallback(FGameplayAbilityTargetDataHandle Data)
+void UAbilityTask_WaitTargetData::OnTargetDataReadyCallback(const FGameplayAbilityTargetDataHandle& Data)
 {
 	check(AbilitySystemComponent);
 	if (!Ability)
@@ -278,7 +279,7 @@ void UAbilityTask_WaitTargetData::OnTargetDataReadyCallback(FGameplayAbilityTarg
 }
 
 /** The TargetActor we spawned locally has called back with a cancel event (they still include the 'last/best' targetdata but the consumer of this may want to discard it) */
-void UAbilityTask_WaitTargetData::OnTargetDataCancelledCallback(FGameplayAbilityTargetDataHandle Data)
+void UAbilityTask_WaitTargetData::OnTargetDataCancelledCallback(const FGameplayAbilityTargetDataHandle& Data)
 {
 	check(AbilitySystemComponent);
 
diff --git a/Engine/Source/Runtime/GameplayAbilities/Private/AbilitySystemBlueprintLibrary.cpp b/Engine/Source/Runtime/GameplayAbilities/Private/AbilitySystemBlueprintLibrary.cpp
--- a/Engine/Source/Runtime/GameplayAbilities/Private/AbilitySystemBlueprintLibrary.cpp
+++ b/Engine/Source/Runtime/GameplayAbilities/Private/AbilitySystemBlueprintLibrary.cpp
@@ -59,9 +59,9 @@ float UAbilitySystemBlueprintLibrary::GetFloatAttributeFromAbilitySystemComponen
 	return Result;
 }
 
-FGameplayAbilityTargetDataHandle UAbilitySystemBlueprintLibrary::AppendTargetDataHandle(FGameplayAbilityTargetDataHandle TargetHandle, FGameplayAbilityTargetDataHandle HandleToAdd)
+FGameplayAbilityTargetDataHandle UAbilitySystemBlueprintLibrary::AppendTargetDataHandle(FGameplayAbilityTargetDataHandle TargetHandle, const FGameplayAbilityTargetDataHandle& HandleToAdd)
 {
-	TargetHandle.Append(&HandleToAdd);
+	TargetHandle.Append(HandleToAdd);
 	return TargetHandle;
 }
 
@@ -86,18 +86,18 @@ FGameplayAbilityTargetDataHandle UAbilitySystemBlueprintLibrary::AbilityTargetDa
 	FGameplayAbilityTargetDataHandle		Handle(NewData);
 	return Handle;
 }
-FGameplayAbilityTargetDataHandle UAbilitySystemBlueprintLibrary::AbilityTargetDataFromActorArray(TArray<AActor*> ActorArray, bool OneTargetPerHandle)
+FGameplayAbilityTargetDataHandle UAbilitySystemBlueprintLibrary::AbilityTargetDataFromActorArray(const TArray<AActor*>& ActorArray, bool OneTargetPerHandle)
 {
 	// Construct TargetData
 	if (OneTargetPerHandle)
 	{
-		FGameplayAbilityTargetDataHandle		Handle;
+		FGameplayAbilityTargetDataHandle Handle;
 		for (int32 i = 0; i < ActorArray.Num(); ++i)
 		{
 			if (IsValid(ActorArray[i]))
 			{
 				FGameplayAbilityTargetDataHandle TempHandle = AbilityTargetDataFromActor(ActorArray[i]);
-				Handle.Append(&TempHandle);
+				Handle.Append(TempHandle);
 			}
 		}
 		return Handle;
@@ -115,13 +115,13 @@ FGameplayAbilityTargetDataHandle UAbilitySystemBlueprintLibrary::AbilityTargetDa
 	}
 }
 
-FGameplayAbilityTargetDataHandle UAbilitySystemBlueprintLibrary::FilterTargetData(FGameplayAbilityTargetDataHandle TargetDataHandle, FGameplayTargetDataFilterHandle FilterHandle)
+FGameplayAbilityTargetDataHandle UAbilitySystemBlueprintLibrary::FilterTargetData(const FGameplayAbilityTargetDataHandle& TargetDataHandle, FGameplayTargetDataFilterHandle FilterHandle)
 {
 	FGameplayAbilityTargetDataHandle ReturnDataHandle;
 	
 	for (int32 i = 0; TargetDataHandle.IsValid(i); ++i)
 	{
-		FGameplayAbilityTargetData* UnfilteredData = TargetDataHandle.Get(i);
+		const FGameplayAbilityTargetData* UnfilteredData = TargetDataHandle.Get(i);
 		check(UnfilteredData);
 		if (UnfilteredData->GetActors().Num() > 0)
 		{
@@ -129,7 +129,7 @@ FGameplayAbilityTargetDataHandle UAbilitySystemBlueprintLibrary::FilterTargetDat
 			if (FilteredActors.Num() > 0)
 			{
 				//Copy the data first, since we don't understand the internals of it
-				UScriptStruct* ScriptStruct = UnfilteredData->GetScriptStruct();
+				const UScriptStruct* ScriptStruct = UnfilteredData->GetScriptStruct();
 				FGameplayAbilityTargetData* NewData = (FGameplayAbilityTargetData*)FMemory::Malloc(ScriptStruct->GetCppStructOps()->GetSize());
 				ScriptStruct->InitializeStruct(NewData);
 				ScriptStruct->CopyScriptStruct(NewData, UnfilteredData);
@@ -165,7 +165,7 @@ FGameplayEffectSpecHandle UAbilitySystemBlueprintLibrary::MakeSpecHandle(UGamepl
 	return FGameplayEffectSpecHandle(new FGameplayEffectSpec(InGameplayEffect, FGameplayEffectContextHandle(EffectContext), InLevel));
 }
 
-FGameplayAbilityTargetDataHandle UAbilitySystemBlueprintLibrary::AbilityTargetDataFromHitResult(FHitResult HitResult)
+FGameplayAbilityTargetDataHandle UAbilitySystemBlueprintLibrary::AbilityTargetDataFromHitResult(const FHitResult& HitResult)
 {
 	// Construct TargetData
 	FGameplayAbilityTargetData_SingleTargetHit* TargetData = new FGameplayAbilityTargetData_SingleTargetHit(HitResult);
@@ -177,12 +177,12 @@ FGameplayAbilityTargetDataHandle UAbilitySystemBlueprintLibrary::AbilityTargetDa
 	return Handle;
 }
 
-int32 UAbilitySystemBlueprintLibrary::GetDataCountFromTargetData(FGameplayAbilityTargetDataHandle TargetData)
+int32 UAbilitySystemBlueprintLibrary::GetDataCountFromTargetData(const FGameplayAbilityTargetDataHandle& TargetData)
 {
 	return TargetData.Data.Num();
 }
 
-TArray<AActor*> UAbilitySystemBlueprintLibrary::GetActorsFromTargetData(FGameplayAbilityTargetDataHandle TargetData, int32 Index)
+TArray<AActor*> UAbilitySystemBlueprintLibrary::GetActorsFromTargetData(const FGameplayAbilityTargetDataHandle& TargetData, int32 Index)
 {
 	if (TargetData.Data.IsValidIndex(Index))
 	{
@@ -201,7 +201,7 @@ TArray<AActor*> UAbilitySystemBlueprintLibrary::GetActorsFromTargetData(FGamepla
 	return TArray<AActor*>();
 }
 
-bool UAbilitySystemBlueprintLibrary::DoesTargetDataContainActor(FGameplayAbilityTargetDataHandle TargetData, int32 Index, AActor* Actor)
+bool UAbilitySystemBlueprintLibrary::DoesTargetDataContainActor(const FGameplayAbilityTargetDataHandle& TargetData, int32 Index, AActor* Actor)
 {
 	if (TargetData.Data.IsValidIndex(Index))
 	{
@@ -220,7 +220,7 @@ bool UAbilitySystemBlueprintLibrary::DoesTargetDataContainActor(FGameplayAbility
 	return false;
 }
 
-bool UAbilitySystemBlueprintLibrary::TargetDataHasActor(FGameplayAbilityTargetDataHandle TargetData, int32 Index)
+bool UAbilitySystemBlueprintLibrary::TargetDataHasActor(const FGameplayAbilityTargetDataHandle& TargetData, int32 Index)
 {
 	if (TargetData.Data.IsValidIndex(Index))
 	{
@@ -233,7 +233,7 @@ bool UAbilitySystemBlueprintLibrary::TargetDataHasActor(FGameplayAbilityTargetDa
 	return false;
 }
 
-bool UAbilitySystemBlueprintLibrary::TargetDataHasHitResult(FGameplayAbilityTargetDataHandle TargetData, int32 Index)
+bool UAbilitySystemBlueprintLibrary::TargetDataHasHitResult(const FGameplayAbilityTargetDataHandle& TargetData, int32 Index)
 {
 	if (TargetData.Data.IsValidIndex(Index))
 	{
@@ -246,7 +246,7 @@ bool UAbilitySystemBlueprintLibrary::TargetDataHasHitResult(FGameplayAbilityTarg
 	return false;
 }
 
-FHitResult UAbilitySystemBlueprintLibrary::GetHitResultFromTargetData(FGameplayAbilityTargetDataHandle TargetData, int32 Index)
+FHitResult UAbilitySystemBlueprintLibrary::GetHitResultFromTargetData(const FGameplayAbilityTargetDataHandle& TargetData, int32 Index)
 {
 	if (TargetData.Data.IsValidIndex(Index))
 	{
@@ -264,7 +264,7 @@ FHitResult UAbilitySystemBlueprintLibrary::GetHitResultFromTargetData(FGameplayA
 	return FHitResult();
 }
 
-bool UAbilitySystemBlueprintLibrary::TargetDataHasOrigin(FGameplayAbilityTargetDataHandle TargetData, int32 Index)
+bool UAbilitySystemBlueprintLibrary::TargetDataHasOrigin(const FGameplayAbilityTargetDataHandle& TargetData, int32 Index)
 {
 	if (TargetData.Data.IsValidIndex(Index) == false)
 	{
@@ -279,7 +279,7 @@ bool UAbilitySystemBlueprintLibrary::TargetDataHasOrigin(FGameplayAbilityTargetD
 	return false;
 }
 
-FTransform UAbilitySystemBlueprintLibrary::GetTargetDataOrigin(FGameplayAbilityTargetDataHandle TargetData, int32 Index)
+FTransform UAbilitySystemBlueprintLibrary::GetTargetDataOrigin(const FGameplayAbilityTargetDataHandle& TargetData, int32 Index)
 {
 	if (TargetData.Data.IsValidIndex(Index) == false)
 	{
@@ -306,7 +306,7 @@ FTransform UAbilitySystemBlueprintLibrary::GetTargetDataOrigin(FGameplayAbilityT
 	return FTransform::Identity;
 }
 
-bool UAbilitySystemBlueprintLibrary::TargetDataHasEndPoint(FGameplayAbilityTargetDataHandle TargetData, int32 Index)
+bool UAbilitySystemBlueprintLibrary::TargetDataHasEndPoint(const FGameplayAbilityTargetDataHandle& TargetData, int32 Index)
 {
 	if (TargetData.Data.IsValidIndex(Index))
 	{
@@ -319,7 +319,7 @@ bool UAbilitySystemBlueprintLibrary::TargetDataHasEndPoint(FGameplayAbilityTarge
 	return false;
 }
 
-FVector UAbilitySystemBlueprintLibrary::GetTargetDataEndPoint(FGameplayAbilityTargetDataHandle TargetData, int32 Index)
+FVector UAbilitySystemBlueprintLibrary::GetTargetDataEndPoint(const FGameplayAbilityTargetDataHandle& TargetData, int32 Index)
 {
 	if (TargetData.Data.IsValidIndex(Index))
 	{
@@ -341,7 +341,7 @@ FVector UAbilitySystemBlueprintLibrary::GetTargetDataEndPoint(FGameplayAbilityTa
 	return FVector::ZeroVector;
 }
 
-FTransform UAbilitySystemBlueprintLibrary::GetTargetDataEndPointTransform(FGameplayAbilityTargetDataHandle TargetData, int32 Index)
+FTransform UAbilitySystemBlueprintLibrary::GetTargetDataEndPointTransform(const FGameplayAbilityTargetDataHandle& TargetData, int32 Index)
 {
 	if (TargetData.Data.IsValidIndex(Index))
 	{
diff --git a/Engine/Source/Runtime/GameplayAbilities/Private/AbilitySystemComponent_Abilities.cpp b/Engine/Source/Runtime/GameplayAbilities/Private/AbilitySystemComponent_Abilities.cpp
--- a/Engine/Source/Runtime/GameplayAbilities/Private/AbilitySystemComponent_Abilities.cpp
+++ b/Engine/Source/Runtime/GameplayAbilities/Private/AbilitySystemComponent_Abilities.cpp
@@ -172,7 +172,7 @@ void UAbilitySystemComponent::RefreshAbilityActorInfo()
 	AbilityActorInfo->InitFromActor(AbilityActorInfo->OwnerActor.Get(), AbilityActorInfo->AvatarActor.Get(), this);
 }
 
-FGameplayAbilitySpecHandle UAbilitySystemComponent::GiveAbility(FGameplayAbilitySpec Spec)
+FGameplayAbilitySpecHandle UAbilitySystemComponent::GiveAbility(const FGameplayAbilitySpec& Spec)
 {	
 	check(Spec.Ability);
 	check(IsOwnerActorAuthoritative());	// Should be called on authority
@@ -198,7 +198,7 @@ FGameplayAbilitySpecHandle UAbilitySystemComponent::GiveAbility(FGameplayAbility
 	return OwnedSpec.Handle;
 }
 
-FGameplayAbilitySpecHandle UAbilitySystemComponent::GiveAbilityAndActivateOnce(FGameplayAbilitySpec Spec)
+FGameplayAbilitySpecHandle UAbilitySystemComponent::GiveAbilityAndActivateOnce(const FGameplayAbilitySpec& Spec)
 {
 	check(Spec.Ability);
 
@@ -1734,7 +1734,7 @@ void UAbilitySystemComponent::NotifyAbilityActivated(FGameplayAbilitySpecHandle
 	AbilityActivatedCallbacks.Broadcast(Ability);
 }
 
-void UAbilitySystemComponent::NotifyAbilityFailed(const FGameplayAbilitySpecHandle Handle, UGameplayAbility* Ability, FGameplayTagContainer FailureReason)
+void UAbilitySystemComponent::NotifyAbilityFailed(const FGameplayAbilitySpecHandle Handle, UGameplayAbility* Ability, const FGameplayTagContainer& FailureReason)
 {
 	AbilityFailedCallbacks.Broadcast(Ability, FailureReason);
 }
@@ -2310,7 +2310,7 @@ void UAbilitySystemComponent::OnRep_ReplicatedAnimMontage()
 			{
 				int32 NextSectionID = AnimInstance->Montage_GetNextSectionID(LocalAnimMontageInfo.AnimMontage, RepSectionID);
 
-				// If NextSectionID is different thant the replicated one, then set it.
+				// If NextSectionID is different than the replicated one, then set it.
 				if( NextSectionID != RepNextSectionID )
 				{
 					AnimInstance->Montage_SetNextSection(LocalAnimMontageInfo.AnimMontage->GetSectionName(RepSectionID), LocalAnimMontageInfo.AnimMontage->GetSectionName(RepNextSectionID), LocalAnimMontageInfo.AnimMontage);
@@ -2318,10 +2318,11 @@ void UAbilitySystemComponent::OnRep_ReplicatedAnimMontage()
 
 				// Make sure we haven't received that update too late and the client hasn't already jumped to another section. 
 				int32 CurrentSectionID = LocalAnimMontageInfo.AnimMontage->GetSectionIndexFromPosition(AnimInstance->Montage_GetPosition(LocalAnimMontageInfo.AnimMontage));
-				if( (CurrentSectionID != RepSectionID) && (CurrentSectionID != RepNextSectionID) )
+				if ((CurrentSectionID != RepSectionID) && (CurrentSectionID != RepNextSectionID))
 				{
-					// Client is in a wrong section, jump to replicated position.
-					AnimInstance->Montage_SetPosition(LocalAnimMontageInfo.AnimMontage, RepAnimMontageInfo.Position);
+					// Client is in a wrong section, teleport him into the begining of the right section
+					const float SectionStartTime = LocalAnimMontageInfo.AnimMontage->GetAnimCompositeSection(RepSectionID).GetTime();
+					AnimInstance->Montage_SetPosition(LocalAnimMontageInfo.AnimMontage, SectionStartTime);
 				}
 			}
 
@@ -2331,6 +2332,12 @@ void UAbilitySystemComponent::OnRep_ReplicatedAnimMontage()
 			// Only check threshold if we are located in the same section. Different sections require a bit more work as we could be jumping around the timeline.
 			if ((CurrentSectionID == RepSectionID) && (FMath::Abs(CurrentPosition - RepAnimMontageInfo.Position) > MONTAGE_REP_POS_ERR_THRESH) && RepAnimMontageInfo.IsStopped == 0)
 			{
+				// fast forward to server position and trigger notifies
+				if (FAnimMontageInstance* MontageInstance = AnimInstance->GetActiveInstanceForMontage(*RepAnimMontageInfo.AnimMontage))
+				{
+					MontageInstance->HandleEvents(CurrentPosition, RepAnimMontageInfo.Position, nullptr);
+					AnimInstance->TriggerAnimNotifies(0.f);
+				}
 				AnimInstance->Montage_SetPosition(LocalAnimMontageInfo.AnimMontage, RepAnimMontageInfo.Position);
 			}
 
@@ -2710,7 +2717,7 @@ void UAbilitySystemComponent::ClientSetReplicatedEvent_Implementation(EAbilityGe
 
 // -------
 
-void UAbilitySystemComponent::ServerSetReplicatedTargetData_Implementation(FGameplayAbilitySpecHandle AbilityHandle, FPredictionKey AbilityOriginalPredictionKey, FGameplayAbilityTargetDataHandle ReplicatedTargetDataHandle, FGameplayTag ApplicationTag, FPredictionKey CurrentPredictionKey)
+void UAbilitySystemComponent::ServerSetReplicatedTargetData_Implementation(FGameplayAbilitySpecHandle AbilityHandle, FPredictionKey AbilityOriginalPredictionKey, const FGameplayAbilityTargetDataHandle& ReplicatedTargetDataHandle, FGameplayTag ApplicationTag, FPredictionKey CurrentPredictionKey)
 {
 	FScopedPredictionWindow ScopedPrediction(this, CurrentPredictionKey);
 
@@ -2733,7 +2740,7 @@ void UAbilitySystemComponent::ServerSetReplicatedTargetData_Implementation(FGame
 	ReplicatedData.TargetSetDelegate.Broadcast(ReplicatedData.TargetData, ReplicatedData.ApplicationTag);
 }
 
-bool UAbilitySystemComponent::ServerSetReplicatedTargetData_Validate(FGameplayAbilitySpecHandle AbilityHandle, FPredictionKey AbilityOriginalPredictionKey, FGameplayAbilityTargetDataHandle ReplicatedTargetDataHandle, FGameplayTag ApplicationTag, FPredictionKey CurrentPredictionKey)
+bool UAbilitySystemComponent::ServerSetReplicatedTargetData_Validate(FGameplayAbilitySpecHandle AbilityHandle, FPredictionKey AbilityOriginalPredictionKey, const FGameplayAbilityTargetDataHandle& ReplicatedTargetDataHandle, FGameplayTag ApplicationTag, FPredictionKey CurrentPredictionKey)
 {
 	return true;
 }
diff --git a/Engine/Source/Runtime/GameplayAbilities/Private/GameplayAbilityTargetTypes.cpp b/Engine/Source/Runtime/GameplayAbilities/Private/GameplayAbilityTargetTypes.cpp
--- a/Engine/Source/Runtime/GameplayAbilities/Private/GameplayAbilityTargetTypes.cpp
+++ b/Engine/Source/Runtime/GameplayAbilities/Private/GameplayAbilityTargetTypes.cpp
@@ -48,7 +48,7 @@ TArray<FActiveGameplayEffectHandle> FGameplayAbilityTargetData::ApplyGameplayEff
 	return AppliedHandles;
 }
 
-void FGameplayAbilityTargetData::AddTargetDataToContext(FGameplayEffectContextHandle& Context, bool bIncludeActorArray)
+void FGameplayAbilityTargetData::AddTargetDataToContext(FGameplayEffectContextHandle& Context, bool bIncludeActorArray) const
 {
 	if (bIncludeActorArray && (GetActors().Num() > 0))
 	{
@@ -66,10 +66,8 @@ void FGameplayAbilityTargetData::AddTargetDataToContext(FGameplayEffectContextHa
 	}
 }
 
-void FGameplayAbilityTargetData::AddTargetDataToGameplayCueParameters(FGameplayCueParameters& Parameters)
+void FGameplayAbilityTargetData::AddTargetDataToGameplayCueParameters(FGameplayCueParameters& Parameters) const
 {
-	
-	
 }
 
 FString FGameplayAbilityTargetData::ToString() const
@@ -77,7 +75,7 @@ FString FGameplayAbilityTargetData::ToString() const
 	return TEXT("BASE CLASS");
 }
 
-FGameplayAbilityTargetDataHandle FGameplayAbilityTargetingLocationInfo::MakeTargetDataHandleFromHitResult(TWeakObjectPtr<UGameplayAbility> Ability, FHitResult HitResult) const
+FGameplayAbilityTargetDataHandle FGameplayAbilityTargetingLocationInfo::MakeTargetDataHandleFromHitResult(TWeakObjectPtr<UGameplayAbility> Ability, const FHitResult& HitResult) const
 {
 	TArray<FHitResult> HitResults;
 	HitResults.Add(HitResult);
@@ -99,7 +97,7 @@ FGameplayAbilityTargetDataHandle FGameplayAbilityTargetingLocationInfo::MakeTarg
 	return ReturnDataHandle;
 }
 
-FGameplayAbilityTargetDataHandle FGameplayAbilityTargetingLocationInfo::MakeTargetDataHandleFromActors(TArray<TWeakObjectPtr<AActor> > TargetActors, bool OneActorPerHandle) const
+FGameplayAbilityTargetDataHandle FGameplayAbilityTargetingLocationInfo::MakeTargetDataHandleFromActors(const TArray<TWeakObjectPtr<AActor> >& TargetActors, bool OneActorPerHandle) const
 {
 	/** Note: This is cleaned up by the FGameplayAbilityTargetDataHandle (via an internal TSharedPtr) */
 	FGameplayAbilityTargetData_ActorArray* ReturnData = new FGameplayAbilityTargetData_ActorArray();
@@ -109,18 +107,18 @@ FGameplayAbilityTargetDataHandle FGameplayAbilityTargetingLocationInfo::MakeTarg
 	{
 		if (TargetActors.Num() > 0)
 		{
-			if (TargetActors[0].IsValid())
+			if (AActor* TargetActor = TargetActors[0].Get())
 			{
-				ReturnData->TargetActorArray.Add(TargetActors[0].Get());
+				ReturnData->TargetActorArray.Add(TargetActor);
 			}
 
 			for (int32 i = 1; i < TargetActors.Num(); ++i)
 			{
-				if (TargetActors[i].IsValid())
+				if (AActor* TargetActor = TargetActors[i].Get())
 				{
 					FGameplayAbilityTargetData_ActorArray* CurrentData = new FGameplayAbilityTargetData_ActorArray();
 					CurrentData->SourceLocation = *this;
-					CurrentData->TargetActorArray.Add(TargetActors[i].Get());
+					CurrentData->TargetActorArray.Add(TargetActor);
 					ReturnDataHandle.Add(CurrentData);
 				}
 			}
diff --git a/Engine/Source/Runtime/GameplayAbilities/Private/GameplayCueNotify_Actor.cpp b/Engine/Source/Runtime/GameplayAbilities/Private/GameplayCueNotify_Actor.cpp
--- a/Engine/Source/Runtime/GameplayAbilities/Private/GameplayCueNotify_Actor.cpp
+++ b/Engine/Source/Runtime/GameplayAbilities/Private/GameplayCueNotify_Actor.cpp
@@ -82,7 +82,7 @@ bool AGameplayCueNotify_Actor::HandlesEvent(EGameplayCueEvent::Type EventType) c
 int32 GameplayCueNotifyActorStacking = 0;
 static FAutoConsoleVariableRef CVarGameplayCueNotifyActorStacking(TEXT("AbilitySystem.GameplayCueNotifyActorStacking"), GameplayCueNotifyActorStacking, TEXT("Enable simple stacking rules for gameplaycue actors"), ECVF_Default );
 
-void AGameplayCueNotify_Actor::HandleGameplayCue(AActor* MyTarget, EGameplayCueEvent::Type EventType, FGameplayCueParameters Parameters)
+void AGameplayCueNotify_Actor::HandleGameplayCue(AActor* MyTarget, EGameplayCueEvent::Type EventType, const FGameplayCueParameters& Parameters)
 {
 	SCOPE_CYCLE_COUNTER(STAT_HandleGameplayCueNotifyActor);
 
@@ -161,22 +161,22 @@ void AGameplayCueNotify_Actor::OnOwnerDestroyed()
 	GameplayCueFinishedCallback();
 }
 
-bool AGameplayCueNotify_Actor::OnExecute_Implementation(AActor* MyTarget, FGameplayCueParameters Parameters)
+bool AGameplayCueNotify_Actor::OnExecute_Implementation(AActor* MyTarget, const FGameplayCueParameters& Parameters)
 {
 	return false;
 }
 
-bool AGameplayCueNotify_Actor::OnActive_Implementation(AActor* MyTarget, FGameplayCueParameters Parameters)
+bool AGameplayCueNotify_Actor::OnActive_Implementation(AActor* MyTarget, const FGameplayCueParameters& Parameters)
 {
 	return false;
 }
 
-bool AGameplayCueNotify_Actor::WhileActive_Implementation(AActor* MyTarget, FGameplayCueParameters Parameters)
+bool AGameplayCueNotify_Actor::WhileActive_Implementation(AActor* MyTarget, const FGameplayCueParameters& Parameters)
 {
 	return false;
 }
 
-bool AGameplayCueNotify_Actor::OnRemove_Implementation(AActor* MyTarget, FGameplayCueParameters Parameters)
+bool AGameplayCueNotify_Actor::OnRemove_Implementation(AActor* MyTarget, const FGameplayCueParameters& Parameters)
 {
 	return false;
 }
diff --git a/Engine/Source/Runtime/GameplayAbilities/Private/GameplayCueNotify_HitImpact.cpp b/Engine/Source/Runtime/GameplayAbilities/Private/GameplayCueNotify_HitImpact.cpp
--- a/Engine/Source/Runtime/GameplayAbilities/Private/GameplayCueNotify_HitImpact.cpp
+++ b/Engine/Source/Runtime/GameplayAbilities/Private/GameplayCueNotify_HitImpact.cpp
@@ -14,7 +14,7 @@ bool UGameplayCueNotify_HitImpact::HandlesEvent(EGameplayCueEvent::Type EventTyp
 	return (EventType == EGameplayCueEvent::Executed);
 }
 
-void UGameplayCueNotify_HitImpact::HandleGameplayCue(AActor* Self, EGameplayCueEvent::Type EventType, FGameplayCueParameters Parameters)
+void UGameplayCueNotify_HitImpact::HandleGameplayCue(AActor* Self, EGameplayCueEvent::Type EventType, const FGameplayCueParameters& Parameters)
 {
 	check(EventType == EGameplayCueEvent::Executed);
 	check(Self);
diff --git a/Engine/Source/Runtime/GameplayAbilities/Private/GameplayCueNotify_Static.cpp b/Engine/Source/Runtime/GameplayAbilities/Private/GameplayCueNotify_Static.cpp
--- a/Engine/Source/Runtime/GameplayAbilities/Private/GameplayCueNotify_Static.cpp
+++ b/Engine/Source/Runtime/GameplayAbilities/Private/GameplayCueNotify_Static.cpp
@@ -58,7 +58,7 @@ bool UGameplayCueNotify_Static::HandlesEvent(EGameplayCueEvent::Type EventType)
 	return true;
 }
 
-void UGameplayCueNotify_Static::HandleGameplayCue(AActor* MyTarget, EGameplayCueEvent::Type EventType, FGameplayCueParameters Parameters)
+void UGameplayCueNotify_Static::HandleGameplayCue(AActor* MyTarget, EGameplayCueEvent::Type EventType, const FGameplayCueParameters& Parameters)
 {
 	SCOPE_CYCLE_COUNTER(STAT_HandleGameplayCueNotifyStatic);
 
@@ -95,22 +95,22 @@ void UGameplayCueNotify_Static::OnOwnerDestroyed()
 {
 }
 
-bool UGameplayCueNotify_Static::OnExecute_Implementation(AActor* MyTarget, FGameplayCueParameters Parameters) const
+bool UGameplayCueNotify_Static::OnExecute_Implementation(AActor* MyTarget, const FGameplayCueParameters& Parameters) const
 {
 	return false;
 }
 
-bool UGameplayCueNotify_Static::OnActive_Implementation(AActor* MyTarget, FGameplayCueParameters Parameters) const
+bool UGameplayCueNotify_Static::OnActive_Implementation(AActor* MyTarget, const FGameplayCueParameters& Parameters) const
 {
 	return false;
 }
 
-bool UGameplayCueNotify_Static::WhileActive_Implementation(AActor* MyTarget, FGameplayCueParameters Parameters) const
+bool UGameplayCueNotify_Static::WhileActive_Implementation(AActor* MyTarget, const FGameplayCueParameters& Parameters) const
 {
 	return false;
 }
 
-bool UGameplayCueNotify_Static::OnRemove_Implementation(AActor* MyTarget, FGameplayCueParameters Parameters) const
+bool UGameplayCueNotify_Static::OnRemove_Implementation(AActor* MyTarget, const FGameplayCueParameters& Parameters) const
 {
 	return false;
 }
diff --git a/Engine/Source/Runtime/GameplayAbilities/Private/GameplayCueSet.cpp b/Engine/Source/Runtime/GameplayAbilities/Private/GameplayCueSet.cpp
--- a/Engine/Source/Runtime/GameplayAbilities/Private/GameplayCueSet.cpp
+++ b/Engine/Source/Runtime/GameplayAbilities/Private/GameplayCueSet.cpp
@@ -19,7 +19,7 @@ UGameplayCueSet::UGameplayCueSet(const FObjectInitializer& ObjectInitializer)
 
 }
 
-bool UGameplayCueSet::HandleGameplayCue(AActor* TargetActor, FGameplayTag GameplayCueTag, EGameplayCueEvent::Type EventType, FGameplayCueParameters Parameters)
+bool UGameplayCueSet::HandleGameplayCue(AActor* TargetActor, FGameplayTag GameplayCueTag, EGameplayCueEvent::Type EventType, const FGameplayCueParameters& Parameters)
 {
 	// GameplayCueTags could have been removed from the dictionary but not content. When the content is resaved the old tag will be cleaned up, but it could still come through here
 	// at runtime. Since we only populate the map with dictionary gameplaycue tags, we may not find it here.
diff --git a/Engine/Source/Runtime/GameplayAbilities/Private/GameplayEffectTypes.cpp b/Engine/Source/Runtime/GameplayAbilities/Private/GameplayEffectTypes.cpp
--- a/Engine/Source/Runtime/GameplayAbilities/Private/GameplayEffectTypes.cpp
+++ b/Engine/Source/Runtime/GameplayAbilities/Private/GameplayEffectTypes.cpp
@@ -69,7 +69,7 @@ void FGameplayEffectContext::AddInstigator(class AActor *InInstigator, class AAc
 	}
 }
 
-void FGameplayEffectContext::AddActors(TArray<TWeakObjectPtr<AActor>> InActors, bool bReset)
+void FGameplayEffectContext::AddActors(const TArray<TWeakObjectPtr<AActor>>& InActors, bool bReset)
 {
 	if (bReset && Actors.Num())
 	{
@@ -79,7 +79,7 @@ void FGameplayEffectContext::AddActors(TArray<TWeakObjectPtr<AActor>> InActors,
 	Actors.Append(InActors);
 }
 
-void FGameplayEffectContext::AddHitResult(const FHitResult InHitResult, bool bReset)
+void FGameplayEffectContext::AddHitResult(const FHitResult& InHitResult, bool bReset)
 {
 	if (bReset && HitResult.IsValid())
 	{
diff --git a/Engine/Source/Runtime/GameplayAbilities/Public/Abilities/GameplayAbility.h b/Engine/Source/Runtime/GameplayAbilities/Public/Abilities/GameplayAbility.h
--- a/Engine/Source/Runtime/GameplayAbilities/Public/Abilities/GameplayAbility.h
+++ b/Engine/Source/Runtime/GameplayAbilities/Public/Abilities/GameplayAbility.h
@@ -475,18 +475,18 @@ protected:
 	TArray<FActiveGameplayEffectHandle> K2_ApplyGameplayEffectToTarget(FGameplayAbilityTargetDataHandle TargetData, const UGameplayEffect* GameplayEffect, int32 GameplayEffectLevel = 1, int32 Stacks = 1);
 
 	/** Non blueprintcallable, safe to call on CDO/NonInstance abilities */
-	TArray<FActiveGameplayEffectHandle> ApplyGameplayEffectToTarget(const FGameplayAbilitySpecHandle Handle, const FGameplayAbilityActorInfo* ActorInfo, const FGameplayAbilityActivationInfo ActivationInfo, FGameplayAbilityTargetDataHandle Target, TSubclassOf<UGameplayEffect> GameplayEffectClass, float GameplayEffectLevel, int32 Stacks = 1) const;
+	TArray<FActiveGameplayEffectHandle> ApplyGameplayEffectToTarget(const FGameplayAbilitySpecHandle Handle, const FGameplayAbilityActorInfo* ActorInfo, const FGameplayAbilityActivationInfo ActivationInfo, const FGameplayAbilityTargetDataHandle& Target, TSubclassOf<UGameplayEffect> GameplayEffectClass, float GameplayEffectLevel, int32 Stacks = 1) const;
 
 	UFUNCTION(BlueprintCallable, Category = Ability, DisplayName = "ApplyGameplayEffectSpecToTarget")
 	TArray<FActiveGameplayEffectHandle> K2_ApplyGameplayEffectSpecToTarget(const FGameplayEffectSpecHandle EffectSpecHandle, FGameplayAbilityTargetDataHandle TargetData);
 
-	TArray<FActiveGameplayEffectHandle> ApplyGameplayEffectSpecToTarget(const FGameplayAbilitySpecHandle AbilityHandle, const FGameplayAbilityActorInfo* ActorInfo, const FGameplayAbilityActivationInfo ActivationInfo, const FGameplayEffectSpecHandle SpecHandle, FGameplayAbilityTargetDataHandle TargetData) const;
+	TArray<FActiveGameplayEffectHandle> ApplyGameplayEffectSpecToTarget(const FGameplayAbilitySpecHandle AbilityHandle, const FGameplayAbilityActorInfo* ActorInfo, const FGameplayAbilityActivationInfo ActivationInfo, const FGameplayEffectSpecHandle SpecHandle, const FGameplayAbilityTargetDataHandle& TargetData) const;
 
 	// ---------
 	// Remove Self
 	// ---------
 	
-	/** Removes GAmeplayEffects from owner which match the given asset level tags. */
+	/** Removes GameplayEffects from owner which match the given asset level tags. */
 	UFUNCTION(BlueprintCallable, Category = Ability, DisplayName="RemoveGameplayEffectFromOwnerWithAssetTags")
 	void BP_RemoveGameplayEffectFromOwnerWithAssetTags(FGameplayTagContainer WithAssetTags, int32 StacksToRemove = -1);
 
diff --git a/Engine/Source/Runtime/GameplayAbilities/Public/Abilities/GameplayAbilityTargetActor_Radius.h b/Engine/Source/Runtime/GameplayAbilities/Public/Abilities/GameplayAbilityTargetActor_Radius.h
--- a/Engine/Source/Runtime/GameplayAbilities/Public/Abilities/GameplayAbilityTargetActor_Radius.h
+++ b/Engine/Source/Runtime/GameplayAbilities/Public/Abilities/GameplayAbilityTargetActor_Radius.h
@@ -27,6 +27,6 @@ protected:
 
 	TArray<TWeakObjectPtr<AActor> >	PerformOverlap(const FVector& Origin);
 
-	FGameplayAbilityTargetDataHandle MakeTargetData(const TArray<TWeakObjectPtr<AActor> > Actors, const FVector& Origin) const;
+	FGameplayAbilityTargetDataHandle MakeTargetData(const TArray<TWeakObjectPtr<AActor>>& Actors, const FVector& Origin) const;
 
 };
diff --git a/Engine/Source/Runtime/GameplayAbilities/Public/Abilities/GameplayAbilityTargetActor_Trace.h b/Engine/Source/Runtime/GameplayAbilities/Public/Abilities/GameplayAbilityTargetActor_Trace.h
--- a/Engine/Source/Runtime/GameplayAbilities/Public/Abilities/GameplayAbilityTargetActor_Trace.h
+++ b/Engine/Source/Runtime/GameplayAbilities/Public/Abilities/GameplayAbilityTargetActor_Trace.h
@@ -44,7 +44,7 @@ public:
 protected:
 	virtual FHitResult PerformTrace(AActor* InSourceActor) PURE_VIRTUAL(AGameplayAbilityTargetActor_Trace, return FHitResult(););
 
-	FGameplayAbilityTargetDataHandle MakeTargetData(FHitResult HitResult) const;
+	FGameplayAbilityTargetDataHandle MakeTargetData(const FHitResult& HitResult) const;
 
 	TWeakObjectPtr<AGameplayAbilityWorldReticle> ReticleActor;
 };
\ No newline at end of file
diff --git a/Engine/Source/Runtime/GameplayAbilities/Public/Abilities/GameplayAbilityTargetTypes.h b/Engine/Source/Runtime/GameplayAbilities/Public/Abilities/GameplayAbilityTargetTypes.h
--- a/Engine/Source/Runtime/GameplayAbilities/Public/Abilities/GameplayAbilityTargetTypes.h
+++ b/Engine/Source/Runtime/GameplayAbilities/Public/Abilities/GameplayAbilityTargetTypes.h
@@ -67,13 +67,13 @@ struct GAMEPLAYABILITIES_API FGameplayAbilityTargetData
 
 	virtual TArray<FActiveGameplayEffectHandle> ApplyGameplayEffectSpec(FGameplayEffectSpec& Spec, FPredictionKey PredictionKey = FPredictionKey());
 
-	virtual void AddTargetDataToContext(FGameplayEffectContextHandle& Context, bool bIncludeActorArray);
+	virtual void AddTargetDataToContext(FGameplayEffectContextHandle& Context, bool bIncludeActorArray) const;
 
-	virtual void AddTargetDataToGameplayCueParameters(FGameplayCueParameters& Parameters);
+	virtual void AddTargetDataToGameplayCueParameters(FGameplayCueParameters& Parameters) const;
 
-	virtual TArray<TWeakObjectPtr<AActor> >	GetActors() const
+	virtual TArray<TWeakObjectPtr<AActor>> GetActors() const
 	{
-		return TArray<TWeakObjectPtr<AActor> >();
+		return TArray<TWeakObjectPtr<AActor>>();
 	}
 
 	virtual bool SetActors(TArray<TWeakObjectPtr<AActor>> NewActorArray)
@@ -91,7 +91,7 @@ struct GAMEPLAYABILITIES_API FGameplayAbilityTargetData
 
 	virtual const FHitResult* GetHitResult() const
 	{
-		return NULL;
+		return nullptr;
 	}
 
 	// -------------------------------------
@@ -125,7 +125,7 @@ struct GAMEPLAYABILITIES_API FGameplayAbilityTargetData
 
 	// -------------------------------------
 
-	virtual UScriptStruct* GetScriptStruct()
+	virtual UScriptStruct* GetScriptStruct() const
 	{
 		return FGameplayAbilityTargetData::StaticStruct();
 	}
@@ -192,9 +192,14 @@ struct GAMEPLAYABILITIES_API FGameplayAbilityTargetDataHandle
 		return (Index < Data.Num() && Data[Index].IsValid());
 	}
 
+	const FGameplayAbilityTargetData* Get(int32 Index) const
+	{
+		return IsValid(Index) ? Data[Index].Get() : nullptr;
+	}
+
 	FGameplayAbilityTargetData* Get(int32 Index)
 	{
-		return IsValid(Index) ? Data[Index].Get() : NULL;
+		return IsValid(Index) ? Data[Index].Get() : nullptr;
 	}
 
 	void Add(struct FGameplayAbilityTargetData* DataPtr)
@@ -202,18 +207,21 @@ struct GAMEPLAYABILITIES_API FGameplayAbilityTargetDataHandle
 		Data.Add(TSharedPtr<FGameplayAbilityTargetData>(DataPtr));
 	}
 
-	void Append(struct FGameplayAbilityTargetDataHandle* OtherHandle)
+	DEPRECATED(4.11, "Pass Handle by reference, not pointer")
+	void Append(FGameplayAbilityTargetDataHandle* OtherHandle)
 	{
-		for (int32 i = 0; i < OtherHandle->Data.Num(); ++i)
-		{
-			Data.Add(OtherHandle->Data[i]);
-		}
+		Append(*OtherHandle);
+	}
+
+	void Append(const FGameplayAbilityTargetDataHandle& OtherHandle)
+	{
+		Data.Append(OtherHandle.Data);
 	}
 
 	bool NetSerialize(FArchive& Ar, class UPackageMap* Map, bool& bOutSuccess);
 
 	/** Comparison operator */
-	bool operator==(FGameplayAbilityTargetDataHandle const& Other) const
+	bool operator==(const FGameplayAbilityTargetDataHandle& Other) const
 	{
 		// Both invalid structs or both valid and Pointer compare (???) // deep comparison equality
 		if (Data.Num() != Other.Data.Num())
@@ -235,7 +243,7 @@ struct GAMEPLAYABILITIES_API FGameplayAbilityTargetDataHandle
 	}
 
 	/** Comparison operator */
-	bool operator!=(FGameplayAbilityTargetDataHandle const& Other) const
+	bool operator!=(const FGameplayAbilityTargetDataHandle& Other) const
 	{
 		return !(FGameplayAbilityTargetDataHandle::operator==(Other));
 	}
@@ -302,9 +310,9 @@ public:
 		return FTransform::Identity;
 	}
 
-	FGameplayAbilityTargetDataHandle MakeTargetDataHandleFromHitResult(TWeakObjectPtr<UGameplayAbility> Ability, FHitResult HitResult) const;
+	FGameplayAbilityTargetDataHandle MakeTargetDataHandleFromHitResult(TWeakObjectPtr<UGameplayAbility> Ability, const FHitResult& HitResult) const;
 	FGameplayAbilityTargetDataHandle MakeTargetDataHandleFromHitResults(TWeakObjectPtr<UGameplayAbility> Ability, const TArray<FHitResult>& HitResults) const;
-	FGameplayAbilityTargetDataHandle MakeTargetDataHandleFromActors(TArray<TWeakObjectPtr<AActor>> TargetActors, bool OneActorPerHandle = false) const;
+	FGameplayAbilityTargetDataHandle MakeTargetDataHandleFromActors(const TArray<TWeakObjectPtr<AActor>>& TargetActors, bool OneActorPerHandle = false) const;
 
 	/** Type of location used - will determine what data is transmitted over the network and what fields are used when calculating position. */
 	UPROPERTY(BlueprintReadWrite, meta = (ExposeOnSpawn = true), Category = Targeting)
@@ -339,7 +347,7 @@ public:
 
 	bool NetSerialize(FArchive& Ar, class UPackageMap* Map, bool& bOutSuccess);
 
-	virtual UScriptStruct* GetScriptStruct()
+	virtual UScriptStruct* GetScriptStruct() const
 	{
 		return FGameplayAbilityTargetingLocationInfo::StaticStruct();
 	}
@@ -393,7 +401,7 @@ struct GAMEPLAYABILITIES_API FGameplayAbilityTargetData_LocationInfo : public FG
 
 	// -------------------------------------
 
-	virtual UScriptStruct* GetScriptStruct() override
+	virtual UScriptStruct* GetScriptStruct() const override
 	{
 		return FGameplayAbilityTargetData_LocationInfo::StaticStruct();
 	}
@@ -495,7 +503,7 @@ struct GAMEPLAYABILITIES_API FGameplayAbilityTargetData_ActorArray : public FGam
 
 	// -------------------------------------
 
-	virtual UScriptStruct* GetScriptStruct() override
+	virtual UScriptStruct* GetScriptStruct() const override
 	{
 		return FGameplayAbilityTargetData_ActorArray::StaticStruct();
 	}
@@ -526,7 +534,7 @@ struct GAMEPLAYABILITIES_API FGameplayAbilityTargetData_SingleTargetHit : public
 	{ }
 
 	FGameplayAbilityTargetData_SingleTargetHit(const FHitResult InHitResult)
-		: HitResult(InHitResult)
+		: HitResult(MoveTemp(InHitResult))
 	{ }
 
 	// -------------------------------------
@@ -536,7 +544,7 @@ struct GAMEPLAYABILITIES_API FGameplayAbilityTargetData_SingleTargetHit : public
 		TArray<TWeakObjectPtr<AActor> >	Actors;
 		if (HitResult.Actor.IsValid())
 		{
-			Actors.Push(HitResult.Actor.Get());
+			Actors.Push(HitResult.Actor);
 		}
 		return Actors;
 	}
@@ -582,7 +590,7 @@ struct GAMEPLAYABILITIES_API FGameplayAbilityTargetData_SingleTargetHit : public
 
 	bool NetSerialize(FArchive& Ar, class UPackageMap* Map, bool& bOutSuccess);
 
-	virtual UScriptStruct* GetScriptStruct() override
+	virtual UScriptStruct* GetScriptStruct() const override
 	{
 		return FGameplayAbilityTargetData_SingleTargetHit::StaticStruct();
 	}
@@ -598,13 +606,13 @@ struct TStructOpsTypeTraits<FGameplayAbilityTargetData_SingleTargetHit> : public
 };
 
 /** Generic callback for returning when target data is available */
-DECLARE_MULTICAST_DELEGATE_OneParam(FAbilityTargetData, FGameplayAbilityTargetDataHandle);
+DECLARE_MULTICAST_DELEGATE_OneParam(FAbilityTargetData, const FGameplayAbilityTargetDataHandle&);
 
 
 // ----------------------------------------------------
 
 /** Generic callback for returning when target data is available */
-DECLARE_MULTICAST_DELEGATE_TwoParams(FAbilityTargetDataSetDelegate, FGameplayAbilityTargetDataHandle, FGameplayTag);
+DECLARE_MULTICAST_DELEGATE_TwoParams(FAbilityTargetDataSetDelegate, const FGameplayAbilityTargetDataHandle&, FGameplayTag);
 
 /** These are generic, nonpayload carrying events that are replicated between the client and server */
 UENUM()
diff --git a/Engine/Source/Runtime/GameplayAbilities/Public/Abilities/Tasks/AbilityTask_WaitOverlap.h b/Engine/Source/Runtime/GameplayAbilities/Public/Abilities/Tasks/AbilityTask_WaitOverlap.h
--- a/Engine/Source/Runtime/GameplayAbilities/Public/Abilities/Tasks/AbilityTask_WaitOverlap.h
+++ b/Engine/Source/Runtime/GameplayAbilities/Public/Abilities/Tasks/AbilityTask_WaitOverlap.h
@@ -4,7 +4,7 @@
 #include "Abilities/GameplayAbilityTypes.h"
 #include "AbilityTask_WaitOverlap.generated.h"
 
-DECLARE_DYNAMIC_MULTICAST_DELEGATE_OneParam(FWaitOverlapDelegate, FGameplayAbilityTargetDataHandle, TargetData);
+DECLARE_DYNAMIC_MULTICAST_DELEGATE_OneParam(FWaitOverlapDelegate, const FGameplayAbilityTargetDataHandle&, TargetData);
 
 class AActor;
 class UPrimitiveComponent;
diff --git a/Engine/Source/Runtime/GameplayAbilities/Public/Abilities/Tasks/AbilityTask_WaitTargetData.h b/Engine/Source/Runtime/GameplayAbilities/Public/Abilities/Tasks/AbilityTask_WaitTargetData.h
--- a/Engine/Source/Runtime/GameplayAbilities/Public/Abilities/Tasks/AbilityTask_WaitTargetData.h
+++ b/Engine/Source/Runtime/GameplayAbilities/Public/Abilities/Tasks/AbilityTask_WaitTargetData.h
@@ -7,7 +7,7 @@
 
 class AGameplayAbilityTargetActor;
 
-DECLARE_DYNAMIC_MULTICAST_DELEGATE_OneParam(FWaitTargetDataDelegate, FGameplayAbilityTargetDataHandle, Data);
+DECLARE_DYNAMIC_MULTICAST_DELEGATE_OneParam(FWaitTargetDataDelegate, const FGameplayAbilityTargetDataHandle&, Data);
 
 /** Wait for targeting actor (spawned from parameter) to provide data. Can be set not to end upon outputting data. Can be ended by task name. */
 UCLASS(notplaceable)
@@ -22,16 +22,16 @@ class GAMEPLAYABILITIES_API UAbilityTask_WaitTargetData: public UAbilityTask
 	FWaitTargetDataDelegate	Cancelled;
 
 	UFUNCTION()
-	void OnTargetDataReplicatedCallback(FGameplayAbilityTargetDataHandle Data, FGameplayTag ActivationTag);
+	void OnTargetDataReplicatedCallback(const FGameplayAbilityTargetDataHandle& Data, FGameplayTag ActivationTag);
 
 	UFUNCTION()
 	void OnTargetDataReplicatedCancelledCallback();
 
 	UFUNCTION()
-	void OnTargetDataReadyCallback(FGameplayAbilityTargetDataHandle Data);
+	void OnTargetDataReadyCallback(const FGameplayAbilityTargetDataHandle& Data);
 
 	UFUNCTION()
-	void OnTargetDataCancelledCallback(FGameplayAbilityTargetDataHandle Data);
+	void OnTargetDataCancelledCallback(const FGameplayAbilityTargetDataHandle& Data);
 
 	/** Spawns target actor and waits for it to return valid data or to be canceled. */
 	UFUNCTION(BlueprintCallable, meta=(HidePin = "WorldContextObject", DefaultToSelf = "WorldContextObject", BlueprintInternalUseOnly = "true", HideSpawnParms="Instigator"), Category="Ability|Tasks")
diff --git a/Engine/Source/Runtime/GameplayAbilities/Public/AbilitySystemBlueprintLibrary.h b/Engine/Source/Runtime/GameplayAbilities/Public/AbilitySystemBlueprintLibrary.h
--- a/Engine/Source/Runtime/GameplayAbilities/Public/AbilitySystemBlueprintLibrary.h
+++ b/Engine/Source/Runtime/GameplayAbilities/Public/AbilitySystemBlueprintLibrary.h
@@ -45,26 +45,26 @@ class GAMEPLAYABILITIES_API UAbilitySystemBlueprintLibrary : public UBlueprintFu
 	// -------------------------------------------------------------------------------
 
 	UFUNCTION(BlueprintCallable, Category = "Ability|TargetData")
-	static FGameplayAbilityTargetDataHandle AppendTargetDataHandle(FGameplayAbilityTargetDataHandle TargetHandle, FGameplayAbilityTargetDataHandle HandleToAdd);
+	static FGameplayAbilityTargetDataHandle AppendTargetDataHandle(FGameplayAbilityTargetDataHandle TargetHandle, const FGameplayAbilityTargetDataHandle& HandleToAdd);
 
 	UFUNCTION(BlueprintPure, Category = "Ability|TargetData")
 	static FGameplayAbilityTargetDataHandle	AbilityTargetDataFromLocations(const FGameplayAbilityTargetingLocationInfo& SourceLocation, const FGameplayAbilityTargetingLocationInfo& TargetLocation);
 
 	UFUNCTION(BlueprintPure, Category = "Ability|TargetData")
-	static FGameplayAbilityTargetDataHandle	AbilityTargetDataFromHitResult(FHitResult HitResult);
+	static FGameplayAbilityTargetDataHandle	AbilityTargetDataFromHitResult(const FHitResult& HitResult);
 
 	UFUNCTION(BlueprintPure, Category = "Ability|TargetData")
-	static int32 GetDataCountFromTargetData(FGameplayAbilityTargetDataHandle TargetData);
+	static int32 GetDataCountFromTargetData(const FGameplayAbilityTargetDataHandle& TargetData);
 
 	UFUNCTION(BlueprintPure, Category = "Ability|TargetData")
 	static FGameplayAbilityTargetDataHandle	AbilityTargetDataFromActor(AActor* Actor);
 
 	UFUNCTION(BlueprintPure, Category = "Ability|TargetData")
-	static FGameplayAbilityTargetDataHandle	AbilityTargetDataFromActorArray(TArray<AActor*> ActorArray, bool OneTargetPerHandle);
+	static FGameplayAbilityTargetDataHandle	AbilityTargetDataFromActorArray(const TArray<AActor*>& ActorArray, bool OneTargetPerHandle);
 
 	/** Create a new target data handle with filtration performed on the data */
 	UFUNCTION(BlueprintPure, Category = "Ability|TargetData")
-	static FGameplayAbilityTargetDataHandle	FilterTargetData(FGameplayAbilityTargetDataHandle TargetDataHandle, FGameplayTargetDataFilterHandle ActorFilterClass);
+	static FGameplayAbilityTargetDataHandle	FilterTargetData(const FGameplayAbilityTargetDataHandle& TargetDataHandle, FGameplayTargetDataFilterHandle ActorFilterClass);
 
 	/** Create a handle for filtering target data, filling out all fields */
 	UFUNCTION(BlueprintPure, Category = "Filter")
@@ -75,36 +75,36 @@ class GAMEPLAYABILITIES_API UAbilitySystemBlueprintLibrary : public UBlueprintFu
 	static FGameplayEffectSpecHandle MakeSpecHandle(UGameplayEffect* InGameplayEffect, AActor* InInstigator, AActor* InEffectCauser, float InLevel = 1.0f);
 
 	UFUNCTION(BlueprintPure, Category = "Ability|TargetData")
-	static TArray<AActor*> GetActorsFromTargetData(FGameplayAbilityTargetDataHandle TargetData, int32 Index);
+	static TArray<AActor*> GetActorsFromTargetData(const FGameplayAbilityTargetDataHandle& TargetData, int32 Index);
 
 	/** Returns true if the given TargetData has the actor passed in targeted */
 	UFUNCTION(BlueprintPure, Category = "Ability|TargetData")
-	static bool DoesTargetDataContainActor(FGameplayAbilityTargetDataHandle TargetData, int32 Index, AActor* Actor);
+	static bool DoesTargetDataContainActor(const FGameplayAbilityTargetDataHandle& TargetData, int32 Index, AActor* Actor);
 
 	/** Returns true if the given TargetData has at least 1 actor targeted */
 	UFUNCTION(BlueprintPure, Category = "Ability|TargetData")
-	static bool TargetDataHasActor(FGameplayAbilityTargetDataHandle TargetData, int32 Index);
+	static bool TargetDataHasActor(const FGameplayAbilityTargetDataHandle& TargetData, int32 Index);
 
 	UFUNCTION(BlueprintPure, Category = "Ability|TargetData")
-	static bool TargetDataHasHitResult(FGameplayAbilityTargetDataHandle HitResult, int32 Index);
+	static bool TargetDataHasHitResult(const FGameplayAbilityTargetDataHandle& HitResult, int32 Index);
 
 	UFUNCTION(BlueprintPure, Category = "Ability|TargetData")
-	static FHitResult GetHitResultFromTargetData(FGameplayAbilityTargetDataHandle HitResult, int32 Index);
+	static FHitResult GetHitResultFromTargetData(const FGameplayAbilityTargetDataHandle& HitResult, int32 Index);
 
 	UFUNCTION(BlueprintPure, Category = "Ability|TargetData")
-	static bool TargetDataHasOrigin(FGameplayAbilityTargetDataHandle TargetData, int32 Index);
+	static bool TargetDataHasOrigin(const FGameplayAbilityTargetDataHandle& TargetData, int32 Index);
 
 	UFUNCTION(BlueprintPure, Category = "Ability|TargetData")
-	static FTransform GetTargetDataOrigin(FGameplayAbilityTargetDataHandle TargetData, int32 Index);
+	static FTransform GetTargetDataOrigin(const FGameplayAbilityTargetDataHandle& TargetData, int32 Index);
 
 	UFUNCTION(BlueprintPure, Category = "Ability|TargetData")
-	static bool TargetDataHasEndPoint(FGameplayAbilityTargetDataHandle TargetData, int32 Index);
+	static bool TargetDataHasEndPoint(const FGameplayAbilityTargetDataHandle& TargetData, int32 Index);
 
 	UFUNCTION(BlueprintPure, Category = "Ability|TargetData")
-	static FVector GetTargetDataEndPoint(FGameplayAbilityTargetDataHandle TargetData, int32 Index);
+	static FVector GetTargetDataEndPoint(const FGameplayAbilityTargetDataHandle& TargetData, int32 Index);
 
 	UFUNCTION(BlueprintPure, Category = "Ability|TargetData")
-	static FTransform GetTargetDataEndPointTransform(FGameplayAbilityTargetDataHandle TargetData, int32 Index);
+	static FTransform GetTargetDataEndPointTransform(const FGameplayAbilityTargetDataHandle& TargetData, int32 Index);
 
 	// -------------------------------------------------------------------------------
 	//		GameplayEffectContext
