#Animation crash with state reinitialization without update

https://jira.ol.epicgames.net/browse/UE-24688

diff --git a/Engine/Source/Runtime/AnimGraphRuntime/Public/AnimNodes/AnimNode_BlendListByBool.h b/Engine/Source/Runtime/AnimGraphRuntime/Public/AnimNodes/AnimNode_BlendListByBool.h
--- a/Engine/Source/Runtime/AnimGraphRuntime/Public/AnimNodes/AnimNode_BlendListByBool.h
+++ b/Engine/Source/Runtime/AnimGraphRuntime/Public/AnimNodes/AnimNode_BlendListByBool.h
@@ -27,6 +27,6 @@ public:
 	}
 
 protected:
-	virtual int32 GetActiveChildIndex();
-	virtual FString GetNodeName(FNodeDebugData& DebugData) { return DebugData.GetNodeName(this); }
+	virtual int32 GetActiveChildIndex() override;
+	virtual FString GetNodeName(FNodeDebugData& DebugData) override { return DebugData.GetNodeName(this); }
 };
diff --git a/Engine/Source/Runtime/AnimGraphRuntime/Public/AnimNodes/AnimNode_BlendListByEnum.h b/Engine/Source/Runtime/AnimGraphRuntime/Public/AnimNodes/AnimNode_BlendListByEnum.h
--- a/Engine/Source/Runtime/AnimGraphRuntime/Public/AnimNodes/AnimNode_BlendListByEnum.h
+++ b/Engine/Source/Runtime/AnimGraphRuntime/Public/AnimNodes/AnimNode_BlendListByEnum.h
@@ -26,6 +26,6 @@ public:
 	}
 
 protected:
-	virtual int32 GetActiveChildIndex();
-	virtual FString GetNodeName(FNodeDebugData& DebugData) { return DebugData.GetNodeName(this); }
+	virtual int32 GetActiveChildIndex() override;
+	virtual FString GetNodeName(FNodeDebugData& DebugData) override { return DebugData.GetNodeName(this); }
 };
diff --git a/Engine/Source/Runtime/Engine/Private/Animation/AnimNode_StateMachine.cpp b/Engine/Source/Runtime/Engine/Private/Animation/AnimNode_StateMachine.cpp
--- a/Engine/Source/Runtime/Engine/Private/Animation/AnimNode_StateMachine.cpp
+++ b/Engine/Source/Runtime/Engine/Private/Animation/AnimNode_StateMachine.cpp
@@ -337,10 +337,10 @@ const FAnimationTransitionBetweenStates& FAnimNode_StateMachine::GetTransitionIn
 void FAnimNode_StateMachine::Update(const FAnimationUpdateContext& Context)
 {
 	// If we just became relevant and haven't been initialized yet, then reinitialize state machine.
+	bool bShouldReinitialize = false;
 	if (!bFirstUpdate && (UpdateCounter.Get() != INDEX_NONE) && !UpdateCounter.WasSynchronizedInTheLastFrame(Context.AnimInstanceProxy->GetUpdateCounter()))
 	{
-		FAnimationInitializeContext InitializationContext(Context.AnimInstanceProxy);
-		this->Initialize(InitializationContext);
+		bShouldReinitialize = true;
 	}
 	UpdateCounter.SynchronizeWith(Context.AnimInstanceProxy->GetUpdateCounter());
 
@@ -501,6 +501,12 @@ void FAnimNode_StateMachine::Update(const FAnimationUpdateContext& Context)
 	// Update the only active state if there are no transitions still in flight
 	if (ActiveTransitionArray.Num() == 0 && !IsAConduitState(CurrentState) && !StatesUpdated.Contains(CurrentState))
 	{
+		if (bShouldReinitialize)
+		{
+			FAnimationInitializeContext InitializationContext(Context.AnimInstanceProxy);
+			StatePoseLinks[CurrentState].Initialize(InitializationContext);
+		}
+
 		StatePoseLinks[CurrentState].Update(Context);
 	}
 
