Copying //UE4/Orion-Staging to //UE4/Main (Origin: //Orion/Dev-General [REDACTED_HANDLE])

diff --git a/Engine/Source/Runtime/Engine/Private/AudioDevice.cpp b/Engine/Source/Runtime/Engine/Private/AudioDevice.cpp
--- a/Engine/Source/Runtime/Engine/Private/AudioDevice.cpp
+++ b/Engine/Source/Runtime/Engine/Private/AudioDevice.cpp
@@ -850,6 +850,234 @@ bool FAudioDevice::HandleAudio3dVisualizeCommand(const TCHAR* Cmd, FOutputDevice
 	}
 	return true;
 }
+
+
+bool FAudioDevice::HandleAudioMemoryInfo(const TCHAR* Cmd, FOutputDevice& Ar)
+{
+	struct FSoundWaveInfo
+	{
+		USoundWave* SoundWave;
+		SIZE_T ResourceSize;
+		FString SoundGroupName;
+		float Duration;
+		bool bDecompressed;
+
+		FSoundWaveInfo(USoundWave* InSoundWave, SIZE_T InResourceSize, const FString& InSoundGroupName, float InDuration, bool bInDecompressed)
+			: SoundWave(InSoundWave)
+			, ResourceSize(InResourceSize)
+			, SoundGroupName(InSoundGroupName)
+			, Duration(InDuration)
+			, bDecompressed(bInDecompressed)
+		{}
+	};
+
+	struct FSoundWaveGroupInfo
+	{
+		SIZE_T ResourceSize;
+		SIZE_T CompressedResourceSize;
+
+		FSoundWaveGroupInfo()
+			: ResourceSize(0)
+			, CompressedResourceSize(0)
+		{}
+	};
+
+	// Alpha sort the objects by path name
+	struct FCompareSoundWave
+	{
+		FORCEINLINE bool operator()(const FSoundWaveInfo& A, const FSoundWaveInfo& B) const
+		{
+			return A.SoundWave->GetPathName() < B.SoundWave->GetPathName();
+		}
+	};
+
+	const FString PathName = *(FPaths::ProfilingDir() + TEXT("MemReports/"));
+	IFileManager::Get().MakeDirectory(*PathName);
+
+	const FString Filename = CreateProfileFilename(TEXT("_audio_memreport.csv"), true);
+	FString FilenameFull = PathName + Filename;
+
+	UE_LOG(LogEngine, Log, TEXT("AudioMemReport: saving to %s"), *FilenameFull);
+
+	FArchive* FileAr = IFileManager::Get().CreateDebugFileWriter(*FilenameFull);
+	FOutputDeviceArchiveWrapper* FileArWrapper = new FOutputDeviceArchiveWrapper(FileAr);
+	FOutputDevice* ReportAr = FileArWrapper;
+
+	// Get the sound wave class
+	UClass* SoundWaveClass = nullptr;
+	ParseObject<UClass>(TEXT("class=SoundWave"), TEXT("CLASS="), SoundWaveClass, ANY_PACKAGE);
+
+	TArray<FSoundWaveInfo> SoundWaveObjects;
+	TMap<FString, FSoundWaveGroupInfo> SoundWaveGroupSizes;
+	TArray<FString> SoundWaveGroupFolders;
+
+	// Grab the list of folders to specifically track memory usage for
+	FConfigSection* TrackedFolders = GConfig->GetSectionPrivate(TEXT("AudioMemReportFolders"), 0, 1, GEngineIni);
+	if (TrackedFolders)
+	{
+		for (FConfigSectionMap::TIterator It(*TrackedFolders); It; ++It)
+		{
+			FString SoundFolder = *It.Value();
+			SoundWaveGroupSizes.Add(SoundFolder, FSoundWaveGroupInfo());
+			SoundWaveGroupFolders.Add(SoundFolder);
+		}
+	}
+
+	SIZE_T TotalResourceSize = 0;
+	SIZE_T CompressedResourceSize = 0;
+	SIZE_T DecompressedResourceSize = 0;
+	int32 CompressedResourceCount = 0;
+
+	if (SoundWaveClass != nullptr)
+	{
+		// Loop through all objects and find only sound wave objects
+		for (TObjectIterator<USoundWave> It; It; ++It)
+		{
+			if (It->IsTemplate(RF_ClassDefaultObject))
+			{
+				continue;
+			}
+
+			// Get the resource size of the sound wave
+			const SIZE_T TrueResourceSize = It->GetResourceSize(EResourceSizeMode::Exclusive);
+			if (TrueResourceSize == 0)
+			{
+				continue;
+			}
+
+			USoundWave* SoundWave = *It;
+
+			const FSoundGroup& SoundGroup = GetDefault<USoundGroups>()->GetSoundGroup(SoundWave->SoundGroup);
+			float Duration = SoundWave->GetDuration();
+			bool bDecompressed = SoundGroup.bAlwaysDecompressOnLoad || Duration < SoundGroup.DecompressedDuration;
+
+			FString SoundGroupName;
+			switch (SoundWave->SoundGroup)
+			{
+				case ESoundGroup::SOUNDGROUP_Default:
+					SoundGroupName = TEXT("Default");
+					break;
+
+				case ESoundGroup::SOUNDGROUP_Effects:
+					SoundGroupName = TEXT("Effects");
+					break;
+
+				case ESoundGroup::SOUNDGROUP_UI:
+					SoundGroupName = TEXT("UI");
+					break;
+
+				case ESoundGroup::SOUNDGROUP_Music:
+					SoundGroupName = TEXT("Music");
+					break;
+
+				case ESoundGroup::SOUNDGROUP_Voice:
+					SoundGroupName = TEXT("Voice");
+					break;
+
+				default:
+					SoundGroupName = SoundGroup.DisplayName;
+					break;
+			}
+
+			// Add the info to the SoundWaveObjects array
+			SoundWaveObjects.Add(FSoundWaveInfo(SoundWave, TrueResourceSize, SoundGroupName, Duration, bDecompressed));
+
+			// Track total resource usage
+			TotalResourceSize += TrueResourceSize;
+
+			if (bDecompressed)
+			{
+				DecompressedResourceSize += TrueResourceSize;
+				++CompressedResourceCount;
+			}
+			else
+			{
+				CompressedResourceSize += TrueResourceSize;
+			}
+
+			// Get the sound object path
+			FString SoundWavePath = SoundWave->GetPathName();
+
+			// Now track the resource size according to all the sub-directories
+			FString SubDir;
+			int32 Index = 0;
+
+			for (int32 i = 0; i < SoundWavePath.Len(); ++i)
+			{
+				if (SoundWavePath[i] == '/')
+				{
+					if (SubDir.Len() > 0)
+					{
+						FSoundWaveGroupInfo* SubDirSize = SoundWaveGroupSizes.Find(SubDir);
+						if (SubDirSize)
+						{
+							SubDirSize->ResourceSize += TrueResourceSize;
+							if (bDecompressed)
+							{
+								SubDirSize->CompressedResourceSize += TrueResourceSize;
+							}
+						}
+					}
+					SubDir = TEXT("");
+				}
+				else
+				{
+					SubDir.AppendChar(SoundWavePath[i]);
+				}
+			}
+		}
+
+		ReportAr->Log(TEXT("Sound Wave Memory Report"));
+		ReportAr->Log(TEXT(""));
+
+		if (SoundWaveObjects.Num())
+		{
+			// Alpha sort the sound wave objects
+			SoundWaveObjects.Sort(FCompareSoundWave());
+
+			// Log the sound wave objects
+			
+			ReportAr->Logf(TEXT("Memory (MB),Count"));
+			ReportAr->Logf(TEXT("Total,%.3f,%d"), TotalResourceSize / 1024.f / 1024.f, SoundWaveObjects.Num());
+			ReportAr->Logf(TEXT("Decompressed,%.3f,%d"), DecompressedResourceSize / 1024.f / 1024.f, CompressedResourceCount);
+			ReportAr->Logf(TEXT("Compressed,%.3f,%d"), CompressedResourceSize / 1024.f / 1024.f, SoundWaveObjects.Num() - CompressedResourceCount);
+
+			if (SoundWaveGroupFolders.Num())
+			{
+				ReportAr->Log(TEXT(""));
+				ReportAr->Log(TEXT("Memory Usage and Count for Specified Folders (Folders defined in [AudioMemReportFolders] section in DefaultEngine.ini file):"));
+				ReportAr->Log(TEXT(""));
+				ReportAr->Logf(TEXT("%s,%s,%s"), TEXT("Directory"), TEXT("Total (MB)"), TEXT("Compressed (MB)"));
+				for (const FString& SoundWaveGroupFolder : SoundWaveGroupFolders)
+				{
+					FSoundWaveGroupInfo* SubDirSize = SoundWaveGroupSizes.Find(SoundWaveGroupFolder);
+					check(SubDirSize);
+					ReportAr->Logf(TEXT("%s,%10.2f,%10.2f"), *SoundWaveGroupFolder, SubDirSize->ResourceSize / 1024.0f / 1024.0f, SubDirSize->CompressedResourceSize / 1024.0f / 1024.0f);
+				}
+			}
+
+			ReportAr->Log(TEXT(""));
+			ReportAr->Log(TEXT("All Sound Wave Objects Sorted Alphebetically:"));
+			ReportAr->Log(TEXT(""));
+
+			ReportAr->Logf(TEXT("%s,%s,%s,%s,%s,%s"), TEXT("SoundWave"), TEXT("KB"), TEXT("MB"), TEXT("SoundGroup"), TEXT("Duration"), TEXT("CompressionState"));
+			for (const FSoundWaveInfo& Info : SoundWaveObjects)
+			{
+				float Kbytes = Info.ResourceSize / 1024.0f;
+				ReportAr->Logf(TEXT("%s,%10.2f,%10.2f,%s,%10.2f,%s"), *Info.SoundWave->GetPathName(), Kbytes, Kbytes / 1024.0f, *Info.SoundGroupName, Info.Duration, Info.bDecompressed ? TEXT("Decompressed") : TEXT("Compressed"));
+			}
+		}
+
+	}
+
+	// Shutdown and free archive resources
+	FileArWrapper->TearDown();
+	delete FileArWrapper;
+	delete FileAr;
+
+	return true;
+}
+
 #endif // !UE_BUILD_SHIPPING
 
 EDebugState FAudioDevice::GetMixDebugState( void )
@@ -964,6 +1192,10 @@ bool FAudioDevice::Exec( UWorld* InWorld, const TCHAR* Cmd, FOutputDevice& Ar )
 	{
 		return HandleAudio3dVisualizeCommand(Cmd, Ar);
 	}
+	else if (FParse::Command(&Cmd, TEXT("AudioMemReport")))
+	{
+		return HandleAudioMemoryInfo(Cmd, Ar);
+	}
 #endif // !UE_BUILD_SHIPPING
 
 	return false;
@@ -2533,6 +2765,14 @@ void FAudioDevice::Precache(USoundWave* SoundWave, bool bSynchronous, bool bTrac
 	{
 		const FSoundGroup& SoundGroup = GetDefault<USoundGroups>()->GetSoundGroup(SoundWave->SoundGroup);
 
+		float CompressedDurationThreshold = SoundGroup.DecompressedDuration;
+		/*
+		if (CompressedDurationThreshold > 0.0f)
+		{
+			CompressedDurationThreshold = 1.0f;
+		}
+		*/
+
 		// handle audio decompression
 		if (FPlatformProperties::SupportsAudioStreaming() && SoundWave->IsStreaming())
 		{
@@ -2540,7 +2780,7 @@ void FAudioDevice::Precache(USoundWave* SoundWave, bool bSynchronous, bool bTrac
 			SoundWave->bCanProcessAsync = false;
 		}
 		else if (SupportsRealtimeDecompression() && 
-			(bDisableAudioCaching || (!SoundGroup.bAlwaysDecompressOnLoad && SoundWave->Duration > SoundGroup.DecompressedDuration)))
+				 (bDisableAudioCaching || (!SoundGroup.bAlwaysDecompressOnLoad && SoundWave->Duration > CompressedDurationThreshold)))
 		{
 			// Store as compressed data and decompress in realtime
 			SoundWave->DecompressionType = DTYPE_RealTime;
diff --git a/Engine/Source/Runtime/Engine/Private/BatchedElements.cpp b/Engine/Source/Runtime/Engine/Private/BatchedElements.cpp
--- a/Engine/Source/Runtime/Engine/Private/BatchedElements.cpp
+++ b/Engine/Source/Runtime/Engine/Private/BatchedElements.cpp
@@ -382,7 +382,9 @@ FBatchedElements::FSimpleElementBSSContainer FBatchedElements::GammaAlphaOnlySha
 /** Global alpha ref test value for rendering masked batched elements */
 float GBatchedElementAlphaRefVal = 128.f;
 /** Global smoothing width for rendering batched elements with distance field blend modes */
-float GBatchedElementSmoothWidth = 12;
+float GBatchedElementSmoothWidth = 4;
+
+FAutoConsoleVariableRef CVarWellCanvasDistanceFieldSmoothWidth(TEXT("Canvas.DistanceFieldSmoothness"), GBatchedElementSmoothWidth, TEXT("Global sharpness of distance field fonts/shapes rendered by canvas."), ECVF_Default);
 
 template<class TSimpleElementPixelShader>
 static TSimpleElementPixelShader* GetPixelShader(bool bEncoded, ESimpleElementBlendMode BlendMode, ERHIFeatureLevel::Type FeatureLevel)
diff --git a/Engine/Source/Runtime/Engine/Private/Camera/CameraAnimInst.cpp b/Engine/Source/Runtime/Engine/Private/Camera/CameraAnimInst.cpp
--- a/Engine/Source/Runtime/Engine/Private/Camera/CameraAnimInst.cpp
+++ b/Engine/Source/Runtime/Engine/Private/Camera/CameraAnimInst.cpp
@@ -167,6 +167,8 @@ static const FName NAME_CameraComponentFieldOfViewPropertyName(TEXT("CameraCompo
 
 void UCameraAnimInst::Play(UCameraAnim* Anim, class AActor* CamActor, float InRate, float InScale, float InBlendInTime, float InBlendOutTime, bool bInLooping, bool bRandomStartTime, float Duration)
 {
+	check(IsInGameThread());
+
 	if (Anim && Anim->CameraInterpGroup)
 	{
 		// make sure any previous anim has been terminated correctly
@@ -197,6 +199,8 @@ void UCameraAnimInst::Play(UCameraAnim* Anim, class AActor* CamActor, float InRa
 		}
 		InterpGroupInst->InitGroupInst(CamAnim->CameraInterpGroup, CamActor);
 
+		checkf(CamAnim->CameraInterpGroup->InterpTracks.Num() == InterpGroupInst->TrackInst.Num(), TEXT("Track count mismatch! Outer = %s"), *GetNameSafe(CamAnim));
+
 		// cache move track refs
 		for (int32 Idx = 0; Idx < InterpGroupInst->TrackInst.Num(); ++Idx)
 		{
@@ -242,11 +246,14 @@ void UCameraAnimInst::Play(UCameraAnim* Anim, class AActor* CamActor, float InRa
 
 void UCameraAnimInst::Stop(bool bImmediate)
 {
+	check(IsInGameThread());
+
 	if ( bImmediate || (BlendOutTime <= 0.f) )
 	{
 		if (InterpGroupInst->Group != NULL)
 		{
 			InterpGroupInst->TermGroupInst(true);
+			InterpGroupInst->Group = nullptr;
 		}
 		MoveTrack = NULL;
 		MoveInst = NULL;
diff --git a/Engine/Source/Runtime/Engine/Private/Camera/CameraModifier_CameraShake.cpp b/Engine/Source/Runtime/Engine/Private/Camera/CameraModifier_CameraShake.cpp
--- a/Engine/Source/Runtime/Engine/Private/Camera/CameraModifier_CameraShake.cpp
+++ b/Engine/Source/Runtime/Engine/Private/Camera/CameraModifier_CameraShake.cpp
@@ -35,7 +35,7 @@ bool UCameraModifier_CameraShake::ModifyCamera(float DeltaTime, FMinimalViewInfo
 		// Delete any obsolete shakes
 		for (int32 i=ActiveShakes.Num()-1; i>=0; i--)
 		{
-			UCameraShake* ShakeInst = ActiveShakes[i];
+			UCameraShake* const ShakeInst = ActiveShakes[i];
 			if ((ShakeInst == nullptr) || ShakeInst->IsFinished())
 			{
 				ActiveShakes.RemoveAt(i, 1);
@@ -106,39 +106,49 @@ UCameraShake* UCameraModifier_CameraShake::AddCameraShake(TSubclassOf<class UCam
 	return nullptr;
 }
 
-void UCameraModifier_CameraShake::RemoveCameraShake(UCameraShake* ShakeInst)
+void UCameraModifier_CameraShake::RemoveCameraShake(UCameraShake* ShakeInst, bool bImmediately)
 {
 	for (int32 i = 0; i < ActiveShakes.Num(); ++i)
 	{
 		if (ActiveShakes[i] == ShakeInst)
 		{
-			ShakeInst->StopShake();
-			ActiveShakes.RemoveAt(i, 1);
+			ShakeInst->StopShake(bImmediately);
+
+			if (bImmediately)
+			{
+				ActiveShakes.RemoveAt(i, 1);
+			}
 			break;
 		}
 	}
 }
 
-void UCameraModifier_CameraShake::RemoveAllCameraShakesOfClass(TSubclassOf<class UCameraShake> ShakeClass)
+void UCameraModifier_CameraShake::RemoveAllCameraShakesOfClass(TSubclassOf<class UCameraShake> ShakeClass, bool bImmediately)
 {
 	for (int32 i = ActiveShakes.Num()-1; i >= 0; --i)
 	{
 		if ( ActiveShakes[i] && (ActiveShakes[i]->GetClass()->IsChildOf(ShakeClass)) )
 		{
-			ActiveShakes[i]->StopShake();
-			ActiveShakes.RemoveAt(i, 1);
+			ActiveShakes[i]->StopShake(bImmediately);
+			if (bImmediately)
+			{
+				ActiveShakes.RemoveAt(i, 1);
+			}
 		}
 	}
 }
 
-void UCameraModifier_CameraShake::RemoveAllCameraShakes()
+void UCameraModifier_CameraShake::RemoveAllCameraShakes(bool bImmediately)
 {
 	// clean up any active camera shake anims
 	for (UCameraShake* Inst : ActiveShakes)
 	{
-		Inst->StopShake();
+		Inst->StopShake(bImmediately);
 	}
 
-	// clear ActiveShakes array
-	ActiveShakes.Empty();
+	if (bImmediately)
+	{
+		// clear ActiveShakes array
+		ActiveShakes.Empty();
+	}
 }
diff --git a/Engine/Source/Runtime/Engine/Private/Camera/CameraShake.cpp b/Engine/Source/Runtime/Engine/Private/Camera/CameraShake.cpp
--- a/Engine/Source/Runtime/Engine/Private/Camera/CameraShake.cpp
+++ b/Engine/Source/Runtime/Engine/Private/Camera/CameraShake.cpp
@@ -40,16 +40,33 @@ UCameraShake::UCameraShake(const FObjectInitializer& ObjectInitializer)
 	OscillationBlendOutTime = 0.2f;
 }
 
-void UCameraShake::StopShake()
+void UCameraShake::StopShake(bool bImmediately)
 {
-	// stop cam anim if playing
-	if ((AnimInst != nullptr) && !AnimInst->bFinished)
+	if (bImmediately)
 	{
-		CameraOwner->StopCameraAnimInst(AnimInst, true);
+		// stop cam anim if playing
+		if ((AnimInst != nullptr) && !AnimInst->bFinished)
+		{
+			CameraOwner->StopCameraAnimInst(AnimInst, true);
+		}
+
 		AnimInst = nullptr;
+
+		// stop oscillation
+		OscillatorTimeRemaining = 0.f;
+	}
+	else
+	{
+		// advance to the blend out time
+		OscillatorTimeRemaining = FMath::Min(OscillatorTimeRemaining, OscillationBlendOutTime);
+
+		if ((AnimInst != nullptr) && !AnimInst->bFinished)
+		{
+			CameraOwner->StopCameraAnimInst(AnimInst, false);
+		}
 	}
 
-	ReceiveStopShake();
+	ReceiveStopShake(bImmediately);
 }
 
 void UCameraShake::PlayShake(APlayerCameraManager* Camera, float Scale, ECameraAnimPlaySpace::Type InPlaySpace, FRotator UserPlaySpaceRot)
diff --git a/Engine/Source/Runtime/Engine/Private/ComponentInstanceDataCache.cpp b/Engine/Source/Runtime/Engine/Private/ComponentInstanceDataCache.cpp
--- a/Engine/Source/Runtime/Engine/Private/ComponentInstanceDataCache.cpp
+++ b/Engine/Source/Runtime/Engine/Private/ComponentInstanceDataCache.cpp
@@ -157,7 +157,7 @@ FComponentInstanceDataCache::FComponentInstanceDataCache(const AActor* Actor)
 {
 	if (Actor != nullptr)
 	{
-		const bool bIsChildActor = Actor->ParentComponentActor.IsValid();
+		const bool bIsChildActor = Actor->IsChildActor();
 
 		TInlineComponentArray<UActorComponent*> Components(Actor);
 
@@ -201,7 +201,7 @@ void FComponentInstanceDataCache::ApplyToActor(AActor* Actor, const ECacheApplyP
 {
 	if (Actor != nullptr)
 	{
-		const bool bIsChildActor = Actor->ParentComponentActor.IsValid();
+		const bool bIsChildActor = Actor->IsChildActor();
 
 		TInlineComponentArray<UActorComponent*> Components(Actor);
 
diff --git a/Engine/Source/Runtime/Engine/Private/Components/CharacterMovementComponent.cpp b/Engine/Source/Runtime/Engine/Private/Components/CharacterMovementComponent.cpp
--- a/Engine/Source/Runtime/Engine/Private/Components/CharacterMovementComponent.cpp
+++ b/Engine/Source/Runtime/Engine/Private/Components/CharacterMovementComponent.cpp
@@ -467,8 +467,8 @@ void UCharacterMovementComponent::SetUpdatedComponent(USceneComponent* NewUpdate
 
 bool UCharacterMovementComponent::HasValidData() const
 {
-#if ENABLE_NAN_DIAGNOSTIC
 	bool bIsValid = UpdatedComponent && IsValid(CharacterOwner);
+#if ENABLE_NAN_DIAGNOSTIC
 	if (bIsValid)
 	{
 		// NaN-checking updates
@@ -484,11 +484,13 @@ bool UCharacterMovementComponent::HasValidData() const
 		{
 			logOrEnsureNanError(TEXT("UCharacterMovementComponent::HasValidData detected NaN/INF in UpdatedComponent ComponentTransform!"));
 		}
+		if (UpdatedComponent->GetComponentRotation().ContainsNaN())
+		{
+			logOrEnsureNanError(TEXT("UCharacterMovementComponent::HasValidData detected NaN/INF in UpdatedComponent->GetComponentRotation()! %s"), *UpdatedComponent->GetComponentRotation().ToString());
+		}
 	}
-	return bIsValid;
-#else
-	return UpdatedComponent && IsValid(CharacterOwner);
 #endif
+	return bIsValid;
 }
 
 FCollisionShape UCharacterMovementComponent::GetPawnCapsuleCollisionShape(const EShrinkCapsuleExtent ShrinkMode, const float CustomShrinkAmount) const
diff --git a/Engine/Source/Runtime/Engine/Private/Components/ChildActorComponent.cpp b/Engine/Source/Runtime/Engine/Private/Components/ChildActorComponent.cpp
--- a/Engine/Source/Runtime/Engine/Private/Components/ChildActorComponent.cpp
+++ b/Engine/Source/Runtime/Engine/Private/Components/ChildActorComponent.cpp
@@ -243,6 +243,17 @@ void UChildActorComponent::SetChildActorClass(TSubclassOf<AActor> Class)
 	}
 }
 
+struct FActorParentComponentSetter
+{
+private:
+	static void Set(AActor* ChildActor, UChildActorComponent* ParentComponent)
+	{
+		ChildActor->ParentComponent = ParentComponent;
+	}
+
+	friend UChildActorComponent;
+};
+
 void UChildActorComponent::CreateChildActor()
 {
 	// Kill spawned actor if we have one
@@ -265,14 +276,21 @@ void UChildActorComponent::CreateChildActor()
 					bSpawn = false;
 					UE_LOG(LogChildActorComponent, Error, TEXT("Found cycle in child actor component '%s'.  Not spawning Actor of class '%s' to break."), *GetPathName(), *ChildActorClass->GetName());
 				}
-				Actor = Actor->ParentComponentActor.Get();
+				if (UChildActorComponent* ParentComponent = Actor->GetParentComponent())
+				{
+					Actor = ParentComponent->GetOwner();
+				}
+				else
+				{
+					Actor = nullptr;
+				}
 			}
 
 			if (bSpawn)
 			{
 				FActorSpawnParameters Params;
 				Params.SpawnCollisionHandlingOverride = ESpawnActorCollisionHandlingMethod::AlwaysSpawn;
-				Params.bDeferConstruction = true; // We defer construction so that we set ParentComponentActor prior to component registration so they appear selected
+				Params.bDeferConstruction = true; // We defer construction so that we set ParentComponent prior to component registration so they appear selected
 				Params.bAllowDuringConstructionScript = true;
 				Params.OverrideLevel = (MyOwner ? MyOwner->GetLevel() : nullptr);
 				Params.Name = ChildActorName;
@@ -291,8 +309,8 @@ void UChildActorComponent::CreateChildActor()
 				{
 					ChildActorName = ChildActor->GetFName();
 
-					// Remember which actor spawned it (for selection in editor etc)
-					ChildActor->ParentComponentActor = MyOwner;
+					// Remember which component spawned it (for selection in editor etc)
+					FActorParentComponentSetter::Set(ChildActor, this);
 
 					// Parts that we deferred from SpawnActor
 					const FComponentInstanceDataCache* ComponentInstanceData = (CachedInstanceData ? CachedInstanceData->ComponentInstanceData : nullptr);
diff --git a/Engine/Source/Runtime/Engine/Private/Components/PrimitiveComponent.cpp b/Engine/Source/Runtime/Engine/Private/Components/PrimitiveComponent.cpp
--- a/Engine/Source/Runtime/Engine/Private/Components/PrimitiveComponent.cpp
+++ b/Engine/Source/Runtime/Engine/Private/Components/PrimitiveComponent.cpp
@@ -27,6 +27,7 @@
 #include "CollisionDebugDrawingPublic.h"
 #include "GameFramework/CheatManager.h"
 #include "GameFramework/DamageType.h"
+#include "Components/ChildActorComponent.h"
 
 #define LOCTEXT_NAMESPACE "PrimitiveComponent"
 
@@ -1030,14 +1031,21 @@ bool UPrimitiveComponent::IsComponentIndividuallySelected() const
 
 bool UPrimitiveComponent::ShouldRenderSelected() const
 {
-	const AActor* Owner = GetOwner();
-	return(	bSelectable && 
-			Owner != NULL && 
-#if WITH_EDITOR
-			(Owner->IsSelected() || (Owner->ParentComponentActor != NULL && Owner->ParentComponentActor->IsSelected())) );
-#else
-			Owner->IsSelected() );
-#endif
+	if (bSelectable)
+	{
+		if (const AActor* Owner = GetOwner())
+		{
+			if (Owner->IsSelected())
+			{
+				return true;
+			}
+			else if (UChildActorComponent* ParentComponent = Owner->GetParentComponent())
+			{
+				return ParentComponent->GetOwner()->IsSelected();
+			}
+		}
+	}
+	return false;
 }
 
 void UPrimitiveComponent::SetCastShadow(bool NewCastShadow)
diff --git a/Engine/Source/Runtime/Engine/Private/Components/SceneComponent.cpp b/Engine/Source/Runtime/Engine/Private/Components/SceneComponent.cpp
--- a/Engine/Source/Runtime/Engine/Private/Components/SceneComponent.cpp
+++ b/Engine/Source/Runtime/Engine/Private/Components/SceneComponent.cpp
@@ -15,6 +15,7 @@
 #include "Engine/SCS_Node.h"
 #include "SNotificationList.h"
 #include "NotificationManager.h"
+#include "Components/ChildActorComponent.h"
 
 #define LOCTEXT_NAMESPACE "SceneComponent"
 
@@ -2270,6 +2271,18 @@ bool USceneComponent::IsVisibleInEditor() const
 bool USceneComponent::ShouldRender() const
 {
 	AActor* Owner = GetOwner();
+
+	if (Owner)
+	{
+		if (UChildActorComponent* ParentComponent = Owner->GetParentComponent())
+		{
+			if (!ParentComponent->ShouldRender())
+			{
+				return false;
+			}
+		}
+	}
+	
 	const bool bShowInEditor = 
 #if WITH_EDITOR
 		GIsEditor ? (!Owner || !Owner->IsHiddenEd()) : false;
@@ -2286,6 +2299,18 @@ bool USceneComponent::ShouldRender() const
 bool USceneComponent::CanEverRender() const
 {
 	AActor* Owner = GetOwner();
+
+	if (Owner)
+	{
+		if (UChildActorComponent* ParentComponent = Owner->GetParentComponent())
+		{
+			if (!ParentComponent->CanEverRender())
+			{
+				return false;
+			}
+		}
+	}
+
 	const bool bShowInEditor =
 #if WITH_EDITOR
 		GIsEditor ? (!Owner || !Owner->IsHiddenEd()) : false;
diff --git a/Engine/Source/Runtime/Engine/Private/Components/StaticMeshComponent.cpp b/Engine/Source/Runtime/Engine/Private/Components/StaticMeshComponent.cpp
--- a/Engine/Source/Runtime/Engine/Private/Components/StaticMeshComponent.cpp
+++ b/Engine/Source/Runtime/Engine/Private/Components/StaticMeshComponent.cpp
@@ -2035,6 +2035,14 @@ void FStaticMeshComponentLODInfo::ImportText(const TCHAR** SourceText)
 	}
 }
 
+int32 GKeepKeepOverrideVertexColorsOnCPU = 1;
+FAutoConsoleVariableRef CKeepOverrideVertexColorsOnCPU(
+	TEXT("r.KeepOverrideVertexColorsOnCPU"),
+	GKeepKeepOverrideVertexColorsOnCPU,
+	TEXT("Keeps a CPU copy of override vertex colors.  May be required for some blueprints / object spawning."),
+	ECVF_Scalability | ECVF_RenderThreadSafe
+	);
+
 FArchive& operator<<(FArchive& Ar,FStaticMeshComponentLODInfo& I)
 {
 	const uint8 OverrideColorsStripFlag = 1;
@@ -2077,7 +2085,9 @@ FArchive& operator<<(FArchive& Ar,FStaticMeshComponentLODInfo& I)
 				I.OverrideVertexColors = new FColorVertexBuffer;
 			}
 
-			I.OverrideVertexColors->Serialize( Ar, true );
+			//we want to discard the vertex colors after rhi init when in cooked/client builds.
+			const bool bNeedsCPUAccess = !Ar.IsLoading() || GIsEditor || IsRunningCommandlet() || (GKeepKeepOverrideVertexColorsOnCPU != 0);
+			I.OverrideVertexColors->Serialize(Ar, bNeedsCPUAccess);
 		}
 	}
 
diff --git a/Engine/Source/Runtime/Engine/Private/ContentStreaming.cpp b/Engine/Source/Runtime/Engine/Private/ContentStreaming.cpp
--- a/Engine/Source/Runtime/Engine/Private/ContentStreaming.cpp
+++ b/Engine/Source/Runtime/Engine/Private/ContentStreaming.cpp
@@ -48,6 +48,7 @@ DEFINE_STAT(STAT_TotalForcedHeuristicSize);
 
 DEFINE_STAT(STAT_LightmapMemorySize);
 DEFINE_STAT(STAT_LightmapDiskSize);
+DEFINE_STAT(STAT_HLODTextureMemorySize);
 DEFINE_STAT(STAT_IntermediateTexturesSize);
 DEFINE_STAT(STAT_RequestSizeCurrentFrame);
 DEFINE_STAT(STAT_RequestSizeTotal);
@@ -128,6 +129,14 @@ static TAutoConsoleVariable<float> CVarStreamingBoost(
 	ECVF_Default
 	);
 
+static TAutoConsoleVariable<float> CVarStreamingScreenSizeEffectiveMax(
+	TEXT("r.Streaming.MaxEffectiveScreenSize"),
+	0,
+	TEXT("0: Use current actual vertical screen size\n")	
+	TEXT("> 0: Clamp wanted mip size calculation to this value for the vertical screen size component."),
+	ECVF_Scalability
+	);
+
 #if PLATFORM_SUPPORTS_TEXTURE_STREAMING
 static TAutoConsoleVariable<int32> CVarSetTextureStreaming(
 	TEXT("r.TextureStreaming"),
@@ -150,6 +159,13 @@ static TAutoConsoleVariable<int32> CVarStreamingPoolSize(
 	TEXT("-1: Default texture pool size, otherwise the size in MB"),
 	ECVF_Scalability);
 
+static TAutoConsoleVariable<int32> CVarStreamingMaxEjectSize(
+	TEXT("r.Streaming.MaxSingleStreamOutSize"),
+	50*1024*1024,
+	TEXT("Streaming out textures costs memory because two copies of the texture are briefly alive at the same time.\n")
+	TEXT("This sets a limit on much memory to use at once.\n"),
+	ECVF_Default);
+
 static TAutoConsoleVariable<int32> CVarStreamingShowWantedMips(
 	TEXT("r.Streaming.ShowWantedMips"),
 	0,
@@ -246,7 +262,6 @@ struct FTexturePriority
 	int32 TextureIndex;
 };
 
-
 /** Self-contained structure to manage a streaming texture, possibly on a separate thread. */
 struct FStreamingTexture
 {
@@ -256,9 +271,21 @@ struct FStreamingTexture
 		WantedMips = InTexture->ResidentMips;
 		DynamicScreenSize = 0.0f;
 		MipCount = InTexture->GetNumMips();
-		STAT( PerfectWantedMips = InTexture->ResidentMips );
+		PerfectWantedMips = InTexture->ResidentMips;
 		STAT( MostResidentMips = InTexture->ResidentMips );
 		LODGroup = (TextureGroup) InTexture->LODGroup;
+
+		//@TODO: Temporary: Done because there is currently no texture group for HLOD textures and adding one in a branch isn't ideal; this should be removed / replaced once there is a real HLOD texture group
+		bIsHLODTexture = !Texture->HasAnyFlags(RF_Public) && Texture->GetName().StartsWith(TEXT("T_LOD_"), ESearchCase::CaseSensitive);
+
+		// Landscape textures and HLOD textures should not be affected by MipBias
+		//@TODO: Long term this should probably be a property of the texture group instead
+		bCanBeAffectedByMipBias = true;
+		if ((LODGroup == TEXTUREGROUP_Terrain_Heightmap) || (LODGroup == TEXTUREGROUP_Terrain_Weightmap) || bIsHLODTexture)
+		{
+			bCanBeAffectedByMipBias = false;
+		}
+
 		NumNonStreamingMips = InTexture->GetNumNonStreamingMips();
 		ForceLoadRefCount = 0;
 		bIsStreamingLightmap = IsStreamingLightmap( Texture );
@@ -289,7 +316,7 @@ struct FStreamingTexture
 		RequestedMips = Texture->RequestedMips;
 		MinAllowedMips = 1;			//FMath::Min(Texture->ResidentMips, Texture->RequestedMips);
 		MaxAllowedMips = MipCount;	//FMath::Max(Texture->ResidentMips, Texture->RequestedMips);
-		STAT( MaxAllowedOptimalMips = MaxAllowedMips );
+		MaxAllowedOptimalMips = MaxAllowedMips;
 		STAT( MostResidentMips = FMath::Max(MostResidentMips, Texture->ResidentMips) );
 		float LastRenderTimeForTexture = Texture->GetLastRenderTimeForStreaming();
 		LastRenderTime = (FApp::GetCurrentTime() > LastRenderTimeForTexture) ? float( FApp::GetCurrentTime() - LastRenderTimeForTexture ) : 0.0f;
@@ -427,11 +454,11 @@ struct FStreamingTexture
 	int32			WantedMips;
 	/** Largest screen size occupied by this texture, counting dynamic primitives only. Call FFloatMipLevel::FromScreenSizeInTexels() to get DynamicWantedMips. */
 	float			DynamicScreenSize;
-#if STATS
 	/** Legacy: Same as WantedMips, but not clamped by fudge factor. */
 	int32			PerfectWantedMips;
 	/** Same as MaxAllowedMips, but not clamped by the -reducepoolsize command line option. */
 	int32			MaxAllowedOptimalMips;
+#if STATS
 	/** Most number of mip-levels this texture has ever had resident in memory. */
 	int32			MostResidentMips;
 #endif
@@ -493,6 +520,11 @@ struct FStreamingTexture
 	uint32			bUsesOrphanedHeuristics : 1;
 	/** Whether this texture has been replaced in a material and all primitives using it call NotifyPrimitiveUpdated(). */
 	uint32			bNeedPrimitiveUpdate : 1;
+	/** Can this texture be affected by mip bias? */
+	uint32			bCanBeAffectedByMipBias : 1;
+	/** Is this texture a HLOD texture? */
+	//@TODO: Temporary: Done because there is currently no texture group for HLOD textures and adding one in a branch isn't ideal; this should be removed / replaced once there is a real HLOD texture group
+	uint32			bIsHLODTexture : 1;
 };
 
 // ----------------------------------------------------------------------
@@ -512,14 +544,9 @@ int32 FFloatMipLevel::ComputeMip(const FStreamingTexture* StreamingTexture, floa
 	if(StreamingTexture)
 	{
 		MinAllowedMips = StreamingTexture->MinAllowedMips;
-		MaxAllowedMips = 
-#if STATS
-			bOptimal ? StreamingTexture->MaxAllowedOptimalMips : 
-#endif
-			StreamingTexture->MaxAllowedMips;
-		// Landscape textures should not be affected by MipBias
-		if (StreamingTexture->LODGroup == TEXTUREGROUP_Terrain_Heightmap || 
-			StreamingTexture->LODGroup == TEXTUREGROUP_Terrain_Weightmap)
+		MaxAllowedMips = bOptimal ? StreamingTexture->MaxAllowedOptimalMips : StreamingTexture->MaxAllowedMips;
+
+		if (!StreamingTexture->bCanBeAffectedByMipBias)
 		{
 			MipBias = 0.0f;
 		}
@@ -924,6 +951,7 @@ struct FStreamingContext
 		ThisFrameTotalStreamingTexturesMaxSize			= 0;
 		ThisFrameTotalLightmapMemorySize				= 0;
 		ThisFrameTotalLightmapDiskSize					= 0;
+		ThisFrameTotalHLODMemorySize					= 0;
 		ThisFrameTotalMipCountIncreaseRequestsInFlight	= 0;
 		ThisFrameOptimalWantedSize						= 0;
 		ThisFrameTotalStaticTextureHeuristicSize		= 0;
@@ -983,6 +1011,7 @@ struct FStreamingContext
 		ThisFrameTotalStreamingTexturesMaxSize			+= Other.ThisFrameTotalStreamingTexturesMaxSize;
 		ThisFrameTotalLightmapMemorySize				+= Other.ThisFrameTotalLightmapMemorySize;
 		ThisFrameTotalLightmapDiskSize					+= Other.ThisFrameTotalLightmapDiskSize;
+		ThisFrameTotalHLODMemorySize					+= Other.ThisFrameTotalHLODMemorySize;
 		ThisFrameTotalMipCountIncreaseRequestsInFlight	+= Other.ThisFrameTotalMipCountIncreaseRequestsInFlight;
 		ThisFrameOptimalWantedSize						+= Other.ThisFrameOptimalWantedSize;
 		ThisFrameTotalStaticTextureHeuristicSize		+= Other.ThisFrameTotalStaticTextureHeuristicSize;
@@ -1022,6 +1051,7 @@ struct FStreamingContext
 	uint64 ThisFrameTotalStreamingTexturesMaxSize;
 	uint64 ThisFrameTotalLightmapMemorySize;
 	uint64 ThisFrameTotalLightmapDiskSize;
+	uint64 ThisFrameTotalHLODMemorySize;
 	uint64 ThisFrameTotalMipCountIncreaseRequestsInFlight;
 	uint64 ThisFrameOptimalWantedSize;
 	/** Number of bytes using StaticTexture heuristics this frame, currently in memory. */
@@ -1080,7 +1110,7 @@ public:
 		void Reset()
 		{
 			TotalResidentSize = 0;
-			STAT( TotalRequiredSize = 0 );
+			TotalRequiredSize = 0;
 			TempStreamingSize = 0;
 			PendingStreamInSize = 0;
 			PendingStreamOutSize = 0;
@@ -1091,7 +1121,7 @@ public:
 		/** Total number of bytes currently in memory */
 		int64 TotalResidentSize;
 		/** Total number of bytes required, using PerfectWantedMips */
-		STAT( int64 TotalRequiredSize );
+		int64 TotalRequiredSize;
 		/** Currently being streamed in, in bytes. */
 		SIZE_T PendingStreamInSize;
 		/** Currently being streamed out, in bytes. */
@@ -1202,8 +1232,8 @@ private:
 				}
 			}
 
-			STAT( int32 PerfectWantedTextureSize = StreamingTexture.GetSize( StreamingTexture.PerfectWantedMips ) );
-			STAT( ThreadStats.TotalRequiredSize += PerfectWantedTextureSize );
+			const int32 PerfectWantedTextureSize = StreamingTexture.GetSize( StreamingTexture.PerfectWantedMips );
+			ThreadStats.TotalRequiredSize += PerfectWantedTextureSize;
 			StreamingManager.UpdateFrameStats( ThreadContext, StreamingTexture, Index );
 
 			// Reset the boost factor
@@ -2107,6 +2137,7 @@ FStreamingManagerTexture::FStreamingManagerTexture()
 ,	TotalStreamingTexturesMaxSize(0)
 ,	TotalLightmapMemorySize(0)
 ,	TotalLightmapDiskSize(0)
+,	TotalHLODMemorySize(0)
 ,	TotalMipCountIncreaseRequestsInFlight(0)
 ,	TotalOptimalWantedSize(0)
 ,	TotalStaticTextureHeuristicSize(0)
@@ -2114,6 +2145,7 @@ FStreamingManagerTexture::FStreamingManagerTexture()
 ,	TotalLastRenderHeuristicSize(0)
 ,	TotalForcedHeuristicSize(0)
 ,	MemoryOverBudget(0)
+,	MaxEverRequired(0)
 ,	OriginalTexturePoolSize(0)
 ,	PreviousPoolSizeTimestamp(0.0)
 ,	PreviousPoolSizeSetting(-1)
@@ -2418,6 +2450,7 @@ void Renderthread_StreamOutTextureData(FRHICommandListImmediate& RHICmdList, TAr
  **/
 bool FStreamingManagerTexture::StreamOutTextureData( int64 RequiredMemorySize )
 {
+	const int64 MaxTempSizeForStreamout = FMath::Max(1024 * 1024, CVarStreamingMaxEjectSize.GetValueOnGameThread());
 	RequiredMemorySize = FMath::Max<int64>(RequiredMemorySize, MinEvictSize);
 
 	// Array of candidates for reducing mip-levels.
@@ -2427,6 +2460,11 @@ bool FStreamingManagerTexture::StreamOutTextureData( int64 RequiredMemorySize )
 	// Don't stream out character textures (to begin with)
 	float CurrentTime = float(FPlatformTime::Seconds() - GStartTime);
 
+	volatile bool bSucceeded = false;
+	
+	//resizing textures actually creates a temp copy so we can only resize so many at a time without running out of memory during the eject itself.
+	int64 MemoryCostToResize = 0;	
+
 	// Collect all textures will be considered for streaming out.
 	for (int32 StreamingIndex = 0; StreamingIndex < StreamingTextures.Num(); ++StreamingIndex)
 	{
@@ -2466,25 +2504,54 @@ bool FStreamingManagerTexture::StreamOutTextureData( int64 RequiredMemorySize )
 														Texture->LODGroup == TEXTUREGROUP_CharacterSpecular ||
 														Texture->LODGroup == TEXTUREGROUP_CharacterNormalMap;
 							uint32 TextureDataAddress = 0;
+							MemoryCostToResize += Texture->CalcTextureMemorySize(FMath::Max(0, Texture->ResidentMips - 1));
 							CandidateTextures.Add( FTextureSortElement(Texture, Texture->CalcTextureMemorySize(Texture->ResidentMips), bIsCharacterTexture ? 1 : 0, TextureDataAddress, NumRequiredResidentMips) );
 						}
 					}
 				}
 			}
 		}
-	}
 
-	volatile bool bSucceeded = false;
+		if (MemoryCostToResize >= MaxTempSizeForStreamout)
+		{			
+			// Queue up the process on the render thread.
+			ENQUEUE_UNIQUE_RENDER_COMMAND_THREEPARAMETER(
+				StreamOutTextureDataCommand,
+				TArray<FTextureSortElement>*, CandidateTextures, &CandidateTextures,
+				int64, RequiredMemorySize, RequiredMemorySize,
+				volatile bool*, bSucceeded, &bSucceeded,
+				{				
+				FRHICommandListExecutor::GetImmediateCommandList().ImmediateFlush(EImmediateFlushType::FlushRHIThreadFlushResources);				
+				RHIFlushResources();
+
+				Renderthread_StreamOutTextureData(RHICmdList, CandidateTextures, RequiredMemorySize, bSucceeded);
+
+				FRHICommandListExecutor::GetImmediateCommandList().ImmediateFlush(EImmediateFlushType::FlushRHIThreadFlushResources);
+				RHIFlushResources();
+			});
+
+			// Block until the command has finished executing.
+			FlushRenderingCommands();
+			MemoryCostToResize = 0;			
+			CandidateTextures.Reset();
+		}
+	}	
 
-	// Queue up the process on the render thread.
-	ENQUEUE_UNIQUE_RENDER_COMMAND_THREEPARAMETER(
-		StreamOutTextureDataCommand,
-		TArray<FTextureSortElement>*,CandidateTextures,&CandidateTextures,
-		int64,RequiredMemorySize,RequiredMemorySize,
-		volatile bool*,bSucceeded,&bSucceeded,
+	if (CandidateTextures.Num() > 0)
 	{
-		Renderthread_StreamOutTextureData(RHICmdList, CandidateTextures, RequiredMemorySize, bSucceeded);
-	});
+		// Queue up the process on the render thread.
+		ENQUEUE_UNIQUE_RENDER_COMMAND_THREEPARAMETER(
+			StreamOutTextureDataCommand,
+			TArray<FTextureSortElement>*, CandidateTextures, &CandidateTextures,
+			int64, RequiredMemorySize, RequiredMemorySize,
+			volatile bool*, bSucceeded, &bSucceeded,
+			{
+			Renderthread_StreamOutTextureData(RHICmdList, CandidateTextures, RequiredMemorySize, bSucceeded);
+
+			FRHICommandListExecutor::GetImmediateCommandList().ImmediateFlush(EImmediateFlushType::FlushRHIThreadFlushResources);
+			RHIFlushResources();
+		});
+	}
 
 	// Block until the command has finished executing.
 	FlushRenderingCommands();
@@ -3155,6 +3222,7 @@ void FStreamingManagerTexture::ResetStreamingStats()
 	TotalStreamingTexturesMaxSize						= 0;
 	TotalLightmapMemorySize								= 0;
 	TotalLightmapDiskSize								= 0;
+	TotalHLODMemorySize									= 0;
 	TotalMipCountIncreaseRequestsInFlight				= 0;
 	TotalOptimalWantedSize								= 0;
 	TotalStaticTextureHeuristicSize						= 0;
@@ -3182,6 +3250,7 @@ void FStreamingManagerTexture::UpdateStreamingStats( const FStreamingContext& Co
 	TotalStreamingTexturesMaxSize			+= Context.ThisFrameTotalStreamingTexturesMaxSize;
 	TotalLightmapMemorySize					+= Context.ThisFrameTotalLightmapMemorySize;
 	TotalLightmapDiskSize					+= Context.ThisFrameTotalLightmapDiskSize;
+	TotalHLODMemorySize						+= Context.ThisFrameTotalHLODMemorySize;
 	TotalMipCountIncreaseRequestsInFlight	+= Context.ThisFrameTotalMipCountIncreaseRequestsInFlight;
 	TotalOptimalWantedSize					+= Context.ThisFrameOptimalWantedSize;
 	TotalStaticTextureHeuristicSize			+= Context.ThisFrameTotalStaticTextureHeuristicSize;
@@ -3205,6 +3274,7 @@ void FStreamingManagerTexture::UpdateStreamingStats( const FStreamingContext& Co
 		SET_DWORD_STAT( STAT_StreamingTexturesMaxSize,		TotalStreamingTexturesMaxSize	);
 		SET_DWORD_STAT( STAT_LightmapMemorySize,			TotalLightmapMemorySize			);
 		SET_DWORD_STAT( STAT_LightmapDiskSize,				TotalLightmapDiskSize			);
+		SET_DWORD_STAT( STAT_HLODTextureMemorySize,			TotalHLODMemorySize				);
 		SET_FLOAT_STAT( STAT_StreamingBandwidth,			BandwidthAverage/1024.0f/1024.0f);
 		SET_FLOAT_STAT( STAT_StreamingLatency,				LatencyAverage					);
 		SET_MEMORY_STAT( STAT_StreamingTexturesSize,		TotalStreamingTexturesSize		);
@@ -3383,7 +3453,7 @@ void FStreamingManagerTexture::StreamTextures( bool bProcessEverything )
 	ThreadStats.PendingStreamInSize += (Stats.PendingMemoryAdjustment > 0) ? Stats.PendingMemoryAdjustment : 0;
 	ThreadStats.PendingStreamOutSize += (Stats.PendingMemoryAdjustment < 0) ? -Stats.PendingMemoryAdjustment : 0;
 
-	int64 AvailableNow, AvailableLater, TempMemoryUsed;
+	int64 AvailableNow, AvailableLater, TempMemoryUsed, LargestAlloc;
 	if ( Stats.IsUsingLimitedPoolSize() )
 	{
 		ThreadStats.TempStreamingSize = GStreamMemoryTracker.CalcTempMemory();
@@ -3392,11 +3462,13 @@ void FStreamingManagerTexture::StreamTextures( bool bProcessEverything )
 		// Note: 'AvailableNow' and 'AvailableLater' are allowed to be negative.
 		AvailableNow = GStreamMemoryTracker.CalcAvailableNow( Stats.ComputeAvailableMemorySize(), MemoryMargin );
 		AvailableLater = GStreamMemoryTracker.CalcAvailableLater( Stats.ComputeAvailableMemorySize(), MemoryMargin );
+		LargestAlloc = Stats.LargestContiguousAllocation;
 
-		STAT( int64 NonStreamingUsage = Stats.AllocatedMemorySize - ThreadStats.TotalResidentSize - int64(ThreadStats.PendingStreamInSize) - TempMemoryUsed );
-		STAT( int64 MemoryBudget = Stats.TexturePoolSize - NonStreamingUsage - MemoryMargin );
-		STAT( int64 LocalMemoryOverBudget = FMath::Max(ThreadStats.TotalRequiredSize - MemoryBudget, 0LL) );
-		STAT( MemoryOverBudget = FMath::Max(LocalMemoryOverBudget, 0LL) );
+		const int64 NonStreamingUsage = Stats.AllocatedMemorySize - ThreadStats.TotalResidentSize - int64(ThreadStats.PendingStreamInSize) - TempMemoryUsed;
+		const int64 MemoryBudget = Stats.TexturePoolSize - NonStreamingUsage - MemoryMargin;
+		const int64 LocalMemoryOverBudget = FMath::Max(ThreadStats.TotalRequiredSize - MemoryBudget, 0LL);
+		MemoryOverBudget = FMath::Max(LocalMemoryOverBudget, 0LL);
+		MaxEverRequired = FMath::Max(ThreadStats.TotalRequiredSize, MaxEverRequired);
 
 		SET_MEMORY_STAT( STAT_StreamingOverBudget,	MemoryOverBudget);
 		SET_MEMORY_STAT( STAT_StreamingUnderBudget,	FMath::Max(-LocalMemoryOverBudget,0ll) );
@@ -3407,6 +3479,8 @@ void FStreamingManagerTexture::StreamTextures( bool bProcessEverything )
 		TempMemoryUsed = ThreadStats.TempStreamingSize;
 		AvailableNow = MAX_int64;
 		AvailableLater = MAX_int64;
+		LargestAlloc = AvailableNow;
+
 		MemoryOverBudget = 0;
 		SET_MEMORY_STAT( STAT_StreamingOverBudget, 0 );
 		SET_MEMORY_STAT( STAT_StreamingUnderBudget, 0 );
@@ -3461,9 +3535,10 @@ void FStreamingManagerTexture::StreamTextures( bool bProcessEverything )
 					AvailableLater -= StreamSize;
 
 					// Can we stream in this texture now?
-					if ( StreamSize <= AvailableNow && TempMemoryUsed < MaxTempMemoryUsed )
+					if (StreamSize <= LargestAlloc && StreamSize <= AvailableNow && TempMemoryUsed < MaxTempMemoryUsed)
 					{
 						AvailableNow -= StreamSize;
+						LargestAlloc -= StreamSize;
 						TempMemoryUsed += TempSize;
 						StartStreaming( HighPrioTexture, HighPrioTexture.WantedMips, Context, true );
 					}
@@ -4016,7 +4091,7 @@ void FStreamingManagerTexture::CalcMinMaxMips( FStreamingTexture& StreamingTextu
 		StreamingTexture.MaxAllowedMips );
 
 	// Take the reduced texture pool in to account.
-	STAT( StreamingTexture.MaxAllowedOptimalMips = StreamingTexture.MaxAllowedMips );
+	StreamingTexture.MaxAllowedOptimalMips = StreamingTexture.MaxAllowedMips;
 	if ( GIsOperatingWithReducedTexturePool )
 	{				
 		int32 MaxTextureMipCount = FMath::Max( GMaxTextureMipCount - 2, 0 );
@@ -4083,8 +4158,9 @@ void FStreamingManagerTexture::UpdateFrameStats( FStreamingContext& Context, FSt
 	}
 	Context.ThisFrameNumStreamingTextures++;
 	Context.ThisFrameTotalStreamingTexturesMaxSize += PotentialSize;
-	Context.ThisFrameTotalLightmapMemorySize += StreamingTexture.bIsStreamingLightmap ? ResidentSize : 0;
+	Context.ThisFrameTotalLightmapMemorySize += (StreamingTexture.bIsStreamingLightmap || StreamingTexture.LODGroup == TEXTUREGROUP_Lightmap || StreamingTexture.LODGroup == TEXTUREGROUP_Shadowmap) ? ResidentSize : 0;
 	Context.ThisFrameTotalLightmapDiskSize += StreamingTexture.bIsStreamingLightmap ? PotentialSize : 0;
+	Context.ThisFrameTotalHLODMemorySize += StreamingTexture.bIsHLODTexture ? ResidentSize : 0;
 #endif
 }
 
@@ -4101,7 +4177,7 @@ void FStreamingManagerTexture::CalcWantedMips( FStreamingTexture& StreamingTextu
 	}
 
 	float MinDistance = StreamingTexture.bUsesDynamicHeuristics ? FMath::Sqrt(StreamingTexture.DynamicMinDistanceSq) : FLT_MAX;
-	STAT(FFloatMipLevel PerfectWantedMips = WantedMips);
+	FFloatMipLevel PerfectWantedMips = WantedMips;
 
 	// Figure out miplevels to request based on handlers.
 	if ( StreamingTexture.MinAllowedMips != StreamingTexture.MaxAllowedMips )
@@ -4113,7 +4189,7 @@ void FStreamingManagerTexture::CalcWantedMips( FStreamingTexture& StreamingTextu
 			float HandlerDistance = FLT_MAX;
 			FFloatMipLevel HandlerWantedMips = TextureStreamingHandler->GetWantedMips( *this, StreamingTexture, HandlerDistance );
 			WantedMips = FMath::Max( WantedMips, HandlerWantedMips );
-			STAT( PerfectWantedMips = FMath::Max( PerfectWantedMips, HandlerWantedMips ) );
+			PerfectWantedMips = FMath::Max( PerfectWantedMips, HandlerWantedMips );
 			MinDistance = FMath::Min( MinDistance, HandlerDistance );
 		}
 
@@ -4140,7 +4216,7 @@ void FStreamingManagerTexture::CalcWantedMips( FStreamingTexture& StreamingTextu
 				FFloatMipLevel HandlerWantedMips = GetWantedMipsForOrphanedTexture(StreamingTexture, HandlerDistance);
 				WantedMips = FMath::Max(WantedMips, HandlerWantedMips);
 				MinDistance = FMath::Min(MinDistance, HandlerDistance);
-				STAT(PerfectWantedMips = FMath::Max(PerfectWantedMips, HandlerWantedMips));
+				PerfectWantedMips = FMath::Max(PerfectWantedMips, HandlerWantedMips);
 			}
 
 			// Still wasn't handled? Use the LastRenderTime handler as last resort.
@@ -4152,20 +4228,20 @@ void FStreamingManagerTexture::CalcWantedMips( FStreamingTexture& StreamingTextu
 				FFloatMipLevel HandlerWantedMips = FallbackStreamingHandler.GetWantedMips( *this, StreamingTexture, HandlerDistance );
 				WantedMips = FMath::Max( WantedMips, HandlerWantedMips );
 				MinDistance = FMath::Min( MinDistance, HandlerDistance );
-				STAT( PerfectWantedMips = FMath::Max( PerfectWantedMips, HandlerWantedMips ) );
+				PerfectWantedMips = FMath::Max( PerfectWantedMips, HandlerWantedMips );
 			}
 		}
 	}
 	// Stream in all allowed miplevels if requested.
 	else
 	{
-		STAT( PerfectWantedMips = FFloatMipLevel::FromMipLevel(StreamingTexture.MaxAllowedOptimalMips) );
+		PerfectWantedMips = FFloatMipLevel::FromMipLevel(StreamingTexture.MaxAllowedOptimalMips);
 		WantedMips = FFloatMipLevel::FromMipLevel(StreamingTexture.MaxAllowedMips);
 	}
 
 	StreamingTexture.WantedMips = WantedMips.ComputeMip(&StreamingTexture, ThreadSettings.MipBias, false);
 	StreamingTexture.MinDistance = MinDistance;
-	STAT( StreamingTexture.PerfectWantedMips = PerfectWantedMips.ComputeMip(&StreamingTexture, ThreadSettings.MipBias, true) );
+	StreamingTexture.PerfectWantedMips = PerfectWantedMips.ComputeMip(&StreamingTexture, ThreadSettings.MipBias, true);
 }
 
 /**
@@ -4438,6 +4514,14 @@ bool FStreamingManagerTexture::HandleListStreamingTexturesReportCommand( const T
 #endif
 
 #if !UE_BUILD_SHIPPING
+
+bool FStreamingManagerTexture::HandleResetMaxEverRequiredTexturesCommand(const TCHAR* Cmd, FOutputDevice& Ar)
+{
+	Ar.Logf(TEXT("OldMax: %u MaxEverRequired Reset."), MaxEverRequired);
+	ResetMaxEverRequired();	
+	return true;
+}
+
 bool FStreamingManagerTexture::HandleLightmapStreamingFactorCommand( const TCHAR* Cmd, FOutputDevice& Ar )
 {
 	FString FactorString(FParse::Token(Cmd, 0));
@@ -4746,6 +4830,10 @@ bool FStreamingManagerTexture::Exec( UWorld* InWorld, const TCHAR* Cmd, FOutputD
 	}
 #endif
 #if !UE_BUILD_SHIPPING
+	if (FParse::Command(&Cmd, TEXT("ResetMaxEverRequiredTextures")))
+	{
+		return HandleResetMaxEverRequiredTexturesCommand(Cmd, Ar);
+	}
 	if (FParse::Command(&Cmd,TEXT("LightmapStreamingFactor")))
 	{
 		return HandleLightmapStreamingFactorCommand( Cmd, Ar );
@@ -5343,7 +5431,10 @@ FFloatMipLevel FStreamingHandlerTextureStatic::GetWantedMips( FStreamingManagerT
 					for( int32 ViewIndex=0; ViewIndex < StreamingManager.ThreadNumViews() && !bShouldAbortLoop; ViewIndex++ )
 					{
 						const FStreamingViewInfo& ViewInfo = StreamingManager.ThreadGetView(ViewIndex);
-						const float ScreenSizeFloat = ViewInfo.ScreenSize * ViewInfo.BoostFactor * ScreenSizeFactor;
+
+						const float MaxEffectiveScreenSize = CVarStreamingScreenSizeEffectiveMax.GetValueOnAnyThread();
+						const float EffectiveScreenSize = (MaxEffectiveScreenSize > 0.0f) ? FMath::Min(MaxEffectiveScreenSize, ViewInfo.ScreenSize) : ViewInfo.ScreenSize;
+						const float ScreenSizeFloat = EffectiveScreenSize * ViewInfo.BoostFactor * ScreenSizeFactor;
 						const VectorRegister ScreenSize = VectorLoadFloat1( &ScreenSizeFloat );
 						const VectorRegister ViewOriginX = VectorLoadFloat1( &ViewInfo.ViewOrigin.X );
 						const VectorRegister ViewOriginY = VectorLoadFloat1( &ViewInfo.ViewOrigin.Y );
diff --git a/Engine/Source/Runtime/Engine/Private/CurveBase.cpp b/Engine/Source/Runtime/Engine/Private/CurveBase.cpp
--- a/Engine/Source/Runtime/Engine/Private/CurveBase.cpp
+++ b/Engine/Source/Runtime/Engine/Private/CurveBase.cpp
@@ -532,6 +532,11 @@ TArray<FRichCurveKey> FRichCurve::GetCopyOfKeys() const
 	return Keys;
 }
 
+const TArray<FRichCurveKey>& FRichCurve::GetConstRefOfKeys() const
+{
+	return Keys;
+}
+
 
 TArray<FRichCurveKey>::TConstIterator FRichCurve::GetKeyIterator() const
 {
diff --git a/Engine/Source/Runtime/Engine/Private/DemoNetDriver.cpp b/Engine/Source/Runtime/Engine/Private/DemoNetDriver.cpp
--- a/Engine/Source/Runtime/Engine/Private/DemoNetDriver.cpp
+++ b/Engine/Source/Runtime/Engine/Private/DemoNetDriver.cpp
@@ -725,12 +725,12 @@ bool UDemoNetDriver::InitListen( FNetworkNotify* InNotify, FURL& ListenURL, bool
 	return true;
 }
 
-bool UDemoNetDriver::IsRecording()
+bool UDemoNetDriver::IsRecording() const
 {
 	return ClientConnections.Num() > 0 && ClientConnections[0] != nullptr && ClientConnections[0]->State != USOCK_Closed;
 }
 
-bool UDemoNetDriver::IsPlaying()
+bool UDemoNetDriver::IsPlaying() const
 {
 	return ServerConnection != nullptr && ServerConnection->State != USOCK_Closed;
 }
diff --git a/Engine/Source/Runtime/Engine/Private/EndUserSettings.cpp b/Engine/Source/Runtime/Engine/Private/EndUserSettings.cpp
--- a/Engine/Source/Runtime/Engine/Private/EndUserSettings.cpp
+++ b/Engine/Source/Runtime/Engine/Private/EndUserSettings.cpp
@@ -10,6 +10,8 @@
 UEndUserSettings::UEndUserSettings(const FObjectInitializer& ObjectInitializer)
 	: Super(ObjectInitializer)
 	, bSendAnonymousUsageDataToEpic(true)
+	, bSendMeanTimeBetweenFailureDataToEpic(false)
+	, bAllowUserIdInUsageData(false)
 {
 }
 
@@ -88,7 +90,8 @@ void UEndUserSettings::OnSendAnonymousUsageDataToEpicChanged()
 		if (FEngineAnalytics::IsAvailable())
 		{
 			FEngineAnalytics::GetProvider().RecordEvent(FString("Engine.Privacy.EndUserOptOut"));
-			FEngineAnalytics::Shutdown();
+			const bool bIsEngineShutdown = false;
+			FEngineAnalytics::Shutdown(bIsEngineShutdown);
 		}
 	}
 }
diff --git a/Engine/Source/Runtime/Engine/Private/EngineAnalytics.cpp b/Engine/Source/Runtime/Engine/Private/EngineAnalytics.cpp
--- a/Engine/Source/Runtime/Engine/Private/EngineAnalytics.cpp
+++ b/Engine/Source/Runtime/Engine/Private/EngineAnalytics.cpp
@@ -6,11 +6,13 @@
 #include "Runtime/Analytics/Analytics/Public/Analytics.h"
 #include "Runtime/Analytics/Analytics/Public/Interfaces/IAnalyticsProvider.h"
 #include "GeneralProjectSettings.h"
+#include "EngineSessionManager.h"
 
 bool FEngineAnalytics::bIsInitialized;
 bool FEngineAnalytics::bIsEditorRun;
 bool FEngineAnalytics::bIsGameRun;
 TSharedPtr<IAnalyticsProvider> FEngineAnalytics::Analytics;
+TSharedPtr<FEngineSessionManager> FEngineAnalytics::SessionManager;
 
 /**
  * Engine analytics config log to initialize the engine analytics provider.
@@ -112,7 +114,7 @@ void FEngineAnalytics::Initialize()
 			if (Analytics.IsValid())
 			{
 				// Use an anonymous user id in-game
-				if (bIsGameRun)
+				if (bIsGameRun && GEngine->AreGameAnalyticsAnonymous())
 				{
 					FString AnonymousId;
 					if (!FPlatformMisc::GetStoredValue(TEXT("Epic Games"), TEXT("Unreal Engine/Privacy"), TEXT("AnonymousID"), AnonymousId) || AnonymousId.IsEmpty())
@@ -142,12 +144,37 @@ void FEngineAnalytics::Initialize()
 				bIsInitialized = true;
 			}
 		}
+
+		// Create the session manager singleton
+		if (!SessionManager.IsValid())
+		{
+			if (bIsEditorRun || (bIsGameRun && GEngine->AreGameMTBFEventsEnabled()))
+			{
+				SessionManager = MakeShareable(new FEngineSessionManager(bIsEditorRun ? EEngineSessionManagerMode::Editor : EEngineSessionManagerMode::Game));
+				SessionManager->Initialize();
+			}
+		}
 	}
 }
 
 
-void FEngineAnalytics::Shutdown()
+void FEngineAnalytics::Shutdown(bool bIsEngineShutdown)
 {
 	Analytics.Reset();
 	bIsInitialized = false;
+
+	// Destroy the session manager singleton if it exists
+	if (SessionManager.IsValid() && bIsEngineShutdown)
+	{
+		SessionManager->Shutdown();
+		SessionManager.Reset();
+	}
+}
+
+void FEngineAnalytics::Tick(float DeltaTime)
+{
+	if (SessionManager.IsValid())
+	{
+		SessionManager->Tick(DeltaTime);
+	}
 }
diff --git a/Engine/Source/Runtime/Engine/Private/EngineSessionManager.cpp b/Engine/Source/Runtime/Engine/Private/EngineSessionManager.cpp
new file mode 100644
--- /dev/null
+++ b/Engine/Source/Runtime/Engine/Private/EngineSessionManager.cpp
@@ -0,0 +1,495 @@
+// Copyright 1998-2015 Epic Games, Inc. All Rights Reserved.
+
+#include "EnginePrivate.h"
+#include "EngineSessionManager.h"
+#include "EngineAnalytics.h"
+#include "IAnalyticsProvider.h"
+#include "AnalyticsEventAttribute.h"
+#include "GeneralProjectSettings.h"
+#include "UserActivityTracking.h"
+
+#define LOCTEXT_NAMESPACE "SessionManager"
+
+namespace SessionManagerDefs
+{
+	static const FTimespan SessionRecordExpiration(30, 0, 0, 0);	// 30 days
+	static const FTimespan SessionRecordTimeout(0, 30, 0);			// 30 minutes
+	static const FTimespan GlobalLockWaitTimeout(0, 0, 0, 0, 500);	// 1/2 second
+	static const float HeartbeatPeriodSeconds(300.0f);				// 5 minutes
+	static const FString DefaultUserActivity(TEXT("Unknown"));
+	static const FString StoreId(TEXT("Epic Games"));
+	static const FString CrashSessionToken(TEXT("Crashed"));
+	static const FString DebuggerSessionToken(TEXT("Debugger"));
+	static const FString AbnormalSessionToken(TEXT("AbnormalShutdown"));
+	static const FString PS4SessionToken(TEXT("AbnormalShutdownPS4"));
+	static const FString SessionRecordListSection(TEXT("List"));
+	static const FString EditorSessionRecordSectionPrefix(TEXT("Unreal Engine/Editor Sessions/"));
+	static const FString GameSessionRecordSectionPrefix(TEXT("Unreal Engine/Game Sessions/"));
+	static const FString SessionsVersionString(TEXT("1_3"));
+	static const FString ModeStoreKey(TEXT("Mode"));
+	static const FString ProjectNameStoreKey(TEXT("ProjectName"));
+	static const FString CrashStoreKey(TEXT("IsCrash"));
+	static const FString DeactivatedStoreKey(TEXT("IsDeactivated"));
+	static const FString BackgroundStoreKey(TEXT("IsInBackground"));
+	static const FString EngineVersionStoreKey(TEXT("EngineVersion"));
+	static const FString TimestampStoreKey(TEXT("Timestamp"));
+	static const FString DebuggerStoreKey(TEXT("IsDebugger"));
+	static const FString UserActivityStoreKey(TEXT("CurrentUserActivity"));
+	static const FString GlobalLockName(TEXT("UE4_SessionManager_Lock"));
+	static const FString FalseValueString(TEXT("0"));
+	static const FString TrueValueString(TEXT("1"));
+	static const FString EditorValueString(TEXT("Editor"));
+	static const FString GameValueString(TEXT("Game"));
+	static const FString UnknownProjectValueString(TEXT("UnknownProject"));
+}
+
+namespace
+{
+	FString TimestampToString(FDateTime InTimestamp)
+	{
+		return LexicalConversion::ToString(InTimestamp.ToUnixTimestamp());
+	}
+
+	FDateTime StringToTimestamp(FString InString)
+	{
+		int64 TimestampUnix;
+		if (LexicalConversion::TryParseString(TimestampUnix, *InString))
+		{
+			return FDateTime::FromUnixTimestamp(TimestampUnix);
+		}
+		return FDateTime::MinValue();
+	}
+}
+
+/* FEngineSessionManager */
+
+void FEngineSessionManager::Initialize()
+{
+	// Register for crash and app state callbacks
+	FCoreDelegates::OnHandleSystemError.AddRaw(this, &FEngineSessionManager::OnCrashing);
+	FCoreDelegates::ApplicationHasReactivatedDelegate.AddRaw(this, &FEngineSessionManager::OnAppReactivate);
+	FCoreDelegates::ApplicationWillDeactivateDelegate.AddRaw(this, &FEngineSessionManager::OnAppDeactivate);
+	FCoreDelegates::ApplicationWillEnterBackgroundDelegate.AddRaw(this, &FEngineSessionManager::OnAppBackground);
+	FCoreDelegates::ApplicationHasEnteredForegroundDelegate.AddRaw(this, &FEngineSessionManager::OnAppForeground);
+	FUserActivityTracking::OnActivityChanged.AddRaw(this, &FEngineSessionManager::OnUserActivity);
+
+	const bool bFirstInitAttempt = true;
+	InitializeRecords(bFirstInitAttempt);
+}
+
+void FEngineSessionManager::InitializeRecords(bool bFirstAttempt)
+{
+	if (!FEngineAnalytics::IsAvailable())
+	{
+		return;
+	}
+
+	TArray<FSessionRecord> SessionRecordsToReport;
+
+	{
+		// Scoped lock
+		FSystemWideCriticalSection StoredValuesLock(SessionManagerDefs::GlobalLockName, bFirstAttempt ? SessionManagerDefs::GlobalLockWaitTimeout : FTimespan::Zero());
+
+		// Get list of sessions in storage
+		if (StoredValuesLock.IsValid() && BeginReadWriteRecords())
+		{
+			TArray<FSessionRecord> SessionRecordsToDelete;
+
+			// Attempt check each stored session
+			for (FSessionRecord& Record : SessionRecords)
+			{
+				FTimespan RecordAge = FDateTime::UtcNow() - Record.Timestamp;
+
+				if (Record.bCrashed)
+				{
+					// Crashed sessions
+					SessionRecordsToReport.Add(Record);
+					SessionRecordsToDelete.Add(Record);
+				}
+				else if (RecordAge > SessionManagerDefs::SessionRecordExpiration)
+				{
+					// Delete expired session records
+					SessionRecordsToDelete.Add(Record);
+				}
+				else if (RecordAge > SessionManagerDefs::SessionRecordTimeout)
+				{
+					// Timed out sessions
+					SessionRecordsToReport.Add(Record);
+					SessionRecordsToDelete.Add(Record);
+				}
+			}
+
+			for (FSessionRecord& DeletingRecord : SessionRecordsToDelete)
+			{
+				DeleteStoredRecord(DeletingRecord);
+			}
+
+			// Create a session record for this session
+			CreateAndWriteRecordForSession();
+
+			// Update and release list of sessions in storage
+			EndReadWriteRecords();
+
+			bInitializedRecords = true;
+		}
+	}
+
+	for (FSessionRecord& ReportingSession : SessionRecordsToReport)
+	{
+		// Send error report for session that timed out or crashed
+		SendAbnormalShutdownReport(ReportingSession);
+	}
+}
+
+void FEngineSessionManager::Tick(float DeltaTime)
+{
+	HeartbeatTimeElapsed += DeltaTime;
+
+	if (HeartbeatTimeElapsed > SessionManagerDefs::HeartbeatPeriodSeconds && !bShutdown)
+	{
+		HeartbeatTimeElapsed = 0.0f;
+
+		if (!bInitializedRecords)
+		{
+			// Try late initialization
+			const bool bFirstInitAttempt = false;
+			InitializeRecords(bFirstInitAttempt);
+		}
+
+		// Update timestamp in the session record for this session 
+		if (bInitializedRecords)
+		{
+			CurrentSession.Timestamp = FDateTime::UtcNow();
+
+			FPlatformMisc::SetStoredValue(SessionManagerDefs::StoreId, CurrentSessionSectionName, SessionManagerDefs::TimestampStoreKey, TimestampToString(CurrentSession.Timestamp));
+		}
+	}
+}
+
+void FEngineSessionManager::Shutdown()
+{
+	FCoreDelegates::OnHandleSystemError.RemoveAll(this);
+	FCoreDelegates::ApplicationHasReactivatedDelegate.RemoveAll(this);
+	FCoreDelegates::ApplicationWillDeactivateDelegate.RemoveAll(this);
+	FCoreDelegates::ApplicationWillEnterBackgroundDelegate.RemoveAll(this);
+	FCoreDelegates::ApplicationHasEnteredForegroundDelegate.RemoveAll(this);
+
+	// Clear the session record for this session
+	if (bInitializedRecords)
+	{
+		if (!CurrentSession.bCrashed)
+		{
+			FPlatformMisc::DeleteStoredValue(SessionManagerDefs::StoreId, CurrentSessionSectionName, SessionManagerDefs::ModeStoreKey);
+			FPlatformMisc::DeleteStoredValue(SessionManagerDefs::StoreId, CurrentSessionSectionName, SessionManagerDefs::ProjectNameStoreKey);
+			FPlatformMisc::DeleteStoredValue(SessionManagerDefs::StoreId, CurrentSessionSectionName, SessionManagerDefs::CrashStoreKey);
+			FPlatformMisc::DeleteStoredValue(SessionManagerDefs::StoreId, CurrentSessionSectionName, SessionManagerDefs::EngineVersionStoreKey);
+			FPlatformMisc::DeleteStoredValue(SessionManagerDefs::StoreId, CurrentSessionSectionName, SessionManagerDefs::TimestampStoreKey);
+			FPlatformMisc::DeleteStoredValue(SessionManagerDefs::StoreId, CurrentSessionSectionName, SessionManagerDefs::DebuggerStoreKey);
+			FPlatformMisc::DeleteStoredValue(SessionManagerDefs::StoreId, CurrentSessionSectionName, SessionManagerDefs::DeactivatedStoreKey);
+			FPlatformMisc::DeleteStoredValue(SessionManagerDefs::StoreId, CurrentSessionSectionName, SessionManagerDefs::BackgroundStoreKey);
+			FPlatformMisc::DeleteStoredValue(SessionManagerDefs::StoreId, CurrentSessionSectionName, SessionManagerDefs::UserActivityStoreKey);
+		}
+
+		bInitializedRecords = false;
+		bShutdown = true;
+	}
+}
+
+bool FEngineSessionManager::BeginReadWriteRecords()
+{
+	SessionRecords.Empty();
+
+	// Lock and read the list of sessions in storage
+	FString ListSectionName = GetStoreSectionString(SessionManagerDefs::SessionRecordListSection);
+
+	// Write list to SessionRecords member
+	FString SessionListString;
+	FPlatformMisc::GetStoredValue(SessionManagerDefs::StoreId, ListSectionName, TEXT("SessionList"), SessionListString);
+
+	// Parse SessionListString for session ids
+	TArray<FString> SessionIds;
+	SessionListString.ParseIntoArray(SessionIds, TEXT(","));
+
+	// Retrieve all the sessions in the list from storage
+	for (FString& SessionId : SessionIds)
+	{
+		FString SectionName = GetStoreSectionString(SessionId);
+
+		FString IsCrashString;
+		FString EngineVersionString;
+		FString TimestampString;
+		FString IsDebuggerString;
+
+		// Read manditory values
+		if (FPlatformMisc::GetStoredValue(SessionManagerDefs::StoreId, SectionName, SessionManagerDefs::CrashStoreKey, IsCrashString) &&
+			FPlatformMisc::GetStoredValue(SessionManagerDefs::StoreId, SectionName, SessionManagerDefs::EngineVersionStoreKey, EngineVersionString) &&
+			FPlatformMisc::GetStoredValue(SessionManagerDefs::StoreId, SectionName, SessionManagerDefs::TimestampStoreKey, TimestampString) &&
+			FPlatformMisc::GetStoredValue(SessionManagerDefs::StoreId, SectionName, SessionManagerDefs::DebuggerStoreKey, IsDebuggerString))
+		{
+			// Read optional values
+			FString ModeString = SessionManagerDefs::EditorValueString;
+			FPlatformMisc::GetStoredValue(SessionManagerDefs::StoreId, SectionName, SessionManagerDefs::ModeStoreKey, ModeString);
+			FString ProjectName = SessionManagerDefs::UnknownProjectValueString;
+			FPlatformMisc::GetStoredValue(SessionManagerDefs::StoreId, SectionName, SessionManagerDefs::ProjectNameStoreKey, ProjectName);
+			FString IsDeactivatedString = SessionManagerDefs::FalseValueString;
+			FPlatformMisc::GetStoredValue(SessionManagerDefs::StoreId, SectionName, SessionManagerDefs::DeactivatedStoreKey, IsDeactivatedString);
+			FString IsInBackgroundString = SessionManagerDefs::FalseValueString;
+			FPlatformMisc::GetStoredValue(SessionManagerDefs::StoreId, SectionName, SessionManagerDefs::BackgroundStoreKey, IsInBackgroundString);
+			FString UserActivityString = SessionManagerDefs::DefaultUserActivity;
+			FPlatformMisc::GetStoredValue(SessionManagerDefs::StoreId, SectionName, SessionManagerDefs::UserActivityStoreKey, UserActivityString);
+
+			// Create new record from the read values
+			FSessionRecord NewRecord;
+			NewRecord.SessionId = SessionId;
+			NewRecord.Mode = ModeString == SessionManagerDefs::EditorValueString ? EEngineSessionManagerMode::Editor : EEngineSessionManagerMode::Game;
+			NewRecord.ProjectName = ProjectName;
+			NewRecord.EngineVersion = EngineVersionString;
+			NewRecord.Timestamp = StringToTimestamp(TimestampString);
+			NewRecord.bCrashed = IsCrashString == SessionManagerDefs::TrueValueString;
+			NewRecord.bIsDebugger = IsDebuggerString == SessionManagerDefs::TrueValueString;
+			NewRecord.bIsDeactivated = IsDeactivatedString == SessionManagerDefs::TrueValueString;
+			NewRecord.bIsInBackground = IsInBackgroundString == SessionManagerDefs::TrueValueString;
+			NewRecord.CurrentUserActivity = UserActivityString;
+
+			SessionRecords.Add(NewRecord);
+		}
+		else
+		{
+			// Clean up orphaned values, if there are any
+			FPlatformMisc::DeleteStoredValue(SessionManagerDefs::StoreId, SectionName, SessionManagerDefs::ModeStoreKey);
+			FPlatformMisc::DeleteStoredValue(SessionManagerDefs::StoreId, SectionName, SessionManagerDefs::ProjectNameStoreKey);
+			FPlatformMisc::DeleteStoredValue(SessionManagerDefs::StoreId, SectionName, SessionManagerDefs::CrashStoreKey);
+			FPlatformMisc::DeleteStoredValue(SessionManagerDefs::StoreId, SectionName, SessionManagerDefs::EngineVersionStoreKey);
+			FPlatformMisc::DeleteStoredValue(SessionManagerDefs::StoreId, SectionName, SessionManagerDefs::TimestampStoreKey);
+			FPlatformMisc::DeleteStoredValue(SessionManagerDefs::StoreId, SectionName, SessionManagerDefs::DebuggerStoreKey);
+			FPlatformMisc::DeleteStoredValue(SessionManagerDefs::StoreId, SectionName, SessionManagerDefs::DeactivatedStoreKey);
+			FPlatformMisc::DeleteStoredValue(SessionManagerDefs::StoreId, SectionName, SessionManagerDefs::BackgroundStoreKey);
+			FPlatformMisc::DeleteStoredValue(SessionManagerDefs::StoreId, SectionName, SessionManagerDefs::UserActivityStoreKey);
+		}
+	}
+
+	return true;
+}
+
+void FEngineSessionManager::EndReadWriteRecords()
+{
+	// Update the list of sessions in storage to match SessionRecords
+	FString SessionListString;
+	if (SessionRecords.Num() > 0)
+	{
+		for (FSessionRecord& Session : SessionRecords)
+		{
+			SessionListString.Append(Session.SessionId);
+			SessionListString.Append(TEXT(","));
+		}
+		SessionListString.RemoveAt(SessionListString.Len() - 1);
+	}
+
+	FString ListSectionName = GetStoreSectionString(SessionManagerDefs::SessionRecordListSection);
+	FPlatformMisc::SetStoredValue(SessionManagerDefs::StoreId, ListSectionName, TEXT("SessionList"), SessionListString);
+
+	// Clear SessionRecords member
+	SessionRecords.Empty();
+}
+
+void FEngineSessionManager::DeleteStoredRecord(const FSessionRecord& Record)
+{
+	// Delete the session record in storage
+	FString SessionId = Record.SessionId;
+	FString SectionName = GetStoreSectionString(SessionId);
+
+	FPlatformMisc::DeleteStoredValue(SessionManagerDefs::StoreId, SectionName, SessionManagerDefs::ModeStoreKey);
+	FPlatformMisc::DeleteStoredValue(SessionManagerDefs::StoreId, SectionName, SessionManagerDefs::ProjectNameStoreKey);
+	FPlatformMisc::DeleteStoredValue(SessionManagerDefs::StoreId, SectionName, SessionManagerDefs::CrashStoreKey);
+	FPlatformMisc::DeleteStoredValue(SessionManagerDefs::StoreId, SectionName, SessionManagerDefs::EngineVersionStoreKey);
+	FPlatformMisc::DeleteStoredValue(SessionManagerDefs::StoreId, SectionName, SessionManagerDefs::TimestampStoreKey);
+	FPlatformMisc::DeleteStoredValue(SessionManagerDefs::StoreId, SectionName, SessionManagerDefs::DebuggerStoreKey);
+	FPlatformMisc::DeleteStoredValue(SessionManagerDefs::StoreId, SectionName, SessionManagerDefs::DeactivatedStoreKey);
+	FPlatformMisc::DeleteStoredValue(SessionManagerDefs::StoreId, SectionName, SessionManagerDefs::BackgroundStoreKey);
+	FPlatformMisc::DeleteStoredValue(SessionManagerDefs::StoreId, SectionName, SessionManagerDefs::UserActivityStoreKey);
+
+	// Remove the session record from SessionRecords list
+	SessionRecords.RemoveAll([&SessionId](const FSessionRecord& X){ return X.SessionId == SessionId; });
+}
+
+void FEngineSessionManager::SendAbnormalShutdownReport(const FSessionRecord& Record)
+{
+#if PLATFORM_WINDOWS
+	const FString PlatformName(TEXT("Windows"));
+#elif PLATFORM_MAC
+	const FString PlatformName(TEXT("Mac"));
+#elif PLATFORM_LINUX
+	const FString PlatformName(TEXT("Linux"));
+#elif PLATFORM_PS4
+	const FString PlatformName(TEXT("PS4"));
+	if (Record.bIsDeactivated && !Record.bCrashed)
+	{
+		// Shutting down in deactivated state on PS4 is normal - don't report it
+		return;
+	}
+#elif PLATFORM_XBOXONE
+	const FString PlatformName(TEXT("XBoxOne"));
+	if (Record.bIsInBackground && !Record.bCrashed)
+	{
+		// Shutting down in background state on XB1 is normal - don't report it
+		return;
+	}
+#else
+	const FString PlatformName(TEXT("Unknown"));
+	return; // TODO: CWood: disabled on other platforms
+#endif
+
+	FGuid SessionId;
+	FString SessionIdString = Record.SessionId;
+	if (FGuid::Parse(SessionIdString, SessionId))
+	{
+		// convert session guid to one with braces for sending to analytics
+		SessionIdString = SessionId.ToString(EGuidFormats::DigitsWithHyphensInBraces);
+	}
+
+#if !PLATFORM_PS4
+	FString ShutdownTypeString = Record.bCrashed ? SessionManagerDefs::CrashSessionToken :
+		(Record.bIsDebugger ? SessionManagerDefs::DebuggerSessionToken : SessionManagerDefs::AbnormalSessionToken);
+#else
+	// PS4 cannot set the crash flag so report abnormal shutdowns with a specific token meaning "crash or abnormal shutdown".
+	FString ShutdownTypeString = Record.bIsDebugger ? SessionManagerDefs::DebuggerSessionToken : SessionManagerDefs::PS4SessionToken;
+#endif
+
+	FString RunTypeString = Record.Mode == EEngineSessionManagerMode::Editor ? SessionManagerDefs::EditorValueString : SessionManagerDefs::GameValueString;
+
+	TArray< FAnalyticsEventAttribute > AbnormalShutdownAttributes;
+	AbnormalShutdownAttributes.Add(FAnalyticsEventAttribute(TEXT("RunType"), RunTypeString));
+	AbnormalShutdownAttributes.Add(FAnalyticsEventAttribute(TEXT("ProjectName"), Record.ProjectName));
+	AbnormalShutdownAttributes.Add(FAnalyticsEventAttribute(TEXT("Platform"), PlatformName));
+	AbnormalShutdownAttributes.Add(FAnalyticsEventAttribute(TEXT("SessionId"), SessionIdString));
+	AbnormalShutdownAttributes.Add(FAnalyticsEventAttribute(TEXT("EngineVersion"), Record.EngineVersion));
+	AbnormalShutdownAttributes.Add(FAnalyticsEventAttribute(TEXT("ShutdownType"), ShutdownTypeString));
+	AbnormalShutdownAttributes.Add(FAnalyticsEventAttribute(TEXT("Timestamp"), Record.Timestamp.ToIso8601()));
+	AbnormalShutdownAttributes.Add(FAnalyticsEventAttribute(TEXT("CurrentUserActivity"), Record.CurrentUserActivity));
+
+	FEngineAnalytics::GetProvider().RecordEvent(TEXT("Engine.AbnormalShutdown"), AbnormalShutdownAttributes);
+}
+
+void FEngineSessionManager::CreateAndWriteRecordForSession()
+{
+	FGuid SessionId;
+	if (FGuid::Parse(FEngineAnalytics::GetProvider().GetSessionID(), SessionId))
+	{
+		// convert session guid to one without braces or other chars that might not be suitable for storage
+		CurrentSession.SessionId = SessionId.ToString(EGuidFormats::DigitsWithHyphens);
+	}
+	else
+	{
+		CurrentSession.SessionId = FEngineAnalytics::GetProvider().GetSessionID();
+	}
+
+	const UGeneralProjectSettings& ProjectSettings = *GetDefault<UGeneralProjectSettings>();
+
+	CurrentSession.Mode = Mode;
+	CurrentSession.ProjectName = ProjectSettings.ProjectName;
+	CurrentSession.EngineVersion = FEngineVersion::Current().ToString(EVersionComponent::Changelist);
+	CurrentSession.Timestamp = FDateTime::UtcNow();
+	CurrentSession.bIsDebugger = FPlatformMisc::IsDebuggerPresent();
+	CurrentSession.CurrentUserActivity = GetUserActivityString();
+	CurrentSessionSectionName = GetStoreSectionString(CurrentSession.SessionId);
+
+	FString ModeString = CurrentSession.Mode == EEngineSessionManagerMode::Editor ? SessionManagerDefs::EditorValueString : SessionManagerDefs::GameValueString;
+	FString IsDebuggerString = CurrentSession.bIsDebugger ? SessionManagerDefs::TrueValueString : SessionManagerDefs::FalseValueString;
+	FString IsDeactivatedString = CurrentSession.bIsDeactivated ? SessionManagerDefs::TrueValueString : SessionManagerDefs::FalseValueString;
+	FString IsInBackgroundString = CurrentSession.bIsInBackground ? SessionManagerDefs::TrueValueString : SessionManagerDefs::FalseValueString;
+
+	FPlatformMisc::SetStoredValue(SessionManagerDefs::StoreId, CurrentSessionSectionName, SessionManagerDefs::ModeStoreKey, ModeString);
+	FPlatformMisc::SetStoredValue(SessionManagerDefs::StoreId, CurrentSessionSectionName, SessionManagerDefs::ProjectNameStoreKey, CurrentSession.ProjectName);
+	FPlatformMisc::SetStoredValue(SessionManagerDefs::StoreId, CurrentSessionSectionName, SessionManagerDefs::CrashStoreKey, SessionManagerDefs::FalseValueString);
+	FPlatformMisc::SetStoredValue(SessionManagerDefs::StoreId, CurrentSessionSectionName, SessionManagerDefs::EngineVersionStoreKey, CurrentSession.EngineVersion);
+	FPlatformMisc::SetStoredValue(SessionManagerDefs::StoreId, CurrentSessionSectionName, SessionManagerDefs::TimestampStoreKey, TimestampToString(CurrentSession.Timestamp));
+	FPlatformMisc::SetStoredValue(SessionManagerDefs::StoreId, CurrentSessionSectionName, SessionManagerDefs::DebuggerStoreKey, IsDebuggerString);
+	FPlatformMisc::SetStoredValue(SessionManagerDefs::StoreId, CurrentSessionSectionName, SessionManagerDefs::DeactivatedStoreKey, IsDeactivatedString);
+	FPlatformMisc::SetStoredValue(SessionManagerDefs::StoreId, CurrentSessionSectionName, SessionManagerDefs::BackgroundStoreKey, IsInBackgroundString);
+	FPlatformMisc::SetStoredValue(SessionManagerDefs::StoreId, CurrentSessionSectionName, SessionManagerDefs::UserActivityStoreKey, CurrentSession.CurrentUserActivity);
+
+	SessionRecords.Add(CurrentSession);
+}
+
+void FEngineSessionManager::OnCrashing()
+{
+	if (!CurrentSession.bCrashed && bInitializedRecords)
+	{
+		CurrentSession.bCrashed = true;
+		FPlatformMisc::SetStoredValue(SessionManagerDefs::StoreId, CurrentSessionSectionName, SessionManagerDefs::CrashStoreKey, SessionManagerDefs::TrueValueString);
+	}
+}
+
+void FEngineSessionManager::OnAppReactivate()
+{
+	if (CurrentSession.bIsDeactivated)
+	{
+		CurrentSession.bIsDeactivated = false;
+		FPlatformMisc::SetStoredValue(SessionManagerDefs::StoreId, CurrentSessionSectionName, SessionManagerDefs::DeactivatedStoreKey, SessionManagerDefs::FalseValueString);
+	}
+}
+
+void FEngineSessionManager::OnAppDeactivate()
+{
+	if (!CurrentSession.bIsDeactivated)
+	{
+		CurrentSession.bIsDeactivated = true;
+		FPlatformMisc::SetStoredValue(SessionManagerDefs::StoreId, CurrentSessionSectionName, SessionManagerDefs::DeactivatedStoreKey, SessionManagerDefs::TrueValueString);
+	}
+}
+
+void FEngineSessionManager::OnAppBackground()
+{
+	if (!CurrentSession.bIsInBackground)
+	{
+		CurrentSession.bIsInBackground = true;
+		FPlatformMisc::SetStoredValue(SessionManagerDefs::StoreId, CurrentSessionSectionName, SessionManagerDefs::BackgroundStoreKey, SessionManagerDefs::TrueValueString);
+	}
+}
+
+void FEngineSessionManager::OnAppForeground()
+{
+	if (CurrentSession.bIsInBackground)
+	{
+		CurrentSession.bIsInBackground = false;
+		FPlatformMisc::SetStoredValue(SessionManagerDefs::StoreId, CurrentSessionSectionName, SessionManagerDefs::BackgroundStoreKey, SessionManagerDefs::FalseValueString);
+	}
+}
+
+FString FEngineSessionManager::GetStoreSectionString(FString InSuffix)
+{
+	check(Mode == EEngineSessionManagerMode::Editor || Mode == EEngineSessionManagerMode::Game)
+
+	if (Mode == EEngineSessionManagerMode::Editor)
+	{
+		return FString::Printf(TEXT("%s%s/%s"), *SessionManagerDefs::EditorSessionRecordSectionPrefix, *SessionManagerDefs::SessionsVersionString, *InSuffix);
+	}
+	else
+	{
+		const UGeneralProjectSettings& ProjectSettings = *GetDefault<UGeneralProjectSettings>();
+		return FString::Printf(TEXT("%s%s/%s/%s"), *SessionManagerDefs::GameSessionRecordSectionPrefix, *SessionManagerDefs::SessionsVersionString, *ProjectSettings.ProjectName, *InSuffix);
+	}
+}
+
+void FEngineSessionManager::OnUserActivity(const FUserActivity& UserActivity)
+{
+	if (!CurrentSession.bCrashed && bInitializedRecords)
+	{
+		CurrentSession.CurrentUserActivity = GetUserActivityString();
+		FPlatformMisc::SetStoredValue(SessionManagerDefs::StoreId, CurrentSessionSectionName, SessionManagerDefs::UserActivityStoreKey, CurrentSession.CurrentUserActivity);
+	}
+}
+
+FString FEngineSessionManager::GetUserActivityString() const
+{
+	const FUserActivity& UserActivity = FUserActivityTracking::GetUserActivity();
+	
+	if (UserActivity.ActionName.IsEmpty())
+	{
+		return SessionManagerDefs::DefaultUserActivity;
+	}
+
+	return UserActivity.ActionName;
+}
+
+#undef LOCTEXT_NAMESPACE
diff --git a/Engine/Source/Runtime/Engine/Private/EngineSessionManager.h b/Engine/Source/Runtime/Engine/Private/EngineSessionManager.h
new file mode 100644
--- /dev/null
+++ b/Engine/Source/Runtime/Engine/Private/EngineSessionManager.h
@@ -0,0 +1,71 @@
+// Copyright 1998-2015 Epic Games, Inc. All Rights Reserved.
+
+#pragma once
+
+enum class EEngineSessionManagerMode
+{
+	Editor,
+	Game
+};
+
+/* Handles writing session records to platform's storage to track crashed and timed-out editor sessions */
+class FEngineSessionManager
+{
+public:
+	FEngineSessionManager(EEngineSessionManagerMode InMode)
+		: Mode(InMode)
+		, bInitializedRecords(false)
+		, bShutdown(false)
+		, HeartbeatTimeElapsed(0.0f)
+	{
+		CurrentSession.bCrashed = false;
+		CurrentSession.bIsDeactivated = false;
+		CurrentSession.bIsInBackground = false;
+	}
+
+	void Initialize();
+
+	void Tick(float DeltaTime);
+
+	void Shutdown();
+
+private:
+	struct FSessionRecord
+	{
+		FString SessionId;
+		EEngineSessionManagerMode Mode;
+		FString ProjectName;
+		FString EngineVersion;
+		FDateTime Timestamp;
+		bool bCrashed;
+		bool bIsDebugger;
+		bool bIsDeactivated;
+		bool bIsInBackground;
+		FString CurrentUserActivity;
+	};
+
+private:
+	void InitializeRecords(bool bFirstAttempt);
+	bool BeginReadWriteRecords();
+	void EndReadWriteRecords();
+	void DeleteStoredRecord(const FSessionRecord& Record);
+	void SendAbnormalShutdownReport(const FSessionRecord& Record);
+	void CreateAndWriteRecordForSession();
+	void OnCrashing();
+	void OnAppReactivate();
+	void OnAppDeactivate();
+	void OnAppBackground();
+	void OnAppForeground();
+	FString GetStoreSectionString(FString InSuffix);
+	void OnUserActivity(const FUserActivity& UserActivity);
+	FString GetUserActivityString() const;
+
+private:
+	EEngineSessionManagerMode Mode;
+	bool bInitializedRecords;
+	bool bShutdown;
+	float HeartbeatTimeElapsed;
+	FSessionRecord CurrentSession;
+	FString CurrentSessionSectionName;
+	TArray<FSessionRecord> SessionRecords;
+};
diff --git a/Engine/Source/Runtime/Engine/Private/GameUserSettings.cpp b/Engine/Source/Runtime/Engine/Private/GameUserSettings.cpp
--- a/Engine/Source/Runtime/Engine/Private/GameUserSettings.cpp
+++ b/Engine/Source/Runtime/Engine/Private/GameUserSettings.cpp
@@ -6,6 +6,7 @@
 #include "Scalability.h"
 #include "GameFramework/GameUserSettings.h"
 #include "Engine/GameEngine.h"
+#include "Kismet/KismetSystemLibrary.h"
 
 extern EWindowMode::Type GetWindowModeType(EWindowMode::Type WindowMode);
 
@@ -32,10 +33,25 @@ FIntPoint UGameUserSettings::GetLastConfirmedScreenResolution() const
 	return FIntPoint(LastUserConfirmedResolutionSizeX, LastUserConfirmedResolutionSizeY);
 }
 
-void UGameUserSettings::SetScreenResolution( FIntPoint Resolution )
+FIntPoint UGameUserSettings::GetDesktopResolution() const
+{
+	FDisplayMetrics DisplayMetrics;
+	if (FSlateApplication::IsInitialized())
+	{
+		FSlateApplication::Get().GetInitialDisplayMetrics(DisplayMetrics);
+	}
+	else
+	{
+		FDisplayMetrics::GetDisplayMetrics(DisplayMetrics);
+	}
+	return FIntPoint(DisplayMetrics.PrimaryDisplayWidth, DisplayMetrics.PrimaryDisplayHeight);
+}
+
+void UGameUserSettings::SetScreenResolution(FIntPoint Resolution)
 {
 	ResolutionSizeX = Resolution.X;
 	ResolutionSizeY = Resolution.Y;
+	UpdateMinResolutionScaling();
 }
 
 EWindowMode::Type UGameUserSettings::GetFullscreenMode() const
@@ -66,6 +82,8 @@ void UGameUserSettings::SetFullscreenMode( EWindowMode::Type InFullscreenMode )
 			FullscreenMode = 2;
 			break;
 	}
+
+	UpdateMinResolutionScaling();
 }
 
 void UGameUserSettings::SetVSyncEnabled( bool bEnable )
@@ -139,8 +157,11 @@ void UGameUserSettings::SetToDefaults()
 	WindowPosY = GetDefaultWindowPosition().Y;
 	FullscreenMode = GetDefaultWindowMode();
 	FrameRateLimit = 0.0f;
+	MinResolutionScale = Scalability::MinResolutionScale;
 
 	ScalabilityQuality.SetDefaults();
+
+	UpdateMinResolutionScaling();
 }
 
 bool UGameUserSettings::IsVersionValid()
@@ -153,6 +174,14 @@ void UGameUserSettings::UpdateVersion()
 	Version = UE_GAMEUSERSETTINGS_VERSION;
 }
 
+void UGameUserSettings::UpdateMinResolutionScaling()
+{
+	const int32 MinHeight = UKismetSystemLibrary::GetMinYResolutionForUI();
+	const int32 ScreenHeight = FullscreenMode == EWindowMode::WindowedFullscreen ? GetDesktopResolution().Y : ResolutionSizeY;
+	MinResolutionScale = FMath::Max<int32>(Scalability::MinResolutionScale, ((float)MinHeight / (float)ScreenHeight) * 100);
+	ScalabilityQuality.ResolutionQuality = FMath::Max(ScalabilityQuality.ResolutionQuality, MinResolutionScale);
+}
+
 void UGameUserSettings::ValidateSettings()
 {
 	// Should we wipe all user settings?
@@ -358,7 +387,7 @@ FIntPoint UGameUserSettings::GetDefaultWindowPosition()
 
 EWindowMode::Type UGameUserSettings::GetDefaultWindowMode()
 {
-	return EWindowMode::Windowed;
+	return EWindowMode::WindowedFullscreen;
 }
 
 void UGameUserSettings::ResetToCurrentSettings()
@@ -382,6 +411,8 @@ void UGameUserSettings::ResetToCurrentSettings()
 
 		// Reset the quality settings to the current levels
 		ScalabilityQuality = Scalability::GetQualityLevels();
+
+		UpdateMinResolutionScaling();
 	}
 }
 
@@ -419,7 +450,8 @@ float UGameUserSettings::GetFrameRateLimit() const
 void UGameUserSettings::SetOverallScalabilityLevel(int32 Value)
 {
 	Value = FMath::Clamp(Value, 0, 3);
-	ScalabilityQuality.SetFromSingleQualityLevel(Value);
+
+	ScalabilityQuality.SetFromSingleQualityLevel(Value);	
 }
 
 int32 UGameUserSettings::GetOverallScalabilityLevel() const
@@ -430,19 +462,19 @@ int32 UGameUserSettings::GetOverallScalabilityLevel() const
 void UGameUserSettings::GetResolutionScaleInformation(float& CurrentScaleNormalized, int32& CurrentScaleValue, int32& MinScaleValue, int32& MaxScaleValue) const
 {
 	CurrentScaleValue = ScalabilityQuality.ResolutionQuality;
-	MinScaleValue = Scalability::MinResolutionScale;
+	MinScaleValue = MinResolutionScale;
 	MaxScaleValue = Scalability::MaxResolutionScale;
 	CurrentScaleNormalized = ((float)CurrentScaleValue - (float)MinScaleValue) / (float)(MaxScaleValue - MinScaleValue);
 }
 
 void UGameUserSettings::SetResolutionScaleValue(int32 NewScaleValue)
 {
-	ScalabilityQuality.ResolutionQuality = FMath::Clamp(NewScaleValue, Scalability::MinResolutionScale, Scalability::MaxResolutionScale);
+	ScalabilityQuality.ResolutionQuality = FMath::Clamp(NewScaleValue, MinResolutionScale, Scalability::MaxResolutionScale);
 }
 
 void UGameUserSettings::SetResolutionScaleNormalized(float NewScaleNormalized)
 {
-	const int32 RemappedValue = (int32)FMath::Lerp((float)Scalability::MinResolutionScale, (float)Scalability::MaxResolutionScale, NewScaleNormalized);
+	const int32 RemappedValue = (int32)FMath::Lerp((float)MinResolutionScale, (float)Scalability::MaxResolutionScale, NewScaleNormalized);
 	SetResolutionScaleValue(RemappedValue);
 }
 
diff --git a/Engine/Source/Runtime/Engine/Private/GameViewportClient.cpp b/Engine/Source/Runtime/Engine/Private/GameViewportClient.cpp
--- a/Engine/Source/Runtime/Engine/Private/GameViewportClient.cpp
+++ b/Engine/Source/Runtime/Engine/Private/GameViewportClient.cpp
@@ -305,6 +305,7 @@ void UGameViewportClient::Init(struct FWorldContext& WorldContext, UGameInstance
 	AddCursor(EMouseCursor::Default, UISettings->DefaultCursor);
 	AddCursor(EMouseCursor::TextEditBeam, UISettings->TextEditBeamCursor);
 	AddCursor(EMouseCursor::Crosshairs, UISettings->CrosshairsCursor);
+	AddCursor(EMouseCursor::Hand, UISettings->HandCursor);
 	AddCursor(EMouseCursor::GrabHand, UISettings->GrabHandCursor);
 	AddCursor(EMouseCursor::GrabHandClosed, UISettings->GrabHandClosedCursor);
 	AddCursor(EMouseCursor::SlashedCircle, UISettings->SlashedCircleCursor);
@@ -2758,6 +2759,7 @@ bool UGameViewportClient::HandleToggleFullscreenCommand()
 		if( UserSettings != nullptr )
 		{
 			UserSettings->SetFullscreenMode( FullScreenMode );
+			UserSettings->ConfirmVideoMode();
 		}
 	}
 
diff --git a/Engine/Source/Runtime/Engine/Private/Interpolation.cpp b/Engine/Source/Runtime/Engine/Private/Interpolation.cpp
--- a/Engine/Source/Runtime/Engine/Private/Interpolation.cpp
+++ b/Engine/Source/Runtime/Engine/Private/Interpolation.cpp
@@ -2089,7 +2089,7 @@ void UInterpGroup::PostLoad()
 
 void UInterpGroup::UpdateGroup(float NewPosition, UInterpGroupInst* GrInst, bool bPreview, bool bJump)
 {
-	check( InterpTracks.Num() == GrInst->TrackInst.Num() );
+	checkf( InterpTracks.Num() == GrInst->TrackInst.Num(), TEXT("UpdateGroup track mismatch! Outer = %s"), *GetNameSafe(GetOuter()) );
 
 	// Update animation state of Actor.
 #if 0
diff --git a/Engine/Source/Runtime/Engine/Private/KismetSystemLibrary.cpp b/Engine/Source/Runtime/Engine/Private/KismetSystemLibrary.cpp
--- a/Engine/Source/Runtime/Engine/Private/KismetSystemLibrary.cpp
+++ b/Engine/Source/Runtime/Engine/Private/KismetSystemLibrary.cpp
@@ -17,6 +17,7 @@
 #include "OnlineSubsystemTypes.h"
 #include "OnlineSubsystemUtils.h"
 #include "OnlineIdentityInterface.h"
+#include "UserActivityTracking.h"
 
 //////////////////////////////////////////////////////////////////////////
 // UKismetSystemLibrary
@@ -3021,16 +3022,38 @@ int32 UKismetSystemLibrary::GetRenderingMaterialQualityLevel()
 
 bool UKismetSystemLibrary::GetSupportedFullscreenResolutions(TArray<FIntPoint>& Resolutions)
 {
+	uint32 MinYResolution = UKismetSystemLibrary::GetMinYResolutionForUI();
+
 	FScreenResolutionArray SupportedResolutions;
 	if ( RHIGetAvailableResolutions(SupportedResolutions, true) )
 	{
+		uint32 LargestY = 0;
 		for ( const FScreenResolutionRHI& SupportedResolution : SupportedResolutions )
 		{
-			FIntPoint Resolution;
-			Resolution.X = SupportedResolution.Width;
-			Resolution.Y = SupportedResolution.Height;
+			LargestY = FMath::Max(LargestY, SupportedResolution.Height);
+			if(SupportedResolution.Height >= MinYResolution)
+			{
+				FIntPoint Resolution;
+				Resolution.X = SupportedResolution.Width;
+				Resolution.Y = SupportedResolution.Height;
+
+				Resolutions.Add(Resolution);
+			}
+		}
+		if(!Resolutions.Num())
+		{
+			// if we don't have any resolution, we take the largest one(s)
+			for ( const FScreenResolutionRHI& SupportedResolution : SupportedResolutions )
+			{
+				if(SupportedResolution.Height == LargestY)
+				{
+					FIntPoint Resolution;
+					Resolution.X = SupportedResolution.Width;
+					Resolution.Y = SupportedResolution.Height;
 
-			Resolutions.Add(Resolution);
+					Resolutions.Add(Resolution);
+				}
+			}
 		}
 
 		return true;
@@ -3039,6 +3062,20 @@ bool UKismetSystemLibrary::GetSupportedFullscreenResolutions(TArray<FIntPoint>&
 	return false;
 }
 
+static TAutoConsoleVariable<int32> CVarMinYResolution(
+	TEXT("r.MinYResolution"),
+	720,
+	TEXT("Defines the smallest Y resolution we want to support in the UI"),
+	ECVF_RenderThreadSafe
+	);
+
+int32 UKismetSystemLibrary::GetMinYResolutionForUI()
+{
+	int32 Value = FMath::Clamp(CVarMinYResolution.GetValueOnGameThread(), 200, 8 * 1024);
+
+	return Value;
+}
+
 void UKismetSystemLibrary::LaunchURL(const FString& URL)
 {
 	if (!URL.IsEmpty())
@@ -3319,4 +3356,9 @@ void UKismetSystemLibrary::LoadAssetClass(UObject* WorldContextObject, const TAs
 void UKismetSystemLibrary::RegisterForRemoteNotifications()
 {
 	FPlatformMisc::RegisterForRemoteNotifications();
-}
\ No newline at end of file
+}
+
+void UKismetSystemLibrary::SetUserActivity(const FUserActivity& UserActivity)
+{
+	FUserActivityTracking::SetActivity(UserActivity);
+}
diff --git a/Engine/Source/Runtime/Engine/Private/LODActor.cpp b/Engine/Source/Runtime/Engine/Private/LODActor.cpp
--- a/Engine/Source/Runtime/Engine/Private/LODActor.cpp
+++ b/Engine/Source/Runtime/Engine/Private/LODActor.cpp
@@ -159,28 +159,32 @@ void ALODActor::PostEditChangeProperty(FPropertyChangedEvent& PropertyChangedEve
 	UProperty* PropertyThatChanged = PropertyChangedEvent.Property;
 	FName PropertyName = PropertyThatChanged != NULL ? PropertyThatChanged->GetFName() : NAME_None;
 	
-	if (PropertyName == GET_MEMBER_NAME_CHECKED(ALODActor, LODDrawDistance))
+	if (PropertyName == GET_MEMBER_NAME_CHECKED(ALODActor, bOverrideTransitionScreenSize) || PropertyName == GET_MEMBER_NAME_CHECKED(ALODActor, TransitionScreenSize))
 	{
-		for (auto& Actor: SubActors)
-		{
-			if (Actor)
-			{
-				TArray<UPrimitiveComponent*> InnerComponents;
-				Actor->GetComponents<UPrimitiveComponent>(InnerComponents);
+		float CalculateSreenSize = 0.0f;
 
-				for(auto& Component: InnerComponents)
-				{
-					UPrimitiveComponent* ParentComp = Component->GetLODParentPrimitive();
-					if (ParentComp)
-					{
-						ParentComp->MinDrawDistance = LODDrawDistance;
-						ParentComp->MarkRenderStateDirty();
-					}
-				}
-			}
+		if (bOverrideTransitionScreenSize)
+		{
+			CalculateSreenSize = TransitionScreenSize;
+		}
+		else
+		{
+			UWorld* World = GetWorld();
+			check(World != nullptr);
+			AWorldSettings* WorldSettings = World->GetWorldSettings();
+			checkf(WorldSettings->HierarchicalLODSetup.IsValidIndex(LODLevel - 1), TEXT("Out of range HLOD level (%i) found in LODActor (%s)"), LODLevel - 1, *GetName());
+			CalculateSreenSize = WorldSettings->HierarchicalLODSetup[LODLevel - 1].TransitionScreenSize;
 		}
+
+		RecalculateDrawingDistance(CalculateSreenSize);
+	}
+	else if (PropertyName == GET_MEMBER_NAME_CHECKED(ALODActor, bOverrideScreenSize) || PropertyName == GET_MEMBER_NAME_CHECKED(ALODActor, ScreenSize)
+		|| PropertyName == GET_MEMBER_NAME_CHECKED(ALODActor, bOverrideMaterialMergeSettings) || PropertyName == GET_MEMBER_NAME_CHECKED(ALODActor, MaterialSettings))
+	{
+		// If we change override settings dirty the actor
 		SetIsDirty(true);
 	}
+	
 	Super::PostEditChangeProperty(PropertyChangedEvent);
 }
 
@@ -366,6 +370,17 @@ void ALODActor::SetIsDirty(const bool bNewState)
 			}
 		}
 
+		// Set static mesh to null (this so we can revert without destroying the previously build static mesh)
+		StaticMeshComponent->StaticMesh = nullptr;
+		// Mark render state dirty to update viewport
+		StaticMeshComponent->MarkRenderStateDirty();
+
+		// Propagate to sub actors that we no longer have a static mesh
+		for (auto& SubActor : SubActors)
+		{
+			SubActor->SetLODParent(nullptr, LODDrawDistance);
+		}
+
 		// Broadcast actor marked dirty event
 		if (GEngine)
 		{
@@ -481,6 +496,7 @@ void ALODActor::UpdateSubActorLODParents()
 
 void ALODActor::CleanSubActorArray()
 {
+	bool bIsDirty = false;
 	for (int32 SubActorIndex = 0; SubActorIndex < SubActors.Num(); ++SubActorIndex)
 	{
 		auto& Actor = SubActors[SubActorIndex];
@@ -488,8 +504,30 @@ void ALODActor::CleanSubActorArray()
 		{
 			SubActors.RemoveAt(SubActorIndex);
 			SubActorIndex--;
+			bIsDirty = true;
 		}
 	}
+
+	if (bIsDirty)
+	{
+		SetIsDirty(true);
+	}
+}
+
+void ALODActor::RecalculateDrawingDistance(const float TransitionScreenSize)
+{
+	// At the moment this assumes a fixed field of view of 90 degrees (horizontal and vertical axi)
+	static const float FOVRad = 90.0f * (float)PI / 360.0f;
+	static const FMatrix ProjectionMatrix = FPerspectiveMatrix(FOVRad, 1920, 1080, 0.01f);
+	FBoxSphereBounds Bounds = GetStaticMeshComponent()->CalcBounds(FTransform());
+
+	// Get projection multiple accounting for view scaling.
+	const float ScreenMultiple = FMath::Max(1920.0f / 2.0f * ProjectionMatrix.M[0][0],
+		1080.0f / 2.0f * ProjectionMatrix.M[1][1]);
+	// (ScreenMultiple * SphereRadius) / Sqrt(Screensize * 1920 * 1080.0f * PI) = Distance
+	LODDrawDistance = (ScreenMultiple * Bounds.SphereRadius) / FMath::Sqrt((TransitionScreenSize * 1920.0f * 1080.0f) / PI);
+
+	UpdateSubActorLODParents();
 }
 
 #endif // WITH_EDITOR
diff --git a/Engine/Source/Runtime/Engine/Private/Level.cpp b/Engine/Source/Runtime/Engine/Private/Level.cpp
--- a/Engine/Source/Runtime/Engine/Private/Level.cpp
+++ b/Engine/Source/Runtime/Engine/Private/Level.cpp
@@ -738,30 +738,125 @@ void ULevel::UpdateLevelComponents(bool bRerunConstructionScripts)
 	IncrementalUpdateComponents( 0, bRerunConstructionScripts );
 }
 
+namespace FLevelSortUtils
+{
+	void AddToListSafe(AActor* TestActor, TArray<AActor*>& List)
+	{
+		if (TestActor)
+		{
+			const bool bAlreadyAdded = List.Contains(TestActor);
+			if (bAlreadyAdded)
+			{
+				FString ListItemDesc;
+				for (int32 Idx = 0; Idx < List.Num(); Idx++)
+				{
+					if (Idx > 0)
+					{
+						ListItemDesc += TEXT(", ");
+					}
+
+					ListItemDesc += GetNameSafe(List[Idx]);
+				}
+
+				UE_LOG(LogLevel, Warning, TEXT("Found a cycle in actor's parent chain: %s"), *ListItemDesc);
+			}
+			else
+			{
+				List.Add(TestActor);
+			}
+		}
+	}
+
+	// Finds list of parents from an entry in ParentMap, returns them in provided array and removes from map
+	// Logs an error when cycle is found
+	void FindAndRemoveParentChain(TMap<AActor*, AActor*>& ParentMap, TArray<AActor*>& ParentChain)
+	{
+		check(ParentMap.Num());
+		
+		// seed from first entry
+		TMap<AActor*, AActor*>::TIterator It(ParentMap);
+		ParentChain.Add(It.Key());
+		ParentChain.Add(It.Value());
+		It.RemoveCurrent();
+
+		// fill chain's parent nodes
+		bool bLoop = true;
+		while (bLoop)
+		{
+			AActor* MapValue = nullptr;
+			bLoop = ParentMap.RemoveAndCopyValue(ParentChain.Last(), MapValue);
+			AddToListSafe(MapValue, ParentChain);
+		}
+
+		// find chain's child nodes, ignore cycle detection since it would've triggered already from previous loop
+		for (AActor* const* MapKey = ParentMap.FindKey(ParentChain[0]); MapKey; MapKey = ParentMap.FindKey(ParentChain[0]))
+		{
+			AActor* MapValue = nullptr;
+			ParentMap.RemoveAndCopyValue((AActor*)MapKey, MapValue);
+			ParentChain.Insert(MapValue, 0);
+		}
+	}
+
+	struct FDepthSort
+	{
+		TMap<AActor*, int32> DepthMap;
+
+		bool operator()(AActor* A, AActor* B) const
+		{
+			const int32 DepthA = A ? DepthMap.FindRef(A) : MAX_int32;
+			const int32 DepthB = B ? DepthMap.FindRef(B) : MAX_int32;
+			return DepthA < DepthB;
+		}
+	};
+}
+
 /**
 *	Sorts actors such that parent actors will appear before children actors in the list
 *	Stable sort
 */
-static void SortActorsHierarchy(TTransArray<AActor*>& Actors)
+static void SortActorsHierarchy(TTransArray<AActor*>& Actors, UObject* Level)
 {
-	auto CalcAttachDepth = [](AActor* InActor) -> int32 {
-		int32 Depth = MAX_int32;
-		if (InActor)
+	const double StartTime = FPlatformTime::Seconds();
+
+	// Precalculate parent map to avoid processing cycles during sort
+	TMap<AActor*, AActor*> ParentMap;
+	for (int32 Idx = 0; Idx < Actors.Num(); Idx++)
+	{
+		if (Actors[Idx])
 		{
-			Depth = 0;
-			if (InActor->GetRootComponent())
+			AActor* ParentActor = Actors[Idx]->GetAttachParentActor();
+			if (ParentActor)
 			{
-				for (const USceneComponent* Test = InActor->GetRootComponent()->AttachParent; Test != nullptr; Test = Test->AttachParent, Depth++);
+				ParentMap.Add(Actors[Idx], ParentActor);
 			}
 		}
-		return Depth;
-	};
-	
-	// Unfortunately TArray.StableSort assumes no null entries in the array
-	// So it forces me to use internal unrestricted version
-	StableSortInternal(Actors.GetData(), Actors.Num(), [&](AActor* L, AActor* R) {
-			return CalcAttachDepth(L) < CalcAttachDepth(R);
-	});
+	}
+
+	if (ParentMap.Num())
+	{
+		FLevelSortUtils::FDepthSort DepthSorter;
+		TArray<AActor*> ParentChain;
+		while (ParentMap.Num())
+		{
+			ParentChain.Reset();
+			FLevelSortUtils::FindAndRemoveParentChain(ParentMap, ParentChain);
+
+			for (int32 Idx = 0; Idx < ParentChain.Num(); Idx++)
+			{
+				DepthSorter.DepthMap.Add(ParentChain[Idx], ParentChain.Num() - Idx - 1);
+			}
+		}
+
+		// Unfortunately TArray.StableSort assumes no null entries in the array
+		// So it forces me to use internal unrestricted version
+		StableSortInternal(Actors.GetData(), Actors.Num(), DepthSorter);
+	}
+
+	const float ElapsedTime = (float)(FPlatformTime::Seconds() - StartTime);
+	if (ElapsedTime > 1.0f)
+	{
+		UE_LOG(LogLevel, Warning, TEXT("SortActorsHierarchy(%s) took %f seconds"), Level ? *GetNameSafe(Level->GetOutermost()) : TEXT("??"), ElapsedTime);
+	}
 
 	// Since all the null entries got sorted to the end, lop them off right now
 	int32 RemoveAtIndex = Actors.Num();
@@ -790,7 +885,7 @@ void ULevel::IncrementalUpdateComponents(int32 NumComponentsToUpdate, bool bReru
 	{
 		UpdateModelComponents();
 		// Sort actors to ensure that parent actors will be registered before child actors
-		SortActorsHierarchy(Actors);
+		SortActorsHierarchy(Actors, this);
 	}
 
 	// Find next valid actor to process components registration
diff --git a/Engine/Source/Runtime/Engine/Private/LevelStreaming.cpp b/Engine/Source/Runtime/Engine/Private/LevelStreaming.cpp
--- a/Engine/Source/Runtime/Engine/Private/LevelStreaming.cpp
+++ b/Engine/Source/Runtime/Engine/Private/LevelStreaming.cpp
@@ -324,7 +324,7 @@ void ULevelStreaming::DiscardPendingUnloadLevel(UWorld* PersistentWorld)
 	}
 }
 
-bool ULevelStreaming::RequestLevel(UWorld* PersistentWorld, bool bAllowLevelLoadRequests, bool bBlockOnLoad)
+bool ULevelStreaming::RequestLevel(UWorld* PersistentWorld, bool bAllowLevelLoadRequests, EReqLevelBlock BlockPolicy)
 {
 	// Quit early in case load request already issued
 	if (bHasLoadRequestPending)
@@ -420,7 +420,7 @@ bool ULevelStreaming::RequestLevel(UWorld* PersistentWorld, bool bAllowLevelLoad
 		const FString NonPrefixedLevelName = UWorld::StripPIEPrefixFromPackageName(DesiredPackageName.ToString(), PersistentWorld->StreamingLevelsPrefix);
 		UPackage* EditorLevelPackage = FindObjectFast<UPackage>(nullptr, FName(*NonPrefixedLevelName));
 		
-		bool bShouldDuplicate = EditorLevelPackage && (bBlockOnLoad || EditorLevelPackage->IsDirty() || !GEngine->PreferToStreamLevelsInPIE());
+		bool bShouldDuplicate = EditorLevelPackage && (BlockPolicy == AlwaysBlock || EditorLevelPackage->IsDirty() || !GEngine->PreferToStreamLevelsInPIE());
 		if (bShouldDuplicate)
 		{
 			// Do the duplication
@@ -485,7 +485,7 @@ bool ULevelStreaming::RequestLevel(UWorld* PersistentWorld, bool bAllowLevelLoad
 
 			// streamingServer: server loads everything?
 			// Editor immediately blocks on load and we also block if background level streaming is disabled.
-			if (bBlockOnLoad || ShouldBeAlwaysLoaded())
+			if (BlockPolicy == AlwaysBlock || (ShouldBeAlwaysLoaded() && BlockPolicy != NeverBlock))
 			{
 				// Finish all async loading.
 				FlushAsyncLoading();
diff --git a/Engine/Source/Runtime/Engine/Private/LocalPlayer.cpp b/Engine/Source/Runtime/Engine/Private/LocalPlayer.cpp
--- a/Engine/Source/Runtime/Engine/Private/LocalPlayer.cpp
+++ b/Engine/Source/Runtime/Engine/Private/LocalPlayer.cpp
@@ -1166,7 +1166,7 @@ bool ULocalPlayer::HandleExecCommand( const TCHAR* Cmd, FOutputDevice& Ar )
 
 bool ULocalPlayer::HandleToggleDrawEventsCommand( const TCHAR* Cmd, FOutputDevice& Ar )
 {
-#if !(UE_BUILD_SHIPPING || UE_BUILD_TEST)
+#if WITH_PROFILEGPU
 	if( GEmitDrawEvents )
 	{
 		GEmitDrawEvents = false;
@@ -1354,7 +1354,7 @@ bool ULocalPlayer::Exec(UWorld* InWorld, const TCHAR* Cmd,FOutputDevice& Ar)
 	{
 		return HandleExecCommand( Cmd, Ar );
 	}
-#if !(UE_BUILD_SHIPPING || UE_BUILD_TEST)
+#if WITH_PROFILEGPU
 	else if( FParse::Command(&Cmd,TEXT("TOGGLEDRAWEVENTS")) )
 	{
 		return HandleToggleDrawEventsCommand( Cmd, Ar );
diff --git a/Engine/Source/Runtime/Engine/Private/Materials/HLSLMaterialTranslator.h b/Engine/Source/Runtime/Engine/Private/Materials/HLSLMaterialTranslator.h
--- a/Engine/Source/Runtime/Engine/Private/Materials/HLSLMaterialTranslator.h
+++ b/Engine/Source/Runtime/Engine/Private/Materials/HLSLMaterialTranslator.h
@@ -857,7 +857,6 @@ public:
 
 		FString CustomUVAssignments;
 
-		{
 			int32 LastProperty = -1;
 			for (uint32 CustomUVIndex = 0; CustomUVIndex < NumUserTexCoords; CustomUVIndex++)
 			{
@@ -876,7 +875,6 @@ public:
 				}
 				CustomUVAssignments += FString::Printf(TEXT("\tOutTexCoords[%u] = %s;") LINE_TERMINATOR, CustomUVIndex, *TranslatedCodeChunks[MP_CustomizedUVs0 + CustomUVIndex]);
 			}
-		}
 
 		LazyPrintf.PushParam(*CustomUVAssignments);
 
diff --git a/Engine/Source/Runtime/Engine/Private/Materials/MaterialExpressions.cpp b/Engine/Source/Runtime/Engine/Private/Materials/MaterialExpressions.cpp
--- a/Engine/Source/Runtime/Engine/Private/Materials/MaterialExpressions.cpp
+++ b/Engine/Source/Runtime/Engine/Private/Materials/MaterialExpressions.cpp
@@ -8294,7 +8294,7 @@ int32 UMaterialExpressionMaterialProxyReplace::Compile(class FMaterialCompiler*
 	{
 		int32 Arg1 = Realtime.Compile(Compiler);
 		int32 Arg2 = MaterialProxy.Compile(Compiler);
-		return Compiler->LightmassReplace(Arg1, Arg2);
+		return Compiler->MaterialProxyReplace(Arg1, Arg2);
 	}
 }
 
