Fix for variable aliasing

diff --git a/Engine/Source/Runtime/Engine/Private/LODActor.cpp b/Engine/Source/Runtime/Engine/Private/LODActor.cpp
--- a/Engine/Source/Runtime/Engine/Private/LODActor.cpp
+++ b/Engine/Source/Runtime/Engine/Private/LODActor.cpp
@@ -514,7 +514,7 @@ void ALODActor::CleanSubActorArray()
 	}
 }
 
-void ALODActor::RecalculateDrawingDistance(const float TransitionScreenSize)
+void ALODActor::RecalculateDrawingDistance(const float InTransitionScreenSize)
 {
 	// At the moment this assumes a fixed field of view of 90 degrees (horizontal and vertical axi)
 	static const float FOVRad = 90.0f * (float)PI / 360.0f;
@@ -525,7 +525,7 @@ void ALODActor::RecalculateDrawingDistance(const float TransitionScreenSize)
 	const float ScreenMultiple = FMath::Max(1920.0f / 2.0f * ProjectionMatrix.M[0][0],
 		1080.0f / 2.0f * ProjectionMatrix.M[1][1]);
 	// (ScreenMultiple * SphereRadius) / Sqrt(Screensize * 1920 * 1080.0f * PI) = Distance
-	LODDrawDistance = (ScreenMultiple * Bounds.SphereRadius) / FMath::Sqrt((TransitionScreenSize * 1920.0f * 1080.0f) / PI);
+	LODDrawDistance = (ScreenMultiple * Bounds.SphereRadius) / FMath::Sqrt((InTransitionScreenSize * 1920.0f * 1080.0f) / PI);
 
 	UpdateSubActorLODParents();
 }
