Copying //UE4/Dev-Framework to //UE4/Main @ 2830052

diff --git a/Engine/Source/Runtime/Engine/Private/ActiveSound.cpp b/Engine/Source/Runtime/Engine/Private/ActiveSound.cpp
--- a/Engine/Source/Runtime/Engine/Private/ActiveSound.cpp
+++ b/Engine/Source/Runtime/Engine/Private/ActiveSound.cpp
@@ -8,14 +8,17 @@
 #include "Sound/SoundNodeAttenuation.h"
 #include "SubtitleManager.h"
 
+
 FActiveSound::FActiveSound()
 	: Sound(nullptr)
 	, World(nullptr)
 	, AudioComponent(nullptr)
 	, AudioDevice(nullptr)
 	, ConcurrencyGroupID(0)
+	, ConcurrencyGeneration(0)
 	, ConcurrencySettings(nullptr)
 	, SoundClassOverride(nullptr)
+	, bHasCheckedOcclusion(false)
 	, bIsOccluded(false)
 	, bAllowSpatialization(true)
 	, bHasAttenuationSettings(false)
@@ -52,8 +55,11 @@ FActiveSound::FActiveSound()
 	, CurrentOcclusionFilterFrequency(MAX_FILTER_FREQUENCY)
 	, CurrentOcclusionVolumeAttenuation(1.0f)
 	, ConcurrencyVolumeScale(1.f)
+	, ConcurrencyDuckingVolumeScale(1.f)
 	, SubtitlePriority(0.f)
 	, Priority(1.0f)
+	, FocusPriorityScale(1.0f)
+	, FocusDistanceScale(1.0f)
 	, VolumeConcurrency(0.0f)
 	, OcclusionCheckInterval(0.f)
 	, LastOcclusionCheckTime(0.f)
@@ -122,20 +128,7 @@ USoundClass* FActiveSound::GetSoundClass() const
 
 int32 FActiveSound::FindClosestListener( const TArray<FListener>& InListeners ) const
 {
-	int32 ClosestListenerIndex = 0;
-	float ClosestDistSq = FVector::DistSquared(Transform.GetTranslation(), InListeners[0].Transform.GetTranslation() );
-
-	for( int32 i = 1; i < InListeners.Num(); i++ )
-	{
-		const float DistSq = FVector::DistSquared(Transform.GetTranslation(), InListeners[i].Transform.GetTranslation());
-		if( DistSq < ClosestDistSq )
-		{
-			ClosestListenerIndex = i;
-			ClosestDistSq = DistSq;
-		}
-	}
-
-	return ClosestListenerIndex;
+	return FAudioDevice::FindClosestListenerIndex(Transform, InListeners);
 }
 
 uint32 FActiveSound::TryGetOwnerID() const
@@ -214,8 +207,11 @@ void FActiveSound::UpdateWaveInstances( TArray<FWaveInstance*> &InWaveInstances,
 	// Cache the closest listener ptr 
 	ClosestListenerPtr = &AudioDevice->Listeners[ClosestListenerIndex];
 
+	// The apparent max distance factors the actual max distance of the sound scaled with the distance scale due to focus effects
+	float ApparentMaxDistance = MaxDistance * FocusDistanceScale;
+
 	// Update whether or not his sound is out of range
-	bIsAudible = AudioDevice->LocationIsAudible(Transform.GetTranslation(), *ClosestListenerPtr, MaxDistance);
+	bIsAudible = AudioDevice->LocationIsAudible(Transform.GetTranslation(), *ClosestListenerPtr, ApparentMaxDistance);
 
 	FSoundParseParameters ParseParams;
 	ParseParams.Transform = Transform;
@@ -352,22 +348,25 @@ void FActiveSound::UpdateAdjustVolumeMultiplier(const float DeltaTime)
 
 void FActiveSound::UpdateOcclusion(const FAttenuationSettings* AttenuationSettingsPtr)
 {
+	float InterpolationTime = bHasCheckedOcclusion ? AttenuationSettingsPtr->OcclusionInterpolationTime : 0.0f;
+	bHasCheckedOcclusion = true;
+
 	if (bIsOccluded)
 	{
 		if (CurrentOcclusionFilterFrequency.GetTargetValue() > AttenuationSettingsPtr->OcclusionLowPassFilterFrequency)
 		{
-			CurrentOcclusionFilterFrequency.Set(AttenuationSettingsPtr->OcclusionLowPassFilterFrequency, AttenuationSettingsPtr->OcclusionInterpolationTime);
+			CurrentOcclusionFilterFrequency.Set(AttenuationSettingsPtr->OcclusionLowPassFilterFrequency, InterpolationTime);
 		}
 
 		if (CurrentOcclusionVolumeAttenuation.GetTargetValue() > AttenuationSettingsPtr->OcclusionVolumeAttenuation)
 		{
-			CurrentOcclusionVolumeAttenuation.Set(AttenuationSettingsPtr->OcclusionVolumeAttenuation, AttenuationSettingsPtr->OcclusionInterpolationTime);
+			CurrentOcclusionVolumeAttenuation.Set(AttenuationSettingsPtr->OcclusionVolumeAttenuation, InterpolationTime);
 		}
 	}
 	else
 	{
-		CurrentOcclusionFilterFrequency.Set(MAX_FILTER_FREQUENCY, AttenuationSettingsPtr->OcclusionInterpolationTime);
-		CurrentOcclusionVolumeAttenuation.Set(1.0f, AttenuationSettingsPtr->OcclusionInterpolationTime);
+		CurrentOcclusionFilterFrequency.Set(MAX_FILTER_FREQUENCY, InterpolationTime);
+		CurrentOcclusionVolumeAttenuation.Set(1.0f, InterpolationTime);
 	}
 
 	UWorld* WorldPtr = World.Get();
@@ -742,30 +741,6 @@ void FActiveSound::CollectAttenuationShapesForVisualization(TMultiMap<EAttenuati
 	}
 }
 
-void FActiveSound::GetAttenuationListenerData(FAttenuationListenerData& ListenerData, const FSoundParseParameters& ParseParams, const FListener& Listener) const
-{
-	// Only perform this calculation once and explicitly require a reset.
-	if (ListenerData.bDataComputed)
-	{
-		return;
-	}
-
-	const FTransform& SoundTransform = ParseParams.Transform;
-	const FVector& ListenerLocation = Listener.Transform.GetTranslation();
-
-	FVector ListenerToSound = SoundTransform.GetTranslation() - ListenerLocation;
-	ListenerToSound.ToDirectionAndLength(ListenerData.ListenerToSoundDir, ListenerData.ListenerToSoundDistance);
-
-	ListenerData.AttenuationDistance = 0.0f;
-
-	if ((AttenuationSettings.bAttenuate && AttenuationSettings.AttenuationShape == EAttenuationShape::Sphere) || AttenuationSettings.bAttenuateWithLPF)
-	{
-		ListenerData.AttenuationDistance = FMath::Max(ListenerData.ListenerToSoundDistance - AttenuationSettings.AttenuationShapeExtents.X, 0.f);
-	}
-
-	ListenerData.bDataComputed = true;
-}
-
 void FActiveSound::ApplyAttenuation(FSoundParseParameters& ParseParams, const FListener& Listener, const FAttenuationSettings* SettingsAttenuationNode)
 {
 	float& Volume = ParseParams.Volume;
@@ -777,57 +752,27 @@ void FActiveSound::ApplyAttenuation(FSoundParseParameters& ParseParams, const FL
 
 	FAttenuationListenerData ListenerData;
 
-	// Init focus factor to 1.0f in case this sound is ignoring focus
-	float FocusFactor = 1.0f;
-
-	// How much we want to scale the listener-emitter distance based on focus
-	float DistanceScale = 1.0f;
+	// Get the current focus factor
+	const float FocusFactor = AudioDevice->GetFocusFactor(ListenerData, Sound, SoundTransform, *Settings, &Listener);
 
+	// Reset distance and priority scale to 1.0 in case changed in editor
+	FocusDistanceScale = 1.0f;
+	FocusPriorityScale = 1.0f;
+	
 	// Computer the focus factor if needed
 	check(Sound);
 	if (Settings->bSpatialize && Settings->bEnableListenerFocus && !Sound->bIgnoreFocus)
 	{
-		// Get the attenuation listener data
-		GetAttenuationListenerData(ListenerData, ParseParams, Listener);
-
-		// Compute the focus factor
-		const FVector& ListenerForwardDir = Listener.Transform.GetUnitAxis(EAxis::X);
-		const FVector& ListenerToSoundDir = ListenerData.ListenerToSoundDir;
-		const float FocusDotProduct = FVector::DotProduct(ListenerForwardDir, ListenerToSoundDir);
-		const float FocusAngle = FMath::RadiansToDegrees(FMath::Acos(FocusDotProduct));
-
-		const float FocusAzimuth = FMath::Clamp(Settings->FocusAzimuth, 0.0f, 360.0f);
-		const float NonFocusAzimuth = FMath::Clamp(Settings->NonFocusAzimuth, 0.0f, 360.0f);
-
-		if (FocusAzimuth != NonFocusAzimuth)
-		{
-			FocusFactor = (FocusAngle - FocusAzimuth) / (NonFocusAzimuth - FocusAzimuth);
-			FocusFactor = FMath::Clamp(FocusFactor, 0.0f, 1.0f);
-		}
-		else
-		{
-			FocusFactor = 0.0f;
-			if (FocusAngle < FocusAzimuth)
-			{
-				FocusFactor = 1.0f;
-			}
-		}
-
 		// Get the volume scale to apply the volume calculation based on the focus factor
-		const float NonFocusVolumeAttenuation = FMath::Clamp(Settings->NonFocusVolumeAttenuation, 0.0f, 1.0f);
-		const float FocusVolumeAttenuation = FMath::Lerp(1.0f, NonFocusVolumeAttenuation, FocusFactor);
+		const float FocusVolumeAttenuation = Settings->GetFocusAttenuation(AudioDevice->GlobalFocusSettings, FocusFactor);
 		Volume *= FocusVolumeAttenuation;
 
 		// Scale the volume-weighted priority scale value we use for sorting this sound for voice-stealing
-		const float NonFocusPriorityScale = FMath::Max(Settings->NonFocusPriorityScale, 0.0f);
-		const float FocusPriorityScale = FMath::Max(Settings->FocusPriorityScale, 0.0f);
-		const float PriorityScale = FMath::Lerp(FocusPriorityScale, NonFocusPriorityScale, FocusFactor);
-		ParseParams.Priority *= PriorityScale;
+		FocusPriorityScale = Settings->GetFocusPriorityScale(AudioDevice->GlobalFocusSettings, FocusFactor);
+		ParseParams.Priority *= FocusPriorityScale;
 
-		// Get the distance scale to use when computing distance-calculations for 3d atttenuation
-		const float FocusDistanceScale = FMath::Max(Settings->FocusDistanceScale, 0.0f);
-		const float NonFocusDistanceScale = FMath::Max(Settings->NonFocusDistanceScale, 0.0f);
-		DistanceScale = FMath::Lerp(FocusDistanceScale, NonFocusDistanceScale, FocusFactor);
+		// Get the distance scale to use when computing distance-calculations for 3d attenuation
+		FocusDistanceScale = Settings->GetFocusDistanceScale(AudioDevice->GlobalFocusSettings, FocusFactor);
 	}
 
 	// Attenuate the volume based on the model
@@ -838,21 +783,21 @@ void FActiveSound::ApplyAttenuation(FSoundParseParameters& ParseParams, const FL
 			case EAttenuationShape::Sphere:
 			{
 				// Update attenuation data in-case it hasn't been updated
-				GetAttenuationListenerData(ListenerData, ParseParams, Listener);
-				Volume *= Settings->AttenuationEval(ListenerData.AttenuationDistance, Settings->FalloffDistance, DistanceScale);
+				AudioDevice->GetAttenuationListenerData(ListenerData, SoundTransform, *Settings, &Listener);
+				Volume *= Settings->AttenuationEval(ListenerData.AttenuationDistance, Settings->FalloffDistance, FocusDistanceScale);
 				break;
 			}
 
 			case EAttenuationShape::Box:
-			Volume *= Settings->AttenuationEvalBox(SoundTransform, ListenerLocation, DistanceScale);
+			Volume *= Settings->AttenuationEvalBox(SoundTransform, ListenerLocation, FocusDistanceScale);
 			break;
 
 			case EAttenuationShape::Capsule:
-			Volume *= Settings->AttenuationEvalCapsule(SoundTransform, ListenerLocation, DistanceScale);
+			Volume *= Settings->AttenuationEvalCapsule(SoundTransform, ListenerLocation, FocusDistanceScale);
 			break;
 
 			case EAttenuationShape::Cone:
-			Volume *= Settings->AttenuationEvalCone(SoundTransform, ListenerLocation, DistanceScale);
+			Volume *= Settings->AttenuationEvalCone(SoundTransform, ListenerLocation, FocusDistanceScale);
 			break;
 
 			default:
@@ -875,7 +820,7 @@ void FActiveSound::ApplyAttenuation(FSoundParseParameters& ParseParams, const FL
 	// Attenuate with the low pass filter if necessary
 	if (Settings->bAttenuateWithLPF)
 	{
-		GetAttenuationListenerData(ListenerData, ParseParams, Listener);
+		AudioDevice->GetAttenuationListenerData(ListenerData, SoundTransform, *Settings, &Listener);
 
 		// Attenuate with the low pass filter if necessary
 		FVector2D InputRange(Settings->LPFRadiusMin, Settings->LPFRadiusMax);
diff --git a/Engine/Source/Runtime/Engine/Private/Actor.cpp b/Engine/Source/Runtime/Engine/Private/Actor.cpp
--- a/Engine/Source/Runtime/Engine/Private/Actor.cpp
+++ b/Engine/Source/Runtime/Engine/Private/Actor.cpp
@@ -1148,17 +1148,9 @@ void AActor::GetOverlappingActors(TArray<AActor*>& OutOverlappingActors, UClass*
 	OutOverlappingActors.Reset();
 	TArray<AActor*> OverlappingActorsForCurrentComponent;
 
-	// use a stack to walk this actor's attached component tree
-	TInlineComponentArray<USceneComponent*> ComponentStack;
-	USceneComponent const* CurrentComponent = GetRootComponent();
-	while (CurrentComponent)
+	for(UActorComponent* OwnedComp : OwnedComponents)
 	{
-		// push children on the stack so they get tested later
-		ComponentStack.Append(CurrentComponent->AttachChildren);
-
-		// get list of actors from the component
-		UPrimitiveComponent const* const PrimComp = Cast<const UPrimitiveComponent>(CurrentComponent);
-		if (PrimComp)
+		if(UPrimitiveComponent* PrimComp = Cast<UPrimitiveComponent>(OwnedComp))
 		{
 			PrimComp->GetOverlappingActors(OverlappingActorsForCurrentComponent, ClassFilter);
 
@@ -1172,10 +1164,6 @@ void AActor::GetOverlappingActors(TArray<AActor*>& OutOverlappingActors, UClass*
 				}
 			}
 		}
-
-		// advance to next component
-		const bool bAllowShrinking = false;
-		CurrentComponent = (ComponentStack.Num() > 0) ? ComponentStack.Pop(bAllowShrinking) : nullptr;
 	}
 }
 
@@ -1184,16 +1172,9 @@ void AActor::GetOverlappingComponents(TArray<UPrimitiveComponent*>& OutOverlappi
 	OutOverlappingComponents.Reset();
 	TArray<UPrimitiveComponent*> OverlappingComponentsForCurrentComponent;
 
-	// use a stack to walk this actor's attached component tree
-	TInlineComponentArray<USceneComponent*> ComponentStack;
-	USceneComponent* CurrentComponent = GetRootComponent();
-	while (CurrentComponent)
+	for (UActorComponent* OwnedComp : OwnedComponents)
 	{
-		// push children on the stack so they get tested later
-		ComponentStack.Append(CurrentComponent->AttachChildren);
-
-		UPrimitiveComponent const* const PrimComp = Cast<const UPrimitiveComponent>(CurrentComponent);
-		if (PrimComp)
+		if (UPrimitiveComponent* PrimComp = Cast<UPrimitiveComponent>(OwnedComp))
 		{
 			// get list of components from the component
 			PrimComp->GetOverlappingComponents(OverlappingComponentsForCurrentComponent);
@@ -1204,10 +1185,6 @@ void AActor::GetOverlappingComponents(TArray<UPrimitiveComponent*>& OutOverlappi
 				OutOverlappingComponents.AddUnique(*CompIt);
 			}
 		}
-
-		// advance to next component
-		const bool bAllowShrinking = false;
-		CurrentComponent = (ComponentStack.Num() > 0) ? ComponentStack.Pop(bAllowShrinking) : nullptr;
 	}
 }
 
@@ -1307,6 +1284,12 @@ static void MarkOwnerRelevantComponentsDirty(AActor* TheActor)
 	}
 }
 
+bool AActor::WasRecentlyRendered(float Tolerance) const
+{
+	UWorld* World = GetWorld();
+	return (World) ? (World->GetTimeSeconds() - GetLastRenderTime() <= Tolerance) : false;
+}
+
 float AActor::GetLastRenderTime() const
 {
 	// return most recent of Components' LastRenderTime
diff --git a/Engine/Source/Runtime/Engine/Private/Animation/AnimCurveTypes.cpp b/Engine/Source/Runtime/Engine/Private/Animation/AnimCurveTypes.cpp
--- a/Engine/Source/Runtime/Engine/Private/Animation/AnimCurveTypes.cpp
+++ b/Engine/Source/Runtime/Engine/Private/Animation/AnimCurveTypes.cpp
@@ -5,9 +5,6 @@
 #include "Animation/AnimInstance.h"
 #include "AnimationRuntime.h"
 
-DECLARE_CYCLE_STAT(TEXT("BlendedCurve InitFrom"), STAT_BlendedCurve_InitFrom, STATGROUP_Anim);
-
-
 /////////////////////////////////////////////////////
 // FFloatCurve
 
@@ -428,228 +425,4 @@ bool FRawCurveTracks::DuplicateCurveDataImpl(TArray<DataType> & Curves, USkeleto
 	}
 	return false;
 }
-/////////////////////////////////////////////
-
-FORCEINLINE void ValidateCurve(FBlendedCurve* Curve)
-{
-#if 0 
-	for (auto& Element : Curve->Elements)
-	{
-		if (Element.Flags & ACF_DrivesMorphTarget)
-		{
-			ensure (Element.Value < 1.2f);
-		}
-	}
-#endif
-}
-
-void FBlendedCurve::Set(USkeleton::AnimCurveUID InUid, float InValue, int32 InFlags)
-{
-	int32 ArrayIndex;
-
-	check(bInitialized);
-
-	if (UIDList->Find(InUid, ArrayIndex))
-	{
-		Elements[ArrayIndex].Value = InValue;
-		Elements[ArrayIndex].Flags = InFlags;
-	}
-
-	ValidateCurve(this);
-}
-
-void FBlendedCurve::Reset(int32 Count)
-{
-	Elements.Reset();
-	Elements.Reserve(Count);
-}
-
-//@Todo curve flags won't transfer over - it only overwrites
-void FBlendedCurve::Blend(const FBlendedCurve& A, const FBlendedCurve& B, float Alpha)
-{
-	check(A.Num()==B.Num());
-	if(FMath::Abs(Alpha) <= ZERO_ANIMWEIGHT_THRESH)
-	{
-		// if blend is all the way for child1, then just copy its bone atoms
-		Override(A);
-	}
-	else if(FMath::Abs(Alpha - 1.0f) <= ZERO_ANIMWEIGHT_THRESH)
-	{
-		// if blend is all the way for child2, then just copy its bone atoms
-		Override(B);
-	}
-	else
-	{
-		InitFrom(A);
-		for(int32 CurveId=0; CurveId<A.Elements.Num(); ++CurveId)
-		{
-			Elements[CurveId].Value = FMath::Lerp(A.Elements[CurveId].Value, B.Elements[CurveId].Value, Alpha); 
-			Elements[CurveId].Flags = (A.Elements[CurveId].Flags) | (B.Elements[CurveId].Flags);
-		}
-	}
-
-	ValidateCurve(this);
-}
-
-void FBlendedCurve::BlendWith(const FBlendedCurve& Other, float Alpha)
-{
-	check(Num()==Other.Num());
-	if(FMath::Abs(Alpha) <= ZERO_ANIMWEIGHT_THRESH)
-	{
-		return;
-	}
-	else if(FMath::Abs(Alpha - 1.0f) <= ZERO_ANIMWEIGHT_THRESH)
-	{
-		// if blend is all the way for child2, then just copy its bone atoms
-		Override(Other);
-	}
-	else
-	{
-		for(int32 CurveId=0; CurveId<Elements.Num(); ++CurveId)
-		{
-			Elements[CurveId].Value = FMath::Lerp(Elements[CurveId].Value, Other.Elements[CurveId].Value, Alpha);
-			Elements[CurveId].Flags |= (Other.Elements[CurveId].Flags);
-		}
-	}
-	 
-	ValidateCurve(this);
-}
-
-void FBlendedCurve::ConvertToAdditive(const FBlendedCurve& BaseCurve)
-{
-	check(bInitialized);
-	check(Num()==BaseCurve.Num());
-
-	for(int32 CurveId=0; CurveId<Elements.Num(); ++CurveId)
-	{
-		Elements[CurveId].Value -= BaseCurve.Elements[CurveId].Value;
-		Elements[CurveId].Flags |= BaseCurve.Elements[CurveId].Flags;
-	}
-	ValidateCurve(this);
-}
-
-void FBlendedCurve::Accumulate(const FBlendedCurve& AdditiveCurve, float Weight)
-{
-	check(bInitialized);
-	check(Num()==AdditiveCurve.Num());
-
-	if (Weight > ZERO_ANIMWEIGHT_THRESH)
-	{
-		for(int32 CurveId=0; CurveId<Elements.Num(); ++CurveId)
-		{
-			Elements[CurveId].Value += AdditiveCurve.Elements[CurveId].Value * Weight;
-			Elements[CurveId].Flags |= AdditiveCurve.Elements[CurveId].Flags;
-		}
-	}
-
-	ValidateCurve(this);
-}
-
-void FBlendedCurve::Override(const FBlendedCurve& CurveToOverrideFrom, float Weight)
-{
-	InitFrom(CurveToOverrideFrom);
-
-	if ( FMath::IsNearlyEqual(Weight, 1.f) )
-	{
-		Override(CurveToOverrideFrom);
-	}
-	else
-	{
-		for(int32 CurveId=0; CurveId<CurveToOverrideFrom.Elements.Num(); ++CurveId)
-		{
-			Elements[CurveId].Value = CurveToOverrideFrom.Elements[CurveId].Value * Weight;
-			Elements[CurveId].Flags |= CurveToOverrideFrom.Elements[CurveId].Flags;
-		}
-	}
-
-	ValidateCurve(this);
-}
-
-void FBlendedCurve::Override(const FBlendedCurve& CurveToOverrideFrom)
-{
-	InitFrom(CurveToOverrideFrom);
-	Elements.Reset();
-	Elements.Append(CurveToOverrideFrom.Elements);
-
-	ValidateCurve(this);
-}
-
-void FBlendedCurve::InitFrom(const USkeleton* Skeleton)
-{
-	SCOPE_CYCLE_COUNTER(STAT_BlendedCurve_InitFrom);
-	InitFrom((TArray<FSmartNameMapping::UID>*)&(const_cast<USkeleton*>(Skeleton)->GetCachedAnimCurveMappingNameUids()));
-}
-
-void FBlendedCurve::InitFrom(TArray<FSmartNameMapping::UID> const * InSmartNameUIDs)
-{
-	SCOPE_CYCLE_COUNTER(STAT_BlendedCurve_InitFrom);
-
-	check(InSmartNameUIDs != nullptr);
-	UIDList = InSmartNameUIDs;
-	Elements.Reset();
-	Elements.AddZeroed(UIDList->Num());
-
-	// no name, means no curve
-	bInitialized = true;
-}
-
-
-void FBlendedCurve::InitFrom(const FBlendedCurve& InCurveToInitFrom)
-{
-	SCOPE_CYCLE_COUNTER(STAT_BlendedCurve_InitFrom);	
-
-	// make sure this doesn't happen
-	if (ensure(&InCurveToInitFrom != this))
-	{
-		check(InCurveToInitFrom.UIDList != nullptr);
-		UIDList = InCurveToInitFrom.UIDList;
-		Elements.Reset();
-		Elements.AddZeroed(UIDList->Num());
-
-		bInitialized = true;
-	}
-}
-
-void FBlendedCurve::CopyFrom(const FBlendedCurve& CurveToCopyFrom)
-{
-	if (&CurveToCopyFrom != this)
-	{
-		UIDList = CurveToCopyFrom.UIDList;
-		Elements.Reset();
-		Elements.Append(CurveToCopyFrom.Elements);
-		bInitialized = true;
-		ValidateCurve(this);
-	}
-}
-
-void FBlendedCurve::MoveFrom(FBlendedCurve& CurveToMoveFrom)
-{
-	UIDList = CurveToMoveFrom.UIDList;
-	CurveToMoveFrom.UIDList = nullptr;
-	Elements = MoveTemp(CurveToMoveFrom.Elements);
-	bInitialized = true;
-	CurveToMoveFrom.bInitialized = false;
-	ValidateCurve(this);
-}
-
-void FBlendedCurve::Combine(const FBlendedCurve& CurveToCombine)
-{
-	check(bInitialized);
-	check(Num()==CurveToCombine.Num());
-
-	for(int32 CurveId=0; CurveId<CurveToCombine.Elements.Num(); ++CurveId)
-	{
-		// if target value is non zero, we accpet target's value
-		// originally this code was doing max, but that doesn't make sense since the values can be negative
-		// we could try to pick non-zero, but if target value is non-zero, I think we should accept that value 
-		// if source is non zero, it will be overriden
-		if (CurveToCombine.Elements[CurveId].Value != 0.f)
-		{
-			Elements[CurveId].Value = CurveToCombine.Elements[CurveId].Value; 
-		}
-
-		Elements[CurveId].Flags |= CurveToCombine.Elements[CurveId].Flags;
-	}
 
-	ValidateCurve(this);
-}
\ No newline at end of file
diff --git a/Engine/Source/Runtime/Engine/Private/Animation/AnimInstance.cpp b/Engine/Source/Runtime/Engine/Private/Animation/AnimInstance.cpp
--- a/Engine/Source/Runtime/Engine/Private/Animation/AnimInstance.cpp
+++ b/Engine/Source/Runtime/Engine/Private/Animation/AnimInstance.cpp
@@ -500,8 +500,9 @@ bool UAnimInstance::ParallelCanEvaluate(const USkeletalMesh* InSkeletalMesh) con
 	return Proxy.GetRequiredBones().IsValid() && (Proxy.GetRequiredBones().GetAsset() == InSkeletalMesh);
 }
 
-void UAnimInstance::ParallelEvaluateAnimation(bool bForceRefPose, const USkeletalMesh* InSkeletalMesh, TArray<FTransform>& OutLocalAtoms, TArray<FActiveVertexAnim>& OutVertexAnims, FBlendedCurve& OutCurve)
+void UAnimInstance::ParallelEvaluateAnimation(bool bForceRefPose, const USkeletalMesh* InSkeletalMesh, TArray<FTransform>& OutLocalAtoms, TArray<FActiveVertexAnim>& OutVertexAnims, FBlendedHeapCurve& OutCurve)
 {
+	FMemMark Mark(FMemStack::Get());
 	FAnimInstanceProxy& Proxy = GetProxyOnAnyThread<FAnimInstanceProxy>();
 
 	if( !bForceRefPose )
@@ -513,7 +514,7 @@ void UAnimInstance::ParallelEvaluateAnimation(bool bForceRefPose, const USkeleta
 		// Run the anim blueprint
 		Proxy.EvaluateAnimation(EvaluationContext);
 		// Move the curves
-		OutCurve.MoveFrom(EvaluationContext.Curve);
+		OutCurve.CopyFrom(EvaluationContext.Curve);
 			
 		// can we avoid that copy?
 		if( EvaluationContext.Pose.GetNumBones() > 0 )
@@ -1075,7 +1076,7 @@ void UAnimInstance::TickSyncGroupWriteIndex()
 	GetProxyOnGameThread<FAnimInstanceProxy>().TickSyncGroupWriteIndex();
 }
 
-void UAnimInstance::UpdateCurves(const FBlendedCurve& InCurve)
+void UAnimInstance::UpdateCurves(const FBlendedHeapCurve& InCurve)
 {
 	SCOPE_CYCLE_COUNTER(STAT_UpdateCurves);
 
diff --git a/Engine/Source/Runtime/Engine/Private/Animation/AnimNodeBase.cpp b/Engine/Source/Runtime/Engine/Private/Animation/AnimNodeBase.cpp
--- a/Engine/Source/Runtime/Engine/Private/Animation/AnimNodeBase.cpp
+++ b/Engine/Source/Runtime/Engine/Private/Animation/AnimNodeBase.cpp
@@ -229,7 +229,7 @@ void FPoseLink::Evaluate(FPoseContext& Output)
 #endif
 		LinkedNode->Evaluate(Output);
 #if ENABLE_ANIMNODE_POSE_DEBUG
-		CurrentPose = Output.Pose;
+		CurrentPose.CopyBonesFrom(Output.Pose);
 #endif
 	}
 	else
diff --git a/Engine/Source/Runtime/Engine/Private/Animation/AnimNode_StateMachine.cpp b/Engine/Source/Runtime/Engine/Private/Animation/AnimNode_StateMachine.cpp
--- a/Engine/Source/Runtime/Engine/Private/Animation/AnimNode_StateMachine.cpp
+++ b/Engine/Source/Runtime/Engine/Private/Animation/AnimNode_StateMachine.cpp
@@ -260,11 +260,8 @@ void FAnimNode_StateMachine::Initialize(const FAnimationInitializeContext& Conte
 			StatesUpdated.Reset();
 			ActiveTransitionArray.Reset();
 
-			if (StateCacheBoneCounters.Num() != Machine->States.Num())
-			{
-				StateCacheBoneCounters.Reset(Machine->States.Num());
-				StateCacheBoneCounters.AddDefaulted(Machine->States.Num());
-			}
+			StateCacheBoneCounters.Reset(Machine->States.Num());
+			StateCacheBoneCounters.AddDefaulted(Machine->States.Num());
 		
 			// Move to the default state
 			SetState(Context, Machine->InitialState);
@@ -876,6 +873,14 @@ void FAnimNode_StateMachine::SetState(const FAnimationBaseContext& Context, int3
 			OnGraphStatesExited[CurrentState].ExecuteIfBound(*this, CurrentState, NewStateIndex);
 		}
 
+		bool bForceReset = false;
+
+		if(PRIVATE_MachineDescription->States.IsValidIndex(NewStateIndex))
+		{
+			const FBakedAnimationState& BakedCurrentState = PRIVATE_MachineDescription->States[NewStateIndex];
+			bForceReset = BakedCurrentState.bAlwaysResetOnEntry;
+		}
+
 		// Determine if the new state is active or not
 		const bool bAlreadyActive = GetStateWeight(NewStateIndex) > 0.0f;
 
@@ -891,7 +896,7 @@ void FAnimNode_StateMachine::SetState(const FAnimationBaseContext& Context, int3
 			}
 		}
 
-		if (!bAlreadyActive && !IsAConduitState(NewStateIndex))
+		if ((!bAlreadyActive || bForceReset) && !IsAConduitState(NewStateIndex))
 		{
 			// Initialize the new state since it's not part of an active transition (and thus not still initialized)
 			FAnimationInitializeContext InitContext(Context.AnimInstanceProxy);
diff --git a/Engine/Source/Runtime/Engine/Private/Animation/AnimNotifyState_Trail.cpp b/Engine/Source/Runtime/Engine/Private/Animation/AnimNotifyState_Trail.cpp
--- a/Engine/Source/Runtime/Engine/Private/Animation/AnimNotifyState_Trail.cpp
+++ b/Engine/Source/Runtime/Engine/Private/Animation/AnimNotifyState_Trail.cpp
@@ -9,6 +9,8 @@
 #include "MessageLog.h"
 #include "AssertionMacros.h"
 #include "Animation/AnimInstance.h"
+#include "Animation/AnimSequenceBase.h"
+#include "ParticleHelper.h"
 
 #define LOCTEXT_NAMESPACE "AnimNotifyState_Trail"
 
@@ -73,6 +75,18 @@ void UAnimNotifyState_Trail::NotifyBegin(class USkeletalMeshComponent * MeshComp
 		return;
 	}
 
+	UParticleSystem* ParticleSystemTemplate = GetOverridenPSTemplate(MeshComp, Animation);
+	if (ParticleSystemTemplate != nullptr)
+	{
+		PSTemplate = ParticleSystemTemplate;
+	}
+
+	if(PSTemplate == nullptr)
+	{
+		UE_LOG(LogParticles, Warning, TEXT("Trail Notify: Null PSTemplate for trail notify in anim: %s"), *GetNameSafe(Animation)); 
+		return;
+	}
+
 	ParticleSystemComponentArray Children;
 	GetCandidateSystems(*MeshComp, Children);
 
@@ -83,12 +97,6 @@ void UAnimNotifyState_Trail::NotifyBegin(class USkeletalMeshComponent * MeshComp
 		Width = AnimInst->GetCurveValue(WidthScaleCurve);
 	}
 
-	UParticleSystem* ParticleSystemTemplate = GetOverridenPSTemplate(MeshComp, Animation);
-	if (ParticleSystemTemplate != nullptr)
-	{
-		PSTemplate = ParticleSystemTemplate;
-	}
-
 	UParticleSystemComponent* RecycleCandidates[3] = {nullptr, nullptr, nullptr}; // in order of priority
 	bool bFoundExistingTrail = false;
 	for (UParticleSystemComponent* ParticleComp : Children)
diff --git a/Engine/Source/Runtime/Engine/Private/Animation/AnimNotify_PlayParticleEffect.cpp b/Engine/Source/Runtime/Engine/Private/Animation/AnimNotify_PlayParticleEffect.cpp
--- a/Engine/Source/Runtime/Engine/Private/Animation/AnimNotify_PlayParticleEffect.cpp
+++ b/Engine/Source/Runtime/Engine/Private/Animation/AnimNotify_PlayParticleEffect.cpp
@@ -2,8 +2,10 @@
 
 #include "EnginePrivate.h"
 #include "Animation/AnimNotifies/AnimNotify_PlayParticleEffect.h"
+#include "Animation/AnimSequenceBase.h"
 #include "Kismet/GameplayStatics.h"
 #include "Particles/ParticleSystem.h"
+#include "ParticleHelper.h"
 
 /////////////////////////////////////////////////////
 // UAnimNotify_PlayParticleEffect
@@ -55,6 +57,10 @@ void UAnimNotify_PlayParticleEffect::Notify(class USkeletalMeshComponent* MeshCo
 			UGameplayStatics::SpawnEmitterAtLocation(MeshComp->GetWorld(), PSTemplate, SpawnTransform);
 		}
 	}
+	else
+	{
+		UE_LOG(LogParticles, Warning, TEXT("Particle Notify: Null PSTemplate for particle notify in anim: %s"), *GetNameSafe(Animation));
+	}
 }
 
 FString UAnimNotify_PlayParticleEffect::GetNotifyName_Implementation() const
diff --git a/Engine/Source/Runtime/Engine/Private/Animation/AnimSequence.cpp b/Engine/Source/Runtime/Engine/Private/Animation/AnimSequence.cpp
--- a/Engine/Source/Runtime/Engine/Private/Animation/AnimSequence.cpp
+++ b/Engine/Source/Runtime/Engine/Private/Animation/AnimSequence.cpp
@@ -22,7 +22,6 @@
 #define LOCTEXT_NAMESPACE "AnimSequence"
 
 DECLARE_CYCLE_STAT(TEXT("AnimSeq GetBonePose"), STAT_AnimSeq_GetBonePose, STATGROUP_Anim);
-DECLARE_CYCLE_STAT(TEXT("BlendedCurve InitFrom"), STAT_BlendedCurve_InitFrom, STATGROUP_Anim);
 DECLARE_CYCLE_STAT(TEXT("Build Anim Track Pairs"), STAT_BuildAnimTrackPairs, STATGROUP_Anim);
 DECLARE_CYCLE_STAT(TEXT("Extract Pose From Anim Data"), STAT_ExtractPoseFromAnimData, STATGROUP_Anim);
 
diff --git a/Engine/Source/Runtime/Engine/Private/Animation/AnimSingleNodeInstanceProxy.cpp b/Engine/Source/Runtime/Engine/Private/Animation/AnimSingleNodeInstanceProxy.cpp
--- a/Engine/Source/Runtime/Engine/Private/Animation/AnimSingleNodeInstanceProxy.cpp
+++ b/Engine/Source/Runtime/Engine/Private/Animation/AnimSingleNodeInstanceProxy.cpp
@@ -46,7 +46,7 @@ bool FAnimSingleNodeInstanceProxy::Evaluate(FPoseContext& Output)
 
 				if (bCanProcessAdditiveAnimationsLocal)
 				{
-					Sequence->GetAdditiveBasePose(Output.Pose, Output.Curve, ExtractionContext);
+				Sequence->GetAdditiveBasePose(Output.Pose, Output.Curve, ExtractionContext);
 				}
 				else
 				{
@@ -332,8 +332,9 @@ void FAnimSingleNodeInstanceProxy::SetAnimationAsset(class UAnimationAsset* NewA
 	PlayRate = InPlayRate;
 	CurrentTime = 0.f;
 	BlendSpaceInput = FVector::ZeroVector;
-	UpdateBlendspaceSamples(BlendSpaceInput);
+	BlendSampleData.Reset();
 	MarkerTickRecord.Reset();
+	UpdateBlendspaceSamples(BlendSpaceInput);
 
 #if WITH_EDITORONLY_DATA
 	PreviewPoseCurrentTime = 0.0f;
diff --git a/Engine/Source/Runtime/Engine/Private/Animation/BonePose.cpp b/Engine/Source/Runtime/Engine/Private/Animation/BonePose.cpp
--- a/Engine/Source/Runtime/Engine/Private/Animation/BonePose.cpp
+++ b/Engine/Source/Runtime/Engine/Private/Animation/BonePose.cpp
@@ -14,15 +14,6 @@ void FMeshPose::ResetToIdentity()
 	FAnimationRuntime::InitializeTransform(*BoneContainer, Bones);
 }
 
-void FMeshPose::InitFrom(FCompactPose& CompactPose)
-{
-	SetBoneContainer(&CompactPose.GetBoneContainer());
-	for (const FCompactPoseBoneIndex BoneIndex : CompactPose.ForEachBoneIndex())
-	{
-		FMeshPoseBoneIndex MeshPoseBoneIndex = GetBoneContainer().MakeMeshPoseIndex(BoneIndex);
-		Bones[MeshPoseBoneIndex.GetInt()] = CompactPose[BoneIndex];
-	}
-}
 
 bool FMeshPose::ContainsNaN() const
 {
@@ -55,72 +46,3 @@ bool FMeshPose::IsNormalized() const
 	return true;
 }
 
-void FCompactPose::ResetToRefPose(const FBoneContainer& RequiredBones)
-{
-	const TArray<FTransform>& RefPoseCompactArray = RequiredBones.GetRefPoseCompactArray();
-	Bones.Reset();
-	Bones.Append(RefPoseCompactArray);
-
-	// If retargeting is disabled, copy ref pose from Skeleton, rather than mesh.
-	// this is only used in editor and for debugging.
-	if (RequiredBones.GetDisableRetargeting())
-	{
-		checkSlow(RequiredBones.IsValid());
-		// Only do this if we have a mesh. otherwise we're not retargeting animations.
-		if (RequiredBones.GetSkeletalMeshAsset())
-		{
-			TArray<FTransform> const & SkeletonRefPose = RequiredBones.GetSkeletonAsset()->GetRefLocalPoses();
-
-			for (const FCompactPoseBoneIndex BoneIndex : ForEachBoneIndex())
-			{
-				const int32 SkeletonBoneIndex = GetBoneContainer().GetSkeletonIndex(BoneIndex);
-
-				// Pose bone index should always exist in Skeleton
-				checkSlow(SkeletonBoneIndex != INDEX_NONE);
-				Bones[BoneIndex.GetInt()] = SkeletonRefPose[SkeletonBoneIndex];
-			}
-		}
-	}
-}
-
-void FCompactPose::ResetToIdentity()
-{
-	for (FTransform& Bone : Bones)
-	{
-		Bone.SetIdentity();
-	}
-}
-
-bool FCompactPose::IsNormalized() const
-{
-	for (const FTransform& Bone : Bones)
-	{
-		if (!Bone.IsRotationNormalized())
-		{
-			return false;
-		}
-	}
-
-	return true;
-}
-
-bool FCompactPose::ContainsNaN() const
-{
-	for (const FTransform& Bone : Bones)
-	{
-		if (Bone.ContainsNaN())
-		{
-			return true;
-		}
-	}
-
-	return false;
-}
-
-void FCompactPose::NormalizeRotations()
-{
-	for (FTransform& Bone : Bones)
-	{
-		Bone.NormalizeRotation();
-	}
-}
diff --git a/Engine/Source/Runtime/Engine/Private/AudioDecompress.cpp b/Engine/Source/Runtime/Engine/Private/AudioDecompress.cpp
--- a/Engine/Source/Runtime/Engine/Private/AudioDecompress.cpp
+++ b/Engine/Source/Runtime/Engine/Private/AudioDecompress.cpp
@@ -331,7 +331,7 @@ void FAsyncAudioDecompressWorker::DoWork()
 			}
 			else
 			{
-				check(Wave->DecompressionType == DTYPE_Native);
+				check(Wave->DecompressionType == DTYPE_Native || Wave->DecompressionType == DTYPE_Procedural);
 
 				Wave->RawPCMDataSize = QualityInfo.SampleDataSize;
 				check(Wave->RawPCMData == nullptr);
diff --git a/Engine/Source/Runtime/Engine/Private/AudioDevice.cpp b/Engine/Source/Runtime/Engine/Private/AudioDevice.cpp
--- a/Engine/Source/Runtime/Engine/Private/AudioDevice.cpp
+++ b/Engine/Source/Runtime/Engine/Private/AudioDevice.cpp
@@ -247,8 +247,8 @@ void FAudioDevice::Teardown()
 		Sources[Index]->Stop();
 		delete Sources[Index];
 	}
-	Sources.Empty();
-	FreeSources.Empty();
+	Sources.Reset();
+	FreeSources.Reset();
 
 	if (SpatializationPlugin != nullptr)
 	{
@@ -331,8 +331,8 @@ void FAudioDevice::ResetInterpolation( void )
 		}
 	}
 
-	SoundMixModifiers.Empty();
-	PrevPassiveSoundMixModifiers.Empty();
+	SoundMixModifiers.Reset();
+	PrevPassiveSoundMixModifiers.Reset();
 	BaseSoundMix = NULL;
 
 	// reset audio effects
@@ -2499,7 +2499,6 @@ bool FAudioDevice::LocationIsAudible(const FVector& Location, const float MaxDis
 		return true;
 	}
 
-	const float MaxDistanceSquared = MaxDistance * MaxDistance;
 	for (const FListener& Listener : Listeners)
 	{
 		if (LocationIsAudible(Location, Listener, MaxDistance))
@@ -2522,6 +2521,141 @@ bool FAudioDevice::LocationIsAudible(const FVector& Location, const FListener& L
 	return (Listener.Transform.GetTranslation() - Location).SizeSquared() < MaxDistanceSquared;
 }
 
+void FAudioDevice::GetMaxDistanceAndFocusFactor(USoundBase* Sound, const class UWorld* World, const FVector& Location, const FAttenuationSettings* AttenuationSettingsToApply, float *OutMaxDistance, float* OutFocusFactor)
+{
+	check(Sound);
+	check(OutMaxDistance);
+	check(OutFocusFactor);
+
+	bool bIsInGameWorld = World ? World->IsGameWorld() : true;
+	bool bHasAttenuationSettings = (bIsInGameWorld && AttenuationSettingsToApply);
+
+	if (bHasAttenuationSettings)
+	{
+		FTransform SoundTransform;
+		SoundTransform.SetTranslation(Location);
+
+		*OutMaxDistance = AttenuationSettingsToApply->GetMaxDimension();
+
+		// Now scale the max distance based on the focus settings in the attenuation settings
+		FAttenuationListenerData ListenerData;
+		*OutFocusFactor = GetFocusFactor(ListenerData, Sound, SoundTransform, *AttenuationSettingsToApply);
+	}
+	else
+	{
+		// No need to scale the distance by focus factor since we're not using any attenuation settings
+		*OutMaxDistance = Sound->GetMaxAudibleDistance();
+		*OutFocusFactor = 1.0f;
+	}
+}
+
+bool FAudioDevice::SoundIsAudible(USoundBase* Sound, const class UWorld* World, const FVector& Location, const FAttenuationSettings* AttenuationSettingsToApply, float MaxDistance, float FocusFactor)
+{
+	bool bIsInGameWorld = World ? World->IsGameWorld() : true;
+	bool bHasAttenuationSettings = (bIsInGameWorld && AttenuationSettingsToApply);
+	float DistanceScale = 1.0f;
+	if (bHasAttenuationSettings)
+	{
+		DistanceScale = AttenuationSettingsToApply->GetFocusDistanceScale(GlobalFocusSettings, FocusFactor);
+	}
+
+	DistanceScale = FMath::Max(DistanceScale, 0.0001f);
+	return LocationIsAudible(Location, MaxDistance / DistanceScale);
+}
+
+int32 FAudioDevice::FindClosestListenerIndex(const FTransform& SoundTransform, const TArray<FListener>& InListeners)
+{
+	int32 ClosestListenerIndex = 0;
+	if (InListeners.Num() > 0)
+	{
+		float ClosestDistSq = FVector::DistSquared(SoundTransform.GetTranslation(), InListeners[0].Transform.GetTranslation());
+
+		for (int32 i = 1; i < InListeners.Num(); i++)
+		{
+			const float DistSq = FVector::DistSquared(SoundTransform.GetTranslation(), InListeners[i].Transform.GetTranslation());
+			if (DistSq < ClosestDistSq)
+			{
+				ClosestListenerIndex = i;
+				ClosestDistSq = DistSq;
+			}
+		}
+	}
+
+	return ClosestListenerIndex;
+}
+
+int32 FAudioDevice::FindClosestListenerIndex(const FTransform& SoundTransform) const
+{
+	return FindClosestListenerIndex(SoundTransform, Listeners);
+}
+
+void FAudioDevice::GetAttenuationListenerData(FAttenuationListenerData& OutListenerData, const FTransform& SoundTransform, const FAttenuationSettings& AttenuationSettings, const FListener* InListener) const
+{
+	// Only compute various components of the listener of it hasn't been computed yet
+	if (!OutListenerData.bDataComputed)
+	{
+		// Use the optional input listener param
+		OutListenerData.Listener = InListener;
+
+		// If not set, then we need to find the closest listener
+		if (!OutListenerData.Listener)
+		{
+			const int32 ClosestListenerIndex = FindClosestListenerIndex(SoundTransform);
+			OutListenerData.Listener = &Listeners[ClosestListenerIndex];
+		}
+
+		const FVector& ListenerLocation = OutListenerData.Listener->Transform.GetTranslation();
+		FVector ListenerToSound = SoundTransform.GetTranslation() - ListenerLocation;
+		ListenerToSound.ToDirectionAndLength(OutListenerData.ListenerToSoundDir, OutListenerData.ListenerToSoundDistance);
+
+		OutListenerData.AttenuationDistance = 0.0f;
+
+		if ((AttenuationSettings.bAttenuate && AttenuationSettings.AttenuationShape == EAttenuationShape::Sphere) || AttenuationSettings.bAttenuateWithLPF)
+		{
+			OutListenerData.AttenuationDistance = FMath::Max(OutListenerData.ListenerToSoundDistance - AttenuationSettings.AttenuationShapeExtents.X, 0.f);
+		}
+
+		OutListenerData.bDataComputed = true;
+	}
+}
+
+float FAudioDevice::GetFocusFactor(FAttenuationListenerData& OutListenerData, const USoundBase* Sound, const FTransform& SoundTransform, const FAttenuationSettings& AttenuationSettings, const FListener* InListener) const
+{
+	check(Sound);
+
+	// 0.0f means we are in focus, 1.0f means we are out of focus
+	float FocusFactor = 0.0f;
+
+	if (AttenuationSettings.bSpatialize && AttenuationSettings.bEnableListenerFocus && !Sound->bIgnoreFocus)
+	{
+		// Compute the focus factor based on listener emitter geometry (SoundTransform)
+		GetAttenuationListenerData(OutListenerData, SoundTransform, AttenuationSettings, InListener);
+
+		check(OutListenerData.Listener);
+
+		const FVector& ListenerForwardDir = OutListenerData.Listener->Transform.GetUnitAxis(EAxis::X);
+
+		const float FocusDotProduct = FMath::Clamp(FVector::DotProduct(ListenerForwardDir, OutListenerData.ListenerToSoundDir), 0.0f, 1.0f);
+		const float FocusAngleRadians = FMath::Acos(FocusDotProduct);
+		const float FocusAngle = FMath::RadiansToDegrees(FMath::Acos(FocusDotProduct));
+
+		const float FocusAzimuth = FMath::Clamp(GlobalFocusSettings.FocusAzimuthScale * AttenuationSettings.FocusAzimuth, 0.0f, 180.0f);
+		const float NonFocusAzimuth = FMath::Clamp(GlobalFocusSettings.NonFocusAzimuthScale * AttenuationSettings.NonFocusAzimuth, 0.0f, 180.0f);
+
+		if (FocusAzimuth != NonFocusAzimuth)
+		{
+			FocusFactor = (FocusAngle - FocusAzimuth) / (NonFocusAzimuth - FocusAzimuth);
+			FocusFactor = FMath::Clamp(FocusFactor, 0.0f, 1.0f);
+		}
+		else if (FocusAngle >= FocusAzimuth)
+		{
+			FocusFactor = 1.0f;
+		}
+	}
+
+	return FocusFactor;
+}
+
 UAudioComponent* FAudioDevice::CreateComponent(USoundBase* Sound, UWorld* World, AActor* Actor, bool bPlay, bool bStopWhenOwnerDestroyed, const FVector* Location, USoundAttenuation* AttenuationSettings, USoundConcurrency* ConcurrencySettings)
 {
 	UAudioComponent* AudioComponent = NULL;
@@ -2546,54 +2680,68 @@ UAudioComponent* FAudioDevice::CreateComponent(USoundBase* Sound, UWorld* World,
 			// Don't create component on destroyed actor.
 		}
 		// Either no actor or actor is still alive.
-		else if (Location && !Sound->IsAudibleSimple(AudioDevice, *Location, AttenuationSettings))
-		{
-			// Don't create a sound component for short sounds that start out of range of any listener
-			UE_LOG(LogAudio, Log, TEXT( "AudioComponent not created for out of range Sound %s" ), *Sound->GetName() );
-		}
 		else
 		{
+			// Listener position could change before long sounds finish
+			const FAttenuationSettings* AttenuationSettingsToApply = (AttenuationSettings ? &AttenuationSettings->Attenuation : Sound->GetAttenuationSettingsToApply());
 
-
-			// Use actor as outer if we have one.
-			if( Actor )
+			bool bIsAudible = true;
+			// If a sound is a long duration, the position might change before sound finishes so assume it's audible
+			if (Location && Sound->GetDuration() <= 1.0f)
 			{
-				AudioComponent = NewObject<UAudioComponent>(Actor);
+				float MaxDistance = 0.0f;
+				float FocusFactor = 0.0f;
+				AudioDevice->GetMaxDistanceAndFocusFactor(Sound, World, *Location, AttenuationSettingsToApply, &MaxDistance, &FocusFactor);
+				bIsAudible = AudioDevice->SoundIsAudible(Sound, World, *Location, AttenuationSettingsToApply, MaxDistance, FocusFactor);
 			}
-			// Let engine pick the outer (transient package).
-			else
+
+			if (bIsAudible)
 			{
-				AudioComponent = NewObject<UAudioComponent>();
-			}
+				// Use actor as outer if we have one.
+				if (Actor)
+				{
+					AudioComponent = NewObject<UAudioComponent>(Actor);
+				}
+				// Let engine pick the outer (transient package).
+				else
+				{
+					AudioComponent = NewObject<UAudioComponent>();
+				}
 
-			check( AudioComponent );
+				check(AudioComponent);
 
-			AudioComponent->Sound = Sound;
-			AudioComponent->bAutoActivate = false;
-			AudioComponent->bIsUISound = false;
-			AudioComponent->bAutoDestroy = bPlay;
-			AudioComponent->bStopWhenOwnerDestroyed = bStopWhenOwnerDestroyed;
+				AudioComponent->Sound = Sound;
+				AudioComponent->bAutoActivate = false;
+				AudioComponent->bIsUISound = false;
+				AudioComponent->bAutoDestroy = bPlay;
+				AudioComponent->bStopWhenOwnerDestroyed = bStopWhenOwnerDestroyed;
 #if WITH_EDITORONLY_DATA
-			AudioComponent->bVisualizeComponent	= false;
+				AudioComponent->bVisualizeComponent = false;
 #endif
-			AudioComponent->AttenuationSettings = AttenuationSettings;
-			AudioComponent->ConcurrencySettings = ConcurrencySettings;
+				AudioComponent->AttenuationSettings = AttenuationSettings;
+				AudioComponent->ConcurrencySettings = ConcurrencySettings;
 
-			if (Location)
-			{
-				AudioComponent->SetWorldLocation(*Location);
-			}
+				if (Location)
+				{
+					AudioComponent->SetWorldLocation(*Location);
+				}
 
-			// AudioComponent used in PlayEditorSound sets World to NULL to avoid situations where the world becomes invalid
-			// and the component is left with invalid pointer.
-			if( World )
-			{
-				AudioComponent->RegisterComponentWithWorld(World);
-			}
+				// AudioComponent used in PlayEditorSound sets World to NULL to avoid situations where the world becomes invalid
+				// and the component is left with invalid pointer.
+				if (World)
+				{
+					AudioComponent->RegisterComponentWithWorld(World);
+				}
 
-			if( bPlay )
+				if (bPlay)
+				{
+					AudioComponent->Play();
+				}
+			}
+			else
 			{
-				AudioComponent->Play();
+				// Don't create a sound component for short sounds that start out of range of any listener
+				UE_LOG(LogAudio, Log, TEXT("AudioComponent not created for out of range Sound %s"), *Sound->GetName());
 			}
 		}
 	}
@@ -2608,7 +2756,13 @@ void FAudioDevice::PlaySoundAtLocation(USoundBase* Sound, UWorld* World, float V
 		return;
 	}
 
-	if (Sound->IsAudibleSimple(this, Location, AttenuationSettings))
+	const FAttenuationSettings* AttenuationSettingsToApply = (AttenuationSettings ? &AttenuationSettings->Attenuation : Sound->GetAttenuationSettingsToApply());
+	float MaxDistance = 0.0f;
+	float FocusFactor = 0.0f;
+
+	GetMaxDistanceAndFocusFactor(Sound, World, Location, AttenuationSettingsToApply, &MaxDistance, &FocusFactor);
+
+	if (Sound->GetDuration() > 1.0f || SoundIsAudible(Sound, World, Location, AttenuationSettingsToApply, MaxDistance, FocusFactor))
 	{
 		const bool bIsInGameWorld = World->IsGameWorld();
 
@@ -2625,18 +2779,15 @@ void FAudioDevice::PlaySoundAtLocation(USoundBase* Sound, UWorld* World, float V
 		NewActiveSound.bHandleSubtitles = true;
 		NewActiveSound.SubtitlePriority = 10000.f; // Todo: Fix this. Add it to USoundBase
 
-		const FAttenuationSettings* AttenuationSettingsToApply = (AttenuationSettings ? &AttenuationSettings->Attenuation : Sound->GetAttenuationSettingsToApply());
 		NewActiveSound.bHasAttenuationSettings = (bIsInGameWorld && AttenuationSettingsToApply);
 		if (NewActiveSound.bHasAttenuationSettings)
 		{
 			NewActiveSound.AttenuationSettings = *AttenuationSettingsToApply;
-			NewActiveSound.MaxDistance = NewActiveSound.AttenuationSettings.GetMaxDimension();
+			NewActiveSound.FocusPriorityScale = AttenuationSettingsToApply->GetFocusPriorityScale(GlobalFocusSettings, FocusFactor);
+			NewActiveSound.FocusDistanceScale = AttenuationSettingsToApply->GetFocusDistanceScale(GlobalFocusSettings, FocusFactor);
 		}
-		else
-		{
-			NewActiveSound.MaxDistance = Sound->GetMaxAudibleDistance();
-		}
-			
+
+		NewActiveSound.MaxDistance = MaxDistance;
 		NewActiveSound.ConcurrencySettings = ConcurrencySettings;
 		NewActiveSound.Priority = Sound->Priority;
 
@@ -2691,13 +2842,13 @@ void FAudioDevice::Flush( UWorld* WorldToFlush, bool bClearActivatedReverb )
 	ProcessingPendingActiveSoundStops();
 
 	// Anytime we flush, make sure to clear all the listeners.  We'll get the right ones soon enough.
-	Listeners.Empty();
+	Listeners.Reset();
 	Listeners.Add(FListener());
 
 	// Clear all the activated reverb effects
 	if (bClearActivatedReverb)
 	{
-		ActivatedReverbs.Empty();
+		ActivatedReverbs.Reset();
 		HighestPriorityReverb = nullptr;
 	}
 
@@ -2726,7 +2877,7 @@ void FAudioDevice::Flush( UWorld* WorldToFlush, bool bClearActivatedReverb )
 				Sources[ SourceIndex ]->Stop();
 			}
 
-			WaveInstanceSourceMap.Empty();
+			WaveInstanceSourceMap.Reset();
 		}
 	}
 }
diff --git a/Engine/Source/Runtime/Engine/Private/Collision/Collision.cpp b/Engine/Source/Runtime/Engine/Private/Collision/Collision.cpp
--- a/Engine/Source/Runtime/Engine/Private/Collision/Collision.cpp
+++ b/Engine/Source/Runtime/Engine/Private/Collision/Collision.cpp
@@ -183,7 +183,12 @@ static FORCEINLINE_DEBUGGABLE bool CheckForCollision(const AActor* Actor)
 	// tell phsx that we still require queries( See FBodyInstance::UpdatePhysicsFilterData).  Doing so allows us to perform editor only traces against objects with collision disabled.  Due to this, we cannot assume that the collision enabled flag here
 	// setting is correct compared to physx and so we still ignore specified components regardless of their collision setting
 #if WITH_EDITOR
-	const bool bCheckForCollision = Actor && Actor->GetWorld() && Actor->GetWorld()->IsGameWorld();
+	bool bCheckForCollision = false;
+	if (Actor)
+	{
+		const UWorld* World = Actor->GetWorld();
+		bCheckForCollision = World && World->IsGameWorld();
+	}
 #else
 	const bool bCheckForCollision = true;
 #endif
@@ -197,7 +202,12 @@ static FORCEINLINE_DEBUGGABLE bool CheckForCollision(const UPrimitiveComponent*
 	// tell phsx that we still require queries( See FBodyInstance::UpdatePhysicsFilterData).  Doing so allows us to perform editor only traces against objects with collision disabled.  Due to this, we cannot assume that the collision enabled flag here
 	// setting is correct compared to physx and so we still ignore specified components regardless of their collision setting
 #if WITH_EDITOR
-	bool bCheckForCollision = PrimComponent && PrimComponent->GetWorld() && PrimComponent->GetWorld()->IsGameWorld();
+	bool bCheckForCollision = false;
+	if (PrimComponent)
+	{
+		const UWorld* World = PrimComponent->GetWorld();
+		bCheckForCollision = World && World->IsGameWorld();
+	}
 #else
 	const bool bCheckForCollision = true;
 #endif
@@ -209,22 +219,7 @@ void FCollisionQueryParams::AddIgnoredActor(const AActor* InIgnoreActor)
 {
 	if (InIgnoreActor)
 	{	
-		const bool bCheckForCollision = CheckForCollision(InIgnoreActor);
-		const int32 InitialCount = IgnoreComponents.Num();
-		for (const UActorComponent* ActorComponent : InIgnoreActor->GetComponents())
-		{
-			const UPrimitiveComponent* PrimComponent = Cast<const UPrimitiveComponent>(ActorComponent);
-			if (PrimComponent && (!bCheckForCollision || IsQueryCollisionEnabled(PrimComponent)) )
-			{
-				IgnoreComponents.Add(PrimComponent->GetUniqueID());
-			}
-		}
-
-		// If we added entries to a non-empty array, we aren't sure if the list is unique. Assumes GetComponents() above contains no duplicates.
-		if ((InitialCount != 0) && (IgnoreComponents.Num() - InitialCount > 0))
-		{
-			bComponentListUnique = false;
-		}
+		IgnoreActors.Add(InIgnoreActor->GetUniqueID());
 	}
 }
 
diff --git a/Engine/Source/Runtime/Engine/Private/Collision/PhysXCollision.cpp b/Engine/Source/Runtime/Engine/Private/Collision/PhysXCollision.cpp
--- a/Engine/Source/Runtime/Engine/Private/Collision/PhysXCollision.cpp
+++ b/Engine/Source/Runtime/Engine/Private/Collision/PhysXCollision.cpp
@@ -621,6 +621,7 @@ PxSceneQueryHitType::Enum FPxQueryFilterCallback::preFilter(const PxFilterData&
 
 	// Check if the shape is the right complexity for the trace 
 	PxFilterData ShapeFilter = shape->getQueryFilterData();
+	PxFilterData ShapeSimFilter = shape->getSimulationFilterData();	//This is a bit of a hack. We do this because word2 has our component ID
 #define ENABLE_PREFILTER_LOGGING 0
 #if ENABLE_PREFILTER_LOGGING
 	static bool bLoggingEnabled=false;
@@ -639,8 +640,8 @@ PxSceneQueryHitType::Enum FPxQueryFilterCallback::preFilter(const PxFilterData&
 	}
 #endif // ENABLE_PREFILTER_LOGGING
 
-	// See if we are ignoring the actor this shape belongs to (word0 of shape filterdata is actorID)
-	if(IgnoreComponents.Contains(ShapeFilter.word0))
+	// See if we are ignoring the actor this shape belongs to (word0 of shape filterdata is actorID) or the component (word2 of shape sim filter data is componentID)
+	if(IgnoreActors.Contains(ShapeFilter.word0) || IgnoreComponents.Contains(ShapeSimFilter.word2))
 	{
 		//UE_LOG(LogTemp, Log, TEXT("Ignoring Actor: %d"), ShapeFilter.word0);
 		return (PrefilterReturnValue = PxSceneQueryHitType::eNONE);
diff --git a/Engine/Source/Runtime/Engine/Private/Collision/PhysXCollision.h b/Engine/Source/Runtime/Engine/Private/Collision/PhysXCollision.h
--- a/Engine/Source/Runtime/Engine/Private/Collision/PhysXCollision.h
+++ b/Engine/Source/Runtime/Engine/Private/Collision/PhysXCollision.h
@@ -76,7 +76,10 @@ PxGeometry * GetGeometryFromShape(GeometryFromShapeStorage & LocalStorage, const
 // FILTER
 
 /** TArray typedef of components to ignore. */
-typedef TArray<uint32, TInlineAllocator<NumInlinedActorComponents>> FilterIgnoreComponentsArrayType;
+typedef FCollisionQueryParams::IgnoreComponentsArrayType FilterIgnoreComponentsArrayType;
+
+/** TArray typedef of actors to ignore. */
+typedef FCollisionQueryParams::IgnoreActorsArrayType FilterIgnoreActorsArrayType;
 
 
 /** Unreal PhysX scene query filter callback object */
@@ -84,8 +87,11 @@ class FPxQueryFilterCallback : public PxSceneQueryFilterCallback
 {
 public:
 
-	/** List of ActorIds for this query to ignore */
+	/** List of ComponentIds for this query to ignore */
 	FilterIgnoreComponentsArrayType IgnoreComponents;
+
+	/** List of ActorIds for this query to ignore */
+	FilterIgnoreActorsArrayType IgnoreActors;
 	
 	/** Result of PreFilter callback. */
 	PxSceneQueryHitType::Enum PrefilterReturnValue;
@@ -109,6 +115,7 @@ public:
 		PrefilterReturnValue = PxSceneQueryHitType::eNONE;
 		
 		IgnoreComponents = InQueryParams.GetIgnoredComponents();
+		IgnoreActors = InQueryParams.GetIgnoredActors();
 		bIgnoreTouches = false;
 		bIgnoreBlocks = InQueryParams.bIgnoreBlocks;
 	}
diff --git a/Engine/Source/Runtime/Engine/Private/Collision/PhysicsFiltering.h b/Engine/Source/Runtime/Engine/Private/Collision/PhysicsFiltering.h
--- a/Engine/Source/Runtime/Engine/Private/Collision/PhysicsFiltering.h
+++ b/Engine/Source/Runtime/Engine/Private/Collision/PhysicsFiltering.h
@@ -43,7 +43,7 @@ struct FPhysicsFilterBuilder
 		Word3 |= ExtraFilterMask;
 	}
 
-	inline void GetQueryData(uint32 ComponentID, uint32& OutWord0, uint32& OutWord1, uint32& OutWord2, uint32& OutWord3) const
+	inline void GetQueryData(uint32 ActorID, uint32& OutWord0, uint32& OutWord1, uint32& OutWord2, uint32& OutWord3) const
 	{
 		/**
 		 * Format for QueryData : 
@@ -52,13 +52,13 @@ struct FPhysicsFilterBuilder
 		 *		word2 (touching channels)
 		 *		word3 (ExtraFilter (top 4) MyChannel (top 5) as ECollisionChannel + Flags (lower 23))
 		 */
-		OutWord0 = ComponentID;
+		OutWord0 = ActorID;
 		OutWord1 = BlockingBits;
 		OutWord2 = TouchingBits;
 		OutWord3 = Word3;
 	}
 
-	inline void GetSimData(uint32 BodyIndex, uint32 SkelMeshCompID, uint32& OutWord0, uint32& OutWord1, uint32& OutWord2, uint32& OutWord3) const
+	inline void GetSimData(uint32 BodyIndex, uint32 ComponentID, uint32& OutWord0, uint32& OutWord1, uint32& OutWord2, uint32& OutWord3) const
 	{
 		/**
 		 * Format for SimData : 
@@ -69,7 +69,7 @@ struct FPhysicsFilterBuilder
 		 */
 		OutWord0 = BodyIndex;
 		OutWord1 = BlockingBits;
-		OutWord2 = SkelMeshCompID;
+		OutWord2 = ComponentID;
 		OutWord3 = Word3;
 	}
 
@@ -91,9 +91,9 @@ private:
 inline void CreateShapeFilterData(
 	const uint8 MyChannel,
 	const FMaskFilter MaskFilter,
-	const int32 ComponentID,
+	const int32 ActorID,
 	const FCollisionResponseContainer& ResponseToChannels,
-	uint32 SkelMeshCompID,
+	uint32 ComponentID,
 	uint16 BodyIndex,
 	PxFilterData& OutQueryData,
 	PxFilterData& OutSimData,
@@ -110,8 +110,8 @@ inline void CreateShapeFilterData(
 
 	OutQueryData.setToDefault();
 	OutSimData.setToDefault();
-	Builder.GetQueryData(ComponentID, OutQueryData.word0, OutQueryData.word1, OutQueryData.word2, OutQueryData.word3);
-	Builder.GetSimData(BodyIndex, SkelMeshCompID, OutSimData.word0, OutSimData.word1, OutSimData.word2, OutSimData.word3);
+	Builder.GetQueryData(ActorID, OutQueryData.word0, OutQueryData.word1, OutQueryData.word2, OutQueryData.word3);
+	Builder.GetSimData(BodyIndex, ComponentID, OutSimData.word0, OutSimData.word1, OutSimData.word2, OutSimData.word3);
 }
 #endif //WITH_PHYSX
 
diff --git a/Engine/Source/Runtime/Engine/Private/Components/AudioComponent.cpp b/Engine/Source/Runtime/Engine/Private/Components/AudioComponent.cpp
--- a/Engine/Source/Runtime/Engine/Private/Components/AudioComponent.cpp
+++ b/Engine/Source/Runtime/Engine/Private/Components/AudioComponent.cpp
@@ -171,76 +171,84 @@ void UAudioComponent::PlayInternal(const float StartTime, const float FadeInDura
 	{
 		if (FAudioDevice* AudioDevice = GetAudioDevice())
 		{
-			FActiveSound NewActiveSound;
-			NewActiveSound.SetAudioComponent(this);
-			NewActiveSound.World = GetWorld();
-			NewActiveSound.Sound = Sound;
-			NewActiveSound.SoundClassOverride = SoundClassOverride;
-			NewActiveSound.ConcurrencySettings = ConcurrencySettings;
-
-			NewActiveSound.VolumeMultiplier = (VolumeModulationMax + ((VolumeModulationMin - VolumeModulationMax) * FMath::SRand())) * VolumeMultiplier;
-			// The priority used for the active sound is the audio component's priority scaled with the sound's priority
-			if (bOverridePriority)
-			{
-				NewActiveSound.Priority = Priority;
-			}
-			else
-			{
-				NewActiveSound.Priority = Sound->Priority;
-			}
-			NewActiveSound.PitchMultiplier = (PitchModulationMax + ((PitchModulationMin - PitchModulationMax) * FMath::SRand())) * PitchMultiplier;
-			NewActiveSound.bEnableLowPassFilter = bEnableLowPassFilter;
-			NewActiveSound.LowPassFilterFrequency = LowPassFilterFrequency;
-			NewActiveSound.RequestedStartTime = FMath::Max(0.f, StartTime);
-			NewActiveSound.SubtitlePriority = SubtitlePriority;
-
-			NewActiveSound.bShouldRemainActiveIfDropped = bShouldRemainActiveIfDropped;
-			NewActiveSound.bHandleSubtitles = (!bSuppressSubtitles || OnQueueSubtitles.IsBound());
-			NewActiveSound.bIgnoreForFlushing = bIgnoreForFlushing;
-
-			NewActiveSound.bIsUISound = bIsUISound;
-			NewActiveSound.bIsMusic = bIsMusic;
-			NewActiveSound.bAlwaysPlay = bAlwaysPlay;
-			NewActiveSound.bReverb = bReverb;
-			NewActiveSound.bCenterChannelOnly = bCenterChannelOnly;
-
-			NewActiveSound.bLocationDefined = !bPreviewComponent;
-			if (NewActiveSound.bLocationDefined)
-			{
-				NewActiveSound.Transform = ComponentToWorld;
-			}
-
 			const FAttenuationSettings* AttenuationSettingsToApply = (bAllowSpatialization ? GetAttenuationSettingsToApply() : nullptr);
-			NewActiveSound.bAllowSpatialization = bAllowSpatialization;
-			NewActiveSound.bHasAttenuationSettings = (AttenuationSettingsToApply != nullptr);
-			if (NewActiveSound.bHasAttenuationSettings)
-			{
-				NewActiveSound.AttenuationSettings = *AttenuationSettingsToApply;
-				NewActiveSound.MaxDistance = NewActiveSound.AttenuationSettings.GetMaxDimension();
-			}
-			else
-			{
-				NewActiveSound.MaxDistance = Sound->GetMaxAudibleDistance();
-			}
 
-			NewActiveSound.InstanceParameters = InstanceParameters;
+			float MaxDistance = 0.0f;
+			float FocusFactor = 0.0f;
+			FVector Location = ComponentToWorld.GetLocation();
 
-			NewActiveSound.TargetAdjustVolumeMultiplier = FadeVolumeLevel;
+			AudioDevice->GetMaxDistanceAndFocusFactor(Sound, World, Location, AttenuationSettingsToApply, &MaxDistance, &FocusFactor);
 
-			if (FadeInDuration > 0.0f)
-			{
-				NewActiveSound.CurrentAdjustVolumeMultiplier = 0.f;
-				NewActiveSound.TargetAdjustVolumeStopTime = FadeInDuration;
-			}
-			else
+			if (Sound->IsLooping() || AudioDevice->SoundIsAudible(Sound, World, Location, AttenuationSettingsToApply, MaxDistance, FocusFactor))
 			{
-				NewActiveSound.CurrentAdjustVolumeMultiplier = FadeVolumeLevel;
-			}
+				FActiveSound NewActiveSound;
+				NewActiveSound.SetAudioComponent(this);
+				NewActiveSound.World = GetWorld();
+				NewActiveSound.Sound = Sound;
+				NewActiveSound.SoundClassOverride = SoundClassOverride;
+				NewActiveSound.ConcurrencySettings = ConcurrencySettings;
+
+				NewActiveSound.VolumeMultiplier = (VolumeModulationMax + ((VolumeModulationMin - VolumeModulationMax) * FMath::SRand())) * VolumeMultiplier;
+				// The priority used for the active sound is the audio component's priority scaled with the sound's priority
+				if (bOverridePriority)
+				{
+					NewActiveSound.Priority = Priority;
+				}
+				else
+				{
+					NewActiveSound.Priority = Sound->Priority;
+				}
 
-			// TODO - Audio Threading. This call would be a task call to dispatch to the audio thread
-			AudioDevice->AddNewActiveSound(NewActiveSound);
+				NewActiveSound.PitchMultiplier = (PitchModulationMax + ((PitchModulationMin - PitchModulationMax) * FMath::SRand())) * PitchMultiplier;
+				NewActiveSound.bEnableLowPassFilter = bEnableLowPassFilter;
+				NewActiveSound.LowPassFilterFrequency = LowPassFilterFrequency;
+				NewActiveSound.RequestedStartTime = FMath::Max(0.f, StartTime);
+				NewActiveSound.SubtitlePriority = SubtitlePriority;
 
-			bIsActive = true;
+				NewActiveSound.bShouldRemainActiveIfDropped = bShouldRemainActiveIfDropped;
+				NewActiveSound.bHandleSubtitles = (!bSuppressSubtitles || OnQueueSubtitles.IsBound());
+				NewActiveSound.bIgnoreForFlushing = bIgnoreForFlushing;
+
+				NewActiveSound.bIsUISound = bIsUISound;
+				NewActiveSound.bIsMusic = bIsMusic;
+				NewActiveSound.bAlwaysPlay = bAlwaysPlay;
+				NewActiveSound.bReverb = bReverb;
+				NewActiveSound.bCenterChannelOnly = bCenterChannelOnly;
+
+				NewActiveSound.bLocationDefined = !bPreviewComponent;
+				if (NewActiveSound.bLocationDefined)
+				{
+					NewActiveSound.Transform = ComponentToWorld;
+				}
+
+				NewActiveSound.bAllowSpatialization = bAllowSpatialization;
+				NewActiveSound.bHasAttenuationSettings = (AttenuationSettingsToApply != nullptr);
+				if (NewActiveSound.bHasAttenuationSettings)
+				{
+					NewActiveSound.AttenuationSettings = *AttenuationSettingsToApply;
+					NewActiveSound.FocusPriorityScale = AttenuationSettingsToApply->GetFocusPriorityScale(AudioDevice->GlobalFocusSettings, FocusFactor);
+				}
+
+				NewActiveSound.MaxDistance = MaxDistance;
+
+				NewActiveSound.InstanceParameters = InstanceParameters;
+				NewActiveSound.TargetAdjustVolumeMultiplier = FadeVolumeLevel;
+
+				if (FadeInDuration > 0.0f)
+				{
+					NewActiveSound.CurrentAdjustVolumeMultiplier = 0.f;
+					NewActiveSound.TargetAdjustVolumeStopTime = FadeInDuration;
+				}
+				else
+				{
+					NewActiveSound.CurrentAdjustVolumeMultiplier = FadeVolumeLevel;
+				}
+
+				// TODO - Audio Threading. This call would be a task call to dispatch to the audio thread
+				AudioDevice->AddNewActiveSound(NewActiveSound);
+
+				bIsActive = true;
+			}
 		}
 	}
 }
diff --git a/Engine/Source/Runtime/Engine/Private/Components/ChildActorComponent.cpp b/Engine/Source/Runtime/Engine/Private/Components/ChildActorComponent.cpp
--- a/Engine/Source/Runtime/Engine/Private/Components/ChildActorComponent.cpp
+++ b/Engine/Source/Runtime/Engine/Private/Components/ChildActorComponent.cpp
@@ -1,7 +1,6 @@
 // Copyright 1998-2016 Epic Games, Inc. All Rights Reserved.
 
 #include "EnginePrivate.h"
-#include "ComponentInstanceDataCache.h"
 #include "Components/ChildActorComponent.h"
 
 DEFINE_LOG_CATEGORY_STATIC(LogChildActorComponent, Warning, All);
@@ -94,72 +93,53 @@ void UChildActorComponent::OnComponentDestroyed(bool bDestroyingHierarchy)
 	DestroyChildActor(MyWorld && !MyWorld->IsGameWorld());
 }
 
-class FChildActorComponentInstanceData : public FSceneComponentInstanceData
+FChildActorComponentInstanceData::FChildActorComponentInstanceData(const UChildActorComponent* Component)
+	: FSceneComponentInstanceData(Component)
+	, ChildActorName(Component->ChildActorName)
+	, ComponentInstanceData(nullptr)
 {
-public:
-	FChildActorComponentInstanceData(const UChildActorComponent* Component)
-		: FSceneComponentInstanceData(Component)
-		, ChildActorName(Component->ChildActorName)
-		, ComponentInstanceData(nullptr)
+	if (Component->ChildActor)
 	{
-		if (Component->ChildActor)
+		ComponentInstanceData = new FComponentInstanceDataCache(Component->ChildActor);
+		// If it is empty dump it
+		if (!ComponentInstanceData->HasInstanceData())
 		{
-			ComponentInstanceData = new FComponentInstanceDataCache(Component->ChildActor);
-			// If it is empty dump it
-			if (!ComponentInstanceData->HasInstanceData())
-			{
-				delete ComponentInstanceData;
-				ComponentInstanceData = nullptr;
-			}
+			delete ComponentInstanceData;
+			ComponentInstanceData = nullptr;
+		}
 
-			USceneComponent* ChildRootComponent = Component->ChildActor->GetRootComponent();
-			if (ChildRootComponent)
+		USceneComponent* ChildRootComponent = Component->ChildActor->GetRootComponent();
+		if (ChildRootComponent)
+		{
+			for (USceneComponent* AttachedComponent : ChildRootComponent->AttachChildren)
 			{
-				for (USceneComponent* AttachedComponent : ChildRootComponent->AttachChildren)
+				if (AttachedComponent)
 				{
-					if (AttachedComponent)
+					AActor* AttachedActor = AttachedComponent->GetOwner();
+					if (AttachedActor != Component->ChildActor)
 					{
-						AActor* AttachedActor = AttachedComponent->GetOwner();
-						if (AttachedActor != Component->ChildActor)
-						{
-							FAttachedActorInfo Info;
-							Info.Actor = AttachedActor;
-							Info.SocketName = AttachedComponent->AttachSocketName;
-							Info.RelativeTransform = AttachedComponent->GetRelativeTransform();
-							AttachedActors.Add(Info);
-						}
+						FAttachedActorInfo Info;
+						Info.Actor = AttachedActor;
+						Info.SocketName = AttachedComponent->AttachSocketName;
+						Info.RelativeTransform = AttachedComponent->GetRelativeTransform();
+						AttachedActors.Add(Info);
 					}
 				}
 			}
 		}
 	}
+}
 
-	virtual ~FChildActorComponentInstanceData()
-	{
-		delete ComponentInstanceData;
-	}
-
-	virtual void ApplyToComponent(UActorComponent* Component, const ECacheApplyPhase CacheApplyPhase) override
-	{
-		FSceneComponentInstanceData::ApplyToComponent(Component, CacheApplyPhase);
-		CastChecked<UChildActorComponent>(Component)->ApplyComponentInstanceData(this, CacheApplyPhase);
-	}
-
-	// The name of the spawned child actor so it (attempts to) remain constant across construction script reruns
-	FName ChildActorName;
-
-	// The component instance data cache for the ChildActor spawned by this component
-	FComponentInstanceDataCache* ComponentInstanceData;
-
-	struct FAttachedActorInfo
-	{
-		TWeakObjectPtr<AActor> Actor;
-		FName SocketName;
-		FTransform RelativeTransform;
-	};
+FChildActorComponentInstanceData::~FChildActorComponentInstanceData()
+{
+	delete ComponentInstanceData;
+}
 
-	TArray<FAttachedActorInfo> AttachedActors;
-};
+void FChildActorComponentInstanceData::ApplyToComponent(UActorComponent* Component, const ECacheApplyPhase CacheApplyPhase)
+{
+	FSceneComponentInstanceData::ApplyToComponent(Component, CacheApplyPhase);
+	CastChecked<UChildActorComponent>(Component)->ApplyComponentInstanceData(this, CacheApplyPhase);
+}
 
 void UChildActorComponent::BeginDestroy()
 {
diff --git a/Engine/Source/Runtime/Engine/Private/Components/DestructibleComponent.cpp b/Engine/Source/Runtime/Engine/Private/Components/DestructibleComponent.cpp
--- a/Engine/Source/Runtime/Engine/Private/Components/DestructibleComponent.cpp
+++ b/Engine/Source/Runtime/Engine/Private/Components/DestructibleComponent.cpp
@@ -293,7 +293,7 @@ void UDestructibleComponent::CreatePhysicsState()
 
 	// Passing AssetInstanceID = 0 so we'll have self-collision
 	AActor* Owner = GetOwner();
-	CreateShapeFilterData(MoveChannel, FMaskFilter(0), GetUniqueID(), CollResponse, 0, 0, PQueryFilterData, PSimFilterData, BodyInstance.bUseCCD, bEnableImpactDamage, false, bEnableContactModification);
+	CreateShapeFilterData(MoveChannel, FMaskFilter(0), Owner->GetUniqueID(), CollResponse, GetUniqueID(), 0, PQueryFilterData, PSimFilterData, BodyInstance.bUseCCD, bEnableImpactDamage, false, bEnableContactModification);
 
 	// Build filterData variations for complex and simple
 	PSimFilterData.word3 |= EPDF_SimpleCollision | EPDF_ComplexCollision;
@@ -1458,7 +1458,7 @@ void UDestructibleComponent::SetCollisionResponseForActor(PxRigidDynamic* Actor,
 		physx::PxU32 SupportDepth = TheDestructibleMesh->ApexDestructibleAsset->getChunkDepth(ChunkIdx);
 
 		const bool bEnableImpactDamage = IsImpactDamageEnabled(TheDestructibleMesh, SupportDepth);
-		CreateShapeFilterData(MoveChannel, FMaskFilter(0), GetUniqueID(), UseResponse, 0, ChunkIdxToBoneIdx(ChunkIdx), PQueryFilterData, PSimFilterData, BodyInstance.bUseCCD, bEnableImpactDamage, false);
+		CreateShapeFilterData(MoveChannel, FMaskFilter(0), Owner->GetUniqueID(), UseResponse, GetUniqueID(), ChunkIdxToBoneIdx(ChunkIdx), PQueryFilterData, PSimFilterData, BodyInstance.bUseCCD, bEnableImpactDamage, false);
 		
 		PQueryFilterData.word3 |= EPDF_SimpleCollision | EPDF_ComplexCollision;
 
@@ -1501,7 +1501,7 @@ void UDestructibleComponent::SetCollisionResponseForShape(PxShape* Shape, int32
 		bool bLargeChunk = IsChunkLarge(Shape->getActor());
 		const FCollisionResponse& ColResponse = bLargeChunk ? LargeChunkCollisionResponse : SmallChunkCollisionResponse;
 		//TODO: we currently assume chunks will not have impact damage as it's very expensive. Should look into exposing this a bit more
-		CreateShapeFilterData(MoveChannel, FMaskFilter(0), (Owner ? Owner->GetUniqueID() : 0), ColResponse.GetResponseContainer(), 0, ChunkIdxToBoneIdx(ChunkIdx), PQueryFilterData, PSimFilterData, BodyInstance.bUseCCD, false, false);
+		CreateShapeFilterData(MoveChannel, FMaskFilter(0), (Owner ? Owner->GetUniqueID() : 0), ColResponse.GetResponseContainer(), GetUniqueID(), ChunkIdxToBoneIdx(ChunkIdx), PQueryFilterData, PSimFilterData, BodyInstance.bUseCCD, false, false);
 
 		PQueryFilterData.word3 |= EPDF_SimpleCollision | EPDF_ComplexCollision;
 
diff --git a/Engine/Source/Runtime/Engine/Private/Components/SceneComponent.cpp b/Engine/Source/Runtime/Engine/Private/Components/SceneComponent.cpp
--- a/Engine/Source/Runtime/Engine/Private/Components/SceneComponent.cpp
+++ b/Engine/Source/Runtime/Engine/Private/Components/SceneComponent.cpp
@@ -2290,7 +2290,7 @@ bool USceneComponent::ShouldRender() const
 		false;
 #endif
 	UWorld *World = GetWorld();
-	const bool bInGameWorld = GetWorld() && GetWorld()->UsesGameHiddenFlags();
+	const bool bInGameWorld = World && World->UsesGameHiddenFlags();
 
 	const bool bShowInGame = IsVisible() && (!Owner || !Owner->bHidden);
 	return ((bInGameWorld && bShowInGame) || (!bInGameWorld && bShowInEditor)) && bVisible == true;
diff --git a/Engine/Source/Runtime/Engine/Private/Components/SkeletalMeshComponent.cpp b/Engine/Source/Runtime/Engine/Private/Components/SkeletalMeshComponent.cpp
--- a/Engine/Source/Runtime/Engine/Private/Components/SkeletalMeshComponent.cpp
+++ b/Engine/Source/Runtime/Engine/Private/Components/SkeletalMeshComponent.cpp
@@ -1009,7 +1009,7 @@ void USkeletalMeshComponent::RecalcRequiredBones(int32 LODIndex)
 	CachedSpaceBases.Empty();
 }
 
-void USkeletalMeshComponent::EvaluateAnimation(const USkeletalMesh* InSkeletalMesh, UAnimInstance* InAnimInstance, TArray<FTransform>& OutLocalAtoms, TArray<FActiveVertexAnim>& OutVertexAnims, FVector& OutRootBoneTranslation, FBlendedCurve& OutCurve) const
+void USkeletalMeshComponent::EvaluateAnimation(const USkeletalMesh* InSkeletalMesh, UAnimInstance* InAnimInstance, TArray<FTransform>& OutLocalAtoms, TArray<FActiveVertexAnim>& OutVertexAnims, FVector& OutRootBoneTranslation, FBlendedHeapCurve& OutCurve) const
 {
 	ANIM_MT_SCOPE_CYCLE_COUNTER(SkeletalComponentAnimEvaluate, IsRunningParallelEvaluation());
 
@@ -1059,12 +1059,10 @@ void USkeletalMeshComponent::UpdateSlaveComponent()
 	Super::UpdateSlaveComponent();
 }
 
-void USkeletalMeshComponent::PerformAnimationEvaluation(const USkeletalMesh* InSkeletalMesh, UAnimInstance* InAnimInstance, TArray<FTransform>& OutSpaceBases, TArray<FTransform>& OutLocalAtoms, TArray<FActiveVertexAnim>& OutVertexAnims, FVector& OutRootBoneTranslation, FBlendedCurve& OutCurve) const
+void USkeletalMeshComponent::PerformAnimationEvaluation(const USkeletalMesh* InSkeletalMesh, UAnimInstance* InAnimInstance, TArray<FTransform>& OutSpaceBases, TArray<FTransform>& OutLocalAtoms, TArray<FActiveVertexAnim>& OutVertexAnims, FVector& OutRootBoneTranslation, FBlendedHeapCurve& OutCurve) const
 {
 	ANIM_MT_SCOPE_CYCLE_COUNTER(PerformAnimEvaluation, IsRunningParallelEvaluation());
 
-	FMemMark StackMemoryMark(FMemStack::Get());
-
 	// Can't do anything without a SkeletalMesh
 	// Do nothing more if no bones in skeleton.
 	if (!InSkeletalMesh || OutSpaceBases.Num() == 0)
diff --git a/Engine/Source/Runtime/Engine/Private/Components/SkinnedMeshComponent.cpp b/Engine/Source/Runtime/Engine/Private/Components/SkinnedMeshComponent.cpp
--- a/Engine/Source/Runtime/Engine/Private/Components/SkinnedMeshComponent.cpp
+++ b/Engine/Source/Runtime/Engine/Private/Components/SkinnedMeshComponent.cpp
@@ -516,15 +516,16 @@ void USkinnedMeshComponent::TickComponent(float DeltaTime, enum ELevelTick TickT
 	// See if this mesh was rendered recently. This has to happen first because other data will rely on this
 	bRecentlyRendered = (LastRenderTime > GetWorld()->TimeSeconds - 1.0f);
 
+	// Update component's LOD settings
+	// This must be done BEFORE animation Update and Evaluate (TickPose and RefreshBoneTransforms respectively)
+	const bool bLODHasChanged = UpdateLODStatus();
+
 	// Tick Pose first
 	if (ShouldTickPose())
 	{
 		TickPose(DeltaTime, false);
 	}
 
-	// Update component's LOD settings
-	const bool bLODHasChanged = UpdateLODStatus();
-
 	// If we have been recently rendered, and bForceRefPose has been on for at least a frame, or the LOD changed, update bone matrices.
 	if( ShouldUpdateTransform(bLODHasChanged) )
 	{
diff --git a/Engine/Source/Runtime/Engine/Private/Components/TextRenderComponent.cpp b/Engine/Source/Runtime/Engine/Private/Components/TextRenderComponent.cpp
--- a/Engine/Source/Runtime/Engine/Private/Components/TextRenderComponent.cpp
+++ b/Engine/Source/Runtime/Engine/Private/Components/TextRenderComponent.cpp
@@ -173,9 +173,9 @@ public:
 // @param It must be a valid initialized text iterator
 // @param Font 0 is silently ignored
 FVector2D ComputeTextSize(FTextIterator It, class UFont* Font,
-	float XScale, float YScale, float HorizSpacingAdjust)
+	float XScale, float YScale, float HorizSpacingAdjust, float VertSpacingAdjust)
 {
-	FVector2D Ret(0, 0);
+	FVector2D Ret(0.f, 0.f);
 
 	if(!Font)
 	{
@@ -184,7 +184,7 @@ FVector2D ComputeTextSize(FTextIterator It, class UFont* Font,
 
 	const float CharIncrement = ( (float)Font->Kerning + HorizSpacingAdjust ) * XScale;
 
-	float LineX = 0;
+	float LineX = 0.f;
 
 	const int32 PageIndex = 0;
 
@@ -215,14 +215,14 @@ FVector2D ComputeTextSize(FTextIterator It, class UFont* Font,
 			const float X		= LineX;
 			const float Y		= Char.VerticalOffset * YScale;
 			float SizeX			= Char.USize * XScale;
-			const float SizeY	= Char.VSize * YScale;
+			const float SizeY	= (Char.VSize + VertSpacingAdjust) * YScale;
 			const float U		= Char.StartU * InvTextureSize.X;
 			const float V		= Char.StartV * InvTextureSize.Y;
 			const float SizeU	= Char.USize * InvTextureSize.X;
 			const float SizeV	= Char.VSize * InvTextureSize.Y;
 
-			float Right = X + SizeX;
-			float Bottom = Y + SizeY;
+			const float Right = X + SizeX;
+			const float Bottom = Y + SizeY;
 
 			Ret.X = FMath::Max(Ret.X, Right);
 			Ret.Y = FMath::Max(Ret.Y, Bottom);
@@ -244,7 +244,7 @@ FVector2D ComputeTextSize(FTextIterator It, class UFont* Font,
 // compute the left top depending on the alignment
 static float ComputeHorizontalAlignmentOffset(FVector2D Size, EHorizTextAligment HorizontalAlignment)
 {
-	float Ret = 0;
+	float Ret = 0.f;
 
 	switch (HorizontalAlignment)
 	{
@@ -290,7 +290,7 @@ float ComputeVerticalAlignmentOffset(float SizeY, EVerticalTextAligment Vertical
 		}
 	case EVRTA_TextTop:
 		{
-			return 0;
+			return 0.f;
 		}
 	case EVRTA_TextCenter:
 		{
@@ -299,7 +299,7 @@ float ComputeVerticalAlignmentOffset(float SizeY, EVerticalTextAligment Vertical
 	default:
 		{
 			check(0);
-			return 0;
+			return 0.f;
 		}
 	}
 }
@@ -314,25 +314,26 @@ float CalculateVerticalAlignmentOffset(
 	float XScale, 
 	float YScale, 
 	float HorizSpacingAdjust,
+	float VertSpacingAdjust,
 	EVerticalTextAligment VerticalAlignment)
 {
 	check(Text);
 
 	if(!Font)
 	{
-		return 0;
+		return 0.f;
 	}
 
-	float FirstLineHeight = -1; // Only kept around for legacy positioning support
-	float StartY = 0;
+	float FirstLineHeight = -1.f; // Only kept around for legacy positioning support
+	float StartY = 0.f;
 
 	FTextIterator It(Text);
 
 	while (It.NextLine())
 	{
-		FVector2D LineSize = ComputeTextSize(It, Font, XScale, YScale, HorizSpacingAdjust);
+		FVector2D LineSize = ComputeTextSize(It, Font, XScale, YScale, HorizSpacingAdjust, VertSpacingAdjust);
 
-		if (FirstLineHeight < 0)
+		if (FirstLineHeight < 0.f)
 		{
 			FirstLineHeight = LineSize.Y;
 		}
@@ -345,7 +346,7 @@ float CalculateVerticalAlignmentOffset(
 		}
 
 		// Move Y position down to next line. If the current line is empty, move by max char height in font
-		StartY += LineSize.Y > 0 ? LineSize.Y : Font->GetMaxCharHeight();
+		StartY += LineSize.Y > 0.f ? LineSize.Y : Font->GetMaxCharHeight();
 	}
 
 	// Calculate a vertical translation to create the correct vertical alignment
@@ -391,6 +392,7 @@ private:
 	float XScale;
 	float YScale;
 	float HorizSpacingAdjust;
+	float VertSpacingAdjust;
 	EHorizTextAligment HorizontalAlignment;
 	EVerticalTextAligment VerticalAlignment;
 	bool bAlwaysRenderAsText;
@@ -404,12 +406,13 @@ FTextRenderSceneProxy::FTextRenderSceneProxy( UTextRenderComponent* Component) :
 	XScale(Component->WorldSize * Component->XScale * Component->InvDefaultSize),
 	YScale(Component->WorldSize * Component->YScale * Component->InvDefaultSize),
 	HorizSpacingAdjust(Component->HorizSpacingAdjust),
+	VertSpacingAdjust(Component->VertSpacingAdjust),
 	HorizontalAlignment(Component->HorizontalAlignment),
 	VerticalAlignment(Component->VerticalAlignment),
 	bAlwaysRenderAsText(Component->bAlwaysRenderAsText)
 {
-	WireframeColor = FLinearColor(1, 0, 0);
-	UMaterialInterface* EffectiveMaterial = 0;
+	WireframeColor = FLinearColor(1.f, 0.f, 0.f);
+	UMaterialInterface* EffectiveMaterial = nullptr;
 
 	if(Component->TextMaterial)
 	{
@@ -579,27 +582,27 @@ bool  FTextRenderSceneProxy::BuildStringMesh( TArray<FDynamicMeshVertex>& OutVer
 		return false;
 	}
 
-	float FirstLineHeight = -1; // Only kept around for legacy positioning support
-	float StartY = 0;
+	float FirstLineHeight = -1.f; // Only kept around for legacy positioning support
+	float StartY = 0.f;
 
 	const float CharIncrement = ( (float)Font->Kerning + HorizSpacingAdjust ) * XScale;
 
-	float LineX = 0;
+	float LineX = 0.f;
 
 	const int32 PageIndex = 0;
 
 	FTextIterator It(*Text.ToString());
 	while (It.NextLine())
 	{
-		FVector2D LineSize = ComputeTextSize(It, Font, XScale, YScale, HorizSpacingAdjust);
+		FVector2D LineSize = ComputeTextSize(It, Font, XScale, YScale, HorizSpacingAdjust, VertSpacingAdjust);
 		float StartX = ComputeHorizontalAlignmentOffset(LineSize, HorizontalAlignment);
 
-		if (FirstLineHeight < 0)
+		if (FirstLineHeight < 0.f)
 		{
 			FirstLineHeight = LineSize.Y;
 		}
 
-		LineX = 0;
+		LineX = 0.f;
 		int32 Ch;
 
 		while (It.NextCharacterInLine(Ch))
@@ -634,25 +637,25 @@ bool  FTextRenderSceneProxy::BuildStringMesh( TArray<FDynamicMeshVertex>& OutVer
 				const float SizeU	= Char.USize * InvTextureSize.X;
 				const float SizeV	= Char.VSize * InvTextureSize.Y;			
 
-				float Left = X;
-				float Top = Y;
-				float Right = X + SizeX;
-				float Bottom = Y + SizeY;
+				const float Left = X;
+				const float Top = Y;
+				const float Right = X + SizeX;
+				const float Bottom = Y + SizeY;
 
 				// axis choice and sign to get good alignment when placed on surface
-				FVector4 V0 = FVector4(0, -Left, -Top, 0);
-				FVector4 V1 = FVector4(0, -Right, -Top, 0);
-				FVector4 V2 = FVector4(0, -Left, -Bottom, 0);
-				FVector4 V3 = FVector4(0, -Right, -Bottom, 0);
+				const FVector4 V0(0.f, -Left, -Top, 0.f);
+				const FVector4 V1(0.f, -Right, -Top, 0.f);
+				const FVector4 V2(0.f, -Left, -Bottom, 0.f);
+				const FVector4 V3(0.f, -Right, -Bottom, 0.f);
 
-				FVector TangentX(0, -1, 0);
-				FVector TangentY(0, 0, -1);
-				FVector TangentZ(1, 0, 0);
+				const FVector TangentX(0.f, -1.f, 0.f);
+				const FVector TangentY(0.f, 0.f, -1.f);
+				const FVector TangentZ(1.f, 0.f, 0.f);
 
-				int32 V00 = OutVertices.Add(FDynamicMeshVertex(V0, TangentX, TangentZ, FVector2D(U, V), TextRenderColor));
-				int32 V10 = OutVertices.Add(FDynamicMeshVertex(V1, TangentX, TangentZ, FVector2D(U + SizeU, V), TextRenderColor));
-				int32 V01 = OutVertices.Add(FDynamicMeshVertex(V2, TangentX, TangentZ, FVector2D(U, V + SizeV), TextRenderColor));
-				int32 V11 = OutVertices.Add(FDynamicMeshVertex(V3, TangentX, TangentZ, FVector2D(U + SizeU, V + SizeV), TextRenderColor));
+				const int32 V00 = OutVertices.Add(FDynamicMeshVertex(V0, TangentX, TangentZ, FVector2D(U, V), TextRenderColor));
+				const int32 V10 = OutVertices.Add(FDynamicMeshVertex(V1, TangentX, TangentZ, FVector2D(U + SizeU, V), TextRenderColor));
+				const int32 V01 = OutVertices.Add(FDynamicMeshVertex(V2, TangentX, TangentZ, FVector2D(U, V + SizeV), TextRenderColor));
+				const int32 V11 = OutVertices.Add(FDynamicMeshVertex(V3, TangentX, TangentZ, FVector2D(U + SizeU, V + SizeV), TextRenderColor));
 
 				check(V00 < 65536);
 				check(V10 < 65536);
@@ -679,7 +682,7 @@ bool  FTextRenderSceneProxy::BuildStringMesh( TArray<FDynamicMeshVertex>& OutVer
 		}
 
 		// Move Y position down to next line. If the current line is empty, move by max char height in font
-		StartY += LineSize.Y > 0 ? LineSize.Y : Font->GetMaxCharHeight();
+		StartY += LineSize.Y > 0.f ? LineSize.Y : Font->GetMaxCharHeight();
 	}
 
 	// Avoid initializing RHI resources when no vertices are generated.
@@ -754,9 +757,10 @@ UTextRenderComponent::UTextRenderComponent(const FObjectInitializer& ObjectIniti
 
 		SetCollisionProfileName(UCollisionProfile::NoCollision_ProfileName);
 		TextRenderColor = FColor::White;
-		XScale = 1;
-		YScale = 1;
-		HorizSpacingAdjust = 0;
+		XScale = 1.f;
+		YScale = 1.f;
+		HorizSpacingAdjust = 0.f;
+		VertSpacingAdjust = 0.f;
 		HorizontalAlignment = EHTA_Left;
 		VerticalAlignment = EVRTA_TextBottom;
 
@@ -801,7 +805,7 @@ void UTextRenderComponent::SetMaterial(int32 ElementIndex, UMaterialInterface* I
 
 UMaterialInterface* UTextRenderComponent::GetMaterial(int32 ElementIndex) const
 {
-	return (ElementIndex == 0) ? TextMaterial : NULL;
+	return (ElementIndex == 0) ? TextMaterial : nullptr;
 }
 
 bool UTextRenderComponent::ShouldRecreateProxyOnUpdateTransform() const
@@ -815,25 +819,25 @@ FBoxSphereBounds UTextRenderComponent::CalcBounds(const FTransform& LocalToWorld
 {
 	if(!Text.IsEmpty() && Font)
 	{
-		FVector2D Size(FLT_MIN, 0);
+		FVector2D Size(FLT_MIN, 0.f);
 		FVector2D LeftTop(FLT_MAX, FLT_MAX);
-		float FirstLineHeight = -1;
+		float FirstLineHeight = -1.f;
 
 		FTextIterator It(*Text.ToString());
 
-		float AdjustedXScale = WorldSize * XScale * InvDefaultSize;
-		float AdjustedYScale = WorldSize * YScale * InvDefaultSize;
+		const float AdjustedXScale = WorldSize * XScale * InvDefaultSize;
+		const float AdjustedYScale = WorldSize * YScale * InvDefaultSize;
 
 		while (It.NextLine())
 		{
-			FVector2D LineSize = ComputeTextSize(It, Font, AdjustedXScale, AdjustedYScale, HorizSpacingAdjust);
-			float LineLeft = ComputeHorizontalAlignmentOffset(LineSize, HorizontalAlignment);
+			const FVector2D LineSize = ComputeTextSize(It, Font, AdjustedXScale, AdjustedYScale, HorizSpacingAdjust, VertSpacingAdjust);
+			const float LineLeft = ComputeHorizontalAlignmentOffset(LineSize, HorizontalAlignment);
 
 			Size.X = FMath::Max(LineSize.X, Size.X);
-			Size.Y += LineSize.Y > 0 ? LineSize.Y : Font->GetMaxCharHeight();
+			Size.Y += LineSize.Y > 0.f ? LineSize.Y : Font->GetMaxCharHeight();
 			LeftTop.X = FMath::Min(LeftTop.X, LineLeft);
 
-			if (FirstLineHeight < 0)
+			if (FirstLineHeight < 0.f)
 			{
 				FirstLineHeight = LineSize.Y;
 			}
@@ -843,7 +847,7 @@ FBoxSphereBounds UTextRenderComponent::CalcBounds(const FTransform& LocalToWorld
 		}
 
 		LeftTop.Y = ComputeVerticalAlignmentOffset(Size.Y, VerticalAlignment, FirstLineHeight);
-		FBox LocalBox(FVector(0, -LeftTop.X, -LeftTop.Y), FVector(0, -(LeftTop.X + Size.X), -(LeftTop.Y + Size.Y)));
+		const FBox LocalBox(FVector(0.f, -LeftTop.X, -LeftTop.Y), FVector(0.f, -(LeftTop.X + Size.X), -(LeftTop.Y + Size.Y)));
 
 		FBoxSphereBounds Ret(LocalBox.TransformBy(LocalToWorld));
 
@@ -863,18 +867,18 @@ FMatrix UTextRenderComponent::GetRenderMatrix() const
 	// Adjust LocalToWorld transform to account for vertical text alignment when rendering.
 	if(!Text.IsEmpty() && Font)
 	{
-		float SizeY = 0;
-		float FirstLineHeight = -1;
-		float AdjustedXScale = WorldSize * XScale * InvDefaultSize;
-		float AdjustedYScale = WorldSize * YScale * InvDefaultSize;
+		float SizeY = 0.f;
+		float FirstLineHeight = -1.f;
+		const float AdjustedXScale = WorldSize * XScale * InvDefaultSize;
+		const float AdjustedYScale = WorldSize * YScale * InvDefaultSize;
 
 		FTextIterator It(*Text.ToString());
 		while (It.NextLine())
 		{
-			FVector2D LineSize = ComputeTextSize(It, Font, AdjustedXScale, AdjustedYScale, HorizSpacingAdjust);
-			SizeY += LineSize.Y > 0 ? LineSize.Y : Font->GetMaxCharHeight();
+			const FVector2D LineSize = ComputeTextSize(It, Font, AdjustedXScale, AdjustedYScale, HorizSpacingAdjust, VertSpacingAdjust);
+			SizeY += LineSize.Y > 0.f ? LineSize.Y : Font->GetMaxCharHeight();
 
-			if (FirstLineHeight < 0)
+			if (FirstLineHeight < 0.f)
 			{
 				FirstLineHeight = LineSize.Y;
 			}
@@ -885,8 +889,8 @@ FMatrix UTextRenderComponent::GetRenderMatrix() const
 
 		// Calculate a vertical translation to create the correct vertical alignment
 		FMatrix VerticalTransform = FMatrix::Identity;
-		float VerticalAlignmentOffset = -ComputeVerticalAlignmentOffset(SizeY, VerticalAlignment, FirstLineHeight);
-		VerticalTransform = VerticalTransform.ConcatTranslation(FVector(0, 0, VerticalAlignmentOffset));
+		const float VerticalAlignmentOffset = -ComputeVerticalAlignmentOffset(SizeY, VerticalAlignment, FirstLineHeight);
+		VerticalTransform = VerticalTransform.ConcatTranslation(FVector(0.f, 0.f, VerticalAlignmentOffset));
 
 		return VerticalTransform *  ComponentToWorld.ToMatrixWithScale();
 
@@ -961,6 +965,12 @@ void UTextRenderComponent::SetHorizSpacingAdjust(float Value)
 	MarkRenderStateDirty();	
 }
 
+void UTextRenderComponent::SetVertSpacingAdjust(float Value)
+{
+	VertSpacingAdjust = Value;
+	MarkRenderStateDirty();	
+}
+
 void UTextRenderComponent::SetWorldSize(float Value)
 {
 	WorldSize = Value;
@@ -969,13 +979,13 @@ void UTextRenderComponent::SetWorldSize(float Value)
 
 FVector UTextRenderComponent::GetTextLocalSize() const
 {
-	FBoxSphereBounds Bounds = CalcBounds(FTransform::Identity);
+	const FBoxSphereBounds Bounds = CalcBounds(FTransform::Identity);
 	return Bounds.GetBox().GetSize();
 }
 
 FVector UTextRenderComponent::GetTextWorldSize() const
 {
-	FBoxSphereBounds Bounds = CalcBounds(ComponentToWorld);
+	const FBoxSphereBounds Bounds = CalcBounds(ComponentToWorld);
 	return Bounds.GetBox().GetSize();
 }
 
@@ -985,8 +995,8 @@ void UTextRenderComponent::PostLoad()
 	// fix with its own version, use the version number closest to that CL
 	if (GetLinkerUE4Version() < VER_UE4_PACKAGE_REQUIRES_LOCALIZATION_GATHER_FLAGGING)
 	{
-		float Offset = CalculateVerticalAlignmentOffset(*Text.ToString(), Font, XScale, YScale, HorizSpacingAdjust, VerticalAlignment);
-		FTransform RelativeTransform = GetRelativeTransform();
+		const float Offset = CalculateVerticalAlignmentOffset(*Text.ToString(), Font, XScale, YScale, HorizSpacingAdjust, VertSpacingAdjust, VerticalAlignment);
+		const FTransform RelativeTransform = GetRelativeTransform();
 		FTransform CorrectionLeft = FTransform::Identity;
 		FTransform CorrectionRight = FTransform::Identity;
 		CorrectionLeft.SetTranslation(FVector(0.0f, 0.0f, -Offset));
diff --git a/Engine/Source/Runtime/Engine/Private/Controller.cpp b/Engine/Source/Runtime/Engine/Private/Controller.cpp
--- a/Engine/Source/Runtime/Engine/Private/Controller.cpp
+++ b/Engine/Source/Runtime/Engine/Private/Controller.cpp
@@ -42,6 +42,7 @@ AController::AController(const FObjectInitializer& ObjectInitializer)
 
 	bCanBeDamaged = false;
 	bAttachToPawn = false;
+	bIsPlayerController = false;
 
 	if (RootComponent)
 	{
@@ -56,13 +57,14 @@ void AController::K2_DestroyActor()
 	// do nothing, disallow destroying controller from Blueprints
 }
 
-bool AController::IsLocalPlayerController() const
-{
-	return false;
-}
-
 bool AController::IsLocalController() const
 {
+	if (FGenericPlatformProperties::IsServerOnly())
+	{
+		// Never local on dedicated server. IsServerOnly() is checked at compile time and optimized out if false.
+		return false;
+	}
+
 	const ENetMode NetMode = GetNetMode();
 
 	if (NetMode == NM_Standalone)
