Copying //UE4/Dev-Framework to //UE4/Main @ 2830052

diff --git a/Engine/Source/Runtime/Engine/Private/DrawDebugHelpers.cpp b/Engine/Source/Runtime/Engine/Private/DrawDebugHelpers.cpp
--- a/Engine/Source/Runtime/Engine/Private/DrawDebugHelpers.cpp
+++ b/Engine/Source/Runtime/Engine/Private/DrawDebugHelpers.cpp
@@ -83,7 +83,7 @@ void DrawDebugDirectionalArrow(const UWorld* InWorld, FVector const& LineStart,
 	}
 }
 
-void DrawDebugBox(const UWorld* InWorld, FVector const& Center, FVector const& Box, FColor const& Color, bool bPersistentLines, float LifeTime, uint8 DepthPriority)
+void DrawDebugBox(const UWorld* InWorld, FVector const& Center, FVector const& Box, FColor const& Color, bool bPersistentLines, float LifeTime, uint8 DepthPriority, float Thickness)
 {
 	// no debug line drawing on dedicated server
 	if (GEngine->GetNetMode(InWorld) != NM_DedicatedServer)
@@ -94,25 +94,25 @@ void DrawDebugBox(const UWorld* InWorld, FVector const& Center, FVector const& B
 		{
 			float LineLifeTime = (LifeTime > 0.f)? LifeTime : LineBatcher->DefaultLifeTime;
 
-			LineBatcher->DrawLine(Center + FVector( Box.X,  Box.Y,  Box.Z), Center + FVector( Box.X, -Box.Y, Box.Z), Color, DepthPriority, 0.f, LineLifeTime);
-			LineBatcher->DrawLine(Center + FVector( Box.X, -Box.Y,  Box.Z), Center + FVector(-Box.X, -Box.Y, Box.Z), Color, DepthPriority, 0.f, LineLifeTime);
-			LineBatcher->DrawLine(Center + FVector(-Box.X, -Box.Y,  Box.Z), Center + FVector(-Box.X,  Box.Y, Box.Z), Color, DepthPriority, 0.f, LineLifeTime);
-			LineBatcher->DrawLine(Center + FVector(-Box.X,  Box.Y,  Box.Z), Center + FVector( Box.X,  Box.Y, Box.Z), Color, DepthPriority, 0.f, LineLifeTime);
+			LineBatcher->DrawLine(Center + FVector( Box.X,  Box.Y,  Box.Z), Center + FVector( Box.X, -Box.Y, Box.Z), Color, DepthPriority, Thickness, LineLifeTime);
+			LineBatcher->DrawLine(Center + FVector( Box.X, -Box.Y,  Box.Z), Center + FVector(-Box.X, -Box.Y, Box.Z), Color, DepthPriority, Thickness, LineLifeTime);
+			LineBatcher->DrawLine(Center + FVector(-Box.X, -Box.Y,  Box.Z), Center + FVector(-Box.X,  Box.Y, Box.Z), Color, DepthPriority, Thickness, LineLifeTime);
+			LineBatcher->DrawLine(Center + FVector(-Box.X,  Box.Y,  Box.Z), Center + FVector( Box.X,  Box.Y, Box.Z), Color, DepthPriority, Thickness, LineLifeTime);
 
-			LineBatcher->DrawLine(Center + FVector( Box.X,  Box.Y, -Box.Z), Center + FVector( Box.X, -Box.Y, -Box.Z), Color, DepthPriority, 0.f, LineLifeTime);
-			LineBatcher->DrawLine(Center + FVector( Box.X, -Box.Y, -Box.Z), Center + FVector(-Box.X, -Box.Y, -Box.Z), Color, DepthPriority, 0.f, LineLifeTime);
-			LineBatcher->DrawLine(Center + FVector(-Box.X, -Box.Y, -Box.Z), Center + FVector(-Box.X,  Box.Y, -Box.Z), Color, DepthPriority, 0.f, LineLifeTime);
-			LineBatcher->DrawLine(Center + FVector(-Box.X,  Box.Y, -Box.Z), Center + FVector( Box.X,  Box.Y, -Box.Z), Color, DepthPriority, 0.f, LineLifeTime);
+			LineBatcher->DrawLine(Center + FVector( Box.X,  Box.Y, -Box.Z), Center + FVector( Box.X, -Box.Y, -Box.Z), Color, DepthPriority, Thickness, LineLifeTime);
+			LineBatcher->DrawLine(Center + FVector( Box.X, -Box.Y, -Box.Z), Center + FVector(-Box.X, -Box.Y, -Box.Z), Color, DepthPriority, Thickness, LineLifeTime);
+			LineBatcher->DrawLine(Center + FVector(-Box.X, -Box.Y, -Box.Z), Center + FVector(-Box.X,  Box.Y, -Box.Z), Color, DepthPriority, Thickness, LineLifeTime);
+			LineBatcher->DrawLine(Center + FVector(-Box.X,  Box.Y, -Box.Z), Center + FVector( Box.X,  Box.Y, -Box.Z), Color, DepthPriority, Thickness, LineLifeTime);
 
-			LineBatcher->DrawLine(Center + FVector( Box.X,  Box.Y,  Box.Z), Center + FVector( Box.X,  Box.Y, -Box.Z), Color, DepthPriority, 0.f, LineLifeTime);
-			LineBatcher->DrawLine(Center + FVector( Box.X, -Box.Y,  Box.Z), Center + FVector( Box.X, -Box.Y, -Box.Z), Color, DepthPriority, 0.f, LineLifeTime);
-			LineBatcher->DrawLine(Center + FVector(-Box.X, -Box.Y,  Box.Z), Center + FVector(-Box.X, -Box.Y, -Box.Z), Color, DepthPriority, 0.f, LineLifeTime);
-			LineBatcher->DrawLine(Center + FVector(-Box.X,  Box.Y,  Box.Z), Center + FVector(-Box.X,  Box.Y, -Box.Z), Color, DepthPriority, 0.f, LineLifeTime);
+			LineBatcher->DrawLine(Center + FVector( Box.X,  Box.Y,  Box.Z), Center + FVector( Box.X,  Box.Y, -Box.Z), Color, DepthPriority, Thickness, LineLifeTime);
+			LineBatcher->DrawLine(Center + FVector( Box.X, -Box.Y,  Box.Z), Center + FVector( Box.X, -Box.Y, -Box.Z), Color, DepthPriority, Thickness, LineLifeTime);
+			LineBatcher->DrawLine(Center + FVector(-Box.X, -Box.Y,  Box.Z), Center + FVector(-Box.X, -Box.Y, -Box.Z), Color, DepthPriority, Thickness, LineLifeTime);
+			LineBatcher->DrawLine(Center + FVector(-Box.X,  Box.Y,  Box.Z), Center + FVector(-Box.X,  Box.Y, -Box.Z), Color, DepthPriority, Thickness, LineLifeTime);
 		}
 	}
 }
 
-void DrawDebugBox(const UWorld* InWorld, FVector const& Center, FVector const& Box, const FQuat& Rotation, FColor const& Color, bool bPersistentLines, float LifeTime, uint8 DepthPriority)
+void DrawDebugBox(const UWorld* InWorld, FVector const& Center, FVector const& Box, const FQuat& Rotation, FColor const& Color, bool bPersistentLines, float LifeTime, uint8 DepthPriority, float Thickness)
 {
 	// no debug line drawing on dedicated server
 	if (GEngine->GetNetMode(InWorld) != NM_DedicatedServer)
@@ -127,51 +127,51 @@ void DrawDebugBox(const UWorld* InWorld, FVector const& Center, FVector const& B
 			FTransform const Transform(Rotation);
 			FVector Start = Transform.TransformPosition(FVector( Box.X,  Box.Y,  Box.Z));
 			FVector End = Transform.TransformPosition(FVector( Box.X, -Box.Y, Box.Z));
-			new(Lines) FBatchedLine(Center + Start, Center + End, Color, LineLifeTime, 0.f, DepthPriority);
+			new(Lines) FBatchedLine(Center + Start, Center + End, Color, LineLifeTime, Thickness, DepthPriority);
 
 			Start = Transform.TransformPosition(FVector( Box.X, -Box.Y,  Box.Z));
 			End = Transform.TransformPosition(FVector(-Box.X, -Box.Y, Box.Z));
-			new(Lines) FBatchedLine(Center + Start, Center + End, Color, LineLifeTime, 0.f, DepthPriority);
+			new(Lines) FBatchedLine(Center + Start, Center + End, Color, LineLifeTime, Thickness, DepthPriority);
 
 			Start = Transform.TransformPosition(FVector(-Box.X, -Box.Y,  Box.Z));
 			End = Transform.TransformPosition(FVector(-Box.X,  Box.Y, Box.Z));
-			new(Lines) FBatchedLine(Center + Start, Center + End, Color, LineLifeTime, 0.f, DepthPriority);
+			new(Lines) FBatchedLine(Center + Start, Center + End, Color, LineLifeTime, Thickness, DepthPriority);
 
 			Start = Transform.TransformPosition(FVector(-Box.X,  Box.Y,  Box.Z));
 			End = Transform.TransformPosition(FVector( Box.X,  Box.Y, Box.Z));
-			new(Lines) FBatchedLine(Center + Start, Center + End, Color, LineLifeTime, 0.f, DepthPriority);
+			new(Lines) FBatchedLine(Center + Start, Center + End, Color, LineLifeTime, Thickness, DepthPriority);
 
 			Start = Transform.TransformPosition(FVector( Box.X,  Box.Y, -Box.Z));
 			End = Transform.TransformPosition(FVector( Box.X, -Box.Y, -Box.Z));
-			new(Lines) FBatchedLine(Center + Start, Center + End, Color, LineLifeTime, 0.f, DepthPriority);
+			new(Lines) FBatchedLine(Center + Start, Center + End, Color, LineLifeTime, Thickness, DepthPriority);
 
 			Start = Transform.TransformPosition(FVector( Box.X, -Box.Y, -Box.Z));
 			End = Transform.TransformPosition(FVector(-Box.X, -Box.Y, -Box.Z));
-			new(Lines) FBatchedLine(Center + Start, Center + End, Color, LineLifeTime, 0.f, DepthPriority);
+			new(Lines) FBatchedLine(Center + Start, Center + End, Color, LineLifeTime, Thickness, DepthPriority);
 
 			Start = Transform.TransformPosition(FVector(-Box.X, -Box.Y, -Box.Z));
 			End = Transform.TransformPosition(FVector(-Box.X,  Box.Y, -Box.Z));
-			new(Lines) FBatchedLine(Center + Start, Center + End, Color, LineLifeTime, 0.f, DepthPriority);
+			new(Lines) FBatchedLine(Center + Start, Center + End, Color, LineLifeTime, Thickness, DepthPriority);
 
 			Start = Transform.TransformPosition(FVector(-Box.X,  Box.Y, -Box.Z));
 			End = Transform.TransformPosition(FVector( Box.X,  Box.Y, -Box.Z));
-			new(Lines )FBatchedLine(Center + Start, Center + End, Color, LineLifeTime, 0.f, DepthPriority);
+			new(Lines )FBatchedLine(Center + Start, Center + End, Color, LineLifeTime, Thickness, DepthPriority);
 
 			Start = Transform.TransformPosition(FVector( Box.X,  Box.Y,  Box.Z));
 			End = Transform.TransformPosition(FVector( Box.X,  Box.Y, -Box.Z));
-			new(Lines) FBatchedLine(Center + Start, Center + End, Color, LineLifeTime, 0.f, DepthPriority);
+			new(Lines) FBatchedLine(Center + Start, Center + End, Color, LineLifeTime, Thickness, DepthPriority);
 
 			Start = Transform.TransformPosition(FVector( Box.X, -Box.Y,  Box.Z));
 			End = Transform.TransformPosition(FVector( Box.X, -Box.Y, -Box.Z));
-			new(Lines) FBatchedLine(Center + Start, Center + End, Color, LineLifeTime, 0.f, DepthPriority);
+			new(Lines) FBatchedLine(Center + Start, Center + End, Color, LineLifeTime, Thickness, DepthPriority);
 
 			Start = Transform.TransformPosition(FVector(-Box.X, -Box.Y,  Box.Z));
 			End = Transform.TransformPosition(FVector(-Box.X, -Box.Y, -Box.Z));
-			new(Lines) FBatchedLine(Center + Start, Center + End, Color, LineLifeTime, 0.f, DepthPriority);
+			new(Lines) FBatchedLine(Center + Start, Center + End, Color, LineLifeTime, Thickness, DepthPriority);
 
 			Start = Transform.TransformPosition(FVector(-Box.X,  Box.Y,  Box.Z));
 			End = Transform.TransformPosition(FVector(-Box.X,  Box.Y, -Box.Z));
-			new(Lines) FBatchedLine(Center + Start, Center + End, Color, LineLifeTime, 0.f, DepthPriority);
+			new(Lines) FBatchedLine(Center + Start, Center + End, Color, LineLifeTime, Thickness, DepthPriority);
 
 			LineBatcher->DrawLines(Lines);
 		}
@@ -289,7 +289,7 @@ ENGINE_API void DrawDebugSolidPlane(const UWorld* InWorld, FPlane const& P, FVec
 	}
 }
 
-void DrawDebugCoordinateSystem(const UWorld* InWorld, FVector const& AxisLoc, FRotator const& AxisRot, float Scale, bool bPersistentLines, float LifeTime, uint8 DepthPriority)
+void DrawDebugCoordinateSystem(const UWorld* InWorld, FVector const& AxisLoc, FRotator const& AxisRot, float Scale, bool bPersistentLines, float LifeTime, uint8 DepthPriority, float Thickness)
 {
 	// no debug line drawing on dedicated server
 	if (GEngine->GetNetMode(InWorld) != NM_DedicatedServer)
@@ -303,9 +303,9 @@ void DrawDebugCoordinateSystem(const UWorld* InWorld, FVector const& AxisLoc, FR
 		ULineBatchComponent* const LineBatcher = GetDebugLineBatcher( InWorld, bPersistentLines, LifeTime, (DepthPriority == SDPG_Foreground) );
 		if(LineBatcher != NULL)
 		{
-			LineBatcher->DrawLine(AxisLoc, AxisLoc + X*Scale, FColor::Red, DepthPriority, 0.f, LifeTime );
-			LineBatcher->DrawLine(AxisLoc, AxisLoc + Y*Scale, FColor::Green, DepthPriority, 0.f, LifeTime );
-			LineBatcher->DrawLine(AxisLoc, AxisLoc + Z*Scale, FColor::Blue, DepthPriority, 0.f, LifeTime );
+			LineBatcher->DrawLine(AxisLoc, AxisLoc + X*Scale, FColor::Red, DepthPriority, Thickness, LifeTime );
+			LineBatcher->DrawLine(AxisLoc, AxisLoc + Y*Scale, FColor::Green, DepthPriority, Thickness, LifeTime );
+			LineBatcher->DrawLine(AxisLoc, AxisLoc + Z*Scale, FColor::Blue, DepthPriority, Thickness, LifeTime );
 		}
 	}
 }
@@ -418,7 +418,7 @@ void DrawDebugCircle(const UWorld* InWorld, FVector Center, float Radius, int32
 	);
 }
 
-void DrawDebug2DDonut(const UWorld* InWorld, const FMatrix& TransformMatrix, float InnerRadius, float OuterRadius, int32 Segments, const FColor& Color, bool bPersistentLines, float LifeTime, uint8 DepthPriority)
+void DrawDebug2DDonut(const UWorld* InWorld, const FMatrix& TransformMatrix, float InnerRadius, float OuterRadius, int32 Segments, const FColor& Color, bool bPersistentLines, float LifeTime, uint8 DepthPriority, float Thickness)
 {
 	// no debug line drawing on dedicated server
 	if( GEngine->GetNetMode(InWorld) != NM_DedicatedServer )
@@ -439,16 +439,16 @@ void DrawDebug2DDonut(const UWorld* InWorld, const FMatrix& TransformMatrix, flo
 
 			TArray<FBatchedLine> Lines;
 			Lines.Empty(4);
-			Lines.Add(FBatchedLine(Center - OuterRadius * AxisY, Center - InnerRadius * AxisY, Color, LineLifeTime, 0.0f, DepthPriority));
-			Lines.Add(FBatchedLine(Center + OuterRadius * AxisY, Center + InnerRadius * AxisY, Color, LineLifeTime, 0.0f, DepthPriority));
-			Lines.Add(FBatchedLine(Center - OuterRadius * AxisZ, Center - InnerRadius * AxisZ, Color, LineLifeTime, 0.0f, DepthPriority));
-			Lines.Add(FBatchedLine(Center + OuterRadius * AxisZ, Center + InnerRadius * AxisZ, Color, LineLifeTime, 0.0f, DepthPriority));
+			Lines.Add(FBatchedLine(Center - OuterRadius * AxisY, Center - InnerRadius * AxisY, Color, LineLifeTime, Thickness, DepthPriority));
+			Lines.Add(FBatchedLine(Center + OuterRadius * AxisY, Center + InnerRadius * AxisY, Color, LineLifeTime, Thickness, DepthPriority));
+			Lines.Add(FBatchedLine(Center - OuterRadius * AxisZ, Center - InnerRadius * AxisZ, Color, LineLifeTime, Thickness, DepthPriority));
+			Lines.Add(FBatchedLine(Center + OuterRadius * AxisZ, Center + InnerRadius * AxisZ, Color, LineLifeTime, Thickness, DepthPriority));
 			LineBatcher->DrawLines(Lines);
 		}
 	}
 }
 
-void DrawDebugSphere(const UWorld* InWorld, FVector const& Center, float Radius, int32 Segments, FColor const& Color, bool bPersistentLines, float LifeTime, uint8 DepthPriority)
+void DrawDebugSphere(const UWorld* InWorld, FVector const& Center, float Radius, int32 Segments, FColor const& Color, bool bPersistentLines, float LifeTime, uint8 DepthPriority, float Thickness)
 {
 	// no debug line drawing on dedicated server
 	if (GEngine->GetNetMode(InWorld) != NM_DedicatedServer)
@@ -489,8 +489,8 @@ void DrawDebugSphere(const UWorld* InWorld, FVector const& Center, float Radius,
 				Vertex2 = FVector((CosX * SinY1), (SinX * SinY1), CosY1) * Radius + Center;
 				Vertex4 = FVector((CosX * SinY2), (SinX * SinY2), CosY2) * Radius + Center;
 
-				Lines.Add(FBatchedLine(Vertex1, Vertex2, Color, LineLifeTime, 0.0f, DepthPriority));
-				Lines.Add(FBatchedLine(Vertex1, Vertex3, Color, LineLifeTime, 0.0f, DepthPriority));
+				Lines.Add(FBatchedLine(Vertex1, Vertex2, Color, LineLifeTime, Thickness, DepthPriority));
+				Lines.Add(FBatchedLine(Vertex1, Vertex3, Color, LineLifeTime, Thickness, DepthPriority));
 
 				Vertex1 = Vertex2;
 				Vertex3 = Vertex4;
@@ -504,7 +504,7 @@ void DrawDebugSphere(const UWorld* InWorld, FVector const& Center, float Radius,
 	}
 }
 
-void DrawDebugCylinder(const UWorld* InWorld, FVector const& Start, FVector const& End, float Radius, int32 Segments, FColor const& Color, bool bPersistentLines, float LifeTime, uint8 DepthPriority)
+void DrawDebugCylinder(const UWorld* InWorld, FVector const& Start, FVector const& End, float Radius, int32 Segments, FColor const& Color, bool bPersistentLines, float LifeTime, uint8 DepthPriority, float Thickness)
 {
 	// no debug line drawing on dedicated server
 	if (GEngine->GetNetMode(InWorld) != NM_DedicatedServer)
@@ -545,9 +545,9 @@ void DrawDebugCylinder(const UWorld* InWorld, FVector const& Start, FVector cons
 				P2 = Segment + Start;
 				P4 = Segment + End;
 
-				LineBatcher->DrawLine(P2, P4, Color, DepthPriority, 0.f, LineLifeTime);
-				LineBatcher->DrawLine(P1, P2, Color, DepthPriority, 0.f, LineLifeTime);
-				LineBatcher->DrawLine(P3, P4, Color, DepthPriority, 0.f, LineLifeTime);
+				LineBatcher->DrawLine(P2, P4, Color, DepthPriority, Thickness, LineLifeTime);
+				LineBatcher->DrawLine(P1, P2, Color, DepthPriority, Thickness, LineLifeTime);
+				LineBatcher->DrawLine(P3, P4, Color, DepthPriority, Thickness, LineLifeTime);
 
 				P1 = P2;
 				P3 = P4;
@@ -558,7 +558,7 @@ void DrawDebugCylinder(const UWorld* InWorld, FVector const& Start, FVector cons
 }
 
 /** Used by gameplay when defining a cone by a vertical and horizontal dot products. */
-void DrawDebugAltCone(const UWorld* InWorld, FVector const& Origin, FRotator const& Rotation, float Length, float AngleWidth, float AngleHeight, FColor const& DrawColor, bool bPersistentLines, float LifeTime, uint8 DepthPriority)
+void DrawDebugAltCone(const UWorld* InWorld, FVector const& Origin, FRotator const& Rotation, float Length, float AngleWidth, float AngleHeight, FColor const& DrawColor, bool bPersistentLines, float LifeTime, uint8 DepthPriority, float Thickness)
 {
 	FRotationMatrix const RM( Rotation );
 	FVector const AxisX = RM.GetScaledAxis( EAxis::X );
@@ -584,40 +584,40 @@ void DrawDebugAltCone(const UWorld* InWorld, FVector const& Origin, FRotator con
 	FVector D = EndPoint - Up - Right;
 
 	// Corners
-	Lines.Add(FBatchedLine(Origin, A, DrawColor, LineLifeTime, 0.0f, DepthPriority));
-	Lines.Add(FBatchedLine(Origin, B, DrawColor, LineLifeTime, 0.0f, DepthPriority));
-	Lines.Add(FBatchedLine(Origin, C, DrawColor, LineLifeTime, 0.0f, DepthPriority));
-	Lines.Add(FBatchedLine(Origin, D, DrawColor, LineLifeTime, 0.0f, DepthPriority));
+	Lines.Add(FBatchedLine(Origin, A, DrawColor, LineLifeTime, Thickness, DepthPriority));
+	Lines.Add(FBatchedLine(Origin, B, DrawColor, LineLifeTime, Thickness, DepthPriority));
+	Lines.Add(FBatchedLine(Origin, C, DrawColor, LineLifeTime, Thickness, DepthPriority));
+	Lines.Add(FBatchedLine(Origin, D, DrawColor, LineLifeTime, Thickness, DepthPriority));
 
 	// Further most plane/frame
-	Lines.Add(FBatchedLine(A, B, DrawColor, LineLifeTime, 0.0f, DepthPriority));
-	Lines.Add(FBatchedLine(B, C, DrawColor, LineLifeTime, 0.0f, DepthPriority));
-	Lines.Add(FBatchedLine(C, D, DrawColor, LineLifeTime, 0.0f, DepthPriority));
-	Lines.Add(FBatchedLine(D, A, DrawColor, LineLifeTime, 0.0f, DepthPriority));
+	Lines.Add(FBatchedLine(A, B, DrawColor, LineLifeTime, Thickness, DepthPriority));
+	Lines.Add(FBatchedLine(B, C, DrawColor, LineLifeTime, Thickness, DepthPriority));
+	Lines.Add(FBatchedLine(C, D, DrawColor, LineLifeTime, Thickness, DepthPriority));
+	Lines.Add(FBatchedLine(D, A, DrawColor, LineLifeTime, Thickness, DepthPriority));
 
 	// Mid points
-	Lines.Add(FBatchedLine(Origin, EndPoint + Up, DrawColor, LineLifeTime, 0.0f, DepthPriority));
-	Lines.Add(FBatchedLine(Origin, EndPoint - Up, DrawColor, LineLifeTime, 0.0f, DepthPriority));
-	Lines.Add(FBatchedLine(Origin, EndPoint + Right, DrawColor, LineLifeTime, 0.0f, DepthPriority));
-	Lines.Add(FBatchedLine(Origin, EndPoint - Right, DrawColor, LineLifeTime, 0.0f, DepthPriority));
+	Lines.Add(FBatchedLine(Origin, EndPoint + Up, DrawColor, LineLifeTime, Thickness, DepthPriority));
+	Lines.Add(FBatchedLine(Origin, EndPoint - Up, DrawColor, LineLifeTime, Thickness, DepthPriority));
+	Lines.Add(FBatchedLine(Origin, EndPoint + Right, DrawColor, LineLifeTime, Thickness, DepthPriority));
+	Lines.Add(FBatchedLine(Origin, EndPoint - Right, DrawColor, LineLifeTime, Thickness, DepthPriority));
 
 	// Inbetween
-	Lines.Add(FBatchedLine(Origin, EndPoint + Up - HalfRight, DrawColor, LineLifeTime, 0.0f, DepthPriority));
-	Lines.Add(FBatchedLine(Origin, EndPoint + Up + HalfRight, DrawColor, LineLifeTime, 0.0f, DepthPriority));
+	Lines.Add(FBatchedLine(Origin, EndPoint + Up - HalfRight, DrawColor, LineLifeTime, Thickness, DepthPriority));
+	Lines.Add(FBatchedLine(Origin, EndPoint + Up + HalfRight, DrawColor, LineLifeTime, Thickness, DepthPriority));
 
-	Lines.Add(FBatchedLine(Origin, EndPoint - Up - HalfRight, DrawColor, LineLifeTime, 0.0f, DepthPriority));
-	Lines.Add(FBatchedLine(Origin, EndPoint - Up + HalfRight, DrawColor, LineLifeTime, 0.0f, DepthPriority));
+	Lines.Add(FBatchedLine(Origin, EndPoint - Up - HalfRight, DrawColor, LineLifeTime, Thickness, DepthPriority));
+	Lines.Add(FBatchedLine(Origin, EndPoint - Up + HalfRight, DrawColor, LineLifeTime, Thickness, DepthPriority));
 
-	Lines.Add(FBatchedLine(Origin, EndPoint + Right - HalfUp, DrawColor, LineLifeTime, 0.0f, DepthPriority));
-	Lines.Add(FBatchedLine(Origin, EndPoint + Right + HalfUp, DrawColor, LineLifeTime, 0.0f, DepthPriority));
+	Lines.Add(FBatchedLine(Origin, EndPoint + Right - HalfUp, DrawColor, LineLifeTime, Thickness, DepthPriority));
+	Lines.Add(FBatchedLine(Origin, EndPoint + Right + HalfUp, DrawColor, LineLifeTime, Thickness, DepthPriority));
 
-	Lines.Add(FBatchedLine(Origin, EndPoint - Right - HalfUp, DrawColor, LineLifeTime, 0.0f, DepthPriority));
-	Lines.Add(FBatchedLine(Origin, EndPoint - Right + HalfUp, DrawColor, LineLifeTime, 0.0f, DepthPriority));
+	Lines.Add(FBatchedLine(Origin, EndPoint - Right - HalfUp, DrawColor, LineLifeTime, Thickness, DepthPriority));
+	Lines.Add(FBatchedLine(Origin, EndPoint - Right + HalfUp, DrawColor, LineLifeTime, Thickness, DepthPriority));
 
 	LineBatcher->DrawLines(Lines);
 }
 
-void DrawDebugCone(const UWorld* InWorld, FVector const& Origin, FVector const& Direction, float Length, float AngleWidth, float AngleHeight, int32 NumSides, FColor const& DrawColor, bool bPersistentLines, float LifeTime, uint8 DepthPriority)
+void DrawDebugCone(const UWorld* InWorld, FVector const& Origin, FVector const& Direction, float Length, float AngleWidth, float AngleHeight, int32 NumSides, FColor const& DrawColor, bool bPersistentLines, float LifeTime, uint8 DepthPriority, float Thickness)
 {
 	// no debug line drawing on dedicated server
 	if (GEngine->GetNetMode(InWorld) != NM_DedicatedServer)
@@ -680,12 +680,12 @@ void DrawDebugCone(const UWorld* InWorld, FVector const& Origin, FVector const&
 			for(int32 i = 0; i < NumSides; i++)
 			{
 				CurrentPoint = ConeToWorld.TransformPosition(ConeVerts[i]);
-				Lines.Add(FBatchedLine(ConeToWorld.GetOrigin(), CurrentPoint, DrawColor, LineLifeTime, 0.0f, DepthPriority));
+				Lines.Add(FBatchedLine(ConeToWorld.GetOrigin(), CurrentPoint, DrawColor, LineLifeTime, Thickness, DepthPriority));
 
 				// PrevPoint must be defined to draw junctions
 				if( i > 0 )
 				{
-					Lines.Add(FBatchedLine(PrevPoint, CurrentPoint, DrawColor, LineLifeTime, 0.0f, DepthPriority));
+					Lines.Add(FBatchedLine(PrevPoint, CurrentPoint, DrawColor, LineLifeTime, Thickness, DepthPriority));
 				}
 				else
 				{
@@ -695,7 +695,7 @@ void DrawDebugCone(const UWorld* InWorld, FVector const& Origin, FVector const&
 				PrevPoint = CurrentPoint;
 			}
 			// Connect last junction to first
-			Lines.Add(FBatchedLine(CurrentPoint, FirstPoint, DrawColor, LineLifeTime, 0.0f, DepthPriority));
+			Lines.Add(FBatchedLine(CurrentPoint, FirstPoint, DrawColor, LineLifeTime, Thickness, DepthPriority));
 
 			LineBatcher->DrawLines(Lines);
 		}
@@ -736,7 +736,7 @@ void FlushDebugStrings( const UWorld* InWorld )
 	}	
 }
 
-void DrawDebugFrustum(const UWorld* InWorld, const FMatrix& FrustumToWorld, FColor const& Color, bool bPersistentLines, float LifeTime, uint8 DepthPriority)
+void DrawDebugFrustum(const UWorld* InWorld, const FMatrix& FrustumToWorld, FColor const& Color, bool bPersistentLines, float LifeTime, uint8 DepthPriority, float Thickness)
 {
 	// no debug line drawing on dedicated server
 	if (GEngine->GetNetMode(InWorld) != NM_DedicatedServer)
@@ -761,20 +761,20 @@ void DrawDebugFrustum(const UWorld* InWorld, const FMatrix& FrustumToWorld, FCol
 			}
 		}
 
-		DrawDebugLine(InWorld, Vertices[0][0][0], Vertices[0][0][1],Color,  bPersistentLines, LifeTime, DepthPriority);
-		DrawDebugLine(InWorld, Vertices[1][0][0], Vertices[1][0][1],Color,  bPersistentLines, LifeTime, DepthPriority);
-		DrawDebugLine(InWorld, Vertices[0][1][0], Vertices[0][1][1],Color,  bPersistentLines, LifeTime, DepthPriority);
-		DrawDebugLine(InWorld, Vertices[1][1][0], Vertices[1][1][1],Color,  bPersistentLines, LifeTime, DepthPriority);
+		DrawDebugLine(InWorld, Vertices[0][0][0], Vertices[0][0][1],Color,  bPersistentLines, LifeTime, DepthPriority, Thickness);
+		DrawDebugLine(InWorld, Vertices[1][0][0], Vertices[1][0][1],Color,  bPersistentLines, LifeTime, DepthPriority, Thickness);
+		DrawDebugLine(InWorld, Vertices[0][1][0], Vertices[0][1][1],Color,  bPersistentLines, LifeTime, DepthPriority, Thickness);
+		DrawDebugLine(InWorld, Vertices[1][1][0], Vertices[1][1][1],Color,  bPersistentLines, LifeTime, DepthPriority, Thickness);
 
-		DrawDebugLine(InWorld, Vertices[0][0][0], Vertices[0][1][0],Color,  bPersistentLines, LifeTime, DepthPriority);
-		DrawDebugLine(InWorld, Vertices[1][0][0], Vertices[1][1][0],Color,  bPersistentLines, LifeTime, DepthPriority);
-		DrawDebugLine(InWorld, Vertices[0][0][1], Vertices[0][1][1],Color,  bPersistentLines, LifeTime, DepthPriority);
-		DrawDebugLine(InWorld, Vertices[1][0][1], Vertices[1][1][1],Color,  bPersistentLines, LifeTime, DepthPriority);
+		DrawDebugLine(InWorld, Vertices[0][0][0], Vertices[0][1][0],Color,  bPersistentLines, LifeTime, DepthPriority, Thickness);
+		DrawDebugLine(InWorld, Vertices[1][0][0], Vertices[1][1][0],Color,  bPersistentLines, LifeTime, DepthPriority, Thickness);
+		DrawDebugLine(InWorld, Vertices[0][0][1], Vertices[0][1][1],Color,  bPersistentLines, LifeTime, DepthPriority, Thickness);
+		DrawDebugLine(InWorld, Vertices[1][0][1], Vertices[1][1][1],Color,  bPersistentLines, LifeTime, DepthPriority, Thickness);
 
-		DrawDebugLine(InWorld, Vertices[0][0][0], Vertices[1][0][0],Color,  bPersistentLines, LifeTime, DepthPriority);
-		DrawDebugLine(InWorld, Vertices[0][1][0], Vertices[1][1][0],Color,  bPersistentLines, LifeTime, DepthPriority);
-		DrawDebugLine(InWorld, Vertices[0][0][1], Vertices[1][0][1],Color,  bPersistentLines, LifeTime, DepthPriority);
-		DrawDebugLine(InWorld, Vertices[0][1][1], Vertices[1][1][1],Color,  bPersistentLines, LifeTime, DepthPriority);
+		DrawDebugLine(InWorld, Vertices[0][0][0], Vertices[1][0][0],Color,  bPersistentLines, LifeTime, DepthPriority, Thickness);
+		DrawDebugLine(InWorld, Vertices[0][1][0], Vertices[1][1][0],Color,  bPersistentLines, LifeTime, DepthPriority, Thickness);
+		DrawDebugLine(InWorld, Vertices[0][0][1], Vertices[1][0][1],Color,  bPersistentLines, LifeTime, DepthPriority, Thickness);
+		DrawDebugLine(InWorld, Vertices[0][1][1], Vertices[1][1][1],Color,  bPersistentLines, LifeTime, DepthPriority, Thickness);
 	}
 }
 
diff --git a/Engine/Source/Runtime/Engine/Private/GameMode.cpp b/Engine/Source/Runtime/Engine/Private/GameMode.cpp
--- a/Engine/Source/Runtime/Engine/Private/GameMode.cpp
+++ b/Engine/Source/Runtime/Engine/Private/GameMode.cpp
@@ -728,8 +728,19 @@ void AGameMode::SetMatchState(FName NewState)
 
 	MatchState = NewState;
 
-	// Call change callbacks
+	OnMatchStateSet();
+
+	if (GameState)
+	{
+		GameState->SetMatchState(NewState);
+	}
 
+	K2_OnSetMatchState(NewState);
+}
+
+void AGameMode::OnMatchStateSet()
+{
+	// Call change callbacks
 	if (MatchState == MatchState::WaitingToStart)
 	{
 		HandleMatchIsWaitingToStart();
@@ -750,13 +761,6 @@ void AGameMode::SetMatchState(FName NewState)
 	{
 		HandleMatchAborted();
 	}
-
-	if (GameState)
-	{
-		GameState->SetMatchState(NewState);
-	}
-
-	K2_OnSetMatchState(NewState);
 }
 
 void AGameMode::Tick(float DeltaSeconds)
@@ -989,21 +993,27 @@ void AGameMode::ProcessServerTravel(const FString& URL, bool bAbsolute)
 }
 
 
-void AGameMode::GetSeamlessTravelActorList(bool bToEntry, TArray<AActor*>& ActorList)
+void AGameMode::GetSeamlessTravelActorList(bool bToTransition, TArray<AActor*>& ActorList)
 {
 	UWorld* World = GetWorld();
+
+	// Get allocations for the elements we're going to add handled in one go
+	const int32 ActorsToAddCount = World->GameState->PlayerArray.Num() + (bToTransition ?  3 : 0);
+	ActorList.Reserve(ActorsToAddCount);
+
 	// always keep PlayerStates, so that after we restart we can keep players on the same team, etc
-	for (int32 i = 0; i < World->GameState->PlayerArray.Num(); i++)
-	{
-		ActorList.Add(World->GameState->PlayerArray[i]);
-	}
+	ActorList.Append(World->GameState->PlayerArray);
 
-	if (bToEntry)
+	if (bToTransition)
 	{
+		// keep ourselves until we transition to the final destination
+		ActorList.Add(this);
 		// keep general game state until we transition to the final destination
 		ActorList.Add(World->GameState);
 		// keep the game session state until we transition to the final destination
 		ActorList.Add(GameSession);
+
+		// If adding in this section best to increase the literal above for the ActorsToAddCount
 	}
 }
 
diff --git a/Engine/Source/Runtime/Engine/Private/GameplayStatics.cpp b/Engine/Source/Runtime/Engine/Private/GameplayStatics.cpp
--- a/Engine/Source/Runtime/Engine/Private/GameplayStatics.cpp
+++ b/Engine/Source/Runtime/Engine/Private/GameplayStatics.cpp
@@ -650,9 +650,13 @@ UParticleSystemComponent* UGameplayStatics::SpawnEmitterAtLocation(UObject* Worl
 		if (UWorld* World = GEngine->GetWorldFromContextObject(WorldContextObject))
 		{
 			UParticleSystemComponent* PSC = CreateParticleSystem(EmitterTemplate, World, World->GetWorldSettings(), bAutoDestroy);
-			PSC->SetAbsolute(true, true, true);
-			PSC->SetWorldLocationAndRotation(SpawnLocation, SpawnRotation);
-			PSC->SetRelativeScale3D(FVector(1.f));
+
+			PSC->bAbsoluteLocation = true;
+			PSC->bAbsoluteRotation = true;
+			PSC->bAbsoluteScale = true;
+			PSC->RelativeLocation = SpawnLocation;
+			PSC->RelativeRotation = SpawnRotation;
+			PSC->RelativeScale3D = FVector(1.f);
 
 			PSC->RegisterComponentWithWorld(World);
 
@@ -680,8 +684,12 @@ UParticleSystemComponent* UGameplayStatics::SpawnEmitterAtLocation(UWorld* World
 	{
 		PSC = CreateParticleSystem(EmitterTemplate, World, World->GetWorldSettings(), bAutoDestroy);
 
-		PSC->SetAbsolute(true, true, true);
-		PSC->SetWorldTransform(SpawnTransform);
+		PSC->bAbsoluteLocation = true;
+		PSC->bAbsoluteRotation = true;
+		PSC->bAbsoluteScale = true;
+		PSC->RelativeLocation = SpawnTransform.GetLocation();
+		PSC->RelativeRotation = SpawnTransform.GetRotation().Rotator();
+		PSC->RelativeScale3D = SpawnTransform.GetScale3D();
 
 		PSC->RegisterComponentWithWorld(World);
 
@@ -716,16 +724,25 @@ UParticleSystemComponent* UGameplayStatics::SpawnEmitterAttached(UParticleSystem
 			if (World && World->GetNetMode() != NM_DedicatedServer)
 			{
 				PSC = CreateParticleSystem(EmitterTemplate, World, AttachToComponent->GetOwner(), bAutoDestroy);
-				PSC->AttachTo(AttachToComponent, AttachPointName);
+				
+				PSC->AttachParent = AttachToComponent;
+				PSC->AttachSocketName = AttachPointName;
+
 				if (LocationType == EAttachLocation::KeepWorldPosition)
 				{
-					PSC->SetWorldLocationAndRotation(Location, Rotation);
+					const FTransform ParentToWorld = AttachToComponent->GetSocketTransform(AttachPointName);
+					const FTransform ComponentToWorld(Rotation, Location);
+					const FTransform RelativeTM = ComponentToWorld.GetRelativeTransform(ParentToWorld);
+					PSC->RelativeLocation = RelativeTM.GetLocation();
+					PSC->RelativeRotation = RelativeTM.GetRotation().Rotator();
 				}
 				else
 				{
-					PSC->SetRelativeLocationAndRotation(Location, Rotation);
+					PSC->RelativeLocation = Location;
+					PSC->RelativeRotation = Rotation;
 				}
-				PSC->SetRelativeScale3D(FVector(1.f));
+				PSC->RelativeScale3D = FVector(1.f);
+
 				PSC->RegisterComponentWithWorld(World);
 				PSC->ActivateSystem(true);
 
@@ -808,6 +825,31 @@ void UGameplayStatics::SetGlobalPitchModulation(UObject* WorldContextObject, flo
 	}
 }
 
+void UGameplayStatics::SetGlobalListenerFocusParameters(UObject* WorldContextObject, float FocusAzimuthScale, float NonFocusAzimuthScale, float FocusDistanceScale, float NonFocusDistanceScale, float FocusVolumeScale, float NonFocusVolumeScale, float FocusPriorityScale, float NonFocusPriorityScale)
+{
+	if (!GEngine || !GEngine->UseSound())
+	{
+		return;
+	}
+
+	UWorld* ThisWorld = GEngine->GetWorldFromContextObject(WorldContextObject);
+	if (!ThisWorld || !ThisWorld->bAllowAudioPlayback || ThisWorld->GetNetMode() == NM_DedicatedServer)
+	{
+		return;
+	}
+
+	if (FAudioDevice* AudioDevice = ThisWorld->GetAudioDevice())
+	{
+		AudioDevice->GlobalFocusSettings.FocusAzimuthScale = FMath::Max(FocusAzimuthScale, 0.0f);
+		AudioDevice->GlobalFocusSettings.NonFocusAzimuthScale = FMath::Max(NonFocusAzimuthScale, 0.0f);
+		AudioDevice->GlobalFocusSettings.FocusDistanceScale = FMath::Max(FocusDistanceScale, 0.0f);
+		AudioDevice->GlobalFocusSettings.NonFocusDistanceScale = FMath::Max(NonFocusDistanceScale, 0.0f);
+		AudioDevice->GlobalFocusSettings.FocusVolumeScale = FMath::Max(FocusVolumeScale, 0.0f);
+		AudioDevice->GlobalFocusSettings.NonFocusVolumeScale = FMath::Max(NonFocusVolumeScale, 0.0f);
+		AudioDevice->GlobalFocusSettings.FocusPriorityScale = FMath::Max(FocusPriorityScale, 0.0f);
+		AudioDevice->GlobalFocusSettings.NonFocusPriorityScale = FMath::Max(NonFocusPriorityScale, 0.0f);
+	}
+}
 
 void UGameplayStatics::PlaySound2D(UObject* WorldContextObject, class USoundBase* Sound, float VolumeMultiplier, float PitchMultiplier, float StartTime, class USoundConcurrency* ConcurrencySettings)
 {
@@ -950,7 +992,7 @@ class UAudioComponent* UGameplayStatics::SpawnSoundAttached(class USoundBase* So
 	}
 
 	UAudioComponent* AudioComponent = FAudioDevice::CreateComponent(Sound, AttachToComponent->GetWorld(), AttachToComponent->GetOwner(), false, bStopWhenAttachedToDestroyed, &TestLocation, AttenuationSettings, ConcurrencySettings);
-	if (AudioComponent)
+	if (AudioComponent && AudioComponent->GetWorld())
 	{
 		const bool bIsInGameWorld = AudioComponent->GetWorld()->IsGameWorld();
 
diff --git a/Engine/Source/Runtime/Engine/Private/HUD.cpp b/Engine/Source/Runtime/Engine/Private/HUD.cpp
--- a/Engine/Source/Runtime/Engine/Private/HUD.cpp
+++ b/Engine/Source/Runtime/Engine/Private/HUD.cpp
@@ -810,12 +810,13 @@ void AHUD::DrawRect(FLinearColor Color, float ScreenX, float ScreenY, float Widt
 	}
 }
 
-void AHUD::DrawLine(float StartScreenX, float StartScreenY, float EndScreenX, float EndScreenY, FLinearColor LineColor)
+void AHUD::DrawLine(float StartScreenX, float StartScreenY, float EndScreenX, float EndScreenY, FLinearColor LineColor, float LineThickness)
 {
 	if (IsCanvasValid_WarnIfNot())
 	{
 		FCanvasLineItem LineItem( FVector2D(StartScreenX, StartScreenY), FVector2D(EndScreenX, EndScreenY) );
 		LineItem.SetColor( LineColor );
+		LineItem.LineThickness = LineThickness;
 		Canvas->DrawItem( LineItem );
 	}
 }
diff --git a/Engine/Source/Runtime/Engine/Private/KismetSystemLibrary.cpp b/Engine/Source/Runtime/Engine/Private/KismetSystemLibrary.cpp
--- a/Engine/Source/Runtime/Engine/Private/KismetSystemLibrary.cpp
+++ b/Engine/Source/Runtime/Engine/Private/KismetSystemLibrary.cpp
@@ -2690,88 +2690,88 @@ void UKismetSystemLibrary::DrawDebugPoint(UObject* WorldContextObject, FVector c
 }
 
 /** Draw directional arrow, pointing from LineStart to LineEnd. */
-void UKismetSystemLibrary::DrawDebugArrow(UObject* WorldContextObject, FVector const LineStart, FVector const LineEnd, float ArrowSize, FLinearColor Color, float LifeTime)
+void UKismetSystemLibrary::DrawDebugArrow(UObject* WorldContextObject, FVector const LineStart, FVector const LineEnd, float ArrowSize, FLinearColor Color, float LifeTime, float Thickness)
 {
 	UWorld* World = GEngine->GetWorldFromContextObject(WorldContextObject);
 	if (World != nullptr)
 	{
-		::DrawDebugDirectionalArrow(World, LineStart, LineEnd, ArrowSize, Color.ToFColor(true), false, LifeTime, SDPG_World);
+		::DrawDebugDirectionalArrow(World, LineStart, LineEnd, ArrowSize, Color.ToFColor(true), false, LifeTime, SDPG_World, Thickness);
 	}
 }
 
 /** Draw a debug box */
-void UKismetSystemLibrary::DrawDebugBox(UObject* WorldContextObject, FVector const Center, FVector Extent, FLinearColor Color, const FRotator Rotation, float LifeTime)
+void UKismetSystemLibrary::DrawDebugBox(UObject* WorldContextObject, FVector const Center, FVector Extent, FLinearColor Color, const FRotator Rotation, float LifeTime, float Thickness)
 {
 	UWorld* World = GEngine->GetWorldFromContextObject(WorldContextObject);
 	if(World != nullptr)
 	{
 		if (Rotation == FRotator::ZeroRotator)
 		{
-			::DrawDebugBox(World, Center, Extent, Color.ToFColor(true), false, LifeTime, SDPG_World);
+			::DrawDebugBox(World, Center, Extent, Color.ToFColor(true), false, LifeTime, SDPG_World, Thickness);
 		}
 		else
 		{
-			::DrawDebugBox(World, Center, Extent, Rotation.Quaternion(), Color.ToFColor(true), false, LifeTime, SDPG_World);
+			::DrawDebugBox(World, Center, Extent, Rotation.Quaternion(), Color.ToFColor(true), false, LifeTime, SDPG_World, Thickness);
 		}
 	}
 }
 
 /** Draw a debug coordinate system. */
-void UKismetSystemLibrary::DrawDebugCoordinateSystem(UObject* WorldContextObject, FVector const AxisLoc, FRotator const AxisRot, float Scale, float LifeTime)
+void UKismetSystemLibrary::DrawDebugCoordinateSystem(UObject* WorldContextObject, FVector const AxisLoc, FRotator const AxisRot, float Scale, float LifeTime, float Thickness)
 {
 	UWorld* World = GEngine->GetWorldFromContextObject(WorldContextObject);
 	if (World != nullptr)
 	{
-		::DrawDebugCoordinateSystem(World, AxisLoc, AxisRot, Scale, false, LifeTime, SDPG_World);
+		::DrawDebugCoordinateSystem(World, AxisLoc, AxisRot, Scale, false, LifeTime, SDPG_World, Thickness);
 	}
 }
 
 /** Draw a debug sphere */
-void UKismetSystemLibrary::DrawDebugSphere(UObject* WorldContextObject, FVector const Center, float Radius, int32 Segments, FLinearColor Color, float LifeTime)
+void UKismetSystemLibrary::DrawDebugSphere(UObject* WorldContextObject, FVector const Center, float Radius, int32 Segments, FLinearColor Color, float LifeTime, float Thickness)
 {
 	UWorld* World = GEngine->GetWorldFromContextObject(WorldContextObject);
 	if (World != nullptr)
 	{
-		::DrawDebugSphere(World, Center, Radius, Segments, Color.ToFColor(true), false, LifeTime, SDPG_World);
+		::DrawDebugSphere(World, Center, Radius, Segments, Color.ToFColor(true), false, LifeTime, SDPG_World, Thickness);
 	}
 }
 
 /** Draw a debug cylinder */
-void UKismetSystemLibrary::DrawDebugCylinder(UObject* WorldContextObject, FVector const Start, FVector const End, float Radius, int32 Segments, FLinearColor Color, float LifeTime)
+void UKismetSystemLibrary::DrawDebugCylinder(UObject* WorldContextObject, FVector const Start, FVector const End, float Radius, int32 Segments, FLinearColor Color, float LifeTime, float Thickness)
 {
 	UWorld* World = GEngine->GetWorldFromContextObject(WorldContextObject);
 	if (World != nullptr)
 	{
-		::DrawDebugCylinder(World, Start, End, Radius, Segments, Color.ToFColor(true), false, LifeTime, SDPG_World);
+		::DrawDebugCylinder(World, Start, End, Radius, Segments, Color.ToFColor(true), false, LifeTime, SDPG_World, Thickness);
 	}
 }
 
 /** Draw a debug cone */
-void UKismetSystemLibrary::DrawDebugCone(UObject* WorldContextObject, FVector const Origin, FVector const Direction, float Length, float AngleWidth, float AngleHeight, int32 NumSides, FLinearColor Color)
+void UKismetSystemLibrary::DrawDebugCone(UObject* WorldContextObject, FVector const Origin, FVector const Direction, float Length, float AngleWidth, float AngleHeight, int32 NumSides, FLinearColor Color, float Duration, float Thickness)
 {
 	UWorld* World = GEngine->GetWorldFromContextObject(WorldContextObject);
 	if (World != nullptr)
 	{
-		::DrawDebugCone(World, Origin, Direction, Length, AngleWidth, AngleHeight, NumSides, Color.ToFColor(true), false, -1.f, SDPG_World);
+		::DrawDebugCone(World, Origin, Direction, Length, AngleWidth, AngleHeight, NumSides, Color.ToFColor(true), false, Duration, SDPG_World, Thickness);
 	}
 }
 
-void UKismetSystemLibrary::DrawDebugConeInDegrees(UObject* WorldContextObject, FVector const Origin, FVector const Direction, float Length, float AngleWidth, float AngleHeight, int32 NumSides, FLinearColor Color, float LifeTime)
+void UKismetSystemLibrary::DrawDebugConeInDegrees(UObject* WorldContextObject, FVector const Origin, FVector const Direction, float Length, float AngleWidth, float AngleHeight, int32 NumSides, FLinearColor Color, float LifeTime, float Thickness)
 {
 	UWorld* World = GEngine->GetWorldFromContextObject(WorldContextObject);
 	if (World != nullptr)
 	{
-		::DrawDebugCone(World, Origin, Direction, Length, FMath::DegreesToRadians(AngleWidth), FMath::DegreesToRadians(AngleHeight), NumSides, Color.ToFColor(true), false, LifeTime, SDPG_World);
+		::DrawDebugCone(World, Origin, Direction, Length, FMath::DegreesToRadians(AngleWidth), FMath::DegreesToRadians(AngleHeight), NumSides, Color.ToFColor(true), false, LifeTime, SDPG_World, Thickness);
 	}
 }
 
 /** Draw a debug capsule */
-void UKismetSystemLibrary::DrawDebugCapsule(UObject* WorldContextObject, FVector const Center, float HalfHeight, float Radius, const FRotator Rotation, FLinearColor Color, float LifeTime)
+void UKismetSystemLibrary::DrawDebugCapsule(UObject* WorldContextObject, FVector const Center, float HalfHeight, float Radius, const FRotator Rotation, FLinearColor Color, float LifeTime, float Thickness)
 {
 	UWorld* World = GEngine->GetWorldFromContextObject(WorldContextObject);
 	if (World != nullptr)
 	{
-		::DrawDebugCapsule(World, Center, HalfHeight, Radius, Rotation.Quaternion(), Color.ToFColor(true), false, LifeTime, SDPG_World);
+		::DrawDebugCapsule(World, Center, HalfHeight, Radius, Rotation.Quaternion(), Color.ToFColor(true), false, LifeTime, SDPG_World, Thickness);
 	}
 }
 
@@ -2816,13 +2816,13 @@ void UKismetSystemLibrary::FlushPersistentDebugLines(UObject* WorldContextObject
 }
 
 /** Draws a debug frustum. */
-void UKismetSystemLibrary::DrawDebugFrustum(UObject* WorldContextObject, const FTransform& FrustumTransform, FLinearColor FrustumColor, float Duration)
+void UKismetSystemLibrary::DrawDebugFrustum(UObject* WorldContextObject, const FTransform& FrustumTransform, FLinearColor FrustumColor, float Duration, float Thickness)
 {
 	UWorld* World = GEngine->GetWorldFromContextObject(WorldContextObject);
 	if( World != nullptr && FrustumTransform.IsRotationNormalized() )
 	{
 		FMatrix FrustumToWorld =  FrustumTransform.ToMatrixWithScale();
-		::DrawDebugFrustum(World, FrustumToWorld, FrustumColor.ToFColor(true), false, Duration, SDPG_World);
+		::DrawDebugFrustum(World, FrustumToWorld, FrustumColor.ToFColor(true), false, Duration, SDPG_World, Thickness);
 	}
 }
 
diff --git a/Engine/Source/Runtime/Engine/Private/Particles/ParticleComponents.cpp b/Engine/Source/Runtime/Engine/Private/Particles/ParticleComponents.cpp
--- a/Engine/Source/Runtime/Engine/Private/Particles/ParticleComponents.cpp
+++ b/Engine/Source/Runtime/Engine/Private/Particles/ParticleComponents.cpp
@@ -1529,8 +1529,8 @@ void UParticleEmitter::Build()
 		if (HighLODLevel->TypeDataModule != nullptr)
 		{
 			if(HighLODLevel->TypeDataModule->RequiresBuild())
-		{
-			FParticleEmitterBuildInfo EmitterBuildInfo;
+			{
+				FParticleEmitterBuildInfo EmitterBuildInfo;
 				HighLODLevel->CompileModules( EmitterBuildInfo );
 				HighLODLevel->TypeDataModule->Build( EmitterBuildInfo );
 			}
@@ -2306,6 +2306,8 @@ void UParticleSystem::GetAssetRegistryTags(TArray<FAssetRegistryTag>& OutTags) c
 
 	OutTags.Add(FAssetRegistryTag("NumLODs", LexicalConversion::ToString(LODDistances.Num()), FAssetRegistryTag::TT_Numerical));
 
+	OutTags.Add(FAssetRegistryTag("WarmupTime", LexicalConversion::ToString(WarmupTime), FAssetRegistryTag::TT_Numerical));
+
 	// Done here instead of as an AssetRegistrySearchable string to avoid the long prefix on the enum value string
 	FString LODMethodString = TEXT("Unknown");
 	switch (LODMethod)
@@ -2325,13 +2327,13 @@ void UParticleSystem::GetAssetRegistryTags(TArray<FAssetRegistryTag>& OutTags) c
 	}
 	OutTags.Add(FAssetRegistryTag("LODMethod", LODMethodString, FAssetRegistryTag::TT_Alphabetical));
 
+	OutTags.Add(FAssetRegistryTag("CPUCollision", UsesCPUCollision() ? TEXT("True") : TEXT("False"), FAssetRegistryTag::TT_Alphabetical));
 	OutTags.Add(FAssetRegistryTag("Looping", bAnyEmitterLoopsForever ? TEXT("True") : TEXT("False"), FAssetRegistryTag::TT_Alphabetical));
 	OutTags.Add(FAssetRegistryTag("Immortal", IsPotentiallyImmortal() ? TEXT("True") : TEXT("False"), FAssetRegistryTag::TT_Alphabetical));
 
 	Super::GetAssetRegistryTags(OutTags);
 }
 
-
 bool UParticleSystem::IsPotentiallyImmortal() const
 {
 	if (bAutoDeactivate)
@@ -2340,9 +2342,9 @@ bool UParticleSystem::IsPotentiallyImmortal() const
 	}
 
 	// Run thru the emitters and see if any will loop forever
-	for (int32 EmitterIndex = 0; EmitterIndex < Emitters.Num(); ++EmitterIndex)
+	for (const UParticleEmitter* Emitter : Emitters)
 	{
-		if (const UParticleEmitter* Emitter = Emitters[EmitterIndex])
+		if (Emitter != nullptr)
 		{
 			for (const UParticleLODLevel* LODLevel : Emitter->LODLevels)
 			{
@@ -2368,6 +2370,34 @@ bool UParticleSystem::IsPotentiallyImmortal() const
 	return false;
 }
 
+bool UParticleSystem::UsesCPUCollision() const
+{
+	for (const UParticleEmitter* Emitter : Emitters)
+	{
+		if (Emitter != nullptr)
+		{
+			// If we have not yet found a CPU collision module, and we have some enabled LODs to look in..
+			if (Emitter->HasAnyEnabledLODs() && Emitter->LODLevels.Num() > 0)
+			{
+				if (const UParticleLODLevel* HighLODLevel = Emitter->LODLevels[0])
+				{
+					// Iterate over modules of highest LOD (will have all the modules)
+					for (const UParticleModule* Module : HighLODLevel->Modules)
+					{
+						// If an enabled CPU collision module 
+						if (Module->bEnabled && Module->IsA<UParticleModuleCollision>())
+						{
+							return true;
+						}
+					}
+				}
+			}
+		}
+	}
+
+	return false;
+}
+
 bool UParticleSystem::CanBeClusterRoot() const
 {
 	return true;
@@ -4081,6 +4111,9 @@ public:
 	}
 };
 
+TAutoConsoleVariable<int32> CVarFXEarlySchedule(TEXT("FX.EarlyScheduleAsync"), 0, TEXT("If 1, particle system components that can run async will be scheduled earlier in the frame"));
+
+DECLARE_CYCLE_STAT(TEXT("PSys Comp Marshall Time"),STAT_UParticleSystemComponent_Marshall,STATGROUP_Particles);
 
 void UParticleSystemComponent::TickComponent(float DeltaTime, enum ELevelTick TickType, FActorComponentTickFunction *ThisTickFunction)
 {
@@ -4149,7 +4182,8 @@ void UParticleSystemComponent::TickComponent(float DeltaTime, enum ELevelTick Ti
 		}
 		return;
 	} 
-	else if (bWarmingUp == false)
+	// See if DetailMode has changed since the last time we checked
+	else if (bWarmingUp == false && LastCheckedDetailMode != DetailModeCVar)
 	{
 		SCOPE_CYCLE_COUNTER(STAT_UParticleSystemComponent_CheckForReset);
 		bool bRequiresReset = false;
@@ -4165,11 +4199,15 @@ void UParticleSystemComponent::TickComponent(float DeltaTime, enum ELevelTick Ti
 				}
 			}
 		}
+
 		if (bRequiresReset)
 		{
 			ResetParticles(true);
 			InitializeSystem();
 		}
+
+		// Save the detail mode we just checked
+		LastCheckedDetailMode = DetailModeCVar;
 	}
 
 	// Bail out if MaxSecondsBeforeInactive > 0 and we haven't been rendered the last MaxSecondsBeforeInactive seconds.
@@ -4236,6 +4274,27 @@ void UParticleSystemComponent::TickComponent(float DeltaTime, enum ELevelTick Ti
 
 	AccumTickTime += DeltaTime;
 
+	// Save player locations
+	PlayerLocations.Reset();
+	PlayerLODDistanceFactor.Reset();
+
+	if (World->IsGameWorld())
+	{
+		for (FConstPlayerControllerIterator Iterator = World->GetPlayerControllerIterator(); Iterator; ++Iterator)
+		{
+			APlayerController* PlayerController = *Iterator;
+			if (PlayerController->IsLocalPlayerController())
+			{
+				FVector POVLoc;
+				FRotator POVRotation;
+				PlayerController->GetPlayerViewPoint(POVLoc, POVRotation);
+
+				PlayerLocations.Add(POVLoc);
+				PlayerLODDistanceFactor.Add(PlayerController->LocalPlayerCachedLODDistanceFactor);
+			}
+		}
+	}
+
 	// Orient the Z axis toward the camera
 	if (Template->bOrientZAxisTowardCamera)
 	{
@@ -4274,13 +4333,25 @@ void UParticleSystemComponent::TickComponent(float DeltaTime, enum ELevelTick Ti
 	else
 	{
 		SCOPE_CYCLE_COUNTER(STAT_UParticleSystemComponent_QueueTasks);
+		{
+			SCOPE_CYCLE_COUNTER(STAT_UParticleSystemComponent_Marshall)
 		bAsyncDataCopyIsValid = true;
 		check(!bParallelRenderThreadUpdate);
 		AsyncComponentToWorld = ComponentToWorld;
 		AsyncInstanceParameters.Reset();
 		AsyncInstanceParameters.Append(InstanceParameters);
+			AsyncBounds = Bounds;
+			AsyncPartSysVelocity = PartSysVelocity;
+
+			//cache component to world of each actor that trails may use
+			for (FParticleSysParam& ParticleSysParam : AsyncInstanceParameters)
+			{
+				ParticleSysParam.UpdateAsyncActorCache();
+			}
 
 		bAsyncWorkOutstanding = true;
+		}
+		
 		{
 			SCOPE_CYCLE_COUNTER(STAT_UParticleSystemComponent_QueueAsync);
 			AsyncWork = TGraphTask<FParticleAsyncTask>::CreateTask(nullptr, ENamedThreads::GameThread).ConstructAndDispatchWhenReady(this);
@@ -4297,6 +4368,16 @@ void UParticleSystemComponent::TickComponent(float DeltaTime, enum ELevelTick Ti
 			ThisTickFunction->GetCompletionHandle()->DontCompleteUntil(Finalize);
 		}
 #endif
+
+		if(CVarFXEarlySchedule.GetValueOnGameThread())
+		{
+			PrimaryComponentTick.TickGroup = TG_PrePhysics; 
+			PrimaryComponentTick.EndTickGroup = TG_PostPhysics;
+		}
+		else
+		{
+			PrimaryComponentTick.TickGroup = TG_DuringPhysics;
+		}
 	}
 }
 
@@ -4364,6 +4445,16 @@ void UParticleSystemComponent::ComputeTickComponent_Concurrent()
 void UParticleSystemComponent::FinalizeTickComponent()
 {
 	SCOPE_CYCLE_COUNTER(STAT_ParticleFinalizeTickTime);
+
+	if(bAsyncDataCopyIsValid)
+	{
+		//reset async actor to world
+		for (FParticleSysParam& ParticleSysParam : AsyncInstanceParameters)
+		{
+			ParticleSysParam.ResetAsyncActorCache();
+		}
+	}
+
 	bAsyncDataCopyIsValid = false;
 	if (!bNeedsFinalize)
 	{
@@ -4381,22 +4472,22 @@ void UParticleSystemComponent::FinalizeTickComponent()
 			FParticleEmitterInstance* Instance = EmitterInstances[EmitterIndex];
 			if (Instance && Instance->bEnabled)
 			{
-				if (EmitterIndex + 1 < EmitterInstances.Num())
-				{
-					FParticleEmitterInstance* NextInstance = EmitterInstances[EmitterIndex + 1];
-					FPlatformMisc::Prefetch(NextInstance);
-				}
+			if (EmitterIndex + 1 < EmitterInstances.Num())
+			{
+				FParticleEmitterInstance* NextInstance = EmitterInstances[EmitterIndex+1];
+				FPlatformMisc::Prefetch(NextInstance);
+			}
 
 				if (Instance->SpriteTemplate)
+			{
+				UParticleLODLevel* SpriteLODLevel = Instance->SpriteTemplate->GetCurrentLODLevel(Instance);
+				if (SpriteLODLevel && SpriteLODLevel->bEnabled)
 				{
-					UParticleLODLevel* SpriteLODLevel = Instance->SpriteTemplate->GetCurrentLODLevel(Instance);
-					if (SpriteLODLevel && SpriteLODLevel->bEnabled)
-					{
-						Instance->ProcessParticleEvents(DeltaTimeTick, bSuppressSpawning);
-					}
+					Instance->ProcessParticleEvents(DeltaTimeTick, bSuppressSpawning);
 				}
 			}
 		}
+		}
 
 		AParticleEventManager* EventManager = (World ? World->MyParticleEventManager : NULL);
 		if (EventManager)
@@ -4570,18 +4661,18 @@ void UParticleSystemComponent::InitParticles(bool bReInitExistingEmitters)
 			UParticleEmitter* Emitter = Template->Emitters[Idx];
 			FParticleEmitterInstance* Instance = NumInstances == 0 ? NULL : EmitterInstances[Idx];
 			
-			const bool bDetailModeAllowsRendering = DetailMode <= GlobalDetailMode;
+			const bool bDetailModeAllowsRendering	= DetailMode <= GlobalDetailMode;
 			bool bShouldCreateAndOrInit = bDetailModeAllowsRendering && Emitter->HasAnyEnabledLODs() && bCanEverRender;
 			const bool bIsSignificant = (int32)Emitter->SignificanceLevel >= (int32)RequiredSignificance;
 			bool bFailedOnlyOnSignificance = false;
 			if (bShouldCreateAndOrInit && !bIsSignificant)
-			{
+				{
 				bFailedOnlyOnSignificance = true;
 				bShouldCreateAndOrInit = false;
-			}
+					}
 
 			if (bShouldCreateAndOrInit)
-			{
+					{
 				if (Instance)
 				{
 					Instance->SetHaltSpawning(false);
@@ -4590,10 +4681,10 @@ void UParticleSystemComponent::InitParticles(bool bReInitExistingEmitters)
 				{
 					Instance = Emitter->CreateInstance(this);
 					EmitterInstances[Idx] = Instance;
-				}
+			}
 
 				if (Instance)
-				{
+			{
 					Instance->bEnabled = true;
 					if (bReInitExistingEmitters)
 					{
@@ -4603,10 +4694,10 @@ void UParticleSystemComponent::InitParticles(bool bReInitExistingEmitters)
 						PreferredLODLevel = FMath::Min(PreferredLODLevel, Emitter->LODLevels.Num());
 						bSetLodLevels |= !bIsFirstCreate;//Only set lod levels if we init any instances and it's not the first creation time.
 					}
-				}
-			}
-			else
-			{
+						}
+					}
+					else
+					{
 				if (Instance)
 				{
 					if(!bFailedOnlyOnSignificance)
@@ -4619,20 +4710,20 @@ void UParticleSystemComponent::InitParticles(bool bReInitExistingEmitters)
 						EmitterInstances[Idx] = NULL;
 						bClearDynamicData = true;
 					}
-					else
-					{
+				else
+				{
 						//otherwise, we just disable it.
 						Instance->bEnabled = false;
 					}
-				}
+					}
 				bMissingEmittersDueToSignificance |= bFailedOnlyOnSignificance;
+				}
 			}
-		}
 
-		if (bClearDynamicData)
-		{
-			ClearDynamicData();
-		}
+			if (bClearDynamicData)
+			{
+				ClearDynamicData();
+			}
 
 		if (bSetLodLevels)
 		{
@@ -4649,7 +4740,7 @@ void UParticleSystemComponent::InitParticles(bool bReInitExistingEmitters)
 				// set the LOD levels here
 				if (Instance)
 				{
-					Instance->CurrentLODLevelIndex = LODLevel;
+					Instance->CurrentLODLevelIndex	= LODLevel;
 
 					// small safety net for OR-11322; can be removed if the ensure never fires after the change in SetTemplate (reset all instances LOD indices to 0)
 					if (Instance->CurrentLODLevelIndex >= Instance->SpriteTemplate->LODLevels.Num())
@@ -4657,7 +4748,7 @@ void UParticleSystemComponent::InitParticles(bool bReInitExistingEmitters)
 						Instance->CurrentLODLevelIndex = Instance->SpriteTemplate->LODLevels.Num()-1;
 						ensureMsgf(false, TEXT("LOD access out of bounds (OR-11322). Please let olaf.piesche or simon.tovey know."));
 					}
-					Instance->CurrentLODLevel = Instance->SpriteTemplate->LODLevels[Instance->CurrentLODLevelIndex];
+					Instance->CurrentLODLevel		= Instance->SpriteTemplate->LODLevels[Instance->CurrentLODLevelIndex];
 				}
 			}
 		}
diff --git a/Engine/Source/Runtime/Engine/Private/Particles/ParticleModules_Collision.cpp b/Engine/Source/Runtime/Engine/Private/Particles/ParticleModules_Collision.cpp
--- a/Engine/Source/Runtime/Engine/Private/Particles/ParticleModules_Collision.cpp
+++ b/Engine/Source/Runtime/Engine/Private/Particles/ParticleModules_Collision.cpp
@@ -171,11 +171,8 @@ void UParticleModuleCollision::Update(FParticleEmitterInstance* Owner, int32 Off
 	const int32 MeshRotationOffset = Owner->GetMeshRotationOffset();
 	const bool bMeshRotationActive = Owner->IsMeshRotationActive();
 
-	FVector ParentScale = FVector(1.0f, 1.0f, 1.0f);
-	if (Owner->Component)
-	{
-		ParentScale = Owner->Component->ComponentToWorld.GetScale3D();
-	}
+	const FTransform& OwnerTM = Owner->Component->GetAsyncComponentToWorld();
+	const FVector ParentScale = OwnerTM.GetScale3D();
 
 	FParticleEventInstancePayload* EventPayload = NULL;
 	if (LODLevel->EventGenerator)
@@ -191,11 +188,11 @@ void UParticleModuleCollision::Update(FParticleEmitterInstance* Owner, int32 Off
 
 	FParticleCollisionInstancePayload* CollisionInstPayload = (FParticleCollisionInstancePayload*)(Owner->GetModuleInstanceData(this));
 
-	TArray<FVector> PlayerLocations;
-	TArray<float> PlayerLODDistanceFactor;
-	int32 PlayerCount = 0;
+	const TArray<FVector>& PlayerLocations = Owner->Component->GetPlayerLocations();
+	TArray<float> PlayerLODDistanceFactor = Owner->Component->GetPlayerLODDistanceFactor();	//Make a copy because we need to square it later
+	const int32 PlayerCount = PlayerLocations.Num();
 
-	if (Owner->GetWorld()->IsGameWorld())
+	if (World->IsGameWorld())
 	{
 		bool bIgnoreAllCollision = false;
 
@@ -211,24 +208,6 @@ void UParticleModuleCollision::Update(FParticleEmitterInstance* Owner, int32 Off
 			// If the MaxCollisionDistance is greater than WORLD_MAX, they obviously want the check disabled...
 			if (MaxCollisionDistance < WORLD_MAX)
 			{
-				// Store off the player locations and LOD distance factors
-				
-				for( FConstPlayerControllerIterator Iterator = World->GetPlayerControllerIterator(); Iterator; ++Iterator )
-				{
-					APlayerController* PlayerController = *Iterator;
-					if (PlayerController->IsLocalPlayerController())
-					{
-						FVector POVLoc;
-						FRotator POVRotation;
-						PlayerController->GetPlayerViewPoint(POVLoc, POVRotation);
-
-						PlayerLocations.Add(POVLoc);
-						PlayerLODDistanceFactor.Add(PlayerController->LocalPlayerCachedLODDistanceFactor);
-					}
-				}
-
-				PlayerCount = PlayerLocations.Num();
-
 				// If we have at least a few particles, do a simple check vs. the bounds
 				if (Owner->ActiveParticles > 7)
 				{
@@ -238,12 +217,12 @@ void UParticleModuleCollision::Update(FParticleEmitterInstance* Owner, int32 Off
 						BoundingBox.Init();
 						if (Owner->Component->Template && Owner->Component->Template->bUseFixedRelativeBoundingBox)
 						{
-							BoundingBox = Owner->Component->Template->FixedRelativeBoundingBox.TransformBy(Owner->Component->ComponentToWorld);
+							BoundingBox = Owner->Component->Template->FixedRelativeBoundingBox.TransformBy(OwnerTM);
 						}
 						else
 						{
 							// A frame behind, but shouldn't be an issue...
-							BoundingBox = Owner->Component->Bounds.GetBox();
+							BoundingBox = Owner->Component->GetAsyncBounds().GetBox();
 						}
 
 						// see if any player is within the extended bounds...
@@ -330,8 +309,8 @@ void UParticleModuleCollision::Update(FParticleEmitterInstance* Owner, int32 Off
 		if (LODLevel->RequiredModule->bUseLocalSpace)
 		{
 			// Transform the location and old location into world space
-			Location		= Owner->Component->ComponentToWorld.TransformPosition(Location);
-			OldLocation		= Owner->Component->ComponentToWorld.TransformPosition(Particle.OldLocation);
+			Location		= OwnerTM.TransformPosition(Location);
+			OldLocation		= OwnerTM.TransformPosition(Particle.OldLocation);
 		}
 		else
 		{
@@ -361,7 +340,7 @@ void UParticleModuleCollision::Update(FParticleEmitterInstance* Owner, int32 Off
 
 		FVector End = Location + Direction * Size / DirScalar;
 
-		if ((Owner->GetWorld()->IsGameWorld() == true) && (MaxCollisionDistance < WORLD_MAX))
+		if ((World->IsGameWorld() == true) && (MaxCollisionDistance < WORLD_MAX))
 		{
 			// LOD collision by distance
 			bool bCloseEnough = false;
@@ -415,11 +394,11 @@ void UParticleModuleCollision::Update(FParticleEmitterInstance* Owner, int32 Off
 					if (LODLevel->RequiredModule->bUseLocalSpace)
 					{
 						// Transform the particle velocity to world space
-						FVector OldVelocity		= Owner->Component->ComponentToWorld.TransformVector(Particle.Velocity);
-						FVector	BaseVelocity	= Owner->Component->ComponentToWorld.TransformVector(Particle.BaseVelocity);
+						FVector OldVelocity		= OwnerTM.TransformVector(Particle.Velocity);
+						FVector	BaseVelocity	= OwnerTM.TransformVector(Particle.BaseVelocity);
 						BaseVelocity			= BaseVelocity.MirrorByVector(Hit.Normal) * CollisionPayload.UsedDampingFactor;
 
-						Particle.BaseVelocity		= Owner->Component->ComponentToWorld.InverseTransformVector(BaseVelocity);
+						Particle.BaseVelocity		= OwnerTM.InverseTransformVector(BaseVelocity);
 						Particle.BaseRotationRate	= Particle.BaseRotationRate * CollisionPayload.UsedDampingFactorRotation.X;
 						if (bMeshRotationActive && MeshRotationOffset > 0)
 						{
@@ -433,7 +412,7 @@ void UParticleModuleCollision::Update(FParticleEmitterInstance* Owner, int32 Off
 
 						// New location
 						FVector	NewLocation		= Location + NewVelocity * (1.f - Hit.Time);
-						Particle.Location		= Owner->Component->ComponentToWorld.InverseTransformPosition(NewLocation);
+						Particle.Location		= OwnerTM.InverseTransformPosition(NewLocation);
 
 						if (bApplyPhysics)
 						{
@@ -487,14 +466,14 @@ void UParticleModuleCollision::Update(FParticleEmitterInstance* Owner, int32 Off
 				{
 					if (LODLevel->RequiredModule->bUseLocalSpace == true)
 					{
-						Size = Owner->Component->ComponentToWorld.TransformVector(Size);
+						Size = OwnerTM.TransformVector(Size);
 					}
 					Particle.Location = Hit.Location + (Size / 2.0f);
 					if (LODLevel->RequiredModule->bUseLocalSpace == true)
 					{
 						// We need to transform the location back relative to the PSys.
 						// NOTE: LocalSpace makes sense only for stationary emitters that use collision.
-						Particle.Location = Owner->Component->ComponentToWorld.InverseTransformPosition(Particle.Location);
+						Particle.Location = OwnerTM.InverseTransformPosition(Particle.Location);
 					}
 					switch (CollisionCompletionOption)
 					{
@@ -549,7 +528,7 @@ void UParticleModuleCollision::Update(FParticleEmitterInstance* Owner, int32 Off
 
 bool UParticleModuleCollision::CanTickInAnyThread()
 {
-	return false; //TODO: make all of the update function thread safe before turning this one /*!bIgnoreTriggerVolumes && !bApplyPhysics && !bPawnsDoNotDecrementCount;*/
+	return !bApplyPhysics;
 }
 
 void UParticleModuleCollision::SetToSensibleDefaults(UParticleEmitter* Owner)
diff --git a/Engine/Source/Runtime/Engine/Private/Particles/ParticleModules_Location.cpp b/Engine/Source/Runtime/Engine/Private/Particles/ParticleModules_Location.cpp
--- a/Engine/Source/Runtime/Engine/Private/Particles/ParticleModules_Location.cpp
+++ b/Engine/Source/Runtime/Engine/Private/Particles/ParticleModules_Location.cpp
@@ -341,6 +341,7 @@ void UParticleModuleLocationDirect::Spawn(FParticleEmitterInstance* Owner, int32
 
 void UParticleModuleLocationDirect::Update(FParticleEmitterInstance* Owner, int32 Offset, float DeltaTime)
 {
+	const FTransform& OwnerTM = Owner->Component->GetAsyncComponentToWorld();
 	BEGIN_UPDATE_LOOP;
 	{
 		FVector	NewLoc;
@@ -353,7 +354,7 @@ void UParticleModuleLocationDirect::Update(FParticleEmitterInstance* Owner, int3
 		else
 		{
 			FVector Loc			= Location.GetValue(Particle.RelativeTime, Owner->Component);
-			Loc = Owner->Component->ComponentToWorld.TransformPosition(Loc);
+			Loc = OwnerTM.TransformPosition(Loc);
 			NewLoc	= Loc;
 		}
 
@@ -1605,6 +1606,8 @@ void UParticleModuleLocationBoneSocket::Update(FParticleEmitterInstance* Owner,
 	const int32 MeshRotationOffset = Owner->GetMeshRotationOffset();
 	const bool bMeshRotationActive = MeshRotationOffset > 0 && Owner->IsMeshRotationActive();
 	FQuat* SourceRotation = (bMeshRotationActive) ? NULL : &RotationQuat;
+	const FTransform& OwnerTM = Owner->Component->GetAsyncComponentToWorld();
+
 	BEGIN_UPDATE_LOOP;
 	{
 		FModuleLocationBoneSocketParticlePayload* ParticlePayload = (FModuleLocationBoneSocketParticlePayload*)((uint8*)&Particle + Offset);
@@ -1617,7 +1620,7 @@ void UParticleModuleLocationBoneSocket::Update(FParticleEmitterInstance* Owner,
 				PayloadData->Rotation = RotationQuat.Euler();
 				if (Owner->CurrentLODLevel->RequiredModule->bUseLocalSpace == true)
 				{
-					PayloadData->Rotation = Owner->Component->ComponentToWorld.InverseTransformVectorNoScale(PayloadData->Rotation);
+					PayloadData->Rotation = OwnerTM.InverseTransformVectorNoScale(PayloadData->Rotation);
 				}
 			}
 		}
@@ -2247,6 +2250,8 @@ void UParticleModuleLocationSkelVertSurface::Update(FParticleEmitterInstance* Ow
 	FQuat SourceRotation;
 	const int32 MeshRotationOffset = Owner->GetMeshRotationOffset();
 	const bool bMeshRotationActive = MeshRotationOffset > 0 && Owner->IsMeshRotationActive();
+	const FTransform& OwnerTM = Owner->Component->GetAsyncComponentToWorld();
+
 	BEGIN_UPDATE_LOOP;
 	{
 		FModuleLocationVertSurfaceParticlePayload* ParticlePayload = (FModuleLocationVertSurfaceParticlePayload*)((uint8*)&Particle + Offset);
@@ -2270,7 +2275,7 @@ void UParticleModuleLocationSkelVertSurface::Update(FParticleEmitterInstance* Ow
 				FVector Rot = SourceRotation.Euler();
 				if (Owner->CurrentLODLevel->RequiredModule->bUseLocalSpace == true)
 				{
-					Rot = Owner->Component->ComponentToWorld.InverseTransformVectorNoScale(Rot);
+					Rot = OwnerTM.InverseTransformVectorNoScale(Rot);
 				}
 				PayloadData->Rotation = Rot;
 			}
diff --git a/Engine/Source/Runtime/Engine/Private/Particles/ParticleModules_Velocity.cpp b/Engine/Source/Runtime/Engine/Private/Particles/ParticleModules_Velocity.cpp
--- a/Engine/Source/Runtime/Engine/Private/Particles/ParticleModules_Velocity.cpp
+++ b/Engine/Source/Runtime/Engine/Private/Particles/ParticleModules_Velocity.cpp
@@ -270,9 +270,10 @@ void UParticleModuleVelocityOverLifetime::Update(FParticleEmitterInstance* Owner
 	UParticleLODLevel* LODLevel	= Owner->SpriteTemplate->GetCurrentLODLevel(Owner);
 	check(LODLevel);
 	FVector OwnerScale(1.0f);
+	const FTransform& OwnerTM = Owner->Component->GetAsyncComponentToWorld();
 	if ((bApplyOwnerScale == true) && Owner->Component)
 	{
-		OwnerScale = Owner->Component->ComponentToWorld.GetScale3D();
+		OwnerScale = OwnerTM.GetScale3D();
 	}
 	if (Absolute)
 	{
@@ -281,7 +282,7 @@ void UParticleModuleVelocityOverLifetime::Update(FParticleEmitterInstance* Owner
 			if (bInWorldSpace == false)
 			{
 				FVector Vel;
-				const FMatrix LocalToWorld = Owner->Component->ComponentToWorld.ToMatrixNoScale();
+				const FMatrix LocalToWorld = OwnerTM.ToMatrixNoScale();
 				BEGIN_UPDATE_LOOP;
 				{
 					Vel = VelOverLife.GetValue(Particle.RelativeTime, Owner->Component);
@@ -311,7 +312,7 @@ void UParticleModuleVelocityOverLifetime::Update(FParticleEmitterInstance* Owner
 			else
 			{
 				FVector Vel;
-				const FMatrix LocalToWorld = Owner->Component->ComponentToWorld.ToMatrixNoScale();
+				const FMatrix LocalToWorld = OwnerTM.ToMatrixNoScale();
 				const FMatrix InvMat = LocalToWorld.InverseFast();
 				BEGIN_UPDATE_LOOP;
 				{
@@ -329,7 +330,7 @@ void UParticleModuleVelocityOverLifetime::Update(FParticleEmitterInstance* Owner
 			FVector Vel;
 			if (bInWorldSpace == false)
 			{
-				const FMatrix LocalToWorld = Owner->Component->ComponentToWorld.ToMatrixNoScale();
+				const FMatrix LocalToWorld = OwnerTM.ToMatrixNoScale();
 				BEGIN_UPDATE_LOOP;
 				{
 					Vel = VelOverLife.GetValue(Particle.RelativeTime, Owner->Component);
@@ -359,7 +360,7 @@ void UParticleModuleVelocityOverLifetime::Update(FParticleEmitterInstance* Owner
 			else
 			{
 				FVector Vel;
-				const FMatrix LocalToWorld = Owner->Component->ComponentToWorld.ToMatrixNoScale();
+				const FMatrix LocalToWorld = OwnerTM.ToMatrixNoScale();
 				const FMatrix InvMat = LocalToWorld.InverseFast();
 				BEGIN_UPDATE_LOOP;
 				{
diff --git a/Engine/Source/Runtime/Engine/Private/Particles/ParticleTrail2EmitterInstance.cpp b/Engine/Source/Runtime/Engine/Private/Particles/ParticleTrail2EmitterInstance.cpp
--- a/Engine/Source/Runtime/Engine/Private/Particles/ParticleTrail2EmitterInstance.cpp
+++ b/Engine/Source/Runtime/Engine/Private/Particles/ParticleTrail2EmitterInstance.cpp
@@ -2036,7 +2036,6 @@ void FParticleRibbonEmitterInstance::SetupTrailModules()
 
 void FParticleRibbonEmitterInstance::ResolveSource()
 {
-	check(IsInGameThread());
 	if (SourceModule && SourceModule->SourceName != NAME_None)
 	{
 		switch (SourceModule->SourceMethod)
@@ -2044,10 +2043,11 @@ void FParticleRibbonEmitterInstance::ResolveSource()
 		case PET2SRCM_Actor:
 			if (SourceActor == NULL)
 			{
+				const TArray<struct FParticleSysParam>& AsyncInstanceParameters = Component->GetAsyncInstanceParameters();
 				FParticleSysParam Param;
-				for (int32 ParamIdx = 0; ParamIdx < Component->InstanceParameters.Num(); ParamIdx++)
+				for (int32 ParamIdx = 0; ParamIdx < AsyncInstanceParameters.Num(); ParamIdx++)
 				{
-					Param = Component->InstanceParameters[ParamIdx];
+					Param = AsyncInstanceParameters[ParamIdx];
 					if (Param.Name == SourceModule->SourceName)
 					{
 						SourceActor = Param.Actor;
@@ -2057,9 +2057,9 @@ void FParticleRibbonEmitterInstance::ResolveSource()
 
 				if (SourceModule->SourceOffsetCount > 0)
 				{
-					for (int32 ParamIdx = 0; ParamIdx < Component->InstanceParameters.Num(); ParamIdx++)
+					for (int32 ParamIdx = 0; ParamIdx < AsyncInstanceParameters.Num(); ParamIdx++)
 					{
-						Param = Component->InstanceParameters[ParamIdx];
+						Param = AsyncInstanceParameters[ParamIdx];
 						FString ParamName = Param.Name.ToString();
 						const TCHAR* TrailSourceOffset = FCString::Strstr(*ParamName, TEXT("TrailSourceOffset"));
 						if (TrailSourceOffset)
@@ -2161,6 +2161,7 @@ void FParticleRibbonEmitterInstance::UpdateSourceData(float DeltaTime, bool bFir
 bool FParticleRibbonEmitterInstance::ResolveSourcePoint(int32 InTrailIdx, 
 	FVector& OutPosition, FQuat& OutRotation, FVector& OutUp, FVector& OutTangent, float& OutTangentStrength)
 {
+	const FTransform& AsyncComponentToWorld = Component->GetAsyncComponentToWorld();
 	bool bSourceWasSet = false;
 	// Resolve the source point...
 	if (SourceModule)
@@ -2290,15 +2291,15 @@ bool FParticleRibbonEmitterInstance::ResolveSourcePoint(int32 InTrailIdx,
 					else
 					{
 						// Fall back to the emitter location??
-						OutPosition = SourceEmitter->Component->GetComponentLocation();
-						OutTangent = Component->PartSysVelocity;
+						OutPosition = AsyncComponentToWorld.GetLocation();
+						OutTangent = Component->GetAsyncPartSysVelocity();
 						//@todo. How to handle this... can potentially cause a jump from the emitter to the
 						// particle...
 						SourceTimes[InTrailIdx] = 0.0f;
 					}
 					OutTangentStrength = OutTangent.SizeSquared();
 					//@todo. Allow particle rotation to define up??
-					OutUp = SourceEmitter->Component->ComponentToWorld.GetScaledAxis( EAxis::Z );
+					OutUp = AsyncComponentToWorld.GetScaledAxis( EAxis::Z );
 
 					//@todo. Where to get rotation from????
 					OutRotation = FQuat(0,0,0,1);
@@ -2319,11 +2320,29 @@ bool FParticleRibbonEmitterInstance::ResolveSourcePoint(int32 InTrailIdx,
 
 				if (SourceActor)
 				{
-					FTransform ActorToWorld = SourceActor->ActorToWorld();
+					FTransform ActorToWorld = FTransform::Identity;
+					FVector SourceActorVelocity = FVector::ZeroVector;
+					if(IsInGameThread())
+					{
+						ActorToWorld = SourceActor->ActorToWorld();
+						SourceActorVelocity = SourceActor->GetVelocity();
+					}
+					else
+					{
+						const TArray<FParticleSysParam>& AsyncParticleSysParams = Component->GetAsyncInstanceParameters();
+						for(const FParticleSysParam& ParticleSysParam : AsyncParticleSysParams)
+						{
+							if(ParticleSysParam.Actor == SourceActor)
+							{
+								ActorToWorld = ParticleSysParam.GetAsyncActorToWorld();
+								SourceActorVelocity = ParticleSysParam.GetAsyncActorVelocity();
+							}
+						}
+					}
 					OutPosition = ActorToWorld.GetLocation();
 					FRotator TempRotator = ActorToWorld.Rotator();
 					OutRotation = FQuat(TempRotator);
-					OutTangent = SourceActor->GetVelocity();
+					OutTangent = SourceActorVelocity;
 					OutTangentStrength = OutTangent.SizeSquared();
 
 					OutUp = ActorToWorld.TransformVector(FVector(0.f,0.f,1.f));
@@ -2337,7 +2356,7 @@ bool FParticleRibbonEmitterInstance::ResolveSourcePoint(int32 InTrailIdx,
 
 	if (bSourceWasSet == false)
 	{
-		OutPosition = Component->GetComponentLocation();
+		OutPosition = AsyncComponentToWorld.GetLocation();
 		if (SourceModule && (SourceModule->SourceOffsetCount > 0))
 		{
 			FVector SourceOffsetValue;
@@ -2346,15 +2365,15 @@ bool FParticleRibbonEmitterInstance::ResolveSourcePoint(int32 InTrailIdx,
 				if (CurrentLODLevel && (CurrentLODLevel->RequiredModule->bUseLocalSpace == false))
 				{
 					// Transform it
-					SourceOffsetValue = Component->ComponentToWorld.TransformVector(SourceOffsetValue);
+					SourceOffsetValue = AsyncComponentToWorld.TransformVector(SourceOffsetValue);
 				}
 				OutPosition += SourceOffsetValue;
 			}
 		}
-		OutRotation = Component->GetComponentQuat();
-		OutTangent = Component->PartSysVelocity;
+		OutRotation = AsyncComponentToWorld.GetRotation();
+		OutTangent = Component->GetAsyncPartSysVelocity();
 		OutTangentStrength = OutTangent.SizeSquared();
-		OutUp = Component->ComponentToWorld.GetScaledAxis( EAxis::Z );
+		OutUp = AsyncComponentToWorld.GetScaledAxis( EAxis::Z );
 
 		bSourceWasSet = true;
 	}
diff --git a/Engine/Source/Runtime/Engine/Private/PhysicsEngine/BodyInstance.cpp b/Engine/Source/Runtime/Engine/Private/PhysicsEngine/BodyInstance.cpp
--- a/Engine/Source/Runtime/Engine/Private/PhysicsEngine/BodyInstance.cpp
+++ b/Engine/Source/Runtime/Engine/Private/PhysicsEngine/BodyInstance.cpp
@@ -703,7 +703,7 @@ void ExecuteOnPxShapeWrite(FBodyInstance* BodyInstance, PxShape* PShape, Lambda
 }
 
 
-void FBodyInstance::UpdatePhysicsShapeFilterData(uint32 SkelMeshCompID, bool bUseComplexAsSimple, bool bUseSimpleAsComplex, bool bPhysicsStatic, const TEnumAsByte<ECollisionEnabled::Type> * CollisionEnabledOverride, FCollisionResponseContainer * ResponseOverride, bool * bNotifyOverride)
+void FBodyInstance::UpdatePhysicsShapeFilterData(uint32 ComponentID, bool bUseComplexAsSimple, bool bUseSimpleAsComplex, bool bPhysicsStatic, const TEnumAsByte<ECollisionEnabled::Type> * CollisionEnabledOverride, FCollisionResponseContainer * ResponseOverride, bool * bNotifyOverride)
 {
 	ExecuteOnPhysicsReadWrite([&]
 	{
@@ -727,7 +727,8 @@ void FBodyInstance::UpdatePhysicsShapeFilterData(uint32 SkelMeshCompID, bool bUs
 
 
 			UPrimitiveComponent* OwnerPrimitiveComponent = BI->OwnerComponent.Get();
-			int32 CompID = (OwnerPrimitiveComponent != nullptr) ? OwnerPrimitiveComponent->GetUniqueID() : 0;
+			AActor* OwnerActor = OwnerPrimitiveComponent ? OwnerPrimitiveComponent->GetOwner() : nullptr;
+			int32 ActorID = (OwnerActor != nullptr) ? OwnerActor->GetUniqueID() : 0;
 
 			// Create the filterdata structs
 			PxFilterData PSimFilterData;
@@ -735,7 +736,7 @@ void FBodyInstance::UpdatePhysicsShapeFilterData(uint32 SkelMeshCompID, bool bUs
 			PxFilterData PComplexQueryData;
 			if (UseCollisionEnabled != ECollisionEnabled::NoCollision)
 			{
-				CreateShapeFilterData(BI->ObjectType, MaskFilter, CompID, UseResponse, SkelMeshCompID, InstanceBodyIndex, PSimpleQueryData, PSimFilterData, bUseCCD && !bPhysicsStatic, bUseNotify, bPhysicsStatic);	//InstanceBodyIndex and CCD are determined by root body in case of welding
+				CreateShapeFilterData(BI->ObjectType, MaskFilter, ActorID, UseResponse, ComponentID, InstanceBodyIndex, PSimpleQueryData, PSimFilterData, bUseCCD && !bPhysicsStatic, bUseNotify, bPhysicsStatic);	//InstanceBodyIndex and CCD are determined by root body in case of welding
 				PComplexQueryData = PSimpleQueryData;
 
 				// Build filterdata variations for complex and simple
@@ -929,12 +930,8 @@ void FBodyInstance::UpdatePhysicsFilterData()
 	bool bResponseOverride = false;
 	bool bNotifyOverride = false;
 
-	// Get skelmeshcomp ID
-	uint32 SkelMeshCompID = 0;
 	if (USkeletalMeshComponent* SkelMeshComp = Cast<USkeletalMeshComponent>(OwnerComponentInst))
 	{
-		SkelMeshCompID = SkelMeshComp->GetUniqueID();
-
 		// In skeletal case, collision enable/disable/movement should be overriden by mesh component
 		// being in the physics asset, and not having collision is a waste and it can cause a bug where disconnected bodies
 		if (Owner)
@@ -990,11 +987,14 @@ void FBodyInstance::UpdatePhysicsFilterData()
 	const bool bUseComplexAsSimple = (BodySetup.Get()->GetCollisionTraceFlag() == CTF_UseComplexAsSimple);
 	const bool bUseSimpleAsComplex = (BodySetup.Get()->GetCollisionTraceFlag() == CTF_UseSimpleAsComplex);
 
+	// Get component ID
+	uint32 ComponentID = OwnerComponentInst->GetUniqueID();
+
 #if WITH_PHYSX
 	const TEnumAsByte<ECollisionEnabled::Type>* CollisionEnabledOverride = bUseCollisionEnabledOverride ? &UseCollisionEnabled : NULL;
 	FCollisionResponseContainer * ResponseOverride = bResponseOverride ? &UseResponse : NULL;
 	bool * bNotifyOverridePtr = bNotifyOverride ? &bUseNotifyRBCollision : NULL;
-	UpdatePhysicsShapeFilterData(SkelMeshCompID, bUseComplexAsSimple, bUseSimpleAsComplex, bPhysicsStatic, CollisionEnabledOverride, ResponseOverride, bNotifyOverridePtr);
+	UpdatePhysicsShapeFilterData(ComponentID, bUseComplexAsSimple, bUseSimpleAsComplex, bPhysicsStatic, CollisionEnabledOverride, ResponseOverride, bNotifyOverridePtr);
 #endif
 
 #if WITH_BOX2D
@@ -1010,7 +1010,7 @@ void FBodyInstance::UpdatePhysicsFilterData()
 
 			b2Filter BoxSimFilterData;
 			FilterBuilder.GetCombinedData(/*out*/ BoxSimFilterData.BlockingChannels, /*out*/ BoxSimFilterData.TouchingChannels, /*out*/ BoxSimFilterData.ObjectTypeAndFlags);
-			BoxSimFilterData.UniqueComponentID = SkelMeshCompID;
+			BoxSimFilterData.UniqueComponentID = ComponentID;
 			BoxSimFilterData.BodyIndex = InstanceBodyIndex;
 
 			// Update the body data
@@ -4700,12 +4700,8 @@ void FBodyInstance::GetFilterData_AssumesLocked(FShapeData& ShapeData, bool bFor
 	bool bResponseOverride = false;
 	bool bNotifyOverride = false;
 
-	// Get skelmeshcomp ID
-	uint32 SkelMeshCompID = 0;
 	if(USkeletalMeshComponent* SkelMeshComp = Cast<USkeletalMeshComponent>(OwnerComponentInst))
 	{
-		SkelMeshCompID = SkelMeshComp->GetUniqueID();
-
 		// In skeletal case, collision enable/disable/movement should be overriden by mesh component
 		// being in the physics asset, and not having collision is a waste and it can cause a bug where disconnected bodies
 		if(Owner)
@@ -4770,8 +4766,9 @@ void FBodyInstance::GetFilterData_AssumesLocked(FShapeData& ShapeData, bool bFor
 			PxFilterData PSimFilterData;
 			PxFilterData PSimpleQueryData;
 			PxFilterData PComplexQueryData;
-			int32 CompID = (OwnerComponentInst != nullptr) ? OwnerComponentInst->GetUniqueID() : 0;
-			CreateShapeFilterData(ObjectType, MaskFilter, CompID, UseResponse, SkelMeshCompID, InstanceBodyIndex, PSimpleQueryData, PSimFilterData, bUseCCD && !bPhysicsStatic, bUseNotifyRBCollision, bPhysicsStatic);	//CCD is determined by root body in case of welding
+			uint32 ActorID = Owner ? Owner->GetUniqueID() : 0;
+			uint32 CompID = (OwnerComponentInst != nullptr) ? OwnerComponentInst->GetUniqueID() : 0;
+			CreateShapeFilterData(ObjectType, MaskFilter, ActorID, UseResponse, CompID, InstanceBodyIndex, PSimpleQueryData, PSimFilterData, bUseCCD && !bPhysicsStatic, bUseNotifyRBCollision, bPhysicsStatic);	//CCD is determined by root body in case of welding
 			PComplexQueryData = PSimpleQueryData;
 			
 			// Set output sim data
diff --git a/Engine/Source/Runtime/Engine/Private/PhysicsEngine/PhysLevel.cpp b/Engine/Source/Runtime/Engine/Private/PhysicsEngine/PhysLevel.cpp
--- a/Engine/Source/Runtime/Engine/Private/PhysicsEngine/PhysLevel.cpp
+++ b/Engine/Source/Runtime/Engine/Private/PhysicsEngine/PhysLevel.cpp
@@ -349,6 +349,14 @@ void InitGamePhys()
 	check(ErrorCode == APEX_CE_NO_ERROR);
 	check(GApexSDK);
 
+
+#if UE_BUILD_SHIPPING
+	GApexSDK->setEnableApexStats(false);
+#endif
+
+
+
+
 #if APEX_STATICALLY_LINKED
 	// We need to instantiate the module if we have statically linked them
 	// Otherwise all createModule functions will fail
diff --git a/Engine/Source/Runtime/Engine/Private/PhysicsEngine/PhysXSupport.cpp b/Engine/Source/Runtime/Engine/Private/PhysicsEngine/PhysXSupport.cpp
--- a/Engine/Source/Runtime/Engine/Private/PhysicsEngine/PhysXSupport.cpp
+++ b/Engine/Source/Runtime/Engine/Private/PhysicsEngine/PhysXSupport.cpp
@@ -325,8 +325,8 @@ PxFilterFlags PhysXSimFilterShader(	PxFilterObjectAttributes attributes0, PxFilt
 		return PxFilterFlag::eSUPPRESS;
 	}
 
-	// if these bodies are from the same skeletal mesh component, use the disable table to see if we should disable collision
-	if((filterData0.word2 == filterData1.word2) && (filterData0.word2 != 0))
+	// if these bodies are from the same component, use the disable table to see if we should disable collision. This case should only happen for things like skeletalmesh and destruction. The table is only created for skeletal mesh components at the moment
+	if(filterData0.word2 == filterData1.word2)
 	{
 		check(constantBlockSize == sizeof(FPhysSceneShaderInfo));
 		const FPhysSceneShaderInfo* PhysSceneShaderInfo = (const FPhysSceneShaderInfo*) constantBlock;
diff --git a/Engine/Source/Runtime/Engine/Private/PhysicsEngine/RadialForceComponent.cpp b/Engine/Source/Runtime/Engine/Private/PhysicsEngine/RadialForceComponent.cpp
--- a/Engine/Source/Runtime/Engine/Private/PhysicsEngine/RadialForceComponent.cpp
+++ b/Engine/Source/Runtime/Engine/Private/PhysicsEngine/RadialForceComponent.cpp
@@ -52,27 +52,36 @@ void URadialForceComponent::TickComponent(float DeltaTime, enum ELevelTick TickT
 
 		GetWorld()->OverlapMultiByObjectType(Overlaps, Origin, FQuat::Identity, CollisionObjectQueryParams, FCollisionShape::MakeSphere(Radius), Params);
 
-		// Iterate over each and apply force
-		for ( int32 OverlapIdx=0; OverlapIdx<Overlaps.Num(); OverlapIdx++ )
+		// A component can have multiple physics presences (e.g. destructible mesh components).
+		// The component should handle the radial force for all of the physics objects it contains
+		// so here we grab all of the unique components to avoid applying impulses more than once.
+		TArray<UPrimitiveComponent*, TInlineAllocator<1>> AffectedComponents;
+		AffectedComponents.Reserve(Overlaps.Num());
+
+		for(FOverlapResult& OverlapResult : Overlaps)
 		{
-			UPrimitiveComponent* PokeComp = Overlaps[OverlapIdx].Component.Get();
-			if(PokeComp)
+			if(UPrimitiveComponent* PrimitiveComponent = OverlapResult.Component.Get())
 			{
-				PokeComp->AddRadialForce( Origin, Radius, ForceStrength, Falloff );
+				AffectedComponents.AddUnique(PrimitiveComponent);
+			}
+		}
+
+		for(UPrimitiveComponent* PrimitiveComponent : AffectedComponents)
+		{
+			PrimitiveComponent->AddRadialForce(Origin, Radius, ForceStrength, Falloff);
 
-				// see if this is a target for a movement component
-				AActor* PokeOwner = PokeComp->GetOwner();
-				if(PokeOwner)
+			// see if this is a target for a movement component
+			AActor* ComponentOwner = PrimitiveComponent->GetOwner();
+			if(ComponentOwner)
+			{
+				TInlineComponentArray<UMovementComponent*> MovementComponents;
+				ComponentOwner->GetComponents<UMovementComponent>(MovementComponents);
+				for(const auto& MovementComponent : MovementComponents)
 				{
-					TInlineComponentArray<UMovementComponent*> MovementComponents;
-					PokeOwner->GetComponents<UMovementComponent>(MovementComponents);
-					for(const auto& MovementComponent : MovementComponents)
+					if(MovementComponent->UpdatedComponent == PrimitiveComponent)
 					{
-						if(MovementComponent->UpdatedComponent == PokeComp)
-						{
-							MovementComponent->AddRadialForce( Origin, Radius, ForceStrength, Falloff );
-							break;
-						}
+						MovementComponent->AddRadialForce(Origin, Radius, ForceStrength, Falloff);
+						break;
 					}
 				}
 			}
@@ -106,35 +115,42 @@ void URadialForceComponent::FireImpulse()
 
 	GetWorld()->OverlapMultiByObjectType(Overlaps, Origin, FQuat::Identity, CollisionObjectQueryParams, FCollisionShape::MakeSphere(Radius), Params);
 
-	// Iterate over each and apply an impulse
-	for ( int32 OverlapIdx=0; OverlapIdx<Overlaps.Num(); OverlapIdx++ )
+	// A component can have multiple physics presences (e.g. destructible mesh components).
+	// The component should handle the radial force for all of the physics objects it contains
+	// so here we grab all of the unique components to avoid applying impulses more than once.
+	TArray<UPrimitiveComponent*, TInlineAllocator<1>> AffectedComponents;
+	AffectedComponents.Reserve(Overlaps.Num());
+
+	for(FOverlapResult& OverlapResult : Overlaps)
+	{
+		if(UPrimitiveComponent* PrimitiveComponent = OverlapResult.Component.Get())
+		{
+			AffectedComponents.AddUnique(PrimitiveComponent);
+		}
+	}
+
+	for(UPrimitiveComponent* PrimitiveComponent : AffectedComponents)
 	{
-		UPrimitiveComponent* PokeComp = Overlaps[OverlapIdx].Component.Get();
-		if(PokeComp)
+		if(DestructibleDamage > SMALL_NUMBER)
 		{
-			// If DestructibleDamage is non-zero, see if this is a destructible, and do damage if so.
-			if(DestructibleDamage > SMALL_NUMBER)
+			if(UDestructibleComponent* DestructibleComponent = Cast<UDestructibleComponent>(PrimitiveComponent))
 			{
-				UDestructibleComponent* DestructibleComp = Cast<UDestructibleComponent>(PokeComp);
-				if(DestructibleComp != NULL)
-				{
-					DestructibleComp->ApplyRadiusDamage(DestructibleDamage, Origin, Radius, ImpulseStrength, (Falloff == RIF_Constant));
-				}
+				DestructibleComponent->ApplyRadiusDamage(DestructibleDamage, Origin, Radius, ImpulseStrength, Falloff == RIF_Constant);
 			}
+		}
 
-			// Do impulse after
-			PokeComp->AddRadialImpulse( Origin, Radius, ImpulseStrength, Falloff, bImpulseVelChange );
+		// Apply impulse
+		PrimitiveComponent->AddRadialImpulse(Origin, Radius, ImpulseStrength, Falloff, bImpulseVelChange);
 
-			// see if this is a target for a movement component
-			TInlineComponentArray<UMovementComponent*> MovementComponents;
-			PokeComp->GetOwner()->GetComponents<UMovementComponent>(MovementComponents);
-			for(const auto& MovementComponent : MovementComponents)
+		// See if this is a target for a movement component, if so apply the impulse to it
+		TInlineComponentArray<UMovementComponent*> MovementComponents;
+		PrimitiveComponent->GetOwner()->GetComponents<UMovementComponent>(MovementComponents);
+		for(const auto& MovementComponent : MovementComponents)
+		{
+			if(MovementComponent->UpdatedComponent == PrimitiveComponent)
 			{
-				if(MovementComponent->UpdatedComponent == PokeComp)
-				{
-					MovementComponent->AddRadialImpulse( Origin, Radius, ImpulseStrength, Falloff, bImpulseVelChange );
-					break;
-				}
+				MovementComponent->AddRadialImpulse(Origin, Radius, ImpulseStrength, Falloff, bImpulseVelChange);
+				break;
 			}
 		}
 	}
diff --git a/Engine/Source/Runtime/Engine/Private/PlayerCameraManager.cpp b/Engine/Source/Runtime/Engine/Private/PlayerCameraManager.cpp
--- a/Engine/Source/Runtime/Engine/Private/PlayerCameraManager.cpp
+++ b/Engine/Source/Runtime/Engine/Private/PlayerCameraManager.cpp
@@ -1230,7 +1230,7 @@ float APlayerCameraManager::CalcRadialShakeScale(APlayerCameraManager* Camera, F
 {
 	// using camera location so stuff like spectator cameras get shakes applied sensibly as well
 	// need to ensure server has reasonably accurate camera position
-	FVector POVLoc = Camera->GetActorLocation();
+	FVector POVLoc = Camera->GetCameraLocation();
 
 	if (InnerRadius < OuterRadius)
 	{
diff --git a/Engine/Source/Runtime/Engine/Private/PlayerController.cpp b/Engine/Source/Runtime/Engine/Private/PlayerController.cpp
--- a/Engine/Source/Runtime/Engine/Private/PlayerController.cpp
+++ b/Engine/Source/Runtime/Engine/Private/PlayerController.cpp
@@ -86,6 +86,7 @@ APlayerController::APlayerController(const FObjectInitializer& ObjectInitializer
 	bAutoManageActiveCameraTarget = true;
 	bHidePawnInCinematicMode = false;
 
+	bIsPlayerController = true;
 	bIsLocalPlayerController = false;
 
 	if (RootComponent)
@@ -144,27 +145,42 @@ bool APlayerController::DestroyNetworkActorHandled()
 
 bool APlayerController::IsLocalController() const
 {
+	// Never local on dedicated server, always local on clients. IsServerOnly() and IsClientOnly() are checked at compile time and optimized out appropriately.
+	if (FGenericPlatformProperties::IsServerOnly())
+	{
+		check(!bIsLocalPlayerController);
+		return false;
+	}
+	else if (FGenericPlatformProperties::IsClientOnly())
+	{
+		bIsLocalPlayerController = true;
+		return true;
+	}
+	
+	// Fast path if we have this bool set.
+	if (bIsLocalPlayerController)
+	{
+		return true;
+	}
+
 	ENetMode NetMode = GetNetMode();
 	if (NetMode == NM_DedicatedServer)
 	{
+		// This is still checked for the PIE case, which would not be caught in the IsServerOnly() check above.
 		check(!bIsLocalPlayerController);
 		return false;
 	}
 
-	if (NetMode == NM_Client)
+	if (NetMode == NM_Client || NetMode == NM_Standalone)
 	{
-		// Clients only receive their own PC. We are not ROLE_AutonomousProxy until after PostInitializeComponents so we can't check that.
+		// Clients or Standalone only receive their own PC. We are not ROLE_AutonomousProxy until after PostInitializeComponents so we can't check that.
+		bIsLocalPlayerController = true;
 		return true;
 	}
 
 	return bIsLocalPlayerController;
 }
 
-bool APlayerController::IsLocalPlayerController() const
-{
-	// We automatically pass the "IsPlayer" part because we are a PlayerController...
-	return IsLocalController();
-}
 
 void APlayerController::FailedToSpawnPawn()
 {
diff --git a/Engine/Source/Runtime/Engine/Private/SceneView.cpp b/Engine/Source/Runtime/Engine/Private/SceneView.cpp
--- a/Engine/Source/Runtime/Engine/Private/SceneView.cpp
+++ b/Engine/Source/Runtime/Engine/Private/SceneView.cpp
@@ -1408,7 +1408,7 @@ void FSceneView::EndFinalPostprocessSettings(const FSceneViewInitOptions& ViewIn
 
 	if(!Family->EngineShowFlags.AmbientCubemap)
 	{
-		FinalPostProcessSettings.ContributingCubemaps.Empty();
+		FinalPostProcessSettings.ContributingCubemaps.Reset();
 	}
 
 	if(!Family->EngineShowFlags.LensFlares)
diff --git a/Engine/Source/Runtime/Engine/Private/SoundAttenuation.cpp b/Engine/Source/Runtime/Engine/Private/SoundAttenuation.cpp
--- a/Engine/Source/Runtime/Engine/Private/SoundAttenuation.cpp
+++ b/Engine/Source/Runtime/Engine/Private/SoundAttenuation.cpp
@@ -68,6 +68,30 @@ float FAttenuationSettings::GetMaxDimension() const
 	return MaxDimension;
 }
 
+float FAttenuationSettings::GetFocusPriorityScale(const struct FGlobalFocusSettings& FocusSettings, float FocusFactor) const
+{
+	float Focus = FocusSettings.FocusPriorityScale * FocusPriorityScale;
+	float NonFocus = FocusSettings.NonFocusPriorityScale * NonFocusPriorityScale;
+	float Result = FMath::Lerp(Focus, NonFocus, FocusFactor);
+	return FMath::Max(0.0f, Result);
+}
+
+float FAttenuationSettings::GetFocusAttenuation(const struct FGlobalFocusSettings& FocusSettings, float FocusFactor) const
+{
+	float Focus = FocusSettings.FocusVolumeScale * FocusVolumeAttenuation;
+	float NonFocus = FocusSettings.NonFocusVolumeScale * NonFocusVolumeAttenuation;
+	float Result = FMath::Lerp(Focus, NonFocus, FocusFactor);
+	return FMath::Max(0.0f, Result);
+}
+
+float FAttenuationSettings::GetFocusDistanceScale(const struct FGlobalFocusSettings& FocusSettings, float FocusFactor) const
+{
+	float Focus = FocusSettings.FocusDistanceScale * FocusDistanceScale;
+	float NonFocus = FocusSettings.NonFocusDistanceScale * NonFocusDistanceScale;
+	float Result = FMath::Lerp(Focus, NonFocus, FocusFactor);
+	return FMath::Max(0.0f, Result);
+}
+
 /**
  * Calculate the attenuation value.
  * @param DistanceModel - Which math model of attenuation is used
@@ -85,7 +109,9 @@ float FAttenuationSettings::AttenuationEval(const float Distance, const float Fa
 	// is greater than the falloff value, it'll use the algorithm/curve value evaluated at Falloff distance,
 	// which could be 1.0 (and not 0.0f).
 
-	float DistanceCopy = FMath::Clamp(Distance, 0.0f, Falloff);
+	float FalloffCopy = FMath::Max(Falloff, 1.0f);
+	float DistanceCopy = FMath::Clamp(Distance, 0.0f, FalloffCopy);
+
 	DistanceCopy *= DistanceScale;
 
 	float Result = 0.0f;
@@ -93,34 +119,34 @@ float FAttenuationSettings::AttenuationEval(const float Distance, const float Fa
 	{
 		case ATTENUATION_Linear:
 
-			Result = (1.0f - (DistanceCopy / Falloff));
+			Result = (1.0f - (DistanceCopy / FalloffCopy));
 			break;
 
 		case ATTENUATION_Logarithmic:
 
-			Result = 0.5f * -FMath::Loge(DistanceCopy / Falloff);
+			Result = 0.5f * -FMath::Loge(DistanceCopy / FalloffCopy);
 			break;
 
 		case ATTENUATION_Inverse:
 
-			Result = 0.02f / (DistanceCopy / Falloff);
+			Result = 0.02f / (DistanceCopy / FalloffCopy);
 			break;
 
 		case ATTENUATION_LogReverse:
 
-			Result = 1.0f + 0.5f * FMath::Loge(1.0f - (DistanceCopy / Falloff));
+			Result = 1.0f + 0.5f * FMath::Loge(1.0f - (DistanceCopy / FalloffCopy));
 			break;
 
 		case ATTENUATION_NaturalSound:
 		{
 			check( dBAttenuationAtMax <= 0.0f );
-			Result = FMath::Pow(10.0f, ((DistanceCopy / Falloff) * dBAttenuationAtMax) / 20.0f);
+			Result = FMath::Pow(10.0f, ((DistanceCopy / FalloffCopy) * dBAttenuationAtMax) / 20.0f);
 			break;
 		}
 
 		case ATTENUATION_Custom:
 
-			Result = CustomAttenuationCurve.GetRichCurveConst()->Eval(DistanceCopy / Falloff);
+			Result = CustomAttenuationCurve.GetRichCurveConst()->Eval(DistanceCopy / FalloffCopy);
 			break;
 
 		default:
@@ -210,6 +236,7 @@ bool FAttenuationSettings::operator==(const FAttenuationSettings& Other) const
 			&& FocusDistanceScale		== Other.FocusDistanceScale
 			&& FocusPriorityScale		== Other.FocusPriorityScale
 			&& NonFocusPriorityScale	== Other.NonFocusPriorityScale
+			&& FocusVolumeAttenuation	== Other.FocusVolumeAttenuation
 			&& NonFocusVolumeAttenuation == Other.NonFocusVolumeAttenuation);
 }
 
diff --git a/Engine/Source/Runtime/Engine/Private/SoundBase.cpp b/Engine/Source/Runtime/Engine/Private/SoundBase.cpp
--- a/Engine/Source/Runtime/Engine/Private/SoundBase.cpp
+++ b/Engine/Source/Runtime/Engine/Private/SoundBase.cpp
@@ -49,54 +49,6 @@ float USoundBase::GetMaxAudibleDistance()
 	return 0.f;
 }
 
-bool USoundBase::IsAudibleSimple(class FAudioDevice* AudioDevice, const FVector Location, USoundAttenuation* InAttenuationSettings)
-{
-	// No audio device means no listeners to check against
-	if (!AudioDevice)
-	{
-		return false;
-	}
-
-	// Listener position could change before long sounds finish
-	if (GetDuration() > 1.0f)
-	{
-		return true;
-	}
-
-	// Is this SourceActor within the MaxAudibleDistance of any of the listeners?
-	float MaxAudibleDistance = InAttenuationSettings != nullptr ? InAttenuationSettings->Attenuation.GetMaxDimension() : GetMaxAudibleDistance();
-
-	return AudioDevice->LocationIsAudible(Location, MaxAudibleDistance);
-}
-
-bool USoundBase::IsAudible( const FVector &SourceLocation, const FVector &ListenerLocation, AActor* SourceActor, bool& bIsOccluded, bool bCheckOcclusion )
-{
-	//@fixme - naive implementation, needs to be optimized
-	// for now, check max audible distance, and also if looping
-	check( SourceActor );
-	// Account for any portals
-	const FVector ModifiedSourceLocation = SourceLocation;
-
-	float MaxDist = GetMaxAudibleDistance();
-
-	if( MaxDist * MaxDist >= ( ListenerLocation - ModifiedSourceLocation ).SizeSquared() )
-	{
-		// Can't line check through portals
-		if( bCheckOcclusion && ( MaxDist != WORLD_MAX ) && ( ModifiedSourceLocation == SourceLocation ) )
-		{
-			static FName NAME_IsAudible(TEXT("IsAudible"));
-
-			// simple trace occlusion check - reduce max audible distance if occluded
-			bIsOccluded = SourceActor->GetWorld()->LineTraceTestByChannel(ModifiedSourceLocation, ListenerLocation, ECC_Visibility, FCollisionQueryParams(NAME_IsAudible, true, SourceActor));
-		}
-		return true;
-	}
-	else
-	{
-		return false;
-	}
-}
-
 float USoundBase::GetDuration()
 {
 	return Duration;
@@ -112,6 +64,12 @@ float USoundBase::GetPitchMultiplier()
 	return 1.f;
 }
 
+bool USoundBase::IsLooping()
+{ 
+	return (GetDuration() >= INDEFINITELY_LOOPING_DURATION); 
+}
+
+
 USoundClass* USoundBase::GetSoundClass() const
 {
 	return SoundClassObject;
