Copying //UE4/Dev-Framework to //UE4/Main @ 2830052

diff --git a/Engine/Source/Runtime/Engine/Private/SoundConcurrency.cpp b/Engine/Source/Runtime/Engine/Private/SoundConcurrency.cpp
--- a/Engine/Source/Runtime/Engine/Private/SoundConcurrency.cpp
+++ b/Engine/Source/Runtime/Engine/Private/SoundConcurrency.cpp
@@ -13,9 +13,6 @@
 USoundConcurrency::USoundConcurrency(const FObjectInitializer& ObjectInitializer)
 	: Super(ObjectInitializer)
 {
-	Concurrency.MaxCount = 16;
-	Concurrency.ResolutionRule = EMaxConcurrentResolutionRule::PreventNew;
-	Concurrency.VolumeScale = 1.0f;
 }
 
 /************************************************************************/
@@ -23,9 +20,10 @@ USoundConcurrency::USoundConcurrency(const FObjectInitializer& ObjectInitializer
 /************************************************************************/
 
 FConcurrencyGroup::FConcurrencyGroup()
-	: MaxActiveSounds(0)
+	: MaxActiveSounds(16)
 	, ConcurrencyGroupID(0)
-	, ResolutionRule(EMaxConcurrentResolutionRule::PreventNew)
+	, ResolutionRule(EMaxConcurrentResolutionRule::StopFarthestThenPreventNew)
+	, Generation(0)
 {
 }
 
@@ -40,6 +38,7 @@ FConcurrencyGroup::FConcurrencyGroup(FActiveSound* ActiveSound)
 
 	MaxActiveSounds = ConcurrencySettings->MaxCount;
 	ResolutionRule = ConcurrencySettings->ResolutionRule;
+	Generation = 0;
 
 	ActiveSounds.Add(ActiveSound);
 	ActiveSound->ConcurrencyGroupID = ConcurrencyGroupID;
@@ -50,6 +49,7 @@ FConcurrencyGroup::FConcurrencyGroup(const FConcurrencyGroup& Other)
 	, MaxActiveSounds(Other.MaxActiveSounds)
 	, ConcurrencyGroupID(Other.ConcurrencyGroupID)
 	, ResolutionRule(Other.ResolutionRule)
+	, Generation(Other.Generation)
 {
 	check(ConcurrencyGroupID != 0);
 }
@@ -59,6 +59,7 @@ FConcurrencyGroup::FConcurrencyGroup(const FConcurrencyGroup&& Other)
 	, MaxActiveSounds(Other.MaxActiveSounds)
 	, ConcurrencyGroupID(Other.ConcurrencyGroupID)
 	, ResolutionRule(Other.ResolutionRule)
+	, Generation(Other.Generation)
 {
 	check(ConcurrencyGroupID != 0);
 }
@@ -69,6 +70,7 @@ FConcurrencyGroup& FConcurrencyGroup::operator=(const FConcurrencyGroup& Other)
 	MaxActiveSounds = Other.MaxActiveSounds;
 	ConcurrencyGroupID = Other.ConcurrencyGroupID;
 	ResolutionRule = Other.ResolutionRule;
+	Generation = Other.Generation;
 	check(ConcurrencyGroupID != 0);
 	return *this;
 }
@@ -79,6 +81,7 @@ FConcurrencyGroup& FConcurrencyGroup::operator=(const FConcurrencyGroup&& Other)
 	MaxActiveSounds = Other.MaxActiveSounds;
 	ConcurrencyGroupID = Other.ConcurrencyGroupID;
 	ResolutionRule = Other.ResolutionRule;
+	Generation = Other.Generation;
 	check(ConcurrencyGroupID != 0);
 	return *this;
 }
@@ -93,6 +96,7 @@ void FConcurrencyGroup::AddActiveSound(FActiveSound* ActiveSound)
 	check(ConcurrencyGroupID != 0);
 	ActiveSounds.Add(ActiveSound);
 	ActiveSound->ConcurrencyGroupID = ConcurrencyGroupID;
+	ActiveSound->ConcurrencyGeneration = Generation++;
 }
 
 void FConcurrencyGroup::StopQuietSoundsDueToMaxConcurrency()
@@ -236,14 +240,14 @@ FActiveSound* FSoundConcurrencyManager::ResolveConcurrency(const FActiveSound& N
 				// Find the current lowest priority sound.
 				for (FActiveSound* CurrSound : ActiveSounds)
 				{
-					if (SoundToStop == nullptr || CurrSound->Priority < SoundToStop->Priority)
+					if (SoundToStop == nullptr || CurrSound->GetPriority() < SoundToStop->GetPriority())
 					{
 						SoundToStop = CurrSound;
 					}
 				}
 
 				// Only stop any sounds if the *lowest* priority is lower than the incoming NewActiveSound
-				if (SoundToStop->Priority >= NewActiveSound.Priority)
+				if (SoundToStop->GetPriority() >= NewActiveSound.GetPriority())
 				{
 					SoundToStop = nullptr;
 				}
@@ -284,18 +288,27 @@ FActiveSound* FSoundConcurrencyManager::ResolveConcurrency(const FActiveSound& N
 		FActiveSound* OutActiveSound = MakeNewActiveSound(NewActiveSound);
 		check(OutActiveSound);
 
-		// And add it to to the concurrency group
-		ConcurrencyGroup.AddActiveSound(OutActiveSound);
-
-		// Factor in the volume-scale due to sound instance concurrency count
-		// E.g. if 3 sounds are playing in the same group, and its set to be 0.9 VolumeScale, then the 
-		// 4th sound in the concurrency group is going to be attenuated by 0.9^4
+		// If we're ducking older sounds in the concurrency group, then loop through each sound in the concurrency group
+		// and update their duck amount based on each sound's generation and the next generation count. The older the sound, the more ducking.
 		if (Concurrency->VolumeScale < 1.0f)
 		{
-			int32 NumActiveSoundsInConcurrencyGroup = ConcurrencyGroup.GetActiveSounds().Num();
-			OutActiveSound->ConcurrencyVolumeScale = FMath::Pow(Concurrency->VolumeScale, (float)NumActiveSoundsInConcurrencyGroup);
+			check(Concurrency->VolumeScale >= 0.0f);
+
+			ActiveSounds = ConcurrencyGroup.GetActiveSounds();
+
+			int32 NextGeneration = ConcurrencyGroup.GetGeneration() + 1;
+
+			for (FActiveSound* ActiveSound : ActiveSounds)
+			{
+				const int32 ActiveSoundGeneration = ActiveSound->ConcurrencyGeneration;
+				const float GenerationDelta = (float)(NextGeneration - ActiveSoundGeneration);
+				ActiveSound->ConcurrencyVolumeScale = FMath::Pow(Concurrency->VolumeScale, GenerationDelta);
+			}
 		}
 
+		// And add it to to the concurrency group. This automatically updates generation counts.
+		ConcurrencyGroup.AddActiveSound(OutActiveSound);
+
 		// Stop any sounds now if needed
 		if (SoundToStop)
 		{
diff --git a/Engine/Source/Runtime/Engine/Private/StreamableManager.cpp b/Engine/Source/Runtime/Engine/Private/StreamableManager.cpp
--- a/Engine/Source/Runtime/Engine/Private/StreamableManager.cpp
+++ b/Engine/Source/Runtime/Engine/Private/StreamableManager.cpp
@@ -249,13 +249,13 @@ void FStreamableManager::RequestAsyncLoad(const TArray<FStringAssetReference>& T
 	TSharedRef<FStreamableRequest> NewRequest = MakeShareable(new FStreamableRequest());
 	NewRequest->CompletionDelegate = DelegateToCall;
 
-	TArray<FStreamable *> ExistingStreamables;
-	ExistingStreamables.SetNum(TargetsToStream.Num());
+	static TArray<FStreamable*> ExistingStreamables;
+	ExistingStreamables.Reserve(TargetsToStream.Num());
 
 	for (int32 i = 0; i < TargetsToStream.Num(); i++)
 	{
 		FStreamable* Existing = StreamInternal(TargetsToStream[i], Priority);
-		ExistingStreamables[i] = Existing;
+		ExistingStreamables.Add(Existing);
 
 		if (Existing)
 		{
@@ -273,6 +273,8 @@ void FStreamableManager::RequestAsyncLoad(const TArray<FStringAssetReference>& T
 			CheckCompletedRequests(TargetsToStream[i], Existing);
 		}
 	}
+
+	ExistingStreamables.Reset();
 }
 
 void FStreamableManager::RequestAsyncLoad(const FStringAssetReference& TargetToStream, FStreamableDelegate DelegateToCall, TAsyncLoadPriority Priority)
diff --git a/Engine/Source/Runtime/Engine/Private/UnrealEngine.cpp b/Engine/Source/Runtime/Engine/Private/UnrealEngine.cpp
--- a/Engine/Source/Runtime/Engine/Private/UnrealEngine.cpp
+++ b/Engine/Source/Runtime/Engine/Private/UnrealEngine.cpp
@@ -7233,7 +7233,7 @@ struct FCompareFSoundInfoByWaveInstNum
 /** draws a property of the given object on the screen similarly to stats */
 static void DrawProperty(UCanvas* CanvasObject, UObject* Obj, const FDebugDisplayProperty& PropData, UProperty* Prop, int32 X, int32& Y)
 {
-#if UE_BUILD_SHIPPING
+#if !UE_BUILD_SHIPPING
 	checkSlow(PropData.bSpecialProperty || Prop != NULL);
 	checkSlow(Prop == NULL || Obj->GetClass()->IsChildOf(Prop->GetOwnerClass()));
 
@@ -12281,7 +12281,7 @@ bool UEngine::ToggleStatSounds(UWorld* World, FCommonViewportClient* ViewportCli
 
 int32 UEngine::RenderStatSounds(UWorld* World, FViewport* Viewport, FCanvas* Canvas, int32 X, int32 Y, const FVector* ViewLocation, const FRotator* ViewRotation)
 {
-#if UE_BUILD_SHIPPING
+#if !UE_BUILD_SHIPPING
 
 #if UE_BUILD_DEBUG
 
diff --git a/Engine/Source/Runtime/Engine/Private/Vehicles/WheeledVehicleMovementComponent.cpp b/Engine/Source/Runtime/Engine/Private/Vehicles/WheeledVehicleMovementComponent.cpp
--- a/Engine/Source/Runtime/Engine/Private/Vehicles/WheeledVehicleMovementComponent.cpp
+++ b/Engine/Source/Runtime/Engine/Private/Vehicles/WheeledVehicleMovementComponent.cpp
@@ -263,7 +263,7 @@ void UWheeledVehicleMovementComponent::SetupVehicleShapes()
 			CollisionResponse.SetAllChannels(ECR_Ignore);
 
 			PxFilterData PWheelQueryFilterData, PDummySimData;
-			CreateShapeFilterData(ECC_Vehicle, FMaskFilter(0), UpdatedComponent->GetUniqueID(), CollisionResponse, 0, 0, PWheelQueryFilterData, PDummySimData, false, false, false);
+			CreateShapeFilterData(ECC_Vehicle, FMaskFilter(0), UpdatedComponent->GetOwner()->GetUniqueID(), CollisionResponse, UpdatedComponent->GetUniqueID(), 0, PWheelQueryFilterData, PDummySimData, false, false, false);
 
 			//// Give suspension raycasts the same group ID as the chassis so that they don't hit each other
 			PWheelShape->setQueryFilterData(PWheelQueryFilterData);
diff --git a/Engine/Source/Runtime/Engine/Private/VisualLogger/VisualLogger.cpp b/Engine/Source/Runtime/Engine/Private/VisualLogger/VisualLogger.cpp
--- a/Engine/Source/Runtime/Engine/Private/VisualLogger/VisualLogger.cpp
+++ b/Engine/Source/Runtime/Engine/Private/VisualLogger/VisualLogger.cpp
@@ -73,6 +73,24 @@ bool FVisualLogger::CheckVisualLogInputInternal(const UObject* Object, const FLo
 	return true;
 }
 
+void FVisualLogger::AddWhitelistedClass(UClass& InClass)
+{
+	ClassWhitelist.AddUnique(&InClass);
+}
+
+bool FVisualLogger::IsClassWhitelisted(const UClass& InClass) const
+{
+	for (const UClass* WhitelistedClass : ClassWhitelist)
+	{
+		if (InClass.IsChildOf(WhitelistedClass))
+		{
+			return true;
+		}
+	}
+
+	return false;
+}
+
 FVisualLogEntry* FVisualLogger::GetLastEntryForObject(const UObject* Object)
 {
 	UObject * LogOwner = FVisualLogger::FindRedirection(Object);
@@ -83,38 +101,41 @@ FVisualLogEntry* FVisualLogger::GetEntryToWrite(const UObject* Object, float Tim
 {
 	FVisualLogEntry* CurrentEntry = nullptr;
 	UObject * LogOwner = FVisualLogger::FindRedirection(Object);
-	if (LogOwner == nullptr || (LogOwner != Object && CurrentEntryPerObject.Contains(LogOwner) == false))
+	if (LogOwner == nullptr)
 	{
 		return nullptr;
 	}
 
 	bool InitializeNewEntry = false;
 
-	const UWorld* World = GetWorldForVisualLogger(Object);
+	const UWorld* World = GetWorldForVisualLogger(LogOwner);
 
 	if (CurrentEntryPerObject.Contains(LogOwner))
 	{
 		CurrentEntry = &CurrentEntryPerObject[LogOwner];
-		InitializeNewEntry = TimeStamp > CurrentEntry->TimeStamp && ShouldCreate == ECreateIfNeeded::Create;
-		if (World && IsInGameThread())
+		if (CurrentEntry->bIsAllowedToLog)
 		{
-			World->GetTimerManager().ClearTimer(VisualLoggerCleanupTimerHandle);
-			for (auto& CurrentPair : CurrentEntryPerObject)
+			InitializeNewEntry = TimeStamp > CurrentEntry->TimeStamp && ShouldCreate == ECreateIfNeeded::Create;
+			if (World && IsInGameThread())
 			{
-				FVisualLogEntry* Entry = &CurrentPair.Value;
-				if (Entry->TimeStamp >= 0 && Entry->TimeStamp < TimeStamp)
+				World->GetTimerManager().ClearTimer(VisualLoggerCleanupTimerHandle);
+				for (auto& CurrentPair : CurrentEntryPerObject)
 				{
-					for (FVisualLogDevice* Device : OutputDevices)
+					FVisualLogEntry* Entry = &CurrentPair.Value;
+					if (Entry->TimeStamp >= 0 && Entry->TimeStamp < TimeStamp)
 					{
-						Device->Serialize(CurrentPair.Key, ObjectToNameMap[CurrentPair.Key], ObjectToClassNameMap[CurrentPair.Key], *Entry);
+						for (FVisualLogDevice* Device : OutputDevices)
+						{
+							Device->Serialize(CurrentPair.Key, ObjectToNameMap[CurrentPair.Key], ObjectToClassNameMap[CurrentPair.Key], *Entry);
+						}
+						Entry->Reset();
 					}
-					Entry->Reset();
 				}
 			}
 		}
 	}
 
-	if (!CurrentEntry)
+	if (CurrentEntry == nullptr)
 	{
 		// It's first and only one usage of LogOwner as regular object to get names. We assume once that LogOwner is correct here and only here.
 		CurrentEntry = &CurrentEntryPerObject.Add(LogOwner);
@@ -122,7 +143,13 @@ FVisualLogEntry* FVisualLogger::GetEntryToWrite(const UObject* Object, float Tim
 		ObjectToClassNameMap.Add(LogOwner, *(LogOwner->GetClass()->GetName()));
 		ObjectToPointerMap.Add(LogOwner, LogOwner);
 		ObjectToWorldMap.Add(LogOwner, World);
-		InitializeNewEntry = true;
+		
+		// IsClassWhitelisted isn't super fast, but this gets calculated only once for every 
+		// object trying to log something
+		const bool bCanBeLogged = ClassWhitelist.Num() == 0 || IsClassWhitelisted(*LogOwner->GetClass()) 
+			|| IsClassWhitelisted(*Object->GetClass());
+		InitializeNewEntry = bCanBeLogged;
+		CurrentEntry->bIsAllowedToLog = bCanBeLogged;
 	}
 
 	if (InitializeNewEntry)
@@ -155,11 +182,11 @@ FVisualLogEntry* FVisualLogger::GetEntryToWrite(const UObject* Object, float Tim
 		}
 		else
 		{
-			const AActor* ObjectAsActor = Cast<AActor>(Object);
+			const AActor* ObjectAsActor = Cast<AActor>(LogOwner);
 			if (ObjectAsActor)
 			{
 				CurrentEntry->Location = ObjectAsActor->GetActorLocation();
-				const IVisualLoggerDebugSnapshotInterface* DebugSnapshotInterface = Cast<const IVisualLoggerDebugSnapshotInterface>(Object);
+				const IVisualLoggerDebugSnapshotInterface* DebugSnapshotInterface = Cast<const IVisualLoggerDebugSnapshotInterface>(LogOwner);
 				if (DebugSnapshotInterface)
 				{
 					DebugSnapshotInterface->GrabDebugSnapshot(CurrentEntry);
@@ -168,28 +195,7 @@ FVisualLogEntry* FVisualLogger::GetEntryToWrite(const UObject* Object, float Tim
 		}
 	}
 
-	if (World && IsInGameThread())
-	{
-		//set next tick timer to flush obsolete/old entries
-		World->GetTimerManager().SetTimer(VisualLoggerCleanupTimerHandle, FTimerDelegate::CreateLambda(
-			[this, World](){
-			for (auto& CurrentPair : CurrentEntryPerObject)
-			{
-				FVisualLogEntry* Entry = &CurrentPair.Value;
-				if (Entry->TimeStamp >= 0 && (!World || Entry->TimeStamp < World->GetTimeSeconds())) // CurrentEntry->TimeStamp == -1 means it's not initialized entry information
-				{
-					for (FVisualLogDevice* Device : OutputDevices)
-					{
-						Device->Serialize(CurrentPair.Key, ObjectToNameMap[CurrentPair.Key], ObjectToClassNameMap[CurrentPair.Key], *Entry);
-					}
-					Entry->Reset();
-				}
-			}
-		}
-		), 0.1, false);
-	}
-
-	return CurrentEntry;
+	return CurrentEntry && CurrentEntry->bIsAllowedToLog ? CurrentEntry : nullptr;
 }
 
 
@@ -283,9 +289,9 @@ void FVisualLogger::EventLog(const UObject* Object, const FName EventTag1, const
 void FVisualLogger::NavigationDataDump(const UObject* Object, const FLogCategoryBase& Category, ELogVerbosity::Type Verbosity, int32 UniqueLogId, const FBox& Box)
 {
 	SCOPE_CYCLE_COUNTER(STAT_VisualLog);
-	UWorld *World = nullptr;
-	FVisualLogEntry *CurrentEntry = nullptr;
-	if (CheckVisualLogInputInternal(Object, Category, Verbosity, &World, &CurrentEntry) == false)
+	UWorld* World = nullptr;
+	FVisualLogEntry* CurrentEntry = nullptr;
+	if (CheckVisualLogInputInternal(Object, Category, Verbosity, &World, &CurrentEntry) == false || CurrentEntry == nullptr)
 	{
 		return;
 	}
@@ -332,6 +338,7 @@ void FVisualLogger::Cleanup(UWorld* OldWorld, bool bReleaseMemory)
 		SetIsRecordingToFile(false);
 	}
 
+	Flush();
 	for (FVisualLogDevice* Device : FVisualLogger::Get().OutputDevices)
 	{
 		Device->Cleanup(bReleaseMemory);
@@ -341,16 +348,28 @@ void FVisualLogger::Cleanup(UWorld* OldWorld, bool bReleaseMemory)
 	{
 		WorldToRedirectionMap.Remove(OldWorld);
 
-		for (auto It = ObjectToWorldMap.CreateIterator(); It; ++It)
+		if (WorldToRedirectionMap.Num() == 0)
 		{
-			if (It.Value() == OldWorld)
+			WorldToRedirectionMap.Reset();
+			ObjectToWorldMap.Reset();
+			CurrentEntryPerObject.Reset();
+			ObjectToNameMap.Reset();
+			ObjectToClassNameMap.Reset();
+			ObjectToPointerMap.Reset();
+		}
+		else
+		{
+			for (auto It = ObjectToWorldMap.CreateIterator(); It; ++It)
 			{
-				const UObject* Obj = It.Key();
-				ObjectToWorldMap.Remove(Obj);
-				CurrentEntryPerObject.Remove(Obj);
-				ObjectToNameMap.Remove(Obj);
-				ObjectToClassNameMap.Remove(Obj);
-				ObjectToPointerMap.Remove(Obj);
+				if (It.Value() == OldWorld)
+				{
+					const UObject* Obj = It.Key();
+					ObjectToWorldMap.Remove(Obj);
+					CurrentEntryPerObject.Remove(Obj);
+					ObjectToNameMap.Remove(Obj);
+					ObjectToClassNameMap.Remove(Obj);
+					ObjectToPointerMap.Remove(Obj);
+				}
 			}
 		}
 	}
@@ -540,7 +559,7 @@ public:
 				{
 					FString Category = FParse::Token(Cmd, 1);
 					FVisualLogger::Get().BlockAllCategories(true);
-					FVisualLogger::Get().AddCategortyToWhiteList(*Category);
+					FVisualLogger::Get().AddCategoryToWhitelist(*Category);
 					return true;
 				}
 #if WITH_EDITOR
diff --git a/Engine/Source/Runtime/Engine/Private/World.cpp b/Engine/Source/Runtime/Engine/Private/World.cpp
--- a/Engine/Source/Runtime/Engine/Private/World.cpp
+++ b/Engine/Source/Runtime/Engine/Private/World.cpp
@@ -4844,14 +4844,8 @@ UWorld* FSeamlessTravelHandler::Tick()
 			FUObjectAnnotationSparseBool KeepAnnotation;
 			TArray<AActor*> KeepActors;
 
-			AGameMode* AuthGameMode = CurrentWorld->GetAuthGameMode();
-			if (AuthGameMode)
+			if (AGameMode* AuthGameMode = CurrentWorld->GetAuthGameMode())
 			{
-				// keep GameMode if traveling to transition map
-				if (!bSwitchedToDefaultMap)
-				{
-					KeepAnnotation.Set(CurrentWorld->GetAuthGameMode());
-				}
 				AuthGameMode->GetSeamlessTravelActorList(!bSwitchedToDefaultMap, KeepActors);
 			}
 
diff --git a/Engine/Source/Runtime/Engine/Public/ActiveSound.h b/Engine/Source/Runtime/Engine/Public/ActiveSound.h
--- a/Engine/Source/Runtime/Engine/Public/ActiveSound.h
+++ b/Engine/Source/Runtime/Engine/Public/ActiveSound.h
@@ -125,12 +125,18 @@ public:
 	/** The group of active concurrent sounds that this sound is playing in. */
 	FConcurrencyGroupID ConcurrencyGroupID;
 
+	/** The generation of this sound in the concurrency group. */
+	int32 ConcurrencyGeneration;
+
 	/** Optional USoundConcurrency to override sound */
 	USoundConcurrency* ConcurrencySettings;
 
 	/** Optional SoundClass to override Sound */
 	USoundClass* SoundClassOverride;
 
+	/** Whether or not the sound has checked if it was occluded already. Used to initialize a sound as occluded and bypassing occlusion interpolation. */
+	uint32 bHasCheckedOcclusion:1;
+
 	/** whether we were occluded the last time we checked */
 	uint32 bIsOccluded:1;
 
@@ -218,14 +224,23 @@ public:
 
 	FDynamicParameter CurrentOcclusionVolumeAttenuation;
 
-	/** A volume scale to apply to a sound based on the concurrency count of the active sound when it started */
+	/** A volume scale to apply to a sound based on the concurrency count of the active sound when it started. Will reduce volume of new sounds if many sounds are playing in concurrency group. */
 	float ConcurrencyVolumeScale;
 
+	/** A volume to apply to a sound based on the concurrency generation and the current generation count. Will reduce volume of older sounds as new sounds play in concurrency group. */
+	float ConcurrencyDuckingVolumeScale;
+
 	float SubtitlePriority;
 
 	/** The product of the component priority and the USoundBase priority */
 	float Priority;
 
+	/** The amount priority is scaled due to focus */
+	float FocusPriorityScale;
+
+	/** The amount distance is scaled due to focus */
+	float FocusDistanceScale;
+
 	/** Frequency with which to check for occlusion from its closest listener */
 	// The volume used to determine concurrency resolution for "quietest" active sound
 	float VolumeConcurrency;
@@ -352,6 +367,9 @@ public:
 	/** Applies the active sound's attenuation settings to the input parse params using the given listener */
 	void ApplyAttenuation(FSoundParseParameters& ParseParams, const FListener& Listener, const FAttenuationSettings* SettingsAttenuationNode = nullptr);
 
+	/** Returns the effective priority of the active sound */
+	float GetPriority() const { return Priority * FocusPriorityScale; }
+
 private:
 	
 	/** Whether or not this sound is audible */
@@ -383,22 +401,4 @@ private:
 	/** Apply the interior settings to the ambient sound as appropriate */
 	void HandleInteriorVolumes( const FListener& Listener, struct FSoundParseParameters& ParseParams );
 
-	/** Struct used to cache listener attenuation vector math results */
-	struct FAttenuationListenerData
-	{
-		FVector ListenerToSoundDir;
-		float AttenuationDistance;
-		float ListenerToSoundDistance;
-		bool bDataComputed;
-
-		FAttenuationListenerData()
-			: ListenerToSoundDir(FVector::ZeroVector)
-			, AttenuationDistance(0.0f)
-			, ListenerToSoundDistance(0.0f)
-			, bDataComputed(false)
-		{}
-	};
-
-	/** Performs listener-based calculations for spatializing and attenuating the active sound */
-	void GetAttenuationListenerData(FAttenuationListenerData& Data, const FSoundParseParameters& ParseParams, const FListener& Listener) const;
 };
\ No newline at end of file
diff --git a/Engine/Source/Runtime/Engine/Public/Animation/AnimCurveTypes.h b/Engine/Source/Runtime/Engine/Public/Animation/AnimCurveTypes.h
--- a/Engine/Source/Runtime/Engine/Public/Animation/AnimCurveTypes.h
+++ b/Engine/Source/Runtime/Engine/Public/Animation/AnimCurveTypes.h
@@ -235,75 +235,249 @@ struct FCurveElement
 /**
  * This struct is used to create curve snap shot of current time when extracted
  */
-struct ENGINE_API FBlendedCurve
+template <typename InAllocator>
+struct FBaseBlendedCurve
 {
+	typedef InAllocator   Allocator;
 	/**
 	* List of curve elements for this pose
 	*/
-	TArray<FCurveElement> Elements;
+	TArray<FCurveElement, Allocator> Elements;
 
 	/**
 	* List of SmartName UIDs, retrieved from AnimInstanceProxy (which keeps authority)
 	*/
 	TArray<FSmartNameMapping::UID> const * UIDList;
-	
+
+
 	/**
 	 * constructor
 	 */
-	FBlendedCurve()
+	FBaseBlendedCurve()
 		: UIDList(nullptr)
 		, bInitialized(false)
 	{
 	}
-	
+
 	/** Initialize Curve Data from following data */
 	DEPRECATED(4.11, "Use new InitFrom(TArray<FSmartNameMapping::UID>* InSmartNameUIDs) signature")
-	void InitFrom(const class USkeleton* Skeleton);
-	void InitFrom(TArray<FSmartNameMapping::UID> const * InSmartNameUIDs);
-	void InitFrom(const FBlendedCurve& InCurveToInitFrom);
+	void InitFrom(const class USkeleton* Skeleton)
+	{
+		InitFrom((TArray<FSmartNameMapping::UID>*)&(const_cast<USkeleton*>(Skeleton)->GetCachedAnimCurveMappingNameUids()));
+	}
+
+	void InitFrom(TArray<FSmartNameMapping::UID> const * InSmartNameUIDs)
+	{
+		check(InSmartNameUIDs != nullptr);
+		UIDList = InSmartNameUIDs;
+		Elements.Reset();
+		Elements.AddZeroed(UIDList->Num());
+
+		// no name, means no curve
+		bInitialized = true;
+	}
+
+	template <typename OtherAllocator>
+	void InitFrom(const FBaseBlendedCurve<OtherAllocator>& InCurveToInitFrom)
+	{
+		// make sure this doesn't happen
+		check(InCurveToInitFrom.UIDList != nullptr);
+		UIDList = InCurveToInitFrom.UIDList;
+		Elements.Reset();
+		Elements.AddZeroed(UIDList->Num());
+		
+		bInitialized = true;
+	}
+
+	void InitFrom(const FBaseBlendedCurve<Allocator>& InCurveToInitFrom)
+	{
+		// make sure this doesn't happen
+		if (ensure(&InCurveToInitFrom != this))
+		{
+			check(InCurveToInitFrom.UIDList != nullptr);
+			UIDList = InCurveToInitFrom.UIDList;
+			Elements.Reset();
+			Elements.AddZeroed(UIDList->Num());
+
+			bInitialized = true;
+		}
+	}
 
 	/** Set value of InUID to InValue */
-	void Set(USkeleton::AnimCurveUID InUid, float InValue, int32 InFlags);
+	void Set(USkeleton::AnimCurveUID InUid, float InValue, int32 InFlags)
+	{
+		int32 ArrayIndex;
+
+		check(bInitialized);
+
+		if (UIDList->Find(InUid, ArrayIndex))
+		{
+			Elements[ArrayIndex].Value = InValue;
+			Elements[ArrayIndex].Flags = InFlags;
+		}
+	}
 
 	/**
 	 * Blend (A, B) using Alpha, same as Lerp
 	 */
-	void Blend(const FBlendedCurve& A, const FBlendedCurve& B, float Alpha);
+	//@Todo curve flags won't transfer over - it only overwrites
+	void Blend(const FBaseBlendedCurve& A, const FBaseBlendedCurve& B, float Alpha)
+	{
+		check(A.Num() == B.Num());
+		if (FMath::Abs(Alpha) <= ZERO_ANIMWEIGHT_THRESH)
+		{
+			// if blend is all the way for child1, then just copy its bone atoms
+			Override(A);
+		}
+		else if (FMath::Abs(Alpha - 1.0f) <= ZERO_ANIMWEIGHT_THRESH)
+		{
+			// if blend is all the way for child2, then just copy its bone atoms
+			Override(B);
+		}
+		else
+		{
+			InitFrom(A);
+			for (int32 CurveId = 0; CurveId < A.Elements.Num(); ++CurveId)
+			{
+				Elements[CurveId].Value = FMath::Lerp(A.Elements[CurveId].Value, B.Elements[CurveId].Value, Alpha);
+				Elements[CurveId].Flags = (A.Elements[CurveId].Flags) | (B.Elements[CurveId].Flags);
+			}
+		}
+	}
+
 	/**
 	 * Blend with Other using Alpha, same as Lerp 
 	 */
-	void BlendWith(const FBlendedCurve& Other, float Alpha);
+	void BlendWith(const FBaseBlendedCurve& Other, float Alpha)
+	{
+		check(Num() == Other.Num());
+		if (FMath::Abs(Alpha) <= ZERO_ANIMWEIGHT_THRESH)
+		{
+			return;
+		}
+		else if (FMath::Abs(Alpha - 1.0f) <= ZERO_ANIMWEIGHT_THRESH)
+		{
+			// if blend is all the way for child2, then just copy its bone atoms
+			Override(Other);
+		}
+		else
+		{
+			for (int32 CurveId = 0; CurveId < Elements.Num(); ++CurveId)
+			{
+				Elements[CurveId].Value = FMath::Lerp(Elements[CurveId].Value, Other.Elements[CurveId].Value, Alpha);
+				Elements[CurveId].Flags |= (Other.Elements[CurveId].Flags);
+			}
+		}
+	}
 	/**
 	 * Convert current curves to Additive (this - BaseCurve) if same found
 	 */
-	void ConvertToAdditive(const FBlendedCurve& BaseCurve);
+	void ConvertToAdditive(const FBaseBlendedCurve& BaseCurve)
+	{
+		check(bInitialized);
+		check(Num() == BaseCurve.Num());
+
+		for (int32 CurveId = 0; CurveId < Elements.Num(); ++CurveId)
+		{
+			Elements[CurveId].Value -= BaseCurve.Elements[CurveId].Value;
+			Elements[CurveId].Flags |= BaseCurve.Elements[CurveId].Flags;
+		}
+	}
 	/**
 	 * Accumulate the input curve with input Weight
 	 */
-	void Accumulate(const FBlendedCurve& AdditiveCurve, float Weight);
+	void Accumulate(const FBaseBlendedCurve& AdditiveCurve, float Weight)
+	{
+		check(bInitialized);
+		check(Num() == AdditiveCurve.Num());
+
+		if (Weight > ZERO_ANIMWEIGHT_THRESH)
+		{
+			for (int32 CurveId = 0; CurveId < Elements.Num(); ++CurveId)
+			{
+				Elements[CurveId].Value += AdditiveCurve.Elements[CurveId].Value * Weight;
+				Elements[CurveId].Flags |= AdditiveCurve.Elements[CurveId].Flags;
+			}
+		}
+	}
 
 	/**
 	 * This doesn't blend but combine MAX(current weight, curvetocombine weight)
 	 */
-	void Combine(const FBlendedCurve& CurveToCombine);
+	void Combine(const FBaseBlendedCurve& CurveToCombine)
+	{
+		check(bInitialized);
+		check(Num() == CurveToCombine.Num());
+
+		for (int32 CurveId = 0; CurveId < CurveToCombine.Elements.Num(); ++CurveId)
+		{
+			// if target value is non zero, we accpet target's value
+			// originally this code was doing max, but that doesn't make sense since the values can be negative
+			// we could try to pick non-zero, but if target value is non-zero, I think we should accept that value 
+			// if source is non zero, it will be overriden
+			if (CurveToCombine.Elements[CurveId].Value != 0.f)
+			{
+				Elements[CurveId].Value = CurveToCombine.Elements[CurveId].Value;
+			}
+
+			Elements[CurveId].Flags |= CurveToCombine.Elements[CurveId].Flags;
+		}
+	}
 
 	/**
 	 * Override with inupt curve * weight
 	 */
-	void Override(const FBlendedCurve& CurveToOverrideFrom, float Weight);
+	void Override(const FBaseBlendedCurve& CurveToOverrideFrom, float Weight)
+	{
+		InitFrom(CurveToOverrideFrom);
+
+		if (FMath::IsNearlyEqual(Weight, 1.f))
+		{
+			Override(CurveToOverrideFrom);
+		}
+		else
+		{
+			for (int32 CurveId = 0; CurveId < CurveToOverrideFrom.Elements.Num(); ++CurveId)
+			{
+				Elements[CurveId].Value = CurveToOverrideFrom.Elements[CurveId].Value * Weight;
+				Elements[CurveId].Flags |= CurveToOverrideFrom.Elements[CurveId].Flags;
+			}
+		}
+	}
+
 	/**
 	 * Override with inupt curve 
 	 */
-	void Override(const FBlendedCurve& CurveToOverrideFrom);
+	void Override(const FBaseBlendedCurve& CurveToOverrideFrom)
+	{
+		InitFrom(CurveToOverrideFrom);
+		Elements.Reset();
+		Elements.Append(CurveToOverrideFrom.Elements);
+	}
 
 	/** Return number of elements */
 	int32 Num() const { return Elements.Num(); }
 
 	/** CopyFrom as expected. */
-	void CopyFrom(const FBlendedCurve& CurveToCopyFrom);
-	/** CopyFrom/MoveFrom as expected. Once moved, it is invalid */
-	void MoveFrom(FBlendedCurve& CurveToMoveFrom);
+	template <typename OtherAllocator>
+	void CopyFrom(const FBaseBlendedCurve<OtherAllocator>& CurveToCopyFrom)
+	{
+		UIDList = CurveToCopyFrom.UIDList;
+		Elements.Reset();
+		Elements.Append(CurveToCopyFrom.Elements);
+		bInitialized = true;
+	}
 
+	void CopyFrom(const FBaseBlendedCurve<Allocator>& CurveToCopyFrom)
+	{
+		if (&CurveToCopyFrom != this)
+		{
+			UIDList = CurveToCopyFrom.UIDList;
+			Elements.Reset();
+			Elements.Append(CurveToCopyFrom.Elements);
+			bInitialized = true;
+		}
+	}
 	/** Empty */
 	void Empty()
 	{
@@ -312,12 +486,37 @@ struct ENGINE_API FBlendedCurve
 		Elements.Reset();
 		bInitialized = false;
 	}
-private:
+
 	/**  Whether initialized or not */
 	bool bInitialized;
 	/** Empty and allocate Count number */
-	void Reset(int32 Count);
+	void Reset(int32 Count)
+	{
+		Elements.Reset();
+		Elements.Reserve(Count);
+	}
+};
+
+struct FBlendedHeapCurve;
+
+struct ENGINE_API FBlendedCurve : public FBaseBlendedCurve<FAnimStackAllocator>
+{
+};
+
+struct ENGINE_API FBlendedHeapCurve : public FBaseBlendedCurve<FDefaultAllocator>
+{
+	/** Once moved, source is invalid */
+	void MoveFrom(FBlendedHeapCurve& CurveToMoveFrom)
+	{
+    	UIDList = CurveToMoveFrom.UIDList;
+		CurveToMoveFrom.UIDList = nullptr;
+		Elements = MoveTemp(CurveToMoveFrom.Elements);
+		bInitialized = true;
+		CurveToMoveFrom.bInitialized = false;
+	}
+
 };
+
 /**
  * Raw Curve data for serialization
  */
diff --git a/Engine/Source/Runtime/Engine/Public/Animation/AnimTypes.h b/Engine/Source/Runtime/Engine/Public/Animation/AnimTypes.h
--- a/Engine/Source/Runtime/Engine/Public/Animation/AnimTypes.h
+++ b/Engine/Source/Runtime/Engine/Public/Animation/AnimTypes.h
@@ -331,3 +331,6 @@ struct FMarkerSyncData
 	void CollectUniqueNames();
 	void CollectMarkersInRange(float PrevPosition, float NewPosition, TArray<FPassedMarker>& OutMarkersPassedThisTick, float TotalDeltaMove);
 };
+
+// Shortcut for the allocator used by animation nodes.
+class FAnimStackAllocator: public TMemStackAllocator<>{};
diff --git a/Engine/Source/Runtime/Engine/Public/AudioDevice.h b/Engine/Source/Runtime/Engine/Public/AudioDevice.h
--- a/Engine/Source/Runtime/Engine/Public/AudioDevice.h
+++ b/Engine/Source/Runtime/Engine/Public/AudioDevice.h
@@ -163,6 +163,65 @@ struct FActivatedReverb
 	}
 };
 
+/** Struct used to cache listener attenuation vector math results */
+struct FAttenuationListenerData
+{
+	FVector ListenerToSoundDir;
+	const FListener* Listener;
+	float AttenuationDistance;
+	float ListenerToSoundDistance;
+	bool bDataComputed;
+
+	FAttenuationListenerData()
+		: ListenerToSoundDir(FVector::ZeroVector)
+		, Listener(nullptr)
+		, AttenuationDistance(0.0f)
+		, ListenerToSoundDistance(0.0f)
+		, bDataComputed(false)
+	{}
+};
+
+struct FAttenuationFocusData
+{
+	float FocusFactor;
+	float DistanceScale;
+	float PriorityScale;
+	float VolumeScale;
+
+	FAttenuationFocusData()
+		: FocusFactor(1.0f)
+		, DistanceScale(1.0f)
+		, PriorityScale(1.0f)
+		, VolumeScale(1.0f)
+	{}
+};
+
+/*
+* Setting for global focus scaling
+*/
+struct FGlobalFocusSettings
+{
+	float FocusAzimuthScale;
+	float NonFocusAzimuthScale;
+	float FocusDistanceScale;
+	float NonFocusDistanceScale;
+	float FocusVolumeScale;
+	float NonFocusVolumeScale;
+	float FocusPriorityScale;
+	float NonFocusPriorityScale;
+
+	FGlobalFocusSettings()
+		: FocusAzimuthScale(1.0f)
+		, NonFocusAzimuthScale(1.0f)
+		, FocusDistanceScale(1.0f)
+		, NonFocusDistanceScale(1.0f)
+		, FocusVolumeScale(1.0f)
+		, NonFocusVolumeScale(1.0f)
+		, FocusPriorityScale(1.0f)
+		, NonFocusPriorityScale(1.0f)
+	{}
+};
+
 class ENGINE_API FAudioDevice : public FExec
 {
 public:
@@ -550,6 +609,31 @@ public:
 
 	bool IsAudioDeviceMuted() const;
 
+	/** Computes and returns some geometry related to the listener and the given sound transform. */
+	void GetAttenuationListenerData(FAttenuationListenerData& OutListenerData, const FTransform& SoundTransform, const FAttenuationSettings& AttenuationSettings, const FListener* InListener = nullptr) const;
+
+	/** Returns the focus factor of a sound based on its position and listener data. */
+	float GetFocusFactor(FAttenuationListenerData& OutListenerData, const class USoundBase* Sound, const FTransform& SoundTransform, const FAttenuationSettings& AttenuationSettings, const FListener* InListener = nullptr) const;
+
+	/** Gets the max distance and focus factor of a sound. */
+	void GetMaxDistanceAndFocusFactor(USoundBase* Sound, const class UWorld* World, const FVector& Location, const FAttenuationSettings* AttenuationSettingsToApply, float *OutMaxDistance, float* OutFocusFactor);
+
+	/**
+	* Checks if the given sound would be audible.
+	* @param Sound					The sound to check if it would be audible
+	* @param World					The world the sound is playing in
+	* @param Location				The location the sound is playing in the world
+	* @param AttenuationSettings	The (optional) attenuation settings the sound is using
+	* @param MaxDistance			The computed max distance of the sound.
+	* @param FocusFactor			The focus factor of the sound.
+	* @param Returns true if the sound is audible, false otherwise.
+	*/
+	bool SoundIsAudible(USoundBase* Sound, const class UWorld* World, const FVector& Location, const FAttenuationSettings* AttenuationSettingsToApply, float MaxDistance, float FocusFactor);
+
+	/** Returns the index of the listener closest to the given sound transform */
+	static int32 FindClosestListenerIndex(const FTransform& SoundTransform, const TArray<FListener>& InListeners);
+	int32 FindClosestListenerIndex(const FTransform& SoundTransform) const;
+
 protected:
 	friend class FSoundSource;
 
@@ -809,6 +893,9 @@ public:
 	/** Global dynamic pitch scale parameter */
 	FDynamicParameter GlobalPitchScale;
 
+	/** The global focus settings */
+	FGlobalFocusSettings GlobalFocusSettings;
+
 	/** Timestamp of the last update */
 	double LastUpdateTime;
 
diff --git a/Engine/Source/Runtime/Engine/Public/BonePose.h b/Engine/Source/Runtime/Engine/Public/BonePose.h
--- a/Engine/Source/Runtime/Engine/Public/BonePose.h
+++ b/Engine/Source/Runtime/Engine/Public/BonePose.h
@@ -8,8 +8,7 @@
 #include "CustomBoneIndexArray.h"
 #include "AnimEncoding.h"
 #include "Animation/AnimStats.h"
+#include "Animation/AnimTypes.h"
 
-struct FCompactPose;
-
 struct FBoneTransform
 {
@@ -39,7 +38,7 @@
 };
 
 
-template<class BoneIndexType>
+template<class BoneIndexType, typename InAllocator>
 struct FBasePose
 {
 public:
@@ -85,59 +84,9 @@
 		IterType end() { return Pose.MakeEndIterReverse(); }
 	};
 
+	const TArray<FTransform, InAllocator>& GetBones() const { return Bones; }
 protected:
-	TArray<FTransform> Bones;
+	TArray<FTransform, InAllocator> Bones;
 };
 
-
-struct FMeshPose : public FBasePose<FMeshPoseBoneIndex>
-{
-public:
-	typedef FMeshPoseBoneIndex BoneIndexType;
-
-	const FBoneContainer& GetBoneContainer() const
-	{
-		checkSlow(BoneContainer && BoneContainer->IsValid());
-		return *BoneContainer;
-	}
-
-	void SetBoneContainer(const FBoneContainer* InBoneContainer)
-	{
-		check(InBoneContainer && InBoneContainer->IsValid());
-		BoneContainer = InBoneContainer;
-		InitBones(BoneContainer->GetNumBones());
-	}
-
-	void InitFrom(FCompactPose& CompactPose);
-
-	void MoveBonesTo(TArray<FTransform>& OutTransforms)
-	{
-		OutTransforms = MoveTemp(Bones);
-		BoneContainer = NULL;
-	}
-
-	// Sets this pose to its ref pose
-	ENGINE_API void ResetToRefPose();
-
-	// Sets every bone transform to Identity
-	void ResetToIdentity();
-
-	// returns true if all bone rotations are normalized
-	bool IsNormalized() const;
-
-	// Returns true if any bone rotation contains NaN
-	bool ContainsNaN() const;
-
-
-	FORCEINLINE BoneIndexType GetParentBone(const BoneIndexType& BoneIndex)
-	{
-		return BoneIndexType(BoneContainer->GetParentBoneIndex(BoneIndex.GetInt()));
-	}
-
-protected:
-
-	//Reference to our BoneContainer
-	const FBoneContainer* BoneContainer;
-};
-
 struct FCompactPoseBoneIndexIterator
@@ -164,33 +113,34 @@
 	FCompactPoseBoneIndex operator*() const { return FCompactPoseBoneIndex(Index); }
 };
 
-struct ENGINE_API FCompactPose : FBasePose<FCompactPoseBoneIndex>
+template <typename InAllocator>
+struct FBaseCompactPose : FBasePose<FCompactPoseBoneIndex, InAllocator>
 {
 public:
 	typedef FCompactPoseBoneIndex BoneIndexType;
-
+	typedef InAllocator   Allocator;
 	//--------------------------------------------------------------------------
 	//Bone Index Iteration
-		typedef FBasePose<FCompactPoseBoneIndex>::FRangedForSupport<FCompactPose, FCompactPoseBoneIndexIterator> RangedForBoneIndexFwd;
-		typedef FBasePose<FCompactPoseBoneIndex>::FRangedForReverseSupport<FCompactPose, FCompactPoseBoneIndexReverseIterator> RangedForBoneIndexBwd;
+	typedef typename FBasePose<FCompactPoseBoneIndex, Allocator>::template FRangedForSupport<FBaseCompactPose, FCompactPoseBoneIndexIterator> RangedForBoneIndexFwd;
+	typedef typename FBasePose<FCompactPoseBoneIndex, Allocator>::template FRangedForReverseSupport<FBaseCompactPose, FCompactPoseBoneIndexReverseIterator> RangedForBoneIndexBwd;
 
-		FORCEINLINE RangedForBoneIndexFwd ForEachBoneIndex() const
-		{
-			return RangedForBoneIndexFwd(*this);
-		}
+	FORCEINLINE RangedForBoneIndexFwd ForEachBoneIndex() const
+	{
+		return RangedForBoneIndexFwd(*this);
+	}
 
-		FORCEINLINE RangedForBoneIndexBwd ForEachBoneIndexReverse() const
-		{
-			return RangedForBoneIndexBwd(*this);
-		}
+	FORCEINLINE RangedForBoneIndexBwd ForEachBoneIndexReverse() const
+	{
+		return RangedForBoneIndexBwd(*this);
+	}
 
-		FORCEINLINE FCompactPoseBoneIndexIterator MakeBeginIter() const { return FCompactPoseBoneIndexIterator(0); }
+	FORCEINLINE FCompactPoseBoneIndexIterator MakeBeginIter() const { return FCompactPoseBoneIndexIterator(0); }
 
-		FORCEINLINE FCompactPoseBoneIndexIterator MakeEndIter() const { return FCompactPoseBoneIndexIterator(GetNumBones()); }
+	FORCEINLINE FCompactPoseBoneIndexIterator MakeEndIter() const { return FCompactPoseBoneIndexIterator(this->GetNumBones()); }
 
-		FORCEINLINE FCompactPoseBoneIndexReverseIterator MakeBeginIterReverse() const { return FCompactPoseBoneIndexReverseIterator(GetNumBones() - 1); }
+	FORCEINLINE FCompactPoseBoneIndexReverseIterator MakeBeginIterReverse() const { return FCompactPoseBoneIndexReverseIterator(this->GetNumBones() - 1); }
 
-		FORCEINLINE FCompactPoseBoneIndexReverseIterator MakeEndIterReverse() const { return FCompactPoseBoneIndexReverseIterator(-1); }
+	FORCEINLINE FCompactPoseBoneIndexReverseIterator MakeEndIterReverse() const { return FCompactPoseBoneIndexReverseIterator(-1); }
 	//--------------------------------------------------------------------------
 
 	const FBoneContainer& GetBoneContainer() const 
@@ -203,50 +153,43 @@
 	{
 		check(InBoneContainer && InBoneContainer->IsValid());
 		BoneContainer = InBoneContainer;
-		InitBones(BoneContainer->GetBoneIndicesArray().Num());
+		this->InitBones(BoneContainer->GetBoneIndicesArray().Num());
 	}
 
-	void InitFrom(const FCompactPose& SrcPose)
+	void InitFrom(const FBaseCompactPose& SrcPose)
 	{
 		SetBoneContainer(SrcPose.BoneContainer);
-		Bones = SrcPose.Bones;
+		this->Bones = SrcPose.Bones;
 	}
 
-	// Moves transform data out of the supplied SrcPose (SrcPose will be left empty)
-	void MoveBonesFrom(FCompactPose& SrcPose)
-	{
-		Bones = MoveTemp(SrcPose.Bones);
-		BoneContainer = SrcPose.BoneContainer;
-		SrcPose.BoneContainer = NULL;
-	}
-
-	//Moves transform data to supplied array (pose will be left empty)
-	void MoveBonesTo(TArray<FTransform>& OutTransforms)
-	{
-		OutTransforms = MoveTemp(Bones);
-		BoneContainer = NULL;
-	}
-
-	// Moves transform data out of the supplied InTransforms. InTransform will be left empty
-	void MoveBonesFrom(TArray<FTransform>& InTransforms)
-	{
-		// if number doesn't match it is not compatible to start with
-		if (InTransforms.Num() == Bones.Num())
-		{
-			Bones = MoveTemp(InTransforms);
-		}
-	}
-
 	// Copy bone transform from SrcPose to this
-	void CopyBonesFrom(const FCompactPose& SrcPose)
+	template <typename OtherAllocator>
+	void CopyBonesFrom(const FBaseCompactPose<OtherAllocator>& SrcPose)
 	{
+		this->Bones = SrcPose.GetBones();
+		BoneContainer = &SrcPose.GetBoneContainer();
+	}
+
+	void CopyBonesFrom(const FBaseCompactPose<Allocator>& SrcPose)
+	{
 		if (this != &SrcPose)
 		{
-			Bones = SrcPose.Bones;
-			BoneContainer = SrcPose.BoneContainer;
+			this->Bones = SrcPose.GetBones();
+			BoneContainer = &SrcPose.GetBoneContainer();
 		}
 	}
-
+	
+	template <typename OtherAllocator>
+	void CopyBonesFrom(const TArray<FTransform, OtherAllocator>& SrcPoseBones)
+	{
+		// only allow if the size is same
+		// if size doesn't match, we can't guarantee the bonecontainer would work
+		// so we can't accept
+		if (this->Bones.Num() == SrcPoseBones.Num())
+		{
+			this->Bones = SrcPoseBones;
+		}
+	}
 	// Sets this pose to its ref pose
 	void ResetToRefPose()
 	{
@@ -254,19 +197,80 @@
 	}
 
 	// Sets this pose to the supplied BoneContainers ref pose
-	void ResetToRefPose(const FBoneContainer& RequiredBones);
+	void ResetToRefPose(const FBoneContainer& RequiredBones)
+	{
+		const TArray<FTransform>& RefPoseCompactArray = RequiredBones.GetRefPoseCompactArray();
+		this->Bones.Reset();
+		this->Bones.Append(RefPoseCompactArray);
 
+		// If retargeting is disabled, copy ref pose from Skeleton, rather than mesh.
+		// this is only used in editor and for debugging.
+		if (RequiredBones.GetDisableRetargeting())
+		{
+			checkSlow(RequiredBones.IsValid());
+			// Only do this if we have a mesh. otherwise we're not retargeting animations.
+			if (RequiredBones.GetSkeletalMeshAsset())
+			{
+				TArray<FTransform> const & SkeletonRefPose = RequiredBones.GetSkeletonAsset()->GetRefLocalPoses();
+
+				for (const FCompactPoseBoneIndex BoneIndex : ForEachBoneIndex())
+				{
+					const int32 SkeletonBoneIndex = GetBoneContainer().GetSkeletonIndex(BoneIndex);
+
+					// Pose bone index should always exist in Skeleton
+					checkSlow(SkeletonBoneIndex != INDEX_NONE);
+					this->Bones[BoneIndex.GetInt()] = SkeletonRefPose[SkeletonBoneIndex];
+				}
+			}
+		}
+	}
+
 	// Sets every bone transform to Identity
-	void ResetToIdentity();
+	void ResetToIdentity()
+	{
+		for (FTransform& Bone : this->Bones)
+		{
+			Bone.SetIdentity();
+		}
+	}
 
 	// returns true if all bone rotations are normalized
-	bool IsNormalized() const;
+	bool IsNormalized() const
+	{
+		for (const FTransform& Bone : this->Bones)
+		{
+			if (!Bone.IsRotationNormalized())
+			{
+				return false;
+			}
+		}
 
+		return true;
+	}
+
 	// Returns true if any bone rotation contains NaN
-	bool ContainsNaN() const;
+	bool ContainsNaN() const
+	{
+		for (const FTransform& Bone : this->Bones)
+		{
+			if (Bone.ContainsNaN())
+			{
+				return true;
+			}
+		}
 
+		return false;
+	}
+
+
 	// Normalizes all rotations in this pose
-	void NormalizeRotations();
+	void NormalizeRotations()
+	{
+		for (FTransform& Bone : this->Bones)
+		{
+			Bone.NormalizeRotation();
+		}
+	}
 
 	bool IsValid() const
 	{
@@ -293,14 +297,104 @@
 		const BoneTrackArray& ScaleTracks,
 		float Time)
 	{
+		// @todo fixme 
+		FTransformArray LocalBones;
+		LocalBones = this->Bones;
+
 		AnimationFormat_GetAnimationPose(
-			*((FTransformArray*)&Bones), //@TODO:@ANIMATION: Nasty hack
+			LocalBones, //@TODO:@ANIMATION: Nasty hack
 			RotationTracks,
 			TranslationTracks,
 			ScaleTracks,
 			Seq,
 			Time);
+		this->Bones = LocalBones;
 	}
+
 protected:
 
 	//Reference to our BoneContainer
+	const FBoneContainer* BoneContainer;
+};
+
+struct FCompactHeapPose;
+
+struct ENGINE_API FCompactPose : public FBaseCompactPose<FAnimStackAllocator>
+{
+};
+
+struct ENGINE_API FCompactHeapPose : public FBaseCompactPose<FDefaultAllocator>
+{
+	// Moves transform data out of the supplied SrcPose (SrcPose will be left empty)
+	void MoveBonesFrom(FCompactHeapPose& SrcPose)
+	{
+		Bones = MoveTemp(SrcPose.Bones);
+		BoneContainer = SrcPose.BoneContainer;
+		SrcPose.BoneContainer = NULL;
+	}
+
+	//Moves transform data to supplied array (pose will be left empty)
+	void MoveBonesTo(TArray<FTransform>& OutTransforms)
+	{
+		OutTransforms = MoveTemp(Bones);
+		BoneContainer = NULL;
+	}
+
+	// Moves transform data out of the supplied InTransforms. InTransform will be left empty
+	void MoveBonesFrom(TArray<FTransform>& InTransforms)
+	{
+		// if number doesn't match it is not compatible to start with
+		if (InTransforms.Num() == Bones.Num())
+		{
+			Bones = MoveTemp(InTransforms);
+		}
+	}
+};
+
+struct FMeshPose : public FBasePose<FMeshPoseBoneIndex, FDefaultAllocator>
+{
+public:
+	typedef FMeshPoseBoneIndex BoneIndexType;
+
+	const FBoneContainer& GetBoneContainer() const
+	{
+		checkSlow(BoneContainer && BoneContainer->IsValid());
+		return *BoneContainer;
+	}
+
+	void SetBoneContainer(const FBoneContainer* InBoneContainer)
+	{
+		check(InBoneContainer && InBoneContainer->IsValid());
+		BoneContainer = InBoneContainer;
+		InitBones(BoneContainer->GetNumBones());
+	}
+
+	//void InitFrom(FCompactPose& CompactPose);
+
+	void MoveBonesTo(TArray<FTransform>& OutTransforms)
+	{
+		OutTransforms = MoveTemp(Bones);
+		BoneContainer = NULL;
+	}
+
+	// Sets this pose to its ref pose
+	ENGINE_API void ResetToRefPose();
+
+	// Sets every bone transform to Identity
+	void ResetToIdentity();
+
+	// returns true if all bone rotations are normalized
+	bool IsNormalized() const;
+
+	// Returns true if any bone rotation contains NaN
+	bool ContainsNaN() const;
+
+
+	FORCEINLINE BoneIndexType GetParentBone(const BoneIndexType& BoneIndex)
+	{
+		return BoneIndexType(BoneContainer->GetParentBoneIndex(BoneIndex.GetInt()));
+	}
+
+protected:
+
+	//Reference to our BoneContainer
@@ -331,7 +425,17 @@
 		ComponentSpaceFlags.AddZeroed(Pose.GetNumBones());
 		ComponentSpaceFlags[0] = 1;
 	}
+
+	// Copy Pose
+	template <typename OtherPoseType>
+	void CopyPose(const OtherPoseType& SrcPose)
+	{
+		Pose.CopyBonesFrom(SrcPose.GetPose());
+		ComponentSpaceFlags = SrcPose.GetComponentSpaceFlags();
+	}
+
 	const PoseType& GetPose() const { return Pose; }
+	const TCustomBoneIndexArray<uint8, BoneIndexType>& GetComponentSpaceFlags() const { return ComponentSpaceFlags; }
 
 	// Get transform for supplied bone in local space
 	FTransform GetLocalSpaceTransform(BoneIndexType BoneIndex);
diff --git a/Engine/Source/Runtime/Engine/Public/CollisionQueryParams.h b/Engine/Source/Runtime/Engine/Public/CollisionQueryParams.h
--- a/Engine/Source/Runtime/Engine/Public/CollisionQueryParams.h
+++ b/Engine/Source/Runtime/Engine/Public/CollisionQueryParams.h
@@ -39,7 +39,10 @@ struct ENGINE_API FCollisionQueryParams
 	bool bIgnoreBlocks;
 
 	/** TArray typedef of components to ignore. */
-	typedef TArray<uint32, TInlineAllocator<NumInlinedActorComponents>> IgnoreComponentsArrayType;
+	typedef TArray<uint32, TInlineAllocator<8>> IgnoreComponentsArrayType;
+
+	/** TArray typedef of actors to ignore. */
+	typedef TArray<uint32, TInlineAllocator<4>> IgnoreActorsArrayType;
 
 	/** Extra filtering done on the query. See declaration for filtering logic */
 	FMaskFilter IgnoreMask;
@@ -52,6 +55,9 @@ private:
 	/** Set of components to ignore during the trace */
 	mutable IgnoreComponentsArrayType IgnoreComponents;
 
+	/** Set of actors to ignore during the trace */
+	IgnoreActorsArrayType IgnoreActors;
+
 	void Internal_AddIgnoredComponent(const UPrimitiveComponent* InIgnoreComponent);
 
 public:
@@ -59,6 +65,12 @@ public:
 	/** Returns set of unique components to ignore during the trace. Elements are guaranteed to be unique (they are made so internally if they are not already). */
 	const IgnoreComponentsArrayType& GetIgnoredComponents() const;
 
+	/** Returns set of actors to ignore during the trace. Note that elements are NOT guaranteed to be unique. This is less important for actors since it's less likely that duplicates are added.*/
+	const IgnoreActorsArrayType& GetIgnoredActors() const
+	{
+		return IgnoreActors;
+	}
+
 	/** Clears the set of components to ignore during the trace. */
 	void ClearIgnoredComponents()
 	{
diff --git a/Engine/Source/Runtime/Engine/Public/DrawDebugHelpers.h b/Engine/Source/Runtime/Engine/Public/DrawDebugHelpers.h
--- a/Engine/Source/Runtime/Engine/Public/DrawDebugHelpers.h
+++ b/Engine/Source/Runtime/Engine/Public/DrawDebugHelpers.h
@@ -16,11 +16,11 @@ ENGINE_API void DrawDebugPoint(const UWorld* InWorld, FVector const& Position, f
 /** Draw directional arrow **/
 ENGINE_API void DrawDebugDirectionalArrow(const UWorld* InWorld, FVector const& LineStart, FVector const& LineEnd, float ArrowSize, FColor const& Color, bool bPersistentLines = false, float LifeTime=-1.f, uint8 DepthPriority = 0, float Thickness = 0.f);
 /** Draw a debug box */
-ENGINE_API void DrawDebugBox(const UWorld* InWorld, FVector const& Center, FVector const& Extent, FColor const& Color, bool bPersistentLines = false, float LifeTime=-1.f, uint8 DepthPriority = 0);
+ENGINE_API void DrawDebugBox(const UWorld* InWorld, FVector const& Center, FVector const& Extent, FColor const& Color, bool bPersistentLines = false, float LifeTime=-1.f, uint8 DepthPriority = 0, float Thickness = 0.f);
 /** Draw a debug box with rotation */
-ENGINE_API void DrawDebugBox(const UWorld* InWorld, FVector const& Center, FVector const& Box, const FQuat& Rotation, FColor const& Color, bool bPersistentLines = false, float LifeTime=-1.f, uint8 DepthPriority = 0);
+ENGINE_API void DrawDebugBox(const UWorld* InWorld, FVector const& Center, FVector const& Box, const FQuat& Rotation, FColor const& Color, bool bPersistentLines = false, float LifeTime=-1.f, uint8 DepthPriority = 0, float Thickness = 0.f);
 /** Draw Debug coordinate system */
-ENGINE_API void DrawDebugCoordinateSystem(const UWorld* InWorld, FVector const& AxisLoc, FRotator const& AxisRot, float Scale, bool bPersistentLines = false, float LifeTime=-1.f, uint8 DepthPriority = 0);
+ENGINE_API void DrawDebugCoordinateSystem(const UWorld* InWorld, FVector const& AxisLoc, FRotator const& AxisRot, float Scale, bool bPersistentLines = false, float LifeTime=-1.f, uint8 DepthPriority = 0, float Thickness = 0.f);
 /** Draw Debug crosshair */
 ENGINE_API void DrawDebugCrosshairs(const UWorld* InWorld, FVector const& AxisLoc, FRotator const& AxisRot, float Scale, const FColor& Color = FColor::White, bool bPersistentLines = false, float LifeTime = -1.f, uint8 DepthPriority = 0);
 /** Draw Debug Circle */
@@ -28,17 +28,17 @@ ENGINE_API void DrawDebugCircle(const UWorld* InWorld, const FMatrix& TransformM
 /** Draw Debug Circle */
 ENGINE_API void DrawDebugCircle(const UWorld* InWorld, FVector Center, float Radius, int32 Segments, const FColor& Color, bool bPersistentLines = false, float LifeTime=-1.f, uint8 DepthPriority = 0, float Thickness = 0.f, FVector YAxis=FVector(0.f,1.f,0.f), FVector ZAxis=FVector(0.f,0.f,1.f), bool bDrawAxis=true);
 /** Draw Debug 2D donut */
-ENGINE_API void DrawDebug2DDonut(const UWorld* InWorld, const FMatrix& TransformMatrix, float InnerRadius, float OuterRadius, int32 Segments, const FColor& Color, bool bPersistentLines = false, float LifeTime=-1.f, uint8 DepthPriority = 0);
+ENGINE_API void DrawDebug2DDonut(const UWorld* InWorld, const FMatrix& TransformMatrix, float InnerRadius, float OuterRadius, int32 Segments, const FColor& Color, bool bPersistentLines = false, float LifeTime=-1.f, uint8 DepthPriority = 0, float Thickness = 0.f);
 /** Draw a debug sphere */
-ENGINE_API void DrawDebugSphere(const UWorld* InWorld, FVector const& Center, float Radius, int32 Segments, FColor const& Color, bool bPersistentLines = false, float LifeTime=-1.f, uint8 DepthPriority = 0);
+ENGINE_API void DrawDebugSphere(const UWorld* InWorld, FVector const& Center, float Radius, int32 Segments, FColor const& Color, bool bPersistentLines = false, float LifeTime=-1.f, uint8 DepthPriority = 0, float Thickness = 0.f);
 /** Draw a debug cylinder */
-ENGINE_API void DrawDebugCylinder(const UWorld* InWorld, FVector const& Start, FVector const& End, float Radius, int32 Segments, FColor const& Color, bool bPersistentLines = false, float LifeTime=-1.f, uint8 DepthPriority = 0);
+ENGINE_API void DrawDebugCylinder(const UWorld* InWorld, FVector const& Start, FVector const& End, float Radius, int32 Segments, FColor const& Color, bool bPersistentLines = false, float LifeTime=-1.f, uint8 DepthPriority = 0, float Thickness = 0.f);
 /** Draw a debug cone. AngleWidth and AngleHeight are given in radians */
-ENGINE_API void DrawDebugCone(const UWorld* InWorld, FVector const& Origin, FVector const& Direction, float Length, float AngleWidth, float AngleHeight, int32 NumSides, FColor const& Color, bool bPersistentLines=false, float LifeTime=-1.f, uint8 DepthPriority = 0);
+ENGINE_API void DrawDebugCone(const UWorld* InWorld, FVector const& Origin, FVector const& Direction, float Length, float AngleWidth, float AngleHeight, int32 NumSides, FColor const& Color, bool bPersistentLines=false, float LifeTime=-1.f, uint8 DepthPriority = 0, float Thickness = 0.f);
 /** Used by gameplay when defining a cone by a vertical and horizontal dot products. */
-ENGINE_API void DrawDebugAltCone(const UWorld* InWorld, FVector const& Origin, FRotator const& Rotation, float Length, float AngleWidth, float AngleHeight, FColor const& DrawColor, bool bPersistentLines=false, float LifeTime=-1.f, uint8 DepthPriority=0);
+ENGINE_API void DrawDebugAltCone(const UWorld* InWorld, FVector const& Origin, FRotator const& Rotation, float Length, float AngleWidth, float AngleHeight, FColor const& DrawColor, bool bPersistentLines=false, float LifeTime=-1.f, uint8 DepthPriority=0, float Thickness = 0.f);
 ENGINE_API void DrawDebugString(const UWorld* InWorld, FVector const& TextLocation, const FString& Text, class AActor* TestBaseActor = NULL, FColor const& TextColor = FColor::White, float Duration = -1.000000, bool bDrawShadow = false);
-ENGINE_API void DrawDebugFrustum(const UWorld* InWorld, const FMatrix& FrustumToWorld, FColor const& Color, bool bPersistentLines = false, float LifeTime=-1.f, uint8 DepthPriority = 0);
+ENGINE_API void DrawDebugFrustum(const UWorld* InWorld, const FMatrix& FrustumToWorld, FColor const& Color, bool bPersistentLines = false, float LifeTime=-1.f, uint8 DepthPriority = 0, float Thickness = 0.f);
 /** Draw a capsule using the LineBatcher */
 ENGINE_API void DrawDebugCapsule(const UWorld* InWorld, FVector const& Center, float HalfHeight, float Radius, const FQuat& Rotation, FColor const& Color, bool bPersistentLines = false, float LifeTime = -1.f, uint8 DepthPriority = 0, float Thickness = 0);
 /** Draw a debug camera shape.  FOV is full angle in degrees. */
diff --git a/Engine/Source/Runtime/Engine/Public/FinalPostProcessSettings.h b/Engine/Source/Runtime/Engine/Public/FinalPostProcessSettings.h
--- a/Engine/Source/Runtime/Engine/Public/FinalPostProcessSettings.h
+++ b/Engine/Source/Runtime/Engine/Public/FinalPostProcessSettings.h
@@ -15,10 +15,9 @@ public:
 	, HighResScreenshotCaptureRegionMaterial(NULL)
 	, bBufferVisualizationDumpRequired(false)
 	{
-		// to avoid reallocations we reserve some elements
-
-		ContributingCubemaps.Empty(8);
-		ContributingLUTs.Empty(8);
+		// to avoid reallocations we reserve some 
+		ContributingCubemaps.Reserve(8);
+		ContributingLUTs.Reserve(8);
 
 		SetLUT(0);
 	}
@@ -49,6 +48,7 @@ public:
 	void UpdateEntry(const FCubemapEntry &Entry, float Weight)
 	{
 		bool Existing = false;
+		// Note: make sure to read the array size each loop here, since we remove elements.
 		for(int32 i = 0; i < ContributingCubemaps.Num(); ++i)
 		{
 			FCubemapEntry& Local = ContributingCubemaps[i];
@@ -65,7 +65,8 @@ public:
 
 			if(Local.AmbientCubemapTintMulScaleValue.IsAlmostBlack())
 			{
-				ContributingCubemaps.RemoveAt(i);
+				ContributingCubemaps.RemoveAt(i, /*bAllowShrinking=*/ false);
+				i--; // Maintain same index in the loop after loop increment since we removed an element.
 			}
 		}
 
@@ -144,14 +145,14 @@ public:
 	void SetLUT(UTexture* Texture)
 	{
 		// intentionally no deallocations
-		ContributingLUTs.Empty();
+		ContributingLUTs.Reset();
 
 		PushLUT(Texture, 1.0f);
 	}
 
 	// was not merged during blending unlike e.g. BloomThreshold 
-	TArray<FCubemapEntry> ContributingCubemaps;
-	TArray<FLUTBlenderEntry> ContributingLUTs;
+	TArray<FCubemapEntry, TInlineAllocator<8>> ContributingCubemaps;
+	TArray<FLUTBlenderEntry, TInlineAllocator<8>> ContributingLUTs;
 
 	// List of materials to use in the buffer visualization overview
 	TArray<UMaterialInterface*> BufferVisualizationOverviewMaterials;
diff --git a/Engine/Source/Runtime/Engine/Public/VisualLogger/VisualLogger.h b/Engine/Source/Runtime/Engine/Public/VisualLogger/VisualLogger.h
--- a/Engine/Source/Runtime/Engine/Public/VisualLogger/VisualLogger.h
+++ b/Engine/Source/Runtime/Engine/Public/VisualLogger/VisualLogger.h
@@ -198,13 +198,16 @@ public:
 	bool IsBlockedForAllCategories() const { return !!bBlockedAllCategories; }
 
 	/** Returns white list for modifications */
-	const TArray<FName>& GetWhiteList() const { return CategoriesWhiteList; }
+	const TArray<FName>& GetWhiteList() const { return CategoriesWhitelist; }
 
-	bool IsWhiteListed(const FName& Name) const { return CategoriesWhiteList.Find(Name) != INDEX_NONE; }
+	bool IsWhiteListed(const FName& Name) const { return CategoriesWhitelist.Find(Name) != INDEX_NONE; }
 
-	void AddCategortyToWhiteList(FName Category) { CategoriesWhiteList.AddUnique(Category); }
+	void AddCategoryToWhitelist(FName Category) { CategoriesWhitelist.AddUnique(Category); }
 
-	void ClearWhiteList() { CategoriesWhiteList.Reset(); }
+	DEPRECATED(4.12, "Please use the AddCategoryToWhitelist instead")
+	void AddCategortyToWhiteList(FName Category) { AddCategoryToWhitelist(Category); }
+
+	void ClearWhiteList() { CategoriesWhitelist.Reset(); }
 
 	/** Generates and returns Id unique for given timestamp - used to connect different logs between (ex. text log with geometry shape) */
 	int32 GetUniqueId(float Timestamp);
@@ -230,9 +233,12 @@ public:
 	const TArray<FVisualLogDevice*>& GetDevices() const { return OutputDevices; }
 	/** Check if log category can be recorded, verify before using GetEntryToWrite! */
 	bool IsCategoryLogged(const FLogCategoryBase& Category) const;
-	/** Returns  current entry for given TimeStap or creates another one  but first it serialize previous entry as completed to vislog devices. Use VisualLogger::DontCreate to get current entry without serialization*/
+	/** Returns  current entry for given TimeStap or creates another one  but first it serialize previous 
+	 *	entry as completed to vislog devices. Use VisualLogger::DontCreate to get current entry without serialization
+	 *	@note this function can return null */
 	FVisualLogEntry* GetEntryToWrite(const UObject* Object, float TimeStamp, ECreateIfNeeded ShouldCreate = ECreateIfNeeded::Create);
-	/** Retrieves last used entry for given UObject*/
+	/** Retrieves last used entry for given UObject
+	 *	@note this function can return null */
 	FVisualLogEntry* GetLastEntryForObject(const UObject* Object);
 	/** flush and serialize data if timestamp allows it */
 	virtual void Flush() override;
@@ -257,6 +263,10 @@ public:
 
 	typedef TMap<const UObject*, TWeakObjectPtr<const UWorld> > FObjectToWorldMapType;
 	FObjectToWorldMapType& GetObjectToWorldMap() { return ObjectToWorldMap; }
+
+	void AddWhitelistedClass(UClass& InClass);
+	bool IsClassWhitelisted(const UClass& InClass) const;
+
 private:
 	FVisualLogger();
 	virtual void Serialize(const TCHAR* V, ELogVerbosity::Type Verbosity, const FName& Category) override { ensureMsgf(0, TEXT("Regular serialize is forbiden for visual logs")); }
@@ -267,8 +277,12 @@ protected:
 	// Map for inter-objects redirections
 	static TMap<const UWorld*, RedirectionMapType> WorldToRedirectionMap;
 
+	// white-listed classes - only instances of these classes will be logged. 
+	// if ClassWhitelist is empty (default) everything will log
+	TArray<UClass*> ClassWhitelist;
+
 	// white list of categories to bypass blocking
-	TArray<FName>	CategoriesWhiteList;
+	TArray<FName>	CategoriesWhitelist;
 	// Visual Logger extensions map
 	TMap<FName, FVisualLogExtensionInterface*> AllExtensions;
 	// last generated unique id for given times tamp
diff --git a/Engine/Source/Runtime/Engine/Public/VisualLogger/VisualLoggerTypes.h b/Engine/Source/Runtime/Engine/Public/VisualLogger/VisualLoggerTypes.h
--- a/Engine/Source/Runtime/Engine/Public/VisualLogger/VisualLoggerTypes.h
+++ b/Engine/Source/Runtime/Engine/Public/VisualLogger/VisualLoggerTypes.h
@@ -169,6 +169,7 @@ struct ENGINE_API FVisualLogEntry
 #if ENABLE_VISUAL_LOG
 	float TimeStamp;
 	FVector Location;
+	uint8 bIsAllowedToLog : 1;
 
 	TArray<FVisualLogEvent> Events;
 	TArray<FVisualLogLine> LogLines;
@@ -179,7 +180,7 @@ struct ENGINE_API FVisualLogEntry
 
 	FVisualLogEntry() { Reset(); }
 	FVisualLogEntry(const FVisualLogEntry& Entry);
-	FVisualLogEntry(const class AActor* InActor, TArray<TWeakObjectPtr<UObject> >* Children);
+	FVisualLogEntry(const AActor* InActor, TArray<TWeakObjectPtr<UObject> >* Children);
 	FVisualLogEntry(float InTimeStamp, FVector InLocation, const UObject* Object, TArray<TWeakObjectPtr<UObject> >* Children);
 
 	void Reset();
@@ -239,7 +240,7 @@ public:
 	};
 
 
-	virtual void Serialize(const class UObject* LogOwner, FName OwnerName, FName InOwnerClassName, const FVisualLogEntry& LogEntry) = 0;
+	virtual void Serialize(const UObject* LogOwner, FName OwnerName, FName InOwnerClassName, const FVisualLogEntry& LogEntry) = 0;
 	virtual void Cleanup(bool bReleaseMemory = false) { /* Empty */ }
 	virtual void StartRecordingToFile(float TImeStamp) { /* Empty */ }
 	virtual void StopRecordingToFile(float TImeStamp) { /* Empty */ }
diff --git a/Engine/Source/Runtime/Landscape/Private/LandscapeCollision.cpp b/Engine/Source/Runtime/Landscape/Private/LandscapeCollision.cpp
--- a/Engine/Source/Runtime/Landscape/Private/LandscapeCollision.cpp
+++ b/Engine/Source/Runtime/Landscape/Private/LandscapeCollision.cpp
@@ -194,7 +194,7 @@ void ULandscapeHeightfieldCollisionComponent::CreatePhysicsState()
 
 				// Setup filtering
 				PxFilterData PQueryFilterData, PSimFilterData;
-				CreateShapeFilterData(GetCollisionObjectType(), FMaskFilter(0), GetUniqueID(), GetCollisionResponseToChannels(), 0, 0, PQueryFilterData, PSimFilterData, true, false, true);
+				CreateShapeFilterData(GetCollisionObjectType(), FMaskFilter(0), GetOwner()->GetUniqueID(), GetCollisionResponseToChannels(), GetUniqueID(), 0, PQueryFilterData, PSimFilterData, true, false, true);
 
 				// Heightfield is used for simple and complex collision
 				PQueryFilterData.word3 |= bCreateSimpleCollision ? EPDF_ComplexCollision : (EPDF_SimpleCollision | EPDF_ComplexCollision);
@@ -239,7 +239,7 @@ void ULandscapeHeightfieldCollisionComponent::CreatePhysicsState()
 						CollisionResponse.SetAllChannels(ECollisionResponse::ECR_Ignore);
 						CollisionResponse.SetResponse(ECollisionChannel::ECC_Visibility, ECR_Block);
 						PxFilterData PQueryFilterDataEd, PSimFilterDataEd;
-						CreateShapeFilterData(ECollisionChannel::ECC_Visibility, FMaskFilter(0), GetUniqueID(), CollisionResponse, 0, 0, PQueryFilterDataEd, PSimFilterDataEd, true, false, true);
+						CreateShapeFilterData(ECollisionChannel::ECC_Visibility, FMaskFilter(0), GetOwner()->GetUniqueID(), CollisionResponse, GetUniqueID(), 0, PQueryFilterDataEd, PSimFilterDataEd, true, false, true);
 
 						PQueryFilterDataEd.word3 |= (EPDF_SimpleCollision | EPDF_ComplexCollision);
 						HeightFieldEdShapeSync->setQueryFilterData(PQueryFilterDataEd);
@@ -927,7 +927,7 @@ void ULandscapeMeshCollisionComponent::CreatePhysicsState()
 
 				// Setup filtering
 				PxFilterData PQueryFilterData, PSimFilterData;
-				CreateShapeFilterData(GetCollisionObjectType(), FMaskFilter(0), GetUniqueID(), GetCollisionResponseToChannels(), 0, 0, PQueryFilterData, PSimFilterData, false, false, true);
+				CreateShapeFilterData(GetCollisionObjectType(), FMaskFilter(0), GetOwner()->GetUniqueID(), GetCollisionResponseToChannels(), GetUniqueID(), 0, PQueryFilterData, PSimFilterData, false, false, true);
 
 				// Heightfield is used for simple and complex collision
 				PQueryFilterData.word3 |= (EPDF_SimpleCollision | EPDF_ComplexCollision);
@@ -975,7 +975,7 @@ void ULandscapeMeshCollisionComponent::CreatePhysicsState()
 						CollisionResponse.SetAllChannels(ECollisionResponse::ECR_Ignore);
 						CollisionResponse.SetResponse(ECollisionChannel::ECC_Visibility, ECR_Block);
 						PxFilterData PQueryFilterDataEd, PSimFilterDataEd;
-						CreateShapeFilterData(ECollisionChannel::ECC_Visibility, FMaskFilter(0), GetUniqueID(), CollisionResponse, 0, 0, PQueryFilterDataEd, PSimFilterDataEd, true, false, true);
+						CreateShapeFilterData(ECollisionChannel::ECC_Visibility, FMaskFilter(0), GetOwner()->GetUniqueID(), CollisionResponse, GetUniqueID(), 0, PQueryFilterDataEd, PSimFilterDataEd, true, false, true);
 
 						PQueryFilterDataEd.word3 |= (EPDF_SimpleCollision | EPDF_ComplexCollision);
 						MeshShapeEdSync->setQueryFilterData(PQueryFilterDataEd);
diff --git a/Engine/Source/Runtime/Windows/XAudio2/Private/XAudio2Buffer.cpp b/Engine/Source/Runtime/Windows/XAudio2/Private/XAudio2Buffer.cpp
--- a/Engine/Source/Runtime/Windows/XAudio2/Private/XAudio2Buffer.cpp
+++ b/Engine/Source/Runtime/Windows/XAudio2/Private/XAudio2Buffer.cpp
@@ -315,6 +315,12 @@ void FXAudio2SoundBuffer::InitXWMA( USoundWave* Wave, FXMAInfo* XMAInfo )
  */
 bool FXAudio2SoundBuffer::ReadCompressedData( uint8* Destination, bool bLooping )
 {
+	if (!DecompressionState)
+	{
+		UE_LOG(LogXAudio2, Warning, TEXT( "Attempting to read compressed data without a compression state instance for resource '%s'" ), *ResourceName );
+		return false;
+	}
+	
 	const uint32 kPCMBufferSize = MONO_PCM_BUFFER_SIZE * NumChannels;
 	if (SoundFormat == SoundFormat_Streaming)
 	{
diff --git a/Engine/Source/Runtime/Windows/XAudio2/Private/XAudio2Source.cpp b/Engine/Source/Runtime/Windows/XAudio2/Private/XAudio2Source.cpp
--- a/Engine/Source/Runtime/Windows/XAudio2/Private/XAudio2Source.cpp
+++ b/Engine/Source/Runtime/Windows/XAudio2/Private/XAudio2Source.cpp
@@ -134,6 +134,18 @@ void FXAudio2SoundSource::SubmitPCMBuffers( void )
 	XAudio2Buffers[0].AudioBytes = XAudio2Buffer->PCM.PCMDataSize;
 	XAudio2Buffers[0].pContext = this;
 
+	if (!AudioDevice)
+	{
+		UE_LOG(LogXAudio2, Error, TEXT("SubmitPCMBuffers: Audio Device is nullptr"));
+		return;
+	}
+
+	if (!Source)
+	{
+		UE_LOG(LogXAudio2, Error, TEXT("SubmitPCMBuffers: Source (IXAudio2SourceVoice is nullptr"));
+		return;
+	}
+
 	if( WaveInstance->LoopingMode == LOOP_Never )
 	{
 		XAudio2Buffers[0].Flags = XAUDIO2_END_OF_STREAM;
@@ -245,7 +257,7 @@ void FXAudio2SoundSource::SubmitPCMRTBuffers( void )
 
 	ReadMorePCMData(2, DataReadMode);
 
-	if (DataReadMode == EDataReadMode::Synchronous)
+	if (DataReadMode == EDataReadMode::Synchronous || (!bSkipFirstBuffer && WaveInstance->WaveData && !WaveInstance->WaveData->bCanProcessAsync))
 	{
 		AudioDevice->ValidateAPICall(TEXT("SubmitSourceBuffer - PCMRT"),
 									 Source->SubmitSourceBuffer(&XAudio2Buffers[2]));
@@ -458,6 +470,10 @@ bool FXAudio2SoundSource::Init(FWaveInstance* InWaveInstance)
 		XAudio2Buffer = FXAudio2SoundBuffer::Init(BestAudioDevice, InWaveInstance->WaveData, InWaveInstance->StartTime > 0.f);
 		Buffer = XAudio2Buffer;
 
+		// Reset the LPFFrequency values
+		LPFFrequency = MAX_FILTER_FREQUENCY;
+		LastLPFFrequency = MAX_FILTER_FREQUENCY;
+
 		// Buffer failed to be created, or there was an error with the compressed data
 		if (Buffer && Buffer->NumChannels > 0)
 		{
@@ -742,7 +758,7 @@ void FXAudio2SoundSource::GetStereoChannelVolumes(float ChannelVolumes[CHANNEL_M
 		ChannelVolumes[CHANNELOUT_RADIO] = 0.0f;
 		if (WaveInstance->bApplyRadioFilter)
 		{
-			ChannelVolumes[CHANNELOUT_RADIO] = WaveInstance->RadioFilterVolume;
+			ChannelVolumes[CHANNELOUT_RADIO] = AttenuatedVolume * WaveInstance->RadioFilterVolume;
 		}
 	}
 }
