Copying //UE4/Orion-Staging to //UE4/Main (originated from //Orion/Dev-General @ 2831630)

diff --git a/Engine/Source/Runtime/CoreUObject/Private/UObject/Class.cpp b/Engine/Source/Runtime/CoreUObject/Private/UObject/Class.cpp
--- a/Engine/Source/Runtime/CoreUObject/Private/UObject/Class.cpp
+++ b/Engine/Source/Runtime/CoreUObject/Private/UObject/Class.cpp
@@ -144,6 +144,11 @@ void UField::Serialize( FArchive& Ar )
 {
 	Super::Serialize( Ar );
 	Ar << Next;
+	if (Ar.IsLoading())
+	{
+		// Make sure that after loading new assets we will check for new classes and generate their token stream.
+		SetTokenStreamMaybeDirty(true);
+	}
 }
 
 void UField::AddCppProperty( UProperty* Property )
@@ -3487,9 +3492,6 @@ void UClass::Serialize( FArchive& Ar )
 			};
 			Ar.Seek(InterfacesStart + sizeof(NumInterfaces) + NumInterfaces * sizeof(FSerializedInterfaceReference));
 		}
-
-		// At this point token stream is most likely dirty too
-		SetTokenStreamMaybeDirty(true);
 	}
 
 	if (!Ar.IsIgnoringClassGeneratedByRef())
@@ -3501,7 +3503,7 @@ void UClass::Serialize( FArchive& Ar )
 	{
 		checkf(!HasAnyClassFlags(CLASS_Native), TEXT("Class %s loaded with CLASS_Native....we should not be loading any native classes."), *GetFullName());
 		checkf(!HasAnyClassFlags(CLASS_Intrinsic), TEXT("Class %s loaded with CLASS_Intrinsic....we should not be loading any intrinsic classes."), *GetFullName());
-		ClassFlags &= ~ CLASS_ShouldNeverBeLoaded;
+		ClassFlags &= ~(CLASS_ShouldNeverBeLoaded | CLASS_TokenStreamAssembled);
 		if (!(Ar.GetPortFlags() & PPF_Duplicate))
 		{
 			Link(Ar, true);
diff --git a/Engine/Source/Runtime/CoreUObject/Public/Serialization/ArchiveUObject.h b/Engine/Source/Runtime/CoreUObject/Public/Serialization/ArchiveUObject.h
--- a/Engine/Source/Runtime/CoreUObject/Public/Serialization/ArchiveUObject.h
+++ b/Engine/Source/Runtime/CoreUObject/Public/Serialization/ArchiveUObject.h
@@ -1533,6 +1533,11 @@ private:
 	int64							PrecacheEndPos[2];
 	/** Buffer containing precached data.											*/
 	uint8*							PrecacheBuffer[2];
+	/** Precahce buffer size */
+	SIZE_T PrecacheBufferSize[2];
+	/** True if the precache buffer is protected */
+	bool PrecacheBufferProtected[2];
+
 	/** Status of pending read, a value of 0 means no outstanding reads.			*/
 	FThreadSafeCounter				PrecacheReadStatus[2];
 	
diff --git a/Engine/Source/Runtime/CoreUObject/Public/Serialization/PropertyLocalizationDataGathering.h b/Engine/Source/Runtime/CoreUObject/Public/Serialization/PropertyLocalizationDataGathering.h
--- a/Engine/Source/Runtime/CoreUObject/Public/Serialization/PropertyLocalizationDataGathering.h
+++ b/Engine/Source/Runtime/CoreUObject/Public/Serialization/PropertyLocalizationDataGathering.h
@@ -37,6 +37,7 @@ public:
 
 	FPropertyLocalizationDataGatherer(TArray<FGatherableTextData>& InGatherableTextDataArray, const UPackage* const InPackage);
 
+	void GatherLocalizationDataFromObjectWithCallbacks(const UObject* Object, const EPropertyLocalizationGathererTextFlags GatherTextFlags);
 	void GatherLocalizationDataFromObject(const UObject* Object, const EPropertyLocalizationGathererTextFlags GatherTextFlags);
 	void GatherLocalizationDataFromObjectFields(const FString& PathToParent, const UObject* Object, const EPropertyLocalizationGathererTextFlags GatherTextFlags);
 	void GatherLocalizationDataFromStructFields(const FString& PathToParent, const UStruct* Struct, const void* StructData, const EPropertyLocalizationGathererTextFlags GatherTextFlags);
diff --git a/Engine/Source/Runtime/Engine/Classes/Animation/AnimMontage.h b/Engine/Source/Runtime/Engine/Classes/Animation/AnimMontage.h
--- a/Engine/Source/Runtime/Engine/Classes/Animation/AnimMontage.h
+++ b/Engine/Source/Runtime/Engine/Classes/Animation/AnimMontage.h
@@ -183,6 +183,9 @@ struct FAnimMontageInstance
 	// markers that passed in this tick
 	TArray<FPassedMarker> MarkersPassedThisTick;
 
+	// Whether this in this tick's call to Advance we used marker based sync
+	bool bDidUseMarkerSyncThisTick;
+
 private:
 	// list of next sections per section - index of array is section id
 	UPROPERTY()
@@ -278,6 +281,7 @@ public:
 		: Montage(NULL)
 		, bPlaying(false)
 		, DefaultBlendTimeMultiplier(1.0f)
+		, bDidUseMarkerSyncThisTick(false)
 		, AnimInstance(NULL)
 		, Position(0.f)
 		, PlayRate(1.f)
@@ -295,6 +299,7 @@ public:
 		: Montage(NULL)
 		, bPlaying(false)
 		, DefaultBlendTimeMultiplier(1.0f)
+		, bDidUseMarkerSyncThisTick(false)
 		, AnimInstance(InAnimInstance)
 		, Position(0.f)
 		, PlayRate(1.f)
diff --git a/Engine/Source/Runtime/Engine/Classes/Components/DecalComponent.h b/Engine/Source/Runtime/Engine/Classes/Components/DecalComponent.h
--- a/Engine/Source/Runtime/Engine/Classes/Components/DecalComponent.h
+++ b/Engine/Source/Runtime/Engine/Classes/Components/DecalComponent.h
@@ -32,18 +32,22 @@ class ENGINE_API UDecalComponent : public USceneComponent
 	UPROPERTY(EditAnywhere, BlueprintReadOnly, Category = Decal)
 	float FadeScreenSize;
 
-protected:
-	UPROPERTY()
+	/**
+	* Time in seconds to wait before beginning to fade out the decal. Set fade duration and start delay to 0 to make persistent.
+	*/
+	UPROPERTY(EditAnywhere, BlueprintReadOnly, Category = Decal)
 	float FadeStartDelay;
 
-	UPROPERTY()
+	/**
+	* Time in seconds for the decal to fade out. Set fade duration and start delay to 0 to make persistent.
+	*/
+	UPROPERTY(EditAnywhere, BlueprintReadOnly, Category = Decal)
 	float FadeDuration;
 
-public:
 	/**
 	* Automatically destroys the owning actor after fully fading out.
 	*/
-	UPROPERTY()
+	UPROPERTY(EditAnywhere, BlueprintReadOnly, Category = Decal)
 	uint8 bDestroyOwnerAfterFade : 1;
 
 	UFUNCTION(BlueprintCallable, Category = "Rendering|Components|Decal")
diff --git a/Engine/Source/Runtime/Engine/Classes/Kismet/KismetInputLibrary.h b/Engine/Source/Runtime/Engine/Classes/Kismet/KismetInputLibrary.h
--- a/Engine/Source/Runtime/Engine/Classes/Kismet/KismetInputLibrary.h
+++ b/Engine/Source/Runtime/Engine/Classes/Kismet/KismetInputLibrary.h
@@ -23,6 +23,15 @@ class ENGINE_API UKismetInputLibrary : public UBlueprintFunctionLibrary
 	 */
 	UFUNCTION(BlueprintPure, meta=(DisplayName = "Equal (Key)", CompactNodeTitle = "=="), Category="Utilities|Key")
 	static bool EqualEqual_KeyKey(FKey A, FKey B);
+
+	/**
+	* Test if the input chords are equal (A == B)
+	* @param A - The chord to compare against
+	* @param B - The chord to compare
+	* @returns True if the chords are equal, false otherwise
+	*/
+	UFUNCTION( BlueprintPure, meta = ( DisplayName = "Equal (Input Chord)", CompactNodeTitle = "==" ), Category = "Utilities|Key" )
+	static bool EqualEqual_InputChordInputChord( FInputChord A, FInputChord B );
 	
 	/**
 	 * @returns True if the key is a modifier key: Ctrl, Command, Alt, Shift
diff --git a/Engine/Source/Runtime/Engine/Classes/Kismet/KismetTextLibrary.h b/Engine/Source/Runtime/Engine/Classes/Kismet/KismetTextLibrary.h
--- a/Engine/Source/Runtime/Engine/Classes/Kismet/KismetTextLibrary.h
+++ b/Engine/Source/Runtime/Engine/Classes/Kismet/KismetTextLibrary.h
@@ -123,15 +123,24 @@ class ENGINE_API UKismetTextLibrary : public UBlueprintFunctionLibrary
 	/* Converts a passed in float to a text based on formatting options */
 	UFUNCTION(BlueprintPure, meta=(DisplayName = "ToText (float)", AdvancedDisplay = "1", BlueprintAutocast), Category="Utilities|Text")
 	static FText Conv_FloatToText(float Value, TEnumAsByte<ERoundingMode> RoundingMode, bool bUseGrouping = true, int32 MinimumIntegralDigits = 1, int32 MaximumIntegralDigits = 324, int32 MinimumFractionalDigits = 0, int32 MaximumFractionalDigits = 3);
+	
+	/**
+	 * Generate an FText that represents the passed number as currency in the current culture.
+	 * BaseVal is specified in the smallest fractional value of the currency and will be converted for formatting according to the selected culture.
+	 * Keep in mind the CurrencyCode is completely independent of the culture it's displayed in (and they do not imply one another).
+	 * For example: FText::AsCurrencyBase(650, TEXT("EUR")); would return an FText of "€6.50" in most English cultures (en_US/en_UK) and "6,50€" in Spanish (es_ES).
+	 */
+	UFUNCTION(BlueprintPure, meta = (DisplayName = "AsCurrency"), Category = "Utilities|Text")
+	static FText AsCurrencyBase(int32 BaseValue, const FString& CurrencyCode);
 
 	// Default values are duplicated from FNumberFormattingOptions and should be replicated in all functions and in the struct when changed!
 	/* Converts a passed in integer to a text formatted as a currency */
-	UFUNCTION(BlueprintPure, meta=(DisplayName = "AsCurrency (int)", AdvancedDisplay = "1"), Category="Utilities|Text")
+	UFUNCTION(BlueprintPure, meta=(DisplayName = "AsCurrency (int) - DEPRECATED (use AsCurrency)"), Category="Utilities|Text")
 	static FText AsCurrency_Integer(int32 Value, TEnumAsByte<ERoundingMode> RoundingMode, bool bUseGrouping = true, int32 MinimumIntegralDigits = 1, int32 MaximumIntegralDigits = 324, int32 MinimumFractionalDigits = 0, int32 MaximumFractionalDigits = 3, const FString& CurrencyCode = TEXT(""));
 
 	// Default values are duplicated from FNumberFormattingOptions and should be replicated in all functions and in the struct when changed!
 	/* Converts a passed in float to a text formatted as a currency */
-	UFUNCTION(BlueprintPure, meta=(DisplayName = "AsCurrency (float)", AdvancedDisplay = "1"), Category="Utilities|Text")
+	UFUNCTION(BlueprintPure, meta=(DisplayName = "AsCurrency (float) - DEPRECATED (use AsCurrency)"), Category="Utilities|Text")
 	static FText AsCurrency_Float(float Value, TEnumAsByte<ERoundingMode> RoundingMode, bool bUseGrouping = true, int32 MinimumIntegralDigits = 1, int32 MaximumIntegralDigits = 324, int32 MinimumFractionalDigits = 0, int32 MaximumFractionalDigits = 3, const FString& CurrencyCode = TEXT(""));
 
 	// Default values are duplicated from FNumberFormattingOptions and should be replicated in all functions and in the struct when changed!
diff --git a/Engine/Source/Runtime/Engine/Private/Actor.cpp b/Engine/Source/Runtime/Engine/Private/Actor.cpp
--- a/Engine/Source/Runtime/Engine/Private/Actor.cpp
+++ b/Engine/Source/Runtime/Engine/Private/Actor.cpp
@@ -1505,17 +1505,6 @@ void AActor::GetAttachedActors(TArray<class AActor*>& OutActors) const
 			const bool bAllowShrinking = false;
 			USceneComponent* SceneComp = CompsToCheck.Pop(bAllowShrinking);
 
-			//////////////////////////////////////////////////////////////////////////
-			// ORION ONLY - DO NOT INTEGRATE INTO MAIN!
-			// Tracking down https://jira.ol.epicgames.net/browse/OR-3997
-			if(!ensure(SceneComp->IsValidLowLevel()))
-			{
-				UE_LOG(LogSpawn, Error, TEXT("GetAttachedActors called for actor [%s], is accessing an invalid scene component" ), *GetFullName() );
-				continue;
-			}
-			// ORION ONLY
-			//////////////////////////////////////////////////////////////////////////
-
 			// Add it to the 'checked' set, should not already be there!
 			if (!CheckedComps.Contains(SceneComp))
 			{
diff --git a/Engine/Source/Runtime/Engine/Private/Animation/AnimInstance.cpp b/Engine/Source/Runtime/Engine/Private/Animation/AnimInstance.cpp
--- a/Engine/Source/Runtime/Engine/Private/Animation/AnimInstance.cpp
+++ b/Engine/Source/Runtime/Engine/Private/Animation/AnimInstance.cpp
@@ -307,7 +307,7 @@ void UAnimInstance::UpdateMontage(float DeltaSeconds)
 	// time to test sync groups
 	for (auto& MontageInstance : MontageInstances)
 	{
-		if (MontageInstance->CanUseMarkerSync())
+		if (MontageInstance->bDidUseMarkerSyncThisTick)
 		{
 			const int32 GroupIndexToUse = MontageInstance->GetSyncGroupIndex();
 
diff --git a/Engine/Source/Runtime/Engine/Private/Animation/AnimMontage.cpp b/Engine/Source/Runtime/Engine/Private/Animation/AnimMontage.cpp
--- a/Engine/Source/Runtime/Engine/Private/Animation/AnimMontage.cpp
+++ b/Engine/Source/Runtime/Engine/Private/Animation/AnimMontage.cpp
@@ -1052,6 +1052,16 @@ void UAnimMontage::TickAssetPlayer(FAnimTickRecord& Instance, struct FAnimNotify
 
 				MarkerTickContext.MarkersPassedThisTick = *Instance.Montage.MarkersPassedThisTick;
 
+#if DO_CHECK
+				if(MarkerTickContext.MarkersPassedThisTick.Num() == 0)
+				{
+					const FMarkerSyncAnimPosition& StartPosition = MarkerTickContext.GetMarkerSyncStartPosition();
+					const FMarkerSyncAnimPosition& EndPosition = MarkerTickContext.GetMarkerSyncEndPosition();
+					checkf(StartPosition.NextMarkerName == EndPosition.NextMarkerName, TEXT("StartPosition NextMarker = %s, EndPosition NextMarker = %s, Asset = %s"), *StartPosition.NextMarkerName.ToString(), *EndPosition.NextMarkerName.ToString(), *Instance.SourceAsset->GetFullName());
+					checkf(StartPosition.PreviousMarkerName == EndPosition.PreviousMarkerName, TEXT("StartPosition PreviousMarker = %s, EndPosition PreviousMarker = %s, Asset = %s"), *StartPosition.PreviousMarkerName.ToString(), *EndPosition.PreviousMarkerName.ToString(), *Instance.SourceAsset->GetFullName());
+				}
+#endif
+
 				UE_LOG(LogAnimMarkerSync, Log, TEXT("Montage Leading SyncGroup: %s(%s) Start [%s], End [%s]"),
 					*GetNameSafe(this), *SyncGroup.ToString(), *MarkerTickContext.GetMarkerSyncStartPosition().ToString(), *MarkerTickContext.GetMarkerSyncEndPosition().ToString());
 			}
@@ -1644,8 +1654,8 @@ void FAnimMontageInstance::Advance(float DeltaTime, struct FRootMotionMovementPa
 			DeltaMoved = 0.f;
 			PreviousPosition = Position;
 
-			bool bCanUseMarkerSync = CanUseMarkerSync();
-			if (bCanUseMarkerSync)
+			bDidUseMarkerSyncThisTick = CanUseMarkerSync();
+			if (bDidUseMarkerSyncThisTick)
 			{
 				MarkersPassedThisTick.Reset();
 			}
@@ -1694,7 +1704,7 @@ void FAnimMontageInstance::Advance(float DeltaTime, struct FRootMotionMovementPa
 					float PrevPosition = Position;
 					Position = FMath::Clamp<float>(Position + ActualDeltaPos, CurrentSection.GetTime(), CurrentSectionEndPos);
 
-					if (bCanUseMarkerSync)
+					if (bDidUseMarkerSyncThisTick)
 					{
 						Montage->MarkerData.CollectMarkersInRange(PrevPosition, Position, MarkersPassedThisTick, ActualDeltaMove);
 					}
diff --git a/Engine/Source/Runtime/Engine/Private/Animation/AnimSequence.cpp b/Engine/Source/Runtime/Engine/Private/Animation/AnimSequence.cpp
--- a/Engine/Source/Runtime/Engine/Private/Animation/AnimSequence.cpp
+++ b/Engine/Source/Runtime/Engine/Private/Animation/AnimSequence.cpp
@@ -158,6 +158,9 @@ void UAnimSequence::GetAssetRegistryTags(TArray<FAssetRegistryTag>& OutTags) con
 	}
 #endif
 
+	OutTags.Add(FAssetRegistryTag(TEXT("Compression Ratio"), FString::Printf(TEXT("%.03f"), (float)GetApproxCompressedSize() / (float)GetApproxRawSize()), FAssetRegistryTag::TT_Numerical));
+	OutTags.Add(FAssetRegistryTag(TEXT("Compressed Size (KB)"), FString::Printf(TEXT("%.02f"), (float)GetApproxCompressedSize() / 1024.0f), FAssetRegistryTag::TT_Numerical));
+
 	Super::GetAssetRegistryTags(OutTags);
 }
 
@@ -966,7 +969,8 @@ FTransform UAnimSequence::ExtractRootMotion(float StartTime, float DeltaTime, bo
 			const ETypeAdvanceAnim AdvanceType = FAnimationRuntime::AdvanceTime(false, DesiredDeltaMove, CurrentPosition, SequenceLength);
 
 			// Verify position assumptions
-			check(bPlayingBackwards ? (CurrentPosition <= PreviousPosition) : (CurrentPosition >= PreviousPosition));
+			ensureMsgf(bPlayingBackwards ? (CurrentPosition <= PreviousPosition) : (CurrentPosition >= PreviousPosition), TEXT("in Animation %s(Skeleton %s) : bPlayingBackwards(%d), PreviousPosition(%0.2f), Current Position(%0.2f)"),
+				*GetName(), *GetNameSafe(GetSkeleton()), bPlayingBackwards, PreviousPosition, CurrentPosition);
 
 			RootMotionParams.Accumulate(ExtractRootMotionFromRange(PreviousPosition, CurrentPosition));
 
@@ -4321,7 +4325,7 @@ void AdvanceMarkerBackwards(int32& Marker, FName MarkerToFind, bool bLooping, co
 
 	if (!AuthoredSyncMarkers.IsValidIndex(Marker) || (AuthoredSyncMarkers[Marker].MarkerName != MarkerToFind))
 	{
-		Marker = INDEX_NONE;
+		Marker = MarkerIndexSpecialValues::AnimationBoundary;
 	}
 }
 
@@ -4338,7 +4342,11 @@ void UAnimSequence::ValidateCurrentPosition(const FMarkerSyncAnimPosition& Posit
 		if (!MarkerMatchesPosition(this, PreviousMarker.MarkerIndex, Position.PreviousMarkerName))
 		{
 			AdvanceMarkerForwards(PreviousMarker.MarkerIndex, Position.PreviousMarkerName, bLooping, AuthoredSyncMarkers);
-			NextMarker.MarkerIndex = (PreviousMarker.MarkerIndex + 1) % AuthoredSyncMarkers.Num();
+			NextMarker.MarkerIndex = (PreviousMarker.MarkerIndex + 1);
+			if(NextMarker.MarkerIndex >= AuthoredSyncMarkers.Num())
+			{
+				NextMarker.MarkerIndex = bLooping ? NextMarker.MarkerIndex % AuthoredSyncMarkers.Num() : MarkerIndexSpecialValues::AnimationBoundary;
+			}
 		}
 
 		if (!MarkerMatchesPosition(this, NextMarker.MarkerIndex, Position.NextMarkerName))
@@ -4352,7 +4360,14 @@ void UAnimSequence::ValidateCurrentPosition(const FMarkerSyncAnimPosition& Posit
 		if (!MarkerMatchesPosition(this, NextMarker.MarkerIndex, Position.NextMarkerName))
 		{
 			AdvanceMarkerBackwards(NextMarker.MarkerIndex, Position.NextMarkerName, bLooping, AuthoredSyncMarkers);
-			PreviousMarker.MarkerIndex = (NextMarker.MarkerIndex ? NextMarker.MarkerIndex : MarkerRange) - 1;
+			if(NextMarker.MarkerIndex == MarkerIndexSpecialValues::AnimationBoundary || (NextMarker.MarkerIndex == 0 && bLooping))
+			{
+				PreviousMarker.MarkerIndex = AuthoredSyncMarkers.Num() - 1;
+			}
+			else
+			{
+				PreviousMarker.MarkerIndex = NextMarker.MarkerIndex - 1;
+			}
 		}
 		if (!MarkerMatchesPosition(this, PreviousMarker.MarkerIndex, Position.PreviousMarkerName))
 		{
@@ -4556,6 +4571,8 @@ float UAnimSequence::GetCurrentTimeFromMarkers(FMarkerPair& PrevMarker, FMarkerP
 	{
 		CurrentTime += SequenceLength;
 	}
+	CurrentTime = FMath::Clamp<float>(CurrentTime, 0, SequenceLength);
+	
 	PrevMarker.TimeToMarker = PrevTime - CurrentTime;
 	NextMarker.TimeToMarker = NextTime - CurrentTime;
 	return CurrentTime;
diff --git a/Engine/Source/Runtime/Engine/Private/Animation/AnimSequenceBase.cpp b/Engine/Source/Runtime/Engine/Private/Animation/AnimSequenceBase.cpp
--- a/Engine/Source/Runtime/Engine/Private/Animation/AnimSequenceBase.cpp
+++ b/Engine/Source/Runtime/Engine/Private/Animation/AnimSequenceBase.cpp
@@ -140,8 +140,8 @@ void UAnimSequenceBase::GetAnimNotifies(const float& StartTime, const float& Del
 		const ETypeAdvanceAnim AdvanceType = FAnimationRuntime::AdvanceTime(false, DesiredDeltaMove, CurrentPosition, SequenceLength);
 
 		// Verify position assumptions
-		ensureMsgf(bPlayingBackwards ? (CurrentPosition <= PreviousPosition) : (CurrentPosition >= PreviousPosition), TEXT("in Animsequence %s : bPlayingBackwards(%d), PreviousPosition(%0.2f), Current Position(%0.2f)"), 
-			*GetName(), bPlayingBackwards, PreviousPosition, CurrentPosition);
+		ensureMsgf(bPlayingBackwards ? (CurrentPosition <= PreviousPosition) : (CurrentPosition >= PreviousPosition), TEXT("in Animation %s(Skeleton %s) : bPlayingBackwards(%d), PreviousPosition(%0.2f), Current Position(%0.2f)"), 
+			*GetName(), *GetNameSafe(GetSkeleton()), bPlayingBackwards, PreviousPosition, CurrentPosition);
 		
 		GetAnimNotifiesFromDeltaPositions(PreviousPosition, CurrentPosition, OutActiveNotifies);
 	
diff --git a/Engine/Source/Runtime/Engine/Private/Camera/CameraAnimInst.cpp b/Engine/Source/Runtime/Engine/Private/Camera/CameraAnimInst.cpp
--- a/Engine/Source/Runtime/Engine/Private/Camera/CameraAnimInst.cpp
+++ b/Engine/Source/Runtime/Engine/Private/Camera/CameraAnimInst.cpp
@@ -21,7 +21,6 @@ UCameraAnimInst::UCameraAnimInst(const FObjectInitializer& ObjectInitializer)
 	bAutoReleaseWhenFinished = true;
 	PlayRate = 1.0f;
 	TransientScaleModifier = 1.0f;
-	InitialFOV = -1.f;
 	PlaySpace = ECameraAnimPlaySpace::CameraLocal;
 
 	InterpGroupInst = CreateDefaultSubobject<UInterpGroupInst>(TEXT("InterpGroupInst0"));
@@ -222,16 +221,13 @@ void UCameraAnimInst::Play(UCameraAnim* Anim, class AActor* CamActor, float InRa
 			InitialCamToWorld = FTransform(OutRot, OutLoc);		// @todo, store inverted since that's how we use it?
 
 			// find FOV track if it exists, else just use the fov saved in the anim
-			if (InitialFOV <= 0.f)
+			InitialFOV = Anim->BaseFOV;
+			for (int32 Idx = 0; Idx < InterpGroupInst->TrackInst.Num(); ++Idx)
 			{
-				InitialFOV = Anim->BaseFOV;
-				for (int32 Idx = 0; Idx < InterpGroupInst->TrackInst.Num(); ++Idx)
+				UInterpTrackFloatProp* const FloatTrack = Cast<UInterpTrackFloatProp>(CamAnim->CameraInterpGroup->InterpTracks[Idx]);
+				if (FloatTrack && (FloatTrack->PropertyName == NAME_CameraComponentFieldOfViewPropertyName))
 				{
-					UInterpTrackFloatProp* const FloatTrack = Cast<UInterpTrackFloatProp>(CamAnim->CameraInterpGroup->InterpTracks[Idx]);
-					if (FloatTrack && (FloatTrack->PropertyName == NAME_CameraComponentFieldOfViewPropertyName))
-					{
-						InitialFOV = FloatTrack->EvalSub(0, 0.f);
-					}
+					InitialFOV = FloatTrack->EvalSub(0, 0.f);
 				}
 			}
 		}
diff --git a/Engine/Source/Runtime/Engine/Private/ChartCreation.cpp b/Engine/Source/Runtime/Engine/Private/ChartCreation.cpp
--- a/Engine/Source/Runtime/Engine/Private/ChartCreation.cpp
+++ b/Engine/Source/Runtime/Engine/Private/ChartCreation.cpp
@@ -1144,7 +1144,10 @@ void UEngine::DumpFPSChartToAnalyticsParams(float TotalTime, float DeltaTime, in
 	FPlatformMisc::GetOSVersions(OSMajor, OSMinor);
 
 	// Get settings info
-	const Scalability::FQualityLevels& Quality = GEngine->GetGameUserSettings()->ScalabilityQuality;
+	UGameUserSettings* UserSettingsObj = GEngine->GetGameUserSettings();
+	check(UserSettingsObj);
+
+	const Scalability::FQualityLevels& Quality = UserSettingsObj->ScalabilityQuality;
 
 	// Add non-bucket params
 	InParamArray.Add(FAnalyticsEventAttribute(TEXT("ChangeList"), GetChangeListNumberForPerfTesting()));
@@ -1153,8 +1156,8 @@ void UEngine::DumpFPSChartToAnalyticsParams(float TotalTime, float DeltaTime, in
 
 	InParamArray.Add(FAnalyticsEventAttribute(TEXT("Platform"), FString::Printf(TEXT("%s"), FPlatformProperties::PlatformName())));
 	InParamArray.Add(FAnalyticsEventAttribute(TEXT("OS"), FString::Printf(TEXT("%s %s"), *OSMajor, *OSMinor)));
-	InParamArray.Add(FAnalyticsEventAttribute(TEXT("CPU"), FString::Printf(TEXT("%s %s"), *FPlatformMisc::GetCPUVendor(), *FPlatformMisc::GetCPUBrand())));
-	InParamArray.Add(FAnalyticsEventAttribute(TEXT("GPU"), FPlatformMisc::GetPrimaryGPUBrand()));
+	InParamArray.Add(FAnalyticsEventAttribute(TEXT("CPU"), FString::Printf(TEXT("%s %s"), *FPlatformMisc::GetCPUVendor().Trim().TrimTrailing(), *FPlatformMisc::GetCPUBrand().Trim().TrimTrailing())));
+	InParamArray.Add(FAnalyticsEventAttribute(TEXT("GPU"), FPlatformMisc::GetPrimaryGPUBrand().Trim().TrimTrailing()));
 	InParamArray.Add(FAnalyticsEventAttribute(TEXT("ResolutionQuality"), Quality.ResolutionQuality));
 	InParamArray.Add(FAnalyticsEventAttribute(TEXT("ViewDistanceQuality"), Quality.ViewDistanceQuality));
 	InParamArray.Add(FAnalyticsEventAttribute(TEXT("AntiAliasingQuality"), Quality.AntiAliasingQuality));
diff --git a/Engine/Source/Runtime/Engine/Private/Components/MeshComponent.cpp b/Engine/Source/Runtime/Engine/Private/Components/MeshComponent.cpp
--- a/Engine/Source/Runtime/Engine/Private/Components/MeshComponent.cpp
+++ b/Engine/Source/Runtime/Engine/Private/Components/MeshComponent.cpp
@@ -257,7 +257,7 @@ void UMeshComponent::CacheMaterialParameterNameIndices()
 			for (FName& ParameterName : OutParameterNames)
 			{
 				// Add or retrieve TMap entry for this scalar parameter name
-				TArray<int32>& MaterialIndices = ScalarParameterMaterialIndices.Add(ParameterName);
+				TArray<int32>& MaterialIndices = ScalarParameterMaterialIndices.FindOrAdd(ParameterName);
 				// Add the corresponding material index
 				MaterialIndices.Add(MaterialIndex);
 			}
@@ -271,7 +271,7 @@ void UMeshComponent::CacheMaterialParameterNameIndices()
 			for (FName& ParameterName : OutParameterNames)
 			{
 				// Add or retrieve TMap entry for this vector parameter name
-				TArray<int32>& MaterialIndices = VectorParameterMaterialIndices.Add(ParameterName);
+				TArray<int32>& MaterialIndices = VectorParameterMaterialIndices.FindOrAdd(ParameterName);
 				// Add the corresponding material index
 				MaterialIndices.Add(MaterialIndex);
 			}
diff --git a/Engine/Source/Runtime/Engine/Private/Components/SkeletalMeshComponent.cpp b/Engine/Source/Runtime/Engine/Private/Components/SkeletalMeshComponent.cpp
--- a/Engine/Source/Runtime/Engine/Private/Components/SkeletalMeshComponent.cpp
+++ b/Engine/Source/Runtime/Engine/Private/Components/SkeletalMeshComponent.cpp
@@ -976,6 +976,19 @@ void USkeletalMeshComponent::RecalcRequiredBones(int32 LODIndex)
 		MergeInBoneIndexArrays(RequiredBones, ForceAnimatedSocketBones);
 	}
 
+	// Gather any bones referenced by shadow shapes
+	if (FSkeletalMeshSceneProxy* SkeletalMeshProxy = (FSkeletalMeshSceneProxy*)SceneProxy)
+	{
+		TArray<FBoneIndexType> ShadowShapeBones;
+		SkeletalMeshProxy->GetShadowShapeBoneIndices(ShadowShapeBones);
+
+		if (ShadowShapeBones.Num())
+		{
+			// Sort in hierarchy order then merge to required bones array
+			ShadowShapeBones.Sort();
+			MergeInBoneIndexArrays(RequiredBones, ShadowShapeBones);
+		}
+	}
 
 	// Ensure that we have a complete hierarchy down to those bones.
 	FAnimationRuntime::EnsureParentsPresent(RequiredBones, SkeletalMesh);
diff --git a/Engine/Source/Runtime/Engine/Private/DemoNetDriver.cpp b/Engine/Source/Runtime/Engine/Private/DemoNetDriver.cpp
--- a/Engine/Source/Runtime/Engine/Private/DemoNetDriver.cpp
+++ b/Engine/Source/Runtime/Engine/Private/DemoNetDriver.cpp
@@ -36,6 +36,7 @@ static TAutoConsoleVariable<int32> CVarDemoFastForwardDestroyTearOffActors( TEXT
 static TAutoConsoleVariable<int32> CVarDemoFastForwardSkipRepNotifies( TEXT( "demo.FastForwardSkipRepNotifies" ), 1, TEXT( "If true, the driver will optimize fast-forwarding by deferring calls to RepNotify functions until the fast-forward is complete. " ) );
 static TAutoConsoleVariable<int32> CVarDemoQueueCheckpointChannels( TEXT( "demo.QueueCheckpointChannels" ), 1, TEXT( "If true, the driver will put all channels created during checkpoint loading into queuing mode, to amortize the cost of spawning new actors across multiple frames." ) );
 static TAutoConsoleVariable<int32> CVarUseAdaptiveReplayUpdateFrequency( TEXT( "demo.UseAdaptiveReplayUpdateFrequency" ), 0, TEXT( "If 1, NetUpdateFrequency will be calculated based on how often actors actually write something when recording to a replay" ) );
+static TAutoConsoleVariable<int32> CVarDemoAsyncLoadWorld( TEXT( "demo.AsyncLoadWorld" ), 0, TEXT( "If 1, we will use seamless server travel to load the replay world asynchronously" ) );
 
 #if !(UE_BUILD_SHIPPING || UE_BUILD_TEST)
 static TAutoConsoleVariable<int32> CVarDemoForceFailure( TEXT( "demo.ForceFailure" ), 0, TEXT( "" ) );
@@ -562,90 +563,111 @@ bool UDemoNetDriver::InitConnectInternal( FString& Error )
 		UE_LOG( LogDemo, Log, TEXT( "Starting demo playback with streaming demo, metadata file not found. Filename: %s, Version %i" ), *DemoFilename, DemoHeader.Version );
 	}
 	
-	// Bypass UDemoPendingNetLevel
-	FString LoadMapError;
-
-	FURL DemoURL;
-	DemoURL.Map = DemoHeader.LevelName;
+	if ( CVarDemoAsyncLoadWorld.GetValueOnGameThread() > 0 )
+	{
+		TArray<AController*> Controllers;
+		for ( FConstControllerIterator Iterator = World->GetControllerIterator(); Iterator; ++Iterator )
+		{
+			Controllers.Add( *Iterator );
+		}
 
-	FWorldContext * WorldContext = GEngine->GetWorldContextFromWorld( GetWorld() );
+		for ( int i = 0; i < Controllers.Num(); i++ )
+		{
+			if ( Controllers[i] )
+			{
+				Controllers[i]->Destroy();
+			}
+		}
 
-	if ( WorldContext == NULL )
-	{
-		Error = FString::Printf( TEXT( "No world context" ), *DemoFilename );
-		UE_LOG( LogDemo, Error, TEXT( "UDemoNetDriver::InitConnect: %s" ), *Error );
-		GameInstance->HandleDemoPlaybackFailure( EDemoPlayFailure::Generic, FString( TEXT( "No world context" ) ) );
-		return false;
+		World->ServerTravel( DemoHeader.LevelName, true, false );
 	}
+	else
+	{
+		// Bypass UDemoPendingNetLevel
+		FString LoadMapError;
 
-	GetWorld()->DemoNetDriver = NULL;
-	SetWorld( NULL );
+		FURL DemoURL;
+		DemoURL.Map = DemoHeader.LevelName;
 
-	auto NewPendingNetGame = NewObject<UDemoPendingNetGame>();
+		FWorldContext * WorldContext = GEngine->GetWorldContextFromWorld( GetWorld() );
 
-	NewPendingNetGame->DemoNetDriver = this;
+		if ( WorldContext == NULL )
+		{
+			Error = FString::Printf( TEXT( "No world context" ), *DemoFilename );
+			UE_LOG( LogDemo, Error, TEXT( "UDemoNetDriver::InitConnect: %s" ), *Error );
+			GameInstance->HandleDemoPlaybackFailure( EDemoPlayFailure::Generic, FString( TEXT( "No world context" ) ) );
+			return false;
+		}
 
-	WorldContext->PendingNetGame = NewPendingNetGame;
+		GetWorld()->DemoNetDriver = NULL;
+		SetWorld( NULL );
 
-	bool bSuccess = GEngine->LoadMap( *WorldContext, DemoURL, NewPendingNetGame, LoadMapError );
+		auto NewPendingNetGame = NewObject<UDemoPendingNetGame>();
 
-#if !(UE_BUILD_SHIPPING || UE_BUILD_TEST)
-	if ( CVarDemoForceFailure.GetValueOnGameThread() == 2 )
-	{
-		bSuccess = false;
-	}
-#endif
+		NewPendingNetGame->DemoNetDriver = this;
 
-	if ( !bSuccess )
-	{
-		StopDemo();
+		WorldContext->PendingNetGame = NewPendingNetGame;
 
-		// If we don't have a world that means we failed loading the new world.
-		// Since there is no world, we must free the net driver ourselves
-		// Technically the pending net game should handle it, but things aren't quite setup properly to handle that either
-		if ( WorldContext->World() == NULL )
+		bool bSuccess = GEngine->LoadMap( *WorldContext, DemoURL, NewPendingNetGame, LoadMapError );
+
+#if !(UE_BUILD_SHIPPING || UE_BUILD_TEST)
+		if ( CVarDemoForceFailure.GetValueOnGameThread() == 2 )
 		{
-			GEngine->DestroyNamedNetDriver( WorldContext->PendingNetGame, NetDriverName );
+			bSuccess = false;
 		}
+#endif
 
-		WorldContext->PendingNetGame = NULL;
+		if ( !bSuccess )
+		{
+			StopDemo();
 
-		GEngine->BrowseToDefaultMap( *WorldContext );
+			// If we don't have a world that means we failed loading the new world.
+			// Since there is no world, we must free the net driver ourselves
+			// Technically the pending net game should handle it, but things aren't quite setup properly to handle that either
+			if ( WorldContext->World() == NULL )
+			{
+				GEngine->DestroyNamedNetDriver( WorldContext->PendingNetGame, NetDriverName );
+			}
 
-		Error = LoadMapError;
-		UE_LOG( LogDemo, Error, TEXT( "UDemoNetDriver::InitConnect: LoadMap failed: failed: %s" ), *Error );
-		GameInstance->HandleDemoPlaybackFailure( EDemoPlayFailure::Generic, FString( TEXT( "LoadMap failed" ) ) );
-		return false;
-	}
+			WorldContext->PendingNetGame = NULL;
 
-	SetWorld( WorldContext->World() );
-	WorldContext->World()->DemoNetDriver = this;
-	WorldContext->PendingNetGame = NULL;
+			GEngine->BrowseToDefaultMap( *WorldContext );
 
-	// Read meta data, if it exists
-	for ( int32 i = 0; i < MetadataHeader.NumStreamingLevels; ++i )
-	{
-		ULevelStreamingKismet* StreamingLevel = NewObject<ULevelStreamingKismet>(GetWorld(), NAME_None, RF_NoFlags, NULL);
+			Error = LoadMapError;
+			UE_LOG( LogDemo, Error, TEXT( "UDemoNetDriver::InitConnect: LoadMap failed: failed: %s" ), *Error );
+			GameInstance->HandleDemoPlaybackFailure( EDemoPlayFailure::Generic, FString( TEXT( "LoadMap failed" ) ) );
+			return false;
+		}
+
+		SetWorld( WorldContext->World() );
+		WorldContext->World()->DemoNetDriver = this;
+		WorldContext->PendingNetGame = NULL;
+
+		// Read meta data, if it exists
+		for ( int32 i = 0; i < MetadataHeader.NumStreamingLevels; ++i )
+		{
+			ULevelStreamingKismet* StreamingLevel = NewObject<ULevelStreamingKismet>( GetWorld(), NAME_None, RF_NoFlags, NULL );
 
-		StreamingLevel->bShouldBeLoaded		= true;
-		StreamingLevel->bShouldBeVisible	= true;
-		StreamingLevel->bShouldBlockOnLoad	= false;
-		StreamingLevel->bInitiallyLoaded	= true;
-		StreamingLevel->bInitiallyVisible	= true;
+			StreamingLevel->bShouldBeLoaded		= true;
+			StreamingLevel->bShouldBeVisible	= true;
+			StreamingLevel->bShouldBlockOnLoad	= false;
+			StreamingLevel->bInitiallyLoaded	= true;
+			StreamingLevel->bInitiallyVisible	= true;
 
-		FString PackageName;
-		FString PackageNameToLoad;
+			FString PackageName;
+			FString PackageNameToLoad;
 
-		(*MetadataAr) << PackageName;
-		(*MetadataAr) << PackageNameToLoad;
-		(*MetadataAr) << StreamingLevel->LevelTransform;
+			( *MetadataAr ) << PackageName;
+			( *MetadataAr ) << PackageNameToLoad;
+			( *MetadataAr ) << StreamingLevel->LevelTransform;
 
-		StreamingLevel->PackageNameToLoad = FName( *PackageNameToLoad );
-		StreamingLevel->SetWorldAssetByPackageName( FName( *PackageName ) );
+			StreamingLevel->PackageNameToLoad = FName( *PackageNameToLoad );
+			StreamingLevel->SetWorldAssetByPackageName( FName( *PackageName ) );
 
-		GetWorld()->StreamingLevels.Add( StreamingLevel );
+			GetWorld()->StreamingLevels.Add( StreamingLevel );
 
-		UE_LOG( LogDemo, Log, TEXT( "  Loading streamingLevel: %s, %s" ), *PackageName, *PackageNameToLoad );
+			UE_LOG( LogDemo, Log, TEXT( "  Loading streamingLevel: %s, %s" ), *PackageName, *PackageNameToLoad );
+		}
 	}
 
 	return true;
@@ -1679,6 +1701,11 @@ void UDemoNetDriver::AddUserToReplay( const FString& UserString )
 
 void UDemoNetDriver::TickDemoPlayback( float DeltaSeconds )
 {
+	if ( World && World->IsInSeamlessTravel() )
+	{
+		return;
+	}
+
 	if ( !IsPlaying() )
 	{
 		return;
@@ -2342,7 +2369,6 @@ void UDemoNetDriver::NotifyGotoTimeFinished(bool bWasSuccessful)
 	}
 }
 
-
 /*-----------------------------------------------------------------------------
 	UDemoPendingNetGame.
 -----------------------------------------------------------------------------*/
diff --git a/Engine/Source/Runtime/Engine/Private/GameEngine.cpp b/Engine/Source/Runtime/Engine/Private/GameEngine.cpp
--- a/Engine/Source/Runtime/Engine/Private/GameEngine.cpp
+++ b/Engine/Source/Runtime/Engine/Private/GameEngine.cpp
@@ -789,9 +789,6 @@ bool UGameEngine::Exec( UWorld* InWorld, const TCHAR* Cmd, FOutputDevice& Ar )
 
 bool UGameEngine::HandleExitCommand( const TCHAR* Cmd, FOutputDevice& Ar )
 {
-	// temp hack: debugging a crash on exit that doesnt repro. Need to find the calling code
-	ensureMsgf(0, TEXT("Debugging graceful exit from the game during load"));
-	
 	Ar.Log( TEXT("Closing by request") );
 	FGameDelegates::Get().GetExitCommandDelegate().Broadcast();
 	FPlatformMisc::RequestExit( 0 );
@@ -868,7 +865,7 @@ float UGameEngine::GetMaxTickRate(float DeltaTime, bool bAllowFrameRateSmoothing
 			if( NetDriver && (NetDriver->GetNetMode() == NM_DedicatedServer || (NetDriver->GetNetMode() == NM_ListenServer && NetDriver->bClampListenServerTickRate)))
 			{
 				// We're a dedicated server, use the LAN or Net tick rate.
-				MaxTickRate = FMath::Clamp( NetDriver->NetServerMaxTickRate, 10, 120 );
+				MaxTickRate = FMath::Clamp( NetDriver->NetServerMaxTickRate, 1, 1000 );
 			}
 			/*else if( NetDriver && NetDriver->ServerConnection )
 			{
diff --git a/Engine/Source/Runtime/Engine/Private/KismetInputLibrary.cpp b/Engine/Source/Runtime/Engine/Private/KismetInputLibrary.cpp
--- a/Engine/Source/Runtime/Engine/Private/KismetInputLibrary.cpp
+++ b/Engine/Source/Runtime/Engine/Private/KismetInputLibrary.cpp
@@ -28,6 +28,11 @@ bool UKismetInputLibrary::EqualEqual_KeyKey(FKey A, FKey B)
 	return A == B;
 }
 
+bool UKismetInputLibrary::EqualEqual_InputChordInputChord( FInputChord A, FInputChord B )
+{
+	return A == B;
+}
+
 bool UKismetInputLibrary::Key_IsModifierKey(const FKey& Key)
 {
 	return Key.IsModifierKey();
diff --git a/Engine/Source/Runtime/Engine/Private/KismetTextLibrary.cpp b/Engine/Source/Runtime/Engine/Private/KismetTextLibrary.cpp
--- a/Engine/Source/Runtime/Engine/Private/KismetTextLibrary.cpp
+++ b/Engine/Source/Runtime/Engine/Private/KismetTextLibrary.cpp
@@ -145,7 +145,13 @@ FText UKismetTextLibrary::Conv_FloatToText(float Value, TEnumAsByte<ERoundingMod
 	return FText::AsNumber(Value, &NumberFormatOptions);
 }
 
+FText UKismetTextLibrary::AsCurrencyBase(int32 BaseValue, const FString& CurrencyCode)
+{
+	return FText::AsCurrencyBase(BaseValue, CurrencyCode);
+}
 
+PRAGMA_DISABLE_DEPRECATION_WARNINGS
+// FIXME: we need to deprecate this kismet api too
 FText UKismetTextLibrary::AsCurrency_Integer(int32 Value, TEnumAsByte<ERoundingMode> RoundingMode, bool bUseGrouping/* = true*/, int32 MinimumIntegralDigits/* = 1*/, int32 MaximumIntegralDigits/* = 324*/, int32 MinimumFractionalDigits/* = 0*/, int32 MaximumFractionalDigits/* = 3*/, const FString& CurrencyCode)
 {
 	FNumberFormattingOptions NumberFormatOptions;
@@ -155,11 +161,8 @@ FText UKismetTextLibrary::AsCurrency_Integer(int32 Value, TEnumAsByte<ERoundingM
 	NumberFormatOptions.MaximumIntegralDigits = MaximumIntegralDigits;
 	NumberFormatOptions.MinimumFractionalDigits = MinimumFractionalDigits;
 	NumberFormatOptions.MaximumFractionalDigits = MaximumFractionalDigits;
-
 	return FText::AsCurrency(Value, CurrencyCode, &NumberFormatOptions);
 }
-
-
 FText UKismetTextLibrary::AsCurrency_Float(float Value, TEnumAsByte<ERoundingMode> RoundingMode, bool bUseGrouping/* = true*/, int32 MinimumIntegralDigits/* = 1*/, int32 MaximumIntegralDigits/* = 324*/, int32 MinimumFractionalDigits/* = 0*/, int32 MaximumFractionalDigits/* = 3*/, const FString& CurrencyCode)
 {
 	FNumberFormattingOptions NumberFormatOptions;
@@ -169,10 +172,9 @@ FText UKismetTextLibrary::AsCurrency_Float(float Value, TEnumAsByte<ERoundingMod
 	NumberFormatOptions.MaximumIntegralDigits = MaximumIntegralDigits;
 	NumberFormatOptions.MinimumFractionalDigits = MinimumFractionalDigits;
 	NumberFormatOptions.MaximumFractionalDigits = MaximumFractionalDigits;
-
 	return FText::AsCurrency(Value, CurrencyCode, &NumberFormatOptions);
 }
-
+PRAGMA_ENABLE_DEPRECATION_WARNINGS
 
 FText UKismetTextLibrary::AsPercent_Float(float Value, TEnumAsByte<ERoundingMode> RoundingMode, bool bUseGrouping/* = true*/, int32 MinimumIntegralDigits/* = 1*/, int32 MaximumIntegralDigits/* = 324*/, int32 MinimumFractionalDigits/* = 0*/, int32 MaximumFractionalDigits/* = 3*/)
 {
diff --git a/Engine/Source/Runtime/Engine/Private/Particles/ParticleComponents.cpp b/Engine/Source/Runtime/Engine/Private/Particles/ParticleComponents.cpp
--- a/Engine/Source/Runtime/Engine/Private/Particles/ParticleComponents.cpp
+++ b/Engine/Source/Runtime/Engine/Private/Particles/ParticleComponents.cpp
@@ -275,7 +275,11 @@ void UParticleLODLevel::PostLoad()
 		(GetOuter() ? (GetOuter()->GetOuter() ? *(GetOuter()->GetOuter()->GetPathName()) : *(GetOuter()->GetPathName())) : TEXT("???")));
 #endif // WITH_EDITORONLY_DATA
 
-	// Fix up emitters that have bEnabled flag set incorrectly
+	{
+		RequiredModule->ConditionalPostLoad();
+	}
+	// shouldn't ever set another UObjects serialized variable in post load
+	// this causes determinisitc cooking issues due to load order being different
 	/*if (RequiredModule)
 	{
 		RequiredModule->ConditionalPostLoad();
diff --git a/Engine/Source/Runtime/Engine/Private/PlayerCameraManager.cpp b/Engine/Source/Runtime/Engine/Private/PlayerCameraManager.cpp
--- a/Engine/Source/Runtime/Engine/Private/PlayerCameraManager.cpp
+++ b/Engine/Source/Runtime/Engine/Private/PlayerCameraManager.cpp
@@ -336,18 +336,7 @@ void APlayerCameraManager::ApplyAnimToCamera(ACameraActor const* AnimatedCamActo
 	// fov
 	const float FOVMin = 5.f;
 	const float FOVMax = 170.f;
-
-	// Interp the FOV toward the camera component's FOV based on Scale
-	const int32 DesiredDirection = FMath::Sign(AnimatedCamActor->GetCameraComponent()->FieldOfView - InOutPOV.FOV);
-	const int32 InitialDirection = FMath::Sign(AnimatedCamActor->GetCameraComponent()->FieldOfView - AnimInst->InitialFOV);
-	if (DesiredDirection != InitialDirection)
-	{
-		InOutPOV.FOV = FMath::Clamp(InOutPOV.FOV + ((AnimatedCamActor->GetCameraComponent()->FieldOfView - InOutPOV.FOV) * Scale), InOutPOV.FOV, AnimatedCamActor->GetCameraComponent()->FieldOfView);
-	}
-	else
-	{
-		InOutPOV.FOV = FMath::Clamp<float>(InOutPOV.FOV + ((AnimatedCamActor->GetCameraComponent()->FieldOfView - AnimInst->InitialFOV) * Scale), AnimatedCamActor->GetCameraComponent()->FieldOfView, AnimInst->InitialFOV);
-	}
+	InOutPOV.FOV += (AnimatedCamActor->GetCameraComponent()->FieldOfView - AnimInst->InitialFOV) * Scale;
 	InOutPOV.FOV = FMath::Clamp<float>(InOutPOV.FOV, FOVMin, FOVMax);
 
 	// postprocess
@@ -413,7 +402,6 @@ UCameraAnimInst* APlayerCameraManager::PlayCameraAnim(UCameraAnim* Anim, float R
 		if (Inst)
 		{
 			Inst->LastCameraLoc = FVector::ZeroVector;		// clear LastCameraLoc
-			Inst->InitialFOV = ViewTarget.POV.FOV;
 			Inst->Play(Anim, AnimCameraActor, Rate, Scale, BlendInTime, BlendOutTime, bLoop, bRandomStartTime, Duration);
 			Inst->SetPlaySpace(PlaySpace, UserPlaySpaceRot);
 			return Inst;
diff --git a/Engine/Source/Runtime/Engine/Private/SceneView.cpp b/Engine/Source/Runtime/Engine/Private/SceneView.cpp
--- a/Engine/Source/Runtime/Engine/Private/SceneView.cpp
+++ b/Engine/Source/Runtime/Engine/Private/SceneView.cpp
@@ -222,7 +222,8 @@ static TAutoConsoleVariable<float> CVarSceneColorFringeMax(
 	TEXT("r.SceneColorFringe.Max"),
 	-1.0f,
 	TEXT("Allows to clamp the postprocess setting (in percent, Scene chromatic aberration / color fringe to simulate an artifact that happens in real-world lens, mostly visible in the image corners)\n")
-	TEXT("-1: don't clamp (default)"),
+	TEXT("-1: don't clamp (default)\n")
+	TEXT("-2: to test extreme fringe"),
 	ECVF_Scalability | ECVF_RenderThreadSafe);
 
 static TAutoConsoleVariable<int32> CVarTonemapperQuality(
@@ -1530,6 +1531,10 @@ void FSceneView::EndFinalPostprocessSettings(const FSceneViewInitOptions& ViewIn
 		{
 			FinalPostProcessSettings.SceneFringeIntensity = FMath::Min(FinalPostProcessSettings.SceneFringeIntensity, Value);
 		}
+		else if (Value == -2.0f)
+		{
+			FinalPostProcessSettings.SceneFringeIntensity = 5.0f;
+		}
 
 		if(!Family->EngineShowFlags.SceneColorFringe || !Family->EngineShowFlags.CameraImperfections)
 		{
diff --git a/Engine/Source/Runtime/Engine/Private/UnrealClient.cpp b/Engine/Source/Runtime/Engine/Private/UnrealClient.cpp
--- a/Engine/Source/Runtime/Engine/Private/UnrealClient.cpp
+++ b/Engine/Source/Runtime/Engine/Private/UnrealClient.cpp
@@ -438,7 +438,7 @@ int32 FStatUnitData::DrawStat(FViewport* InViewport, FCanvas* InCanvas, int32 In
 				static const auto CVar = IConsoleManager::Get().FindConsoleVariable(TEXT("r.PS4StallsOnMarkers"));
 				int32 Value = CVar->GetInt();
 
-				if (!Value)
+				if (Value)
 				{
 					// good to get Razor aligned GPU profiling but bad for high fps
 					Warnings += TEXT(" r.PS4StallsOnMarkers");
diff --git a/Engine/Source/Runtime/Engine/Private/UnrealEngine.cpp b/Engine/Source/Runtime/Engine/Private/UnrealEngine.cpp
--- a/Engine/Source/Runtime/Engine/Private/UnrealEngine.cpp
+++ b/Engine/Source/Runtime/Engine/Private/UnrealEngine.cpp
@@ -6536,6 +6536,10 @@ void UEngine::OnHardwareSurveyComplete(const FHardwareSurveyResults& SurveyResul
 		HardwareStatsAttribs.Add(FAnalyticsEventAttribute(TEXT( "GPU.Name" ), MainGPUName));
 		HardwareStatsAttribs.Add(FAnalyticsEventAttribute(TEXT( "GPU.VRAM" ), BucketedVRAM));
 		HardwareStatsAttribs.Add(FAnalyticsEventAttribute(TEXT( "GPU.DriverVersion" ), MainGPUDriverVer));
+		HardwareStatsAttribs.Add(FAnalyticsEventAttribute(TEXT( "GPU.RHIAdapterName" ), SurveyResults.RHIAdpater.AdapterName));
+		HardwareStatsAttribs.Add(FAnalyticsEventAttribute(TEXT( "GPU.RHIAdapterInternalDriverVersion" ), SurveyResults.RHIAdpater.AdapterInternalDriverVersion));
+		HardwareStatsAttribs.Add(FAnalyticsEventAttribute(TEXT( "GPU.RHIAdapterUserDriverVersion" ), SurveyResults.RHIAdpater.AdapterUserDriverVersion));
+		HardwareStatsAttribs.Add(FAnalyticsEventAttribute(TEXT( "GPU.RHIAdapterDriverDate" ), SurveyResults.RHIAdpater.AdapterDriverDate));
 		HardwareStatsAttribs.Add(FAnalyticsEventAttribute(TEXT( "RAM" ), BucketedRAM));
 		HardwareStatsAttribs.Add(FAnalyticsEventAttribute(TEXT( "RAM.WEI" ), FString::Printf( TEXT( "%.1f" ), SurveyResults.RAMPerformanceIndex )));
 		HardwareStatsAttribs.Add(FAnalyticsEventAttribute(TEXT( "NumberOfMonitors" ), FString::Printf( TEXT( "%d" ), SurveyResults.DisplayCount )));
@@ -9552,9 +9556,9 @@ bool UEngine::LoadMap( FWorldContext& WorldContext, FURL URL, class UPendingNetG
 	// may be unavailable to load the next one.
 	ENQUEUE_UNIQUE_RENDER_COMMAND(FlushCommand, 
 		{
-			FlushPendingDeleteRHIResources_RenderThread();
+			GRHICommandList.GetImmediateCommandList().ImmediateFlush(EImmediateFlushType::FlushRHIThreadFlushResources);
 			RHIFlushResources();
-			FlushPendingDeleteRHIResources_RenderThread();
+			GRHICommandList.GetImmediateCommandList().ImmediateFlush(EImmediateFlushType::FlushRHIThreadFlushResources);
 		}
 	);
 	FlushRenderingCommands();	  
diff --git a/Engine/Source/Runtime/Engine/Private/World.cpp b/Engine/Source/Runtime/Engine/Private/World.cpp
--- a/Engine/Source/Runtime/Engine/Private/World.cpp
+++ b/Engine/Source/Runtime/Engine/Private/World.cpp
@@ -4757,7 +4757,18 @@ void FSeamlessTravelHandler::StartLoadingDestination()
 
 void FSeamlessTravelHandler::CopyWorldData()
 {
-	CurrentWorld->DestroyDemoNetDriver();
+	// If we are doing seamless travel for replay playback, then make sure to transfer the replay driver over to the new world
+	if ( CurrentWorld->DemoNetDriver && CurrentWorld->DemoNetDriver->IsPlaying() )
+	{
+		UDemoNetDriver* OldDriver = CurrentWorld->DemoNetDriver;
+		CurrentWorld->DemoNetDriver = nullptr;
+		OldDriver->SetWorld( LoadedWorld );
+		LoadedWorld->DemoNetDriver = OldDriver;
+	}
+	else
+	{
+		CurrentWorld->DestroyDemoNetDriver();
+	}
 
 	UNetDriver* const NetDriver = CurrentWorld->GetNetDriver();
 	LoadedWorld->SetNetDriver(NetDriver);
diff --git a/Engine/Source/Runtime/Engine/Public/GPUSkinVertexFactory.h b/Engine/Source/Runtime/Engine/Public/GPUSkinVertexFactory.h
--- a/Engine/Source/Runtime/Engine/Public/GPUSkinVertexFactory.h
+++ b/Engine/Source/Runtime/Engine/Public/GPUSkinVertexFactory.h
@@ -305,6 +305,8 @@ public:
 
 		void ReleaseBoneData()
 		{
+			ensure(IsInRenderingThread());
+
 			UniformBuffer.SafeRelease();
 
 			for(uint32 i = 0; i < 2; ++i)
diff --git a/Engine/Source/Runtime/Engine/Public/ResourcePool.h b/Engine/Source/Runtime/Engine/Public/ResourcePool.h
--- a/Engine/Source/Runtime/Engine/Public/ResourcePool.h
+++ b/Engine/Source/Runtime/Engine/Public/ResourcePool.h
@@ -172,7 +172,7 @@ public:
 	/** Constructor */
 	TRenderResourcePool() :
 		FTickableObjectRenderThread(false)
-	{		
+	{
 	}
 	
 	/** Constructor with policy argument
@@ -195,6 +195,8 @@ public:
 	 */
 	ResourceType CreatePooledResource(ResourceCreationArguments Args)
 	{
+		ensure(IsInRenderingThread());
+
 		if (IsInitialized())
 		{
 			return TResourcePool<ResourceType, ResourcePoolPolicy, ResourceCreationArguments>::CreatePooledResource(Args);
@@ -210,6 +212,8 @@ public:
 	 */
 	void ReleasePooledResource(ResourceType Resource)
 	{
+		ensure(IsInRenderingThread());
+
 		if (IsInitialized())
 		{
 			TResourcePool<ResourceType, ResourcePoolPolicy, ResourceCreationArguments>::ReleasePooledResource(Resource);
@@ -219,6 +223,8 @@ public:
 public: // From FTickableObjectRenderThread
 	virtual void Tick( float DeltaTime ) override
 	{
+		ensure(IsInRenderingThread());
+
 		TResourcePool<ResourceType, ResourcePoolPolicy, ResourceCreationArguments>::DrainPool(false);
 	}
 	
diff --git a/Engine/Source/Runtime/Engine/Public/WorldCollision.h b/Engine/Source/Runtime/Engine/Public/WorldCollision.h
--- a/Engine/Source/Runtime/Engine/Public/WorldCollision.h
+++ b/Engine/Source/Runtime/Engine/Public/WorldCollision.h
@@ -31,6 +31,17 @@ struct FTraceHandle
 		_Data.FrameNumber = InFrameNumber;
 		_Data.Index = InIndex;
 	}
+
+	friend inline uint32 GetTypeHash(const FTraceHandle& Handle)
+	{
+		return GetTypeHash(Handle._Handle);
+	}
+
+	bool operator==(FTraceHandle const& Other) const
+	{
+		return Other._Handle == _Handle;
+	}
+
 };
 
 /** Types of Collision Shapes that are used by Trace **/
diff --git a/Engine/Source/Runtime/GameplayAbilities/Private/Abilities/GameplayAbility.cpp b/Engine/Source/Runtime/GameplayAbilities/Private/Abilities/GameplayAbility.cpp
--- a/Engine/Source/Runtime/GameplayAbilities/Private/Abilities/GameplayAbility.cpp
+++ b/Engine/Source/Runtime/GameplayAbilities/Private/Abilities/GameplayAbility.cpp
@@ -1334,6 +1334,10 @@ FGameplayEffectContextHandle UGameplayAbility::GetEffectContext(const FGameplayA
 	FGameplayEffectContextHandle Context = FGameplayEffectContextHandle(UAbilitySystemGlobals::Get().AllocGameplayEffectContext());
 	// By default use the owner and avatar as the instigator and causer
 	Context.AddInstigator(ActorInfo->OwnerActor.Get(), ActorInfo->AvatarActor.Get());
+
+	// add in the ability tracking here.
+	Context.SetAbility(this);
+
 	return Context;
 }
 
@@ -1611,7 +1615,7 @@ void UGameplayAbility::SetRemoteInstanceHasEnded()
 			// We have a task that is waiting for player input, but the remote player has ended the ability, so he will not send it.
 			// Kill the ability to avoid getting stuck active.
 			
-			ABILITY_LOG(Warning, TEXT("Ability %s is force cancelling because Task %s is waiting on remote player input and the  remote player has just ended the ability."), *GetName(), *Task->GetDebugString());
+			ABILITY_LOG(Log, TEXT("Ability %s is force cancelling because Task %s is waiting on remote player input and the  remote player has just ended the ability."), *GetName(), *Task->GetDebugString());
 			CurrentActorInfo->AbilitySystemComponent->ForceCancelAbilityDueToReplication(this);
 			break;
 		}
@@ -1625,7 +1629,7 @@ void UGameplayAbility::NotifyAbilityTaskWaitingOnPlayerData(class UAbilityTask*
 
 	if (RemoteInstanceEnded)
 	{
-		ABILITY_LOG(Warning, TEXT("Ability %s is force cancelling because Task %s has started after the remote player has ended the ability."), *GetName(), *AbilityTask->GetDebugString());
+		ABILITY_LOG(Log, TEXT("Ability %s is force cancelling because Task %s has started after the remote player has ended the ability."), *GetName(), *AbilityTask->GetDebugString());
 		CurrentActorInfo->AbilitySystemComponent->ForceCancelAbilityDueToReplication(this);
 	}
 }
\ No newline at end of file
diff --git a/Engine/Source/Runtime/GameplayAbilities/Private/GameplayCueManager.cpp b/Engine/Source/Runtime/GameplayAbilities/Private/GameplayCueManager.cpp
--- a/Engine/Source/Runtime/GameplayAbilities/Private/GameplayCueManager.cpp
+++ b/Engine/Source/Runtime/GameplayAbilities/Private/GameplayCueManager.cpp
@@ -307,12 +307,22 @@ void UGameplayCueManager::ReloadObjectLibrary(UWorld* World, const UWorld::Initi
 {
 	if (bAccelerationMapOutdated)
 	{
+		check(GlobalCueSet);
+		GlobalCueSet->Empty();
+
 		LoadObjectLibrary_Internal();
 	}
 }
 #endif
 
 void UGameplayCueManager::LoadObjectLibrary_Internal()
+{
+	FOnGameplayCueNotifySetLoaded OnLoadDelegate = FOnGameplayCueNotifySetLoaded::CreateUObject(this, &UGameplayCueManager::OnGameplayCueNotifyAsyncLoadComplete);
+
+	InitObjectLibraries(LoadedPaths, GameplayCueNotifyActorObjectLibrary, GameplayCueNotifyStaticObjectLibrary, OnLoadDelegate);
+}
+
+void UGameplayCueManager::InitObjectLibraries(TArray<FString> Paths, UObjectLibrary* ActorObjectLibrary, UObjectLibrary* StaticObjectLibrary, FOnGameplayCueNotifySetLoaded OnLoadDelegate, FShouldLoadGCNotifyDelegate ShouldLoadDelegate)
 {
 	DECLARE_SCOPE_CYCLE_COUNTER(TEXT("Loading Library"), STAT_ObjectLibrary, STATGROUP_LoadTime);
 
@@ -323,9 +333,9 @@ void UGameplayCueManager::LoadObjectLibrary_Internal()
 	SlowTask.MakeDialog();
 #endif
 
-	FScopeCycleCounterUObject PreloadScopeActor(GameplayCueNotifyActorObjectLibrary);
-	GameplayCueNotifyActorObjectLibrary->LoadBlueprintAssetDataFromPaths(LoadedPaths);
-	GameplayCueNotifyStaticObjectLibrary->LoadBlueprintAssetDataFromPaths(LoadedPaths);
+	FScopeCycleCounterUObject PreloadScopeActor(ActorObjectLibrary);
+	ActorObjectLibrary->LoadBlueprintAssetDataFromPaths(Paths);
+	StaticObjectLibrary->LoadBlueprintAssetDataFromPaths(Paths);
 
 	// ---------------------------------------------------------
 	// Determine loading scheme.
@@ -341,8 +351,8 @@ void UGameplayCueManager::LoadObjectLibrary_Internal()
 		FString PerfMessage = FString::Printf(TEXT("Fully Loaded GameplayCueNotify object library"));
 		SCOPE_LOG_TIME_IN_SECONDS(*PerfMessage, nullptr)
 #endif
-		GameplayCueNotifyActorObjectLibrary->LoadAssetsFromAssetData();
-		GameplayCueNotifyStaticObjectLibrary->LoadAssetsFromAssetData();
+		ActorObjectLibrary->LoadAssetsFromAssetData();
+		StaticObjectLibrary->LoadAssetsFromAssetData();
 	}
 
 	// ---------------------------------------------------------
@@ -350,28 +360,34 @@ void UGameplayCueManager::LoadObjectLibrary_Internal()
 	// ---------------------------------------------------------
 	
 	TArray<FAssetData> ActorAssetDatas;
-	GameplayCueNotifyActorObjectLibrary->GetAssetDataList(ActorAssetDatas);
+	ActorObjectLibrary->GetAssetDataList(ActorAssetDatas);
 
 	TArray<FAssetData> StaticAssetDatas;
-	GameplayCueNotifyStaticObjectLibrary->GetAssetDataList(StaticAssetDatas);
-
-	check(GlobalCueSet);
-	GlobalCueSet->Empty();
+	StaticObjectLibrary->GetAssetDataList(StaticAssetDatas);
 
 	TArray<FGameplayCueReferencePair> CuesToAdd;
-	BuildCuesToAddToGlobalSet(ActorAssetDatas, GET_MEMBER_NAME_CHECKED(AGameplayCueNotify_Actor, GameplayCueName), bAsyncLoadAtStartup, CuesToAdd);
-	BuildCuesToAddToGlobalSet(StaticAssetDatas, GET_MEMBER_NAME_CHECKED(UGameplayCueNotify_Static, GameplayCueName), bAsyncLoadAtStartup, CuesToAdd);
+	BuildCuesToAddToGlobalSet(ActorAssetDatas, GET_MEMBER_NAME_CHECKED(AGameplayCueNotify_Actor, GameplayCueName), bAsyncLoadAtStartup, CuesToAdd, OnLoadDelegate, ShouldLoadDelegate);
+	BuildCuesToAddToGlobalSet(StaticAssetDatas, GET_MEMBER_NAME_CHECKED(UGameplayCueNotify_Static, GameplayCueName), bAsyncLoadAtStartup, CuesToAdd, OnLoadDelegate, ShouldLoadDelegate);
 
 	check(GlobalCueSet);
 	GlobalCueSet->AddCues(CuesToAdd);
 }
 
-void UGameplayCueManager::BuildCuesToAddToGlobalSet(const TArray<FAssetData>& AssetDataList, FName TagPropertyName, bool bAsyncLoadAfterAdd, TArray<FGameplayCueReferencePair>& OutCuesToAdd)
+void UGameplayCueManager::BuildCuesToAddToGlobalSet(const TArray<FAssetData>& AssetDataList, FName TagPropertyName, bool bAsyncLoadAfterAdd, TArray<FGameplayCueReferencePair>& OutCuesToAdd, FOnGameplayCueNotifySetLoaded OnLoaded, FShouldLoadGCNotifyDelegate ShouldLoad)
 {
 	IGameplayTagsModule& GameplayTagsModule = IGameplayTagsModule::Get();
 
+	TArray<FStringAssetReference> AssetsToLoad;
+	AssetsToLoad.Reserve(AssetDataList.Num());
+
 	for (FAssetData Data: AssetDataList)
 	{
+		// If ShouldLoad delegate is bound and it returns false, don't load this one
+		if (ShouldLoad.IsBound() && (ShouldLoad.Execute(Data) == false))
+		{
+			continue;
+		}
+
 		const FString* FoundGameplayTag = Data.TagsAndValues.Find(TagPropertyName);
 		if (FoundGameplayTag && FoundGameplayTag->Equals(TEXT("None")) == false)
 		{
@@ -393,10 +409,7 @@ void UGameplayCueManager::BuildCuesToAddToGlobalSet(const TArray<FAssetData>& As
 
 				OutCuesToAdd.Add(FGameplayCueReferencePair(GameplayCueTag, StringRef));
 
-				if (bAsyncLoadAfterAdd)
-				{
-					StreamableManager.RequestAsyncLoad(StringRef, FStreamableDelegate::CreateUObject(this, &UGameplayCueManager::OnGameplayCueNotifyAsyncLoadComplete, StringRef));
-				}
+				AssetsToLoad.Add(StringRef);
 			}
 			else
 			{
@@ -404,6 +417,24 @@ void UGameplayCueManager::BuildCuesToAddToGlobalSet(const TArray<FAssetData>& As
 			}
 		}
 	}
+
+	if (bAsyncLoadAfterAdd)
+	{
+		auto ForwardLambda = [](TArray<FStringAssetReference> AssetList, FOnGameplayCueNotifySetLoaded OnLoadedDelegate)
+		{
+			OnLoadedDelegate.ExecuteIfBound(AssetList);
+		};
+
+		if (AssetsToLoad.Num() > 0)
+		{			
+			StreamableManager.RequestAsyncLoad(AssetsToLoad, FStreamableDelegate::CreateStatic( ForwardLambda, AssetsToLoad, OnLoaded));
+		}
+		else
+		{
+			// Still fire the delegate even if nothing was found to load
+			OnLoaded.ExecuteIfBound(AssetsToLoad);
+		}
+	}
 }
 
 void UGameplayCueManager::CheckForTooManyRPCs(FName FuncName, const FGameplayCuePendingExecute& PendingCue, const FString& CueID, const FGameplayEffectContext* EffectContext)
@@ -449,13 +480,16 @@ void UGameplayCueManager::CheckForTooManyRPCs(FName FuncName, const FGameplayCue
 	}
 }
 
-void UGameplayCueManager::OnGameplayCueNotifyAsyncLoadComplete(FStringAssetReference StringRef)
+void UGameplayCueManager::OnGameplayCueNotifyAsyncLoadComplete(TArray<FStringAssetReference> AssetList)
 {
-	UClass* GCClass = FindObject<UClass>(nullptr, *StringRef.ToString());
-	if (ensure(GCClass))
+	for (FStringAssetReference StringRef : AssetList)
 	{
-		LoadedGameplayCueNotifyClasses.Add(GCClass);
-		CheckForPreallocation(GCClass);
+		UClass* GCClass = FindObject<UClass>(nullptr, *StringRef.ToString());
+		if (ensure(GCClass))
+		{
+			LoadedGameplayCueNotifyClasses.Add(GCClass);
+			CheckForPreallocation(GCClass);
+		}
 	}
 }
 
@@ -626,6 +660,15 @@ void UGameplayCueManager::PrintGameplayCueNotifyMap()
 	GlobalCueSet->PrintCues();
 }
 
+void UGameplayCueManager::PrintLoadedGameplayCueNotifyClasses()
+{
+	for (UClass* NotifyClass : LoadedGameplayCueNotifyClasses)
+	{
+		ABILITY_LOG(Display, TEXT("%s"), *GetNameSafe(NotifyClass));
+	}
+	ABILITY_LOG(Display, TEXT("%d total classes"), LoadedGameplayCueNotifyClasses.Num());
+}
+
 static void	PrintGameplayCueNotifyMapConsoleCommandFunc(UWorld* InWorld)
 {
 	UAbilitySystemGlobals::Get().GetGameplayCueManager()->PrintGameplayCueNotifyMap();
@@ -637,6 +680,17 @@ FAutoConsoleCommandWithWorld PrintGameplayCueNotifyMapConsoleCommand(
 	FConsoleCommandWithWorldDelegate::CreateStatic(PrintGameplayCueNotifyMapConsoleCommandFunc)
 	);
 
+static void	PrintLoadedGameplayCueNotifyClasses(UWorld* InWorld)
+{
+	UAbilitySystemGlobals::Get().GetGameplayCueManager()->PrintLoadedGameplayCueNotifyClasses();
+}
+
+FAutoConsoleCommandWithWorld PrintLoadedGameplayCueNotifyClassesCommand(
+	TEXT("GameplayCue.PrintLoadedGameplayCueNotifyClasses"),
+	TEXT("Displays GameplayCue Notify classes that are loaded"),
+	FConsoleCommandWithWorldDelegate::CreateStatic(PrintLoadedGameplayCueNotifyClasses)
+	);
+
 FScopedGameplayCueSendContext::FScopedGameplayCueSendContext()
 {
 	UAbilitySystemGlobals::Get().GetGameplayCueManager()->StartGameplayCueSendContext();
diff --git a/Engine/Source/Runtime/GameplayAbilities/Private/GameplayCueSet.cpp b/Engine/Source/Runtime/GameplayAbilities/Private/GameplayCueSet.cpp
--- a/Engine/Source/Runtime/GameplayAbilities/Private/GameplayCueSet.cpp
+++ b/Engine/Source/Runtime/GameplayAbilities/Private/GameplayCueSet.cpp
@@ -93,6 +93,17 @@ void UGameplayCueSet::RemoveCuesByStringRefs(const TArray<FStringAssetReference>
 	}
 }
 
+void UGameplayCueSet::RemoveLoadedClass(UClass* Class)
+{
+	for (int32 idx = 0; idx < GameplayCueData.Num(); ++idx)
+	{
+		if (GameplayCueData[idx].LoadedGameplayCueClass == Class)
+		{
+			GameplayCueData[idx].LoadedGameplayCueClass = nullptr;
+		}
+	}
+}
+
 #if WITH_EDITOR
 void UGameplayCueSet::UpdateCueByStringRefs(const FStringAssetReference& CueToRemove, FString NewPath)
 {
diff --git a/Engine/Source/Runtime/GameplayAbilities/Private/GameplayEffectTypes.cpp b/Engine/Source/Runtime/GameplayAbilities/Private/GameplayEffectTypes.cpp
--- a/Engine/Source/Runtime/GameplayAbilities/Private/GameplayEffectTypes.cpp
+++ b/Engine/Source/Runtime/GameplayAbilities/Private/GameplayEffectTypes.cpp
@@ -69,6 +69,19 @@ void FGameplayEffectContext::AddInstigator(class AActor *InInstigator, class AAc
 	}
 }
 
+void FGameplayEffectContext::SetAbility(const UGameplayAbility* InGameplayAbility)
+{
+	if (InGameplayAbility)
+	{
+		Ability = InGameplayAbility->GetClass();
+	}
+}
+
+UGameplayAbility* FGameplayEffectContext::GetAbility() const
+{
+	return Ability.GetDefaultObject();
+}
+
 void FGameplayEffectContext::AddActors(const TArray<TWeakObjectPtr<AActor>>& InActors, bool bReset)
 {
 	if (bReset && Actors.Num())
@@ -108,25 +121,29 @@ bool FGameplayEffectContext::NetSerialize(FArchive& Ar, class UPackageMap* Map,
 		{
 			RepBits |= 1 << 1;
 		}
-		if (SourceObject.IsValid() )
+		if (*Ability)
 		{
 			RepBits |= 1 << 2;
 		}
-		if (Actors.Num() > 0)
+		if (SourceObject.IsValid())
 		{
 			RepBits |= 1 << 3;
 		}
-		if (HitResult.IsValid())
+		if (Actors.Num() > 0)
 		{
 			RepBits |= 1 << 4;
 		}
-		if (bHasWorldOrigin)
+		if (HitResult.IsValid())
 		{
 			RepBits |= 1 << 5;
 		}
+		if (bHasWorldOrigin)
+		{
+			RepBits |= 1 << 6;
+		}
 	}
 
-	Ar.SerializeBits(&RepBits, 6);
+	Ar.SerializeBits(&RepBits, 7);
 
 	if (RepBits & (1 << 0))
 	{
@@ -138,13 +155,17 @@ bool FGameplayEffectContext::NetSerialize(FArchive& Ar, class UPackageMap* Map,
 	}
 	if (RepBits & (1 << 2))
 	{
-		Ar << SourceObject;
+		Ar << Ability;
 	}
 	if (RepBits & (1 << 3))
 	{
-		Ar << Actors;;
+		Ar << SourceObject;
 	}
 	if (RepBits & (1 << 4))
+	{
+		Ar << Actors;
+	}
+	if (RepBits & (1 << 5))
 	{
 		if (Ar.IsLoading())
 		{
@@ -155,7 +176,7 @@ bool FGameplayEffectContext::NetSerialize(FArchive& Ar, class UPackageMap* Map,
 		}
 		HitResult->NetSerialize(Ar, Map, bOutSuccess);
 	}
-	if (RepBits & (1 << 5))
+	if (RepBits & (1 << 6))
 	{
 		Ar << WorldOrigin;
 		bHasWorldOrigin = true;
@@ -370,28 +391,29 @@ bool FGameplayTagCountContainer::UpdateTagMap_Internal(const FGameplayTag& Tag,
 		const FGameplayTag& CurTag = *CompleteTagIt;
 
 		// Get the current count of the specified tag. NOTE: Stored as a reference, so subsequent changes propogate to the map.
-		int32& TagCount = GameplayTagCountMap.FindOrAdd(CurTag);
+		int32& TagCountRef = GameplayTagCountMap.FindOrAdd(CurTag);
 
-		const int32 OldCount = TagCount;
+		const int32 OldCount = TagCountRef;
 
 		// Apply the delta to the count in the map
-		TagCount = FMath::Max(TagCount + CountDelta, 0);
+		int32 NewTagCount = FMath::Max(OldCount + CountDelta, 0);
+		TagCountRef = NewTagCount;
 
 		// If a significant change (new addition or total removal) occurred, trigger related delegates
-		bool SignificantChange = (OldCount == 0 || TagCount == 0);
+		bool SignificantChange = (OldCount == 0 || NewTagCount == 0);
 		CreatedSignificantChange |= SignificantChange;
 		if (SignificantChange)
 		{
-			OnAnyTagChangeDelegate.Broadcast(CurTag, TagCount);
+			OnAnyTagChangeDelegate.Broadcast(CurTag, NewTagCount);
 		}
 
 		FDelegateInfo* DelegateInfo = GameplayTagEventMap.Find(CurTag);
 		if (DelegateInfo)
 		{
-			DelegateInfo->OnAnyChange.Broadcast(CurTag, TagCount);
+			DelegateInfo->OnAnyChange.Broadcast(CurTag, NewTagCount);
 			if (SignificantChange)
 			{
-				DelegateInfo->OnNewOrRemove.Broadcast(CurTag, TagCount);
+				DelegateInfo->OnNewOrRemove.Broadcast(CurTag, NewTagCount);
 			}
 		}
 	}
diff --git a/Engine/Source/Runtime/GameplayAbilities/Private/GameplayPrediction.cpp b/Engine/Source/Runtime/GameplayAbilities/Private/GameplayPrediction.cpp
--- a/Engine/Source/Runtime/GameplayAbilities/Private/GameplayPrediction.cpp
+++ b/Engine/Source/Runtime/GameplayAbilities/Private/GameplayPrediction.cpp
@@ -5,47 +5,55 @@
 #include "AbilitySystemComponent.h"
 #include "GameplayPrediction.h"
 
+/** The key to understanding this function is that when a key is received by the server, we note which connection gave it to us. We only serialize the key back to that client.  */
 bool FPredictionKey::NetSerialize(FArchive& Ar, class UPackageMap* Map, bool& bOutSuccess)
 {
-	// Read bit for server initiated first
+	// First bit for valid key for this connection or not. (most keys are not valid)
+	uint8 ValidKeyForConnection = 0;
+	if (Ar.IsSaving())
+	{
+		/**
+		 *	Only serialize the payload if we have no owning connection (Client sending to server)
+		 *	or if the owning connection is this connection (Server only sends the prediction key to the client who gave it to us)
+		 *  or if this is a server initiated key (valid on all connections)
+		 */		
+		ValidKeyForConnection = (PredictiveConnection == nullptr || (Map == PredictiveConnection) || bIsServerInitiated) && (Current > 0);
+	}
+	Ar.SerializeBits(&ValidKeyForConnection, 1);
+
+	// Second bit for base key (only if valid connection)
+	uint8 HasBaseKey = 0;
+	if (ValidKeyForConnection)
+	{
+		if (Ar.IsSaving())
+		{
+			HasBaseKey = Base > 0;
+		}
+		Ar.SerializeBits(&HasBaseKey, 1);
+	}
+
+	// Third bit for server initiated
 	uint8 ServerInitiatedByte = bIsServerInitiated;
 	Ar.SerializeBits(&ServerInitiatedByte, 1);
 	bIsServerInitiated = ServerInitiatedByte & 1;
 
-	if (Ar.IsLoading())
+	// Conditionally Serialize the Current and Base keys
+	if (ValidKeyForConnection)
 	{
 		Ar << Current;
-		if (Current > 0)
+		if (HasBaseKey)
 		{
 			Ar << Base;
 		}
+	}	
+	if (Ar.IsLoading())
+	{
+		// We are reading this key: the connection that gave us this key is the predictive connection, and we will only serialize this key back to it.
 		if (!bIsServerInitiated)
 		{
 			PredictiveConnection = Map;
 		}
 	}
-	else
-	{
-		/**
-		 *	Only serialize the payload if we have no owning connection (Client sending to server)
-		 *	or if the owning connection is this connection (Server only sends the prediction key to the client who gave it to us)
-		 *  or if this is a server initiated key (valid on all connections)
-		 */
-		
-		if (PredictiveConnection == nullptr || (Map == PredictiveConnection) || bIsServerInitiated)
-		{
-			Ar << Current;
-			if (Current > 0)
-			{
-				Ar << Base;
-			}
-		}
-		else
-		{
-			KeyType Payload = 0;
-			Ar << Payload;
-		}
-	}
 
 	bOutSuccess = true;
 	return true;
diff --git a/Engine/Source/Runtime/GameplayAbilities/Public/AbilitySystemGlobals.h b/Engine/Source/Runtime/GameplayAbilities/Public/AbilitySystemGlobals.h
--- a/Engine/Source/Runtime/GameplayAbilities/Public/AbilitySystemGlobals.h
+++ b/Engine/Source/Runtime/GameplayAbilities/Public/AbilitySystemGlobals.h
@@ -195,8 +195,6 @@ protected:
 	bool bIgnoreAbilitySystemCosts;
 #endif // #if !(UE_BUILD_SHIPPING || UE_BUILD_TEST)
 
-private:
-
 	/** Holds all of the valid gameplay-related tags that can be applied to assets */
 	UPROPERTY(config)
 	FString GlobalCurveTableName;
diff --git a/Engine/Source/Runtime/GameplayAbilities/Public/AbilitySystemLog.h b/Engine/Source/Runtime/GameplayAbilities/Public/AbilitySystemLog.h
--- a/Engine/Source/Runtime/GameplayAbilities/Public/AbilitySystemLog.h
+++ b/Engine/Source/Runtime/GameplayAbilities/Public/AbilitySystemLog.h
@@ -11,8 +11,8 @@
 //
 
 
-GAMEPLAYABILITIES_API DECLARE_LOG_CATEGORY_EXTERN(LogAbilitySystem, Warning, All);
-GAMEPLAYABILITIES_API DECLARE_LOG_CATEGORY_EXTERN(VLogAbilitySystem, Warning, All);
+GAMEPLAYABILITIES_API DECLARE_LOG_CATEGORY_EXTERN(LogAbilitySystem, Display, All);
+GAMEPLAYABILITIES_API DECLARE_LOG_CATEGORY_EXTERN(VLogAbilitySystem, Display, All);
 
 #if NO_LOGGING || !PLATFORM_DESKTOP
 
diff --git a/Engine/Source/Runtime/GameplayAbilities/Public/GameplayCueManager.h b/Engine/Source/Runtime/GameplayAbilities/Public/GameplayCueManager.h
--- a/Engine/Source/Runtime/GameplayAbilities/Public/GameplayCueManager.h
+++ b/Engine/Source/Runtime/GameplayAbilities/Public/GameplayCueManager.h
@@ -47,9 +47,8 @@ class UGameplayCueSet;
  *	
  */
 
-/**
- *	A self contained handler of a GameplayCue. These are similar to AnimNotifies in implementation.
- */
+DECLARE_DELEGATE_OneParam(FOnGameplayCueNotifySetLoaded, TArray<FStringAssetReference>);
+DECLARE_DELEGATE_RetVal_OneParam(bool, FShouldLoadGCNotifyDelegate, const FAssetData&);
 
 UCLASS()
 class GAMEPLAYABILITIES_API UGameplayCueManager : public UDataAsset
@@ -133,9 +132,9 @@ class GAMEPLAYABILITIES_API UGameplayCueManager : public UDataAsset
 	// Preload GameplayCue tags that we think we will need:
 	// -------------------------------------------------------------
 
-	void	BeginLoadingGameplayCueNotify(FGameplayTag GameplayCueTag);
+	void BeginLoadingGameplayCueNotify(FGameplayTag GameplayCueTag);
 
-	int32	FinishLoadingGameplayCueNotifies();
+	int32 FinishLoadingGameplayCueNotifies();
 
 	UPROPERTY(transient)
 	FStreamableManager	StreamableManager;
@@ -144,6 +143,8 @@ class GAMEPLAYABILITIES_API UGameplayCueManager : public UDataAsset
 
 	void PrintGameplayCueNotifyMap();
 
+	virtual void PrintLoadedGameplayCueNotifyClasses();
+
 	virtual class UWorld* GetWorld() const override;
 
 #if WITH_EDITOR
@@ -177,19 +178,21 @@ protected:
 
 #if WITH_EDITOR
 	//This handles the case where GameplayCueNotifications have changed between sessions, which is possible in editor.
-	void ReloadObjectLibrary(UWorld* World, const UWorld::InitializationValues IVS);
+	virtual void ReloadObjectLibrary(UWorld* World, const UWorld::InitializationValues IVS);
 #endif
 
 	void LoadObjectLibrary_Internal();
 
+	void InitObjectLibraries(TArray<FString> Paths, UObjectLibrary* ActorObjectLibrary, UObjectLibrary* StaticObjectLibrary, FOnGameplayCueNotifySetLoaded OnLoaded, FShouldLoadGCNotifyDelegate ShouldLoad = FShouldLoadGCNotifyDelegate());
+
 	virtual bool ShouldAsyncLoadAtStartup() const { return true; }
 
-	void BuildCuesToAddToGlobalSet(const TArray<FAssetData>& AssetDataList, FName TagPropertyName, bool bAsyncLoadAfterAdd, TArray<struct FGameplayCueReferencePair>& OutCuesToAdd);
+	void BuildCuesToAddToGlobalSet(const TArray<FAssetData>& AssetDataList, FName TagPropertyName, bool bAsyncLoadAfterAdd, TArray<struct FGameplayCueReferencePair>& OutCuesToAdd, FOnGameplayCueNotifySetLoaded OnLoaded = FOnGameplayCueNotifySetLoaded(), FShouldLoadGCNotifyDelegate = FShouldLoadGCNotifyDelegate());
 
 	/** The cue manager has a tendency to produce a lot of RPCs. This logs out when we are attempting to fire more RPCs than will actually go off */
 	void CheckForTooManyRPCs(FName FuncName, const FGameplayCuePendingExecute& PendingCue, const FString& CueID, const FGameplayEffectContext* EffectContext);
 
-	void OnGameplayCueNotifyAsyncLoadComplete(FStringAssetReference StringRef);
+	void OnGameplayCueNotifyAsyncLoadComplete(TArray<FStringAssetReference> StringRef);
 
 	void CheckForPreallocation(UClass* GCClass);
 
diff --git a/Engine/Source/Runtime/GameplayAbilities/Public/GameplayCueSet.h b/Engine/Source/Runtime/GameplayAbilities/Public/GameplayCueSet.h
--- a/Engine/Source/Runtime/GameplayAbilities/Public/GameplayCueSet.h
+++ b/Engine/Source/Runtime/GameplayAbilities/Public/GameplayCueSet.h
@@ -59,6 +59,9 @@ class GAMEPLAYABILITIES_API UGameplayCueSet : public UDataAsset
 	/** Removes all cues from the set matching the supplied string refs */
 	virtual void RemoveCuesByStringRefs(const TArray<FStringAssetReference>& CuesToRemove);
 
+	/** Nulls reference to the loaded class. Note this doesn't remove the entire cue from the internal data structure, just the hard ref to the loaded class */
+	virtual void RemoveLoadedClass(UClass* Class);
+
 #if WITH_EDITOR
 
 	/** Updates an existing cue */
diff --git a/Engine/Source/Runtime/GameplayAbilities/Public/GameplayEffectTypes.h b/Engine/Source/Runtime/GameplayAbilities/Public/GameplayEffectTypes.h
--- a/Engine/Source/Runtime/GameplayAbilities/Public/GameplayEffectTypes.h
+++ b/Engine/Source/Runtime/GameplayAbilities/Public/GameplayEffectTypes.h
@@ -270,12 +270,14 @@ struct GAMEPLAYABILITIES_API FGameplayEffectContext
 	GENERATED_USTRUCT_BODY()
 
 	FGameplayEffectContext()
-		: bHasWorldOrigin(false)
+	: Ability(nullptr)
+	, bHasWorldOrigin(false)
 	{
 	}
 
 	FGameplayEffectContext(AActor* InInstigator, AActor* InEffectCauser)
-		: bHasWorldOrigin(false)
+		: Ability(nullptr)
+		, bHasWorldOrigin(false)
 	{
 		AddInstigator(InInstigator, InEffectCauser);
 	}
@@ -290,12 +292,17 @@ struct GAMEPLAYABILITIES_API FGameplayEffectContext
 	/** Sets the instigator and effect causer. Instigator is who owns the ability that spawned this, EffectCauser is the actor that is the physical source of the effect, such as a weapon. They can be the same. */
 	virtual void AddInstigator(class AActor *InInstigator, class AActor *InEffectCauser);
 
+	/** Sets the ability that was used to spawn this */
+	virtual void SetAbility(const UGameplayAbility* InGameplayAbility);
+
 	/** Returns the immediate instigator that applied this effect */
 	virtual AActor* GetInstigator() const
 	{
 		return Instigator.Get();
 	}
 
+	virtual UGameplayAbility* GetAbility() const;
+
 	/** Returns the ability system component of the instigator of this effect */
 	virtual UAbilitySystemComponent* GetInstigatorAbilitySystemComponent() const
 	{
@@ -403,6 +410,10 @@ protected:
 	UPROPERTY()
 	TWeakObjectPtr<AActor> EffectCauser;
 
+	/** the ability that is responsible for this effect context */
+	UPROPERTY()
+	TSubclassOf<UGameplayAbility> Ability;
+
 	/** Object this effect was created from, can be an actor or static object. Useful to bind an effect to a gameplay object */
 	UPROPERTY()
 	TWeakObjectPtr<UObject> SourceObject;
@@ -496,6 +507,14 @@ struct FGameplayEffectContextHandle
 		}
 	}
 
+	void SetAbility(const UGameplayAbility* InGameplayAbility)
+	{
+		if (IsValid())
+		{
+			Data->SetAbility(InGameplayAbility);
+		}
+	}
+
 	/** Returns the immediate instigator that applied this effect */
 	virtual AActor* GetInstigator() const
 	{
@@ -506,6 +525,16 @@ struct FGameplayEffectContextHandle
 		return NULL;
 	}
 
+	/** Returns the Ability that  */
+	virtual UGameplayAbility* GetAbility() const
+	{
+		if (IsValid())
+		{
+			return Data->GetAbility();
+		}
+		return nullptr;
+	}
+
 	/** Returns the ability system component of the instigator of this effect */
 	virtual UAbilitySystemComponent* GetInstigatorAbilitySystemComponent() const
 	{
diff --git a/Engine/Source/Runtime/JsonUtilities/Private/JsonObjectConverter.cpp b/Engine/Source/Runtime/JsonUtilities/Private/JsonObjectConverter.cpp
--- a/Engine/Source/Runtime/JsonUtilities/Private/JsonObjectConverter.cpp
+++ b/Engine/Source/Runtime/JsonUtilities/Private/JsonObjectConverter.cpp
@@ -170,22 +170,36 @@ bool FJsonObjectConverter::UStructToJsonAttributes(const UStruct* StructDefiniti
 	return true;
 }
 
-bool FJsonObjectConverter::UStructToJsonObjectString(const UStruct* StructDefinition, const void* Struct, FString& OutJsonString, int64 CheckFlags, int64 SkipFlags, int32 Indent, const CustomExportCallback* ExportCb)
+template<class CharType, class PrintPolicy>
+bool UStructToJsonObjectStringInternal(const TSharedRef<FJsonObject>& JsonObject, FString& OutJsonString, int32 Indent)
+{
+	TSharedRef<TJsonWriter<CharType, PrintPolicy> > JsonWriter = TJsonWriterFactory<CharType, PrintPolicy>::Create(&OutJsonString, Indent);
+	bool bSuccess = FJsonSerializer::Serialize(JsonObject, JsonWriter);
+	JsonWriter->Close();
+	return bSuccess;
+}
+
+bool FJsonObjectConverter::UStructToJsonObjectString(const UStruct* StructDefinition, const void* Struct, FString& OutJsonString, int64 CheckFlags, int64 SkipFlags, int32 Indent, const CustomExportCallback* ExportCb, bool bPrettyPrint)
 {
 	TSharedRef<FJsonObject> JsonObject = MakeShareable( new FJsonObject() );
 	if (UStructToJsonObject(StructDefinition, Struct, JsonObject, CheckFlags, SkipFlags, ExportCb))
 	{
-		TSharedRef<TJsonWriter<> > JsonWriter = TJsonWriterFactory<>::Create(&OutJsonString, Indent);
-
-		if (FJsonSerializer::Serialize( JsonObject, JsonWriter ))
+		bool bSuccess = false;
+		if (bPrettyPrint)
+		{
+			bSuccess = UStructToJsonObjectStringInternal<TCHAR, TPrettyJsonPrintPolicy<TCHAR> >(JsonObject, OutJsonString, Indent);
+		}
+		else
+		{
+			bSuccess = UStructToJsonObjectStringInternal<TCHAR, TCondensedJsonPrintPolicy<TCHAR> >(JsonObject, OutJsonString, Indent);
+		}
+		if (bSuccess)
 		{
-			JsonWriter->Close();
 			return true;
 		}
 		else
 		{
 			UE_LOG(LogJson, Warning, TEXT("UStructToJsonObjectString - Unable to write out json"));
-			JsonWriter->Close();
 		}
 	}
 
diff --git a/Engine/Source/Runtime/JsonUtilities/Public/JsonObjectConverter.h b/Engine/Source/Runtime/JsonUtilities/Public/JsonObjectConverter.h
--- a/Engine/Source/Runtime/JsonUtilities/Public/JsonObjectConverter.h
+++ b/Engine/Source/Runtime/JsonUtilities/Public/JsonObjectConverter.h
@@ -65,10 +65,11 @@ public: // UStruct -> JSON
 	 * @param SkipFlags Skip properties that match any of these flags
 	 * @param Indent How many tabs to add to the json serializer
 	 * @param ExportCb Optional callback for types we don't understand. This is called right before falling back to the generic ToString()
+	 * @param bPrettyPrint Option to use pretty print (e.g., adds line endings) or condensed print
 	 *
 	 * @return False if any properties failed to write
 	 */
-	static bool UStructToJsonObjectString(const UStruct* StructDefinition, const void* Struct, FString& OutJsonString, int64 CheckFlags, int64 SkipFlags, int32 Indent = 0, const CustomExportCallback* ExportCb = nullptr);
+	static bool UStructToJsonObjectString(const UStruct* StructDefinition, const void* Struct, FString& OutJsonString, int64 CheckFlags, int64 SkipFlags, int32 Indent = 0, const CustomExportCallback* ExportCb = nullptr, bool bPrettyPrint = true);
 
 	/**
 	 * Wrapper to UStructToJsonObjectString that allows a print policy to be specified.
diff --git a/Engine/Source/Runtime/Landscape/Private/LandscapeEdit.cpp b/Engine/Source/Runtime/Landscape/Private/LandscapeEdit.cpp
--- a/Engine/Source/Runtime/Landscape/Private/LandscapeEdit.cpp
+++ b/Engine/Source/Runtime/Landscape/Private/LandscapeEdit.cpp
@@ -756,7 +756,7 @@ void ULandscapeComponent::UpdateCollisionHeightData(const FColor* const Heightma
 					if (GrassHeights)
 					{
 						uint16& CollisionHeight = CollisionHeightData[CompVertX + CompVertY * CollisionSizeVerts];
-						const uint16& NewHeight = GrassHeights[CompVertX + CompVertY * MipSizeU];
+						const uint16& NewHeight = GrassHeights[CompVertX + CompVertY * CollisionSizeVerts];
 						CollisionHeight = NewHeight;
 					}
 					else
@@ -840,7 +840,7 @@ void ULandscapeComponent::UpdateCollisionHeightData(const FColor* const Heightma
 						if (SimpleCollisionGrassHeights)
 						{
 							uint16& CollisionHeight = SimpleCollisionHeightData[CompVertX + CompVertY * SimpleCollisionSizeVerts];
-							const uint16& NewHeight = SimpleCollisionGrassHeights[CompVertX + CompVertY * SimpleMipSizeU];
+							const uint16& NewHeight = SimpleCollisionGrassHeights[CompVertX + CompVertY * SimpleCollisionSizeVerts];
 							CollisionHeight = NewHeight;
 						}
 						else
diff --git a/Engine/Source/Runtime/Landscape/Private/LandscapeGrass.cpp b/Engine/Source/Runtime/Landscape/Private/LandscapeGrass.cpp
--- a/Engine/Source/Runtime/Landscape/Private/LandscapeGrass.cpp
+++ b/Engine/Source/Runtime/Landscape/Private/LandscapeGrass.cpp
@@ -1824,7 +1824,7 @@ void ALandscapeProxy::UpdateGrass(const TArray<FVector>& Cameras, bool bForceSyn
 			{
 				int32 NumComponentsRendered = 0;
 				int32 NumComponentsUnableToRender = 0;
-				if (GrassTypes.Num() > 0 && CVarPrerenderGrassmaps.GetValueOnAnyThread() > 0)
+				if ((GrassTypes.Num() > 0 && CVarPrerenderGrassmaps.GetValueOnAnyThread() > 0) || bBakeMaterialPositionOffsetIntoCollision)
 				{
 					// try to render some grassmaps
 					TArray<ULandscapeComponent*> ComponentsToRender;
diff --git a/Engine/Source/Runtime/Launch/Private/LaunchEngineLoop.cpp b/Engine/Source/Runtime/Launch/Private/LaunchEngineLoop.cpp
--- a/Engine/Source/Runtime/Launch/Private/LaunchEngineLoop.cpp
+++ b/Engine/Source/Runtime/Launch/Private/LaunchEngineLoop.cpp
@@ -2918,6 +2918,11 @@ bool FEngineLoop::AppInit( )
 
 	FPageAllocator::LatchProtectedMode();
 
+	if (FParse::Param(FCommandLine::Get(), TEXT("purgatorymallocproxy")))
+	{
+		FMemory::EnablePurgatoryTests();
+	}
+
 #if !UE_BUILD_SHIPPING
 	if (FParse::Param(FCommandLine::Get(), TEXT("BUILDMACHINE")))
 	{
diff --git a/Engine/Source/Runtime/Linux/LinuxCommonStartup/Private/LinuxCommonStartup.cpp b/Engine/Source/Runtime/Linux/LinuxCommonStartup/Private/LinuxCommonStartup.cpp
--- a/Engine/Source/Runtime/Linux/LinuxCommonStartup/Private/LinuxCommonStartup.cpp
+++ b/Engine/Source/Runtime/Linux/LinuxCommonStartup/Private/LinuxCommonStartup.cpp
@@ -22,9 +22,11 @@ extern void DLLIMPORT GenerateCrashInfoAndLaunchReporter(const FLinuxCrashContex
  */
 void CommonLinuxCrashHandler(const FGenericCrashContext& GenericContext)
 {
+	// at this point we should already be using malloc crash handler (see PlatformCrashHandler)
+
 	const FLinuxCrashContext& Context = static_cast< const FLinuxCrashContext& >( GenericContext );
+	printf("CommonLinuxCrashHandler: Signal=%d\n", Context.Signal);
 
-	printf("EngineCrashHandler: Signal=%d\n", Context.Signal);
 	ReportCrash(Context);
 	if (GLog)
 	{
diff --git a/Engine/Source/Runtime/MoviePlayer/Private/DefaultGameMoviePlayer.cpp b/Engine/Source/Runtime/MoviePlayer/Private/DefaultGameMoviePlayer.cpp
--- a/Engine/Source/Runtime/MoviePlayer/Private/DefaultGameMoviePlayer.cpp
+++ b/Engine/Source/Runtime/MoviePlayer/Private/DefaultGameMoviePlayer.cpp
@@ -196,6 +196,7 @@ void FDefaultGameMoviePlayer::Shutdown()
 	if( SyncMechanism)
 	{
 		SyncMechanism->DestroySlateThread();
+		FScopeLock SyncMechanismLock(&SyncMechanismCriticalSection);
 		delete SyncMechanism;
 		SyncMechanism = NULL;
 	}
@@ -239,8 +240,11 @@ bool FDefaultGameMoviePlayer::PlayMovie()
 			LoadingScreenWidgetHolder->SetContent(LoadingScreenAttributes.WidgetLoadingScreen.IsValid() ? LoadingScreenAttributes.WidgetLoadingScreen.ToSharedRef() : SNullWidget::NullWidget);
             LoadingScreenWindowPtr.Pin()->SetContent(LoadingScreenContents.ToSharedRef());
 		
-            SyncMechanism = new FSlateLoadingSynchronizationMechanism();
-            SyncMechanism->Initialize();
+			{
+				FScopeLock SyncMechanismLock(&SyncMechanismCriticalSection);
+				SyncMechanism = new FSlateLoadingSynchronizationMechanism();
+				SyncMechanism->Initialize();
+			}
 
             bBeganPlaying = true;
         }
@@ -268,8 +272,12 @@ void FDefaultGameMoviePlayer::WaitForMovieToFinish()
 		if (SyncMechanism)
 		{
 			SyncMechanism->DestroySlateThread();
-			delete SyncMechanism;
-			SyncMechanism = NULL;
+
+			{
+				FScopeLock SyncMechanismLock(&SyncMechanismCriticalSection);
+				delete SyncMechanism;
+				SyncMechanism = NULL;
+			}
 		}
 
 		if( !bEnforceMinimumTime )
@@ -377,6 +385,7 @@ void FDefaultGameMoviePlayer::Tick( float DeltaTime )
 	{
 		if (!IsLoadingFinished() && SyncMechanism)
 		{
+			FScopeLock SyncMechanismLock(&SyncMechanismCriticalSection);
 			if (SyncMechanism->IsSlateDrawPassEnqueued())
 			{
 				GFrameNumberRenderThread++;
diff --git a/Engine/Source/Runtime/MoviePlayer/Private/DefaultGameMoviePlayer.h b/Engine/Source/Runtime/MoviePlayer/Private/DefaultGameMoviePlayer.h
--- a/Engine/Source/Runtime/MoviePlayer/Private/DefaultGameMoviePlayer.h
+++ b/Engine/Source/Runtime/MoviePlayer/Private/DefaultGameMoviePlayer.h
@@ -77,6 +77,9 @@ private:
 	/** The widget which holds the loading screen widget passed in via the FLoadingScreenAttributes */
 	TSharedPtr<class SBorder> LoadingScreenWidgetHolder;
 
+	/** Controls access to the sync mechanism across threads */
+	FCriticalSection SyncMechanismCriticalSection;
+
 	/** The threading mechanism with which we handle running slate on another thread */
 	class FSlateLoadingSynchronizationMechanism* SyncMechanism;
 
diff --git a/Engine/Source/Runtime/NetworkFile/Private/NetworkPlatformFile.cpp b/Engine/Source/Runtime/NetworkFile/Private/NetworkPlatformFile.cpp
--- a/Engine/Source/Runtime/NetworkFile/Private/NetworkPlatformFile.cpp
+++ b/Engine/Source/Runtime/NetworkFile/Private/NetworkPlatformFile.cpp
@@ -20,6 +20,7 @@
 DEFINE_LOG_CATEGORY(LogNetworkPlatformFile);
 
 FString FNetworkPlatformFile::MP4Extension = TEXT(".mp4");
+FString FNetworkPlatformFile::BulkFileExtension = TEXT(".ubulk");
 
 FNetworkPlatformFile::FNetworkPlatformFile()
 	: bHasLoadedDDCDirectories(false)
@@ -1102,7 +1103,7 @@ void FNetworkPlatformFile::EnsureFileIsLocal(const FString& Filename)
 
 	// we only copy files that actually exist on the server, can greatly reduce network traffic for, say,
 	// the INT file each package tries to load
-	if (!bIsCookable && ServerFiles.FindFile(Filename) == NULL)
+	if (!bIsCookable && (ServerFiles.FindFile(Filename) == NULL) && (Extension != BulkFileExtension))
 	{
 		// Uncomment this to have the server file list dumped
 		// the first time a file requested is not found.
diff --git a/Engine/Source/Runtime/NetworkFile/Public/NetworkPlatformFile.h b/Engine/Source/Runtime/NetworkFile/Public/NetworkPlatformFile.h
--- a/Engine/Source/Runtime/NetworkFile/Public/NetworkPlatformFile.h
+++ b/Engine/Source/Runtime/NetworkFile/Public/NetworkPlatformFile.h
@@ -235,6 +235,7 @@ private:
 	class ITransport* Transport; 
 
 	static FString MP4Extension;
+	static FString BulkFileExtension;
 };
 
 class SOCKETS_API FNetworkFileHandle : public IFileHandle
diff --git a/Engine/Source/Runtime/Online/OnlineFramework/Party/Private/PartyGameState.cpp b/Engine/Source/Runtime/Online/OnlineFramework/Party/Private/PartyGameState.cpp
--- a/Engine/Source/Runtime/Online/OnlineFramework/Party/Private/PartyGameState.cpp
+++ b/Engine/Source/Runtime/Online/OnlineFramework/Party/Private/PartyGameState.cpp
@@ -479,7 +479,7 @@ void UPartyGameState::HandlePartyDataReceived(const TSharedRef<FOnlinePartyData>
 	if (ensure(PartyInt.IsValid()))
 	{
 		check(PartyStateRefDef && PartyStateRef);
-		if (FVariantDataConverter::VariantMapToUStruct(InPartyData->KeyValAttrs, PartyStateRefDef, PartyStateRefScratch, 0, CPF_Transient | CPF_RepSkip))
+		if (FVariantDataConverter::VariantMapToUStruct(InPartyData->GetKeyValAttrs(), PartyStateRefDef, PartyStateRefScratch, 0, CPF_Transient | CPF_RepSkip))
 		{
 			ComparePartyData(*PartyStateRef, *PartyStateRefScratch);
 
@@ -527,7 +527,7 @@ void UPartyGameState::HandlePartyMemberDataReceived(const FUniqueNetId& InMember
 
 			check(CurrentPartyMember->MemberStateRefDef && CurrentPartyMember->MemberStateRefScratch);
 
-			if (FVariantDataConverter::VariantMapToUStruct(InPartyMemberData->KeyValAttrs, CurrentPartyMember->MemberStateRefDef, CurrentPartyMember->MemberStateRef, 0, CPF_Transient | CPF_RepSkip))
+			if (FVariantDataConverter::VariantMapToUStruct(InPartyMemberData->GetKeyValAttrs(), CurrentPartyMember->MemberStateRefDef, CurrentPartyMember->MemberStateRef, 0, CPF_Transient | CPF_RepSkip))
 			{
 				// Broadcast property changes
 				CurrentPartyMember->ComparePartyMemberData(*CurrentPartyMember->MemberStateRefScratch);
@@ -1161,7 +1161,7 @@ void UPartyGameState::UpdatePartyData(const FUniqueNetIdRepl& InLocalUserId)
 			{
 				FOnlinePartyData PartyData;
 				ensure(PartyStateRefDef != nullptr && PartyStateRef != nullptr);
-				if (FVariantDataConverter::UStructToVariantMap(PartyStateRefDef, PartyStateRef, PartyData.KeyValAttrs, 0, CPF_Transient | CPF_RepSkip))
+				if (FVariantDataConverter::UStructToVariantMap(PartyStateRefDef, PartyStateRef, PartyData.GetKeyValAttrs(), 0, CPF_Transient | CPF_RepSkip))
 				{
 					PartyInt->UpdatePartyData(*OwningUserId, *PartyId, PartyData);
 				}
@@ -1203,7 +1203,7 @@ void UPartyGameState::UpdatePartyMemberState(const FUniqueNetIdRepl& InLocalUser
 				{
 					FOnlinePartyData PartyMemberData;
 					ensure(InPartyMemberState->MemberStateRefDef != nullptr && InPartyMemberState->MemberStateRef != nullptr);
-					if (FVariantDataConverter::UStructToVariantMap(InPartyMemberState->MemberStateRefDef, InPartyMemberState->MemberStateRef, PartyMemberData.KeyValAttrs, 0, CPF_Transient | CPF_RepSkip))
+					if (FVariantDataConverter::UStructToVariantMap(InPartyMemberState->MemberStateRefDef, InPartyMemberState->MemberStateRef, PartyMemberData.GetKeyValAttrs(), 0, CPF_Transient | CPF_RepSkip))
 					{
 						PartyInt->UpdatePartyMemberData(*InLocalUserId, *PartyId, PartyMemberData);
 					}
diff --git a/Engine/Source/Runtime/Online/OnlineFramework/Qos/Private/QosEvaluator.cpp b/Engine/Source/Runtime/Online/OnlineFramework/Qos/Private/QosEvaluator.cpp
--- a/Engine/Source/Runtime/Online/OnlineFramework/Qos/Private/QosEvaluator.cpp
+++ b/Engine/Source/Runtime/Online/OnlineFramework/Qos/Private/QosEvaluator.cpp
@@ -147,7 +147,6 @@ void FQosEvaluator::OnFindDatacentersComplete(bool bWasSuccessful, FOnQosSearchC
 		if (SessionInt.IsValid())
 		{
 			SessionInt->ClearOnFindSessionsCompleteDelegate_Handle(OnFindDatacentersCompleteDelegateHandle);
-			//OnFindDatacentersCompleteDelegateHandle.Reset();
 		}
 	}
 
@@ -356,7 +355,9 @@ void FQosEvaluator::OnQosRequestComplete(EQosResponseType QosResponse, int32 Res
 			SearchResult.PingInMs = MAX_QUERY_PING;
 		}
 
-		UE_LOG(LogQos, Verbose, TEXT("Qos response received: %d ms"), ResponseTime);
+		extern ENGINE_API float GAverageFPS;
+		extern ENGINE_API float GAverageMS;
+		UE_LOG(LogQos, Verbose, TEXT("Qos response received: %d ms FPS: %0.2f MS: %0.2f"), ResponseTime, GAverageFPS, GAverageMS);
 
 		if (QosStats.IsValid())
 		{
diff --git a/Engine/Source/Runtime/Online/OnlineFramework/Qos/Private/QosInterface.cpp b/Engine/Source/Runtime/Online/OnlineFramework/Qos/Private/QosInterface.cpp
--- a/Engine/Source/Runtime/Online/OnlineFramework/Qos/Private/QosInterface.cpp
+++ b/Engine/Source/Runtime/Online/OnlineFramework/Qos/Private/QosInterface.cpp
@@ -32,11 +32,12 @@ TSharedRef<FQosInterface> FQosInterface::Get()
 }
 
 FQosInterface::FQosInterface()
-: MinimumPingMs(250)
-, QosEvalResult(EQosCompletionResult::Invalid)
+	: MaximumPingMs(250)
+	, BestRegionPingMs(MAX_int32)
+	, QosEvalResult(EQosCompletionResult::Invalid)
 {
 	check(GConfig);
-	GConfig->GetInt(TEXT("Qos"), TEXT("MinimumPingMs"), MinimumPingMs, GGameIni);
+	GConfig->GetInt(TEXT("Qos"), TEXT("MinimumPingMs"), MaximumPingMs, GGameIni);
 	GConfig->GetString(TEXT("Qos"), TEXT("ForceRegionId"), ForceRegionId, GGameIni);
 
 	// get a forced region id from the command line as an override
@@ -121,12 +122,15 @@ void FQosInterface::OnQosEvaluationComplete(EQosCompletionResult Result, const T
 		QosEvalResult = EQosCompletionResult::Failure;
 	}
 
+	BestRegionPingMs = MAX_int32;
+
 	if (QosEvalResult == EQosCompletionResult::Success)
 	{
 		// copy over region info that meets our minimum ping
 		for (const FQosRegionInfo& Region : RegionInfo)
 		{
-			if (MinimumPingMs <= 0 || Region.PingMs <= MinimumPingMs)
+			BestRegionPingMs = FMath::Min(Region.PingMs, BestRegionPingMs);
+			if ((MaximumPingMs <= 0) || (Region.PingMs <= MaximumPingMs))
 			{
 				UE_LOG(LogQos, Log, TEXT("Region: %s Ping: %d ACCEPTED"), *Region.RegionId, Region.PingMs);
 				RegionOptions.Add(Region);
@@ -182,6 +186,11 @@ FString FQosInterface::GetRegionId() const
 	return SelectedRegion;
 }
 
+int32 FQosInterface::GetLowestReportedPing() const
+{
+	return BestRegionPingMs;
+}
+
 const TArray<FQosRegionInfo>& FQosInterface::GetRegionOptions() const
 {
 	return RegionOptions;
diff --git a/Engine/Source/Runtime/Online/OnlineFramework/Qos/Public/QosInterface.h b/Engine/Source/Runtime/Online/OnlineFramework/Qos/Public/QosInterface.h
--- a/Engine/Source/Runtime/Online/OnlineFramework/Qos/Public/QosInterface.h
+++ b/Engine/Source/Runtime/Online/OnlineFramework/Qos/Public/QosInterface.h
@@ -54,6 +54,11 @@ public:
 	 */
 	FString GetRegionId() const;
 
+	/**
+	 * Returns the lowest ping of any datacenter we queried (or MAX_int32 if none were found or the queries have not completed yet)
+	 */
+	int32 GetLowestReportedPing() const;
+
 	/**
 	 * Get the list of regions that the client can choose from (returned from search and must meet min ping requirements)
 	 *
@@ -104,9 +109,14 @@ private:
 	FString GetSavedRegionId() const;
 	void SaveSelectedRegionId();
 
-	int32 MinimumPingMs;
-	FString ForceRegionId;
+	// The threshold that a region ping must be below to consider as a valid option
+	int32 MaximumPingMs;
+
+	// The best region ping we found, even if we failed all tests
+	int32 BestRegionPingMs;
 
+	FString ForceRegionId;
+	
 	TSharedPtr<FQosEvaluator> Evaluator;
 	EQosCompletionResult QosEvalResult;
 	TArray<FQosRegionInfo> RegionOptions;
diff --git a/Engine/Source/Runtime/Online/OnlineSubsystem/Private/OnlineAsyncTaskManager.cpp b/Engine/Source/Runtime/Online/OnlineSubsystem/Private/OnlineAsyncTaskManager.cpp
--- a/Engine/Source/Runtime/Online/OnlineSubsystem/Private/OnlineAsyncTaskManager.cpp
+++ b/Engine/Source/Runtime/Online/OnlineSubsystem/Private/OnlineAsyncTaskManager.cpp
@@ -9,6 +9,7 @@ int32 FOnlineAsyncTaskManager::InvocationCount = 0;
 #define POLLING_INTERVAL_MS 50
 
 FOnlineAsyncTaskManager::FOnlineAsyncTaskManager() :
+	ActiveTask(nullptr),
 	WorkEvent(nullptr),
 	PollingInterval(POLLING_INTERVAL_MS),
 	bRequestingExit(0),
@@ -38,8 +39,6 @@ uint32 FOnlineAsyncTaskManager::Run(void)
 
 	do 
 	{
-		int32 InitialQueueSize = 0;
-		int32 CurrentQueueSize = 0;
 		FOnlineAsyncTask* Task = NULL;
 
 		// Wait for a trigger event to start work
@@ -47,11 +46,6 @@ uint32 FOnlineAsyncTaskManager::Run(void)
 		if (!bRequestingExit)
 		{
 			SCOPE_CYCLE_COUNTER(STAT_Online_Async);
-			{
-				FScopeLock LockInQueue(&InQueueLock);
-				InitialQueueSize = InQueue.Num();
-				SET_DWORD_STAT(STAT_Online_AsyncTasks, InitialQueueSize);
-			}
 			
 			double TimeElapsed = 0;
 
@@ -97,50 +91,38 @@ uint32 FOnlineAsyncTaskManager::Run(void)
 			}
 
 			// Now process the serial "In" queue
-			do
-			{	
-				Task = NULL;
-
-				// Grab the current task from the queue
+			{
+				FScopeLock LockActiveTask(&ActiveTaskLock);
+				Task = ActiveTask;
+			}
+			if (Task)
+			{
+				Task->Tick();
+	
+				if (Task->IsDone())
 				{
-					FScopeLock LockInQueue(&InQueueLock);
-					CurrentQueueSize = InQueue.Num();
-					if (CurrentQueueSize > 0)
+					if (Task->WasSuccessful())
 					{
-						Task = InQueue[0];
+						UE_LOG(LogOnline, Verbose, TEXT("Async task '%s' succeeded in %f seconds"),
+							*Task->ToString(),
+							Task->GetElapsedTime());
 					}
-				}
-
-				if (Task)
-				{
-					Task->Tick();
-	
-					if (Task->IsDone())
+					else
 					{
-						if (Task->WasSuccessful())
-						{
-							UE_LOG(LogOnline, Verbose, TEXT("Async task '%s' succeeded in %f seconds"),
-								*Task->ToString(),
-								Task->GetElapsedTime());
-						}
-						else
-						{
-							UE_LOG(LogOnline, Warning, TEXT("Async task '%s' failed in %f seconds"),
-								*Task->ToString(),
-								Task->GetElapsedTime());
-						}
-						
-						// Task is done, remove from the incoming queue and add to the outgoing queue
-						PopFromInQueue();
-						AddToOutQueue(Task);
+						UE_LOG(LogOnline, Warning, TEXT("Async task '%s' failed in %f seconds"),
+							*Task->ToString(),
+							Task->GetElapsedTime());
 					}
-					else
+						
+					// Task is done, add to the outgoing queue
+					AddToOutQueue(Task);
+					
 					{
-						Task = NULL;
+						FScopeLock LockActiveTask(&ActiveTaskLock);
+						ActiveTask = nullptr;
 					}
 				}
 			}
-			while (Task != NULL);
 		}
 	} 
 	while (!bRequestingExit);
@@ -167,16 +149,6 @@ void FOnlineAsyncTaskManager::AddToInQueue(FOnlineAsyncTask* NewTask)
 {
 	FScopeLock Lock(&InQueueLock);
 	InQueue.Add(NewTask);
-	WorkEvent->Trigger();
-}
-
-void FOnlineAsyncTaskManager::PopFromInQueue()
-{
-	// assert if not game thread
-	check(FPlatformTLS::GetCurrentThreadId() == OnlineThreadId || !FPlatformProcess::SupportsMultithreading());
-
-	FScopeLock Lock(&InQueueLock);
-	InQueue.RemoveAt(0);
 }
 
 void FOnlineAsyncTaskManager::AddToOutQueue(FOnlineAsyncItem* CompletedItem)
@@ -187,6 +159,8 @@ void FOnlineAsyncTaskManager::AddToOutQueue(FOnlineAsyncItem* CompletedItem)
 
 void FOnlineAsyncTaskManager::AddToParallelTasks(FOnlineAsyncTask* NewTask)
 {
+	NewTask->Initialize();
+
 	FScopeLock LockParallelTasks(&ParallelTasksLock);
 
 	ParallelTasks.Add( NewTask );
@@ -230,62 +204,93 @@ void FOnlineAsyncTaskManager::GameTick()
 		}
 	}
 	while (Item != NULL);
+
+	int32 QueueSize = 0;
+	bool bHasActiveTask = false;
+	{
+		{
+			FScopeLock LockInQueue(&InQueueLock);
+			QueueSize = InQueue.Num();
+		}
+		{
+			FScopeLock LockActiveTask(&ActiveTaskLock);
+			if (ActiveTask != nullptr)
+			{
+				++QueueSize;
+				bHasActiveTask = true;
+			}
+		}
+
+		if (!bHasActiveTask && QueueSize > 0)
+		{
+			// Grab the current task from the queue
+			FOnlineAsyncTask* Task = nullptr;
+			{
+				FScopeLock LockInQueue(&InQueueLock);
+				Task = InQueue[0];
+				InQueue.RemoveAt(0);
+			}
+			{
+				FScopeLock LockActiveTask(&ActiveTaskLock);
+				ActiveTask = Task;
+				ActiveTask->Initialize();
+			}
+			WorkEvent->Trigger();
+		}
+	}
+
+	SET_DWORD_STAT(STAT_Online_AsyncTasks, QueueSize);
 }
 
 void FOnlineAsyncTaskManager::Tick()
 {
-	// parallel Q.
-	FOnlineAsyncTask* Task = NULL;
-
 	// Tick Online services ( possibly callbacks ). 
 	OnlineTick();
 
-	// Tick all the parallel tasks - Tick unrelated tasks together. 
+	{
+		// Tick all the parallel tasks - Tick unrelated tasks together. 
 
-	// Create a copy of existing tasks. 
-	TArray<FOnlineAsyncTask*> CopyParallelTasks = ParallelTasks;
+		// Create a copy of existing tasks. 
+		TArray<FOnlineAsyncTask*> CopyParallelTasks = ParallelTasks;
 
-	// Iterate. 
-	for (auto it = CopyParallelTasks.CreateIterator(); it; ++it)
-	{
-		Task = *it;
-		Task->Tick();
+		// parallel Q.
+		FOnlineAsyncTask* Task = NULL;
 
-		if (Task->IsDone())
+		// Iterate. 
+		for (auto it = CopyParallelTasks.CreateIterator(); it; ++it)
 		{
-			UE_LOG(LogOnline, Log, TEXT("Async parallel Task '%s' completed in %f seconds with %d (Parallel)"),
-				*Task->ToString(),
-				Task->GetElapsedTime(),
-				Task->WasSuccessful());
-
-			// Task is done, fixup the original parallel task queue. 
-			RemoveFromParallelTasks(Task);
-			AddToOutQueue(Task);
+			Task = *it;
+			Task->Tick();
+
+			if (Task->IsDone())
+			{
+				UE_LOG(LogOnline, Log, TEXT("Async parallel Task '%s' completed in %f seconds with %d (Parallel)"),
+					*Task->ToString(),
+					Task->GetElapsedTime(),
+					Task->WasSuccessful());
+
+				// Task is done, fixup the original parallel task queue. 
+				RemoveFromParallelTasks(Task);
+				AddToOutQueue(Task);
+			}
 		}
 	}
 
 	// Serial Q.
-	if (InQueue.Num() == 0)
-		return;
-
-	// Pick up the first element in the queue ( "current task" ). 
-	Task = InQueue[0];
-
-	// Needed ?
-	if (Task == NULL)
+	if (ActiveTask == nullptr)
 		return;
 
-	Task->Tick();
+	ActiveTask->Tick();
 
-	if (Task->IsDone())
+	if (ActiveTask->IsDone())
 	{
 		UE_LOG(LogOnline, Log, TEXT("Async serial task '%s' completed in %f seconds with %d"),
-			*Task->ToString(),
-			Task->GetElapsedTime(),
-			Task->WasSuccessful());
+			*ActiveTask->ToString(),
+			ActiveTask->GetElapsedTime(),
+			ActiveTask->WasSuccessful());
 
-		// Task is done, remove from the incoming queue and add to the outgoing queue
-		PopFromInQueue();
-		AddToOutQueue(Task);
+		// Task is done, add to the outgoing queue
+		AddToOutQueue(ActiveTask);
+		ActiveTask = nullptr;
 	}
 }
