Copying //UE4/Dev-Framework to //UE4/Main

==========================
MAJOR FEATURES + CHANGES
==========================

Change 2830238 on 2016/01/15 by [REDACTED_PERSON]

	Added comment and documentation for CreateAudioComponent function in FAudioDevice

Change 2830283 on 2016/01/15 by [REDACTED_PERSON]

	Adding call to XAudio2->Release() when no audio devices are found.

Change 2830286 on 2016/01/15 by [REDACTED_PERSON]

	Re-enabling voice recycling.

Change 2830287 on 2016/01/15 by [REDACTED_PERSON]

	Reseting LastLPFFrequency to FLT_MAX instead of MAX_FILTER_FREQUENCY so that the xaudio2 voice will get SetFilterParameters called on it at least once.

Change 2830336 on 2016/01/15 by [REDACTED_PERSON]

	#ue4 - Replace FMath::Fmod(x, y) with our own implementation on all platforms. Ensure on y=0 input, and return 0 in this case (rather than NaN).

	- We have been using this replacement on all platforms for a while recently because of an accidental #ifdef instead of #if PLATFORM_WINDOWS introduced about 2 months ago.

Change 2830374 on 2016/01/15 by [REDACTED_PERSON]

	#ue4 - Add comment explaining how AttachmentReplication is updated (via GatherCurrentMovement).

Change 2830520 on 2016/01/15 by [REDACTED_PERSON]

	Fix AActor::IsOverlappingActor to iterate over owned components instead of attached components

Change 2831932 on 2016/01/18 by [REDACTED_PERSON]

	Added extra debug info to help track down UE-24732.

	Issue should no longer be present in 4.11, although there is still no reliable repro, hence adding more info to logging.

Change 2832100 on 2016/01/18 by [REDACTED_PERSON]

	Fixed memory leak in cloth (1.43KB per cloth chunk per render creation) and unsafe memory management of all other vertex factories.

Change 2832497 on 2016/01/18 by [REDACTED_PERSON]

	UE-25579 Reversing coodinates of Z-axis for translating UE4 coordinates to oculus coordinates.

Change 2832511 on 2016/01/18 by [REDACTED_PERSON]

	Checking in assets for QAGame to test HRTF spatialization

Change 2832920 on 2016/01/18 by [REDACTED_PERSON]

	Fix crash on calling SetRootBody on mesh component with MasterPoseComponent on.

	- Should not create physics body for slave mesh component

	- duplicated CL with 2832738

Change 2833052 on 2016/01/18 by [REDACTED_PERSON]

	UE-24237 UE-18686 Audio not playing at Start Time set by Play node on PC

Change 2834793 on 2016/01/19 by [REDACTED_PERSON]

	Deleting asset that is no longer needed

Change 2835560 on 2016/01/20 by [REDACTED_PERSON]

	Anim graphs are now debuggable in the Persona graph editor

	Added pin properties to the debug object map so they can be tracked.
	Added code to route debug data via anim nodes for anim blueprints. This allows us to retrieve the properties of these nodes by looking them up in the correct container (i.e. the node, not the anim instance).
	Changed FBlueprintDebugData::FindClassPropertyForPin to also traverse output pins. This was seemingly some legacy code that was no longer needed and was preveting watches form being displayed even when data was present.
	Fixed SKismetDebuggingView to actually display the debugged object if one is selected. Previously it only displayed the UBlueprint itself or selected actors. The old behaviour is still there if no debug object is displayed.
	Changed UEditorEngine::GetFriendlyName to not crash if passed an actual UStruct property. The function purported to deal with UStructs but in fact assumed that all properties were UClass members by calling Property->GetOwnerClass().

Change 2835924 on 2016/01/20 by [REDACTED_PERSON]

	Bitpack FModifierKeysState key state booleans

Change 2835927 on 2016/01/20 by [REDACTED_PERSON]

	Cache on the stack returned values

Change 2835961 on 2016/01/20 by [REDACTED_PERSON]

	PR #1953: PIE: Support more then one digits in PIE prefix ("UEDPIE_[0-9]+_" instead of "UEDPIE_[0-9]_") (Contributed by [REDACTED_PERSON])

Change 2836033 on 2016/01/20 by [REDACTED_PERSON]

	Clamped cloth solver frequency to always have a frequency of at least 1

Change 2836339 on 2016/01/20 by [REDACTED_PERSON]

	Give the ability to turn off face remap table which saves memory in physx

Change 2836364 on 2016/01/20 by [REDACTED_PERSON]

	Turn off face index remap at runtime.

diff --git a/Engine/Plugins/Runtime/OculusAudio/Source/OculusAudio/Private/OculusAudio.cpp b/Engine/Plugins/Runtime/OculusAudio/Source/OculusAudio/Private/OculusAudio.cpp
--- a/Engine/Plugins/Runtime/OculusAudio/Source/OculusAudio/Private/OculusAudio.cpp
+++ b/Engine/Plugins/Runtime/OculusAudio/Source/OculusAudio/Private/OculusAudio.cpp
@@ -213,7 +213,7 @@ private:
 	/** Helper function to convert from UE coords to OVR coords. */
 	FORCEINLINE FVector ToOVRVector(const FVector& InVec) const
 	{
-		return FVector(float(InVec.Y), float(InVec.Z), float(-InVec.X));
+		return FVector(float(InVec.Y), float(-InVec.Z), float(-InVec.X));
 	}
 
 	/* Whether or not the OVR audio context is initialized. We defer initialization until the first audio callback.*/
diff --git a/Engine/Source/Editor/Kismet/Private/Debugging/SKismetDebuggingView.cpp b/Engine/Source/Editor/Kismet/Private/Debugging/SKismetDebuggingView.cpp
--- a/Engine/Source/Editor/Kismet/Private/Debugging/SKismetDebuggingView.cpp
+++ b/Engine/Source/Editor/Kismet/Private/Debugging/SKismetDebuggingView.cpp
@@ -334,20 +334,24 @@ FText FWatchLineItem::GetDescription() const
 		UObject* ParentObject = ParentObjectRef.Get();
 		if ((ParentBlueprint != ParentObject) && (ParentBlueprint != NULL))
 		{
-			if (UProperty* Property = FKismetDebugUtilities::FindClassPropertyForPin(ParentBlueprint, PinToWatch))
+			FString WatchString;
+			const FKismetDebugUtilities::EWatchTextResult WatchStatus = FKismetDebugUtilities::GetWatchText(WatchString, ParentBlueprint, ParentObject, PinToWatch);
+
+			switch (WatchStatus)
 			{
-				// Runtime debugging something
-				const uint8* SelectedDataInPIE = (uint8*)ParentObject;
+			case FKismetDebugUtilities::EWTR_Valid:
+				return FText::FromString(WatchString);
 
-				FString Description;
-				Property->ExportText_InContainer(/*ArrayElement=*/0, /*inout*/ Description, SelectedDataInPIE, SelectedDataInPIE, /*Parent=*/ NULL, PPF_PropertyWindow|PPF_BlueprintDebugView);
+			case FKismetDebugUtilities::EWTR_NotInScope:
+				return LOCTEXT("NotInScope", "Not in scope");
 
-				return FText::FromString(Description);
-			}
-			else
-			{
-				return LOCTEXT("NothingToWatch", "Nothing to watch.");
-			}
+			case FKismetDebugUtilities::EWTR_NoProperty:
+				return LOCTEXT("UnknownProperty", "No debug data");
+
+			default:
+			case FKismetDebugUtilities::EWTR_NoDebugObject:
+				return LOCTEXT("NoDebugObject", "No debug object");
+			}			
 		}
 	}
 
@@ -1016,17 +1020,22 @@ TSharedPtr<SWidget> SKismetDebuggingView::OnMakeContextMenu()
 
 FText SKismetDebuggingView::GetTopText() const
 {
-	if (GEditor->PlayWorld != NULL)
+	const bool bIsDebugging = GEditor->PlayWorld != nullptr;
+	UBlueprint* BlueprintObj = BlueprintToWatchPtr.Get();
+	UObject* DebuggedObject = BlueprintObj->GetObjectBeingDebugged();
+
+	TSet<UObject*> NewRootSet;
+	if (bIsDebugging && (BlueprintObj != nullptr) && (DebuggedObject != nullptr))
 	{
-		return LOCTEXT("ShowDebugForActors", "Showing debug info for selected actors");
+		return FText::Format(LOCTEXT("ShowDebugForObject", "Showing debug info for {0}"), FText::FromName(DebuggedObject->GetFName()));
 	}
-	else if (BlueprintToWatchPtr.Get() != NULL)
+	else if (!bIsDebugging && (BlueprintObj != nullptr))
 	{
 		return LOCTEXT("ShowDebugForBlueprint", "Showing debug info for this blueprint");
 	}
 	else
 	{
-		return LOCTEXT("Idle", "Idle");
+		return LOCTEXT("ShowDebugForActors", "Showing debug info for selected actors");
 	}
 }
 
@@ -1098,11 +1107,17 @@ void SKismetDebuggingView::Tick( const FGeometry& AllottedGeometry, const double
 	}
 
 	// Gather what we'd like to be the new root set
-	const bool bIsDebugging = GEditor->PlayWorld != NULL;
+	const bool bIsDebugging = GEditor->PlayWorld != nullptr;
 	UBlueprint* BlueprintObj = BlueprintToWatchPtr.Get();
+	UObject* DebuggedObject = BlueprintObj->GetObjectBeingDebugged();
 
 	TSet<UObject*> NewRootSet;
-	if (!bIsDebugging && (BlueprintObj != NULL))
+	if (bIsDebugging && (BlueprintObj != nullptr) && (DebuggedObject != nullptr))
+	{
+		// If we have a currently debugged object & we were launched from a specific Kismet window, display the debugged object
+		NewRootSet.Add(DebuggedObject);
+	}
+	else if (!bIsDebugging && (BlueprintObj != nullptr))
 	{
 		// If not debugging and summoned from a specific Kismet window, just display the currently open blueprint
 		NewRootSet.Add(BlueprintObj);
diff --git a/Engine/Source/Editor/KismetCompiler/Private/AnimBlueprintCompiler.cpp b/Engine/Source/Editor/KismetCompiler/Private/AnimBlueprintCompiler.cpp
--- a/Engine/Source/Editor/KismetCompiler/Private/AnimBlueprintCompiler.cpp
+++ b/Engine/Source/Editor/KismetCompiler/Private/AnimBlueprintCompiler.cpp
@@ -371,6 +371,12 @@ void FAnimBlueprintCompiler::ProcessAnimationNode(UAnimGraphNode_Base* VisualAni
 					EvaluationHandlers.FindOrAdd(EvaluationHandlerName).RegisterPin(SourcePin, SourcePinProperty, SourceArrayIndex);
 					bConsumed = true;
 				}
+
+				UEdGraphPin* TrueSourcePin = MessageLog.FindSourceObjectTypeChecked<UEdGraphPin>(SourcePin);
+				if (TrueSourcePin)
+				{
+					NewAnimBlueprintClass->GetDebugData().RegisterClassPropertyAssociation(TrueSourcePin, SourcePinProperty);
+				}
 			}
 		}
 
diff --git a/Engine/Source/Editor/UnrealEd/Classes/Editor/EditorEngine.h b/Engine/Source/Editor/UnrealEd/Classes/Editor/EditorEngine.h
--- a/Engine/Source/Editor/UnrealEd/Classes/Editor/EditorEngine.h
+++ b/Engine/Source/Editor/UnrealEd/Classes/Editor/EditorEngine.h
@@ -2282,12 +2282,12 @@ public:
 	 * Gets the user-friendly, localized (if exists) name of a property
 	 *
 	 * @param	Property	the property we want to try to et the friendly name of	
-	 * @param	OwnerClass	if specified, uses this class's loc file instead of the property's owner class
+	 * @param	OwnerStruct	if specified, uses this class's loc file instead of the property's owner class
 	 *						useful for overriding the friendly name given a property inherited from a parent class.
 	 *
 	 * @return	the friendly name for the property.  localized first, then metadata, then the property's name.
 	 */
-	static FString GetFriendlyName( const UProperty* Property, UStruct* OwnerClass = NULL );
+	static FString GetFriendlyName( const UProperty* Property, UStruct* OwnerStruct = NULL );
 
 	/**
 	 * Register a client tool to receive undo events 
diff --git a/Engine/Source/Editor/UnrealEd/Private/EditorEngine.cpp b/Engine/Source/Editor/UnrealEd/Private/EditorEngine.cpp
--- a/Engine/Source/Editor/UnrealEd/Private/EditorEngine.cpp
+++ b/Engine/Source/Editor/UnrealEd/Private/EditorEngine.cpp
@@ -4189,27 +4189,27 @@ void UEditorEngine::SetActorLabelUnique(AActor* Actor, const FString& NewActorLa
 	FActorLabelUtilities::SetActorLabelUnique(Actor, NewActorLabel, InExistingActorLabels);
 }
 
-FString UEditorEngine::GetFriendlyName( const UProperty* Property, UStruct* OwnerClass/* = NULL*/ )
+FString UEditorEngine::GetFriendlyName( const UProperty* Property, UStruct* OwnerStruct/* = NULL*/ )
 {
 	// first, try to pull the friendly name from the loc file
 	check( Property );
-	UClass* RealOwnerClass = Property->GetOwnerClass();
-	if ( OwnerClass == NULL)
+	UStruct* RealOwnerStruct = Property->GetOwnerStruct();
+	if ( OwnerStruct == NULL)
 	{
-		OwnerClass = RealOwnerClass;
+		OwnerStruct = RealOwnerStruct;
 	}
-	checkSlow(OwnerClass);
+	checkSlow(OwnerStruct);
 
 	FText FoundText;
 	bool DidFindText = false;
-	UStruct* CurrentClass = OwnerClass;
+	UStruct* CurrentStruct = OwnerStruct;
 	do 
 	{
-		FString PropertyPathName = Property->GetPathName(CurrentClass);
+		FString PropertyPathName = Property->GetPathName(CurrentStruct);
 
-		DidFindText = FText::FindText(*CurrentClass->GetName(), *(PropertyPathName + TEXT(".FriendlyName")), /*OUT*/FoundText );
-		CurrentClass = CurrentClass->GetSuperStruct();
-	} while( CurrentClass != NULL && CurrentClass->IsChildOf(RealOwnerClass) && !DidFindText );
+		DidFindText = FText::FindText(*CurrentStruct->GetName(), *(PropertyPathName + TEXT(".FriendlyName")), /*OUT*/FoundText );
+		CurrentStruct = CurrentStruct->GetSuperStruct();
+	} while( CurrentStruct != NULL && CurrentStruct->IsChildOf(RealOwnerStruct) && !DidFindText );
 
 	if ( !DidFindText )
 	{
diff --git a/Engine/Source/Editor/UnrealEd/Private/Kismet2/KismetDebugUtilities.cpp b/Engine/Source/Editor/UnrealEd/Private/Kismet2/KismetDebugUtilities.cpp
--- a/Engine/Source/Editor/UnrealEd/Private/Kismet2/KismetDebugUtilities.cpp
+++ b/Engine/Source/Editor/UnrealEd/Private/Kismet2/KismetDebugUtilities.cpp
@@ -13,6 +13,7 @@
 #include "Editor/KismetCompiler/Public/KismetCompilerModule.h"
 #include "MessageLog.h"
 #include "UObjectToken.h"
+#include "AnimGraphNode_Base.h"
 
 #define LOCTEXT_NAMESPACE "BlueprintDebugging"
 
@@ -1065,6 +1066,35 @@ FKismetDebugUtilities::EWatchTextResult FKismetDebugUtilities::GetWatchText(FStr
 			}
 #endif // USE_UBER_GRAPH_PERSISTENT_FRAME
 			
+			// see if our WatchPin is on a animation node & if so try to get its property info
+			UAnimBlueprintGeneratedClass* AnimBlueprintGeneratedClass = Cast<UAnimBlueprintGeneratedClass>(Blueprint->GeneratedClass);
+			if(!PropertyBase && AnimBlueprintGeneratedClass)
+			{
+				// are we linked to an anim graph node?
+				const UAnimGraphNode_Base* Node = Cast<UAnimGraphNode_Base>(WatchPin->GetOuter());
+				if(Node == nullptr && WatchPin->LinkedTo.Num() > 0)
+				{
+					Node = Cast<UAnimGraphNode_Base>(WatchPin->LinkedTo[0]->GetOuter());
+				}
+
+				if(Node)
+				{
+					UStructProperty* NodeStructProperty = Cast<UStructProperty>(FKismetDebugUtilities::FindClassPropertyForNode(Blueprint, Node));
+					if (NodeStructProperty)
+					{
+						for (UStructProperty* NodeProperty : AnimBlueprintGeneratedClass->AnimNodeProperties)
+						{
+							if (NodeProperty == NodeStructProperty)
+							{
+								void* NodePtr = NodeProperty->ContainerPtrToValuePtr<void>(ActiveObject);
+								Property->ExportText_InContainer(/*ArrayElement=*/ 0, /*inout*/ OutWatchText, NodePtr, NodePtr, /*Parent=*/ ActiveObject, PPF_PropertyWindow|PPF_BlueprintDebugView);
+								return EWTR_Valid;
+							}
+						}
+					}
+				}
+			}
+
 			// Now either print out the variable value, or that it was out-of-scope
 			if (PropertyBase != nullptr)
 			{
diff --git a/Engine/Source/Runtime/Core/Private/GenericPlatform/GenericPlatformMath.cpp b/Engine/Source/Runtime/Core/Private/GenericPlatform/GenericPlatformMath.cpp
--- a/Engine/Source/Runtime/Core/Private/GenericPlatform/GenericPlatformMath.cpp
+++ b/Engine/Source/Runtime/Core/Private/GenericPlatform/GenericPlatformMath.cpp
@@ -72,6 +72,11 @@ float FGenericPlatformMath::Atan2(float Y, float X)
 	return t3;
 }
 
+void FGenericPlatformMath::FmodReportError(float X, float Y)
+{
+	ensureMsgf(Y != 0, TEXT("FMath::FMod(X=%f, Y=%f) : Y is zero, this is invalid and would result in NaN!"), X, Y);
+}
+
 extern float TheCompilerDoesntKnowThisIsAlwaysZero;
 
 void FGenericPlatformMath::AutoTest() 
diff --git a/Engine/Source/Runtime/Core/Private/Tests/Math/UnrealMathTest.cpp b/Engine/Source/Runtime/Core/Private/Tests/Math/UnrealMathTest.cpp
--- a/Engine/Source/Runtime/Core/Private/Tests/Math/UnrealMathTest.cpp
+++ b/Engine/Source/Runtime/Core/Private/Tests/Math/UnrealMathTest.cpp
@@ -594,10 +594,10 @@ MATHTEST_INLINE VectorRegister TestVectorNormalize_InvSqrtEst(const VectorRegist
 MATHTEST_INLINE VectorRegister TestReferenceMod(const VectorRegister& A, const VectorRegister& M)
 {
 	return MakeVectorRegister(
-		(float)fmod(VectorGetComponent(A, 0), VectorGetComponent(M, 0)),
-		(float)fmod(VectorGetComponent(A, 1), VectorGetComponent(M, 1)),
-		(float)fmod(VectorGetComponent(A, 2), VectorGetComponent(M, 2)),
-		(float)fmod(VectorGetComponent(A, 3), VectorGetComponent(M, 3)));
+		(float)fmodf(VectorGetComponent(A, 0), VectorGetComponent(M, 0)),
+		(float)fmodf(VectorGetComponent(A, 1), VectorGetComponent(M, 1)),
+		(float)fmodf(VectorGetComponent(A, 2), VectorGetComponent(M, 2)),
+		(float)fmodf(VectorGetComponent(A, 3), VectorGetComponent(M, 3)));
 }
 
 // SinCos
@@ -1367,6 +1367,62 @@ bool FVectorRegisterAbstractionTest::RunTest(const FString& Parameters)
 		LogTest(TEXT("VectorQuaternionMultiply2"), TestVectorsEqual(V2, V3, 1e-6f));
 	}
 
+	// FMath::FMod
+	{
+		struct XYPair
+		{
+			float X, Y;
+		};
+
+		static XYPair XYArray[] =
+		{		
+			// Test normal ranges
+			{ 0.0,	 1.0},
+			{ 1.5,	 1.0},
+			{ 2.8,	 0.3},
+			{-2.8,	 0.3},
+			{ 2.8,	-0.3},
+			{-2.8,	-0.3},
+			{-0.4,	 5.5},
+			{ 0.4,	-5.5},
+			{ 2.8,	 2.0 + KINDA_SMALL_NUMBER},
+			{-2.8,	 2.0 - KINDA_SMALL_NUMBER},
+
+			// Commonly used for FRotators and angles
+			{725.2,		360.0},
+			{179.9,		 90.0},
+			{ 5.3*PI,	2.*PI},
+			{-5.3*PI,	2.*PI},
+
+			// Test extreme ranges
+			{ 1.0,			 KINDA_SMALL_NUMBER},
+			{ 1.0,			-KINDA_SMALL_NUMBER},
+			{-SMALL_NUMBER,  SMALL_NUMBER},
+			{ SMALL_NUMBER, -SMALL_NUMBER},
+			{ 1.0,			 MIN_flt},
+			{ 1.0,			-MIN_flt},
+			{ MAX_flt,		 MIN_flt},
+			{ MAX_flt,		-MIN_flt},
+
+			// We define this to be zero and not NaN.
+			// Disabled since we don't want to trigger an ensure, but left here for testing that logic.
+			//{ 1.0,	 0.0}, 
+			//{ 1.0,	-0.0},
+		};
+
+		for (auto XY : XYArray)
+		{
+			const float X = XY.X;
+			const float Y = XY.Y;
+			const float Ours = FMath::Fmod(X, Y);
+			const float Theirs = fmodf(X, Y);
+			if (!FMath::IsNearlyEqual(Ours, Theirs, 1.e-5f))
+			{
+				UE_LOG(LogUnrealMathTest, Log, TEXT("FMath::Fmod(%f, %f)=%f <-> fmodf(%f, %f)=%f: FAILED"), X, Y, Ours, X, Y, Theirs);
+				GPassing = false;
+			}
+		}
+	}
 
 	if (!GPassing)
 	{
diff --git a/Engine/Source/Runtime/Core/Public/GenericPlatform/GenericApplication.h b/Engine/Source/Runtime/Core/Public/GenericPlatform/GenericApplication.h
--- a/Engine/Source/Runtime/Core/Public/GenericPlatform/GenericApplication.h
+++ b/Engine/Source/Runtime/Core/Public/GenericPlatform/GenericApplication.h
@@ -270,31 +270,31 @@ public:
 private:
 
 	/** True if the left shift key was down when this event occurred. */
-	bool bIsLeftShiftDown;
+	uint16 bIsLeftShiftDown:1;
 
 	/** True if the right shift key was down when this event occurred. */
-	bool bIsRightShiftDown;
+	uint16 bIsRightShiftDown:1;
 
 	/** True if the left control key was down when this event occurred. */
-	bool bIsLeftControlDown;
+	uint16 bIsLeftControlDown:1;
 
 	/** True if the right control key was down when this event occurred. */
-	bool bIsRightControlDown;
+	uint16 bIsRightControlDown:1;
 
 	/** True if the left alt key was down when this event occurred. */
-	bool bIsLeftAltDown;
+	uint16 bIsLeftAltDown:1;
 
 	/** True if the right alt key was down when this event occurred. */
-	bool bIsRightAltDown;
+	uint16 bIsRightAltDown:1;
 	
 	/** True if the left command key was down when this event occurred. */
-	bool bIsLeftCommandDown;
+	uint16 bIsLeftCommandDown:1;
 	
 	/** True if the right command key was down when this event occurred. */
-	bool bIsRightCommandDown;
+	uint16 bIsRightCommandDown:1;
 
 	/** True if the Caps Lock key has been toggled to the enabled state. */
-	bool bAreCapsLocked;
+	uint16 bAreCapsLocked:1;
 };
 
 
diff --git a/Engine/Source/Runtime/Core/Public/GenericPlatform/GenericPlatformMath.h b/Engine/Source/Runtime/Core/Public/GenericPlatform/GenericPlatformMath.h
--- a/Engine/Source/Runtime/Core/Public/GenericPlatform/GenericPlatformMath.h
+++ b/Engine/Source/Runtime/Core/Public/GenericPlatform/GenericPlatformMath.h
@@ -236,18 +236,16 @@ struct FGenericPlatformMath
 	*			So for example Fmod(2.8f, 2) gives .8f as you would expect, however, Fmod(-2.8f, 2) gives -.8f, NOT 1.2f 
 	* Use Floor instead when snapping positions that can be negative to a grid
 	*/
-	static FORCEINLINE_DEBUGGABLE float Fmod(float X, float Y)
+	static FORCEINLINE float Fmod(float X, float Y)
 	{
-#ifdef PLATFORM_WINDOWS
-		// There's a compiler bug on Windows, where fmodf will start returning NaNs randomly with valid inputs.
-		// Until this is resolved, we implement our own version.
-		float IntPortion = TruncToFloat(X / Y);
-		float Result = X - Y * IntPortion;
-
+		if (fabsf(Y) <= 1.e-8f)
+		{
+			FmodReportError(X, Y);
+			return 0.f;
+		}
+		const float IntPortion = TruncToFloat(X / Y);
+		const float Result = X - Y * IntPortion;
 		return Result;
-#else
-		return fmodf(X, Y);
-#endif		
 	}
 
 	static FORCEINLINE float Sin( float Value ) { return sinf(Value); }
@@ -618,6 +616,11 @@ struct FGenericPlatformMath
 
 	/** Test some of the tricky functions above **/
 	static void AutoTest();
+
+private:
+
+	/** Error reporting for Fmod. Not inlined to avoid compilation issues and avoid all the checks and error reporting at all callsites. */
+	static CORE_API void FmodReportError(float X, float Y);
 };
 
 /** Float specialization */
diff --git a/Engine/Source/Runtime/Engine/Classes/Components/SkeletalMeshComponent.h b/Engine/Source/Runtime/Engine/Classes/Components/SkeletalMeshComponent.h
--- a/Engine/Source/Runtime/Engine/Classes/Components/SkeletalMeshComponent.h
+++ b/Engine/Source/Runtime/Engine/Classes/Components/SkeletalMeshComponent.h
@@ -901,6 +901,7 @@ public:
 	virtual void OnRegister() override;
 	virtual void OnUnregister() override;
 	virtual void CreateRenderState_Concurrent() override;
+	virtual bool ShouldCreatePhysicsState() const override;
 	virtual void CreatePhysicsState() override;
 	virtual void DestroyPhysicsState() override;
 	virtual void InitializeComponent() override;
diff --git a/Engine/Source/Runtime/Engine/Classes/Engine/BlueprintGeneratedClass.h b/Engine/Source/Runtime/Engine/Classes/Engine/BlueprintGeneratedClass.h
--- a/Engine/Source/Runtime/Engine/Classes/Engine/BlueprintGeneratedClass.h
+++ b/Engine/Source/Runtime/Engine/Classes/Engine/BlueprintGeneratedClass.h
@@ -295,7 +295,7 @@ public:
 	UProperty* FindClassPropertyForPin(const UEdGraphPin* Pin) const
 	{
 		UProperty* PropertyPtr = DebugPinToPropertyMap.FindRef(Pin);
-		if ((PropertyPtr == nullptr) && (Pin->Direction == EGPD_Input) && (Pin->LinkedTo.Num() > 0))
+		if ((PropertyPtr == nullptr) && (Pin->LinkedTo.Num() > 0))
 		{
 			// Try checking the other side of the connection
 			PropertyPtr = DebugPinToPropertyMap.FindRef(Pin->LinkedTo[0]);
diff --git a/Engine/Source/Runtime/Engine/Classes/Engine/SkeletalMesh.h b/Engine/Source/Runtime/Engine/Classes/Engine/SkeletalMesh.h
--- a/Engine/Source/Runtime/Engine/Classes/Engine/SkeletalMesh.h
+++ b/Engine/Source/Runtime/Engine/Classes/Engine/SkeletalMesh.h
@@ -466,7 +466,7 @@ struct FClothPhysicsProperties
 	float SelfCollisionStiffness;
 
 	// A computation parameter for the Solver.   Along with frame rate this probably specifies the number of solver iterations
-	UPROPERTY(EditAnywhere, Category = Solver, meta = (ClampMin = "0.0", UIMin = "0.0", UIMax = "1000.0"))
+	UPROPERTY(EditAnywhere, Category = Solver, meta = (ClampMin = "1.0", UIMin = "1.0", UIMax = "1000.0"))
 	float SolverFrequency;
 
 
diff --git a/Engine/Source/Runtime/Engine/Classes/GameFramework/Actor.h b/Engine/Source/Runtime/Engine/Classes/GameFramework/Actor.h
--- a/Engine/Source/Runtime/Engine/Classes/GameFramework/Actor.h
+++ b/Engine/Source/Runtime/Engine/Classes/GameFramework/Actor.h
@@ -268,7 +268,10 @@ public:
 	struct FRepMovement ReplicatedMovement;
 
 private:
-	/** Used for replicating attachment of this actor's RootComponent to another actor. */
+	/**
+	 * Used for replicating attachment of this actor's RootComponent to another actor.
+	 * This is filled in via GatherCurrentMovement() when the RootComponent has an AttachParent.
+	 */
 	UPROPERTY(Transient, replicatedUsing=OnRep_AttachmentReplication)
 	struct FRepAttachment AttachmentReplication;
 
diff --git a/Engine/Source/Runtime/Engine/Classes/PhysicsEngine/BodySetup.h b/Engine/Source/Runtime/Engine/Classes/PhysicsEngine/BodySetup.h
--- a/Engine/Source/Runtime/Engine/Classes/PhysicsEngine/BodySetup.h
+++ b/Engine/Source/Runtime/Engine/Classes/PhysicsEngine/BodySetup.h
@@ -118,6 +118,15 @@ class UBodySetup : public UObject
 	/** Cooked physics data for each format */
 	FFormatContainer CookedFormatData;
 
+private:
+#if WITH_EDITOR
+	/** Cooked physics data with runtime only optimizations. This allows us to remove editor only data (like face index remap) assuming the project doesn't use it at runtime. At runtime we load this into CookedFormatData */
+	FFormatContainer CookedFormatDataRuntimeOnlyOptimization;
+#endif
+
+	int32 GetRuntimeOnlyCookOptimizationFlags() const;
+public:
+
 	/** Cooked physics data override. This is needed in cases where some other body setup has the cooked data and you don't want to own it or copy it. See per poly skeletal mesh */
 	FFormatContainer* CookedFormatDataOverride;
 
@@ -233,9 +242,15 @@ public:
 	 * Given a format name returns its cooked data.
 	 *
 	 * @param Format Physics format name.
+	 * @param bRuntimeOnlyOptimizedVersion whether we want the data that has runtime only optimizations. At runtime this flag is ignored and we use the runtime only optimized data regardless.
 	 * @return Cooked data or NULL of the data was not found.
 	 */
-	FByteBulkData* GetCookedData(FName Format);
+	FByteBulkData* GetCookedData(FName Format, bool bRuntimeOnlyOptimizedVersion = false);
+
+#if WITH_EDITOR
+	ENGINE_API virtual void BeginCacheForCookedPlatformData(  const ITargetPlatform* TargetPlatform ) override;
+	ENGINE_API virtual void ClearCachedCookedPlatformData(  const ITargetPlatform* TargetPlatform ) override;
+#endif
 
 #if WITH_PHYSX
 	/** 
diff --git a/Engine/Source/Runtime/Engine/Classes/PhysicsEngine/PhysicsSettings.h b/Engine/Source/Runtime/Engine/Classes/PhysicsEngine/PhysicsSettings.h
--- a/Engine/Source/Runtime/Engine/Classes/PhysicsEngine/PhysicsSettings.h
+++ b/Engine/Source/Runtime/Engine/Classes/PhysicsEngine/PhysicsSettings.h
@@ -157,6 +157,11 @@ class ENGINE_API UPhysicsSettings : public UDeveloperSettings
 	UPROPERTY(config, EditAnywhere, Category = Simulation)
 	bool bDefaultHasComplexCollision;
 
+	/**
+	*  If true, the internal physx face to UE face mapping will not be generated. This is a memory optimization available if you do not rely on face indices returned by scene queries. */
+	UPROPERTY(config, EditAnywhere, Category = Optimization)
+	bool bSupressFaceRemapTable;
+
 	/**
 	*  If true CCD will be ignored. This is an optimization when CCD is never used which removes the need for physx to check it internally. */
 	UPROPERTY(config, EditAnywhere, Category = Simulation)
diff --git a/Engine/Source/Runtime/Engine/Private/ActiveSound.cpp b/Engine/Source/Runtime/Engine/Private/ActiveSound.cpp
--- a/Engine/Source/Runtime/Engine/Private/ActiveSound.cpp
+++ b/Engine/Source/Runtime/Engine/Private/ActiveSound.cpp
@@ -186,12 +186,6 @@ void FActiveSound::UpdateWaveInstances( TArray<FWaveInstance*> &InWaveInstances,
 	//@todo audio: Fading in and out is also dependent on the DeltaTime
 	PlaybackTime += DeltaTime;
 
-	// Don't clear the seeking unless the sound is actually playing
-	if (DeltaTime > 0.f)
-	{
-		RequestedStartTime = 0.f;
-	}
-
 	// splitscreen support:
 	// we always pass the 'primary' listener (viewport 0) to the sound nodes and the underlying audio system
 	// then move the AudioComponent's CurrentLocation so that its position relative to that Listener is the same as its real position is relative to the closest Listener
diff --git a/Engine/Source/Runtime/Engine/Private/Actor.cpp b/Engine/Source/Runtime/Engine/Private/Actor.cpp
--- a/Engine/Source/Runtime/Engine/Private/Actor.cpp
+++ b/Engine/Source/Runtime/Engine/Private/Actor.cpp
@@ -1115,30 +1115,17 @@ void AActor::UpdateOverlaps(bool bDoNotifies)
 
 bool AActor::IsOverlappingActor(const AActor* Other) const
 {
-	// use a stack to walk this actor's attached component tree
-	TInlineComponentArray<USceneComponent*> ComponentStack;
-	USceneComponent const* CurrentComponent = GetRootComponent();
-	while (CurrentComponent)
+	for (UActorComponent* OwnedComp : OwnedComponents)
 	{
-		// push children on the stack so they get tested later
-		ComponentStack.Append(CurrentComponent->AttachChildren);
-
-		if(CurrentComponent->GetOwner() == this)	//The component could be attached but be from a different actor
+		if (UPrimitiveComponent* PrimComp = Cast<UPrimitiveComponent>(OwnedComp))
 		{
-			UPrimitiveComponent const* const PrimComp = Cast<const UPrimitiveComponent>(CurrentComponent);
-			if (PrimComp && PrimComp->IsOverlappingActor(Other))
+			if (PrimComp->IsOverlappingActor(Other))
 			{
 				// found one, finished
 				return true;
 			}
 		}
-		
-
-		// advance to next component
-		const bool bAllowShrinking = false;
-		CurrentComponent = (ComponentStack.Num() > 0) ? ComponentStack.Pop(bAllowShrinking) : nullptr;
 	}
-
 	return false;
 }
 
diff --git a/Engine/Source/Runtime/Engine/Private/Animation/AnimNode_StateMachine.cpp b/Engine/Source/Runtime/Engine/Private/Animation/AnimNode_StateMachine.cpp
--- a/Engine/Source/Runtime/Engine/Private/Animation/AnimNode_StateMachine.cpp
+++ b/Engine/Source/Runtime/Engine/Private/Animation/AnimNode_StateMachine.cpp
@@ -351,7 +351,10 @@ void FAnimNode_StateMachine::Update(const FAnimationUpdateContext& Context)
 		else if(!Machine->States.IsValidIndex(CurrentState))
 		{
 			// Attempting to catch a crash where the state machine has been freed.
-			UE_LOG(LogAnimation, Warning, TEXT("FAnimNode_StateMachine::Update - Invalid current state, please report. Attempting to use state %d in state machine %d"), CurrentState, StateMachineIndexInClass);
+			// Reported as a symptom of a crash in UE-24732 for 4.10. This log message should not appear given changes to
+			// re-instancing in 4.11 (see CL 2823202). If it does appear we need to spot integrate CL 2823202 (and supporting 
+			// anim re-init changes, probably 2799786 & 2801372).
+			UE_LOG(LogAnimation, Warning, TEXT("FAnimNode_StateMachine::Update - Invalid current state, please report. Attempting to use state %d of %d in state machine %d (ptr 0x%x)"), CurrentState, Machine->States.Num(), StateMachineIndexInClass, Machine);
 			UE_LOG(LogAnimation, Warning, TEXT("\t\tWhen updating AnimInstance: %s"), *Context.AnimInstanceProxy->GetAnimInstanceObject()->GetName())
 
 			return;
diff --git a/Engine/Source/Runtime/Engine/Private/Collision/CollisionConversions.cpp b/Engine/Source/Runtime/Engine/Private/Collision/CollisionConversions.cpp
--- a/Engine/Source/Runtime/Engine/Private/Collision/CollisionConversions.cpp
+++ b/Engine/Source/Runtime/Engine/Private/Collision/CollisionConversions.cpp
@@ -430,6 +430,16 @@ EConvertQueryResult ConvertQueryImpactHit(const UWorld* World, const PxLocationH
 {
 	SCOPE_CYCLE_COUNTER(STAT_ConvertQueryImpactHit);
 
+#if WITH_EDITOR
+	if(bReturnFaceIndex && World->IsGameWorld())
+	{
+		if(!ensure(UPhysicsSettings::Get()->bSupressFaceRemapTable == false))
+		{
+			UE_LOG(LogPhysics, Error, TEXT("A scene query is relying on face indices, but bSupressFaceRemapTable is false."));
+		}
+	}
+#endif
+
 	checkSlow(PHit.flags & PxHitFlag::eDISTANCE);
 	const bool bInitialOverlap = PHit.hadInitialOverlap();
 	if (bInitialOverlap && Geom != nullptr)
diff --git a/Engine/Source/Runtime/Engine/Private/Components/SkinnedMeshComponent.cpp b/Engine/Source/Runtime/Engine/Private/Components/SkinnedMeshComponent.cpp
--- a/Engine/Source/Runtime/Engine/Private/Components/SkinnedMeshComponent.cpp
+++ b/Engine/Source/Runtime/Engine/Private/Components/SkinnedMeshComponent.cpp
@@ -1100,6 +1100,7 @@ void USkinnedMeshComponent::SetMasterPoseComponent(class USkinnedMeshComponent*
 	}
 
 	AllocateTransformData();
+	RecreatePhysicsState();
 	UpdateMasterBoneMap();
 }
 
diff --git a/Engine/Source/Runtime/Engine/Private/PhysicsEngine/BodySetup.cpp b/Engine/Source/Runtime/Engine/Private/PhysicsEngine/BodySetup.cpp
--- a/Engine/Source/Runtime/Engine/Private/PhysicsEngine/BodySetup.cpp
+++ b/Engine/Source/Runtime/Engine/Private/PhysicsEngine/BodySetup.cpp
@@ -752,20 +752,25 @@ void UBodySetup::Serialize(FArchive& Ar)
 
 	if (bCooked)
 	{
+#if WITH_EDITOR
 		if (Ar.IsCooking())
 		{
 			// Make sure to reset bHasCookedCollision data to true before calling GetCookedData for cooking
 			bHasCookedCollisionData = true;
 			FName Format = Ar.CookingTarget()->GetPhysicsFormat(this);
-			bHasCookedCollisionData = GetCookedData(Format) != NULL; // Get the data from the DDC or build it
+			bool bUseRuntimeOnlyCookedData = !bSharedCookedData;	//For shared cook data we do not optimize for runtime only flags. This is only used by per poly skeletal mesh component at the moment. Might want to add support in future
+			bHasCookedCollisionData = GetCookedData(Format, bUseRuntimeOnlyCookedData) != NULL; // Get the data from the DDC or build it
 
 			TArray<FName> ActualFormatsToSave;
 			ActualFormatsToSave.Add(Format);
 
 			Ar << bHasCookedCollisionData;
-			CookedFormatData.Serialize(Ar, this, &ActualFormatsToSave, !bSharedCookedData);
+
+			FFormatContainer* UseCookedFormatData = bUseRuntimeOnlyCookedData ? &CookedFormatDataRuntimeOnlyOptimization : &CookedFormatData;
+			UseCookedFormatData->Serialize(Ar, this, &ActualFormatsToSave, !bSharedCookedData);
 		}
 		else
+#endif
 		{
 			if (Ar.UE4Ver() >= VER_UE4_STORE_HASCOOKEDDATA_FOR_BODYSETUP)
 			{
@@ -960,7 +965,30 @@ float UBodySetup::GetClosestPointAndNormal(const FVector& WorldPosition, const F
 	return GetClosestPointAndNormalImpl<true>(this, WorldPosition, LocalToWorld, &ClosestWorldPosition, &FeatureNormal);
 }
 
-FByteBulkData* UBodySetup::GetCookedData(FName Format)
+#if WITH_EDITOR
+void UBodySetup::BeginCacheForCookedPlatformData(const ITargetPlatform* TargetPlatform)
+{
+	GetCookedData(TargetPlatform->GetPhysicsFormat(this), true);
+}
+
+void UBodySetup::ClearCachedCookedPlatformData( const ITargetPlatform* TargetPlatform )
+{
+	CookedFormatDataRuntimeOnlyOptimization.FlushData();
+}
+#endif
+
+int32 UBodySetup::GetRuntimeOnlyCookOptimizationFlags() const
+{
+	int32 RuntimeCookFlags = 0;
+	if(UPhysicsSettings::Get()->bSupressFaceRemapTable)
+	{
+		RuntimeCookFlags |= ERuntimePhysxCookOptimizationFlags::SupressFaceRemapTable;
+	}
+
+	return RuntimeCookFlags;
+}
+
+FByteBulkData* UBodySetup::GetCookedData(FName Format, bool bRuntimeOnlyOptimizedVersion)
 {
 	if (IsTemplate())
 	{
@@ -976,7 +1004,12 @@ FByteBulkData* UBodySetup::GetCookedData(FName Format)
 		return NULL;
 	}
 
+#if WITH_EDITOR
+	//We don't support runtime cook optimization for per poly skeletal mesh. This is an edge case we may want to support (only helps memory savings)
+	FFormatContainer* UseCookedData = CookedFormatDataOverride ? CookedFormatDataOverride : (bRuntimeOnlyOptimizedVersion ? &CookedFormatDataRuntimeOnlyOptimization : &CookedFormatData);
+#else
 	FFormatContainer* UseCookedData = CookedFormatDataOverride ? CookedFormatDataOverride : &CookedFormatData;
+#endif
 
 	bool bContainedData = UseCookedData->Contains(Format);
 	FByteBulkData* Result = &UseCookedData->GetFormat(Format);
@@ -993,9 +1026,17 @@ FByteBulkData* UBodySetup::GetCookedData(FName Format)
 			return NULL;
 		}
 
+#if WITH_EDITOR
+		const bool bEligibleForRuntimeOptimization = UseCookedData == &CookedFormatDataRuntimeOnlyOptimization;
+#else
+		const bool bEligibleForRuntimeOptimization = CookedFormatDataOverride == nullptr;	//We don't support runtime cook optimization for per poly skeletal mesh. This is an edge case we may want to support (only helps memory savings)
+#endif
+
 #if WITH_RUNTIME_PHYSICS_COOKING || WITH_EDITOR
 		TArray<uint8> OutData;
-		FDerivedDataPhysXCooker* DerivedPhysXData = new FDerivedDataPhysXCooker(Format, this);
+
+		const int32 CookingFlags = bEligibleForRuntimeOptimization ? GetRuntimeOnlyCookOptimizationFlags() : 0;
+		FDerivedDataPhysXCooker* DerivedPhysXData = new FDerivedDataPhysXCooker(Format, CookingFlags, this);
 		if (DerivedPhysXData->CanBuild())
 		{
 		#if WITH_EDITOR
diff --git a/Engine/Source/Runtime/Engine/Private/PhysicsEngine/PhysDerivedData.cpp b/Engine/Source/Runtime/Engine/Private/PhysicsEngine/PhysDerivedData.cpp
--- a/Engine/Source/Runtime/Engine/Private/PhysicsEngine/PhysDerivedData.cpp
+++ b/Engine/Source/Runtime/Engine/Private/PhysicsEngine/PhysDerivedData.cpp
@@ -7,10 +7,11 @@
 
 #include "IPhysXFormatModule.h"
 
-FDerivedDataPhysXCooker::FDerivedDataPhysXCooker( FName InFormat, UBodySetup* InBodySetup )
+FDerivedDataPhysXCooker::FDerivedDataPhysXCooker(FName InFormat, int32 InRuntimeCookFlags, UBodySetup* InBodySetup)
 	: BodySetup( InBodySetup )
 	, CollisionDataProvider( NULL )
 	, Format( InFormat )
+	, RuntimeCookFlags(InRuntimeCookFlags)
 	, Cooker( NULL )
 {
 	check( BodySetup != NULL );
@@ -124,7 +125,7 @@ int32 FDerivedDataPhysXCooker::BuildConvex( TArray<uint8>& OutData, bool InMirro
 		// Cook and store Result at ResultInfoOffset
 		UE_LOG(LogPhysics, Log, TEXT("Cook Convex: %s %d (FlipX:%d)"), *BodySetup->GetOuter()->GetPathName(), ElementIndex, InMirrored);		
 		const bool bDeformableMesh = CollisionDataProvider->IsA(USplineMeshComponent::StaticClass());
-		const bool Result = Cooker->CookConvex(Format, *MeshVertices, OutData, bDeformableMesh);
+		const bool Result = Cooker->CookConvex(Format, RuntimeCookFlags, *MeshVertices, OutData, bDeformableMesh);
 		if( !Result )
 		{
 			UE_LOG(LogPhysics, Warning, TEXT("Failed to cook convex: %s %d (FlipX:%d). The remaining elements will not get cooked."), *BodySetup->GetOuter()->GetPathName(), ElementIndex, InMirrored);
@@ -176,7 +177,7 @@ int32 FDerivedDataPhysXCooker::BuildTriMesh( TArray<uint8>& OutData, bool InUseA
 			ensure(SkeletalMesh->bEnablePerPolyCollision);
 			bDeformableMesh = true;
 		}
-		bResult = Cooker->CookTriMesh( Format, *MeshVertices, TriangleMeshDesc.Indices, TriangleMeshDesc.MaterialIndices, TriangleMeshDesc.bFlipNormals, OutData, bDeformableMesh );
+		bResult = Cooker->CookTriMesh( Format, RuntimeCookFlags, *MeshVertices, TriangleMeshDesc.Indices, TriangleMeshDesc.MaterialIndices, TriangleMeshDesc.bFlipNormals, OutData, bDeformableMesh );
 		if( !bResult )
 		{
 			UE_LOG(LogPhysics, Warning, TEXT("Failed to cook TriMesh: %s."), *CollisionDataProvider->GetPathName());
diff --git a/Engine/Source/Runtime/Engine/Private/PhysicsEngine/PhysDerivedData.h b/Engine/Source/Runtime/Engine/Private/PhysicsEngine/PhysDerivedData.h
--- a/Engine/Source/Runtime/Engine/Private/PhysicsEngine/PhysDerivedData.h
+++ b/Engine/Source/Runtime/Engine/Private/PhysicsEngine/PhysDerivedData.h
@@ -21,12 +21,13 @@ private:
 	FName Format;
 	bool bGenerateNormalMesh;
 	bool bGenerateMirroredMesh;	
+	int32 RuntimeCookFlags;
 	const class IPhysXFormat* Cooker;
 	FGuid DataGuid;
 	FString MeshId;
 
 public:
-	FDerivedDataPhysXCooker( FName InFormat, UBodySetup* InBodySetup );
+	FDerivedDataPhysXCooker(FName InFormat, int32 InRuntimeCookFlags, UBodySetup* InBodySetup);
 
 	virtual const TCHAR* GetPluginName() const override
 	{
@@ -38,7 +39,7 @@ public:
 		// This is a version string that mimics the old versioning scheme. If you
 		// want to bump this version, generate a new guid using VS->Tools->Create GUID and
 		// return it here. Ex.
-		return TEXT("483E0288C7094C2EB86167082ABE7A89");
+		return TEXT("F4A1A1C74D09407EB40C82094A33556D");
 	}
 
 	virtual FString GetPluginSpecificCacheKeySuffix() const override
@@ -53,12 +54,13 @@ public:
 				((PX_PHYSICS_VERSION_BUGFIX & 0xF) << 4) |
 				((UE_PHYSX_DERIVEDDATA_VER	& 0xF));
 
-		return FString::Printf( TEXT("%s_%s_%s_%d_%d_%hu_%hu"),
+		return FString::Printf( TEXT("%s_%s_%s_%d_%d_%d_%hu_%hu"),
 			*Format.ToString(),
 			*DataGuid.ToString(),
 			*MeshId,
 			(int32)bGenerateNormalMesh,
 			(int32)bGenerateMirroredMesh,
+			(int32)RuntimeCookFlags,
 			PhysXVersion,
 			Cooker ? Cooker->GetVersion( Format ) : 0xffff
 			);
diff --git a/Engine/Source/Runtime/Engine/Private/PhysicsEngine/PhysicsSettings.cpp b/Engine/Source/Runtime/Engine/Private/PhysicsEngine/PhysicsSettings.cpp
--- a/Engine/Source/Runtime/Engine/Private/PhysicsEngine/PhysicsSettings.cpp
+++ b/Engine/Source/Runtime/Engine/Private/PhysicsEngine/PhysicsSettings.cpp
@@ -22,6 +22,7 @@ UPhysicsSettings::UPhysicsSettings(const FObjectInitializer& ObjectInitializer)
 	, MaxAngularVelocity(3600)	//10 revolutions per second
 	, bSimulateSkeletalMeshOnDedicatedServer(true)
 	, bDefaultHasComplexCollision(true)
+	, bSupressFaceRemapTable(false)
 	, MaxPhysicsDeltaTime(1.f / 30.f)
 	, bSubstepping(false)
 	, bSubsteppingAsync(false)
diff --git a/Engine/Source/Runtime/Engine/Private/SkeletalMeshComponentPhysics.cpp b/Engine/Source/Runtime/Engine/Private/SkeletalMeshComponentPhysics.cpp
--- a/Engine/Source/Runtime/Engine/Private/SkeletalMeshComponentPhysics.cpp
+++ b/Engine/Source/Runtime/Engine/Private/SkeletalMeshComponentPhysics.cpp
@@ -1458,6 +1458,14 @@ void USkeletalMeshComponent::UpdateOverlaps(TArray<FOverlapInfo> const* PendingO
 	UPrimitiveComponent::UpdateOverlaps(PendingOverlaps, bDoNotifies, OverlapsAtEndLocation);
 }
 
+bool USkeletalMeshComponent::ShouldCreatePhysicsState() const
+{
+	bool bShouldCreatePhysicsState = Super::ShouldCreatePhysicsState();
+	bShouldCreatePhysicsState &= (MasterPoseComponent.IsValid() == false);
+	
+	return bShouldCreatePhysicsState;
+}
+
 void USkeletalMeshComponent::CreatePhysicsState()
 {
 	//	UE_LOG(LogSkeletalMesh, Warning, TEXT("Creating Physics State (%s : %s)"), *GetNameSafe(GetOuter()),  *GetName());
diff --git a/Engine/Source/Runtime/Engine/Private/SkeletalRenderGPUSkin.cpp b/Engine/Source/Runtime/Engine/Private/SkeletalRenderGPUSkin.cpp
--- a/Engine/Source/Runtime/Engine/Private/SkeletalRenderGPUSkin.cpp
+++ b/Engine/Source/Runtime/Engine/Private/SkeletalRenderGPUSkin.cpp
@@ -322,7 +322,7 @@ void FSkeletalMeshObjectGPUSkin::ProcessUpdatedDynamicData(FRHICommandListImmedi
 
 			FGPUBaseSkinVertexFactory::ShaderDataType& ShaderData = 
 				bClothFactory ? VertexFactoryData.ClothVertexFactories[ChunkIdx]->GetVertexFactory()->GetShaderData() :
-				(DynamicData->NumWeightedActiveVertexAnims > 0 ? VertexFactoryData.MorphVertexFactories[ChunkIdx].GetShaderData() : VertexFactoryData.VertexFactories[ChunkIdx].GetShaderData());
+				(DynamicData->NumWeightedActiveVertexAnims > 0 ? VertexFactoryData.MorphVertexFactories[ChunkIdx]->GetShaderData() : VertexFactoryData.VertexFactories[ChunkIdx]->GetShaderData());
 
 			bool bUseSkinCache = bGPUSkinCacheEnabled;
 			if (bUseSkinCache)
@@ -356,7 +356,7 @@ void FSkeletalMeshObjectGPUSkin::ProcessUpdatedDynamicData(FRHICommandListImmedi
 			// Try to use the GPU skinning cache if possible
 			if (bUseSkinCache)
 			{
-				int32 Key = GGPUSkinCache.StartCacheMesh(RHICmdList, GPUSkinCacheKeys[ChunkIdx], &VertexFactoryData.VertexFactories[ChunkIdx], &VertexFactoryData.PassthroughVertexFactories[ChunkIdx], Chunk, this);
+				int32 Key = GGPUSkinCache.StartCacheMesh(RHICmdList, GPUSkinCacheKeys[ChunkIdx], VertexFactoryData.VertexFactories[ChunkIdx].Get(), VertexFactoryData.PassthroughVertexFactories[ChunkIdx].Get(), Chunk, this);
 				if(Key >= 0)
 				{
 					GPUSkinCacheKeys[ChunkIdx] = (int16)Key;
@@ -513,14 +513,14 @@ const FVertexFactory* FSkeletalMeshObjectGPUSkin::GetVertexFactory(int32 LODInde
 	// cloth simulation is updated & if this ChunkIdx is for ClothVertexFactory
 	if ( DynamicData->ClothSimulUpdateData.Num() > 0 
 		&& LOD.GPUSkinVertexFactories.ClothVertexFactories.IsValidIndex(ChunkIdx)  
-		&& LOD.GPUSkinVertexFactories.ClothVertexFactories[ChunkIdx] != NULL )
+		&& LOD.GPUSkinVertexFactories.ClothVertexFactories[ChunkIdx].IsValid() )
 	{
 		return LOD.GPUSkinVertexFactories.ClothVertexFactories[ChunkIdx]->GetVertexFactory();
 	}
 	// use the morph enabled vertex factory if any active morphs are set
 	else if( DynamicData->NumWeightedActiveVertexAnims > 0 )
 	{
-		return &LOD.GPUSkinVertexFactories.MorphVertexFactories[ChunkIdx];
+		return LOD.GPUSkinVertexFactories.MorphVertexFactories[ChunkIdx].Get();
 	}
 	// use the default gpu skin vertex factory
 	else
@@ -528,11 +528,11 @@ const FVertexFactory* FSkeletalMeshObjectGPUSkin::GetVertexFactory(int32 LODInde
 		// If the GPU skinning cache was used, return the passthrough vertex factory
 		if (GGPUSkinCache.IsElementProcessed(GPUSkinCacheKeys[ChunkIdx]))
 		{
-			return &LOD.GPUSkinVertexFactories.PassthroughVertexFactories[ChunkIdx];
+			return LOD.GPUSkinVertexFactories.PassthroughVertexFactories[ChunkIdx].Get();
 		}
 		else
 		{
-			return &LOD.GPUSkinVertexFactories.VertexFactories[ChunkIdx];
+			return LOD.GPUSkinVertexFactories.VertexFactories[ChunkIdx].Get();
 		}
 	}
 }
@@ -686,13 +686,13 @@ TDynamicUpdateVertexFactoryData<VertexFactoryType>,VertexUpdateData,VertexUpdate
  * Creates a vertex factory entry for the given type and initialize it on the render thread
  */
 template <class VertexFactoryTypeBase, class VertexFactoryType>
-static VertexFactoryType* CreateVertexFactory(TIndirectArray<VertexFactoryTypeBase>& VertexFactories,
+static VertexFactoryType* CreateVertexFactory(TArray<TUniquePtr<VertexFactoryTypeBase>>& VertexFactories,
 						 const FSkeletalMeshObjectGPUSkin::FVertexFactoryBuffers& InVertexBuffers,
 						 ERHIFeatureLevel::Type FeatureLevel
 						 )
 {
 	auto* VertexFactory = new VertexFactoryType(FeatureLevel);
-	VertexFactories.Add(VertexFactory);
+	VertexFactories.Add(TUniquePtr<VertexFactoryTypeBase>(VertexFactory));
 
 	// Setup the update data for enqueue
 	TDynamicUpdateVertexFactoryData<VertexFactoryType> VertexUpdateData(VertexFactory,InVertexBuffers);
@@ -717,11 +717,11 @@ VertexFactoryType*, SourceVertexFactory, SourceVertexFactory,
 });
 
 template<typename VertexFactoryType>
-static void CreatePassthroughVertexFactory(TIndirectArray<FGPUSkinPassthroughVertexFactory>& PassthroughVertexFactories,
+static void CreatePassthroughVertexFactory(TArray<TUniquePtr<FGPUSkinPassthroughVertexFactory>>& PassthroughVertexFactories,
 	VertexFactoryType* SourceVertexFactory)
 {
 	auto* VertexFactory = new FGPUSkinPassthroughVertexFactory();
-	PassthroughVertexFactories.Add(VertexFactory);
+	PassthroughVertexFactories.Add(TUniquePtr<FGPUSkinPassthroughVertexFactory>(VertexFactory));
 
 	// update vertex factory components and sync it
 	ENQUEUE_UNIQUE_RENDER_COMMAND_TWOPARAMETER_CREATE_TEMPLATE(
@@ -749,14 +749,14 @@ TDynamicUpdateVertexFactoryData<VertexFactoryType>,VertexUpdateData,VertexUpdate
  * Creates a vertex factory entry for the given type and initialize it on the render thread
  */
 template <class VertexFactoryTypeBase, class VertexFactoryType>
-static void CreateVertexFactoryMorph(TIndirectArray<VertexFactoryTypeBase>& VertexFactories,
+static void CreateVertexFactoryMorph(TArray<TUniquePtr<VertexFactoryTypeBase>>& VertexFactories,
 						 const FSkeletalMeshObjectGPUSkin::FVertexFactoryBuffers& InVertexBuffers,
 						 ERHIFeatureLevel::Type FeatureLevel
 						 )
 
 {
 	auto* VertexFactory = new VertexFactoryType(FeatureLevel);
-	VertexFactories.Add(VertexFactory);
+	VertexFactories.Add(TUniquePtr<VertexFactoryTypeBase>(VertexFactory));
 						
 	// Setup the update data for enqueue
 	TDynamicUpdateVertexFactoryData<VertexFactoryType> VertexUpdateData(VertexFactory, InVertexBuffers);
@@ -787,14 +787,14 @@ TDynamicUpdateVertexFactoryData<VertexFactoryType>,VertexUpdateData,VertexUpdate
  * Creates a vertex factory entry for the given type and initialize it on the render thread
  */
 template <class VertexFactoryTypeBase, class VertexFactoryType>
-static void CreateVertexFactoryCloth(TArray<VertexFactoryTypeBase*>& VertexFactories,
+static void CreateVertexFactoryCloth(TArray<TUniquePtr<VertexFactoryTypeBase>>& VertexFactories,
 						 const FSkeletalMeshObjectGPUSkin::FVertexFactoryBuffers& InVertexBuffers,
 						 ERHIFeatureLevel::Type FeatureLevel
 						 )
 
 {
 	VertexFactoryType* VertexFactory = new VertexFactoryType(FeatureLevel);
-	VertexFactories.Add(VertexFactory);
+	VertexFactories.Add(TUniquePtr<VertexFactoryTypeBase>(VertexFactory));
 						
 	// Setup the update data for enqueue
 	TDynamicUpdateVertexFactoryData<VertexFactoryType> VertexUpdateData(VertexFactory, InVertexBuffers);
@@ -864,12 +864,12 @@ void FSkeletalMeshObjectGPUSkin::FVertexFactoryData::ReleaseVertexFactories()
 	// Default factories
 	for( int32 FactoryIdx=0; FactoryIdx < VertexFactories.Num(); FactoryIdx++)
 	{
-		BeginReleaseResource(&VertexFactories[FactoryIdx]);
+		BeginReleaseResource(VertexFactories[FactoryIdx].Get());
 	}
 
 	for (int32 FactoryIdx = 0; FactoryIdx < PassthroughVertexFactories.Num(); FactoryIdx++)
 	{
-		BeginReleaseResource(&PassthroughVertexFactories[FactoryIdx]);
+		BeginReleaseResource(PassthroughVertexFactories[FactoryIdx].Get());
 	}
 }
 
@@ -902,8 +902,7 @@ void FSkeletalMeshObjectGPUSkin::FVertexFactoryData::ReleaseMorphVertexFactories
 	// Default morph factories
 	for( int32 FactoryIdx=0; FactoryIdx < MorphVertexFactories.Num(); FactoryIdx++ )
 	{
-		auto& MorphVertexFactory = MorphVertexFactories[FactoryIdx];
-		BeginReleaseResource(&MorphVertexFactory);
+		BeginReleaseResource(MorphVertexFactories[FactoryIdx].Get());
 	}
 }
 
@@ -930,7 +929,7 @@ void FSkeletalMeshObjectGPUSkin::FVertexFactoryData::InitAPEXClothVertexFactorie
 		}
 		else
 		{
-			ClothVertexFactories.Add(NULL);
+			ClothVertexFactories.Add(TUniquePtr<FGPUBaseSkinAPEXClothVertexFactory>(nullptr));
 		}
 	}
 }
@@ -943,7 +942,7 @@ void FSkeletalMeshObjectGPUSkin::FVertexFactoryData::ReleaseAPEXClothVertexFacto
 	// Default APEX cloth factories
 	for( int32 FactoryIdx=0; FactoryIdx < ClothVertexFactories.Num(); FactoryIdx++ )
 	{
-		auto* ClothVertexFactory = ClothVertexFactories[FactoryIdx];
+		TUniquePtr<FGPUBaseSkinAPEXClothVertexFactory>& ClothVertexFactory = ClothVertexFactories[FactoryIdx];
 		if (ClothVertexFactory)
 		{
 			BeginReleaseResource(ClothVertexFactory->GetVertexFactory());
diff --git a/Engine/Source/Runtime/Engine/Private/SkeletalRenderGPUSkin.h b/Engine/Source/Runtime/Engine/Private/SkeletalRenderGPUSkin.h
--- a/Engine/Source/Runtime/Engine/Private/SkeletalRenderGPUSkin.h
+++ b/Engine/Source/Runtime/Engine/Private/SkeletalRenderGPUSkin.h
@@ -271,16 +271,16 @@ private:
 		FVertexFactoryData() {}
 
 		/** one vertex factory for each chunk */
-		TIndirectArray<FGPUBaseSkinVertexFactory> VertexFactories;
+		TArray<TUniquePtr<FGPUBaseSkinVertexFactory>> VertexFactories;
 
 		/** one passthrough vertex factory for each chunk */
-		TIndirectArray<FGPUSkinPassthroughVertexFactory> PassthroughVertexFactories;
+		TArray<TUniquePtr<FGPUSkinPassthroughVertexFactory>> PassthroughVertexFactories;
 
 		/** Vertex factory defining both the base mesh as well as the morph delta vertex decals */
-		TIndirectArray<FGPUBaseSkinVertexFactory> MorphVertexFactories;
+		TArray<TUniquePtr<FGPUBaseSkinVertexFactory>> MorphVertexFactories;
 
 		/** Vertex factory defining both the base mesh as well as the APEX cloth vertex data */
-		TArray<FGPUBaseSkinAPEXClothVertexFactory*> ClothVertexFactories;
+		TArray<TUniquePtr<FGPUBaseSkinAPEXClothVertexFactory>> ClothVertexFactories;
 
 		/** 
 		 * Init default vertex factory resources for this LOD 
diff --git a/Engine/Source/Runtime/Engine/Private/Slate/SceneViewport.cpp b/Engine/Source/Runtime/Engine/Private/Slate/SceneViewport.cpp
--- a/Engine/Source/Runtime/Engine/Private/Slate/SceneViewport.cpp
+++ b/Engine/Source/Runtime/Engine/Private/Slate/SceneViewport.cpp
@@ -832,9 +832,17 @@ FReply FSceneViewport::OnFocusReceived(const FFocusEvent& InFocusEvent)
 			if (IsForegroundWindow())
 			{
 				bool bIsCursorForcedVisible = false;
-				if (ViewportClient->GetWorld() && ViewportClient->GetWorld()->GetGameInstance() && ViewportClient->GetWorld()->GetGameInstance()->GetFirstLocalPlayerController())
+				APlayerController* FirstLocalPC = nullptr;
+				if (UWorld* ViewportClientWorld = ViewportClient->GetWorld())
 				{
-					bIsCursorForcedVisible = ViewportClient->GetWorld()->GetGameInstance()->GetFirstLocalPlayerController()->GetMouseCursor() != EMouseCursor::None;
+					if (UGameInstance* GameInstance = ViewportClientWorld->GetGameInstance())
+					{
+						FirstLocalPC = GameInstance->GetFirstLocalPlayerController();
+					}
+				}
+				if (FirstLocalPC)
+				{
+					bIsCursorForcedVisible = FirstLocalPC->GetMouseCursor() != EMouseCursor::None;
 				}
 
 				const bool bPlayInEditorCapture = !bIsPlayInEditorViewport || InFocusEvent.GetCause() != EFocusCause::SetDirectly || bPlayInEditorGetsMouseControl;
@@ -857,12 +865,13 @@ FReply FSceneViewport::OnFocusReceived(const FFocusEvent& InFocusEvent)
 	}
 
 	// Update key state mappings so that the the viewport modifier states are valid upon focus.
-	KeyStateMap.Add( EKeys::LeftAlt, FSlateApplication::Get().GetModifierKeys().IsLeftAltDown() );
-	KeyStateMap.Add( EKeys::RightAlt, FSlateApplication::Get().GetModifierKeys().IsRightAltDown() );
-	KeyStateMap.Add( EKeys::LeftControl, FSlateApplication::Get().GetModifierKeys().IsLeftControlDown());
-	KeyStateMap.Add( EKeys::RightControl, FSlateApplication::Get().GetModifierKeys().IsRightControlDown());
-	KeyStateMap.Add( EKeys::LeftShift, FSlateApplication::Get().GetModifierKeys().IsLeftShiftDown());
-	KeyStateMap.Add( EKeys::RightShift, FSlateApplication::Get().GetModifierKeys().IsRightShiftDown());
+	const FModifierKeysState KeysState = FSlateApplication::Get().GetModifierKeys();
+	KeyStateMap.Add( EKeys::LeftAlt, KeysState.IsLeftAltDown() );
+	KeyStateMap.Add( EKeys::RightAlt, KeysState.IsRightAltDown() );
+	KeyStateMap.Add( EKeys::LeftControl, KeysState.IsLeftControlDown());
+	KeyStateMap.Add( EKeys::RightControl, KeysState.IsRightControlDown());
+	KeyStateMap.Add( EKeys::LeftShift, KeysState.IsLeftShiftDown());
+	KeyStateMap.Add( EKeys::RightShift, KeysState.IsRightShiftDown());
 
 	return CurrentReplyState;
 }
diff --git a/Engine/Source/Runtime/Engine/Private/World.cpp b/Engine/Source/Runtime/Engine/Private/World.cpp
--- a/Engine/Source/Runtime/Engine/Private/World.cpp
+++ b/Engine/Source/Runtime/Engine/Private/World.cpp
@@ -5356,27 +5356,37 @@ FString UWorld::GetAddressURL() const
 
 FString UWorld::RemovePIEPrefix(const FString &Source)
 {
-	// PIE prefix is: UEDPIE_X_MapName (where X is 0-9)
+	// PIE prefix is: UEDPIE_X_MapName (where X is some decimal number)
 	const FString LookFor = PLAYWORLD_PACKAGE_PREFIX;
-	FString FixedName;
 
 	int32 idx = Source.Find(LookFor);
 	if (idx >= 0)
 	{
-		FixedName = Source.Left( idx );
-
-		FString Blah = Source.Right( idx );
-		FString Blah2 = Source.Right( Source.Len() - idx );
-		FString RightS = Source.Right( Source.Len() - (idx + LookFor.Len() + 3) );
-		
-		FixedName += RightS;
-	}
-	else
-	{
-		FixedName = Source;
+		int32 end = idx + LookFor.Len();
+		if ((end >= Source.Len()) || (Source[end] != '_'))
+		{
+			UE_LOG(LogWorld, Warning, TEXT("Looks like World path invalid PIE prefix (expected '_' characeter after PIE prefix): %s"), *Source);
+			return Source;
+		}
+		for (++end; (end < Source.Len()) && (Source[end] != '_'); ++end)
+		{
+			if ((Source[end] < '0') || (Source[end] > '9'))
+			{
+				UE_LOG(LogWorld, Warning, TEXT("Looks like World have invalid PIE prefix (PIE instance not number): %s"), *Source);
+				return Source;
+			}
+		}
+		if (end >= Source.Len())
+		{
+			UE_LOG(LogWorld, Warning, TEXT("Looks like World path invalid PIE prefix (can't find end of PIE prefix): %s"), *Source);
+			return Source;
+		}
+		const FString Prefix = Source.Left(idx);
+		const FString Suffix = Source.Right(Source.Len() - end - 1);
+		return Prefix + Suffix;
 	}
 
-	return FixedName;
+	return Source;
 }
 
 UWorld* UWorld::FindWorldInPackage(UPackage* Package)
diff --git a/Engine/Source/Runtime/Engine/Public/AudioDevice.h b/Engine/Source/Runtime/Engine/Public/AudioDevice.h
--- a/Engine/Source/Runtime/Engine/Public/AudioDevice.h
+++ b/Engine/Source/Runtime/Engine/Public/AudioDevice.h
@@ -401,7 +401,17 @@ public:
 	void SetReverbSettings( class AAudioVolume* Volume, const FReverbSettings& ReverbSettings );
 
 	/**
-	 * Creates an audio component to handle playing a sound
+	 * Creates an audio component to handle playing a sound.
+	 * Plays a sound at the given location without creating an audio component.
+	 * @param   Sound				The USoundBase to play at the location.
+	 * @param   World				The world this sound is playing in.
+	 * @param   AActor				The optional actor with which to play the sound on.
+	 * @param   Play				Whether or not to automatically call play on the audio component after it is created.
+	 * @param	bStopWhenOwnerDestroyed Whether or not to automatically stop the audio component if its owner is destroyed.
+	 * @param	Location			The sound's location.
+	 * @param	AttenuationSettings	The sound's attenuation settings to use. Will default to the USoundBase's AttenuationSettings if not specified.
+	 * @param	USoundConcurrency	The sound's sound concurrency settings to use. Will use the USoundBase's USoundConcurrency if not specified.
+	 * @return	The created audio component if the function successfully created one or a nullptr if not successful. Note: if audio is disabled or if there were no hardware audio devices available, this will return nullptr.
 	 */
 	static class UAudioComponent* CreateComponent(class USoundBase* Sound, class UWorld* World, AActor*  AActor = nullptr, bool Play = true, bool bStopWhenOwnerDestroyed = false, const FVector* Location = nullptr, USoundAttenuation* AttenuationSettings = nullptr, USoundConcurrency* ConcurrencySettings = nullptr);
 
diff --git a/Engine/Source/Runtime/Engine/Public/GPUSkinVertexFactory.h b/Engine/Source/Runtime/Engine/Public/GPUSkinVertexFactory.h
--- a/Engine/Source/Runtime/Engine/Public/GPUSkinVertexFactory.h
+++ b/Engine/Source/Runtime/Engine/Public/GPUSkinVertexFactory.h
@@ -771,6 +771,7 @@ public:
 		}
 	};
 
+	virtual ~FGPUBaseSkinAPEXClothVertexFactory() {}
 
 	/** accessor */
 	FORCEINLINE ClothShaderType& GetClothShaderData()
diff --git a/Engine/Source/Runtime/Landscape/Private/LandscapeCollision.cpp b/Engine/Source/Runtime/Landscape/Private/LandscapeCollision.cpp
--- a/Engine/Source/Runtime/Landscape/Private/LandscapeCollision.cpp
+++ b/Engine/Source/Runtime/Landscape/Private/LandscapeCollision.cpp
@@ -780,7 +780,7 @@ bool ULandscapeMeshCollisionComponent::CookCollisionData(const FName& Format, bo
 	TArray<uint8> OutData;
 	ITargetPlatformManagerModule* TPM = GetTargetPlatformManager();
 	const IPhysXFormat* Cooker = TPM->FindPhysXFormat(Format);
-	bool Result = Cooker->CookTriMesh(Format, Vertices, Indices, MaterialIndices, bFlipNormals, OutData);
+	bool Result = Cooker->CookTriMesh(Format, 0, Vertices, Indices, MaterialIndices, bFlipNormals, OutData);
 
 	if (Result)
 	{
diff --git a/Engine/Source/Runtime/PhysXFormats/Private/PhysXFormats.cpp b/Engine/Source/Runtime/PhysXFormats/Private/PhysXFormats.cpp
--- a/Engine/Source/Runtime/PhysXFormats/Private/PhysXFormats.cpp
+++ b/Engine/Source/Runtime/PhysXFormats/Private/PhysXFormats.cpp
@@ -84,7 +84,7 @@ public:
 		OutFormats.Add(NAME_PhysXPS4);
 	}
 
-	virtual bool CookConvex(FName Format, const TArray<FVector>& SrcBuffer, TArray<uint8>& OutBuffer, bool bDeformableMesh = false) const override
+	virtual bool CookConvex(FName Format, int32 RuntimeCookFlags, const TArray<FVector>& SrcBuffer, TArray<uint8>& OutBuffer, bool bDeformableMesh = false) const override
 	{
 #if WITH_PHYSX
 		PxPlatform::Enum PhysXFormat = PxPlatform::ePC;
@@ -108,6 +108,12 @@ public:
 		const PxCookingParams& Params = PhysXCooking->getParams();
 		PxCookingParams NewParams = Params;
 		NewParams.targetPlatform = PhysXFormat;
+
+		if(RuntimeCookFlags & ERuntimePhysxCookOptimizationFlags::SupressFaceRemapTable)
+		{
+			NewParams.suppressTriangleMeshRemapTable = true;
+		}
+
 		if (bDeformableMesh)
 		{
 			// Meshes which can be deformed need different cooking parameters to inhibit vertex welding and add an extra skin around the collision mesh for safety.
@@ -152,7 +158,7 @@ public:
 		return false;
 	}
 
-	virtual bool CookTriMesh(FName Format, const TArray<FVector>& SrcVertices, const TArray<FTriIndices>& SrcIndices, const TArray<uint16>& SrcMaterialIndices, const bool FlipNormals, TArray<uint8>& OutBuffer, bool bDeformableMesh = false) const override
+	virtual bool CookTriMesh(FName Format, int32 RuntimeCookFlags, const TArray<FVector>& SrcVertices, const TArray<FTriIndices>& SrcIndices, const TArray<uint16>& SrcMaterialIndices, const bool FlipNormals, TArray<uint8>& OutBuffer, bool bDeformableMesh = false) const override
 	{
 #if WITH_PHYSX
 		PxPlatform::Enum PhysXFormat = PxPlatform::ePC;
@@ -176,6 +182,11 @@ public:
 		NewParams.targetPlatform = PhysXFormat;
 		PxMeshPreprocessingFlags OldCookingFlags = NewParams.meshPreprocessParams;
 
+		if (RuntimeCookFlags & ERuntimePhysxCookOptimizationFlags::SupressFaceRemapTable)
+		{
+			NewParams.suppressTriangleMeshRemapTable = true;
+		}
+
 		if (bDeformableMesh)
 		{
 			// In the case of a deformable mesh, we have to change the cook params
diff --git a/Engine/Source/Runtime/PhysXFormats/Public/IPhysXFormat.h b/Engine/Source/Runtime/PhysXFormats/Public/IPhysXFormat.h
--- a/Engine/Source/Runtime/PhysXFormats/Public/IPhysXFormat.h
+++ b/Engine/Source/Runtime/PhysXFormats/Public/IPhysXFormat.h
@@ -6,6 +6,10 @@
 
 #pragma once
 
+enum ERuntimePhysxCookOptimizationFlags
+{
+	SupressFaceRemapTable = 0x1
+};
 
 /**
  * IPhysXFormat, PhysX cooking and serialization abstraction
@@ -30,23 +34,25 @@ public:
 	 * Cooks the source convex data for the platform and stores the cooked data internally.
 	 *
 	 * @param Format The desired format
+	 * @param RuntimeCookFlags Any special runtime flags we want to pass in
 	 * @param SrcBuffer The source buffer
 	 * @param OutBuffer The resulting cooked data
 	 * @param bDeformableMesh Whether this mesh is deformable and hence needs to be cooked with different parameters
 	 * @return true on success, false otherwise.
 	 */
-	virtual bool CookConvex( FName Format, const TArray<FVector>& SrcBuffer, TArray<uint8>& OutBuffer, bool bDeformableMesh = false ) const = 0;
+	virtual bool CookConvex( FName Format, int32 RuntimeCookFlags, const TArray<FVector>& SrcBuffer, TArray<uint8>& OutBuffer, bool bDeformableMesh = false ) const = 0;
 
 	/**
 	 * Cooks the source Tri-Mesh data for the platform and stores the cooked data internally.
 	 *
 	 * @param Format The desired format.
+	 * @param RuntimeCookFlags Any special runtime flags we want to pass in
 	 * @param SrcBuffer The source buffer.
 	 * @param OutBuffer The resulting cooked data.
 	 * @param bDeformableMesh This is a very special case that requires different cooking parameters set.
 	 * @return true on success, false otherwise.
 	 */
-	virtual bool CookTriMesh( FName Format, const TArray<FVector>& SrcVertices, const TArray<struct FTriIndices>& SrcIndices, const TArray<uint16>& SrcMaterialIndices, const bool FlipNormals, TArray<uint8>& OutBuffer, bool bDeformableMesh = false ) const = 0;
+	virtual bool CookTriMesh( FName Format, int32 RuntimeCookFlags, const TArray<FVector>& SrcVertices, const TArray<struct FTriIndices>& SrcIndices, const TArray<uint16>& SrcMaterialIndices, const bool FlipNormals, TArray<uint8>& OutBuffer, bool bDeformableMesh = false ) const = 0;
 		
 	/**
 	 * Cooks the source height field data for the platform and stores the cooked data internally.
diff --git a/Engine/Source/Runtime/Windows/XAudio2/Private/XAudio2Device.cpp b/Engine/Source/Runtime/Windows/XAudio2/Private/XAudio2Device.cpp
--- a/Engine/Source/Runtime/Windows/XAudio2/Private/XAudio2Device.cpp
+++ b/Engine/Source/Runtime/Windows/XAudio2/Private/XAudio2Device.cpp
@@ -122,6 +122,7 @@ bool FXAudio2Device::InitializeHardware()
 	if( DeviceCount < 1 )
 	{
 		UE_LOG(LogInit, Log, TEXT( "No audio devices found!" ) );
+		DeviceProperties->XAudio2->Release();
 		DeviceProperties->XAudio2 = nullptr;
 		return( false );		
 	}
diff --git a/Engine/Source/Runtime/Windows/XAudio2/Private/XAudio2Source.cpp b/Engine/Source/Runtime/Windows/XAudio2/Private/XAudio2Source.cpp
--- a/Engine/Source/Runtime/Windows/XAudio2/Private/XAudio2Source.cpp
+++ b/Engine/Source/Runtime/Windows/XAudio2/Private/XAudio2Source.cpp
@@ -220,8 +220,9 @@ void FXAudio2SoundSource::SubmitPCMRTBuffers( void )
 	XAudio2Buffers[1].AudioBytes = BufferSize;
 
 	// Only use the cached data if we're starting from the beginning, otherwise we'll have to take a synchronous hit
-	bool bSkipFirstBuffer = false;;
-	if (WaveInstance->WaveData && WaveInstance->WaveData->CachedRealtimeFirstBuffer && WaveInstance->StartTime == 0.f)
+	bool bSkipFirstBuffer = false;
+	bool bIsSeeking = (WaveInstance->StartTime > 0.f);
+	if (WaveInstance->WaveData && WaveInstance->WaveData->CachedRealtimeFirstBuffer && !bIsSeeking)
 	{
 		FMemory::Memcpy((uint8*)XAudio2Buffers[0].pAudioData, WaveInstance->WaveData->CachedRealtimeFirstBuffer, BufferSize);
 		FMemory::Memcpy((uint8*)XAudio2Buffers[1].pAudioData, WaveInstance->WaveData->CachedRealtimeFirstBuffer + BufferSize, BufferSize);
@@ -257,7 +258,7 @@ void FXAudio2SoundSource::SubmitPCMRTBuffers( void )
 
 	ReadMorePCMData(2, DataReadMode);
 
-	if (DataReadMode == EDataReadMode::Synchronous || (!bSkipFirstBuffer && WaveInstance->WaveData && !WaveInstance->WaveData->bCanProcessAsync))
+	if (!bIsSeeking && (DataReadMode == EDataReadMode::Synchronous || (!bSkipFirstBuffer && WaveInstance->WaveData && !WaveInstance->WaveData->bCanProcessAsync)))
 	{
 		AudioDevice->ValidateAPICall(TEXT("SubmitSourceBuffer - PCMRT"),
 									 Source->SubmitSourceBuffer(&XAudio2Buffers[2]));
@@ -472,7 +473,7 @@ bool FXAudio2SoundSource::Init(FWaveInstance* InWaveInstance)
 
 		// Reset the LPFFrequency values
 		LPFFrequency = MAX_FILTER_FREQUENCY;
-		LastLPFFrequency = MAX_FILTER_FREQUENCY;
+		LastLPFFrequency = FLT_MAX;
 
 		// Buffer failed to be created, or there was an error with the compressed data
 		if (Buffer && Buffer->NumChannels > 0)
diff --git a/Engine/Source/Runtime/Windows/XAudio2/Private/XAudio2Support.h b/Engine/Source/Runtime/Windows/XAudio2/Private/XAudio2Support.h
--- a/Engine/Source/Runtime/Windows/XAudio2/Private/XAudio2Support.h
+++ b/Engine/Source/Runtime/Windows/XAudio2/Private/XAudio2Support.h
@@ -703,10 +703,7 @@ struct FXAudioDeviceProperties
 
 		// Release the effect chain
 		Voice->SetEffectChain(nullptr);
-#if PLATFORM_WINDOWS
-		Voice->DestroyVoice();
-		return;
-#else
+
 		// See if there is an existing pool for this source voice
 		FSourceVoicePoolEntry* VoicePoolEntry = nullptr;
 		for (int32 i = 0; i < VoicePool.Num(); ++i)
@@ -732,7 +729,6 @@ struct FXAudioDeviceProperties
 			VoicePoolEntry->MaxEffectChainChannels = MaxEffectChainChannels;
 			VoicePool.Add(VoicePoolEntry);
 		}
-#endif
 	}
 };
 
