Fix compile errors for undefined symbol in BodySetup.cpp
Fix spelling of suppress

diff --git a/Engine/Source/Runtime/Engine/Classes/PhysicsEngine/PhysicsSettings.h b/Engine/Source/Runtime/Engine/Classes/PhysicsEngine/PhysicsSettings.h
--- a/Engine/Source/Runtime/Engine/Classes/PhysicsEngine/PhysicsSettings.h
+++ b/Engine/Source/Runtime/Engine/Classes/PhysicsEngine/PhysicsSettings.h
@@ -160,7 +160,7 @@ class ENGINE_API UPhysicsSettings : public UDeveloperSettings
 	/**
 	*  If true, the internal physx face to UE face mapping will not be generated. This is a memory optimization available if you do not rely on face indices returned by scene queries. */
 	UPROPERTY(config, EditAnywhere, Category = Optimization)
-	bool bSupressFaceRemapTable;
+	bool bSuppressFaceRemapTable;
 
 	/**
 	*  If true CCD will be ignored. This is an optimization when CCD is never used which removes the need for physx to check it internally. */
diff --git a/Engine/Source/Runtime/Engine/Private/Collision/CollisionConversions.cpp b/Engine/Source/Runtime/Engine/Private/Collision/CollisionConversions.cpp
--- a/Engine/Source/Runtime/Engine/Private/Collision/CollisionConversions.cpp
+++ b/Engine/Source/Runtime/Engine/Private/Collision/CollisionConversions.cpp
@@ -433,9 +433,9 @@ EConvertQueryResult ConvertQueryImpactHit(const UWorld* World, const PxLocationH
 #if WITH_EDITOR
 	if(bReturnFaceIndex && World->IsGameWorld())
 	{
-		if(!ensure(UPhysicsSettings::Get()->bSupressFaceRemapTable == false))
+		if(!ensure(UPhysicsSettings::Get()->bSuppressFaceRemapTable == false))
 		{
-			UE_LOG(LogPhysics, Error, TEXT("A scene query is relying on face indices, but bSupressFaceRemapTable is false."));
+			UE_LOG(LogPhysics, Error, TEXT("A scene query is relying on face indices, but bSuppressFaceRemapTable is false."));
 		}
 	}
 #endif
diff --git a/Engine/Source/Runtime/Engine/Private/PhysicsEngine/BodySetup.cpp b/Engine/Source/Runtime/Engine/Private/PhysicsEngine/BodySetup.cpp
--- a/Engine/Source/Runtime/Engine/Private/PhysicsEngine/BodySetup.cpp
+++ b/Engine/Source/Runtime/Engine/Private/PhysicsEngine/BodySetup.cpp
@@ -980,10 +980,12 @@ void UBodySetup::ClearCachedCookedPlatformData( const ITargetPlatform* TargetPla
 int32 UBodySetup::GetRuntimeOnlyCookOptimizationFlags() const
 {
 	int32 RuntimeCookFlags = 0;
-	if(UPhysicsSettings::Get()->bSupressFaceRemapTable)
+#if WITH_PHYSX && (WITH_RUNTIME_PHYSICS_COOKING || WITH_EDITOR)
+	if(UPhysicsSettings::Get()->bSuppressFaceRemapTable)
 	{
-		RuntimeCookFlags |= ERuntimePhysxCookOptimizationFlags::SupressFaceRemapTable;
+		RuntimeCookFlags |= ERuntimePhysxCookOptimizationFlags::SuppressFaceRemapTable;
 	}
+#endif
 
 	return RuntimeCookFlags;
 }
diff --git a/Engine/Source/Runtime/Engine/Private/PhysicsEngine/PhysicsSettings.cpp b/Engine/Source/Runtime/Engine/Private/PhysicsEngine/PhysicsSettings.cpp
--- a/Engine/Source/Runtime/Engine/Private/PhysicsEngine/PhysicsSettings.cpp
+++ b/Engine/Source/Runtime/Engine/Private/PhysicsEngine/PhysicsSettings.cpp
@@ -22,7 +22,7 @@ UPhysicsSettings::UPhysicsSettings(const FObjectInitializer& ObjectInitializer)
 	, MaxAngularVelocity(3600)	//10 revolutions per second
 	, bSimulateSkeletalMeshOnDedicatedServer(true)
 	, bDefaultHasComplexCollision(true)
-	, bSupressFaceRemapTable(false)
+	, bSuppressFaceRemapTable(false)
 	, MaxPhysicsDeltaTime(1.f / 30.f)
 	, bSubstepping(false)
 	, bSubsteppingAsync(false)
diff --git a/Engine/Source/Runtime/PhysXFormats/Public/IPhysXFormat.h b/Engine/Source/Runtime/PhysXFormats/Public/IPhysXFormat.h
--- a/Engine/Source/Runtime/PhysXFormats/Public/IPhysXFormat.h
+++ b/Engine/Source/Runtime/PhysXFormats/Public/IPhysXFormat.h
@@ -1,14 +1,10 @@
 // Copyright 1998-2016 Epic Games, Inc. All Rights Reserved.
 
-/*=============================================================================
-	ITextureFormatModule.h: Declares the ITextureFormatModule interface.
-=============================================================================*/
-
 #pragma once
 
 enum ERuntimePhysxCookOptimizationFlags
 {
-	SupressFaceRemapTable = 0x1
+	SuppressFaceRemapTable = 0x1
 };
 
 /**
