Merge CL# 2832954 by [REDACTED_PERSON]

Do not ensure if trying to detach children from an object that has been marked pending kill, which may possibly happen while unloading a level

diff --git a/Engine/Source/Runtime/Engine/Private/Components/SceneComponent.cpp b/Engine/Source/Runtime/Engine/Private/Components/SceneComponent.cpp
--- a/Engine/Source/Runtime/Engine/Private/Components/SceneComponent.cpp
+++ b/Engine/Source/Runtime/Engine/Private/Components/SceneComponent.cpp
@@ -765,7 +765,7 @@ void USceneComponent::OnComponentDestroyed(bool bDestroyingHierarchy)
 					}
 					if (bNeedsDetach)
 					{
-						if (ensure(Child->AttachParent && Child->AttachParent == this))
+						if (Child->AttachParent && Child->AttachParent == this)
 						{
 							Child->DetachFromParent();
 						}
@@ -776,9 +776,11 @@ void USceneComponent::OnComponentDestroyed(bool bDestroyingHierarchy)
 							if (Child->AttachParent)
 							{
 								UE_LOG(LogSceneComponent, Error, TEXT("Component '%s' has '%s' in its AttachChildren array, however, '%s' believes it is attached to '%s'"), *GetFullName(), *Child->GetFullName(), *Child->GetFullName(), *Child->AttachParent->GetFullName());
+								ensure(Child->AttachParent == this);
 							}
-							else
+							else if (!ensure(IsPendingKill()))
 							{
+								// If we are pending kill, the AttachParent reference to us may have been nulled already, so only error if not pending kill
 								UE_LOG(LogSceneComponent, Error, TEXT("Component '%s' has '%s' in its AttachChildren array, however, '%s' believes it is not attached to anything"), *GetFullName(), *Child->GetFullName(), *Child->GetFullName());
 							}
 							AttachChildren.Pop(false);
@@ -812,7 +814,7 @@ void USceneComponent::OnComponentDestroyed(bool bDestroyingHierarchy)
 					}
 					if (bNeedsDetach)
 					{
-						if (ensure(Child->AttachParent && Child->AttachParent == this))
+						if (Child->AttachParent && Child->AttachParent == this)
 						{
 							Child->DetachFromParent();
 						}
@@ -823,9 +825,11 @@ void USceneComponent::OnComponentDestroyed(bool bDestroyingHierarchy)
 							if (Child->AttachParent)
 							{
 								UE_LOG(LogSceneComponent, Error, TEXT("Component '%s' has '%s' in its AttachChildren array, however, '%s' believes it is attached to '%s'"), *GetFullName(), *Child->GetFullName(), *Child->GetFullName(), *Child->AttachParent->GetFullName());
+								ensure(Child->AttachParent == this);
 							}
-							else
+							else if (!ensure(IsPendingKill()))
 							{
+								// If we are pending kill, the AttachParent reference to us may have been nulled already, so only error if not pending kill
 								UE_LOG(LogSceneComponent, Error, TEXT("Component '%s' has '%s' in its AttachChildren array, however, '%s' believes it is not attached to anything"), *GetFullName(), *Child->GetFullName(), *Child->GetFullName());
 							}
 							AttachChildren.Pop(false);
