Fix VS2015 compilation due to implicit copy constructor causing TArray's copy to try and use the private TUniquePtr copy constructor

diff --git a/Engine/Source/Runtime/Engine/Private/SkeletalRenderGPUSkin.h b/Engine/Source/Runtime/Engine/Private/SkeletalRenderGPUSkin.h
--- a/Engine/Source/Runtime/Engine/Private/SkeletalRenderGPUSkin.h
+++ b/Engine/Source/Runtime/Engine/Private/SkeletalRenderGPUSkin.h
@@ -340,6 +340,10 @@ private:
 
 			return Size;
 		}	
+
+		private:
+			FVertexFactoryData(const FVertexFactoryData&);
+			FVertexFactoryData& operator=(const FVertexFactoryData&);
 	};
 
 	/** vertex data for rendering a single LOD */
