Copying //UE4/Dev-Rendering to Dev-Main (//UE4/Dev-Main)

diff --git a/Engine/Source/Runtime/Core/Public/Async/TaskGraphInterfaces.h b/Engine/Source/Runtime/Core/Public/Async/TaskGraphInterfaces.h
--- a/Engine/Source/Runtime/Core/Public/Async/TaskGraphInterfaces.h
+++ b/Engine/Source/Runtime/Core/Public/Async/TaskGraphInterfaces.h
@@ -54,15 +54,22 @@ namespace ENamedThreads
 		QueueIndexMask =	0x100,
 		QueueIndexShift =	8,
 
-		/** High bits are used for a queue index and priority**/
+		/** High bits are used for a queue index task priority and thread priority**/
+
+		NormalTaskPriority =	0x000,
+		HighTaskPriority =		0x200,
 
-		NormalPriority =	0x000,
-		HighPriority =		0x200,
+		NumTaskPriorities =		2,
+		TaskPriorityMask =		0x200,
+		TaskPriorityShift =		9,
 
-		NumPriorities =		2,
-		PriorityMask =		0x200,
-		PriorityShift =		9,
+		NormalThreadPriority = 0x000,
+		HighThreadPriority = 0x400,
+		BackgroundThreadPriority = 0x800,
 
+		NumThreadPriorities = 3,
+		ThreadPriorityMask = 0xC00,
+		ThreadPriorityShift = 10,
 
 		/** Combinations **/
 #if STATS
@@ -70,10 +77,24 @@ namespace ENamedThreads
 #endif
 		GameThread_Local = GameThread | LocalQueue,
 		ActualRenderingThread_Local = ActualRenderingThread | LocalQueue,
+
+		AnyHiPriThreadNormalTask = AnyThread | HighThreadPriority | NormalTaskPriority,
+		AnyHiPriThreadHiPriTask = AnyThread | HighThreadPriority | HighTaskPriority,
+
+		AnyNormalThreadNormalTask = AnyThread | NormalThreadPriority | NormalTaskPriority,
+		AnyNormalThreadHiPriTask = AnyThread | NormalThreadPriority | HighTaskPriority,
+
+		AnyBackgroundThreadNormalTask = AnyThread | BackgroundThreadPriority | NormalTaskPriority,
+		AnyBackgroundHiPriTask = AnyThread | BackgroundThreadPriority | HighTaskPriority,
 	};
 	extern CORE_API Type RenderThread; // this is not an enum, because if there is no render thread, this is just the game thread.
 	extern CORE_API Type RenderThread_Local; // this is not an enum, because if there is no render thread, this is just the game thread.
 
+	// these allow external things to make custom decisions based on what sorts of task threads we are running now.
+	// this are bools to allow runtime tuning.
+	extern CORE_API int32 bHasBackgroundThreads; 
+	extern CORE_API int32 bHasHighPriorityThreads;
+
 	FORCEINLINE Type GetThreadIndex(Type ThreadAndIndex)
 	{
 		return ((ThreadAndIndex & ThreadIndexMask) == AnyThread) ? AnyThread : Type(ThreadAndIndex & ThreadIndexMask);
@@ -84,17 +105,95 @@ namespace ENamedThreads
 		return (ThreadAndIndex & QueueIndexMask) >> QueueIndexShift;
 	}
 
-	FORCEINLINE int32 GetPriority(Type ThreadAndIndex)
+	FORCEINLINE int32 GetTaskPriority(Type ThreadAndIndex)
+	{
+		return (ThreadAndIndex & TaskPriorityMask) >> TaskPriorityShift;
+	}
+
+	FORCEINLINE int32 GetThreadPriorityIndex(Type ThreadAndIndex)
+	{
+		int32 Result = (ThreadAndIndex & ThreadPriorityMask) >> ThreadPriorityShift;
+		check(Result >= 0 && Result < NumThreadPriorities);
+		return Result;
+	}
+
+	FORCEINLINE Type SetPriorities(Type ThreadAndIndex, Type ThreadPriority, Type TaskPriority)
+	{
+		check(
+			!(ThreadAndIndex & ~ThreadIndexMask) &&  // not a thread index
+			!(ThreadPriority & ~ThreadPriorityMask) && // not a thread priority
+			(ThreadPriority & ThreadPriorityMask) != ThreadPriorityMask && // not a valid thread priority
+			!(TaskPriority & ~TaskPriorityMask) // not a task priority
+			);
+		return Type(ThreadAndIndex | ThreadPriority | TaskPriority);
+	}
+
+	FORCEINLINE Type SetPriorities(Type ThreadAndIndex, int32 PriorityIndex, bool bHiPri)
+	{
+		check(
+			!(ThreadAndIndex & ~ThreadIndexMask) && // not a thread index
+			PriorityIndex >= 0 && PriorityIndex < NumThreadPriorities // not a valid thread priority
+			);
+		return Type(ThreadAndIndex | (PriorityIndex << ThreadPriorityShift) | (bHiPri ? HighTaskPriority : NormalTaskPriority));
+	}
+
+	FORCEINLINE Type SetThreadPriority(Type ThreadAndIndex, Type ThreadPriority)
 	{
-		return (ThreadAndIndex & PriorityMask) >> PriorityShift;
+		check(
+			!(ThreadAndIndex & ~ThreadIndexMask) &&  // not a thread index
+			!(ThreadPriority & ~ThreadPriorityMask) && // not a thread priority
+			(ThreadPriority & ThreadPriorityMask) != ThreadPriorityMask // not a valid thread priority
+			);
+		return Type(ThreadAndIndex | ThreadPriority);
 	}
 
-	FORCEINLINE Type HiPri(Type ThreadAndIndex)
+	FORCEINLINE Type SetTaskPriority(Type ThreadAndIndex, Type TaskPriority)
 	{
-		return Type(ThreadAndIndex | HighPriority);
+		check(
+			!(ThreadAndIndex & ~ThreadIndexMask) &&  // not a thread index
+			!(TaskPriority & ~TaskPriorityMask) // not a task priority
+			);
+		return Type(ThreadAndIndex | TaskPriority);
 	}
+
 }
 
+class CORE_API FAutoConsoleTaskPriority
+{
+	FAutoConsoleCommand Command;
+	FString CommandName;
+	ENamedThreads::Type ThreadPriority;
+	ENamedThreads::Type TaskPriority;
+	ENamedThreads::Type TaskPriorityIfForcedToNormalThreadPriority;
+	void CommandExecute(const TArray<FString>& Args);
+public:
+	FAutoConsoleTaskPriority(const TCHAR* Name, const TCHAR* Help, ENamedThreads::Type DefaultThreadPriority, ENamedThreads::Type DefaultTaskPriority, ENamedThreads::Type DefaultTaskPriorityIfForcedToNormalThreadPriority = ENamedThreads::UnusedAnchor)
+		: Command(Name, Help, FConsoleCommandWithArgsDelegate::CreateRaw(this, &FAutoConsoleTaskPriority::CommandExecute))
+		, CommandName(Name)
+		, ThreadPriority(DefaultThreadPriority)
+		, TaskPriority(DefaultTaskPriority)
+		, TaskPriorityIfForcedToNormalThreadPriority(DefaultTaskPriorityIfForcedToNormalThreadPriority)
+	{
+		// if you are asking for a hi or background thread priority, you must provide a separate task priority to use if those threads are not available.
+		check(TaskPriorityIfForcedToNormalThreadPriority != ENamedThreads::UnusedAnchor || ThreadPriority == ENamedThreads::NormalThreadPriority);
+	}
+
+	FORCEINLINE ENamedThreads::Type Get(ENamedThreads::Type Thread = ENamedThreads::AnyThread)
+	{
+		// if we don't have the high priority thread that was asked for, then use a normal thread priority with the backup task priority
+		if (ThreadPriority == ENamedThreads::HighThreadPriority && !ENamedThreads::bHasHighPriorityThreads)
+		{
+			return ENamedThreads::SetTaskPriority(Thread, TaskPriorityIfForcedToNormalThreadPriority);
+		}
+		// if we don't have the background priority thread that was asked for, then use a normal thread priority with the backup task priority
+		if (ThreadPriority == ENamedThreads::BackgroundThreadPriority && !ENamedThreads::bHasBackgroundThreads)
+		{
+			return ENamedThreads::SetTaskPriority(Thread, TaskPriorityIfForcedToNormalThreadPriority);
+		}
+		return ENamedThreads::SetPriorities(Thread, ThreadPriority, TaskPriority);
+	}
+};
+
 
 namespace ESubsequentsMode
 {
@@ -155,7 +254,10 @@ public:
 	/** Return the current thread type, if known. **/
 	virtual ENamedThreads::Type GetCurrentThreadIfKnown(bool bLocalQueue = false) = 0;
 
-	/** Return the number of worker (non-named) threads **/
+	/** 
+		Return the number of worker (non-named) threads PER PRIORITY SET.
+		This is useful for determining how many tasks to split a job into.
+	**/
 	virtual	int32 GetNumWorkerThreads() = 0;
 
 	/** Return true if the given named thread is processing tasks. This is only a "guess" if you ask for a thread other than yourself because that can change before the function returns. **/
@@ -256,7 +358,7 @@ public:
 		/** Total size in bytes for a small task that will use the custom allocator **/
 		SMALL_TASK_SIZE = 256
 	};
-	typedef TLockFreeFixedSizeAllocator_TLSCache<SMALL_TASK_SIZE> TSmallTaskAllocator;
+	typedef TLockFreeFixedSizeAllocator_TLSCache<SMALL_TASK_SIZE, PLATFORM_CACHE_LINE_SIZE> TSmallTaskAllocator;
 protected:
 	/** 
 	 *	Constructor
@@ -483,7 +585,7 @@ public:
 
 private:
 	friend class TRefCountPtr<FGraphEvent>;
-	friend class TLockFreeClassAllocator_TLSCache<FGraphEvent>;
+	friend class TLockFreeClassAllocator_TLSCache<FGraphEvent, PLATFORM_CACHE_LINE_SIZE>;
 
 	/** 
 	 *	Internal function to call the destructor and recycle a graph event
@@ -541,7 +643,7 @@ public:
 private:
 
 	/** Threadsafe list of subsequents for the event **/
-	TClosableLockFreePointerListUnorderedSingleConsumer<FBaseGraphTask>		SubsequentList;
+	TClosableLockFreePointerListUnorderedSingleConsumer<FBaseGraphTask, 0>		SubsequentList;
 	/** List of events to wait for until firing. This is not thread safe as it is only legal to fill it in within the context of an executing task. **/
 	FGraphEventArray														EventsToWaitFor;
 	/** Number of outstanding references to this graph event **/
@@ -1198,7 +1300,7 @@ public:
 
 	ENamedThreads::Type GetDesiredThread()
 	{
-		return ENamedThreads::HiPri(ENamedThreads::AnyThread);
+		return ENamedThreads::AnyHiPriThreadHiPriTask;
 	}
 
 	static ESubsequentsMode::Type GetSubsequentsMode() { return ESubsequentsMode::TrackSubsequents; }
@@ -1409,7 +1511,7 @@ public:
 **/
 class FCompletionList
 {
-	TLockFreePointerListUnordered<FGraphEvent>	Prerequisites;
+	TLockFreePointerListUnordered<FGraphEvent, 0>	Prerequisites;
 public:
 	/**
 	 * Adds a task to the completion list, can be called from any thread
@@ -1429,7 +1531,7 @@ public:
 	{
 		// this is tricky...
 		// we have waited for a set of pending tasks to execute. However, they may have added more pending tasks that also need to be waited for.
-		FGraphEventRef PendingComplete = CreatePrerequisiteCompletionHandle();
+		FGraphEventRef PendingComplete = CreatePrerequisiteCompletionHandle(CurrentThread);
 		if (PendingComplete.GetReference())
 		{
 			MyCompletionGraphEvent->DontCompleteUntil(PendingComplete);
@@ -1442,7 +1544,7 @@ public:
 	 * this should always be called from the same thread.
 	 * @return The task that when completed, indicates all tasks in the list are completed, including any tasks they added recursively. Will be a NULL reference if there are no tasks
 	 */
-	FGraphEventRef CreatePrerequisiteCompletionHandle()
+	FGraphEventRef CreatePrerequisiteCompletionHandle(ENamedThreads::Type CurrentThread)
 	{
 		FGraphEventRef CompleteHandle;
 		TArray<FGraphEvent*> Pending;
@@ -1464,7 +1566,7 @@ public:
 
 			CompleteHandle = FDelegateGraphTask::CreateAndDispatchWhenReady(
 				FDelegateGraphTask::FDelegate::CreateRaw(this, &FCompletionList::ChainWaitForPrerequisites),
-				GET_STATID(STAT_FDelegateGraphTask_WaitOnCompletionList), &PendingHandles
+				GET_STATID(STAT_FDelegateGraphTask_WaitOnCompletionList), &PendingHandles, CurrentThread, ENamedThreads::AnyHiPriThreadHiPriTask
 			);
 		}
 		return CompleteHandle;
diff --git a/Engine/Source/Runtime/Core/Public/Containers/LockFreeFixedSizeAllocator.h b/Engine/Source/Runtime/Core/Public/Containers/LockFreeFixedSizeAllocator.h
--- a/Engine/Source/Runtime/Core/Public/Containers/LockFreeFixedSizeAllocator.h
+++ b/Engine/Source/Runtime/Core/Public/Containers/LockFreeFixedSizeAllocator.h
@@ -12,7 +12,7 @@
  * never returns free space until program shutdown.
  * alignment isn't handled, assumes FMemory::Malloc will work
  */
-template<int32 SIZE, typename TTrackingCounter = FNoopCounter>
+template<int32 SIZE, int TPaddingForCacheContention, typename TTrackingCounter = FNoopCounter>
 class TLockFreeFixedSizeAllocator
 {
 	/**
@@ -97,7 +97,7 @@ public:
 private:
 
 	/** Lock free list of free memory blocks. */
-	FLockFreePointerListFIFOIntrusive<FFakeType> FreeList;
+	FLockFreePointerListFIFOIntrusive<FFakeType, TPaddingForCacheContention> FreeList;
 
 	/** Total number of blocks outstanding and not in the free list. */
 	TTrackingCounter NumUsed; 
@@ -112,7 +112,7 @@ private:
  * never returns free space until program shutdown.
  * alignment isn't handled, assumes FMemory::Malloc will work
  */
-template<int32 SIZE, typename TTrackingCounter = FNoopCounter>
+template<int32 SIZE, int TPaddingForCacheContention, typename TTrackingCounter = FNoopCounter>
 class TLockFreeFixedSizeAllocator
 {
 public:
@@ -188,7 +188,7 @@ public:
 private:
 
 	/** Lock free list of free memory blocks. */
-	TLockFreePointerListUnordered<void> FreeList;
+	TLockFreePointerListUnordered<void, TPaddingForCacheContention> FreeList;
 
 	/** Total number of blocks outstanding and not in the free list. */
 	TTrackingCounter NumUsed; 
@@ -203,8 +203,8 @@ private:
  * never returns free space, even at shutdown
  * alignment isn't handled, assumes FMemory::Malloc will work
  */
-template<int32 SIZE, typename TTrackingCounter = FNoopCounter>
-class TLockFreeFixedSizeAllocator_TLSCache : public TLockFreeFixedSizeAllocator_TLSCacheBase<SIZE, TLockFreePointerListUnordered<void*>, TTrackingCounter>
+template<int32 SIZE, int TPaddingForCacheContention, typename TTrackingCounter = FNoopCounter>
+class TLockFreeFixedSizeAllocator_TLSCache : public TLockFreeFixedSizeAllocator_TLSCacheBase<SIZE, TLockFreePointerListUnordered<void*, TPaddingForCacheContention>, TTrackingCounter>
 {
 };
 
@@ -213,8 +213,8 @@ class TLockFreeFixedSizeAllocator_TLSCache : public TLockFreeFixedSizeAllocator_
  *
  * Never returns free space until program shutdown.
  */
-template<class T>
-class TLockFreeClassAllocator : private TLockFreeFixedSizeAllocator<sizeof(T)>
+template<class T, int TPaddingForCacheContention>
+class TLockFreeClassAllocator : private TLockFreeFixedSizeAllocator<sizeof(T), TPaddingForCacheContention>
 {
 public:
 	/**
@@ -225,7 +225,7 @@ public:
 	 */
 	void* Allocate()
 	{
-		return TLockFreeFixedSizeAllocator<sizeof(T)>::Allocate();
+		return TLockFreeFixedSizeAllocator<sizeof(T), TPaddingForCacheContention>::Allocate();
 	}
 
 	/**
@@ -248,7 +248,7 @@ public:
 	void Free(T *Item)
 	{
 		Item->~T();
-		TLockFreeFixedSizeAllocator<sizeof(T)>::Free(Item);
+		TLockFreeFixedSizeAllocator<sizeof(T), TPaddingForCacheContention>::Free(Item);
 	}
 };
 
@@ -257,8 +257,8 @@ public:
  *
  * Never returns free space until program shutdown.
  */
-template<class T>
-class TLockFreeClassAllocator_TLSCache : private TLockFreeFixedSizeAllocator_TLSCache<sizeof(T)>
+template<class T, int TPaddingForCacheContention>
+class TLockFreeClassAllocator_TLSCache : private TLockFreeFixedSizeAllocator_TLSCache<sizeof(T), TPaddingForCacheContention>
 {
 public:
 	/**
@@ -269,7 +269,7 @@ public:
 	 */
 	void* Allocate()
 	{
-		return TLockFreeFixedSizeAllocator_TLSCache<sizeof(T)>::Allocate();
+		return TLockFreeFixedSizeAllocator_TLSCache<sizeof(T), TPaddingForCacheContention>::Allocate();
 	}
 
 	/**
@@ -292,6 +292,6 @@ public:
 	void Free(T *Item)
 	{
 		Item->~T();
-		TLockFreeFixedSizeAllocator_TLSCache<sizeof(T)>::Free(Item);
+		TLockFreeFixedSizeAllocator_TLSCache<sizeof(T), TPaddingForCacheContention>::Free(Item);
 	}
 };
diff --git a/Engine/Source/Runtime/Core/Public/Containers/LockFreeListImpl.h b/Engine/Source/Runtime/Core/Public/Containers/LockFreeListImpl.h
--- a/Engine/Source/Runtime/Core/Public/Containers/LockFreeListImpl.h
+++ b/Engine/Source/Runtime/Core/Public/Containers/LockFreeListImpl.h
@@ -158,12 +158,12 @@ public:
 
 #if !USE_NEW_LOCK_FREE_LISTS
 // where you don't care what order they pop in. We choose the fastest implementation we have.
-template <class T>
+template<class T, int TPaddingForCacheContention>
 class TLockFreePointerListUnordered : public TLockFreePointerListLIFOBase<T>
 {
 
 };
-template <class T>
+template<class T, int TPaddingForCacheContention>
 class TClosableLockFreePointerListUnorderedSingleConsumer : public TClosableLockFreePointerListLIFO<T>
 {
 
@@ -171,7 +171,7 @@ class TClosableLockFreePointerListUnorderedSingleConsumer : public TClosableLock
 
 #else
 
-template <class T>
+template<class T, int TPaddingForCacheContention>
 class TClosableLockFreePointerListFIFOSingleConsumer
 {
 public:
@@ -222,10 +222,10 @@ public:
 		return Inner.Reset();
 	}
 private:
-	FCloseableLockFreePointerListFIFOBaseSingleConsumer Inner;
+	FCloseableLockFreePointerListFIFOBaseSingleConsumer<TPaddingForCacheContention> Inner;
 };
 
-template <class T>
+template<class T, int TPaddingForCacheContention>
 class TReopenableLockFreePointerListFIFOSingleConsumer
 {
 public:
@@ -259,11 +259,11 @@ public:
 		return Inner.CloseIfEmpty();
 	}
 private:
-	FCloseableLockFreePointerListFIFOBaseSingleConsumer Inner;
+	FCloseableLockFreePointerListFIFOBaseSingleConsumer<TPaddingForCacheContention> Inner;
 };
 
-template <class T>
-class TClosableLockFreePointerListUnorderedSingleConsumer : public TClosableLockFreePointerListFIFOSingleConsumer<T>
+template<class T, int TPaddingForCacheContention>
+class TClosableLockFreePointerListUnorderedSingleConsumer : public TClosableLockFreePointerListFIFOSingleConsumer<T, TPaddingForCacheContention>
 {
 
 };
diff --git a/Engine/Source/Runtime/Core/Public/Containers/LockFreeVoidPointerListBase.h b/Engine/Source/Runtime/Core/Public/Containers/LockFreeVoidPointerListBase.h
--- a/Engine/Source/Runtime/Core/Public/Containers/LockFreeVoidPointerListBase.h
+++ b/Engine/Source/Runtime/Core/Public/Containers/LockFreeVoidPointerListBase.h
@@ -2,6 +2,9 @@
 
 #pragma once
 
+#include "HAL/Platform.h"
+#include "Misc/Build.h"
+
 #define CHECK_NON_CONCURRENT_ASSUMPTIONS (0)
 
 #define USE_NEW_LOCK_FREE_LISTS (PLATFORM_PS4)
@@ -947,7 +950,7 @@ FORCEINLINE void TestCriticalStall()
 }
 #endif
 
-template<class T>
+template<class T, int TPaddingForCacheContention>
 class FLockFreePointerQueueBaseSingleConsumerIntrusive : public FNoncopyable
 {
 public:
@@ -1025,9 +1028,11 @@ public:
 	}
 
 private:
-
+	uint8 PadToAvoidContention1[TPaddingForCacheContention];
 	T* Head;
+	uint8 PadToAvoidContention2[TPaddingForCacheContention];
 	T* Tail;
+	uint8 PadToAvoidContention3[TPaddingForCacheContention];
 
 	T* GetStub() const
 	{
@@ -1042,6 +1047,7 @@ struct FLockFreeLink
 	void *Payload;
 };
 
+template<int TPaddingForCacheContention>
 class FLockFreePointerQueueBaseSingleConsumer : public FNoncopyable
 {
 public:
@@ -1097,10 +1103,14 @@ public:
 
 private:
 
+	uint8 PadToAvoidContention1[TPaddingForCacheContention];
 	FLockFreeLink* Head;
+	uint8 PadToAvoidContention2[TPaddingForCacheContention];
 	FLockFreeLink* Tail;
+	uint8 PadToAvoidContention3[TPaddingForCacheContention];
 };
 
+template<int TPaddingForCacheContention>
 class FCloseableLockFreePointerQueueBaseSingleConsumer : public FNoncopyable
 {
 public:
@@ -1273,11 +1283,14 @@ private:
 		return (FLockFreeLink*)(UPTRINT(In) & ~UPTRINT(1));
 	}
 
+	uint8 PadToAvoidContention1[TPaddingForCacheContention];
 	FLockFreeLink* Head;
+	uint8 PadToAvoidContention2[TPaddingForCacheContention];
 	FLockFreeLink* Tail;
+	uint8 PadToAvoidContention3[TPaddingForCacheContention];
 };
 
-template<class T>
+template<class T, int TPaddingForCacheContention>
 class FCloseableLockFreePointerQueueBaseSingleBaseConsumerIntrusive : public FNoncopyable
 {
 public:
@@ -1483,8 +1496,11 @@ private:
 	}
 
 
+	uint8 PadToAvoidContention1[TPaddingForCacheContention];
 	T* Head;
+	uint8 PadToAvoidContention2[TPaddingForCacheContention];
 	T* Tail;
+	uint8 PadToAvoidContention3[TPaddingForCacheContention];
 
 	T* GetStub() const
 	{
@@ -1574,6 +1590,7 @@ private:
 	TCounter NumPushed;
 };
 
+template<int TPaddingForCacheContention>
 class FSpinLocked_LIFO
 {
 public:
@@ -1615,7 +1632,9 @@ public:
 	}	
 private:
 	TArray<void **> Queue;
+	uint8 PadToAvoidContention1[TPaddingForCacheContention];
 	int32 Lock;
+	uint8 PadToAvoidContention2[TPaddingForCacheContention];
 };
 
 class FLockFreePointerQueueBaseLinkAllocator : public FNoncopyable
@@ -1655,7 +1674,7 @@ public:
 	}
 
 private:
-	TLockFreeFixedSizeAllocator_TLSCacheBase<sizeof(FLockFreeLink), FSpinLocked_LIFO,
+	TLockFreeFixedSizeAllocator_TLSCacheBase<sizeof(FLockFreeLink), FSpinLocked_LIFO<PLATFORM_CACHE_LINE_SIZE>,
 #if 1 || UE_BUILD_SHIPPING || UE_BUILD_TEST
 		FNoopCounter
 #else
@@ -1665,7 +1684,7 @@ private:
 };
 
 
-template<class T>
+template<class T, int TPaddingForCacheContention>
 class FLockFreePointerListFIFOIntrusive : public FNoncopyable
 {
 public:
@@ -1802,12 +1821,13 @@ public:
 	}	
 
 private:
-	FLockFreePointerQueueBaseSingleConsumerIntrusive<T> IncomingQueue;
+	FLockFreePointerQueueBaseSingleConsumerIntrusive<T, TPaddingForCacheContention> IncomingQueue;
 	FDequeueCache DequeueCache;
 	int32 DequeueLock;
 };
 
 // FIFO of void *'s
+template<int TPaddingForCacheContention>
 class FLockFreePointerListFIFOBase : public FNoncopyable
 {
 public:
@@ -1935,12 +1955,13 @@ public:
 
 private:
 	FLockFreePointerQueueBaseLinkAllocator& LinkAllocator;
-	FLockFreePointerQueueBaseSingleConsumer IncomingQueue;
+	FLockFreePointerQueueBaseSingleConsumer<TPaddingForCacheContention> IncomingQueue;
 	FDequeueCache DequeueCache;
 	int32 DequeueLock;
 };
 
 // closeable FIFO of void *'s
+template<int TPaddingForCacheContention>
 class FCloseableLockFreePointerListFIFOBaseSingleConsumer : public FNoncopyable
 {
 public:
@@ -2053,12 +2074,12 @@ public:
 
 private:
 	FLockFreePointerQueueBaseLinkAllocator& LinkAllocator;
-	FCloseableLockFreePointerQueueBaseSingleConsumer IncomingQueue;
+	FCloseableLockFreePointerQueueBaseSingleConsumer<TPaddingForCacheContention> IncomingQueue;
 };
 
 
-template <class T>
-class TLockFreePointerListFIFO : private FLockFreePointerListFIFOBase 
+template<class T, int TPaddingForCacheContention>
+class TLockFreePointerListFIFO : private FLockFreePointerListFIFOBase<TPaddingForCacheContention>
 {
 public:
 	/**	
@@ -2068,7 +2089,7 @@ public:
 	 */
 	FORCEINLINE void Push(T *NewItem)
 	{
-		FLockFreePointerListFIFOBase::Push(NewItem);
+		FLockFreePointerListFIFOBase<TPaddingForCacheContention>::Push(NewItem);
 	}
 
 	/**	
@@ -2077,7 +2098,7 @@ public:
 	 */
 	FORCEINLINE T* Pop()
 	{
-		return (T*)FLockFreePointerListFIFOBase::Pop();
+		return (T*)FLockFreePointerListFIFOBase<TPaddingForCacheContention>::Pop();
 	}
 
 	/**	
@@ -2107,22 +2128,22 @@ public:
 	 */
 	FORCEINLINE bool IsEmpty() const  
 	{
-		return FLockFreePointerListFIFOBase::IsEmpty();
+		return FLockFreePointerListFIFOBase<TPaddingForCacheContention>::IsEmpty();
 	}
 	/**	
 	 *	Check if the list is empty. This is a faster, less rigorous test that can miss a queue full of stuff. Used to optimize thread restarts.
 	 */
 	FORCEINLINE bool IsEmptyFast()  
 	{
-		return FLockFreePointerListFIFOBase::IsEmptyFast();
+		return FLockFreePointerListFIFOBase<TPaddingForCacheContention>::IsEmptyFast();
 	}
 };
 
 #if USE_NEW_LOCK_FREE_LISTS
 
 // a list where you don't care what order they pop in. We choose the fastest implementation we have.
-template <class T>
-class TLockFreePointerListUnordered : public TLockFreePointerListFIFO<T>
+template<class T, int TPaddingForCacheContention>
+class TLockFreePointerListUnordered : public TLockFreePointerListFIFO<T, TPaddingForCacheContention>
 {
 
 };
diff --git a/Engine/Source/Runtime/Core/Public/HAL/ThreadingBase.h b/Engine/Source/Runtime/Core/Public/HAL/ThreadingBase.h
--- a/Engine/Source/Runtime/Core/Public/HAL/ThreadingBase.h
+++ b/Engine/Source/Runtime/Core/Public/HAL/ThreadingBase.h
@@ -12,17 +12,12 @@ class FRunnableThread;
  */
 enum EThreadPriority
 {
-	/** Normal priority. */
 	TPri_Normal,
-
-	/** Above normal priority. */
 	TPri_AboveNormal,
-
-	/** Below normal priority. */
 	TPri_BelowNormal,
-	
-	/** Highest priority. */
-	TPri_Highest
+	TPri_Highest,
+	TPri_Lowest,
+	TPri_SlightlyBelowNormal,
 };
 
 
diff --git a/Engine/Source/Runtime/Core/Public/Misc/MemStack.h b/Engine/Source/Runtime/Core/Public/Misc/MemStack.h
--- a/Engine/Source/Runtime/Core/Public/Misc/MemStack.h
+++ b/Engine/Source/Runtime/Core/Public/Misc/MemStack.h
@@ -26,6 +26,14 @@ public:
 		PageSize = 64 * 1024,
 		SmallPageSize = 1024
 	};
+#if UE_BUILD_SHIPPING
+	typedef TLockFreeFixedSizeAllocator<PageSize, PLATFORM_CACHE_LINE_SIZE, FNoopCounter> TPageAllocator;
+	typedef TLockFreeFixedSizeAllocator<SmallPageSize, PLATFORM_CACHE_LINE_SIZE, FNoopCounter> TSmallPageAllocator;
+#else
+	typedef TLockFreeFixedSizeAllocator<PageSize, PLATFORM_CACHE_LINE_SIZE, FThreadSafeCounter> TPageAllocator;
+	typedef TLockFreeFixedSizeAllocator<SmallPageSize, PLATFORM_CACHE_LINE_SIZE, FThreadSafeCounter> TSmallPageAllocator;
+#endif
+
 	static void *Alloc();
 	static void Free(void *Mem);
 	static void *AllocSmall();
@@ -38,8 +46,8 @@ private:
 #if STATS
 	static void UpdateStats();
 #endif
-	static TLockFreeFixedSizeAllocator<PageSize, FThreadSafeCounter> TheAllocator;
-	static TLockFreeFixedSizeAllocator<SmallPageSize, FThreadSafeCounter> TheSmallAllocator;
+	static TPageAllocator TheAllocator;
+	static TSmallPageAllocator TheSmallAllocator;
 };
 
 
diff --git a/Engine/Source/Runtime/Core/Public/Stats/Stats2.h b/Engine/Source/Runtime/Core/Public/Stats/Stats2.h
--- a/Engine/Source/Runtime/Core/Public/Stats/Stats2.h
+++ b/Engine/Source/Runtime/Core/Public/Stats/Stats2.h
@@ -1163,7 +1163,7 @@ private:
 	};
 
 	/** Lock free pool of FThreadStats instances. */
-	TLockFreePointerListUnordered<FThreadStats> Pool;
+	TLockFreePointerListUnordered<FThreadStats, 0> Pool;
 
 public:
 	/** Default constructor. */
diff --git a/Engine/Source/Runtime/CoreUObject/Public/UObject/UObjectArray.h b/Engine/Source/Runtime/CoreUObject/Public/UObject/UObjectArray.h
--- a/Engine/Source/Runtime/CoreUObject/Public/UObject/UObjectArray.h
+++ b/Engine/Source/Runtime/CoreUObject/Public/UObject/UObjectArray.h
@@ -697,7 +697,7 @@ private:
 	/** Synchronization object for all live objects.											*/
 	FCriticalSection ObjObjectsCritical;
 	/** Available object indices.											*/
-	TLockFreePointerListUnordered<int32> ObjAvailableList;
+	TLockFreePointerListUnordered<int32, PLATFORM_CACHE_LINE_SIZE> ObjAvailableList;
 #if UE_GC_TRACK_OBJ_AVAILABLE
 	/** Available object index count.										*/
 	FThreadSafeCounter ObjAvailableCount;
diff --git a/Engine/Source/Runtime/Engine/Classes/Animation/VertexAnim/VertexAnimBase.h b/Engine/Source/Runtime/Engine/Classes/Animation/VertexAnim/VertexAnimBase.h
--- a/Engine/Source/Runtime/Engine/Classes/Animation/VertexAnim/VertexAnimBase.h
+++ b/Engine/Source/Runtime/Engine/Classes/Animation/VertexAnim/VertexAnimBase.h
@@ -10,9 +10,7 @@ struct FVertexAnimDelta
 {
 	/** change in position */
 	FVector			PositionDelta;
-	/** old format of change in tangent basis normal */
-	FPackedNormal	TangentZDelta_DEPRECATED;
-	/** new format of change in tangent basis normal */
+	/** Tangent basis normal */
 	FVector			TangentZDelta;
 
 	/** index of source vertex to apply deltas to */
@@ -23,8 +21,19 @@ struct FVertexAnimDelta
 	{
 		if ( Ar.UE4Ver() < VER_UE4_MORPHTARGET_CPU_TANGENTZDELTA_FORMATCHANGE )
 		{
-			Ar << V.PositionDelta << V.TangentZDelta_DEPRECATED << V.SourceIdx;
-			V.TangentZDelta = FVector(V.TangentZDelta_DEPRECATED);
+			/** old format of change in tangent basis normal */
+			FPackedNormal	TangentZDelta_DEPRECATED;
+			if (Ar.IsSaving())
+			{
+				TangentZDelta_DEPRECATED = FPackedNormal(V.TangentZDelta);
+			}
+
+			Ar << V.PositionDelta << TangentZDelta_DEPRECATED << V.SourceIdx;
+
+			if (Ar.IsLoading())
+			{
+				V.TangentZDelta = FVector(TangentZDelta_DEPRECATED);
+			}
 		}
 		else
 		{
diff --git a/Engine/Source/Runtime/Engine/Classes/Engine/RendererSettings.h b/Engine/Source/Runtime/Engine/Classes/Engine/RendererSettings.h
--- a/Engine/Source/Runtime/Engine/Classes/Engine/RendererSettings.h
+++ b/Engine/Source/Runtime/Engine/Classes/Engine/RendererSettings.h
@@ -84,6 +84,19 @@ namespace EAntiAliasingMethodUI
 	};
 }
 
+/** used by FPostProcessSettings AutoExposure*/
+UENUM()
+namespace EAutoExposureMethodUI
+{
+	enum Type
+	{
+		/** Not supported on mobile, requires compute shader to construct 64 bin histogram */
+		AEM_Histogram  UMETA(DisplayName = "Auto Exposure Histogram"),
+		/** Not supported on mobile, faster method that computes single value by downsampling */
+		AEM_Basic      UMETA(DisplayName = "Auto Exposure Basic"),
+		AEM_MAX,
+	};
+}
 
 /**
  * Rendering settings.
@@ -218,6 +231,11 @@ class ENGINE_API URendererSettings : public UDeveloperSettings
 		ConsoleVariable = "r.DefaultFeature.AutoExposure", DisplayName = "Auto Exposure",
 		ToolTip = "Whether the default for AutoExposure is enabled or not (postprocess volume/camera/game setting can still override and enable or disable it independently)"))
 	uint32 bDefaultFeatureAutoExposure : 1;
+	
+	UPROPERTY(config, EditAnywhere, Category = DefaultPostprocessingSettings, meta = (
+		ConsoleVariable = "r.DefaultFeature.AutoExposure.Method", DisplayName = "Auto Exposure",
+		ToolTip = "The default method for AutoExposure(postprocess volume/camera/game setting can still override and enable or disable it independently)"))
+	TEnumAsByte<EAutoExposureMethodUI::Type> DefaultFeatureAutoExposure; 
 
 	UPROPERTY(config, EditAnywhere, Category = DefaultPostprocessingSettings, meta = (
 		ConsoleVariable = "r.DefaultFeature.MotionBlur", DisplayName = "Motion Blur",
diff --git a/Engine/Source/Runtime/Engine/Classes/Engine/Scene.h b/Engine/Source/Runtime/Engine/Classes/Engine/Scene.h
--- a/Engine/Source/Runtime/Engine/Classes/Engine/Scene.h
+++ b/Engine/Source/Runtime/Engine/Classes/Engine/Scene.h
@@ -1084,6 +1084,8 @@ struct FPostProcessSettings
 		LPVVplInjectionBias = 0.64f;
 		LPVGeometryVolumeBias = 0.384f;
 		LPVEmissiveInjectionIntensity = 1.0f;
+		// next value might get overwritten by r.DefaultFeature.AutoExposure.Method
+		AutoExposureMethod = AEM_Histogram;
 		AutoExposureLowPercent = 80.0f;
 		AutoExposureHighPercent = 98.3f;
 		// next value might get overwritten by r.DefaultFeature.AutoExposure
diff --git a/Engine/Source/Runtime/Engine/Classes/Particles/ParticleModule.h b/Engine/Source/Runtime/Engine/Classes/Particles/ParticleModule.h
--- a/Engine/Source/Runtime/Engine/Classes/Particles/ParticleModule.h
+++ b/Engine/Source/Runtime/Engine/Classes/Particles/ParticleModule.h
@@ -37,6 +37,8 @@ enum EModuleType
 	EPMT_Event UMETA(DisplayName="Event"),
 	/** Light related modules							*/
 	EPMT_Light UMETA(DisplayName="Light"),
+	/** SubUV related modules							*/
+	EPMT_SubUV UMETA(DisplayName = "SubUV"),
 	EPMT_MAX,
 };
 
diff --git a/Engine/Source/Runtime/Engine/Classes/Particles/SubUV/ParticleModuleSubUVBase.h b/Engine/Source/Runtime/Engine/Classes/Particles/SubUV/ParticleModuleSubUVBase.h
--- a/Engine/Source/Runtime/Engine/Classes/Particles/SubUV/ParticleModuleSubUVBase.h
+++ b/Engine/Source/Runtime/Engine/Classes/Particles/SubUV/ParticleModuleSubUVBase.h
@@ -10,5 +10,9 @@ class UParticleModuleSubUVBase : public UParticleModule
 {
 	GENERATED_UCLASS_BODY()
 
+	// Begin UParticleModule Interface
+	virtual EModuleType	GetModuleType() const override { return EPMT_SubUV; }
+	//End UParticleModule Interface
+
 };
 
diff --git a/Engine/Source/Runtime/Engine/Private/AI/Navigation/NavigationSystem.cpp b/Engine/Source/Runtime/Engine/Private/AI/Navigation/NavigationSystem.cpp
--- a/Engine/Source/Runtime/Engine/Private/AI/Navigation/NavigationSystem.cpp
+++ b/Engine/Source/Runtime/Engine/Private/AI/Navigation/NavigationSystem.cpp
@@ -1033,6 +1033,15 @@ void UNavigationSystem::AbortAsyncFindPathRequest(uint32 AsynPathQueryID)
 	}
 }
 
+FAutoConsoleTaskPriority CPrio_TriggerAsyncQueries(
+	TEXT("TaskGraph.TaskPriorities.NavTriggerAsyncQueries"),
+	TEXT("Task and thread priority for UNavigationSystem::PerformAsyncQueries."),
+	ENamedThreads::BackgroundThreadPriority, // if we have background priority task threads, then use them...
+	ENamedThreads::NormalTaskPriority, // .. at normal task priority
+	ENamedThreads::NormalTaskPriority // if we don't have background threads, then use normal priority threads at normal task priority instead
+	);
+
+
 void UNavigationSystem::TriggerAsyncQueries(TArray<FAsyncPathFindingQuery>& PathFindingQueries)
 {
 	DECLARE_CYCLE_STAT(TEXT("FSimpleDelegateGraphTask.NavigationSystem batched async queries"),
@@ -1041,7 +1050,7 @@ void UNavigationSystem::TriggerAsyncQueries(TArray<FAsyncPathFindingQuery>& Path
 
 	FSimpleDelegateGraphTask::CreateAndDispatchWhenReady(
 		FSimpleDelegateGraphTask::FDelegate::CreateUObject(this, &UNavigationSystem::PerformAsyncQueries, PathFindingQueries),
-		GET_STATID(STAT_FSimpleDelegateGraphTask_NavigationSystemBatchedAsyncQueries));
+		GET_STATID(STAT_FSimpleDelegateGraphTask_NavigationSystemBatchedAsyncQueries), nullptr, CPrio_TriggerAsyncQueries.Get());
 }
 
 static void AsyncQueryDone(FAsyncPathFindingQuery Query)
diff --git a/Engine/Source/Runtime/Engine/Private/Collision/WorldCollisionAsync.cpp b/Engine/Source/Runtime/Engine/Private/Collision/WorldCollisionAsync.cpp
--- a/Engine/Source/Runtime/Engine/Private/Collision/WorldCollisionAsync.cpp
+++ b/Engine/Source/Runtime/Engine/Private/Collision/WorldCollisionAsync.cpp
@@ -161,6 +161,15 @@ namespace
 
 	#if RUN_ASYNC_TRACE
 
+	FAutoConsoleTaskPriority CPrio_FAsyncTraceTask(
+		TEXT("TaskGraph.TaskPriorities.AsyncTraceTask"),
+		TEXT("Task and thread priority for async traces."),
+		ENamedThreads::BackgroundThreadPriority, // if we have background priority task threads, then use them...
+		ENamedThreads::HighTaskPriority, // .. at normal task priority
+		ENamedThreads::NormalTaskPriority // if we don't have background threads, then use normal priority threads at normal task priority instead
+		);
+
+
 	/** Helper class define the task of Async Trace running**/
 	class FAsyncTraceTask
 	{
@@ -193,16 +202,16 @@ namespace
 			DataCount   = InDataCount;
 		}
 
-		FORCEINLINE TStatId GetStatId() const
+		static FORCEINLINE TStatId GetStatId()
 		{
 			RETURN_QUICK_DECLARE_CYCLE_STAT(FAsyncTraceTask, STATGROUP_TaskGraphTasks);
 		}
 		/** return the thread for this task **/
-		static ENamedThreads::Type GetDesiredThread()
+		static FORCEINLINE ENamedThreads::Type GetDesiredThread()
 		{
-			return ENamedThreads::AnyThread;
+			return CPrio_FAsyncTraceTask.Get();
 		}
-		static ESubsequentsMode::Type GetSubsequentsMode() 
+		static FORCEINLINE ESubsequentsMode::Type GetSubsequentsMode()
 		{ 
 			return ESubsequentsMode::TrackSubsequents; 
 		}
diff --git a/Engine/Source/Runtime/Engine/Private/Components/ArrowComponent.cpp b/Engine/Source/Runtime/Engine/Private/Components/ArrowComponent.cpp
--- a/Engine/Source/Runtime/Engine/Private/Components/ArrowComponent.cpp
+++ b/Engine/Source/Runtime/Engine/Private/Components/ArrowComponent.cpp
@@ -66,7 +66,7 @@ public:
 		if (IsInRenderingThread())
 		{
 			// Initialize the vertex factory's stream components.
-			DataType NewData;
+			FDataType NewData;
 			NewData.PositionComponent = STRUCTMEMBER_VERTEXSTREAMCOMPONENT(VertexBuffer, FDynamicMeshVertex, Position, VET_Float3);
 			NewData.TextureCoordinates.Add(
 				FVertexStreamComponent(VertexBuffer, STRUCT_OFFSET(FDynamicMeshVertex, TextureCoordinate), sizeof(FDynamicMeshVertex), VET_Float2)
@@ -83,7 +83,7 @@ public:
 				const FArrowVertexBuffer*, VertexBuffer, VertexBuffer,
 				{
 				// Initialize the vertex factory's stream components.
-				DataType NewData;
+				FDataType NewData;
 				NewData.PositionComponent = STRUCTMEMBER_VERTEXSTREAMCOMPONENT(VertexBuffer, FDynamicMeshVertex, Position, VET_Float3);
 				NewData.TextureCoordinates.Add(
 					FVertexStreamComponent(VertexBuffer, STRUCT_OFFSET(FDynamicMeshVertex, TextureCoordinate), sizeof(FDynamicMeshVertex), VET_Float2)
diff --git a/Engine/Source/Runtime/Engine/Private/Components/BrushComponent.cpp b/Engine/Source/Runtime/Engine/Private/Components/BrushComponent.cpp
--- a/Engine/Source/Runtime/Engine/Private/Components/BrushComponent.cpp
+++ b/Engine/Source/Runtime/Engine/Private/Components/BrushComponent.cpp
@@ -196,7 +196,7 @@ public:
 			FLocalVertexFactory*,VertexFactory,&VertexFactory,
 			FVertexBuffer*,WireVertexBuffer,&WireVertexBuffer,
 			{
-				FLocalVertexFactory::DataType Data;
+				FLocalVertexFactory::FDataType Data;
 				Data.PositionComponent = STRUCTMEMBER_VERTEXSTREAMCOMPONENT(WireVertexBuffer,FModelWireVertex,Position,VET_Float3);
 				Data.TangentBasisComponents[0] = STRUCTMEMBER_VERTEXSTREAMCOMPONENT(WireVertexBuffer,FModelWireVertex,TangentX,VET_PackedNormal);
 				Data.TangentBasisComponents[1] = STRUCTMEMBER_VERTEXSTREAMCOMPONENT(WireVertexBuffer,FModelWireVertex,TangentZ,VET_PackedNormal);
diff --git a/Engine/Source/Runtime/Engine/Private/Components/DecalComponent.cpp b/Engine/Source/Runtime/Engine/Private/Components/DecalComponent.cpp
--- a/Engine/Source/Runtime/Engine/Private/Components/DecalComponent.cpp
+++ b/Engine/Source/Runtime/Engine/Private/Components/DecalComponent.cpp
@@ -8,6 +8,12 @@
 #include "LevelUtils.h"
 #include "Components/DecalComponent.h"
 
+static TAutoConsoleVariable<float> CVarDecalFadeDurationScale(
+	TEXT("r.Decal.FadeDurationScale"),
+	1.0f,
+	TEXT("Scales the per decal fade durations. Lower values shortens lifetime and fade duration. Default is 1.0f.")
+	);
+
 FDeferredDecalProxy::FDeferredDecalProxy(const UDecalComponent* InComponent)
 	: InvFadeDuration(0.0f)
 	, FadeStartDelayNormalized(1.0f)
@@ -102,8 +108,11 @@ float UDecalComponent::GetFadeDuration() const
 
 void UDecalComponent::SetFadeOut(float StartDelay, float Duration, bool DestroyOwnerAfterFade /*= true*/)
 {
-	FadeStartDelay = StartDelay;
-	FadeDuration = Duration;
+	float FadeDurationScale = CVarDecalFadeDurationScale.GetValueOnGameThread();
+	FadeDurationScale = (FadeDurationScale <= SMALL_NUMBER) ? 0.0f : FadeDurationScale;
+
+	FadeStartDelay = StartDelay * FadeDurationScale;
+	FadeDuration = Duration * FadeDurationScale;
 	bDestroyOwnerAfterFade = DestroyOwnerAfterFade;
 	SetLifeSpan(FadeStartDelay + FadeDuration);
 
diff --git a/Engine/Source/Runtime/Engine/Private/Components/SkeletalMeshComponent.cpp b/Engine/Source/Runtime/Engine/Private/Components/SkeletalMeshComponent.cpp
--- a/Engine/Source/Runtime/Engine/Private/Components/SkeletalMeshComponent.cpp
+++ b/Engine/Source/Runtime/Engine/Private/Components/SkeletalMeshComponent.cpp
@@ -51,6 +51,14 @@ DECLARE_CYCLE_STAT_EXTERN(TEXT("Anim Instance Spawn Time"), STAT_AnimSpawnTime,
 DEFINE_STAT(STAT_AnimSpawnTime);
 DEFINE_STAT(STAT_PostAnimEvaluation);
 
+FAutoConsoleTaskPriority CPrio_ParallelAnimationEvaluationTask(
+	TEXT("TaskGraph.TaskPriorities.ParallelAnimationEvaluationTask"),
+	TEXT("Task and thread priority for FParallelAnimationEvaluationTask"),
+	ENamedThreads::HighThreadPriority, // if we have high priority task threads, then use them...
+	ENamedThreads::NormalTaskPriority, // .. at normal task priority
+	ENamedThreads::HighTaskPriority // if we don't have hi pri threads, then use normal priority threads at high task priority instead
+	);
+
 class FParallelAnimationEvaluationTask
 {
 	TWeakObjectPtr<USkeletalMeshComponent> SkeletalMeshComponent;
@@ -65,11 +73,11 @@ public:
 	{
 		RETURN_QUICK_DECLARE_CYCLE_STAT(FParallelAnimationEvaluationTask, STATGROUP_TaskGraphTasks);
 	}
-	static ENamedThreads::Type GetDesiredThread()
+	static FORCEINLINE ENamedThreads::Type GetDesiredThread()
 	{
-		return ENamedThreads::AnyThread;
+		return CPrio_ParallelAnimationEvaluationTask.Get();
 	}
-	static ESubsequentsMode::Type GetSubsequentsMode()
+	static FORCEINLINE ESubsequentsMode::Type GetSubsequentsMode()
 	{
 		return ESubsequentsMode::TrackSubsequents;
 	}
@@ -79,11 +87,13 @@ public:
 		if (USkeletalMeshComponent* Comp = SkeletalMeshComponent.Get())
 		{
 			FScopeCycleCounterUObject ContextScope(Comp);
+#if !UE_BUILD_TEST && !UE_BUILD_SHIPPING
 			float Stall = CVarStallParallelAnimation.GetValueOnAnyThread();
 			if (Stall > 0.0f)
 			{
 				FPlatformProcess::Sleep(Stall / 1000.0f);
 			}
+#endif
 			Comp->ParallelAnimationEvaluation();
 		}
 	}
diff --git a/Engine/Source/Runtime/Engine/Private/DynamicMeshBuilder.cpp b/Engine/Source/Runtime/Engine/Private/DynamicMeshBuilder.cpp
--- a/Engine/Source/Runtime/Engine/Private/DynamicMeshBuilder.cpp
+++ b/Engine/Source/Runtime/Engine/Private/DynamicMeshBuilder.cpp
@@ -291,7 +291,7 @@ public:
 		// Initialize the vertex factory's stream components.
 		if(IsInRenderingThread())
 		{
-			DataType TheData;
+			FDataType TheData;
 			TheData.PositionComponent = STRUCTMEMBER_VERTEXSTREAMCOMPONENT(VertexBuffer,FDynamicMeshVertex,Position,VET_Float3);
 			TheData.TextureCoordinates.Add(
 				FVertexStreamComponent(VertexBuffer,STRUCT_OFFSET(FDynamicMeshVertex,TextureCoordinate),sizeof(FDynamicMeshVertex),VET_Float2)
@@ -308,7 +308,7 @@ public:
 			    FLocalVertexFactory*,VertexFactory,this,
 			    const FDynamicMeshVertexBuffer*,VertexBuffer,VertexBuffer,
 			    {
-				    DataType TheData;
+				    FDataType TheData;
 				    TheData.PositionComponent = STRUCTMEMBER_VERTEXSTREAMCOMPONENT(VertexBuffer,FDynamicMeshVertex,Position,VET_Float3);
 				    TheData.TextureCoordinates.Add(
 					    FVertexStreamComponent(VertexBuffer,STRUCT_OFFSET(FDynamicMeshVertex,TextureCoordinate),sizeof(FDynamicMeshVertex),VET_Float2)
diff --git a/Engine/Source/Runtime/Engine/Private/GPUSkinCache.cpp b/Engine/Source/Runtime/Engine/Private/GPUSkinCache.cpp
--- a/Engine/Source/Runtime/Engine/Private/GPUSkinCache.cpp
+++ b/Engine/Source/Runtime/Engine/Private/GPUSkinCache.cpp
@@ -71,35 +71,16 @@ TGlobalResource<FGPUSkinCache> GGPUSkinCache;
 IMPLEMENT_UNIFORM_BUFFER_STRUCT(GPUSkinCacheBonesUniformShaderParameters,TEXT("GPUSkinCacheBones"));
 
 
-/** Compute shader that skins a batch of vertices. */
-template <bool bUseExtraBoneInfluencesT>
-class FGPUSkinCacheCS : public FGlobalShader
+class FBaseGPUSkinCacheCS : public FGlobalShader
 {
-	DECLARE_SHADER_TYPE(FGPUSkinCacheCS<bUseExtraBoneInfluencesT>,Global)
 public:
+	FBaseGPUSkinCacheCS() {} 
 
-	static bool ShouldCache(EShaderPlatform Platform)
-	{
-		return GEnableGPUSkinCacheShaders && IsFeatureLevelSupported(Platform, ERHIFeatureLevel::SM5);
-	}
-
-	static void ModifyCompilationEnvironment(EShaderPlatform Platform, FShaderCompilerEnvironment& OutEnvironment)
-	{
-		FGlobalShader::ModifyCompilationEnvironment(Platform,OutEnvironment);
-		const uint32 UseExtraBoneInfluences = bUseExtraBoneInfluencesT;
-		OutEnvironment.SetDefine(TEXT("GPUSKIN_USE_EXTRA_INFLUENCES"), UseExtraBoneInfluences);
-
-		OutEnvironment.SetDefine(TEXT("GPUSKIN_RWBUFFER_OFFSET_POSITION"), FGPUSkinCache::RWPositionOffsetInFloats);
-		OutEnvironment.SetDefine(TEXT("GPUSKIN_RWBUFFER_OFFSET_TANGENT_X"), FGPUSkinCache::RWTangentXOffsetInFloats);
-		OutEnvironment.SetDefine(TEXT("GPUSKIN_RWBUFFER_OFFSET_TANGENT_Z"), FGPUSkinCache::RWTangentZOffsetInFloats);
-		OutEnvironment.SetDefine(TEXT("GPUSKIN_RWBUFFER_NUM_FLOATS"), FGPUSkinCache::RWStrideInFloats);
-	}
-
-	FGPUSkinCacheCS(const ShaderMetaType::CompiledShaderInitializerType& Initializer)
+	FBaseGPUSkinCacheCS(const ShaderMetaType::CompiledShaderInitializerType& Initializer)
 		: FGlobalShader(Initializer)
 	{
-		SkinMeshOriginParameter.Bind(Initializer.ParameterMap,TEXT("MeshOrigin"));
-		SkinMeshExtensionParameter.Bind(Initializer.ParameterMap,TEXT("MeshExtension"));
+		SkinMeshOriginParameter.Bind(Initializer.ParameterMap, TEXT("MeshOrigin"));
+		SkinMeshExtensionParameter.Bind(Initializer.ParameterMap, TEXT("MeshExtension"));
 
 		//DebugParameter.Bind(Initializer.ParameterMap, TEXT("DebugParameter"));
 
@@ -113,10 +94,6 @@ public:
 		SkinCacheBuffer.Bind(Initializer.ParameterMap, TEXT("SkinCacheBuffer"));
 	}
 
-	FGPUSkinCacheCS()
-	{
-	}
-
 	void SetParameters(FRHICommandList& RHICmdList, uint32 InputStreamStride, uint32 InputStreamFloatOffset, uint32 InputStreamVertexCount, uint32 OutputBufferFloatOffset, const FVertexBufferAndSRV& BoneBuffer, FUniformBufferRHIRef UniformBuffer, FShaderResourceViewRHIRef VertexBufferSRV, FRWBuffer& SkinBuffer, const FVector& MeshOrigin, const FVector& MeshExtension)
 	{
 		FComputeShaderRHIParamRef ShaderRHI = GetComputeShader();
@@ -155,7 +132,7 @@ public:
 	}
 
 	virtual bool Serialize(FArchive& Ar) override
-	{		
+	{
 		bool bShaderHasOutdatedParameters = FGlobalShader::Serialize(Ar);
 		Ar  << SkinMeshOriginParameter << SkinMeshExtensionParameter << SkinInputStreamStride
 			<< SkinInputStreamVertexCount << SkinInputStreamFloatOffset << SkinOutputBufferFloatOffset
@@ -183,6 +160,39 @@ private:
 	FShaderResourceParameter SkinCacheBuffer;
 };
 
+/** Compute shader that skins a batch of vertices. */
+template <bool bUseExtraBoneInfluencesT>
+class TGPUSkinCacheCS : public FBaseGPUSkinCacheCS
+{
+	DECLARE_SHADER_TYPE(TGPUSkinCacheCS<bUseExtraBoneInfluencesT>,Global)
+public:
+
+	static bool ShouldCache(EShaderPlatform Platform)
+	{
+		return GEnableGPUSkinCacheShaders && IsFeatureLevelSupported(Platform, ERHIFeatureLevel::SM5);
+	}
+
+	static void ModifyCompilationEnvironment(EShaderPlatform Platform, FShaderCompilerEnvironment& OutEnvironment)
+	{
+		FGlobalShader::ModifyCompilationEnvironment(Platform,OutEnvironment);
+		const uint32 UseExtraBoneInfluences = bUseExtraBoneInfluencesT;
+		OutEnvironment.SetDefine(TEXT("GPUSKIN_USE_EXTRA_INFLUENCES"), UseExtraBoneInfluences);
+
+		OutEnvironment.SetDefine(TEXT("GPUSKIN_RWBUFFER_OFFSET_POSITION"), FGPUSkinCache::RWPositionOffsetInFloats);
+		OutEnvironment.SetDefine(TEXT("GPUSKIN_RWBUFFER_OFFSET_TANGENT_X"), FGPUSkinCache::RWTangentXOffsetInFloats);
+		OutEnvironment.SetDefine(TEXT("GPUSKIN_RWBUFFER_OFFSET_TANGENT_Z"), FGPUSkinCache::RWTangentZOffsetInFloats);
+		OutEnvironment.SetDefine(TEXT("GPUSKIN_RWBUFFER_NUM_FLOATS"), FGPUSkinCache::RWStrideInFloats);
+	}
+
+	TGPUSkinCacheCS(const ShaderMetaType::CompiledShaderInitializerType& Initializer)
+		: FBaseGPUSkinCacheCS(Initializer)
+	{
+	}
+
+	TGPUSkinCacheCS()
+	{
+	}
+};
 
 FGPUSkinCache::FGPUSkinCache()
 	: bInitialized(false)
@@ -245,11 +255,6 @@ void FGPUSkinCache::Cleanup()
 // Kick off compute shader skinning for incoming chunk, return key for fast lookup in draw passes
 int32 FGPUSkinCache::StartCacheMesh(FRHICommandListImmediate& RHICmdList, int32 Key, FGPUBaseSkinVertexFactory* VertexFactory, FGPUSkinPassthroughVertexFactory* TargetVertexFactory, const FSkelMeshChunk& BatchElement, const FSkeletalMeshObjectGPUSkin* Skin)
 {
-	if (!GEnableGPUSkinCache)
-	{
-		return -1;
-	}
-
 	if (CachedChunksThisFrameCount >= GMaxGPUSkinCacheElementsPerFrame && GFrameNumberRenderThread <= FrameCounter)
 	{
 		INC_DWORD_STAT(STAT_GPUSkinCache_SkippedForChunksPerFrame);
@@ -441,23 +446,21 @@ void FGPUSkinCache::TransitionToWriteable(FRHICommandList& RHICmdList)
 
 void FGPUSkinCache::TransitionAllToWriteable(FRHICommandList& RHICmdList)
 {
-	FUnorderedAccessViewRHIParamRef OutUAVs[GPUSKINCACHE_FRAMES];
-
-	for (int32 Index = 0; Index < GPUSKINCACHE_FRAMES; ++Index)
+	if (bInitialized)
 	{
-		OutUAVs[Index] = SkinCacheBuffer[Index].UAV;
-	}
+		FUnorderedAccessViewRHIParamRef OutUAVs[GPUSKINCACHE_FRAMES];
 
-	RHICmdList.TransitionResources(EResourceTransitionAccess::ERWNoBarrier, EResourceTransitionPipeline::EGfxToCompute, OutUAVs, ARRAY_COUNT(OutUAVs));
-}
+		for (int32 Index = 0; Index < GPUSKINCACHE_FRAMES; ++Index)
+		{
+			OutUAVs[Index] = SkinCacheBuffer[Index].UAV;
+		}
 
-bool FGPUSkinCache::IsElementProcessed(int32 Key) const
-{
-	if (!GEnableGPUSkinCache)
-	{
-		return false;
+		RHICmdList.TransitionResources(EResourceTransitionAccess::ERWNoBarrier, EResourceTransitionPipeline::EGfxToCompute, OutUAVs, ARRAY_COUNT(OutUAVs));
 	}
+}
 
+bool FGPUSkinCache::InternalIsElementProcessed(int32 Key) const
+{
 	if (Key >= 0 && Key < CachedElements.Num())
 	{
 		const FElementCacheStatusInfo& CacheInfo = CachedElements.GetData()[Key];
@@ -560,40 +563,33 @@ void FGPUSkinCache::DispatchSkinCacheProcess(FRHICommandListImmediate& RHICmdLis
 	SCOPED_DRAW_EVENT(RHICmdList, SkinCacheDispatch);
 
 	int32 UAVIndex = FrameCounter % GPUSKINCACHE_FRAMES;
-	if (bUseExtraBoneInfluences)
-	{
-		TShaderMapRef<FGPUSkinCacheCS<true> > SkinCacheCS(GetGlobalShaderMap(FeatureLevel));
-		RHICmdList.SetComputeShader(SkinCacheCS->GetComputeShader());
-		SkinCacheCS->SetParameters(RHICmdList, VertexStride, InputStreamFloatOffset, VertexCount, OutputBufferFloatOffset, BoneBuffer, UniformBuffer, VBInfo->VertexBufferSRV, SkinCacheBuffer[UAVIndex], MeshOrigin, MeshExtension);
 
-		RHICmdList.DispatchComputeShader(VertexCountAlign64, 1, 1);
-		SkinCacheCS->UnsetParameters(RHICmdList);
-	}
-	else
-	{
-		TShaderMapRef<FGPUSkinCacheCS<false> > SkinCacheCS(GetGlobalShaderMap(FeatureLevel));
-		RHICmdList.SetComputeShader(SkinCacheCS->GetComputeShader());
-		SkinCacheCS->SetParameters(RHICmdList, VertexStride, InputStreamFloatOffset, VertexCount, OutputBufferFloatOffset, BoneBuffer, UniformBuffer, VBInfo->VertexBufferSRV, SkinCacheBuffer[UAVIndex], MeshOrigin, MeshExtension);
+	TShaderMapRef<TGPUSkinCacheCS<true> > SkinCacheCS1(GetGlobalShaderMap(FeatureLevel));
+	TShaderMapRef<TGPUSkinCacheCS<false> > SkinCacheCS0(GetGlobalShaderMap(FeatureLevel));
+	FBaseGPUSkinCacheCS* Shader = bUseExtraBoneInfluences ? (FBaseGPUSkinCacheCS*)*SkinCacheCS1 : (FBaseGPUSkinCacheCS*)*SkinCacheCS0;
 
-		RHICmdList.DispatchComputeShader(VertexCountAlign64, 1, 1);
-		SkinCacheCS->UnsetParameters(RHICmdList);
-	}
+	RHICmdList.SetComputeShader(Shader->GetComputeShader());
+	Shader->SetParameters(RHICmdList, VertexStride, InputStreamFloatOffset, VertexCount, OutputBufferFloatOffset, BoneBuffer, UniformBuffer, VBInfo->VertexBufferSRV, SkinCacheBuffer[UAVIndex], MeshOrigin, MeshExtension);
+
+	RHICmdList.DispatchComputeShader(VertexCountAlign64, 1, 1);
+	Shader->UnsetParameters(RHICmdList);
 }
 
 void FGPUSkinCache::CVarSinkFunction()
 {
-	int32 NewValue = CVarEnableGPUSkinCache.GetValueOnAnyThread();
+	int32 NewGPUSkinCacheValue = CVarEnableGPUSkinCache.GetValueOnAnyThread();
 	if (!GEnableGPUSkinCacheShaders)
 	{
-		NewValue = 0;
+		NewGPUSkinCacheValue = 0;
 	}
 
-	if (NewValue != GEnableGPUSkinCache)
+	if (NewGPUSkinCacheValue != GEnableGPUSkinCache)
 	{
-		ENQUEUE_UNIQUE_RENDER_COMMAND_ONEPARAMETER(DoEnableSkinCaching, int32, Value, NewValue,
+		ENQUEUE_UNIQUE_RENDER_COMMAND_ONEPARAMETER(DoEnableSkinCaching, 
+			int32, SkinValue, NewGPUSkinCacheValue,
 		{
-			GEnableGPUSkinCache = Value;
-			if (Value)
+			GEnableGPUSkinCache = SkinValue;
+			if (SkinValue)
 			{
 				GGPUSkinCache.TransitionAllToWriteable(RHICmdList);
 			}
@@ -603,5 +599,5 @@ void FGPUSkinCache::CVarSinkFunction()
 
 FAutoConsoleVariableSink FGPUSkinCache::CVarSink(FConsoleCommandDelegate::CreateStatic(&CVarSinkFunction));
 
-IMPLEMENT_SHADER_TYPE(template<>,FGPUSkinCacheCS<false>,TEXT("GpuSkinCacheComputeShader"),TEXT("SkinCacheUpdateBatchCS"),SF_Compute);
-IMPLEMENT_SHADER_TYPE(template<>,FGPUSkinCacheCS<true>,TEXT("GpuSkinCacheComputeShader"),TEXT("SkinCacheUpdateBatchCS"),SF_Compute);
+IMPLEMENT_SHADER_TYPE(template<>,TGPUSkinCacheCS<false>,TEXT("GpuSkinCacheComputeShader"),TEXT("SkinCacheUpdateBatchCS"),SF_Compute);
+IMPLEMENT_SHADER_TYPE(template<>,TGPUSkinCacheCS<true>,TEXT("GpuSkinCacheComputeShader"),TEXT("SkinCacheUpdateBatchCS"),SF_Compute);
diff --git a/Engine/Source/Runtime/Engine/Private/GPUSkinVertexFactory.cpp b/Engine/Source/Runtime/Engine/Private/GPUSkinVertexFactory.cpp
--- a/Engine/Source/Runtime/Engine/Private/GPUSkinVertexFactory.cpp
+++ b/Engine/Source/Runtime/Engine/Private/GPUSkinVertexFactory.cpp
@@ -141,8 +141,8 @@ TStatId FClothBufferPool::GetStatId() const
 	RETURN_QUICK_DECLARE_CYCLE_STAT(FClothBufferPool, STATGROUP_Tickables);
 }
 
-TConsoleVariableData<int32>* FGPUBaseSkinVertexFactory::ShaderDataType::MaxBonesVar = NULL;
-uint32 FGPUBaseSkinVertexFactory::ShaderDataType::MaxGPUSkinBones = 0;
+TConsoleVariableData<int32>* FGPUBaseSkinVertexFactory::FShaderDataType::MaxBonesVar = NULL;
+uint32 FGPUBaseSkinVertexFactory::FShaderDataType::MaxGPUSkinBones = 0;
 
 static TAutoConsoleVariable<int32> CVarRHICmdDeferSkeletalLockAndFillToRHIThread(
 	TEXT("r.RHICmdDeferSkeletalLockAndFillToRHIThread"),
@@ -195,14 +195,14 @@ struct FRHICommandUpdateBoneBuffer : public FRHICommand<FRHICommandUpdateBoneBuf
 	}
 };
 
-void FGPUBaseSkinVertexFactory::ShaderDataType::GoToNextFrame(uint32 FrameNumber)
+void FGPUBaseSkinVertexFactory::FShaderDataType::GoToNextFrame(uint32 FrameNumber)
 {
 	PreviousFrameNumber = CurrentFrameNumber;
 	CurrentFrameNumber = FrameNumber;
 	CurrentBuffer = 1 - CurrentBuffer;
 }
 
-bool FGPUBaseSkinVertexFactory::ShaderDataType::UpdateBoneData(FRHICommandListImmediate& RHICmdList, const TArray<FMatrix>& ReferenceToLocalMatrices,
+bool FGPUBaseSkinVertexFactory::FShaderDataType::UpdateBoneData(FRHICommandListImmediate& RHICmdList, const TArray<FMatrix>& ReferenceToLocalMatrices,
 	const TArray<FBoneIndexType>& BoneMap, uint32 FrameNumber, ERHIFeatureLevel::Type FeatureLevel, bool bUseSkinCache)
 {
 	const uint32 NumBones = BoneMap.Num();
@@ -356,7 +356,7 @@ void TGPUSkinVertexFactory<bExtraBoneInfluencesT>::ModifyCompilationEnvironment(
 template<bool bExtraBoneInfluencesT>
 void TGPUSkinVertexFactory<bExtraBoneInfluencesT>::CopyDataTypeForPassthroughFactory(FGPUSkinPassthroughVertexFactory* PassthroughVertexFactory)
 {
-	FGPUSkinPassthroughVertexFactory::DataType DestDataType;
+	FGPUSkinPassthroughVertexFactory::FDataType DestDataType;
 	DestDataType.PositionComponent = Data.PositionComponent;
 	DestDataType.TangentBasisComponents[0] = Data.TangentBasisComponents[0];
 	DestDataType.TangentBasisComponents[1] = Data.TangentBasisComponents[1];
@@ -371,7 +371,7 @@ void TGPUSkinVertexFactory<bExtraBoneInfluencesT>::CopyDataTypeForPassthroughFac
 * @param OutElements - vertex decl list to modify
 */
 template <bool bExtraBoneInfluencesT>
-void TGPUSkinVertexFactory<bExtraBoneInfluencesT>::AddVertexElements(DataType& InData, FVertexDeclarationElementList& OutElements)
+void TGPUSkinVertexFactory<bExtraBoneInfluencesT>::AddVertexElements(FDataType& InData, FVertexDeclarationElementList& OutElements)
 {
 	// position decls
 	OutElements.Add(AccessStreamComponent(InData.PositionComponent,0));
@@ -440,7 +440,7 @@ void TGPUSkinVertexFactory<bExtraBoneInfluencesT>::InitRHI()
 	AddVertexElements(Data,Elements);	
 
 	// create the actual device decls
-	InitDeclaration(Elements,FVertexFactory::DataType());
+	InitDeclaration(Elements);
 }
 
 template <bool bExtraBoneInfluencesT>
@@ -510,7 +510,7 @@ public:
 
 		if(ShaderRHI)
 		{
-			const FGPUBaseSkinVertexFactory::ShaderDataType& ShaderData = ((const FGPUBaseSkinVertexFactory*)VertexFactory)->GetShaderData();
+			const FGPUBaseSkinVertexFactory::FShaderDataType& ShaderData = ((const FGPUBaseSkinVertexFactory*)VertexFactory)->GetShaderData();
 
 			SetShaderValue(RHICmdList, ShaderRHI, MeshOriginParameter, ShaderData.MeshOrigin);
 			SetShaderValue(RHICmdList, ShaderRHI, MeshExtensionParameter, ShaderData.MeshExtension);
@@ -705,7 +705,7 @@ bool TGPUSkinMorphVertexFactory<bExtraBoneInfluencesT>::ShouldCache(EShaderPlatf
 * @param OutElements - vertex decl list to modify
 */
 template <bool bExtraBoneInfluencesT>
-void TGPUSkinMorphVertexFactory<bExtraBoneInfluencesT>::AddVertexElements(DataType& InData, FVertexDeclarationElementList& OutElements)
+void TGPUSkinMorphVertexFactory<bExtraBoneInfluencesT>::AddVertexElements(FDataType& InData, FVertexDeclarationElementList& OutElements)
 {
 	// add the base gpu skin elements
 	TGPUSkinVertexFactory<bExtraBoneInfluencesT>::AddVertexElements(InData,OutElements);
@@ -726,7 +726,7 @@ void TGPUSkinMorphVertexFactory<bExtraBoneInfluencesT>::InitRHI()
 	AddVertexElements(MorphData,Elements);
 
 	// create the actual device decls
-	FVertexFactory::InitDeclaration(Elements,FVertexFactory::DataType());
+	FVertexFactory::InitDeclaration(Elements);
 }
 
 template <bool bExtraBoneInfluencesT>
@@ -972,7 +972,7 @@ bool TGPUSkinAPEXClothVertexFactory<bExtraBoneInfluencesT>::ShouldCache(EShaderP
 * @param OutElements - vertex decl list to modify
 */
 template <bool bExtraBoneInfluencesT>
-void TGPUSkinAPEXClothVertexFactory<bExtraBoneInfluencesT>::AddVertexElements(DataType& InData, FVertexDeclarationElementList& OutElements)
+void TGPUSkinAPEXClothVertexFactory<bExtraBoneInfluencesT>::AddVertexElements(FDataType& InData, FVertexDeclarationElementList& OutElements)
 {
 	// add the base gpu skin elements
 	TGPUSkinVertexFactory<bExtraBoneInfluencesT>::AddVertexElements(InData,OutElements);
@@ -999,7 +999,7 @@ void TGPUSkinAPEXClothVertexFactory<bExtraBoneInfluencesT>::InitRHI()
 	AddVertexElements(MeshMappingData,Elements);
 
 	// create the actual device decls
-	FVertexFactory::InitDeclaration(Elements,FVertexFactory::DataType());
+	FVertexFactory::InitDeclaration(Elements);
 }
 
 template <bool bExtraBoneInfluencesT>
diff --git a/Engine/Source/Runtime/Engine/Private/HierarchicalInstancedStaticMesh.cpp b/Engine/Source/Runtime/Engine/Private/HierarchicalInstancedStaticMesh.cpp
--- a/Engine/Source/Runtime/Engine/Private/HierarchicalInstancedStaticMesh.cpp
+++ b/Engine/Source/Runtime/Engine/Private/HierarchicalInstancedStaticMesh.cpp
@@ -2349,7 +2349,7 @@ void UHierarchicalInstancedStaticMeshComponent::BuildTreeAsync()
 		bIsAsyncBuilding = true;
 
 		FGraphEventRef BuildTreeAsyncResult(
-			FDelegateGraphTask::CreateAndDispatchWhenReady(FDelegateGraphTask::FDelegate::CreateRaw(&Builder.Get(), &FClusterBuilder::BuildAsync), GET_STATID(STAT_FoliageBuildTime), NULL, ENamedThreads::GameThread, ENamedThreads::AnyThread));
+			FDelegateGraphTask::CreateAndDispatchWhenReady(FDelegateGraphTask::FDelegate::CreateRaw(&Builder.Get(), &FClusterBuilder::BuildAsync), GET_STATID(STAT_FoliageBuildTime), NULL, ENamedThreads::GameThread, ENamedThreads::AnyBackgroundThreadNormalTask));
 
 		// add a dependent task to run on the main thread when build is complete
 		FGraphEventRef UnusedAsyncResult(
diff --git a/Engine/Source/Runtime/Engine/Private/InstancedStaticMesh.cpp b/Engine/Source/Runtime/Engine/Private/InstancedStaticMesh.cpp
--- a/Engine/Source/Runtime/Engine/Private/InstancedStaticMesh.cpp
+++ b/Engine/Source/Runtime/Engine/Private/InstancedStaticMesh.cpp
@@ -213,7 +213,7 @@ void FInstancedStaticMeshVertexFactory::Copy(const FInstancedStaticMeshVertexFac
 	ENQUEUE_UNIQUE_RENDER_COMMAND_TWOPARAMETER(
 		FInstancedStaticMeshVertexFactoryCopyData,
 		FInstancedStaticMeshVertexFactory*,VertexFactory,this,
-		const DataType*,DataCopy,&Other.Data,
+		const FDataType*,DataCopy,&Other.Data,
 	{
 		VertexFactory->Data = *DataCopy;
 	});
@@ -313,7 +313,7 @@ void FInstancedStaticMeshVertexFactory::InitRHI()
 	}
 
 	// we don't need per-vertex shadow or lightmap rendering
-	InitDeclaration(Elements,Data);
+	InitDeclaration(Elements);
 }
 
 
@@ -336,7 +336,7 @@ void FInstancedStaticMeshRenderData::InitStaticMeshVertexFactories(
 	{
 		const FStaticMeshLODResources* RenderData = &InstancedRenderData->LODModels[LODIndex];
 						
-		FInstancedStaticMeshVertexFactory::DataType Data;
+		FInstancedStaticMeshVertexFactory::FDataType Data;
 		Data.PositionComponent = FVertexStreamComponent(
 			&RenderData->PositionVertexBuffer,
 			STRUCT_OFFSET(FPositionVertex,Position),
diff --git a/Engine/Source/Runtime/Engine/Private/InstancedStaticMesh.h b/Engine/Source/Runtime/Engine/Private/InstancedStaticMesh.h
--- a/Engine/Source/Runtime/Engine/Private/InstancedStaticMesh.h
+++ b/Engine/Source/Runtime/Engine/Private/InstancedStaticMesh.h
@@ -178,7 +178,7 @@ struct FInstancedStaticMeshVertexFactory : public FLocalVertexFactory
 {
 	DECLARE_VERTEX_FACTORY_TYPE(FInstancedStaticMeshVertexFactory);
 public:
-	struct DataType : public FLocalVertexFactory::DataType
+	struct FDataType : public FLocalVertexFactory::FDataType
 	{
 		/** The stream to read the mesh transform from. */
 		FVertexStreamComponent InstanceOriginComponent;
@@ -209,7 +209,7 @@ public:
 	/**
 	 * An implementation of the interface used by TSynchronizedResource to update the resource with new data from the game thread.
 	 */
-	void SetData(const DataType& InData)
+	void SetData(const FDataType& InData)
 	{
 		Data = InData;
 		UpdateRHI();
@@ -245,7 +245,7 @@ public:
 #endif
 
 private:
-	DataType Data;
+	FDataType Data;
 };
 
 
diff --git a/Engine/Source/Runtime/Engine/Private/LevelTick.cpp b/Engine/Source/Runtime/Engine/Private/LevelTick.cpp
--- a/Engine/Source/Runtime/Engine/Private/LevelTick.cpp
+++ b/Engine/Source/Runtime/Engine/Private/LevelTick.cpp
@@ -708,7 +708,7 @@ static TAutoConsoleVariable<int32> CVarAllowAsyncRenderThreadUpdates(
 
 static TAutoConsoleVariable<int32> CVarAllowAsyncRenderThreadUpdatesDuringGamethreadUpdates(
 	TEXT("AllowAsyncRenderThreadUpdatesDuringGamethreadUpdates"),
-	0,
+	1,
 	TEXT("If > 0 then we do the gamethread updates _while_ doing parallel updates."));
 
 static TAutoConsoleVariable<int32> CVarAllowAsyncRenderThreadUpdatesEditor(
diff --git a/Engine/Source/Runtime/Engine/Private/LocalVertexFactory.cpp b/Engine/Source/Runtime/Engine/Private/LocalVertexFactory.cpp
--- a/Engine/Source/Runtime/Engine/Private/LocalVertexFactory.cpp
+++ b/Engine/Source/Runtime/Engine/Private/LocalVertexFactory.cpp
@@ -53,7 +53,7 @@ bool FLocalVertexFactory::ShouldCache(EShaderPlatform Platform, const class FMat
 	return true; 
 }
 
-void FLocalVertexFactory::SetData(const DataType& InData)
+void FLocalVertexFactory::SetData(const FDataType& InData)
 {
 	check(IsInRenderingThread());
 
@@ -74,7 +74,7 @@ void FLocalVertexFactory::Copy(const FLocalVertexFactory& Other)
 	ENQUEUE_UNIQUE_RENDER_COMMAND_TWOPARAMETER(
 		FLocalVertexFactoryCopyData,
 		FLocalVertexFactory*,VertexFactory,this,
-		const DataType*,DataCopy,&Other.Data,
+		const FDataType*,DataCopy,&Other.Data,
 	{
 		VertexFactory->Data = *DataCopy;
 	});
@@ -152,7 +152,7 @@ void FLocalVertexFactory::InitRHI()
 
 	check(Streams.Num() > 0);
 
-	InitDeclaration(Elements,Data);
+	InitDeclaration(Elements);
 
 	check(IsValidRef(GetDeclaration()));
 }
diff --git a/Engine/Source/Runtime/Engine/Private/Materials/Material.cpp b/Engine/Source/Runtime/Engine/Private/Materials/Material.cpp
--- a/Engine/Source/Runtime/Engine/Private/Materials/Material.cpp
+++ b/Engine/Source/Runtime/Engine/Private/Materials/Material.cpp
@@ -2728,7 +2728,7 @@ void UMaterial::PostEditChangeProperty(FPropertyChangedEvent& PropertyChangedEve
 		bUsesDistortion = false;
 		// check for a distortion value
 		if (Refraction.Expression
-			|| (Refraction.UseConstant && FMath::Abs(Refraction.Constant) >= KINDA_SMALL_NUMBER))
+			|| (Refraction.UseConstant && FMath::Abs(Refraction.Constant - 1.0f) >= KINDA_SMALL_NUMBER))
 		{
 			bUsesDistortion = true;
 		}
diff --git a/Engine/Source/Runtime/Engine/Private/Materials/MaterialInterface.cpp b/Engine/Source/Runtime/Engine/Private/Materials/MaterialInterface.cpp
--- a/Engine/Source/Runtime/Engine/Private/Materials/MaterialInterface.cpp
+++ b/Engine/Source/Runtime/Engine/Private/Materials/MaterialInterface.cpp
@@ -23,7 +23,7 @@ void FMaterialRelevance::SetPrimitiveViewRelevance(FPrimitiveViewRelevance& OutV
 	OutViewRelevance.bNormalTranslucencyRelevance = bNormalTranslucency;
 
 	{
-		IConsoleVariable* CVar = IConsoleManager::Get().FindConsoleVariable(TEXT("r.ShaderModelDebug"));
+		static IConsoleVariable* CVar = IConsoleManager::Get().FindConsoleVariable(TEXT("r.ShaderModelDebug"));
 		if(CVar->GetInt())
 		{
 			UE_LOG(LogEngine, Log, TEXT("r.ShaderModelDebug: SetPrimitiveViewRelevance %x = %x"),
diff --git a/Engine/Source/Runtime/Engine/Private/MeshParticleVertexFactory.cpp b/Engine/Source/Runtime/Engine/Private/MeshParticleVertexFactory.cpp
--- a/Engine/Source/Runtime/Engine/Private/MeshParticleVertexFactory.cpp
+++ b/Engine/Source/Runtime/Engine/Private/MeshParticleVertexFactory.cpp
@@ -177,7 +177,7 @@ void FMeshParticleVertexFactory::InitRHI()
 
 		if(Streams.Num() > 0)
 		{
-			InitDeclaration(Elements,Data);
+			InitDeclaration(Elements);
 			check(IsValidRef(GetDeclaration()));
 		}
 	}
@@ -211,7 +211,7 @@ bool FMeshParticleVertexFactory::ShouldCache(EShaderPlatform Platform, const cla
 	return (Material->IsUsedWithMeshParticles() || Material->IsSpecialEngineMaterial());
 }
 
-void FMeshParticleVertexFactory::SetData(const DataType& InData)
+void FMeshParticleVertexFactory::SetData(const FDataType& InData)
 {
 	check(IsInRenderingThread());
 	Data = InData;
diff --git a/Engine/Source/Runtime/Engine/Private/NiagaraEffectRenderer.cpp b/Engine/Source/Runtime/Engine/Private/NiagaraEffectRenderer.cpp
--- a/Engine/Source/Runtime/Engine/Private/NiagaraEffectRenderer.cpp
+++ b/Engine/Source/Runtime/Engine/Private/NiagaraEffectRenderer.cpp
@@ -647,7 +647,7 @@ DynamicDataRender(NULL)
 
 void NiagaraEffectRendererMeshes::SetupVertexFactory(FMeshParticleVertexFactory *InVertexFactory, const FStaticMeshLODResources& LODResources) const
 {
-	FMeshParticleVertexFactory::DataType Data;
+	FMeshParticleVertexFactory::FDataType Data;
 
 	Data.PositionComponent = FVertexStreamComponent(
 		&LODResources.PositionVertexBuffer,
diff --git a/Engine/Source/Runtime/Engine/Private/Particles/FXSystem.cpp b/Engine/Source/Runtime/Engine/Private/Particles/FXSystem.cpp
--- a/Engine/Source/Runtime/Engine/Private/Particles/FXSystem.cpp
+++ b/Engine/Source/Runtime/Engine/Private/Particles/FXSystem.cpp
@@ -50,7 +50,7 @@ namespace FXConsoleVariables
 	int32 bAllowCulling = true;
 	int32 bFreezeGPUSimulation = false;
 	int32 bFreezeParticleSimulation = false;
-	int32 bAllowAsyncTick = false;
+	int32 bAllowAsyncTick = !WITH_EDITOR;
 	float ParticleSlackGPU = 0.02f;
 	int32 MaxParticleTilePreAllocation = 100;
 	int32 MaxCPUParticlesPerEmitter = 1000;
diff --git a/Engine/Source/Runtime/Engine/Private/Particles/ParticleComponents.cpp b/Engine/Source/Runtime/Engine/Private/Particles/ParticleComponents.cpp
--- a/Engine/Source/Runtime/Engine/Private/Particles/ParticleComponents.cpp
+++ b/Engine/Source/Runtime/Engine/Private/Particles/ParticleComponents.cpp
@@ -1536,8 +1536,8 @@ void UParticleEmitter::Build()
 		if (HighLODLevel->TypeDataModule != nullptr)
 		{
 			if(HighLODLevel->TypeDataModule->RequiresBuild())
-			{
-				FParticleEmitterBuildInfo EmitterBuildInfo;
+		{
+			FParticleEmitterBuildInfo EmitterBuildInfo;
 				HighLODLevel->CompileModules( EmitterBuildInfo );
 				HighLODLevel->TypeDataModule->Build( EmitterBuildInfo );
 			}
@@ -3655,7 +3655,8 @@ FParticleDynamicData* UParticleSystemComponent::CreateDynamicData()
 				// Fill the emitter dynamic buffers with data from our replay
 				ParticleDynamicData->DynamicEmitterDataArray.Reset();
 				ParticleDynamicData->DynamicEmitterDataArray.Reserve(CurReplayFrame.Emitters.Num());
-				for( int32 CurEmitterIndex = 0; CurEmitterIndex < CurReplayFrame.Emitters.Num(); ++CurEmitterIndex )
+
+				for (int32 CurEmitterIndex = 0; CurEmitterIndex < CurReplayFrame.Emitters.Num(); ++CurEmitterIndex)
 				{
 					const FParticleEmitterReplayFrame& CurEmitter = CurReplayFrame.Emitters[ CurEmitterIndex ];
 
@@ -3670,9 +3671,11 @@ FParticleDynamicData* UParticleSystemComponent::CreateDynamicData()
 						FDynamicEmitterDataBase* NewDynamicEmitterData =
 							CreateDynamicDataFromReplay( EmitterInstances[ CurEmitter.OriginalEmitterIndex ], CurEmitterReplay, IsOwnerSelected() );
 
+
 						if( NewDynamicEmitterData != NULL )
 						{
-							ParticleDynamicData->DynamicEmitterDataArray.Add( NewDynamicEmitterData );
+							ParticleDynamicData->DynamicEmitterDataArray.Add(NewDynamicEmitterData);
+							NewDynamicEmitterData->EmitterIndex = CurEmitterIndex;
 						}
 					}
 				}
@@ -3751,6 +3754,8 @@ FParticleDynamicData* UParticleSystemComponent::CreateDynamicData()
 #endif
 						NewDynamicEmitterData->bValid = true;
 						ParticleDynamicData->DynamicEmitterDataArray.Add( NewDynamicEmitterData );
+						NewDynamicEmitterData->EmitterIndex = EmitterIndex;
+
 						// Are we current capturing particle state?
 						if( ReplayState == PRS_Capturing )
 						{
@@ -3848,6 +3853,13 @@ void UParticleSystemComponent::UpdateDynamicData(FParticleSystemSceneProxy* Prox
 	{
 		// Create the dynamic data for rendering this particle system
 		FParticleDynamicData* ParticleDynamicData = CreateDynamicData();
+		if (ParticleDynamicData)
+		{
+			for (int32 Index = 0; Index < ParticleDynamicData->DynamicEmitterDataArray.Num(); Index++)
+			{
+				static_cast<FParticleSystemSceneProxy*>(SceneProxy)->QueueVertexFactoryCreation(ParticleDynamicData->DynamicEmitterDataArray[Index]);
+			}
+		}
 
 		// Render the particles
 #if !(UE_BUILD_SHIPPING || UE_BUILD_TEST)
@@ -4083,6 +4095,15 @@ public:
 	}
 };
 
+FAutoConsoleTaskPriority CPrio_ParticleAsyncTask(
+	TEXT("TaskGraph.TaskPriorities.ParticleAsyncTask"),
+	TEXT("Task and thread priority for FParticleAsyncTask."),
+	ENamedThreads::HighThreadPriority, // if we have high priority task threads, then use them...
+	ENamedThreads::NormalTaskPriority, // .. at normal task priority
+	ENamedThreads::HighTaskPriority // if we don't have hi pri threads, then use normal priority threads at high task priority instead
+	);
+
+
 class FParticleAsyncTask
 {
 	UParticleSystemComponent* Target;
@@ -4100,7 +4121,7 @@ public:
 
 	ENamedThreads::Type GetDesiredThread()
 	{
-		return ENamedThreads::HiPri(ENamedThreads::AnyThread);
+		return CPrio_ParticleAsyncTask.Get();
 	}
 
 	static ESubsequentsMode::Type GetSubsequentsMode() { return ESubsequentsMode::TrackSubsequents; }
@@ -4479,22 +4500,22 @@ void UParticleSystemComponent::FinalizeTickComponent()
 			FParticleEmitterInstance* Instance = EmitterInstances[EmitterIndex];
 			if (Instance && Instance->bEnabled)
 			{
-			if (EmitterIndex + 1 < EmitterInstances.Num())
-			{
+				if (EmitterIndex + 1 < EmitterInstances.Num())
+				{
 				FParticleEmitterInstance* NextInstance = EmitterInstances[EmitterIndex+1];
-				FPlatformMisc::Prefetch(NextInstance);
-			}
+					FPlatformMisc::Prefetch(NextInstance);
+				}
 
 				if (Instance->SpriteTemplate)
-			{
-				UParticleLODLevel* SpriteLODLevel = Instance->SpriteTemplate->GetCurrentLODLevel(Instance);
-				if (SpriteLODLevel && SpriteLODLevel->bEnabled)
 				{
-					Instance->ProcessParticleEvents(DeltaTimeTick, bSuppressSpawning);
+					UParticleLODLevel* SpriteLODLevel = Instance->SpriteTemplate->GetCurrentLODLevel(Instance);
+					if (SpriteLODLevel && SpriteLODLevel->bEnabled)
+					{
+						Instance->ProcessParticleEvents(DeltaTimeTick, bSuppressSpawning);
+					}
 				}
 			}
 		}
-		}
 
 		AParticleEventManager* EventManager = (World ? World->MyParticleEventManager : NULL);
 		if (EventManager)
@@ -4673,13 +4694,13 @@ void UParticleSystemComponent::InitParticles(bool bReInitExistingEmitters)
 			const bool bIsSignificant = (int32)Emitter->SignificanceLevel >= (int32)RequiredSignificance;
 			bool bFailedOnlyOnSignificance = false;
 			if (bShouldCreateAndOrInit && !bIsSignificant)
-				{
+			{
 				bFailedOnlyOnSignificance = true;
 				bShouldCreateAndOrInit = false;
-					}
+			}
 
 			if (bShouldCreateAndOrInit)
-					{
+			{
 				if (Instance)
 				{
 					Instance->SetHaltSpawning(false);
@@ -4688,10 +4709,10 @@ void UParticleSystemComponent::InitParticles(bool bReInitExistingEmitters)
 				{
 					Instance = Emitter->CreateInstance(this);
 					EmitterInstances[Idx] = Instance;
-			}
+				}
 
 				if (Instance)
-			{
+				{
 					Instance->bEnabled = true;
 					if (bReInitExistingEmitters)
 					{
@@ -4701,10 +4722,10 @@ void UParticleSystemComponent::InitParticles(bool bReInitExistingEmitters)
 						PreferredLODLevel = FMath::Min(PreferredLODLevel, Emitter->LODLevels.Num());
 						bSetLodLevels |= !bIsFirstCreate;//Only set lod levels if we init any instances and it's not the first creation time.
 					}
-						}
-					}
-					else
-					{
+				}
+			}
+			else
+			{
 				if (Instance)
 				{
 					if(!bFailedOnlyOnSignificance)
@@ -4717,20 +4738,20 @@ void UParticleSystemComponent::InitParticles(bool bReInitExistingEmitters)
 						EmitterInstances[Idx] = NULL;
 						bClearDynamicData = true;
 					}
-				else
-				{
+					else
+					{
 						//otherwise, we just disable it.
 						Instance->bEnabled = false;
 					}
-					}
-				bMissingEmittersDueToSignificance |= bFailedOnlyOnSignificance;
 				}
+				bMissingEmittersDueToSignificance |= bFailedOnlyOnSignificance;
 			}
+		}
 
-			if (bClearDynamicData)
-			{
-				ClearDynamicData();
-			}
+		if (bClearDynamicData)
+		{
+			ClearDynamicData();
+		}
 
 		if (bSetLodLevels)
 		{
@@ -4760,6 +4781,7 @@ void UParticleSystemComponent::InitParticles(bool bReInitExistingEmitters)
 			}
 		}
 	}
+
 }
 
 void UParticleSystemComponent::ResetParticles(bool bEmptyInstances)
@@ -4913,7 +4935,10 @@ void UParticleSystemComponent::SetTemplate(class UParticleSystem* NewTemplate)
 			Instance->CurrentLODLevelIndex = 0;
 		}
 	}
-
+	if (SceneProxy)
+	{
+		static_cast<FParticleSystemSceneProxy*>(SceneProxy)->MarkVertexFactoriesDirty();
+	}
 }
 
 void UParticleSystemComponent::ActivateSystem(bool bFlagAsJustAttached)
@@ -5385,6 +5410,11 @@ bool UParticleSystemComponent::HasCompleted()
 
 		if (Instance && Instance->CurrentLODLevel)
 		{
+			if (!Instance->bEmitterIsDone)
+			{
+				bCanBeDeactivated = false;
+			}
+
 			if (Instance->CurrentLODLevel->bEnabled)
 			{
 				if (Instance->CurrentLODLevel->RequiredModule->EmitterLoops > 0 || Instance->IsTrailEmitter())
@@ -5438,12 +5468,6 @@ bool UParticleSystemComponent::HasCompleted()
 						bHasCompleted = false;
 					}
 				}
-
-				if (!Instance->bEmitterIsDone)
-				{
-					bCanBeDeactivated = false;
-				}
-
 			}
 			else
 			{
@@ -5452,15 +5476,11 @@ bool UParticleSystemComponent::HasCompleted()
 				{
 					bHasCompleted = false;
 				}
-				else if (!Instance->bEmitterIsDone)
-				{
-					bCanBeDeactivated = false;
-				}
 			}
 		}
 	}
 
-	if (bCanBeDeactivated && bAutoActivate)
+	if (bCanBeDeactivated && Template && Template->bAutoDeactivate)
 	{
 		DeactivateSystem();
 	}
diff --git a/Engine/Source/Runtime/Engine/Private/Particles/ParticleEmitterInstances.cpp b/Engine/Source/Runtime/Engine/Private/Particles/ParticleEmitterInstances.cpp
--- a/Engine/Source/Runtime/Engine/Private/Particles/ParticleEmitterInstances.cpp
+++ b/Engine/Source/Runtime/Engine/Private/Particles/ParticleEmitterInstances.cpp
@@ -148,13 +148,13 @@ DECLARE_CYCLE_STAT(TEXT("EmitterInstance Resize"), STAT_ParticleEmitterInstance_
 	Fast allocators for small block data being sent to the render thread
 -----------------------------------------------------------------------------*/
 
-static TLockFreeFixedSizeAllocator<256> FastParticleSmallBlockAlloc256;
-static TLockFreeFixedSizeAllocator<384> FastParticleSmallBlockAlloc384;
-static TLockFreeFixedSizeAllocator<512> FastParticleSmallBlockAlloc512;
-static TLockFreeFixedSizeAllocator<768> FastParticleSmallBlockAlloc768;
-static TLockFreeFixedSizeAllocator<1024> FastParticleSmallBlockAlloc1024;
-static TLockFreeFixedSizeAllocator<1792> FastParticleSmallBlockAlloc1792;
-static TLockFreeFixedSizeAllocator<2048> FastParticleSmallBlockAlloc2048;
+static TLockFreeFixedSizeAllocator<256, PLATFORM_CACHE_LINE_SIZE> FastParticleSmallBlockAlloc256;
+static TLockFreeFixedSizeAllocator<384, PLATFORM_CACHE_LINE_SIZE> FastParticleSmallBlockAlloc384;
+static TLockFreeFixedSizeAllocator<512, PLATFORM_CACHE_LINE_SIZE> FastParticleSmallBlockAlloc512;
+static TLockFreeFixedSizeAllocator<768, PLATFORM_CACHE_LINE_SIZE> FastParticleSmallBlockAlloc768;
+static TLockFreeFixedSizeAllocator<1024, PLATFORM_CACHE_LINE_SIZE> FastParticleSmallBlockAlloc1024;
+static TLockFreeFixedSizeAllocator<1792, PLATFORM_CACHE_LINE_SIZE> FastParticleSmallBlockAlloc1792;
+static TLockFreeFixedSizeAllocator<2048, PLATFORM_CACHE_LINE_SIZE> FastParticleSmallBlockAlloc2048;
 
 FORCEINLINE static void* FastParticleSmallBlockAlloc(size_t AllocSize)
 {
@@ -3657,7 +3657,7 @@ void FDynamicEmitterDataBase::operator delete(void *RawMemory, size_t AllocSize)
 	FastParticleSmallBlockFree(RawMemory, AllocSize);
 }	
 
-static TLockFreeFixedSizeAllocator<sizeof(FParticleDynamicData)> ParticleDynamicDataAllocator;
+static TLockFreeFixedSizeAllocator<sizeof(FParticleDynamicData), PLATFORM_CACHE_LINE_SIZE> ParticleDynamicDataAllocator;
 
 void* FParticleDynamicData::operator new(size_t AllocSize)
 {
@@ -3676,6 +3676,7 @@ void FParticleDynamicData::operator delete(void *RawMemory, size_t AllocSize)
 FDynamicEmitterDataBase::FDynamicEmitterDataBase(const UParticleModuleRequired* RequiredModule)
 	: bSelected(false)
 	, ParticleVertexFactory(nullptr)
+	, EmitterIndex(INDEX_NONE)
 {
 }
 
diff --git a/Engine/Source/Runtime/Engine/Private/Particles/ParticleGpuSimulation.cpp b/Engine/Source/Runtime/Engine/Private/Particles/ParticleGpuSimulation.cpp
--- a/Engine/Source/Runtime/Engine/Private/Particles/ParticleGpuSimulation.cpp
+++ b/Engine/Source/Runtime/Engine/Private/Particles/ParticleGpuSimulation.cpp
@@ -604,7 +604,7 @@ TGlobalResource<FGPUSpriteVertexDeclaration> GGPUSpriteVertexDeclaration;
 /**
  * Vertex factory for render sprites from GPU simulated particles.
  */
-class FGPUSpriteVertexFactory : public FVertexFactory
+class FGPUSpriteVertexFactory : public FParticleVertexFactoryBase
 {
 	DECLARE_VERTEX_FACTORY_TYPE(FGPUSpriteVertexFactory);
 
@@ -625,6 +625,12 @@ public:
 	/** Texture containint attributes for all particles. */
 	FTexture2DRHIParamRef AttributesTextureRHI;
 
+
+	FGPUSpriteVertexFactory()
+		: FParticleVertexFactoryBase(PVFT_MAX, ERHIFeatureLevel::Num)
+		, ParticleIndicesOffset(0)
+	{}
+
 	/**
 	 * Constructs render resources for this vertex factory.
 	 */
@@ -2668,18 +2674,17 @@ public:
 class FGPUSpriteCollectorResources : public FOneFrameResource
 {
 public:
-	FGPUSpriteVertexFactory VertexFactory;
+	FGPUSpriteVertexFactory *VertexFactory;
 
 	~FGPUSpriteCollectorResources()
 	{
-		VertexFactory.ReleaseResource();
 	}
 };
 
 struct FNewParticleAlloc
 {
-	TLockFreeFixedSizeAllocator<sizeof(TArray<FNewParticle>), FThreadSafeCounter> FreeTArrayFNewParticleArrays;
-	TLockFreePointerListUnordered<TArray<FNewParticle>>	FreeTArrayFNewParticle;
+	TLockFreeFixedSizeAllocator<sizeof(TArray<FNewParticle>), PLATFORM_CACHE_LINE_SIZE, FThreadSafeCounter> FreeTArrayFNewParticleArrays;
+	TLockFreePointerListUnordered<TArray<FNewParticle>, PLATFORM_CACHE_LINE_SIZE>	FreeTArrayFNewParticle;
 	~FNewParticleAlloc()
 	{
 		while (true)
@@ -2875,7 +2880,14 @@ struct FGPUSpriteDynamicEmitterData : FDynamicEmitterDataBase
 	{		
 	}
 
-	virtual void GetDynamicMeshElementsEmitter(const FParticleSystemSceneProxy* Proxy, const FSceneView* View, const FSceneViewFamily& ViewFamily, int32 ViewIndex, FMeshElementCollector& Collector) const override
+	virtual FParticleVertexFactoryBase *CreateVertexFactory() override
+	{
+		FGPUSpriteVertexFactory *VertexFactory = new FGPUSpriteVertexFactory();
+		VertexFactory->InitResource();
+		return VertexFactory;
+	}
+
+	virtual void GetDynamicMeshElementsEmitter(const FParticleSystemSceneProxy* Proxy, const FSceneView* View, const FSceneViewFamily& ViewFamily, int32 ViewIndex, FMeshElementCollector& Collector, FParticleVertexFactoryBase *InVertexFactory) const override
 	{
 		auto FeatureLevel = ViewFamily.GetFeatureLevel();
 
@@ -2917,9 +2929,10 @@ struct FGPUSpriteDynamicEmitterData : FDynamicEmitterDataBase
 				// Iterate over views and assign parameters for each.
 				FParticleSimulationResources* SimulationResources = FXSystem->GetParticleSimulationResources();
 				FGPUSpriteCollectorResources& CollectorResources = Collector.AllocateOneFrameResource<FGPUSpriteCollectorResources>();
-				CollectorResources.VertexFactory.SetFeatureLevel(FeatureLevel);
-				CollectorResources.VertexFactory.InitResource();
-				FGPUSpriteVertexFactory& VertexFactory = CollectorResources.VertexFactory;
+				//CollectorResources.VertexFactory.InitResource();
+				CollectorResources.VertexFactory = static_cast<FGPUSpriteVertexFactory*>(InVertexFactory);
+				CollectorResources.VertexFactory->SetFeatureLevel(FeatureLevel);
+				FGPUSpriteVertexFactory& VertexFactory = *CollectorResources.VertexFactory;
 
 				if (bAllowSorting && SortMode == PSORTMODE_DistanceToView)
 				{
diff --git a/Engine/Source/Runtime/Engine/Private/Particles/ParticleSystemRender.cpp b/Engine/Source/Runtime/Engine/Private/Particles/ParticleSystemRender.cpp
--- a/Engine/Source/Runtime/Engine/Private/Particles/ParticleSystemRender.cpp
+++ b/Engine/Source/Runtime/Engine/Private/Particles/ParticleSystemRender.cpp
@@ -972,23 +972,18 @@ public:
 
 
 
-void FDynamicSpriteEmitterData::CreateVertexFactory(const FSceneView *View)
+FParticleVertexFactoryBase *FDynamicSpriteEmitterData::CreateVertexFactory()
 {
-	if (VertexFactory == nullptr)
-	{
-		VertexFactory = new FParticleSpriteVertexFactory();
+	FParticleSpriteVertexFactory *VertexFactory = new FParticleSpriteVertexFactory();
 		VertexFactory->SetParticleFactoryType(PVFT_Sprite);
-		VertexFactory->SetFeatureLevel(View->GetFeatureLevel());
 		const USubUVAnimation* SubUVAnimation = GetSourceData()->SubUVAnimation;
 		VertexFactory->SetNumVertsInInstanceBuffer(SubUVAnimation ? SubUVAnimation->GetNumBoundingVertices() : 4);
 		VertexFactory->InitResource();
-	}
-
-	VertexFactory->SetFeatureLevel(View->GetFeatureLevel());
+	return VertexFactory;
 }
 
 
-void FDynamicSpriteEmitterData::GetDynamicMeshElementsEmitter(const FParticleSystemSceneProxy* Proxy, const FSceneView* View, const FSceneViewFamily& ViewFamily, int32 ViewIndex, FMeshElementCollector& Collector) const
+void FDynamicSpriteEmitterData::GetDynamicMeshElementsEmitter(const FParticleSystemSceneProxy* Proxy, const FSceneView* View, const FSceneViewFamily& ViewFamily, int32 ViewIndex, FMeshElementCollector& Collector, FParticleVertexFactoryBase *VertexFactory) const
 {
 	SCOPE_CYCLE_COUNTER(STAT_SpriteRenderingTime);
 
@@ -1022,7 +1017,7 @@ void FDynamicSpriteEmitterData::GetDynamicMeshElementsEmitter(const FParticleSys
 
 			FDynamicSpriteCollectorResources& CollectorResources = Collector.AllocateOneFrameResource<FDynamicSpriteCollectorResources>();
 			VertexFactory->SetFeatureLevel(FeatureLevel);
-			CollectorResources.VertexFactory = VertexFactory;
+			CollectorResources.VertexFactory = static_cast<FParticleSpriteVertexFactory*>(VertexFactory);
 
 			if (SourceData->bUseLocalSpace == false)
 			{
@@ -1313,7 +1308,6 @@ FDynamicMeshEmitterData::FDynamicMeshEmitterData(const UParticleModuleRequired*
 	, bUseCameraFacing(false)
 	, bApplyParticleRotationAsSpin(false)
 	, CameraFacingOption(0)
-	, VertexFactory(nullptr)
 {
 	// only update motion blur transforms if we are not paused
 	// bPlayersOnlyPending allows us to keep the particle transforms 
@@ -1322,11 +1316,6 @@ FDynamicMeshEmitterData::FDynamicMeshEmitterData(const UParticleModuleRequired*
 
 FDynamicMeshEmitterData::~FDynamicMeshEmitterData()
 {
-	if (VertexFactory)
-	{
-		VertexFactory->ReleaseResource();
-		delete VertexFactory;
-	}
 }
 
 /** Initialize this emitter's dynamic rendering data, called after source data has been filled in */
@@ -1464,7 +1453,6 @@ public:
 
 	virtual ~FDynamicMeshEmitterCollectorResources()
 	{
-		//VertexFactory.ReleaseResource();
 	}
 };
 
@@ -1476,31 +1464,28 @@ public:
 };
 
 
-void FDynamicMeshEmitterData::CreateVertexFactory(const FSceneView *View)
+FParticleVertexFactoryBase *FDynamicMeshEmitterData::CreateVertexFactory()
 {
-	if (VertexFactory == nullptr)
-	{
-		const int32 InstanceVertexStride = GetDynamicVertexStride(View->GetFeatureLevel());
-		VertexFactory = ConstructMeshParticleVertexFactory();
+		FMeshParticleVertexFactory *VertexFactory = ConstructMeshParticleVertexFactory();
 		VertexFactory->SetParticleFactoryType(PVFT_Mesh);
-		VertexFactory->SetFeatureLevel(View->GetFeatureLevel());
 		SetupVertexFactory(VertexFactory, StaticMesh->RenderData->LODResources[0]);
+		const int32 InstanceVertexStride = GetDynamicVertexStride(ERHIFeatureLevel::Type::SM5);	// featurelevel is ignored
+		const int32 DynamicParameterVertexStride = GetDynamicParameterVertexStride();
 		VertexFactory->SetStrides(InstanceVertexStride, GetDynamicParameterVertexStride());
 		VertexFactory->InitResource();
-	}
+		return VertexFactory;
 
-	VertexFactory->SetFeatureLevel(View->GetFeatureLevel());
 }
 
 
 
-void FDynamicMeshEmitterData::GetDynamicMeshElementsEmitter(const FParticleSystemSceneProxy* Proxy, const FSceneView* View, const FSceneViewFamily& ViewFamily, int32 ViewIndex, FMeshElementCollector& Collector) const
+void FDynamicMeshEmitterData::GetDynamicMeshElementsEmitter(const FParticleSystemSceneProxy* Proxy, const FSceneView* View, const FSceneViewFamily& ViewFamily, int32 ViewIndex, FMeshElementCollector& Collector, FParticleVertexFactoryBase *VertexFactory) const
 {
 	SCOPE_CYCLE_COUNTER(STAT_MeshRenderingTime);
 
 	const bool bInstanced = GRHISupportsInstancing;
 
-	if (bValid)
+	if (bValid && VertexFactory)
 	{
 		if (Source.EmitterRenderMode == ERM_Normal)
 		{
@@ -1517,7 +1502,7 @@ void FDynamicMeshEmitterData::GetDynamicMeshElementsEmitter(const FParticleSyste
 			const int32 DynamicParameterVertexStride  = GetDynamicParameterVertexStride();
 
 			FDynamicMeshEmitterCollectorResources& CollectorResources = Collector.AllocateOneFrameResource<FDynamicMeshEmitterCollectorResources>();
-			CollectorResources.VertexFactory = VertexFactory;
+			CollectorResources.VertexFactory = static_cast<FMeshParticleVertexFactory*>(VertexFactory);
 
 			// Setup the vertex factory.
 			FMeshParticleVertexFactory* MeshVertexFactory = CollectorResources.VertexFactory;
@@ -2167,7 +2152,7 @@ void FDynamicMeshEmitterData::GetInstanceData(void* InstanceData, void* DynamicP
 
 void FDynamicMeshEmitterData::SetupVertexFactory( FMeshParticleVertexFactory* InVertexFactory, FStaticMeshLODResources& LODResources) const
 {
-		FMeshParticleVertexFactory::DataType Data;
+		FMeshParticleVertexFactory::FDataType Data;
 
 		Data.PositionComponent = FVertexStreamComponent(
 			&LODResources.PositionVertexBuffer,
@@ -2286,11 +2271,6 @@ void FDynamicMeshEmitterData::SetupVertexFactory( FMeshParticleVertexFactory* In
 
 FDynamicBeam2EmitterData::~FDynamicBeam2EmitterData()
 {
-	if (VertexFactory)
-	{
-		VertexFactory->ReleaseResource();
-		delete VertexFactory;
-	}
 }
 
 
@@ -2383,21 +2363,16 @@ public:
 };
 
 
-void FDynamicBeam2EmitterData::CreateVertexFactory(const FSceneView *View)
-{
-	if (VertexFactory == nullptr)
+FParticleVertexFactoryBase *FDynamicBeam2EmitterData::CreateVertexFactory()
 	{
-		VertexFactory = new FParticleBeamTrailVertexFactory();
+		FParticleBeamTrailVertexFactory *VertexFactory = new FParticleBeamTrailVertexFactory();
 		VertexFactory->SetParticleFactoryType(PVFT_BeamTrail);
-		VertexFactory->SetFeatureLevel(View->GetFeatureLevel());
 		VertexFactory->InitResource();
+		return VertexFactory;
 	}
 
-	VertexFactory->SetFeatureLevel(View->GetFeatureLevel());
-}
-
 
-void FDynamicBeam2EmitterData::GetDynamicMeshElementsEmitter(const FParticleSystemSceneProxy* Proxy, const FSceneView* View, const FSceneViewFamily& ViewFamily, int32 ViewIndex, FMeshElementCollector& Collector) const
+void FDynamicBeam2EmitterData::GetDynamicMeshElementsEmitter(const FParticleSystemSceneProxy* Proxy, const FSceneView* View, const FSceneViewFamily& ViewFamily, int32 ViewIndex, FMeshElementCollector& Collector, FParticleVertexFactoryBase *VertexFactory) const
 {
 	SCOPE_CYCLE_COUNTER(STAT_BeamRenderingTime);
 	INC_DWORD_STAT(STAT_BeamParticlesRenderCalls);
@@ -2430,7 +2405,8 @@ void FDynamicBeam2EmitterData::GetDynamicMeshElementsEmitter(const FParticleSyst
 	if (Data.OutTriangleCount > 0)
 	{
 		FDynamicBeamTrailCollectorResources& CollectorResources = Collector.AllocateOneFrameResource<FDynamicBeamTrailCollectorResources>();
-		CollectorResources.VertexFactory = VertexFactory;
+		VertexFactory->SetFeatureLevel(View->GetFeatureLevel());
+		CollectorResources.VertexFactory = static_cast<FParticleBeamTrailVertexFactory*>(VertexFactory);
 
 		// Create and set the uniform buffer for this emitter.
 		FParticleBeamTrailVertexFactory* BeamTrailVertexFactory = CollectorResources.VertexFactory;
@@ -5089,11 +5065,6 @@ int32 FDynamicBeam2EmitterData::FillData_InterpolatedNoise(FAsyncBufferFillData&
 
 FDynamicTrailsEmitterData::~FDynamicTrailsEmitterData()
 {
-	if (VertexFactory)
-	{
-		VertexFactory->ReleaseResource();
-		delete VertexFactory;
-	}
 }
 
 
@@ -5123,26 +5094,21 @@ FParticleVertexFactoryBase* FDynamicTrailsEmitterData::BuildVertexFactory(const
 
 
 
-void FDynamicTrailsEmitterData::CreateVertexFactory(const FSceneView *View)
+FParticleVertexFactoryBase *FDynamicTrailsEmitterData::CreateVertexFactory()
 {
-	if (VertexFactory == nullptr)
-	{
-		VertexFactory = new FParticleBeamTrailVertexFactory();
+	FParticleBeamTrailVertexFactory *VertexFactory = new FParticleBeamTrailVertexFactory();
 		VertexFactory->SetParticleFactoryType(PVFT_BeamTrail);
-		VertexFactory->SetFeatureLevel(View->GetFeatureLevel());
 		VertexFactory->InitResource();
-	}
-
-	VertexFactory->SetFeatureLevel(View->GetFeatureLevel());
+		return VertexFactory;
 }
 
 
-void FDynamicTrailsEmitterData::GetDynamicMeshElementsEmitter(const FParticleSystemSceneProxy* Proxy, const FSceneView* View, const FSceneViewFamily& ViewFamily, int32 ViewIndex, FMeshElementCollector& Collector) const
+void FDynamicTrailsEmitterData::GetDynamicMeshElementsEmitter(const FParticleSystemSceneProxy* Proxy, const FSceneView* View, const FSceneViewFamily& ViewFamily, int32 ViewIndex, FMeshElementCollector& Collector, FParticleVertexFactoryBase *VertexFactory) const
 {
 	SCOPE_CYCLE_COUNTER(STAT_TrailRenderingTime);
 	INC_DWORD_STAT(STAT_TrailParticlesRenderCalls);
 
-	if (bValid == false)
+	if (bValid == false || !VertexFactory)
 	{
 		return;
 	}
@@ -5175,7 +5141,8 @@ void FDynamicTrailsEmitterData::GetDynamicMeshElementsEmitter(const FParticleSys
 	{
 		FDynamicBeamTrailCollectorResources& CollectorResources = Collector.AllocateOneFrameResource<FDynamicBeamTrailCollectorResources>();
 
-		CollectorResources.VertexFactory = VertexFactory;
+		VertexFactory->SetFeatureLevel(View->GetFeatureLevel());
+		CollectorResources.VertexFactory = static_cast<FParticleBeamTrailVertexFactory*>(VertexFactory);
 
 		// Create and set the uniform buffer for this emitter.
 		FParticleBeamTrailVertexFactory* BeamTrailVertexFactory = CollectorResources.VertexFactory;
@@ -6481,12 +6448,14 @@ FParticleSystemSceneProxy::FParticleSystemSceneProxy(const UParticleSystemCompon
 	, VisualizeLODIndex(Component->GetCurrentLODIndex())
 	, LastFramePreRendered(-1)
 	, FirstFreeMeshBatch(0)
+	, bVertexFactoriesDirty(false)
 {
 	WireframeColor = FLinearColor(1.0f, 0.0f, 0.0f);
 	LevelColor = FLinearColor(1.0f, 1.0f, 0.0f);
 	PropertyColor = FLinearColor(1.0f, 1.0f, 1.0f);
 
 	LODMethod = Component->LODMethod;
+
 }
 
 FParticleSystemSceneProxy::~FParticleSystemSceneProxy()
@@ -6495,6 +6464,8 @@ FParticleSystemSceneProxy::~FParticleSystemSceneProxy()
 
 	delete DynamicData;
 	DynamicData = NULL;
+
+	ClearVertexFactories();
 }
 
 FMeshBatch* FParticleSystemSceneProxy::GetPooledMeshBatch()
@@ -6550,6 +6521,8 @@ void FParticleSystemSceneProxy::GetDynamicMeshElements(const TArray<const FScene
 				}
 				FScopeCycleCounter AdditionalScope(Data->StatID);
 
+				FParticleVertexFactoryBase *VertexFactory = EmitterVertexFactoryArray[Data->EmitterIndex];
+
 				//hold on to the emitter index in case we need to access any of its properties
 				DynamicData->EmitterIndex = Index;
 
@@ -6559,8 +6532,9 @@ void FParticleSystemSceneProxy::GetDynamicMeshElements(const TArray<const FScene
 					{
 
 						const FSceneView* View = Views[ViewIndex];
-						Data->CreateVertexFactory(View);
-						Data->GetDynamicMeshElementsEmitter(this, View, ViewFamily, ViewIndex, Collector);
+
+
+						Data->GetDynamicMeshElementsEmitter(this, View, ViewFamily, ViewIndex, Collector, VertexFactory);
 						NumDraws++;
 					}
 				}
@@ -6611,6 +6585,8 @@ void FParticleSystemSceneProxy::CreateRenderThreadResourcesForEmitterData()
 			}
 		}
 	}
+	ClearVertexFactoriesIfDirty();
+	UpdateVertexFactories();
 }
 
 void FParticleSystemSceneProxy::ReleaseRenderThreadResourcesForEmitterData()
@@ -6774,6 +6750,7 @@ FPrimitiveViewRelevance FParticleSystemSceneProxy::GetViewRelevance(const FScene
 	FPrimitiveViewRelevance Result;
 	Result.bDrawRelevance = IsShown(View) && View->Family->EngineShowFlags.Particles;
 	Result.bShadowRelevance = IsShadowCast(View);
+	Result.bRenderCustomDepth = ShouldRenderCustomDepth();
 	Result.bRenderInMainPass = ShouldRenderInMainPass();
 	Result.bUsesLightingChannels = GetLightingChannelMask() != GetDefaultLightingChannelMask();
 	Result.bDynamicRelevance = true;
@@ -6895,6 +6872,13 @@ FPrimitiveSceneProxy* UParticleSystemComponent::CreateSceneProxy()
 			NewProxy = ::new FParticleSystemOcclusionSceneProxy(this,ParticleDynamicData);
 		}
 		check (NewProxy);
+		if (ParticleDynamicData)
+		{
+			for (int32 Index = 0; Index < ParticleDynamicData->DynamicEmitterDataArray.Num(); Index++)
+			{
+				NewProxy->QueueVertexFactoryCreation(ParticleDynamicData->DynamicEmitterDataArray[Index]);
+			}
+		}
 	}
 	
 	// 
diff --git a/Engine/Source/Runtime/Engine/Private/PhysicsEngine/PhysAnim.cpp b/Engine/Source/Runtime/Engine/Private/PhysicsEngine/PhysAnim.cpp
--- a/Engine/Source/Runtime/Engine/Private/PhysicsEngine/PhysAnim.cpp
+++ b/Engine/Source/Runtime/Engine/Private/PhysicsEngine/PhysAnim.cpp
@@ -27,6 +27,13 @@ struct FAssetWorldBoneTM
 };
 
 
+FAutoConsoleTaskPriority CPrio_FParallelBlendPhysicsTask(
+	TEXT("TaskGraph.TaskPriorities.ParallelBlendPhysicsTask"),
+	TEXT("Task and thread priority for FParallelBlendPhysicsTask."),
+	ENamedThreads::HighThreadPriority, // if we have high priority task threads, then use them...
+	ENamedThreads::NormalTaskPriority, // .. at normal task priority
+	ENamedThreads::HighTaskPriority // if we don't have hi pri threads, then use normal priority threads at high task priority instead
+	);
 
 class FParallelBlendPhysicsTask
 {
@@ -38,15 +45,15 @@ public:
 	{
 	}
 
-	FORCEINLINE TStatId GetStatId() const
+	static FORCEINLINE TStatId GetStatId()
 	{
 		RETURN_QUICK_DECLARE_CYCLE_STAT(FParallelBlendPhysicsTask, STATGROUP_TaskGraphTasks);
 	}
-	static ENamedThreads::Type GetDesiredThread()
+	static FORCEINLINE ENamedThreads::Type GetDesiredThread()
 	{
-		return ENamedThreads::AnyThread;
+		return CPrio_FParallelBlendPhysicsTask.Get();
 	}
-	static ESubsequentsMode::Type GetSubsequentsMode()
+	static FORCEINLINE ESubsequentsMode::Type GetSubsequentsMode()
 	{
 		return ESubsequentsMode::TrackSubsequents;
 	}
diff --git a/Engine/Source/Runtime/Engine/Private/PhysicsEngine/PhysScene.cpp b/Engine/Source/Runtime/Engine/Private/PhysicsEngine/PhysScene.cpp
--- a/Engine/Source/Runtime/Engine/Private/PhysicsEngine/PhysScene.cpp
+++ b/Engine/Source/Runtime/Engine/Private/PhysicsEngine/PhysScene.cpp
@@ -125,6 +125,23 @@ struct FPendingApexDamageManager
 
 #if WITH_PHYSX
 
+FAutoConsoleTaskPriority CPrio_FPhysXTask(
+	TEXT("TaskGraph.TaskPriorities.PhysXTask"),
+	TEXT("Task and thread priority for FPhysXTask."),
+	ENamedThreads::HighThreadPriority, // if we have high priority task threads, then use them...
+	ENamedThreads::NormalTaskPriority, // .. at normal task priority
+	ENamedThreads::HighTaskPriority // if we don't have hi pri threads, then use normal priority threads at high task priority instead
+	);
+
+FAutoConsoleTaskPriority CPrio_FPhysXTask_Cloth(
+	TEXT("TaskGraph.TaskPriorities.PhysXTask.Cloth"),
+	TEXT("Task and thread priority for FPhysXTask (cloth)."),
+	ENamedThreads::HighThreadPriority, // if we have high priority task threads, then use them...
+	ENamedThreads::NormalTaskPriority, // .. at normal task priority
+	ENamedThreads::HighTaskPriority // if we don't have hi pri threads, then use normal priority threads at high task priority instead
+	);
+
+
 template <bool IsCloth>
 class FPhysXTask
 {
@@ -141,7 +158,7 @@ public:
 		Task.release();
 	}
 
-	FORCEINLINE TStatId GetStatId() const
+	static FORCEINLINE TStatId GetStatId()
 	{
 		if (!IsCloth)
 		{
@@ -153,11 +170,18 @@ public:
 		}
 
 	}
-	static ENamedThreads::Type GetDesiredThread()
+	static FORCEINLINE ENamedThreads::Type GetDesiredThread()
 	{
-		return ENamedThreads::HiPri(ENamedThreads::AnyThread);
+		if (!IsCloth)
+		{
+			return CPrio_FPhysXTask.Get();
+		}
+		else
+		{
+			return CPrio_FPhysXTask_Cloth.Get();
+		}
 	}
-	static ESubsequentsMode::Type GetSubsequentsMode()
+	static FORCEINLINE ESubsequentsMode::Type GetSubsequentsMode()
 	{
 		return ESubsequentsMode::TrackSubsequents;
 	}
@@ -538,6 +562,14 @@ void FPhysScene::AddPendingDamageEvent(UDestructibleComponent* DestructibleCompo
 }
 #endif
 
+FAutoConsoleTaskPriority CPrio_PhysXStepSimulation(
+	TEXT("TaskGraph.TaskPriorities.PhysXStepSimulation"),
+	TEXT("Task and thread priority for FPhysSubstepTask::StepSimulation."),
+	ENamedThreads::HighThreadPriority, // if we have high priority task threads, then use them...
+	ENamedThreads::NormalTaskPriority, // .. at normal task priority
+	ENamedThreads::HighTaskPriority // if we don't have hi pri threads, then use normal priority threads at high task priority instead
+	);
+
 bool FPhysScene::SubstepSimulation(uint32 SceneType, FGraphEventRef &InOutCompletionEvent)
 {
 #if WITH_PHYSX
@@ -552,7 +584,7 @@ bool FPhysScene::SubstepSimulation(uint32 SceneType, FGraphEventRef &InOutComple
 	{
 		//we have valid scene and subtime so enqueue task
 		PhysXCompletionTask* Task = new PhysXCompletionTask(InOutCompletionEvent, SceneType, PScene->getTaskManager());
-		ENamedThreads::Type NamedThread = PhysSingleThreadedMode() ? ENamedThreads::GameThread : ENamedThreads::HiPri(ENamedThreads::AnyThread);
+		ENamedThreads::Type NamedThread = PhysSingleThreadedMode() ? ENamedThreads::GameThread : CPrio_PhysXStepSimulation.Get();
 
 		DECLARE_CYCLE_STAT(TEXT("FSimpleDelegateGraphTask.SubstepSimulationImp"),
 			STAT_FSimpleDelegateGraphTask_SubstepSimulationImp,
@@ -1237,6 +1269,16 @@ void FPhysScene::SetUpForFrame(const FVector* NewGrav, float InDeltaSeconds, flo
 #endif
 }
 
+FAutoConsoleTaskPriority CPrio_PhyXSceneCompletion(
+	TEXT("TaskGraph.TaskPriorities.PhyXSceneCompletion"),
+	TEXT("Task and thread priority for PhysicsSceneCompletion."),
+	ENamedThreads::HighThreadPriority, // if we have high priority task threads, then use them...
+	ENamedThreads::HighTaskPriority, // .. at high task priority
+	ENamedThreads::HighTaskPriority // if we don't have hi pri threads, then use normal priority threads at high task priority instead
+	);
+
+
+
 void FPhysScene::StartFrame()
 {
 	FGraphEventArray FinishPrerequisites;
@@ -1307,7 +1349,7 @@ void FPhysScene::StartFrame()
 				STATGROUP_TaskGraphTasks);
 
 			PhysicsSceneCompletion = TGraphTask<FNullGraphTask>::CreateTask(&FinishPrerequisites, ENamedThreads::GameThread).ConstructAndDispatchWhenReady(
-				GET_STATID(STAT_FNullGraphTask_ProcessPhysScene_Join), PhysSingleThreadedMode() ? ENamedThreads::GameThread : ENamedThreads::HiPri(ENamedThreads::AnyThread));
+				GET_STATID(STAT_FNullGraphTask_ProcessPhysScene_Join), PhysSingleThreadedMode() ? ENamedThreads::GameThread : CPrio_PhyXSceneCompletion.Get());
 		}
 		else
 		{
