Copying //UE4/Orion-Staging to //UE4/Main (Source: //Orion/Dev-General @ 2845681)

diff --git a/Engine/Source/Runtime/Core/Private/Windows/WindowsPlatformMisc.cpp b/Engine/Source/Runtime/Core/Private/Windows/WindowsPlatformMisc.cpp
--- a/Engine/Source/Runtime/Core/Private/Windows/WindowsPlatformMisc.cpp
+++ b/Engine/Source/Runtime/Core/Private/Windows/WindowsPlatformMisc.cpp
@@ -2393,6 +2393,74 @@ FString FWindowsPlatformMisc::GetPrimaryGPUBrand()
 	return PrimaryGPUBrand;
 }
 
+static void GetVideoDriverDetails(const FString& Key, FGPUDriverInfo& Out)
+{
+	// https://msdn.microsoft.com/en-us/library/windows/hardware/ff569240(v=vs.85).aspx
+
+	const TCHAR* DeviceDescriptionValueName = TEXT("Device Description");
+
+	bool bDevice = FWindowsPlatformMisc::QueryRegKey(HKEY_LOCAL_MACHINE, *Key, DeviceDescriptionValueName, Out.DeviceDescription); // AMD and NVIDIA
+
+	if (!bDevice)
+	{
+		// in the case where this failed we also have:
+		//  "DriverDesc"="NVIDIA GeForce GTX 670"
+		
+		// e.g. "GeForce GTX 680" (no NVIDIA prefix so no good for string comparison with DX)
+		//	FWindowsPlatformMisc::QueryRegKey(HKEY_LOCAL_MACHINE, *Key, TEXT("HardwareInformation.AdapterString"), Out.DeviceDescription); // AMD and NVIDIA
+
+		// Try again in Settings subfolder
+		const FString SettingsSubKey = Key + TEXT("\\Settings");
+		bDevice = FWindowsPlatformMisc::QueryRegKey(HKEY_LOCAL_MACHINE, *SettingsSubKey, DeviceDescriptionValueName, Out.DeviceDescription); // AMD and NVIDIA
+
+		if (!bDevice)
+		{
+			// Neither root nor Settings subfolder contained a "Device Description" value so this is probably not a device
+			Out = FGPUDriverInfo();
+			return;
+		}
+	}
+
+	FWindowsPlatformMisc::QueryRegKey(HKEY_LOCAL_MACHINE, *Key, TEXT("ProviderName"), Out.ProviderName);
+
+	// AMD (not the Catalyst one) and NVIDIA
+	FWindowsPlatformMisc::QueryRegKey(HKEY_LOCAL_MACHINE, *Key, TEXT("DriverVersion"), Out.InternalDriverVersion);
+
+	Out.UserDriverVersion = Out.InternalDriverVersion;
+
+	if(Out.ProviderName == TEXT("NVIDIA"))
+	{
+		// https://forums.geforce.com/default/topic/378546/confusion-over-driver-version-numbers/
+		// we don't care about the windows version so we don't look at the front part of the driver version
+		// "9.18.13.4788" -> "347.88"
+		// "10.18.13.4788" -> "347.88"
+		// the following code works with the current numbering scheme, if needed we have to update that
+		{
+			// e.g. 3.4788
+			FString RightPart = Out.UserDriverVersion.Right(6);
+
+			if(RightPart.Len() == 6 && RightPart[1] == (TCHAR)'.')
+			{
+				// e.g. 347.88
+				Out.UserDriverVersion = RightPart.Left(1) + RightPart[2] + RightPart[3] + TEXT(".") + RightPart[4] + RightPart[5];
+			}
+		}
+	}
+	else 
+	{
+		// we assume Out.ProviderName == TEXT("Advanced Micro Devices, Inc.")
+
+		// e.g. 15.7.1
+		if(FWindowsPlatformMisc::QueryRegKey(HKEY_LOCAL_MACHINE, *Key, TEXT("Catalyst_Version"), Out.UserDriverVersion))
+		{
+			Out.UserDriverVersion = FString(TEXT("Catalyst ")) + Out.UserDriverVersion;
+		}
+	}
+
+	// AMD and NVIDIA
+	FWindowsPlatformMisc::QueryRegKey(HKEY_LOCAL_MACHINE, *Key, TEXT("DriverDate"), Out.DriverDate);
+}
+
 void FWindowsPlatformMisc::GetGPUDriverInfo(const FString DeviceDescription, FString& InternalDriverVersion, FString& UserDriverVersion, FString& DriverDate)
 {
 	// to distinguish failed GetGPUDriverInfo() from call to GetGPUDriverInfo()
@@ -2400,82 +2468,102 @@ void FWindowsPlatformMisc::GetGPUDriverInfo(const FString DeviceDescription, FSt
 	UserDriverVersion = TEXT("Unknown");
 	DriverDate = TEXT("Unknown");
 
+	// for debugging, useful even in shipping to see what went wrong
+	FString ErrorString;
+
+	uint32 FoundDriverCount = 0;
+
 	for(uint32 i = 0;; ++i)
 	{
 		// Iterate all installed display adapters
 		const FString Key = FString::Printf(TEXT("SYSTEM\\CurrentControlSet\\Control\\Class\\{4D36E968-E325-11CE-BFC1-08002BE10318}\\%04d"), i);
-		const TCHAR* DeviceDescriptionValueName = TEXT("Device Description");
-
-		FString LocalDeviceDescription; // e.g. "NVIDIA GeForce GTX 680" or "AMD Radeon R9 200 / HD 7900 Series"
-		bool bDevice = FWindowsPlatformMisc::QueryRegKey(HKEY_LOCAL_MACHINE, *Key, DeviceDescriptionValueName, /*out*/ LocalDeviceDescription); // AMD and NVIDIA
+		
+		FGPUDriverInfo Local;
+		GetVideoDriverDetails(Key, Local);
 
-		if (!bDevice)
+		if(!Local.IsValid())
 		{
-			// Try again in Settings subfolder
-			const FString SettingsSubKey = Key + TEXT("\\Settings");
-			bDevice = FWindowsPlatformMisc::QueryRegKey(HKEY_LOCAL_MACHINE, *SettingsSubKey, DeviceDescriptionValueName, /*out*/ LocalDeviceDescription); // AMD and NVIDIA
+			ErrorString += TEXT("GetVideoDriverDetailsInvalid ");
+		}
 
-			if (!bDevice)
-			{
-				// Neither root nor Settings subfolder contained a "Device Description" value so this is probably not a device
-				break;
-			}
+		if(Local.DeviceDescription == DeviceDescription)
+		{
+			// found the one we are searching for
+			InternalDriverVersion = Local.InternalDriverVersion;
+			UserDriverVersion = Local.UserDriverVersion;
+			DriverDate = Local.DriverDate;
+			++FoundDriverCount;
+			break;
 		}
+	}
 
-		// e.g. "NVIDIA" or "Advanced Micro Devices, Inc."
-		FString LocalProviderName;
-		FWindowsPlatformMisc::QueryRegKey(HKEY_LOCAL_MACHINE, *Key, TEXT("ProviderName"), LocalProviderName);
+	// should be 1 after the code but we've seen it >1 and the code needs to be improved (seems it was from upgrading a machine to a new windows version)
+	// see: http://win7settings.blogspot.com/2014/10/how-to-extract-installed-drivers-from.html
+	// https://support.microsoft.com/en-us/kb/200435
+	// http://www.experts-exchange.com/questions/10198207/Windows-NT-Display-adapter-information.html
+	// alternative: from https://support.microsoft.com/en-us/kb/200435
+	if(FoundDriverCount != 1)
+	{
+		// we start again, this time we only look at the primary adapter
+		InternalDriverVersion = TEXT("Unknown");
+		UserDriverVersion = TEXT("Unknown");
+		DriverDate = TEXT("Unknown");
 
-		// AMD (not the Catalyst one) and NVIDIA
-		// e.g. "15.200.1062.1004"(AMD) "9.18.13.4788"(NVIDIA)
-		FString LocalInternalDriverVersion;
-		FWindowsPlatformMisc::QueryRegKey(HKEY_LOCAL_MACHINE, *Key, TEXT("DriverVersion"), LocalInternalDriverVersion);
+		ErrorString += FString::Printf(TEXT("FoundDriverCount:%d FallbackToPrimary "), FoundDriverCount);
+	
+		// Iterate all installed display adapters
+		const FString Key = TEXT("HARDWARE\\DEVICEMAP\\VIDEO");
 
-		// e.g. "Catalyst 15.7.1"(AMD) or "347.88"(NVIDIA)
-		FString LocalUserDriverVersion = LocalInternalDriverVersion;
+		FString DriverLocation; // e.g. HKEY_LOCAL_MACHINE\System\CurrentControlSet\Services\<videodriver>\Device0
+		bool bOk = FWindowsPlatformMisc::QueryRegKey(HKEY_LOCAL_MACHINE, *Key, TEXT("\\Device\\Video0"), /*out*/ DriverLocation);
 
-		if(LocalProviderName == TEXT("NVIDIA"))
+		if(bOk)
 		{
-			// "9.18.13.4788" -> "347.88"
-			// the following code works with the current numbering scheme, if needd we have to update that
-			if(LocalUserDriverVersion.Left(6) == TEXT("9.18.1"))
+			if(DriverLocation.Left(18) == TEXT("\\Registry\\Machine\\"))		// not case sensitive
 			{
-				// e.g. 3.4788
-				FString RightPart = LocalUserDriverVersion.RightChop(6);
+				DriverLocation = FString(TEXT("\\HKEY_LOCAL_MACHINE\\")) + DriverLocation.RightChop(18);
+			}
+			if(DriverLocation.Left(20) == TEXT("\\HKEY_LOCAL_MACHINE\\"))		// not case sensitive
+			{
+				FString Key = DriverLocation.RightChop(20);
+				
+				FGPUDriverInfo Local;
+				GetVideoDriverDetails(Key, Local);
 
-				if(RightPart.Len() == 6 && RightPart[1] == (TCHAR)'.')
+				if(!Local.IsValid())
 				{
-					// e.g. 347.88
-					LocalUserDriverVersion = RightPart.Left(1) + RightPart[2] + RightPart[3] + TEXT(".") + RightPart[4] + RightPart[5];
+					ErrorString += TEXT("GetVideoDriverDetailsInvalid ");
 				}
-			}
-		}
-		else 
-		{
-			// we assume LocalProviderName == TEXT("Advanced Micro Devices, Inc.")
 
-			// e.g. 15.7.1
-			if(FWindowsPlatformMisc::QueryRegKey(HKEY_LOCAL_MACHINE, *Key, TEXT("Catalyst_Version"), LocalUserDriverVersion))
+				if(Local.DeviceDescription == DeviceDescription)
+				{
+					InternalDriverVersion = Local.InternalDriverVersion;
+					UserDriverVersion = Local.UserDriverVersion;
+					DriverDate = Local.DriverDate;
+					FoundDriverCount = 1;
+				}
+				else
+				{
+					ErrorString += TEXT("PrimaryIsNotTheChoosenAdapter ");
+				}
+			}
+			else
 			{
-				LocalUserDriverVersion = FString(TEXT("Catalyst ")) + LocalUserDriverVersion;
+				ErrorString += TEXT("PrimaryDriverLocationFailed ");
 			}
 		}
-
-		// AMD and NVIDIA
-		// e.g. 3-13-2015
-		FString LocalDriverDate;
-		FWindowsPlatformMisc::QueryRegKey(HKEY_LOCAL_MACHINE, *Key, TEXT("DriverDate"), LocalDriverDate);
-
-		if(LocalDeviceDescription == DeviceDescription)
+		else
 		{
-			// found the one we are searching for (if there are multiple we only get the first one)
-			InternalDriverVersion = LocalInternalDriverVersion;
-			UserDriverVersion = LocalUserDriverVersion;
-			DriverDate = LocalDriverDate;
-			break;
+			ErrorString += TEXT("QueryForPrimaryFailed ");
 		}
 	}
+
+	if(!ErrorString.IsEmpty())
+	{
+		UE_LOG(LogWindows, Error, TEXT("Error: %s"), *ErrorString);
+	}
 }
+
 #include "HideWindowsPlatformTypes.h"
 
 void FWindowsPlatformMisc::GetOSVersions( FString& out_OSVersionLabel, FString& out_OSSubVersionLabel )
diff --git a/Engine/Source/Runtime/Core/Private/Windows/WindowsWindow.cpp b/Engine/Source/Runtime/Core/Private/Windows/WindowsWindow.cpp
--- a/Engine/Source/Runtime/Core/Private/Windows/WindowsWindow.cpp
+++ b/Engine/Source/Runtime/Core/Private/Windows/WindowsWindow.cpp
@@ -4,14 +4,14 @@
 #include "WindowsWindow.h"
 #include "WindowsApplication.h"
 
-
 #include "AllowWindowsPlatformTypes.h"
 #if WINVER > 0x502	// Windows Vista or better required for DWM
 	#include "Dwmapi.h"
 #endif
-	#include <ShlObj.h>
+#include <ShlObj.h>
 #include "HideWindowsPlatformTypes.h"
 
+
 FWindowsWindow::~FWindowsWindow()
 {
 	// NOTE: The HWnd is invalid here!
@@ -476,10 +476,11 @@ void FWindowsWindow::HACK_ForceToFront()
 /** Native windows should implement this function by asking the OS to destroy OS-specific resource associated with the window (e.g. Win32 window handle) */
 void FWindowsWindow::Destroy()
 {
-	if (OLEReferenceCount > 0)
+	if (OLEReferenceCount > 0 && IsWindow(HWnd))
 	{
-		RevokeDragDrop( HWnd );
-		check( OLEReferenceCount == 0 );
+		HRESULT Result = RevokeDragDrop(HWnd);
+		checkf(Result == S_OK, TEXT("Failed in RevokeDragDrop, error-code: %i"), Result);
+		checkf(OLEReferenceCount == 0, TEXT("Not all references to window are released, %i left"), OLEReferenceCount);
 	}
 
 	::DestroyWindow( HWnd );
diff --git a/Engine/Source/Runtime/Core/Public/GenericPlatform/GenericApplicationMessageHandler.h b/Engine/Source/Runtime/Core/Public/GenericPlatform/GenericApplicationMessageHandler.h
--- a/Engine/Source/Runtime/Core/Public/GenericPlatform/GenericApplicationMessageHandler.h
+++ b/Engine/Source/Runtime/Core/Public/GenericPlatform/GenericApplicationMessageHandler.h
@@ -40,6 +40,8 @@ struct CORE_API FGamepadKeyNames
 	static const FName LeftThumb;
 	static const FName RightThumb;
 	static const FName SpecialLeft;
+	static const FName SpecialLeft_X;
+	static const FName SpecialLeft_Y;
 	static const FName SpecialRight;
 	static const FName FaceButtonBottom;
 	static const FName FaceButtonRight;
diff --git a/Engine/Source/Runtime/Core/Public/GenericPlatform/GenericPlatformMisc.h b/Engine/Source/Runtime/Core/Public/GenericPlatform/GenericPlatformMisc.h
--- a/Engine/Source/Runtime/Core/Public/GenericPlatform/GenericPlatformMisc.h
+++ b/Engine/Source/Runtime/Core/Public/GenericPlatform/GenericPlatformMisc.h
@@ -476,6 +476,15 @@ public:
 	 */
 	static void RequestExit( bool Force );
 
+	/**
+	 * Requests application exit with a specified return code. Name is different from RequestExit() so overloads of just one of functions are possible.
+	 *
+	 * @param	Force 	   If true, perform immediate exit (dangerous because config code isn't flushed, etc).
+	 *                     If false, request clean main-loop exit from the platform specific code.
+	 * @param   ReturnCode This value will be returned from the program (on the platforms where it's possible). Limited to 0-255 to conform with POSIX.
+	 */
+	static void RequestExitWithStatus( bool Force, uint8 ReturnCode );
+
 	/**
 	 * Returns the last system error code in string form.  NOTE: Only one return value is valid at a time!
 	 *
diff --git a/Engine/Source/Runtime/Core/Public/GenericPlatform/GenericWindowDefinition.h b/Engine/Source/Runtime/Core/Public/GenericPlatform/GenericWindowDefinition.h
--- a/Engine/Source/Runtime/Core/Public/GenericPlatform/GenericWindowDefinition.h
+++ b/Engine/Source/Runtime/Core/Public/GenericPlatform/GenericWindowDefinition.h
@@ -103,3 +103,37 @@ struct CORE_API FGenericWindowDefinition
 
 	FWindowSizeLimits SizeLimits;
 };
+
+// video driver details
+struct FGPUDriverInfo
+{
+	FGPUDriverInfo()
+		: VendorId(0)
+	{
+	}
+
+	// DirectX vendor Id, use functions below
+	uint32 VendorId;
+	// e.g. "NVIDIA GeForce GTX 680" or "AMD Radeon R9 200 / HD 7900 Series"
+	FString DeviceDescription;
+	// e.g. "NVIDIA" or "Advanced Micro Devices, Inc."
+	FString ProviderName;
+	// e.g. "15.200.1062.1004"(AMD)
+	// e.g. "9.18.13.4788"(NVIDIA) first number is Windows version (e.g. 7:Vista, 6:XP, 4:Me, 9:Win8(1), 10:Win7), last 5 have the UserDriver version encoded
+	FString InternalDriverVersion;	
+	// e.g. "Catalyst 15.7.1"(AMD) or "347.88"(NVIDIA)
+	FString UserDriverVersion;
+	// e.g. 3-13-2015
+	FString DriverDate;
+
+	bool IsValid() const
+	{
+		return !DeviceDescription.IsEmpty();
+	}
+	
+	bool IsAMD() const { return VendorId == 0x1002; }
+
+	bool IsIntel() const { return VendorId == 0x8086; }
+
+	bool IsNVIDIA() const { return VendorId == 0x10DE; }
+};
\ No newline at end of file
diff --git a/Engine/Source/Runtime/Core/Public/HAL/IConsoleManager.h b/Engine/Source/Runtime/Core/Public/HAL/IConsoleManager.h
--- a/Engine/Source/Runtime/Core/Public/HAL/IConsoleManager.h
+++ b/Engine/Source/Runtime/Core/Public/HAL/IConsoleManager.h
@@ -2,6 +2,8 @@
 
 #pragma once
 
+#define TRACK_CONSOLE_FIND_COUNT !(UE_BUILD_SHIPPING || UE_BUILD_TEST)
+
 
 /**
  * Console variable usage guide:
@@ -129,6 +131,12 @@ class IConsoleObject
 
 public:
 
+	IConsoleObject()
+#if TRACK_CONSOLE_FIND_COUNT
+		: FindCallCount(0)
+#endif
+	{}
+
 	virtual ~IConsoleObject() {}
 
 	/**
@@ -200,6 +208,11 @@ public:
 
 private: // -----------------------------------------
 
+#if TRACK_CONSOLE_FIND_COUNT
+	// no longer pure visual, if that causes problems we can change the interface
+	// to track down FindConsoleObject/FindConsoleVariable calls without static
+	uint32 FindCallCount;
+#endif
 
 	/**
 	 *  should only be called by the manager, needs to be implemented for each instance
diff --git a/Engine/Source/Runtime/Core/Public/HAL/IPlatformFileCachedWrapper.h b/Engine/Source/Runtime/Core/Public/HAL/IPlatformFileCachedWrapper.h
--- a/Engine/Source/Runtime/Core/Public/HAL/IPlatformFileCachedWrapper.h
+++ b/Engine/Source/Runtime/Core/Public/HAL/IPlatformFileCachedWrapper.h
@@ -234,7 +234,7 @@ public:
 
 		// Allow a choice between shorter load times or less memory on desktop platforms.
 		// Note: this cannot be in config since they aren't read at that point.
-		if (PLATFORM_DESKTOP)
+#if (PLATFORM_DESKTOP || PLATFORM_PS4)
 		{
 			if (FParse::Param(CmdLine, TEXT("NoCachedReadFile")))
 			{
@@ -247,7 +247,7 @@ public:
 
 			UE_LOG(LogPlatformFile, Log, TEXT("%s cached read wrapper"), bResult ? TEXT("Using") : TEXT("Not using"));
 		}
-
+#endif
 		return bResult;
 	}
 	IPlatformFile* GetLowerLevel() override
diff --git a/Engine/Source/Runtime/Core/Public/HAL/MallocStomp.h b/Engine/Source/Runtime/Core/Public/HAL/MallocStomp.h
--- a/Engine/Source/Runtime/Core/Public/HAL/MallocStomp.h
+++ b/Engine/Source/Runtime/Core/Public/HAL/MallocStomp.h
@@ -6,10 +6,17 @@
 #error "WITH_EDITOR must be defined"
 #endif
 
-#if !WITH_EDITOR && PLATFORM_DESKTOP
+#ifndef USE_MALLOC_STOMP
+	#if !WITH_EDITOR && PLATFORM_DESKTOP
 	#define USE_MALLOC_STOMP 0
+	#endif
 #endif
 
+#ifndef USE_MALLOC_STOMP
+#define USE_MALLOC_STOMP 0
+#endif
+
+
 #if USE_MALLOC_STOMP
 /**
  * Stomp memory allocator. It helps find the following errors:
diff --git a/Engine/Source/Runtime/Core/Public/Internationalization/Text.h b/Engine/Source/Runtime/Core/Public/Internationalization/Text.h
--- a/Engine/Source/Runtime/Core/Public/Internationalization/Text.h
+++ b/Engine/Source/Runtime/Core/Public/Internationalization/Text.h
@@ -171,27 +171,17 @@ public:
 	static FText AsNumber(uint64 Val,	const FNumberFormattingOptions* const Options = NULL, const FCulturePtr& TargetCulture = NULL);
 	static FText AsNumber(long Val,		const FNumberFormattingOptions* const Options = NULL, const FCulturePtr& TargetCulture = NULL);
 
-	DEPRECATED(4.11, "Use FText::AsCurrencyBase instead")
+
 	static FText AsCurrency(float Val,  const FString& CurrencyCode = FString(), const FNumberFormattingOptions* const Options = NULL, const FCulturePtr& TargetCulture = NULL);
-	DEPRECATED(4.11, "Use FText::AsCurrencyBase instead")
 	static FText AsCurrency(double Val, const FString& CurrencyCode = FString(), const FNumberFormattingOptions* const Options = NULL, const FCulturePtr& TargetCulture = NULL);
-	DEPRECATED(4.11, "Use FText::AsCurrencyBase instead")
 	static FText AsCurrency(int8 Val,   const FString& CurrencyCode = FString(), const FNumberFormattingOptions* const Options = NULL, const FCulturePtr& TargetCulture = NULL);
-	DEPRECATED(4.11, "Use FText::AsCurrencyBase instead")
 	static FText AsCurrency(int16 Val,  const FString& CurrencyCode = FString(), const FNumberFormattingOptions* const Options = NULL, const FCulturePtr& TargetCulture = NULL);
-	DEPRECATED(4.11, "Use FText::AsCurrencyBase instead")
 	static FText AsCurrency(int32 Val,  const FString& CurrencyCode = FString(), const FNumberFormattingOptions* const Options = NULL, const FCulturePtr& TargetCulture = NULL);
-	DEPRECATED(4.11, "Use FText::AsCurrencyBase instead")
 	static FText AsCurrency(int64 Val,  const FString& CurrencyCode = FString(), const FNumberFormattingOptions* const Options = NULL, const FCulturePtr& TargetCulture = NULL);
-	DEPRECATED(4.11, "Use FText::AsCurrencyBase instead")
 	static FText AsCurrency(uint8 Val,  const FString& CurrencyCode = FString(), const FNumberFormattingOptions* const Options = NULL, const FCulturePtr& TargetCulture = NULL);
-	DEPRECATED(4.11, "Use FText::AsCurrencyBase instead")
 	static FText AsCurrency(uint16 Val, const FString& CurrencyCode = FString(), const FNumberFormattingOptions* const Options = NULL, const FCulturePtr& TargetCulture = NULL);
-	DEPRECATED(4.11, "Use FText::AsCurrencyBase instead")
 	static FText AsCurrency(uint32 Val, const FString& CurrencyCode = FString(), const FNumberFormattingOptions* const Options = NULL, const FCulturePtr& TargetCulture = NULL);
-	DEPRECATED(4.11, "Use FText::AsCurrencyBase instead")
 	static FText AsCurrency(uint64 Val, const FString& CurrencyCode = FString(), const FNumberFormattingOptions* const Options = NULL, const FCulturePtr& TargetCulture = NULL);
-	DEPRECATED(4.11, "Use FText::AsCurrencyBase instead")
 	static FText AsCurrency(long Val,   const FString& CurrencyCode = FString(), const FNumberFormattingOptions* const Options = NULL, const FCulturePtr& TargetCulture = NULL);
 
 	/**
diff --git a/Engine/Source/Runtime/Core/Public/Linux/LinuxPlatformMisc.h b/Engine/Source/Runtime/Core/Public/Linux/LinuxPlatformMisc.h
--- a/Engine/Source/Runtime/Core/Public/Linux/LinuxPlatformMisc.h
+++ b/Engine/Source/Runtime/Core/Public/Linux/LinuxPlatformMisc.h
@@ -70,6 +70,7 @@ struct CORE_API FLinuxPlatformMisc : public FGenericPlatformMisc
 	static void LowLevelOutputDebugString(const TCHAR *Message);
 	static bool ControlScreensaver(EScreenSaverAction Action);
 
+	static void RequestExitWithStatus(bool Force, uint8 ReturnCode);
 	static const TCHAR* GetSystemErrorMessage(TCHAR* OutBuffer, int32 BufferCount, int32 Error);
 	static void ClipboardCopy(const TCHAR* Str);
 	static void ClipboardPaste(class FString& Dest);
@@ -167,6 +168,15 @@ struct CORE_API FLinuxPlatformMisc : public FGenericPlatformMisc
 	 * Returns whether the program has been started remotely (e.g. over SSH)
 	 */
 	static bool HasBeenStartedRemotely();
+
+	/**
+	 * Determines if return code has been overriden and returns it.
+	 *
+	 * @param OverriddenReturnCodeToUsePtr pointer to an variable that will hold an overriden return code, if any. Can be null.
+	 *
+	 * @return true if the error code has been overriden, false if not
+	 */
+	static bool HasOverriddenReturnCode(uint8 * OverriddenReturnCodeToUsePtr);
 };
 
 typedef FLinuxPlatformMisc FPlatformMisc;
diff --git a/Engine/Source/Runtime/Core/Public/Linux/LinuxPlatformOutputDevices.h b/Engine/Source/Runtime/Core/Public/Linux/LinuxPlatformOutputDevices.h
--- a/Engine/Source/Runtime/Core/Public/Linux/LinuxPlatformOutputDevices.h
+++ b/Engine/Source/Runtime/Core/Public/Linux/LinuxPlatformOutputDevices.h
@@ -11,6 +11,7 @@
 struct CORE_API FLinuxOutputDevices : public FGenericPlatformOutputDevices
 {
 	static void							SetupOutputDevices();
+	static FString						GetAbsoluteLogFilename();
 
 	static FOutputDevice*			GetEventLog();
 	static FOutputDeviceConsole*	GetLogConsole();
diff --git a/Engine/Source/Runtime/Core/Public/Misc/App.h b/Engine/Source/Runtime/Core/Public/Misc/App.h
--- a/Engine/Source/Runtime/Core/Public/Misc/App.h
+++ b/Engine/Source/Runtime/Core/Public/Misc/App.h
@@ -307,6 +307,18 @@ public:
 		return (!IsRunningCommandlet() || IsAllowCommandletRendering()) && !IsRunningDedicatedServer();
 	}
 
+	/**
+	 * Checks whether this application should run with the null RHI.
+	 *
+	 * Distinct from GUsingNullRHI which tells you whether the engine has actually initialized and is running a null RHI.
+	 *
+	 * @return true if the application should use null RHI.
+	 */
+	FORCEINLINE static bool ShouldUseNullRHI()
+	{
+		return (USE_NULL_RHI || FParse::Param(FCommandLine::Get(), TEXT("nullrhi")) || !CanEverRender());
+	}
+
 	/**
 	 * Checks whether this application has been installed.
 	 *
diff --git a/Engine/Source/Runtime/Core/Public/Misc/CallbackDevice.h b/Engine/Source/Runtime/Core/Public/Misc/CallbackDevice.h
--- a/Engine/Source/Runtime/Core/Public/Misc/CallbackDevice.h
+++ b/Engine/Source/Runtime/Core/Public/Misc/CallbackDevice.h
@@ -83,6 +83,10 @@ public:
 	// second param is UserID, third is UserIndex / ControllerId.
 	DECLARE_MULTICAST_DELEGATE_ThreeParams(FOnUserControllerConnectionChange, bool, int32, int32);
 
+	// Callback for platform handling when flushing async loads.
+	DECLARE_MULTICAST_DELEGATE(FOnAsyncLoadingFlush);
+	static FOnAsyncLoadingFlush OnAsyncLoadingFlush;
+
 	// get a hotfix delegate
 	static FHotFixDelegate& GetHotfixDelegate(EHotfixDelegates::Type HotFix);
 
@@ -218,6 +222,11 @@ public:
 	DECLARE_MULTICAST_DELEGATE(FVRHeadsetRecenter);
 	static FVRHeadsetRecenter VRHeadsetRecenter;
 
+	/** Sent when application code changes the user activity hint string for analytics, crash reports, etc */
+	DECLARE_MULTICAST_DELEGATE_OneParam(FOnUserActivityStringChanged, const FString&);
+	static FOnUserActivityStringChanged UserActivityStringChanged;
+
+
 private:
 
 	// Callbacks for hotfixes
diff --git a/Engine/Source/Runtime/Core/Public/Stats/Stats2.h b/Engine/Source/Runtime/Core/Public/Stats/Stats2.h
--- a/Engine/Source/Runtime/Core/Public/Stats/Stats2.h
+++ b/Engine/Source/Runtime/Core/Public/Stats/Stats2.h
@@ -1838,7 +1838,7 @@ struct FStat_##StatName\
 	DECLARE_STAT_GROUP(GroupDesc, GroupId, GroupCat, false, CompileIn);
 
 #if UE_BUILD_DEBUG
-	#define SCOPE_CYCLE_COUNTER_GUARD { const char* ReadTheText = "SCOPE_CYCLE_COUNTER can't be used in the global scope."; static_cast<void const*>(ReadTheText); }
+	#define SCOPE_CYCLE_COUNTER_GUARD { const char* ReadTheText = "SCOPE_CYCLE_COUNTER can't be used in the global scope."; (void)ReadTheText; }
 #else
 	#define SCOPE_CYCLE_COUNTER_GUARD
 #endif // UE_BUILD_DEBUG
diff --git a/Engine/Source/Runtime/CoreUObject/Private/Serialization/AsyncLoading.cpp b/Engine/Source/Runtime/CoreUObject/Private/Serialization/AsyncLoading.cpp
--- a/Engine/Source/Runtime/CoreUObject/Private/Serialization/AsyncLoading.cpp
+++ b/Engine/Source/Runtime/CoreUObject/Private/Serialization/AsyncLoading.cpp
@@ -947,8 +947,9 @@ void FAsyncLoadingThread::CancelAsyncLoading()
 
 void FAsyncLoadingThread::SuspendLoading()
 {
-	check(IsInGameThread());
+	UE_CLOG(!IsInGameThread(), LogStreaming, Fatal, TEXT("Async loading can only be suspended from the main thread"));
 	const int32 SuspendCount = IsLoadingSuspended.Increment();
+	UE_LOG(LogStreaming, Display, TEXT("Suspending async loading (%d)"), SuspendCount);
 	if (IsMultithreaded() && SuspendCount == 1)
 	{
 		ThreadSuspendedEvent->Wait();
@@ -959,6 +960,7 @@ void FAsyncLoadingThread::ResumeLoading()
 {
 	check(IsInGameThread());
 	const int32 SuspendCount = IsLoadingSuspended.Decrement();
+	UE_LOG(LogStreaming, Display, TEXT("Resuming async loading (%d)"), SuspendCount);
 	UE_CLOG(SuspendCount < 0, LogStreaming, Fatal, TEXT("ResumeAsyncLoadingThread: Async loading was resumed more times than it was suspended."));
 	if (IsMultithreaded() && SuspendCount == 0)
 	{
@@ -1377,10 +1379,10 @@ EAsyncPackageState::Type FAsyncPackage::CreateLinker()
 		{
 			const FString NameToLoad = Desc.NameToLoad.ToString();
 			const FGuid* const Guid = Desc.Guid.IsValid() ? &Desc.Guid : nullptr;
-			FString NativeFilename;
-			const bool DoesNativePackageExist = FPackageName::DoesPackageExist(NameToLoad, Guid, &NativeFilename, false);
-			FString LocalizedFilename;
-			const bool DoesLocalizedPackageExist = FPackageName::DoesPackageExist(NameToLoad, Guid, &LocalizedFilename, true);
+			FString NativeFilename, LocalizedFilename;
+			FPackageName::DoesPackageExistWithLocalization(NameToLoad, Guid, &NativeFilename, &LocalizedFilename);
+			const bool DoesNativePackageExist = NativeFilename.Len() > 0;
+			const bool DoesLocalizedPackageExist = LocalizedFilename.Len() > 0;
 
 			FString PackageFileName;
 			bool DoesPackageExist = false;
@@ -2217,7 +2219,7 @@ void FlushAsyncLoading(int32 PackageID /* = INDEX_NONE */)
 	{
 		FAsyncLoadingThread& AsyncThread = FAsyncLoadingThread::Get();
 		// Flushing async loading while loading is suspend will result in infinite stall
-		UE_CLOG(AsyncThread.IsAsyncLoadingSuspended(), LogStreaming, Fatal, TEXT("Cannot Flush Async Loading while async loading is suspended."));
+		UE_CLOG(AsyncThread.IsAsyncLoadingSuspended(), LogStreaming, Fatal, TEXT("Cannot Flush Async Loading while async loading is suspended (%d)"), AsyncThread.GetAsyncLoadingSuspendedCount());
 
 		SCOPE_CYCLE_COUNTER(STAT_FAsyncPackage_FlushAsyncLoadingGameThread);
 
@@ -2226,6 +2228,8 @@ void FlushAsyncLoading(int32 PackageID /* = INDEX_NONE */)
 			return;
 		}
 
+		FCoreDelegates::OnAsyncLoadingFlush.Broadcast();
+
 		// Disallow low priority requests like texture streaming while we are flushing streaming
 		// in order to avoid excessive seeking.
 		FIOSystem::Get().SetMinPriority( AIOP_Normal );
diff --git a/Engine/Source/Runtime/CoreUObject/Private/Serialization/AsyncLoadingThread.h b/Engine/Source/Runtime/CoreUObject/Private/Serialization/AsyncLoadingThread.h
--- a/Engine/Source/Runtime/CoreUObject/Private/Serialization/AsyncLoadingThread.h
+++ b/Engine/Source/Runtime/CoreUObject/Private/Serialization/AsyncLoadingThread.h
@@ -195,6 +195,12 @@ public:
 		return IsLoadingSuspended.GetValue() != 0;
 	}
 
+	FORCEINLINE int32 GetAsyncLoadingSuspendedCount()
+	{
+		FPlatformMisc::MemoryBarrier();
+		return IsLoadingSuspended.GetValue();
+	}
+
 	/** Returns the number of async packages that are currently being processed */
 	FORCEINLINE int32 GetAsyncPackagesCount()
 	{
diff --git a/Engine/Source/Runtime/CoreUObject/Private/UObject/Class.cpp b/Engine/Source/Runtime/CoreUObject/Private/UObject/Class.cpp
--- a/Engine/Source/Runtime/CoreUObject/Private/UObject/Class.cpp
+++ b/Engine/Source/Runtime/CoreUObject/Private/UObject/Class.cpp
@@ -3762,6 +3762,10 @@ void UClass::PurgeClass(bool bRecompilingOnLoad)
 	ClassUnique = 0;
 	ClassReps.Empty();
 	NetFields.Empty();
+	for (TObjectIterator<UPackage> PackageIt; PackageIt; ++PackageIt)
+	{
+		PackageIt->ClassUniqueNameIndexMap.Remove(GetFName());
+	}
 
 #if WITH_EDITOR
 	if (!bRecompilingOnLoad)
diff --git a/Engine/Source/Runtime/CoreUObject/Private/UObject/LazyObjectPtr.cpp b/Engine/Source/Runtime/CoreUObject/Private/UObject/LazyObjectPtr.cpp
--- a/Engine/Source/Runtime/CoreUObject/Private/UObject/LazyObjectPtr.cpp
+++ b/Engine/Source/Runtime/CoreUObject/Private/UObject/LazyObjectPtr.cpp
@@ -132,7 +132,7 @@ void FLazyObjectPtr::PossiblySerializeObjectGuid(UObject *Object, FArchive& Ar)
 					if (OtherObject != NULL)
 					{
 						// IsGame returns true for GIsPlayInEditorWorld
-						UE_CLOG(!((FApp::IsGame() || IsRunningCommandlet()) && Package && Package->ContainsMap()), LogUObjectGlobals, Warning, TEXT("Guid is in use by %s and %s, which should never happen in the editor but could happen at runtime with duplicate level loading or PIE"), *OtherObject->GetFullName(), !!Object ? *Object->GetFullName() : TEXT("NULL"));
+						UE_CLOG(!(FApp::IsGame() && Package && Package->ContainsMap()), LogUObjectGlobals, Warning, TEXT("Guid referenced by %s is already used by %s, which should never happen in the editor but could happen at runtime with duplicate level loading or PIE"), !!Object ? *Object->GetFullName() : TEXT("NULL"), *OtherObject->GetFullName());
 						// This guid is in use, which should never happen in the editor but could happen at runtime with duplicate level loading or PIE. If so give it an invalid GUID and don't add to the annotation map.
 						Guid = FGuid();
 					}
diff --git a/Engine/Source/Runtime/CoreUObject/Private/UObject/LinkerLoad.cpp b/Engine/Source/Runtime/CoreUObject/Private/UObject/LinkerLoad.cpp
--- a/Engine/Source/Runtime/CoreUObject/Private/UObject/LinkerLoad.cpp
+++ b/Engine/Source/Runtime/CoreUObject/Private/UObject/LinkerLoad.cpp
@@ -4127,17 +4127,77 @@ UObject* FLinkerLoad::IndexToObject( FPackageIndex Index )
 {
 	if( Index.IsExport() )
 	{
-		check(ExportMap.IsValidIndex( Index.ToExport() ) );
+		#if PLATFORM_DESKTOP
+			// Show a message box indicating, possible, corrupt data (desktop platforms only)
+			if ( !ExportMap.IsValidIndex( Index.ToExport() ) )
+			{
+				FText ErrorMessage, ErrorCaption;
+				GConfig->GetText(TEXT("/Script/Engine.Engine"),
+									TEXT("SerializationOutOfBoundsErrorMessage"),
+									ErrorMessage,
+									GEngineIni);
+				GConfig->GetText(TEXT("/Script/Engine.Engine"),
+					TEXT("SerializationOutOfBoundsErrorMessageCaption"),
+					ErrorCaption,
+					GEngineIni);
+
+				UE_LOG( LogLinker, Error, TEXT("Invalid export object index=%d while reading %s. File is most likely corrupted. Please verify your installation."), Index.ToExport(), *Filename );
+
+				if (GLog)
+				{
+					GLog->Flush();
+				}
+
+				FPlatformMisc::MessageBoxExt(EAppMsgType::Ok, *ErrorMessage.ToString(), *ErrorCaption.ToString());
+
+				check(false);
+			}
+		#else
+			{
+				UE_CLOG( !ExportMap.IsValidIndex( Index.ToExport() ), LogLinker, Fatal, TEXT("Invalid export object index=%d while reading %s. File is most likely corrupted. Please verify your installation."), Index.ToExport(), *Filename );
+			}
+		#endif
+
 		return CreateExport( Index.ToExport() );
 	}
 	else if( Index.IsImport() )
 	{
-		check(ImportMap.IsValidIndex( Index.ToImport() ) );
+		#if PLATFORM_DESKTOP
+			// Show a message box indicating, possible, corrupt data (desktop platforms only)
+			if ( !ImportMap.IsValidIndex( Index.ToImport() ) )
+			{
+				FText ErrorMessage, ErrorCaption;
+				GConfig->GetText(TEXT("/Script/Engine.Engine"),
+									TEXT("SerializationOutOfBoundsErrorMessage"),
+									ErrorMessage,
+									GEngineIni);
+				GConfig->GetText(TEXT("/Script/Engine.Engine"),
+					TEXT("SerializationOutOfBoundsErrorMessageCaption"),
+					ErrorCaption,
+					GEngineIni);
+
+				UE_LOG( LogLinker, Error, TEXT("Invalid import object index=%d while reading %s. File is most likely corrupted. Please verify your installation."), Index.ToImport(), *Filename );
+
+				if (GLog)
+				{
+					GLog->Flush();
+				}
+
+				FPlatformMisc::MessageBoxExt(EAppMsgType::Ok, *ErrorMessage.ToString(), *ErrorCaption.ToString());
+
+				check(false);
+			}
+		#else
+			{
+				UE_CLOG( !ImportMap.IsValidIndex( Index.ToImport() ), LogLinker, Fatal, TEXT("Invalid import object index=%d while reading %s. File is most likely corrupted. Please verify your installation."), Index.ToImport(), *Filename );
+			}
+		#endif
+
 		return CreateImport( Index.ToImport() );
 	}
 	else 
 	{
-		return NULL;
+		return nullptr;
 	}
 }
 
diff --git a/Engine/Source/Runtime/CoreUObject/Private/UObject/SavePackage.cpp b/Engine/Source/Runtime/CoreUObject/Private/UObject/SavePackage.cpp
--- a/Engine/Source/Runtime/CoreUObject/Private/UObject/SavePackage.cpp
+++ b/Engine/Source/Runtime/CoreUObject/Private/UObject/SavePackage.cpp
@@ -408,7 +408,7 @@ bool IsEditorOnlyObject(const UObject* InObject)
 	{
 		return false;
 	}
-
+	
 	if (InObject->HasAnyMarks(OBJECTMARK_EditorOnly))
 	{
 		return true;
@@ -3205,8 +3205,19 @@ ESavePackageResult UPackage::Save(UPackage* InOuter, UObject* Base, EObjectFlags
 #if WITH_EDITORONLY_DATA
 		if (bIsCooking)
 		{
+			static struct FCanSkipEditorReferencedPackagesWhenCooking
+			{
+				bool bCanSkipEditorReferencedPackagesWhenCooking;
+				FCanSkipEditorReferencedPackagesWhenCooking()
+					: bCanSkipEditorReferencedPackagesWhenCooking(true)
+				{
+					GConfig->GetBool(TEXT("Core.System"), TEXT("CanSkipEditorReferencedPackagesWhenCooking"), bCanSkipEditorReferencedPackagesWhenCooking, GEngineIni);
+				}
+				FORCEINLINE operator bool() const { return bCanSkipEditorReferencedPackagesWhenCooking; }
+			} CanSkipEditorReferencedPackagesWhenCooking;
+
 			// Don't save packages marked as editor-only.
-			if (InOuter->IsLoadedByEditorPropertiesOnly())
+			if (CanSkipEditorReferencedPackagesWhenCooking && InOuter->IsLoadedByEditorPropertiesOnly())
 			{				
 				UE_CLOG(!(SaveFlags & SAVE_NoError), LogSavePackage, Display, TEXT("Package loaded by editor-only properties: %s. Package will not be saved."), *InOuter->GetName());
 				return ESavePackageResult::ReferencedOnlyByEditorOnlyData;
diff --git a/Engine/Source/Runtime/CoreUObject/Private/UObject/UObjectGlobals.cpp b/Engine/Source/Runtime/CoreUObject/Private/UObject/UObjectGlobals.cpp
--- a/Engine/Source/Runtime/CoreUObject/Private/UObject/UObjectGlobals.cpp
+++ b/Engine/Source/Runtime/CoreUObject/Private/UObject/UObjectGlobals.cpp
@@ -1597,7 +1597,18 @@ FName MakeUniqueObjectName( UObject* Parent, UClass* Class, FName InBaseName/*=N
 				}
 				else
 				{
-					TestName = FName(BaseName, ++Class->ClassUnique);
+					int32 NameNumber = 0;
+					if (Parent && (Parent != ANY_PACKAGE) )
+					{
+						UPackage* ParentPackage = Parent->GetOutermost();
+						int32& ClassUnique = ParentPackage->ClassUniqueNameIndexMap.FindOrAdd(Class->GetFName());
+						NameNumber = ++ClassUnique;
+					}
+					else
+					{
+						NameNumber = ++Class->ClassUnique;
+					}
+					TestName = FName(BaseName, NameNumber);
 				}
 
 				if (Parent == ANY_PACKAGE)
diff --git a/Engine/Source/Runtime/CoreUObject/Public/UObject/CoreObject.h b/Engine/Source/Runtime/CoreUObject/Public/UObject/CoreObject.h
--- a/Engine/Source/Runtime/CoreUObject/Public/UObject/CoreObject.h
+++ b/Engine/Source/Runtime/CoreUObject/Public/UObject/CoreObject.h
@@ -164,6 +164,8 @@ public:
 	// World browser information
 	TScopedPointer< FWorldTileInfo > WorldTileInfo;
 
+	TMap<FName, int32> ClassUniqueNameIndexMap;
+
 #if WITH_EDITOR
 	/** Editor only: PIE instance ID this package belongs to, INDEX_NONE otherwise */
 	int32 PIEInstanceID;
diff --git a/Engine/Source/Runtime/Engine/Classes/AI/Navigation/RecastNavMesh.h b/Engine/Source/Runtime/Engine/Classes/AI/Navigation/RecastNavMesh.h
--- a/Engine/Source/Runtime/Engine/Classes/AI/Navigation/RecastNavMesh.h
+++ b/Engine/Source/Runtime/Engine/Classes/AI/Navigation/RecastNavMesh.h
@@ -644,11 +644,13 @@ public:
 		int32 CorridorPolysCount;
 		float HitTime;
 		FVector HitNormal;
+		uint32 bIsRaycastEndInCorridor : 1;
 
 		FRaycastResult()
 			: CorridorPolysCount(0)
 			, HitTime(FLT_MAX)
 			, HitNormal(0.f)
+			, bIsRaycastEndInCorridor(false)
 		{
 			FMemory::Memzero(CorridorPolys);
 			FMemory::Memzero(CorridorCost);
diff --git a/Engine/Source/Runtime/Engine/Classes/Animation/AnimationAsset.h b/Engine/Source/Runtime/Engine/Classes/Animation/AnimationAsset.h
--- a/Engine/Source/Runtime/Engine/Classes/Animation/AnimationAsset.h
+++ b/Engine/Source/Runtime/Engine/Classes/Animation/AnimationAsset.h
@@ -52,6 +52,9 @@ struct FBlendSampleData
 	UPROPERTY()
 	int32 SampleDataIndex;
 
+	UPROPERTY()
+	class UAnimSequence* Animation;
+
 	UPROPERTY()
 	float TotalWeight;
 
@@ -74,6 +77,7 @@ struct FBlendSampleData
 	{}
 	FBlendSampleData(int32 Index)
 		:	SampleDataIndex(Index)
+		,	Animation(nullptr)
 		,	TotalWeight(0.f)
 		,	Time(0.f)
 		,	PreviousTime(0.f)
diff --git a/Engine/Source/Runtime/Engine/Classes/Camera/CameraAnim.h b/Engine/Source/Runtime/Engine/Classes/Camera/CameraAnim.h
--- a/Engine/Source/Runtime/Engine/Classes/Camera/CameraAnim.h
+++ b/Engine/Source/Runtime/Engine/Classes/Camera/CameraAnim.h
@@ -4,6 +4,8 @@
 
 #include "CameraAnim.generated.h"
 
+ENGINE_API DECLARE_LOG_CATEGORY_EXTERN(LogCameraAnim, Log, All);
+
 /**
  * A predefined animation to be played on a camera
  */
diff --git a/Engine/Source/Runtime/Engine/Classes/Camera/CameraAnimInst.h b/Engine/Source/Runtime/Engine/Classes/Camera/CameraAnimInst.h
--- a/Engine/Source/Runtime/Engine/Classes/Camera/CameraAnimInst.h
+++ b/Engine/Source/Runtime/Engine/Classes/Camera/CameraAnimInst.h
@@ -29,30 +29,25 @@ public:
 	/** Current time for the animation */
 	float CurTime;
 
+	/** True if the animation has finished, false otherwise. */
+	uint32 bFinished : 1;
+
 protected:
 	/** True if the animation should loop, false otherwise. */
 	uint32 bLooping:1;
 
-public:
-	/** True if the animation has finished, false otherwise. */
-	uint32 bFinished:1;
+	/** True if currently blending in. */
+	uint32 bBlendingIn : 1;
 
-	/** True if it's ok for the system to auto-release this instance upon completion. */
-	uint32 bAutoReleaseWhenFinished:1;
+	/** True if currently blending out. */
+	uint32 bBlendingOut : 1;
 
-protected:
 	/** Time to interpolate in from zero, for smooth starts. */
 	float BlendInTime;
 
 	/** Time to interpolate out to zero, for smooth finishes. */
 	float BlendOutTime;
 
-	/** True if currently blending in. */
-	uint32 bBlendingIn:1;
-
-	/** True if currently blending out. */
-	uint32 bBlendingOut:1;
-
 	/** Current time for the blend-in.  I.e. how long we have been blending. */
 	float CurBlendInTime;
 
@@ -77,7 +72,6 @@ protected:
 	/** How much longer to play the anim, if a specific duration is desired.  Has no effect if 0.  */
 	float RemainingTime;
 
-
 public:
 	/** cached movement track from the currently playing anim so we don't have to go find it every frame */
 	UPROPERTY(transient)
@@ -140,7 +134,6 @@ public:
 	UFUNCTION(BlueprintCallable, Category = CameraAnimInst)
 	void SetScale(float NewDuration);
 
-
 protected:
 	/** Returns InterpGroupInst subobject **/
 	class UInterpGroupInst* GetInterpGroupInst() const;
diff --git a/Engine/Source/Runtime/Engine/Classes/Engine/EngineTypes.h b/Engine/Source/Runtime/Engine/Classes/Engine/EngineTypes.h
--- a/Engine/Source/Runtime/Engine/Classes/Engine/EngineTypes.h
+++ b/Engine/Source/Runtime/Engine/Classes/Engine/EngineTypes.h
@@ -4056,13 +4056,13 @@ struct FMaterialMergeData
 /**
  * The description of a user activity
  */
-USTRUCT()
+USTRUCT(BlueprintType)
 struct FUserActivity
 {
 	GENERATED_USTRUCT_BODY()
 
 	/** Describes the user's activity */
-	UPROPERTY(EditAnywhere, Category = "Activity")
+	UPROPERTY(BlueprintReadWrite, Category = "Activity")
 	FString ActionName;
 
 	/** Default constructor. */
diff --git a/Engine/Source/Runtime/Engine/Classes/Engine/LocalPlayer.h b/Engine/Source/Runtime/Engine/Classes/Engine/LocalPlayer.h
--- a/Engine/Source/Runtime/Engine/Classes/Engine/LocalPlayer.h
+++ b/Engine/Source/Runtime/Engine/Classes/Engine/LocalPlayer.h
@@ -190,7 +190,6 @@ private:
 public:
 	// UObject interface
 	virtual void PostInitProperties() override;
-	virtual void BeginDestroy() override;
 	virtual void FinishDestroy() override;
 	static void AddReferencedObjects(UObject* InThis, FReferenceCollector& Collector);
 	// End of UObject interface
@@ -228,14 +227,6 @@ protected:
 	/** @todo document */
 	void ExecMacro( const TCHAR* Filename, FOutputDevice& Ar );
 
-	/**
-	 * Handle controller connection events from the platform
-	 * @param bConnected - determines connection or disconnection event
-	 * @param UserId - platform specfic UserId
-	 * @param UserIndex - Engine level UserIndex.
-	 */
-	void HandleControllerConnectionChange(bool bConnected, int32 InUserId, int32 InControllerId);
-
 	/** FReply used to defer some slate operations. */
 	FReply SlateOperations;
 
diff --git a/Engine/Source/Runtime/Engine/Classes/Engine/NetConnection.h b/Engine/Source/Runtime/Engine/Classes/Engine/NetConnection.h
--- a/Engine/Source/Runtime/Engine/Classes/Engine/NetConnection.h
+++ b/Engine/Source/Runtime/Engine/Classes/Engine/NetConnection.h
@@ -179,6 +179,9 @@ class UNetConnection : public UPlayer
 	/** Number of bits used for padding in the current packet. */
 	int NumPaddingBits;
 
+	/** The maximum number of bits all packet handlers will reserve */
+	int32 MaxPacketHandlerBits;
+
 	/** Sets all of the bit-tracking variables to zero. */
 	void ResetPacketBitCounts();
 
@@ -268,12 +271,14 @@ public:
 	int32 InPacketsLost, OutPacketsLost;
 
 	// Packet.
-	FBitWriter		SendBuffer;				// Queued up bits waiting to send
-	double			OutLagTime[256];		// For lag measuring.
-	int32			OutLagPacketId[256];	// For lag measuring.
-	int32			InPacketId;				// Full incoming packet index.
-	int32			OutPacketId;			// Most recently sent packet.
-	int32 			OutAckPacketId;			// Most recently acked outgoing packet.
+	FBitWriter		SendBuffer;						// Queued up bits waiting to send
+	double			OutLagTime[256];				// For lag measuring.
+	int32			OutLagPacketId[256];			// For lag measuring.
+	int32			OutBytesPerSecondHistory[256];	// For saturation measuring.
+	float			RemoteSaturation;
+	int32			InPacketId;						// Full incoming packet index.
+	int32			OutPacketId;					// Most recently sent packet.
+	int32 			OutAckPacketId;					// Most recently acked outgoing packet.
 
 	bool			bLastHasServerFrameTime;
 
@@ -515,6 +520,12 @@ public:
 	ENGINE_API virtual void InitConnection(UNetDriver* InDriver, EConnectionState InState, const FURL& InURL, int32 InConnectionSpeed=0, int32 InMaxPacket=0);
 
 
+	/**
+	 * Initializes the PacketHandler
+	 */
+	ENGINE_API virtual void InitHandler();
+
+
 	/** 
 	* Gets a unique ID for the connection, this ID depends on the underlying connection
 	* For IP connections this is an IP Address and port, for steam this is a SteamID
@@ -568,6 +579,12 @@ public:
 	/** Send a raw bunch. */
 	ENGINE_API int32 SendRawBunch( FOutBunch& Bunch, bool InAllowMerge );
 
+	/** The maximum number of bits allowed within a single bunch. */
+	FORCEINLINE int32 GetMaxSingleBunchSizeBits() const
+	{
+		return (MaxPacket * 8) - MAX_BUNCH_HEADER_BITS - MAX_PACKET_TRAILER_BITS - MAX_PACKET_HEADER_BITS - MaxPacketHandlerBits;
+	}
+
 	/** @return The driver object */
 	UNetDriver* GetDriver() {return Driver;}
 
diff --git a/Engine/Source/Runtime/Engine/Classes/GameFramework/GameUserSettings.h b/Engine/Source/Runtime/Engine/Classes/GameFramework/GameUserSettings.h
--- a/Engine/Source/Runtime/Engine/Classes/GameFramework/GameUserSettings.h
+++ b/Engine/Source/Runtime/Engine/Classes/GameFramework/GameUserSettings.h
@@ -296,6 +296,10 @@ protected:
 	UPROPERTY(config)
 	int32 LastConfirmedFullscreenMode;
 
+	/** Fullscreen mode to use when toggling between windowed and fullscreen. Same values as r.FullScreenMode. */
+	UPROPERTY(config)
+	int32 PreferredFullscreenMode;
+
 	/** All settings will be wiped and set to default if the serialized version differs from UE_GAMEUSERSETTINGS_VERSION. */
 	UPROPERTY(config)
 	uint32 Version;
@@ -322,4 +326,6 @@ protected:
 private:
 
 	void UpdateMinResolutionScaling();
+
+	void SetPreferredFullscreenMode(int32 Mode);
 };
diff --git a/Engine/Source/Runtime/Engine/Classes/GameFramework/PlayerInput.h b/Engine/Source/Runtime/Engine/Classes/GameFramework/PlayerInput.h
--- a/Engine/Source/Runtime/Engine/Classes/GameFramework/PlayerInput.h
+++ b/Engine/Source/Runtime/Engine/Classes/GameFramework/PlayerInput.h
@@ -510,6 +510,8 @@ public:
 	/** Returns the list of keys mapped to the specified Action Name */
 	const TArray<FInputActionKeyMapping>& GetKeysForAction(const FName ActionName);
 
+	static const TArray<FInputActionKeyMapping>& GetEngineDefinedActionMappings() { return EngineDefinedActionMappings; }
+
 private:
 	/** 
 	 * Given raw keystate value, returns the "massaged" value. Override for any custom behavior,
diff --git a/Engine/Source/Runtime/Engine/Classes/Kismet/KismetSystemLibrary.h b/Engine/Source/Runtime/Engine/Classes/Kismet/KismetSystemLibrary.h
--- a/Engine/Source/Runtime/Engine/Classes/Kismet/KismetSystemLibrary.h
+++ b/Engine/Source/Runtime/Engine/Classes/Kismet/KismetSystemLibrary.h
@@ -1371,6 +1371,13 @@ class ENGINE_API UKismetSystemLibrary : public UBlueprintFunctionLibrary
 	UFUNCTION(BlueprintPure, Category="Rendering", meta=(UnsafeDuringActorConstruction = "true"))
 	static int32 GetMinYResolutionForUI();
 
+	/**
+	* Gets the smallest Y resolution we want to support in the 3D view, clamped within reasons
+	* @return value in pixels
+	*/
+	UFUNCTION(BlueprintPure, Category = "Rendering", meta = (UnsafeDuringActorConstruction = "true"))
+	static int32 GetMinYResolutionFor3DView();
+
 	// Opens the specified URL in the platform's web browser of choice
 	UFUNCTION(BlueprintCallable, Category = "Utilities|Platform")
 	static void LaunchURL(const FString& URL);
diff --git a/Engine/Source/Runtime/Engine/Private/AI/Navigation/PImplRecastNavMesh.cpp b/Engine/Source/Runtime/Engine/Private/AI/Navigation/PImplRecastNavMesh.cpp
--- a/Engine/Source/Runtime/Engine/Private/AI/Navigation/PImplRecastNavMesh.cpp
+++ b/Engine/Source/Runtime/Engine/Private/AI/Navigation/PImplRecastNavMesh.cpp
@@ -749,6 +749,9 @@ void FPImplRecastNavMesh::Raycast(const FVector& StartLoc, const FVector& EndLoc
 		NavQuery.findNearestContainingPoly(&RecastStart.X, Extent, QueryFilter, &StartNode, NULL);
 	}
 
+	NavNodeRef EndNode = INVALID_NAVNODEREF;
+	NavQuery.findNearestContainingPoly(&RecastEnd.X, Extent, QueryFilter, &EndNode, NULL);
+
 	if (StartNode != INVALID_NAVNODEREF)
 	{
 		float RecastHitNormal[3];
@@ -758,20 +761,7 @@ void FPImplRecastNavMesh::Raycast(const FVector& StartLoc, const FVector& EndLoc
 			, RaycastResult.CorridorPolys, &RaycastResult.CorridorPolysCount, RaycastResult.GetMaxCorridorSize());
 
 		RaycastResult.HitNormal = Recast2UnrVector(RecastHitNormal);
-
-		if (dtStatusSucceed(RaycastStatus) == false && !RaycastResult.HasHit())
-		{
-			NavNodeRef EndNode = INVALID_NAVNODEREF;
-			NavQuery.findNearestPoly(&RecastEnd.X, Extent, QueryFilter, &EndNode, NULL);
-
-			// if raycast corridor ends in different poly, ray was probably stopped by 2D edge check
-			// return blocking hit at end
-			if (EndNode != RaycastResult.GetLastNodeRef())
-			{
-				RaycastResult.HitTime = 1.f;
-				RaycastResult.HitNormal = (StartLoc - EndLoc).GetSafeNormal();
-			}
-		}
+		RaycastResult.bIsRaycastEndInCorridor = dtStatusSucceed(RaycastStatus) && (RaycastResult.GetLastNodeRef() == EndNode);
 	}
 	else
 	{
diff --git a/Engine/Source/Runtime/Engine/Private/AI/Navigation/RecastNavMesh.cpp b/Engine/Source/Runtime/Engine/Private/AI/Navigation/RecastNavMesh.cpp
--- a/Engine/Source/Runtime/Engine/Private/AI/Navigation/RecastNavMesh.cpp
+++ b/Engine/Source/Runtime/Engine/Private/AI/Navigation/RecastNavMesh.cpp
@@ -1983,7 +1983,7 @@ void ARecastNavMesh::BatchRaycast(TArray<FNavigationRaycastWork>& Workload, FSha
 	const FVector NavExtent = GetModifiedQueryExtent(GetDefaultQueryExtent());
 	const float Extent[3] = { NavExtent.X, NavExtent.Z, NavExtent.Y };
 
-	for (auto& WorkItem : Workload)
+	for (FNavigationRaycastWork& WorkItem : Workload)
 	{
 		ARecastNavMesh::FRaycastResult RaycastResult;
 
@@ -2001,18 +2001,10 @@ void ARecastNavMesh::BatchRaycast(TArray<FNavigationRaycastWork>& Workload, FSha
 				, QueryFilter, &RaycastResult.HitTime, RecastHitNormal
 				, RaycastResult.CorridorPolys, &RaycastResult.CorridorPolysCount, RaycastResult.GetMaxCorridorSize());
 
-			RaycastResult.HitNormal = Recast2UnrealPoint(RecastHitNormal);
-
 			if (dtStatusSucceed(RaycastStatus) && RaycastResult.HasHit())
 			{
-				NavNodeRef ExpectedEndNode = INVALID_NAVNODEREF;
-				NavQuery.findNearestPoly(&RecastEnd.X, Extent, QueryFilter, &ExpectedEndNode, NULL);
-
-				if (ExpectedEndNode == RaycastResult.GetLastNodeRef())
-				{
-					WorkItem.bDidHit = true;
-					WorkItem.HitLocation = FNavLocation(WorkItem.RayStart + (WorkItem.RayEnd - WorkItem.RayStart) * RaycastResult.HitTime, ExpectedEndNode);
-				}
+				WorkItem.bDidHit = true;
+				WorkItem.HitLocation = FNavLocation(WorkItem.RayStart + (WorkItem.RayEnd - WorkItem.RayStart) * RaycastResult.HitTime, RaycastResult.GetLastNodeRef());
 			}
 		}
 	}
@@ -2028,7 +2020,7 @@ bool ARecastNavMesh::IsSegmentOnNavmesh(const FVector& SegmentStart, const FVect
 	FRaycastResult Result;
 	RecastNavMeshImpl->Raycast(SegmentStart, SegmentEnd, GetRightFilterRef(Filter), QueryOwner, Result);
 
-	return Result.HasHit() == false;
+	return Result.bIsRaycastEndInCorridor && !Result.HasHit();
 }
 
 bool ARecastNavMesh::FindStraightPath(const FVector& StartLoc, const FVector& EndLoc, const TArray<NavNodeRef>& PathCorridor, TArray<FNavPathPoint>& PathPoints, TArray<uint32>* CustomLinks) const
diff --git a/Engine/Source/Runtime/Engine/Private/Animation/AnimInstance.cpp b/Engine/Source/Runtime/Engine/Private/Animation/AnimInstance.cpp
--- a/Engine/Source/Runtime/Engine/Private/Animation/AnimInstance.cpp
+++ b/Engine/Source/Runtime/Engine/Private/Animation/AnimInstance.cpp
@@ -326,6 +326,7 @@ void UAnimInstance::UpdateMontage(float DeltaSeconds)
 					SyncGroup->TestMontageTickRecordForLeadership();
 				}
 			}
+			MontageInstance->bDidUseMarkerSyncThisTick = false;
 		}
 	}
 
@@ -473,10 +474,15 @@ void UAnimInstance::ParallelUpdateAnimation()
 
 bool UAnimInstance::NeedsImmediateUpdate(float DeltaSeconds) const
 {
-	const bool bUseParallelUpdateAnimation = bCanUseParallelUpdateAnimation || CVarForceUseParallelAnimUpdate.GetValueOnGameThread() != 0;
+	// If Evaluation Phase is skipped, PostUpdateAnimation() will not get called, so we can't use ParallelUpdateAnimation then.
+	USkeletalMeshComponent* SkelMeshComp = GetSkelMeshComponent();
+	const bool bEvaluationPhaseSkipped = SkelMeshComp && !SkelMeshComp->bRecentlyRendered && (SkelMeshComp->MeshComponentUpdateFlag > EMeshComponentUpdateFlag::AlwaysTickPoseAndRefreshBones);
+
+	const bool bUseParallelUpdateAnimation = bCanUseParallelUpdateAnimation || (CVarForceUseParallelAnimUpdate.GetValueOnGameThread() != 0);
 
 	return
 		GIntraFrameDebuggingGameThread ||
+		bEvaluationPhaseSkipped || 
 		CVarUseParallelAnimUpdate.GetValueOnGameThread() == 0 ||
 		CVarUseParallelAnimationEvaluation.GetValueOnGameThread() == 0 ||
 		GetWorld()->IsServer() ||
diff --git a/Engine/Source/Runtime/Engine/Private/Animation/AnimMontage.cpp b/Engine/Source/Runtime/Engine/Private/Animation/AnimMontage.cpp
--- a/Engine/Source/Runtime/Engine/Private/Animation/AnimMontage.cpp
+++ b/Engine/Source/Runtime/Engine/Private/Animation/AnimMontage.cpp
@@ -1012,6 +1012,11 @@ ENGINE_API void UAnimMontage::UpdateLinkableElements(int32 SlotIdx, int32 Segmen
 
 #endif
 
+FString MakePositionMessage(const FMarkerSyncAnimPosition& Position)
+{
+	return FString::Printf(TEXT("Names(PrevName: %s | NextName: %s) PosBetweenMarkers: %.2f"), *Position.PreviousMarkerName.ToString(), *Position.NextMarkerName.ToString(), Position.PositionBetweenMarkers);
+}
+
 void UAnimMontage::TickAssetPlayer(FAnimTickRecord& Instance, struct FAnimNotifyQueue& NotifyQueue, FAnimAssetTickContext& Context) const
 {
 	// nothing has to happen here
@@ -1057,8 +1062,8 @@ void UAnimMontage::TickAssetPlayer(FAnimTickRecord& Instance, struct FAnimNotify
 				{
 					const FMarkerSyncAnimPosition& StartPosition = MarkerTickContext.GetMarkerSyncStartPosition();
 					const FMarkerSyncAnimPosition& EndPosition = MarkerTickContext.GetMarkerSyncEndPosition();
-					checkf(StartPosition.NextMarkerName == EndPosition.NextMarkerName, TEXT("StartPosition NextMarker = %s, EndPosition NextMarker = %s, Asset = %s"), *StartPosition.NextMarkerName.ToString(), *EndPosition.NextMarkerName.ToString(), *Instance.SourceAsset->GetFullName());
-					checkf(StartPosition.PreviousMarkerName == EndPosition.PreviousMarkerName, TEXT("StartPosition PreviousMarker = %s, EndPosition PreviousMarker = %s, Asset = %s"), *StartPosition.PreviousMarkerName.ToString(), *EndPosition.PreviousMarkerName.ToString(), *Instance.SourceAsset->GetFullName());
+					checkf(StartPosition.NextMarkerName == EndPosition.NextMarkerName, TEXT("StartPosition %s\nEndPosition %s\nPrevTime to CurrentTimeAsset: %.3f - %.3f Delta: %.3f\nAsset = %s"), *MakePositionMessage(StartPosition), *MakePositionMessage(EndPosition), PreviousTime, CurrentTime, MoveDelta, *Instance.SourceAsset->GetFullName());
+					checkf(StartPosition.PreviousMarkerName == EndPosition.PreviousMarkerName, TEXT("StartPosition %s\nEndPosition %s\nPrevTime - CurrentTimeAsset: %.3f - %.3f Delta: %.3f\nAsset = %s"), *MakePositionMessage(StartPosition), *MakePositionMessage(EndPosition), PreviousTime, CurrentTime, MoveDelta, *Instance.SourceAsset->GetFullName());
 				}
 #endif
 
diff --git a/Engine/Source/Runtime/Engine/Private/Animation/AnimSequence.cpp b/Engine/Source/Runtime/Engine/Private/Animation/AnimSequence.cpp
--- a/Engine/Source/Runtime/Engine/Private/Animation/AnimSequence.cpp
+++ b/Engine/Source/Runtime/Engine/Private/Animation/AnimSequence.cpp
@@ -606,8 +606,23 @@ void UAnimSequence::PostLoad()
 		}
 	}
 
+	USkeleton* CurrentSkeleton = GetSkeleton();
+
+	// if valid additive, but if base additive isn't 
+	// this seems to happen from retargeting sometimes, which we still have to investigate why, 
+	// but this causes issue since once this happens this is unrecoverable until you delete from outside of editor
+	if (IsValidAdditive())
+	{
+		if (RefPoseSeq && RefPoseSeq->GetSkeleton() != CurrentSkeleton)
+		{
+			// if this happens, there was a issue with retargeting, 
+			UE_LOG(LogAnimation, Warning, TEXT("Animation %s - Invalid additive animation base animation (%s)"), *GetName(), *RefPoseSeq->GetName());
+			RefPoseSeq = nullptr;
+		}
+	}
+
 #if WITH_EDITOR
-	if (USkeleton * CurrentSkeleton = GetSkeleton())
+	if (CurrentSkeleton)
 	{
 		// Get the name mapping object for curves
 		const FSmartNameMapping* NameMapping = CurrentSkeleton->GetSmartNameContainer(USkeleton::AnimTrackCurveMappingName);
diff --git a/Engine/Source/Runtime/Engine/Private/Animation/BlendSpaceBase.cpp b/Engine/Source/Runtime/Engine/Private/Animation/BlendSpaceBase.cpp
--- a/Engine/Source/Runtime/Engine/Private/Animation/BlendSpaceBase.cpp
+++ b/Engine/Source/Runtime/Engine/Private/Animation/BlendSpaceBase.cpp
@@ -301,9 +301,11 @@ void UBlendSpaceBase::TickAssetPlayer(FAnimTickRecord& Instance, struct FAnimNot
 				{
 					for (FBlendSampleData& CurrentBlendSampleItem : SampleDataList)
 					{
-						if (PrevBlendSampleItem.SampleDataIndex == CurrentBlendSampleItem.SampleDataIndex)
+						// it only can have one animation in the sample, make sure to copy Time
+						if (PrevBlendSampleItem.Animation && PrevBlendSampleItem.Animation == CurrentBlendSampleItem.Animation)
 						{
 							CurrentBlendSampleItem.Time = PrevBlendSampleItem.Time;
+							CurrentBlendSampleItem.PreviousTime = PrevBlendSampleItem.PreviousTime;
 							CurrentBlendSampleItem.MarkerTickRecord = PrevBlendSampleItem.MarkerTickRecord;
 						}
 					}
@@ -508,12 +510,14 @@ bool UBlendSpaceBase::GetSamplesFromBlendInput(const FVector &BlendInput, TArray
 		float GridWeight = GridSample.BlendWeight;
 		FEditorElement& GridElement = GridSample.GridElement;
 
-		for(int Ind = 0; Ind < GridElement.MAX_VERTICES; ++Ind)
+		for(int32 Ind = 0; Ind < GridElement.MAX_VERTICES; ++Ind)
 		{
-			if(GridElement.Indices[Ind] != INDEX_NONE)
+			const int32 SampleDataIndex = GridElement.Indices[Ind];		
+			if(SampleDataIndex != INDEX_NONE)
 			{
-				int32 Index = OutSampleDataList.AddUnique(GridElement.Indices[Ind]);
+				int32 Index = OutSampleDataList.AddUnique(SampleDataIndex);
 				OutSampleDataList[Index].AddWeight(GridElement.Weights[Ind]*GridWeight);
+				OutSampleDataList[Index].Animation = SampleData[SampleDataIndex].Animation;
 			}
 		}
 	}
@@ -521,9 +525,26 @@ bool UBlendSpaceBase::GetSamplesFromBlendInput(const FVector &BlendInput, TArray
 	/** Used to sort by  Weight. */
 	struct FCompareFBlendSampleData
 	{
-		FORCEINLINE bool operator()( const FBlendSampleData& A, const FBlendSampleData& B ) const { return B.TotalWeight < A.TotalWeight; }
+		FORCEINLINE bool operator()(const FBlendSampleData& A, const FBlendSampleData& B) const { return B.TotalWeight < A.TotalWeight; }
 	};
-	OutSampleDataList.Sort( FCompareFBlendSampleData() );
+	OutSampleDataList.Sort(FCompareFBlendSampleData());
+
+	// go through merge down to first sample 
+	for (int32 Index1 = 0; Index1 < OutSampleDataList.Num(); ++Index1)
+	{
+		for (int32 Index2 = Index1 + 1; Index2 < OutSampleDataList.Num(); ++Index2)
+		{
+			// if they have sample sample, remove the Index2, and get out
+			if (OutSampleDataList[Index1].Animation == OutSampleDataList[Index2].Animation)
+			{
+				// add weight
+				OutSampleDataList[Index1].AddWeight(OutSampleDataList[Index2].GetWeight());
+				// as for time or previous time will be the master one(Index1)
+				OutSampleDataList.RemoveAtSwap(Index2, 1, false);
+				--Index2;
+			}
+		}
+	}
 
 	// remove noisy ones
 	int32 TotalSample = OutSampleDataList.Num();
diff --git a/Engine/Source/Runtime/Engine/Private/Camera/CameraAnim.cpp b/Engine/Source/Runtime/Engine/Private/Camera/CameraAnim.cpp
--- a/Engine/Source/Runtime/Engine/Private/Camera/CameraAnim.cpp
+++ b/Engine/Source/Runtime/Engine/Private/Camera/CameraAnim.cpp
@@ -12,6 +12,7 @@
 #include "Matinee/InterpGroupCamera.h"
 #include "Matinee/InterpTrackMove.h"
 
+DEFINE_LOG_CATEGORY(LogCameraAnim);
 
 //////////////////////////////////////////////////////////////////////////
 // UCameraAnim
diff --git a/Engine/Source/Runtime/Engine/Private/Camera/CameraAnimInst.cpp b/Engine/Source/Runtime/Engine/Private/Camera/CameraAnimInst.cpp
--- a/Engine/Source/Runtime/Engine/Private/Camera/CameraAnimInst.cpp
+++ b/Engine/Source/Runtime/Engine/Private/Camera/CameraAnimInst.cpp
@@ -18,7 +18,6 @@ UCameraAnimInst::UCameraAnimInst(const FObjectInitializer& ObjectInitializer)
 	: Super(ObjectInitializer)
 {
 	bFinished = true;
-	bAutoReleaseWhenFinished = true;
 	PlayRate = 1.0f;
 	TransientScaleModifier = 1.0f;
 	PlaySpace = ECameraAnimPlaySpace::CameraLocal;
@@ -93,13 +92,16 @@ void UCameraAnimInst::AdvanceAnim(float DeltaTime, bool bJump)
 		
 		// calculate blend weight. calculating separately and taking the minimum handles overlapping blends nicely.
 		{
-			float BlendInWeight = (bBlendingIn) ? (CurBlendInTime / BlendInTime) : 1.f;
-			float BlendOutWeight = (bBlendingOut) ? (1.f - CurBlendOutTime / BlendOutTime) : 1.f;
+			float const BlendInWeight = (bBlendingIn) ? (CurBlendInTime / BlendInTime) : 1.f;
+			float const BlendOutWeight = (bBlendingOut) ? (1.f - CurBlendOutTime / BlendOutTime) : 1.f;
 			CurrentBlendWeight = FMath::Min(BlendInWeight, BlendOutWeight) * BasePlayScale * TransientScaleModifier;
 		}
 
 		// this will update tracks and apply the effects to the group actor (except move tracks)
-		InterpGroupInst->Group->UpdateGroup(CurTime, InterpGroupInst, false, bJump);
+		if (InterpGroupInst->Group)
+		{
+			InterpGroupInst->Group->UpdateGroup(CurTime, InterpGroupInst, false, bJump);
+		}
 
 		if (bAnimJustFinished)
 		{
@@ -121,40 +123,48 @@ void UCameraAnimInst::AdvanceAnim(float DeltaTime, bool bJump)
 
 void UCameraAnimInst::Update(float NewRate, float NewScale, float NewBlendInTime, float NewBlendOutTime, float NewDuration)
 {
-	if (bBlendingOut)
+	if (bFinished == false)
 	{
-		bBlendingOut = false;
-		CurBlendOutTime = 0.f;
+		if (bBlendingOut)
+		{
+			bBlendingOut = false;
+			CurBlendOutTime = 0.f;
 
-		// stop any blendout and reverse it to a blendin
-		bBlendingIn = true;
-		CurBlendInTime = NewBlendInTime * (1.f - CurBlendOutTime / BlendOutTime);
-	}
+			// stop any blendout and reverse it to a blendin
+			bBlendingIn = true;
+			CurBlendInTime = NewBlendInTime * (1.f - CurBlendOutTime / BlendOutTime);
+		}
 
-	PlayRate = NewRate;
-	BasePlayScale = NewScale;
-	BlendInTime = NewBlendInTime;
-	BlendOutTime = NewBlendOutTime;
-	RemainingTime = (NewDuration > 0.f) ? (NewDuration - BlendOutTime) : 0.f;
-	bFinished = false;
+		PlayRate = NewRate;
+		BasePlayScale = NewScale;
+		BlendInTime = NewBlendInTime;
+		BlendOutTime = NewBlendOutTime;
+		RemainingTime = (NewDuration > 0.f) ? (NewDuration - BlendOutTime) : 0.f;
+	}
 }
 
 void UCameraAnimInst::SetDuration(float NewDuration)
 {
-	// setting a new duration will reset the play timer back to 0 but maintain current playback position
-	// if blending out, stop it and blend back in to reverse it so it's smooth
-	if (bBlendingOut)
+	if (bFinished == false)
 	{
-		bBlendingOut = false;
-		CurBlendOutTime = 0.f;
+		// setting a new duration will reset the play timer back to 0 but maintain current playback position
+		// if blending out, stop it and blend back in to reverse it so it's smooth
+		if (bBlendingOut)
+		{
+			bBlendingOut = false;
+			CurBlendOutTime = 0.f;
 
-		// stop any blendout and reverse it to a blendin
-		bBlendingIn = true;
-		CurBlendInTime = BlendInTime * (1.f - CurBlendOutTime / BlendOutTime);
-	}
+			// stop any blendout and reverse it to a blendin
+			bBlendingIn = true;
+			CurBlendInTime = BlendInTime * (1.f - CurBlendOutTime / BlendOutTime);
+		}
 
-	RemainingTime = (NewDuration > 0.f) ? (NewDuration - BlendOutTime) : 0.f;
-	bFinished = false;
+		RemainingTime = (NewDuration > 0.f) ? (NewDuration - BlendOutTime) : 0.f;
+	}
+	else
+	{
+		UE_LOG(LogCameraAnim, Warning, TEXT("SetDuration called for CameraAnim %s after it finished. Ignored."), *GetNameSafe(CamAnim));
+	}
 }
 
 void UCameraAnimInst::SetScale(float NewScale)
@@ -246,13 +256,13 @@ void UCameraAnimInst::Stop(bool bImmediate)
 
 	if ( bImmediate || (BlendOutTime <= 0.f) )
 	{
-		if (InterpGroupInst->Group != NULL)
+		if (InterpGroupInst->Group != nullptr)
 		{
 			InterpGroupInst->TermGroupInst(true);
 			InterpGroupInst->Group = nullptr;
 		}
-		MoveTrack = NULL;
-		MoveInst = NULL;
+		MoveTrack = nullptr;
+		MoveInst = nullptr;
 		bFinished = true;
  	}
 	else
diff --git a/Engine/Source/Runtime/Engine/Private/Components/ChildActorComponent.cpp b/Engine/Source/Runtime/Engine/Private/Components/ChildActorComponent.cpp
--- a/Engine/Source/Runtime/Engine/Private/Components/ChildActorComponent.cpp
+++ b/Engine/Source/Runtime/Engine/Private/Components/ChildActorComponent.cpp
@@ -336,7 +336,8 @@ void UChildActorComponent::DestroyChildActor(const bool bRequiresRename)
 
 				// We would like to make certain that our name is not going to accidentally get taken from us while we're destroyed
 				// so we increment ClassUnique beyond our index to be certain of it.  This is ... a bit hacky.
-				ChildClass->ClassUnique = FMath::Max(ChildClass->ClassUnique, ChildActor->GetFName().GetNumber());
+				int32& ClassUnique = ChildActor->GetOutermost()->ClassUniqueNameIndexMap.FindOrAdd(ChildClass->GetFName());
+				ClassUnique = FMath::Max(ClassUnique, ChildActor->GetFName().GetNumber());
 
 				if (bRequiresRename)
 				{
diff --git a/Engine/Source/Runtime/Engine/Private/ContentStreaming.cpp b/Engine/Source/Runtime/Engine/Private/ContentStreaming.cpp
--- a/Engine/Source/Runtime/Engine/Private/ContentStreaming.cpp
+++ b/Engine/Source/Runtime/Engine/Private/ContentStreaming.cpp
@@ -49,6 +49,7 @@ DEFINE_STAT(STAT_TotalForcedHeuristicSize);
 DEFINE_STAT(STAT_LightmapMemorySize);
 DEFINE_STAT(STAT_LightmapDiskSize);
 DEFINE_STAT(STAT_HLODTextureMemorySize);
+DEFINE_STAT(STAT_HLODTextureDiskSize);
 DEFINE_STAT(STAT_IntermediateTexturesSize);
 DEFINE_STAT(STAT_RequestSizeCurrentFrame);
 DEFINE_STAT(STAT_RequestSizeTotal);
@@ -72,7 +73,6 @@ DEFINE_STAT(STAT_DynamicStreamingTotal);
 /** Accumulated total time spent on dynamic primitives, in seconds. */
 double GStreamingDynamicPrimitivesTime = 0.0;
 
-
 /*-----------------------------------------------------------------------------
 	Globals.
 -----------------------------------------------------------------------------*/
@@ -178,6 +178,32 @@ ENGINE_API TAutoConsoleVariable<int32> CVarStreamingUseAABB(
 	TEXT("If non-zero, will use AABB to compute distance."),
 	ECVF_Default);
 
+static TAutoConsoleVariable<int32> CVarStreamingHLODStrategy(
+	TEXT("r.Streaming.HLODStrategy"),
+	0,
+	TEXT("Define the HLOD streaming strategy.\n")
+	TEXT("0: stream\n")
+	TEXT("1: stream only mip 0\n")
+	TEXT("2: disable streaming"),
+	ECVF_Default);
+
+static TAutoConsoleVariable<float> CVarStreamingHLODPriorityScale(
+	TEXT("r.Streaming.HLODPriorityScale"),
+	1.5,
+	TEXT("Define the streaming priority to apply to HLOD textures\n"),
+	ECVF_Default
+	);
+
+static TAutoConsoleVariable<float> CVarStreamingHiddenPrimitiveScale(
+	TEXT("r.Streaming.HiddenPrimitiveScale"),
+	0.5,
+	TEXT("Define the resolution scale to apply when not in range.\n")
+	TEXT(".5: drop one mip\n")
+	TEXT("1: ignore visiblity"),
+	ECVF_Default
+	);
+
+
 /** Streaming priority: Linear distance factor from 0 to MAX_STREAMINGDISTANCE. */
 #define MAX_STREAMINGDISTANCE	10000.0f
 #define MAX_MIPDELTA			5.0f
@@ -289,6 +315,7 @@ struct FStreamingTexture
 		NumNonStreamingMips = InTexture->GetNumNonStreamingMips();
 		ForceLoadRefCount = 0;
 		bIsStreamingLightmap = IsStreamingLightmap( Texture );
+		bIsLightmap = bIsStreamingLightmap || LODGroup == TEXTUREGROUP_Lightmap || LODGroup == TEXTUREGROUP_Shadowmap;
 		bUsesStaticHeuristics = false;
 		bUsesDynamicHeuristics = false;
 		bUsesLastRenderHeuristics = false;
@@ -423,6 +450,13 @@ struct FStreamingTexture
 			float MipFactor = float(WantedMips) / float(MAX_TEXTURE_MIP_COUNT);
 			float TimeFactor = GStreamWithTimeFactor ? (FMath::Clamp( LastRenderTime, 1.0f, MAX_LASTRENDERTIME ) / MAX_LASTRENDERTIME) : 0.0f;
 			float Priority = MipFactor + DistanceFactor * (1.0f - TimeFactor * 0.5f) + bForceFullyLoad * 100.0f;
+
+			// Custom priority control for HLOD textures.
+			if (bIsHLODTexture)
+			{
+				Priority *= CVarStreamingHLODPriorityScale.GetValueOnAnyThread();
+			}
+
 			return Priority;
 		}
 		else
@@ -508,6 +542,8 @@ struct FStreamingTexture
 	uint32			bInFlight : 1;
 	/** Whether this is a streaming lightmap or shadowmap (cached from IsStreamingLightmap()). */
 	uint32			bIsStreamingLightmap : 1;
+	/** Whether this is a lightmap or shadowmap. */
+	uint32			bIsLightmap : 1;
 	/** Whether this texture uses StaticTexture heuristics. */
 	uint32			bUsesStaticHeuristics : 1;
 	/** Whether this texture uses DynamicTexture heuristics. */
@@ -952,6 +988,7 @@ struct FStreamingContext
 		ThisFrameTotalLightmapMemorySize				= 0;
 		ThisFrameTotalLightmapDiskSize					= 0;
 		ThisFrameTotalHLODMemorySize					= 0;
+		ThisFrameTotalHLODDiskSize						= 0;
 		ThisFrameTotalMipCountIncreaseRequestsInFlight	= 0;
 		ThisFrameOptimalWantedSize						= 0;
 		ThisFrameTotalStaticTextureHeuristicSize		= 0;
@@ -1012,6 +1049,7 @@ struct FStreamingContext
 		ThisFrameTotalLightmapMemorySize				+= Other.ThisFrameTotalLightmapMemorySize;
 		ThisFrameTotalLightmapDiskSize					+= Other.ThisFrameTotalLightmapDiskSize;
 		ThisFrameTotalHLODMemorySize					+= Other.ThisFrameTotalHLODMemorySize;
+		ThisFrameTotalHLODDiskSize						+= Other.ThisFrameTotalHLODDiskSize;
 		ThisFrameTotalMipCountIncreaseRequestsInFlight	+= Other.ThisFrameTotalMipCountIncreaseRequestsInFlight;
 		ThisFrameOptimalWantedSize						+= Other.ThisFrameOptimalWantedSize;
 		ThisFrameTotalStaticTextureHeuristicSize		+= Other.ThisFrameTotalStaticTextureHeuristicSize;
@@ -1052,6 +1090,7 @@ struct FStreamingContext
 	uint64 ThisFrameTotalLightmapMemorySize;
 	uint64 ThisFrameTotalLightmapDiskSize;
 	uint64 ThisFrameTotalHLODMemorySize;
+	uint64 ThisFrameTotalHLODDiskSize;
 	uint64 ThisFrameTotalMipCountIncreaseRequestsInFlight;
 	uint64 ThisFrameOptimalWantedSize;
 	/** Number of bytes using StaticTexture heuristics this frame, currently in memory. */
@@ -2123,7 +2162,7 @@ FStreamingManagerTexture::FStreamingManagerTexture()
 ,	MaxTempMemoryUsed( 5*1024*1024 )
 ,	bUseDynamicStreaming( false )
 ,	BoostPlayerTextures( 3.0f )
-,	RangePrefetchDistance(400.f)
+,	RangePrefetchDistance(1000.f)
 ,	MemoryMargin(0)
 ,	MinEvictSize(0)
 ,	IndividualStreamingTexture(NULL)
@@ -2138,6 +2177,7 @@ FStreamingManagerTexture::FStreamingManagerTexture()
 ,	TotalLightmapMemorySize(0)
 ,	TotalLightmapDiskSize(0)
 ,	TotalHLODMemorySize(0)
+,	TotalHLODDiskSize(0)
 ,	TotalMipCountIncreaseRequestsInFlight(0)
 ,	TotalOptimalWantedSize(0)
 ,	TotalStaticTextureHeuristicSize(0)
@@ -3223,6 +3263,7 @@ void FStreamingManagerTexture::ResetStreamingStats()
 	TotalLightmapMemorySize								= 0;
 	TotalLightmapDiskSize								= 0;
 	TotalHLODMemorySize									= 0;
+	TotalHLODDiskSize									= 0;
 	TotalMipCountIncreaseRequestsInFlight				= 0;
 	TotalOptimalWantedSize								= 0;
 	TotalStaticTextureHeuristicSize						= 0;
@@ -3251,6 +3292,7 @@ void FStreamingManagerTexture::UpdateStreamingStats( const FStreamingContext& Co
 	TotalLightmapMemorySize					+= Context.ThisFrameTotalLightmapMemorySize;
 	TotalLightmapDiskSize					+= Context.ThisFrameTotalLightmapDiskSize;
 	TotalHLODMemorySize						+= Context.ThisFrameTotalHLODMemorySize;
+	TotalHLODDiskSize						+= Context.ThisFrameTotalHLODDiskSize;
 	TotalMipCountIncreaseRequestsInFlight	+= Context.ThisFrameTotalMipCountIncreaseRequestsInFlight;
 	TotalOptimalWantedSize					+= Context.ThisFrameOptimalWantedSize;
 	TotalStaticTextureHeuristicSize			+= Context.ThisFrameTotalStaticTextureHeuristicSize;
@@ -3275,6 +3317,8 @@ void FStreamingManagerTexture::UpdateStreamingStats( const FStreamingContext& Co
 		SET_DWORD_STAT( STAT_LightmapMemorySize,			TotalLightmapMemorySize			);
 		SET_DWORD_STAT( STAT_LightmapDiskSize,				TotalLightmapDiskSize			);
 		SET_DWORD_STAT( STAT_HLODTextureMemorySize,			TotalHLODMemorySize				);
+		SET_DWORD_STAT( STAT_HLODTextureDiskSize,			TotalHLODDiskSize				);
+		
 		SET_FLOAT_STAT( STAT_StreamingBandwidth,			BandwidthAverage/1024.0f/1024.0f);
 		SET_FLOAT_STAT( STAT_StreamingLatency,				LatencyAverage					);
 		SET_MEMORY_STAT( STAT_StreamingTexturesSize,		TotalStreamingTexturesSize		);
@@ -4082,6 +4126,12 @@ void FStreamingManagerTexture::CalcMinMaxMips( FStreamingTexture& StreamingTextu
 		}
 	}
 
+	const int32 HLODStategy = CVarStreamingHLODStrategy.GetValueOnAnyThread();
+	if (StreamingTexture.bIsHLODTexture && HLODStategy == 2) // disable streaming
+	{
+		StreamingTexture.bForceFullyLoad = true;
+	}
+
 	// Calculate the min/max allowed mips.
 	UTexture2D::CalcAllowedMips(
 		StreamingTexture.MipCount,
@@ -4109,6 +4159,11 @@ void FStreamingManagerTexture::CalcMinMaxMips( FStreamingTexture& StreamingTextu
 		StreamingTexture.MinAllowedMips = FMath::Clamp( StreamingTexture.MipCount - ThreadSettings.NumStreamedMips[StreamingTexture.LODGroup], StreamingTexture.MinAllowedMips, StreamingTexture.MaxAllowedMips );
 	}
 
+	if ( StreamingTexture.bIsHLODTexture && HLODStategy == 1 && StreamingTexture.MaxAllowedMips > StreamingTexture.MinAllowedMips ) // stream only mip 0
+	{
+		StreamingTexture.MinAllowedMips = StreamingTexture.MaxAllowedMips - 1;
+	}
+
 	check( StreamingTexture.MinAllowedMips > 0 && StreamingTexture.MinAllowedMips <= StreamingTexture.MipCount );
 	check( StreamingTexture.MaxAllowedMips >= StreamingTexture.MinAllowedMips && StreamingTexture.MaxAllowedMips <= StreamingTexture.MipCount );
 }
@@ -4158,9 +4213,10 @@ void FStreamingManagerTexture::UpdateFrameStats( FStreamingContext& Context, FSt
 	}
 	Context.ThisFrameNumStreamingTextures++;
 	Context.ThisFrameTotalStreamingTexturesMaxSize += PotentialSize;
-	Context.ThisFrameTotalLightmapMemorySize += (StreamingTexture.bIsStreamingLightmap || StreamingTexture.LODGroup == TEXTUREGROUP_Lightmap || StreamingTexture.LODGroup == TEXTUREGROUP_Shadowmap) ? ResidentSize : 0;
-	Context.ThisFrameTotalLightmapDiskSize += StreamingTexture.bIsStreamingLightmap ? PotentialSize : 0;
+	Context.ThisFrameTotalLightmapMemorySize += StreamingTexture.bIsLightmap ? ResidentSize : 0;
+	Context.ThisFrameTotalLightmapDiskSize += StreamingTexture.bIsLightmap ? PotentialSize : 0;
 	Context.ThisFrameTotalHLODMemorySize += StreamingTexture.bIsHLODTexture ? ResidentSize : 0;
+	Context.ThisFrameTotalHLODDiskSize += StreamingTexture.bIsHLODTexture ? PotentialSize : 0;
 #endif
 }
 
@@ -5383,6 +5439,8 @@ FFloatMipLevel FStreamingHandlerTextureStatic::GetWantedMips( FStreamingManagerT
 	bool bEntryFound = false; // True an entry for this texture exists­.
 
 	const bool bUseAABB = CVarStreamingUseAABB.GetValueOnAnyThread() > 0;
+	const float HiddenScale = CVarStreamingHiddenPrimitiveScale.GetValueOnAnyThread();
+	const float MaxResolution = (float)(0x1 << (StreamingTexture.MaxAllowedMips - 1));
 
 	// Nothing do to if there are no views or instances.
 	if( StreamingManager.ThreadNumViews() /*&& StreamingTexture.Instances.Num() > 0*/ )
@@ -5448,16 +5506,11 @@ FFloatMipLevel FStreamingHandlerTextureStatic::GetWantedMips( FStreamingManagerT
 						Temp = VectorSubtract( ViewOriginZ, CenterZ );
 						DistSq = VectorMultiplyAdd( Temp, Temp, DistSq );
 
-						// If distance is not in range, replace it with FLT_MAX
-						{
-							// ClampedDistSq = clamp(DistSq, TextureInstance.MinDistanceSq, TextureInstance.MaxDistanceSq);
-							VectorRegister ClampedDistSq = VectorMax( VectorLoadAligned( &TextureInstance.MinDistanceSq ), DistSq );
-							ClampedDistSq = VectorMin( VectorLoadAligned( &TextureInstance.MaxDistanceSq ), ClampedDistSq );
-
-							// DistSq = (DistSq == ClampedDistSq) ? DistSq : FLT_MAX
-							const VectorRegister InRangeMask = VectorCompareEQ(DistSq, ClampedDistSq);
-							DistSq = VectorSelect(InRangeMask, DistSq, MaxFloat4);
-						}
+						// The range is handle as the distance to the center of the bounds (and not to the bounds), 
+						// because it needs to match how HLOD visibility is handled.
+						VectorRegister ClampedDistSq = VectorMax( VectorLoadAligned( &TextureInstance.MinDistanceSq ), DistSq );
+						ClampedDistSq = VectorMin( VectorLoadAligned( &TextureInstance.MaxDistanceSq ), ClampedDistSq );
+						const VectorRegister InRangeMask = VectorCompareEQ(DistSq, ClampedDistSq);
 
 						VectorRegister DistSqMinusRadiusSq = VectorZero();
 						if (bUseAABB)
@@ -5488,6 +5541,10 @@ FFloatMipLevel FStreamingHandlerTextureStatic::GetWantedMips( FStreamingManagerT
 							DistSqMinusRadiusSq = VectorMultiply( DistSqMinusRadiusSq, DistSqMinusRadiusSq );
 							DistSqMinusRadiusSq = VectorSubtract( DistSq, DistSqMinusRadiusSq );
 						}
+						// Force at least distance of 2 for entry not in their range, in order to enter in next conditional block.
+						// Also required in order to not enter in the final block computing the WantedMipCount.
+						DistSqMinusRadiusSq = VectorSelect( InRangeMask, DistSqMinusRadiusSq, VectorMax(DistSqMinusRadiusSq, VectorSet(2.f, 2.f, 2.f, 2.f) ) );
+
 						MinDistanceSq4 = VectorMin( MinDistanceSq4, DistSqMinusRadiusSq );
 
 						if ( VectorAllGreaterThan( DistSqMinusRadiusSq, VectorOne() ) )
@@ -5499,6 +5556,12 @@ FFloatMipLevel FStreamingHandlerTextureStatic::GetWantedMips( FStreamingManagerT
 								DistSqMinusRadiusSq = VectorMax( DistSqMinusRadiusSq, VectorOne() );
 								VectorRegister Texels = VectorMultiply( VectorLoadAligned( &TextureInstance.TexelFactor ), VectorReciprocalSqrt(DistSqMinusRadiusSq) );
 								Texels = VectorMultiply( Texels, ScreenSize );
+
+								// Clamp to max resolution before actually scaling, otherwise this could have not effect.
+								Texels = VectorMin(VectorSet(MaxResolution, MaxResolution, MaxResolution, MaxResolution), Texels);
+								VectorRegister VisibilityScale = VectorSelect( InRangeMask, VectorOne(), VectorSet(HiddenScale, HiddenScale, HiddenScale, HiddenScale) );
+								Texels = VectorMultiply(Texels, VisibilityScale);
+
 								MaxTexels = VectorMax( MaxTexels, Texels );
 							}
 							else
diff --git a/Engine/Source/Runtime/Engine/Private/DataBunch.cpp b/Engine/Source/Runtime/Engine/Private/DataBunch.cpp
--- a/Engine/Source/Runtime/Engine/Private/DataBunch.cpp
+++ b/Engine/Source/Runtime/Engine/Private/DataBunch.cpp
@@ -88,7 +88,7 @@ FOutBunch::FOutBunch()
 : FNetBitWriter( 0 )
 {}
 FOutBunch::FOutBunch( UChannel* InChannel, bool bInClose )
-:	FNetBitWriter	( InChannel->Connection->PackageMap, InChannel->Connection->MaxPacket*8-MAX_BUNCH_HEADER_BITS-MAX_PACKET_TRAILER_BITS-MAX_PACKET_HEADER_BITS )
+:	FNetBitWriter	( InChannel->Connection->PackageMap, InChannel->Connection->GetMaxSingleBunchSizeBits())
 ,	Next		( NULL )
 ,	Channel		( InChannel )
 ,	Time		( 0 )
diff --git a/Engine/Source/Runtime/Engine/Private/DataChannel.cpp b/Engine/Source/Runtime/Engine/Private/DataChannel.cpp
--- a/Engine/Source/Runtime/Engine/Private/DataChannel.cpp
+++ b/Engine/Source/Runtime/Engine/Private/DataChannel.cpp
@@ -724,7 +724,7 @@ FPacketIdRange UChannel::SendBunch( FOutBunch* Bunch, bool Merge )
 	check(!OpenTemporary || !Bunch->bReliable);
 
 	// This is the max number of bits we can have in a single bunch
-	const int64 MAX_SINGLE_BUNCH_SIZE_BITS  = Connection->MaxPacket*8-MAX_BUNCH_HEADER_BITS-MAX_PACKET_TRAILER_BITS-MAX_PACKET_HEADER_BITS;
+	const int64 MAX_SINGLE_BUNCH_SIZE_BITS  = Connection->GetMaxSingleBunchSizeBits();
 
 	// Max bytes we'll put in a partial bunch
 	const int64 MAX_SINGLE_BUNCH_SIZE_BYTES = MAX_SINGLE_BUNCH_SIZE_BITS / 8;
diff --git a/Engine/Source/Runtime/Engine/Private/DemoNetDriver.cpp b/Engine/Source/Runtime/Engine/Private/DemoNetDriver.cpp
--- a/Engine/Source/Runtime/Engine/Private/DemoNetDriver.cpp
+++ b/Engine/Source/Runtime/Engine/Private/DemoNetDriver.cpp
@@ -579,7 +579,8 @@ bool UDemoNetDriver::InitConnectInternal( FString& Error )
 			}
 		}
 
-		World->ServerTravel( DemoHeader.LevelName, true, false );
+		// FIXME: Test for failure!!!
+		World->SeamlessTravel( DemoHeader.LevelName, true );
 	}
 	else
 	{
diff --git a/Engine/Source/Runtime/Engine/Private/GameUserSettings.cpp b/Engine/Source/Runtime/Engine/Private/GameUserSettings.cpp
--- a/Engine/Source/Runtime/Engine/Private/GameUserSettings.cpp
+++ b/Engine/Source/Runtime/Engine/Private/GameUserSettings.cpp
@@ -159,9 +159,15 @@ void UGameUserSettings::SetToDefaults()
 	FrameRateLimit = 0.0f;
 	MinResolutionScale = Scalability::MinResolutionScale;
 
+	static const auto CVar = IConsoleManager::Get().FindTConsoleVariableDataInt(TEXT("r.FullScreenMode"));
+	PreferredFullscreenMode = CVar->GetValueOnGameThread();
+
 	ScalabilityQuality.SetDefaults();
 
-	UpdateMinResolutionScaling();
+	if (!IsRunningDedicatedServer())
+	{
+		UpdateMinResolutionScaling();
+	}
 }
 
 bool UGameUserSettings::IsVersionValid()
@@ -176,12 +182,23 @@ void UGameUserSettings::UpdateVersion()
 
 void UGameUserSettings::UpdateMinResolutionScaling()
 {
-	const int32 MinHeight = UKismetSystemLibrary::GetMinYResolutionForUI();
+	const int32 MinHeight = UKismetSystemLibrary::GetMinYResolutionFor3DView();
 	const int32 ScreenHeight = FullscreenMode == EWindowMode::WindowedFullscreen ? GetDesktopResolution().Y : ResolutionSizeY;
 	MinResolutionScale = FMath::Max<int32>(Scalability::MinResolutionScale, ((float)MinHeight / (float)ScreenHeight) * 100);
 	ScalabilityQuality.ResolutionQuality = FMath::Max(ScalabilityQuality.ResolutionQuality, MinResolutionScale);
 }
 
+void UGameUserSettings::SetPreferredFullscreenMode(int32 Mode)
+{
+	PreferredFullscreenMode = Mode;
+
+	auto CVar = IConsoleManager::Get().FindConsoleVariable(TEXT("r.FullScreenMode"));
+	if (CVar)
+	{
+		CVar->Set(Mode, ECVF_SetByGameSetting);
+	}
+}
+
 void UGameUserSettings::ValidateSettings()
 {
 	// Should we wipe all user settings?
@@ -273,6 +290,12 @@ void UGameUserSettings::ApplyResolutionSettings(bool bCheckForCommandLineOverrid
 
 	// Request a resolution change
 	RequestResolutionChange(ResolutionSizeX, ResolutionSizeY, NewFullscreenMode, bCheckForCommandLineOverrides);
+
+	if (NewFullscreenMode == EWindowMode::Fullscreen || NewFullscreenMode == EWindowMode::WindowedFullscreen)
+	{
+		SetPreferredFullscreenMode(NewFullscreenMode == EWindowMode::Fullscreen ? 0 : 1);
+	}
+
 	IConsoleManager::Get().CallAllConsoleVariableSinks();
 }
 
@@ -304,6 +327,9 @@ void UGameUserSettings::LoadSettings( bool bForceReload/*=false*/ )
 	{
 		ConfirmVideoMode();
 	}
+
+	// Update r.FullScreenMode CVar
+	SetPreferredFullscreenMode(PreferredFullscreenMode);
 }
 
 void UGameUserSettings::RequestResolutionChange(int32 InResolutionX, int32 InResolutionY, EWindowMode::Type InWindowMode, bool bInDoOverrides /* = true */)
diff --git a/Engine/Source/Runtime/Engine/Private/KismetSystemLibrary.cpp b/Engine/Source/Runtime/Engine/Private/KismetSystemLibrary.cpp
--- a/Engine/Source/Runtime/Engine/Private/KismetSystemLibrary.cpp
+++ b/Engine/Source/Runtime/Engine/Private/KismetSystemLibrary.cpp
@@ -3062,16 +3062,30 @@ bool UKismetSystemLibrary::GetSupportedFullscreenResolutions(TArray<FIntPoint>&
 	return false;
 }
 
-static TAutoConsoleVariable<int32> CVarMinYResolution(
-	TEXT("r.MinYResolution"),
+static TAutoConsoleVariable<int32> CVarMinYResolutionForUI(
+	TEXT("r.MinYResolutionForUI"),
 	720,
-	TEXT("Defines the smallest Y resolution we want to support in the UI"),
+	TEXT("Defines the smallest Y resolution we want to support in the UI (default is 720)"),
+	ECVF_RenderThreadSafe
+	);
+
+static TAutoConsoleVariable<int32> CVarMinYResolutionFor3DView(
+	TEXT("r.MinYResolutionFor3DView"),
+	360,
+	TEXT("Defines the smallest Y resolution we want to support in the 3D view"),
 	ECVF_RenderThreadSafe
 	);
 
 int32 UKismetSystemLibrary::GetMinYResolutionForUI()
 {
-	int32 Value = FMath::Clamp(CVarMinYResolution.GetValueOnGameThread(), 200, 8 * 1024);
+	int32 Value = FMath::Clamp(CVarMinYResolutionForUI.GetValueOnGameThread(), 200, 8 * 1024);
+
+	return Value;
+}
+
+int32 UKismetSystemLibrary::GetMinYResolutionFor3DView()
+{
+	int32 Value = FMath::Clamp(CVarMinYResolutionFor3DView.GetValueOnGameThread(), 200, 8 * 1024);
 
 	return Value;
 }
diff --git a/Engine/Source/Runtime/Engine/Private/Level.cpp b/Engine/Source/Runtime/Engine/Private/Level.cpp
--- a/Engine/Source/Runtime/Engine/Private/Level.cpp
+++ b/Engine/Source/Runtime/Engine/Private/Level.cpp
@@ -281,12 +281,27 @@ void ULevel::AddReferencedObjects(UObject* InThis, FReferenceCollector& Collecto
 {
 	ULevel* This = CastChecked<ULevel>(InThis);
 
+	// List of textures to remove
+	TArray< UTexture2D*, TInlineAllocator<32> > TexturesToRemove;
+
 	// Let GC know that we're referencing some UTexture2D objects
-	for( auto& It : This->TextureToInstancesMap )
+	for (auto& It : This->TextureToInstancesMap)
 	{
 		UTexture2D* Texture2D = It.Key;
-		Collector.AddReferencedObject( Texture2D, This );
+		TexturesToRemove.Add(Texture2D);
+		Collector.AddReferencedObject(Texture2D, This);
+		if (!Texture2D)
+		{
+			// The texture has been probably marked as pending kill and has been removed by GC
+			// We need to remove it after we exit this loop.
+			TexturesToRemove.Add(It.Key);
+		}
+	}
+	for (UTexture2D* Texture : TexturesToRemove)
+	{
+		This->TextureToInstancesMap.Remove(Texture);
 	}
+	TexturesToRemove.Reset();
 
 	// Let GC know that we're referencing some UTexture2D objects
 	for( auto& It : This->DynamicTextureInstances )
@@ -301,11 +316,22 @@ void ULevel::AddReferencedObjects(UObject* InThis, FReferenceCollector& Collecto
 	}
 
 	// Let GC know that we're referencing some UTexture2D objects
-	for( auto& It : This->ForceStreamTextures )
+	for (auto& It : This->ForceStreamTextures)
 	{
 		UTexture2D* Texture2D = It.Key;
-		Collector.AddReferencedObject( Texture2D, This );
+		Collector.AddReferencedObject(Texture2D, This);
+		if (!Texture2D)
+		{
+			// The texture has been probably marked as pending kill and has been removed by GC
+			// We need to remove it after we exit this loop.
+			TexturesToRemove.Add(It.Key);
+		}
+	}
+	for (UTexture2D* Texture : TexturesToRemove)
+	{
+		This->ForceStreamTextures.Remove(Texture);
 	}
+	TexturesToRemove.Reset();
 
 	// Let GC know that we're referencing some AActor objects
 	for (auto& Actor : This->Actors)
diff --git a/Engine/Source/Runtime/Engine/Private/LocalPlayer.cpp b/Engine/Source/Runtime/Engine/Private/LocalPlayer.cpp
--- a/Engine/Source/Runtime/Engine/Private/LocalPlayer.cpp
+++ b/Engine/Source/Runtime/Engine/Private/LocalPlayer.cpp
@@ -184,11 +184,6 @@ void ULocalPlayer::PostInitProperties()
 			StereoViewState.Allocate();
 		}
 	}
-
-	if (!HasAnyFlags(RF_ClassDefaultObject))
-	{	
-		FCoreDelegates::OnControllerConnectionChange.AddUObject(this, &ULocalPlayer::HandleControllerConnectionChange);
-	}
 }
 
 void ULocalPlayer::PlayerAdded(UGameViewportClient* InViewportClient, int32 InControllerID)
@@ -311,12 +306,6 @@ void ULocalPlayer::SendSplitJoin()
 	}
 }
 
-void ULocalPlayer::BeginDestroy()
-{
-	FCoreDelegates::OnControllerConnectionChange.RemoveAll(this);
-	Super::BeginDestroy();
-}
-
 void ULocalPlayer::FinishDestroy()
 {
 	if ( !IsTemplate() )
@@ -1411,22 +1400,6 @@ void ULocalPlayer::ExecMacro( const TCHAR* Filename, FOutputDevice& Ar )
 	}
 }
 
-void ULocalPlayer::HandleControllerConnectionChange(bool bConnected, int32 InUserId, int32 InControllerId)
-{
-	// if this is an event for this LocalPlayer
-	if (InControllerId == ControllerId)
-	{
-		// if we lost the connection we need to flush all the keys on the PlayerInput to avoid the PC spinning in place, or firing forever, etc.
-		if (!bConnected)
-		{		
-			if (PlayerController && PlayerController->PlayerInput)
-			{
-				PlayerController->PlayerInput->FlushPressedKeys();	
-			}
-		}
-	}
-}
-
 void ULocalPlayer::SetControllerId( int32 NewControllerId )
 {
 	if ( ControllerId != NewControllerId )
diff --git a/Engine/Source/Runtime/Engine/Private/Materials/MaterialInterface.cpp b/Engine/Source/Runtime/Engine/Private/Materials/MaterialInterface.cpp
--- a/Engine/Source/Runtime/Engine/Private/Materials/MaterialInterface.cpp
+++ b/Engine/Source/Runtime/Engine/Private/Materials/MaterialInterface.cpp
@@ -22,9 +22,10 @@ void FMaterialRelevance::SetPrimitiveViewRelevance(FPrimitiveViewRelevance& OutV
 	OutViewRelevance.bSeparateTranslucencyRelevance = bSeparateTranslucency;
 	OutViewRelevance.bNormalTranslucencyRelevance = bNormalTranslucency;
 
+#if !(UE_BUILD_SHIPPING || UE_BUILD_TEST)
 	{
 		static IConsoleVariable* CVar = IConsoleManager::Get().FindConsoleVariable(TEXT("r.ShaderModelDebug"));
-		if(CVar->GetInt())
+		if(CVar && CVar->GetInt())
 		{
 			UE_LOG(LogEngine, Log, TEXT("r.ShaderModelDebug: SetPrimitiveViewRelevance %x = %x"),
 				OutViewRelevance.ShadingModelMaskRelevance,
@@ -32,6 +33,7 @@ void FMaterialRelevance::SetPrimitiveViewRelevance(FPrimitiveViewRelevance& OutV
 			ensure(ShadingModelMask);
 		}
 	}
+#endif
 
 	OutViewRelevance.ShadingModelMaskRelevance = ShadingModelMask;
 	OutViewRelevance.bUsesGlobalDistanceField = bUsesGlobalDistanceField;
diff --git a/Engine/Source/Runtime/Engine/Private/NetConnection.cpp b/Engine/Source/Runtime/Engine/Private/NetConnection.cpp
--- a/Engine/Source/Runtime/Engine/Private/NetConnection.cpp
+++ b/Engine/Source/Runtime/Engine/Private/NetConnection.cpp
@@ -45,6 +45,7 @@ UNetConnection::UNetConnection(const FObjectInitializer& ObjectInitializer)
 ,   OwningActor			( NULL )
 ,	MaxPacket			( 0 )
 ,	InternalAck			( false )
+,	MaxPacketHandlerBits ( 0 )
 ,	State				( USOCK_Invalid )
 ,	PacketOverhead		( 0 )
 ,	ResponseId			( 0 )
@@ -93,7 +94,7 @@ UNetConnection::UNetConnection(const FObjectInitializer& ObjectInitializer)
  * @param InMaxPacket the max packet size that will be used for sending
  * @param InPacketOverhead the packet overhead for this connection type
  */
-void UNetConnection::InitBase(UNetDriver* InDriver,class FSocket* InSocket, const FURL& InURL, EConnectionState InState,int32 InMaxPacket,int32 InPacketOverhead)
+void UNetConnection::InitBase(UNetDriver* InDriver,class FSocket* InSocket, const FURL& InURL, EConnectionState InState, int32 InMaxPacket, int32 InPacketOverhead)
 {
 	// Owning net driver
 	Driver = InDriver;
@@ -101,13 +102,7 @@ void UNetConnection::InitBase(UNetDriver* InDriver,class FSocket* InSocket, cons
 	// Reset Handler
 	Handler.Reset(NULL);
 
-	Handler = MakeUnique<PacketHandler>();
-
-	if(Handler.IsValid())
-	{
-		Handler::Mode Mode = Driver->ServerConnection != nullptr ? Handler::Mode::Client : Handler::Mode::Server;
-		Handler->Initialize(Mode);
-	}
+	InitHandler();
 
 	// Stats
 	StatUpdateTime			= Driver->Time;
@@ -128,7 +123,8 @@ void UNetConnection::InitBase(UNetDriver* InDriver,class FSocket* InSocket, cons
 	// Use the passed in values
 	MaxPacket = InMaxPacket;
 	PacketOverhead = InPacketOverhead;
-	check(MaxPacket && PacketOverhead);
+
+	check(MaxPacket > 0 && PacketOverhead > 0);
 
 #if DO_ENABLE_NET_TEST
 	// Copy the command line settings from the net driver
@@ -168,10 +164,7 @@ void UNetConnection::InitConnection(UNetDriver* InDriver, EConnectionState InSta
 {
 	Driver = InDriver;
 
-	if (!Handler.IsValid())
-	{
-		Handler = MakeUnique<PacketHandler>();
-	}
+	// NOTE: If Handler assignment is added, it must be initialized, and GetPacketOverhead must be factored into MaxPacketHandlerBits
 
 	// We won't be sending any packets, so use a default size
 	MaxPacket = (InMaxPacket == 0 || InMaxPacket > MAX_PACKET_SIZE) ? MAX_PACKET_SIZE : InMaxPacket;
@@ -208,6 +201,21 @@ void UNetConnection::InitConnection(UNetDriver* InDriver, EConnectionState InSta
 	PackageMap = PackageMapClient;
 }
 
+void UNetConnection::InitHandler()
+{
+	check(!Handler.IsValid());
+
+	Handler = MakeUnique<PacketHandler>();
+
+	if (Handler.IsValid())
+	{
+		Handler::Mode Mode = Driver->ServerConnection != nullptr ? Handler::Mode::Client : Handler::Mode::Server;
+		Handler->Initialize(Mode);
+
+		MaxPacketHandlerBits = Handler->GetTotalPacketOverheadBits();
+	}
+}
+
 void UNetConnection::Serialize( FArchive& Ar )
 {
 	UObject::Serialize( Ar );
@@ -474,7 +482,7 @@ void UNetConnection::InitSendBuffer()
 	else
 	{
 		// First time initialization needs to allocate the buffer
-		SendBuffer = FBitWriter(MaxPacket * 8);
+		SendBuffer = FBitWriter((MaxPacket * 8) - MaxPacketHandlerBits);
 	}
 
 	ResetPacketBitCounts();
@@ -487,7 +495,7 @@ void UNetConnection::ReceivedRawPacket( void* InData, int32 Count )
 	uint8* Data = (uint8*)InData;
 
 	// UnProcess the packet
-	if(Handler.IsValid())
+	if (Handler.IsValid())
 	{
 		const ProcessedPacket UnProcessedPacket = Handler->Incoming(Data, Count);
 
@@ -635,8 +643,9 @@ void UNetConnection::FlushNet(bool bIgnoreSimulation)
 		const int32 Index = OutPacketId & (ARRAY_COUNT(OutLagPacketId)-1);
 
 		// Remember the actual time this packet was sent out, so we can compute ping when the ack comes back
-		OutLagPacketId[Index]	= OutPacketId;
-		OutLagTime[Index]		= FPlatformTime::Seconds();
+		OutLagPacketId[Index]			= OutPacketId;
+		OutLagTime[Index]				= FPlatformTime::Seconds();
+		OutBytesPerSecondHistory[Index]	= OutBytesPerSecond / 1024;
 
 		OutPacketId++;
 		Driver->OutPackets++;
@@ -777,6 +786,8 @@ void UNetConnection::ReceivedPacket( FBitReader& Reader )
 				bLastHasServerFrameTime = bHasServerFrameTime;
 			}
 #endif
+			uint32 RemoteInKBytesPerSecond = 0;
+			Reader.SerializeIntPacked( RemoteInKBytesPerSecond );
 
 			// Resend any old reliable packets that the receiver hasn't acknowledged.
 			if( AckPacketId>OutAckPacketId )
@@ -814,6 +825,17 @@ void UNetConnection::ReceivedPacket( FBitReader& Reader )
 				const float NewLag		= FPlatformTime::Seconds() - OutLagTime[Index];
 #endif
 
+				if ( OutBytesPerSecondHistory[Index] > 0 )
+				{
+					RemoteSaturation = ( 1.0f - FMath::Min( ( float )RemoteInKBytesPerSecond / ( float )OutBytesPerSecondHistory[Index], 1.0f ) ) * 100.0f;
+				}
+				else
+				{
+					RemoteSaturation = 0.0f;
+				}
+
+				//UE_LOG( LogNet, Warning, TEXT( "Out: %i, InRemote: %i, Saturation: %f" ), OutBytesPerSecondHistory[Index], RemoteInKBytesPerSecond, RemoteSaturation );
+
 				LagAcc += NewLag;
 				LagCount++;
 
@@ -950,7 +972,9 @@ void UNetConnection::ReceivedPacket( FBitReader& Reader )
 				// Can't handle other channels until control channel exists.
 				if ( Channels[0] == NULL )
 				{
-					CLOSE_CONNECTION_DUE_TO_SECURITY_VIOLATION(this, ESecurityEvent::Malformed_Packet, TEXT( "UNetConnection::ReceivedPacket: Received non-control bunch before control channel was created. ChIndex: %i, ChType: %i" ), Bunch.ChIndex, Bunch.ChType);
+					//CLOSE_CONNECTION_DUE_TO_SECURITY_VIOLATION(this, ESecurityEvent::Malformed_Packet, TEXT( "UNetConnection::ReceivedPacket: Received non-control bunch before control channel was created. ChIndex: %i, ChType: %i" ), Bunch.ChIndex, Bunch.ChType);
+					UE_LOG( LogNetTraffic, Log, TEXT( "UNetConnection::ReceivedPacket: Received non-control bunch before control channel was created. ChIndex: %i, ChType: %i" ), Bunch.ChIndex, Bunch.ChType );
+					Close();
 					return;
 				}
 				// on the server, if we receive bunch data for a channel that doesn't exist while we're still logging in,
@@ -965,11 +989,12 @@ void UNetConnection::ReceivedPacket( FBitReader& Reader )
 			// ignore control channel close if it hasn't been opened yet
 			if ( Bunch.ChIndex == 0 && Channels[0] == NULL && Bunch.bClose && Bunch.ChType == CHTYPE_Control )
 			{
-				CLOSE_CONNECTION_DUE_TO_SECURITY_VIOLATION(this, ESecurityEvent::Malformed_Packet, TEXT( "UNetConnection::ReceivedPacket: Received control channel close before open" ));
+				//CLOSE_CONNECTION_DUE_TO_SECURITY_VIOLATION(this, ESecurityEvent::Malformed_Packet, TEXT( "UNetConnection::ReceivedPacket: Received control channel close before open" ));
+				UE_LOG( LogNetTraffic, Log, TEXT( "UNetConnection::ReceivedPacket: Received control channel close before open" ) );
+				Close();
 				return;
 			}
 
-
 			// Receiving data.
 			UChannel* Channel = Channels[Bunch.ChIndex];
 
@@ -1209,6 +1234,10 @@ void UNetConnection::SendAck(int32 AckPacketId, bool FirstTime/*=1*/)
 		AckData.WriteBit( 0 );
 #endif
 
+		// Notify server of our current rate per second at this time
+		uint32 InKBytesPerSecond = InBytesPerSecond / 1024;
+		AckData.SerializeIntPacked( InKBytesPerSecond );
+
 		NETWORK_PROFILER( GNetworkProfiler.TrackSendAck( AckData.GetNumBits(), this ) );
 
 		WriteBitsToSendBuffer( AckData.GetData(), AckData.GetNumBits(), nullptr, 0, EWriteBitsDataType::Ack );
@@ -1577,7 +1606,7 @@ void UNetConnection::Tick()
 	}
 
 	// Tick Handler
-	if(Handler.IsValid())
+	if (Handler.IsValid())
 	{
 		Handler->Tick(FrameTime);
 		BufferedPacket* QueuedPacket = Handler->GetQueuedPacket();
diff --git a/Engine/Source/Runtime/Engine/Private/NetworkDriver.cpp b/Engine/Source/Runtime/Engine/Private/NetworkDriver.cpp
--- a/Engine/Source/Runtime/Engine/Private/NetworkDriver.cpp
+++ b/Engine/Source/Runtime/Engine/Private/NetworkDriver.cpp
@@ -35,6 +35,7 @@ DEFINE_STAT(STAT_Ping);
 DEFINE_STAT(STAT_Channels);
 DEFINE_STAT(STAT_InRate);
 DEFINE_STAT(STAT_OutRate);
+DEFINE_STAT(STAT_OutSaturation);
 DEFINE_STAT(STAT_InRateClientMax);
 DEFINE_STAT(STAT_InRateClientMin);
 DEFINE_STAT(STAT_InRateClientAvg);
@@ -194,7 +195,7 @@ void UNetDriver::AssertValid()
 void UNetDriver::TickFlush(float DeltaSeconds)
 {
 #if USE_SERVER_PERF_COUNTERS
-	float ServerReplicateActorsTimeMs = 0.0f;
+	double ServerReplicateActorsTimeMs = 0.0f;
 #endif // USE_SERVER_PERF_COUNTERS
 
 	if ( IsServer() && ClientConnections.Num() > 0 && ClientConnections[0]->InternalAck == false )
@@ -203,13 +204,13 @@ void UNetDriver::TickFlush(float DeltaSeconds)
 #if WITH_SERVER_CODE
 
 #if USE_SERVER_PERF_COUNTERS
-		float ServerReplicateActorsTimeStart = FPlatformTime::Seconds();
+		double ServerReplicateActorsTimeStart = FPlatformTime::Seconds();
 #endif // USE_SERVER_PERF_COUNTERS
 
 		int32 Updated = ServerReplicateActors( DeltaSeconds );
 
 #if USE_SERVER_PERF_COUNTERS
-		ServerReplicateActorsTimeMs = (FPlatformTime::Seconds() - ServerReplicateActorsTimeStart) * 1000.0f;
+		ServerReplicateActorsTimeMs = (FPlatformTime::Seconds() - ServerReplicateActorsTimeStart) * 1000.0;
 #endif // USE_SERVER_PERF_COUNTERS
 
 		static int32 LastUpdateCount = 0;
@@ -230,7 +231,7 @@ void UNetDriver::TickFlush(float DeltaSeconds)
 
 	// Update network stats (only main game net driver for now) if stats or perf counters are used
 	if (NetDriverName == NAME_GameNetDriver && 
-		 CurrentRealtimeSeconds - StatUpdateTime > StatPeriod && ( ClientConnections.Num() > 0 || ServerConnection != NULL ) )
+		 CurrentRealtimeSeconds - StatUpdateTime > StatPeriod)
 	{
 		int32 ClientInBytesMax = 0;
 		int32 ClientInBytesMin = 0;
@@ -280,6 +281,7 @@ void UNetDriver::TickFlush(float DeltaSeconds)
 		int32 UnAckCount = 0;
 		int32 PendingCount = 0;
 		int32 NetSaturated = 0;
+		float RemoteSaturationMax = 0.0f;
 
 		if (FThreadStats::IsCollectingData())
 		{
@@ -312,10 +314,13 @@ void UNetDriver::TickFlush(float DeltaSeconds)
 			if (ServerConnection != NULL)
 			{
 				NumOpenChannels = ServerConnection->OpenChannels.Num();
+				RemoteSaturationMax = FMath::Max( RemoteSaturationMax, ServerConnection->RemoteSaturation );
 			}
+
 			for (int32 i = 0; i < ClientConnections.Num(); i++)
 			{
 				NumOpenChannels += ClientConnections[i]->OpenChannels.Num();
+				RemoteSaturationMax = FMath::Max( RemoteSaturationMax, ClientConnections[i]->RemoteSaturation );
 			}
 
 			// Use the elapsed time to keep things scaled to one measured unit
@@ -365,6 +370,7 @@ void UNetDriver::TickFlush(float DeltaSeconds)
 		SET_DWORD_STAT(STAT_InLoss, InPacketsLost);
 		SET_DWORD_STAT(STAT_InRate, InBytes);
 		SET_DWORD_STAT(STAT_OutRate, OutBytes);
+		SET_DWORD_STAT(STAT_OutSaturation, RemoteSaturationMax);
 		SET_DWORD_STAT(STAT_InRateClientMax, ClientInBytesMax);
 		SET_DWORD_STAT(STAT_InRateClientMin, ClientInBytesMin);
 		SET_DWORD_STAT(STAT_InRateClientAvg, ClientInBytesAvg);
@@ -461,8 +467,6 @@ void UNetDriver::TickFlush(float DeltaSeconds)
 			}
 			else
 			{
-				PerfCounters->Set(TEXT("NumConnections"), 0);
-
 				PerfCounters->Set(TEXT("AvgPing"), 0.0f, IPerfCounters::Flags::Transient);
 				PerfCounters->Set(TEXT("MaxPing"), -FLT_MAX, IPerfCounters::Flags::Transient);
 				PerfCounters->Set(TEXT("MinPing"), FLT_MAX, IPerfCounters::Flags::Transient);
diff --git a/Engine/Source/Runtime/Engine/Private/PackageMapClient.cpp b/Engine/Source/Runtime/Engine/Private/PackageMapClient.cpp
--- a/Engine/Source/Runtime/Engine/Private/PackageMapClient.cpp
+++ b/Engine/Source/Runtime/Engine/Private/PackageMapClient.cpp
@@ -856,7 +856,7 @@ bool UPackageMapClient::ExportNetGUID( FNetworkGUID NetGUID, const UObject* Obje
 		{
 			check( ExportNetGUIDCount == 0 );
 
-			CurrentExportBunch = new FOutBunch(this, Connection->MaxPacket*8-MAX_BUNCH_HEADER_BITS-MAX_PACKET_TRAILER_BITS-MAX_PACKET_HEADER_BITS );
+			CurrentExportBunch = new FOutBunch(this, Connection->GetMaxSingleBunchSizeBits());
 			CurrentExportBunch->SetAllowResize(false);
 			CurrentExportBunch->bHasGUIDs = true;
 #if !(UE_BUILD_SHIPPING || UE_BUILD_TEST)
