Copying //UE4/Dev-Framework to Dev-Main (//UE4/Dev-Main) @ 2855699

diff --git a/Engine/Source/Runtime/AnimGraphRuntime/Public/BoneControllers/AnimNode_CopyBoneDelta.h b/Engine/Source/Runtime/AnimGraphRuntime/Public/BoneControllers/AnimNode_CopyBoneDelta.h
new file mode 100644
--- /dev/null
+++ b/Engine/Source/Runtime/AnimGraphRuntime/Public/BoneControllers/AnimNode_CopyBoneDelta.h
@@ -0,0 +1,67 @@
+// Copyright 1998-2016 Epic Games, Inc. All Rights Reserved.
+
+#pragma once
+
+#include "AnimNode_SkeletalControlBase.h"
+#include "AnimNode_CopyBoneDelta.generated.h"
+
+UENUM()
+enum class CopyBoneDeltaMode : uint8
+{
+	Accumulate,
+	Copy
+};
+
+/**
+ *	Simple controller to copy a transform relative to the ref pose to the target bone,
+ *	instead of the copy bone node which copies the absolute transform
+ */
+
+USTRUCT()
+struct ANIMGRAPHRUNTIME_API FAnimNode_CopyBoneDelta : public FAnimNode_SkeletalControlBase
+{
+	GENERATED_USTRUCT_BODY()
+
+	UPROPERTY(EditAnywhere, BlueprintReadWrite, Category = Settings)
+	FBoneReference SourceBone;
+
+	UPROPERTY(EditAnywhere, BlueprintReadWrite, Category = Settings)
+	FBoneReference TargetBone;
+
+	UPROPERTY(EditAnywhere, BlueprintReadWrite, Category = Settings, meta = (PinShownByDefault))
+	bool bCopyTranslation;
+
+	UPROPERTY(EditAnywhere, BlueprintReadWrite, Category = Settings, meta = (PinShownByDefault))
+	bool bCopyRotation;
+
+	UPROPERTY(EditAnywhere, BlueprintReadWrite, Category = Settings, meta = (PinShownByDefault))
+	bool bCopyScale;
+
+	UPROPERTY(EditAnywhere, BlueprintReadWrite, Category = Settings)
+	CopyBoneDeltaMode CopyMode;
+
+	UPROPERTY(EditAnywhere, BlueprintReadWrite, Category = Settings, meta = (PinHiddenByDefault))
+	float TranslationMultiplier;
+
+	UPROPERTY(EditAnywhere, BlueprintReadWrite, Category = Settings, meta = (PinHiddenByDefault))
+	float RotationMultiplier;
+
+	UPROPERTY(EditAnywhere, BlueprintReadWrite, Category = Settings, meta = (PinHiddenByDefault))
+	float ScaleMultiplier;
+
+	FAnimNode_CopyBoneDelta();
+
+	// FAnimNode_Base interface
+	virtual void GatherDebugData(FNodeDebugData& DebugData) override;
+	// End of FAnimNode_Base interface
+
+	// FAnimNode_SkeletalControlBase interface
+	virtual void EvaluateBoneTransforms(USkeletalMeshComponent* SkelComp, FCSPose<FCompactPose>& MeshBases, TArray<FBoneTransform>& OutBoneTransforms) override;
+	virtual bool IsValidToEvaluate(const USkeleton* Skeleton, const FBoneContainer& RequiredBones) override;
+	// End of FAnimNode_SkeletalControlBase interface
+
+private:
+	// FAnimNode_SkeletalControlBase interface
+	virtual void InitializeBoneReferences(const FBoneContainer& RequiredBones) override;
+	// End of FAnimNode_SkeletalControlBase interface
+};
\ No newline at end of file
diff --git a/Engine/Source/Runtime/Core/Private/GenericPlatform/GenericPlatformMath.cpp b/Engine/Source/Runtime/Core/Private/GenericPlatform/GenericPlatformMath.cpp
--- a/Engine/Source/Runtime/Core/Private/GenericPlatform/GenericPlatformMath.cpp
+++ b/Engine/Source/Runtime/Core/Private/GenericPlatform/GenericPlatformMath.cpp
@@ -74,7 +74,10 @@ float FGenericPlatformMath::Atan2(float Y, float X)
 
 void FGenericPlatformMath::FmodReportError(float X, float Y)
 {
-	ensureMsgf(Y != 0, TEXT("FMath::FMod(X=%f, Y=%f) : Y is zero, this is invalid and would result in NaN!"), X, Y);
+	if (Y == 0)
+	{
+		ensureMsgf(Y != 0, TEXT("FMath::FMod(X=%f, Y=%f) : Y is zero, this is invalid and would result in NaN!"), X, Y);
+	}
 }
 
 #if WITH_DEV_AUTOMATION_TESTS
diff --git a/Engine/Source/Runtime/Core/Private/Math/Color.cpp b/Engine/Source/Runtime/Core/Private/Math/Color.cpp
--- a/Engine/Source/Runtime/Core/Private/Math/Color.cpp
+++ b/Engine/Source/Runtime/Core/Private/Math/Color.cpp
@@ -235,9 +235,9 @@ FLinearColor FLinearColor::LinearRGBToHSV() const
 	const float RGBRange = RGBMax - RGBMin;
 
 	const float Hue = (RGBMax == RGBMin ? 0.0f :
-					   RGBMax == R    ? fmod((((G - B) / RGBRange) * 60.0f) + 360.0f, 360.0f) :
-					   RGBMax == G    ?      (((B - R) / RGBRange) * 60.0f) + 120.0f :
-					   RGBMax == B    ?      (((R - G) / RGBRange) * 60.0f) + 240.0f :
+					   RGBMax == R    ? FMath::Fmod((((G - B) / RGBRange) * 60.0f) + 360.0f, 360.0f) :
+					   RGBMax == G    ?             (((B - R) / RGBRange) * 60.0f) + 120.0f :
+					   RGBMax == B    ?             (((R - G) / RGBRange) * 60.0f) + 240.0f :
 					   0.0f);
 	
 	const float Saturation = (RGBMax == 0.0f ? 0.0f : RGBRange / RGBMax);
diff --git a/Engine/Source/Runtime/Core/Private/Tests/Math/UnrealMathTest.cpp b/Engine/Source/Runtime/Core/Private/Tests/Math/UnrealMathTest.cpp
--- a/Engine/Source/Runtime/Core/Private/Tests/Math/UnrealMathTest.cpp
+++ b/Engine/Source/Runtime/Core/Private/Tests/Math/UnrealMathTest.cpp
@@ -1039,13 +1039,13 @@ bool FVectorRegisterAbstractionTest::RunTest(const FString& Parameters)
 	LogTest( TEXT("VectorReciprocalSqrtAccurate"), TestVectorsEqual( VectorOne(), V3, 1e-6f ) );
 
 	// VectorMod
-	V0 = MakeVectorRegister(0.0f, 3.2f, 2.8f, 10.0f);
+	V0 = MakeVectorRegister(0.0f, 3.2f, 2.8f,  1.5f);
 	V1 = MakeVectorRegister(2.0f, 1.2f, 2.0f,  3.0f);
 	V2 = TestReferenceMod(V0, V1);
 	V3 = VectorMod(V0, V1);
 	LogTest( TEXT("VectorMod positive"), TestVectorsEqual(V2, V3));
 
-	V0 = MakeVectorRegister(-2.0f,  3.2f, -2.8f, -10.0f);
+	V0 = MakeVectorRegister(-2.0f,  3.2f, -2.8f,  -1.5f);
 	V1 = MakeVectorRegister(-1.5f, -1.2f,  2.0f,   3.0f);
 	V2 = TestReferenceMod(V0, V1);
 	V3 = VectorMod(V0, V1);
@@ -1389,6 +1389,10 @@ bool FVectorRegisterAbstractionTest::RunTest(const FString& Parameters)
 			{ 2.8,	 2.0 + KINDA_SMALL_NUMBER},
 			{-2.8,	 2.0 - KINDA_SMALL_NUMBER},
 
+			// Analytically should be zero but floating point precision can cause results close to Y (or erroneously negative) depending on the method used.
+			{55.8,	 9.3},
+			{1234.1234, 0.1234},
+
 			// Commonly used for FRotators and angles
 			{725.2,		360.0},
 			{179.9,		 90.0},
@@ -1417,10 +1421,24 @@ bool FVectorRegisterAbstractionTest::RunTest(const FString& Parameters)
 			const float Y = XY.Y;
 			const float Ours = FMath::Fmod(X, Y);
 			const float Theirs = fmodf(X, Y);
-			if (!FMath::IsNearlyEqual(Ours, Theirs, 1.e-5f))
+			
+			// A compiler bug causes stock fmodf() to rarely return NaN for valid input, we don't want to report this as a fatal error.
+			if (Y != 0 && FMath::IsNaN(Theirs))
+			{
+				UE_LOG(LogUnrealMathTest, Warning, TEXT("fmodf(%f, %f) with valid input resulted in NaN!"), X, Y);
+				continue;
+			}
+
+			const float Delta = FMath::Abs(Ours - Theirs);
+			if (Delta > 1e-5f)
 			{
-				UE_LOG(LogUnrealMathTest, Log, TEXT("FMath::Fmod(%f, %f)=%f <-> fmodf(%f, %f)=%f: FAILED"), X, Y, Ours, X, Y, Theirs);
-				GPassing = false;
+				// If we differ significantly, that is likely due to rounding and the difference should be nearly equal to Y.
+				const float FractionalDelta = FMath::Abs(Delta - FMath::Abs(Y));
+				if (FractionalDelta > 1e-4f)
+				{
+					UE_LOG(LogUnrealMathTest, Log, TEXT("FMath::Fmod(%f, %f)=%f <-> fmodf(%f, %f)=%f: FAILED"), X, Y, Ours, X, Y, Theirs);
+					GPassing = false;
+				}
 			}
 		}
 	}
diff --git a/Engine/Source/Runtime/Core/Public/Containers/Set.h b/Engine/Source/Runtime/Core/Public/Containers/Set.h
--- a/Engine/Source/Runtime/Core/Public/Containers/Set.h
+++ b/Engine/Source/Runtime/Core/Public/Containers/Set.h
@@ -791,10 +791,16 @@ public:
 	/** @return the intersection of two sets. (A AND B)*/
 	TSet Intersect(const TSet& OtherSet) const
 	{
+		const bool bOtherSmaller = (Num() > OtherSet.Num());
+		const TSet& A = (bOtherSmaller ? OtherSet : *this);
+		const TSet& B = (bOtherSmaller ? *this : OtherSet);
+
 		TSet Result;
-		for(TConstIterator SetIt(*this);SetIt;++SetIt)
+		Result.Reserve(A.Num()); // Worst case is everything in smaller is in larger
+
+		for(TConstIterator SetIt(A);SetIt;++SetIt)
 		{
-			if(OtherSet.Contains(KeyFuncs::GetSetKey(*SetIt)))
+			if(B.Contains(KeyFuncs::GetSetKey(*SetIt)))
 			{
 				Result.Add(*SetIt);
 			}
@@ -806,6 +812,8 @@ public:
 	TSet Union(const TSet& OtherSet) const
 	{
 		TSet Result;
+		Result.Reserve(Num() + OtherSet.Num()); // Worst case is 2 totally unique Sets
+
 		for(TConstIterator SetIt(*this);SetIt;++SetIt)
 		{
 			Result.Add(*SetIt);
@@ -817,10 +825,12 @@ public:
 		return Result;
 	}
 
-	/** @return the complement of two sets. (A not in B)*/
+	/** @return the complement of two sets. (A not in B where A is this and B is Other)*/
 	TSet Difference(const TSet& OtherSet) const
 	{
 		TSet Result;
+		Result.Reserve(Num()); // Worst case is no elements of this are in Other
+
 		for(TConstIterator SetIt(*this);SetIt;++SetIt)
 		{
 			if(!OtherSet.Contains(KeyFuncs::GetSetKey(*SetIt)))
@@ -841,14 +851,22 @@ public:
 	bool Includes(const TSet<ElementType,KeyFuncs,Allocator>& OtherSet) const
 	{
 		bool bIncludesSet = true;
-		for(TConstIterator OtherSetIt(OtherSet); OtherSetIt; ++OtherSetIt)
+		if (OtherSet.Num() <= Num())
 		{
-			if (!Contains(KeyFuncs::GetSetKey(*OtherSetIt)))
+			for(TConstIterator OtherSetIt(OtherSet); OtherSetIt; ++OtherSetIt)
 			{
-				bIncludesSet = false;
-				break;
+				if (!Contains(KeyFuncs::GetSetKey(*OtherSetIt)))
+				{
+					bIncludesSet = false;
+					break;
+				}
 			}
 		}
+		else
+		{
+			// Not possible to include if it is bigger than us
+			bIncludesSet = false;
+		}
 		return bIncludesSet;
 	}
 
@@ -856,7 +874,7 @@ public:
 	TArray<ElementType> Array() const
 	{
 		TArray<ElementType> Result;
-		Result.Empty(Num());
+		Result.Reserve(Num());
 		for(TConstIterator SetIt(*this);SetIt;++SetIt)
 		{
 			Result.Add(*SetIt);
diff --git a/Engine/Source/Runtime/Core/Public/GenericPlatform/GenericPlatformMath.h b/Engine/Source/Runtime/Core/Public/GenericPlatform/GenericPlatformMath.h
--- a/Engine/Source/Runtime/Core/Public/GenericPlatform/GenericPlatformMath.h
+++ b/Engine/Source/Runtime/Core/Public/GenericPlatform/GenericPlatformMath.h
@@ -243,8 +243,17 @@ struct FGenericPlatformMath
 			FmodReportError(X, Y);
 			return 0.f;
 		}
-		const float IntPortion = TruncToFloat(X / Y);
-		const float Result = X - Y * IntPortion;
+		const float Quotient = TruncToFloat(X / Y);
+		float IntPortion = Y * Quotient;
+
+		// Rounding and imprecision could cause IntPortion to exceed X and cause the result to be outside the expected range.
+		// For example Fmod(55.8, 9.3) would result in a very small negative value!
+		if (fabsf(IntPortion) > fabsf(X))
+		{
+			IntPortion = X;
+		}
+
+		const float Result = X - IntPortion;
 		return Result;
 	}
 
diff --git a/Engine/Source/Runtime/Core/Public/HAL/OutputDevices.h b/Engine/Source/Runtime/Core/Public/HAL/OutputDevices.h
--- a/Engine/Source/Runtime/Core/Public/HAL/OutputDevices.h
+++ b/Engine/Source/Runtime/Core/Public/HAL/OutputDevices.h
@@ -150,9 +150,9 @@ protected:
 	uint8 Verbosity;
 	/** Holds the break flag **/
 	bool DebugBreakOnLog;
-	/** Holds default supression **/
+	/** Holds default suppression **/
 	uint8 DefaultVerbosity;
-	/** Holds compile time supression **/
+	/** Holds compile time suppression **/
 	uint8 CompileTimeVerbosity;
 	/** FName for this category **/
 	FName CategoryFName;
diff --git a/Engine/Source/Runtime/Core/Public/Math/DualQuat.h b/Engine/Source/Runtime/Core/Public/Math/DualQuat.h
new file mode 100644
--- /dev/null
+++ b/Engine/Source/Runtime/Core/Public/Math/DualQuat.h
@@ -0,0 +1,58 @@
+// Copyright 1998-2016 Epic Games, Inc. All Rights Reserved.
+
+#pragma once
+
+/** Dual quaternion class */
+class FDualQuat
+{
+public:
+
+	/** rotation or real part */
+	FQuat R;
+	/** half trans or dual part */
+	FQuat D;
+
+	// Constructors
+	FDualQuat(const FQuat &InR, const FQuat &InD)
+		: R(InR)
+		, D(InD)
+	{}
+
+	FDualQuat(const FTransform &T)
+	{
+		FVector V = T.GetTranslation()*0.5f;
+		*this = FDualQuat(FQuat(0, 0, 0, 1), FQuat(V.X, V.Y, V.Z, 0.f)) * FDualQuat(T.GetRotation(), FQuat(0, 0, 0, 0));
+	}
+
+	/** Dual quat addition */
+	FDualQuat operator+(const FDualQuat &B) const
+	{
+		return{ R + B.R, D + B.D };
+	}
+
+	/** Dual quat product */
+	FDualQuat operator*(const FDualQuat &B) const
+	{
+		return{ R*B.R, D*B.R + B.D*R };
+	}
+
+	/** Scale dual quat */
+	FDualQuat operator*(const float S) const
+	{
+		return{ R*S, D*S };
+	}
+
+	/** Return normalized dual quat */
+	FDualQuat Normalized() const
+	{
+		float MinV = 1.0f / FMath::Sqrt(R | R);
+		return{ R*MinV, D*MinV };
+	}
+
+	/** Convert dual quat to transform */
+	FTransform AsFTransform(FVector Scale = FVector(1.0f, 1.0f, 1.0f))
+	{
+		FQuat TQ = D*FQuat(-R.X, -R.Y, -R.Z, R.W);
+		return FTransform(R, FVector(TQ.X, TQ.Y, TQ.Z)*2.0f, Scale);
+	}
+};
diff --git a/Engine/Source/Runtime/Core/Public/UObject/FrameworkObjectVersion.h b/Engine/Source/Runtime/Core/Public/UObject/FrameworkObjectVersion.h
--- a/Engine/Source/Runtime/Core/Public/UObject/FrameworkObjectVersion.h
+++ b/Engine/Source/Runtime/Core/Public/UObject/FrameworkObjectVersion.h
@@ -9,6 +9,8 @@ struct CORE_API FFrameworkObjectVersion
 		// Before any version changes were made
 		BeforeCustomVersionWasAdded = 0,
 
+		// BodySetup's default instance collision profile is used by default when creating a new instance.
+		UseBodySetupCollisionProfile,
 
 		// -----<new versions can be added above this line>-------------------------------------------------
 		VersionPlusOne,
diff --git a/Engine/Source/Runtime/Engine/Classes/Animation/AnimInstance.h b/Engine/Source/Runtime/Engine/Classes/Animation/AnimInstance.h
--- a/Engine/Source/Runtime/Engine/Classes/Animation/AnimInstance.h
+++ b/Engine/Source/Runtime/Engine/Classes/Animation/AnimInstance.h
@@ -1184,6 +1184,9 @@ public:
 	/** Pass-thru function to proxy - only call on the game thread */
 	void UpdateComponentsMaterialParameters(UPrimitiveComponent* Component);
 
+	/** Given a machine and state index, record a state weight for this frame */
+	void RecordStateWeight(const int32& InMachineClassIndex, const int32& InStateIndex, const float& InStateWeight);
+
 protected:
 	/** 
 	 * Add curve float data, using a curve name. External values should all be added using
diff --git a/Engine/Source/Runtime/Engine/Classes/Components/ActorComponent.h b/Engine/Source/Runtime/Engine/Classes/Components/ActorComponent.h
--- a/Engine/Source/Runtime/Engine/Classes/Components/ActorComponent.h
+++ b/Engine/Source/Runtime/Engine/Classes/Components/ActorComponent.h
@@ -278,7 +278,7 @@ public:
 	/** IsSupportedForNetworking means an object can be referenced over the network */
 	virtual bool IsSupportedForNetworking() const override;
 
-	/** Enable or disable replication. This is the equivelent of RemoteRole for actors (only a bool is required for components) */
+	/** Enable or disable replication. This is the equivalent of RemoteRole for actors (only a bool is required for components) */
 	UFUNCTION(BlueprintCallable, Category="Components")
 	void SetIsReplicated(bool ShouldReplicate);
 
diff --git a/Engine/Source/Runtime/Engine/Classes/Components/MaterialBillboardComponent.h b/Engine/Source/Runtime/Engine/Classes/Components/MaterialBillboardComponent.h
--- a/Engine/Source/Runtime/Engine/Classes/Components/MaterialBillboardComponent.h
+++ b/Engine/Source/Runtime/Engine/Classes/Components/MaterialBillboardComponent.h
@@ -5,7 +5,7 @@
 #include "MaterialBillboardComponent.generated.h"
 
 
-USTRUCT()
+USTRUCT(BlueprintType)
 struct FMaterialSpriteElement
 {
 	GENERATED_USTRUCT_BODY()
@@ -55,12 +55,17 @@ class ENGINE_API UMaterialBillboardComponent : public UPrimitiveComponent
 {
 	GENERATED_UCLASS_BODY()
 
-	UPROPERTY(EditAnywhere, BlueprintReadWrite, Category=Sprite)
-	TArray<struct FMaterialSpriteElement> Elements;
+	/** Current array of material billboard elements */
+	UPROPERTY(EditAnywhere, BlueprintReadOnly, Category=Sprite)
+	TArray<FMaterialSpriteElement> Elements;
+
+	/** Set all elements of this material billboard component */
+	UFUNCTION(BlueprintCallable, Category="Rendering|Components|MaterialSprite")
+	void SetElements(const TArray<FMaterialSpriteElement>& NewElements);
 
 	/** Adds an element to the sprite. */
-	UFUNCTION(BlueprintCallable,Category="Rendering|Components|MaterialSprite")
-	virtual void AddElement(
+	UFUNCTION(BlueprintCallable, Category="Rendering|Components|MaterialSprite")
+	void AddElement(
 		class UMaterialInterface* Material,
 		class UCurveFloat* DistanceToOpacityCurve,
 		bool bSizeIsInScreenSpace,
diff --git a/Engine/Source/Runtime/Engine/Classes/Components/PostProcessComponent.h b/Engine/Source/Runtime/Engine/Classes/Components/PostProcessComponent.h
--- a/Engine/Source/Runtime/Engine/Classes/Components/PostProcessComponent.h
+++ b/Engine/Source/Runtime/Engine/Classes/Components/PostProcessComponent.h
@@ -46,7 +46,7 @@ class UPostProcessComponent : public USceneComponent, public IInterface_PostProc
 	{
 		FPostProcessVolumeProperties Ret;
 		Ret.bIsEnabled = bEnabled != 0;
-		Ret.bIsUnbound = bUnbound != 0 || Cast<UShapeComponent>(AttachParent) == nullptr;
+		Ret.bIsUnbound = bUnbound != 0 || Cast<UShapeComponent>(GetAttachParent()) == nullptr;
 		Ret.BlendRadius = BlendRadius;
 		Ret.BlendWeight = BlendWeight;
 		Ret.Priority = Priority;
diff --git a/Engine/Source/Runtime/Engine/Classes/Components/SceneComponent.h b/Engine/Source/Runtime/Engine/Classes/Components/SceneComponent.h
--- a/Engine/Source/Runtime/Engine/Classes/Components/SceneComponent.h
+++ b/Engine/Source/Runtime/Engine/Classes/Components/SceneComponent.h
@@ -92,20 +92,18 @@ class ENGINE_API USceneComponent : public UActorComponent
 public:
 
 	/** The name to use for the default scene root variable */
-	static const FName& GetDefaultSceneRootVariableName();
+	static FName GetDefaultSceneRootVariableName();
 
 	/**
 	 * UObject constructor that takes an ObjectInitializer
 	 */
 	USceneComponent(const FObjectInitializer& ObjectInitializer = FObjectInitializer::Get());
 
-public:
-
 	virtual void GetLifetimeReplicatedProps(TArray<FLifetimeProperty>& OutLifetimeProps) const override;
 
 	/** What we are currently attached to. If valid, RelativeLocation etc. are used relative to this object */
 	UPROPERTY(ReplicatedUsing=OnRep_AttachParent)
-	class USceneComponent* AttachParent;
+	USceneComponent* AttachParent;
 
 	/** List of child SceneComponents that are attached to us. */
 	UPROPERTY(Replicated, transient)
@@ -231,8 +229,8 @@ public:
 
 private:
 
-	bool bNetUpdateTransform;
-	bool bNetUpdateAttachment;
+	uint8 bNetUpdateTransform:1;
+	uint8 bNetUpdateAttachment:1;
 	FName NetOldAttachSocketName;
 	USceneComponent *NetOldAttachParent;
 
@@ -530,6 +528,9 @@ public:
 	UFUNCTION(BlueprintCallable, Category="Physics")
 	virtual bool IsAnySimulatingPhysics() const;
 
+	/** Get the SceneComponents that are attached to this component. */
+	const TArray<USceneComponent*>& GetAttachChildren() const;
+
 	/** Get the SceneComponent we are attached to. */
 	UFUNCTION(BlueprintCallable, Category="Utilities|Transformation")
 	USceneComponent* GetAttachParent() const;
@@ -540,7 +541,7 @@ public:
 
 	/** Gets all parent components up to and including the root component */
 	UFUNCTION(BlueprintCallable, Category="Components")
-	void GetParentComponents(TArray<class USceneComponent*>& Parents) const;
+	void GetParentComponents(TArray<USceneComponent*>& Parents) const;
 
 	/** Gets the number of attached children components */
 	UFUNCTION(BlueprintCallable, Category="Components")
@@ -548,7 +549,7 @@ public:
 
 	/** Gets the attached child component at the specified location */
 	UFUNCTION(BlueprintCallable, Category="Components")
-	class USceneComponent* GetChildComponent(int32 ChildIndex) const;
+	USceneComponent* GetChildComponent(int32 ChildIndex) const;
 
 	/** 
 	 * Gets all the attached child components
@@ -566,7 +567,7 @@ public:
 	 * @param  bWeldSimulatedBodies Whether to weld together simulated physics bodies.
 	 * @return True if attachment is successful (or already attached to requested parent/socket), false if attachment is rejected and there is no change in AttachParent.
 	 */
-	 bool AttachTo(class USceneComponent* InParent, FName InSocketName = NAME_None, EAttachLocation::Type AttachType = EAttachLocation::KeepRelativeOffset, bool bWeldSimulatedBodies = false);
+	 bool AttachTo(USceneComponent* InParent, FName InSocketName = NAME_None, EAttachLocation::Type AttachType = EAttachLocation::KeepRelativeOffset, bool bWeldSimulatedBodies = false);
 
 	/**
 	 * Attach this component to another scene component, optionally at a named socket. It is valid to call this on components whether or not they have been Registered.
@@ -574,13 +575,14 @@ public:
 	 * @param  InSocketName			Optional socket to attach to on the parent.
 	 * @param  AttachType			How to handle transform when attaching (Keep relative offset, keep world position, etc).
 	 * @param  bWeldSimulatedBodies Whether to weld together simulated physics bodies.
+	 * @return True if attachment is successful (or already attached to requested parent/socket), false if attachment is rejected and there is no change in AttachParent.
 	*/
 	UFUNCTION(BlueprintCallable, Category = "Utilities|Transformation", meta = (DisplayName = "AttachTo", AttachType = "KeepRelativeOffset"))
-	void K2_AttachTo(class USceneComponent* InParent, FName InSocketName = NAME_None, EAttachLocation::Type AttachType = EAttachLocation::KeepRelativeOffset, bool bWeldSimulatedBodies = true);
+	bool K2_AttachTo(USceneComponent* InParent, FName InSocketName = NAME_None, EAttachLocation::Type AttachType = EAttachLocation::KeepRelativeOffset, bool bWeldSimulatedBodies = true);
 
 	/** Zeroes out the relative transform of the component, and calls AttachTo(). Useful for attaching directly to a scene component or socket location  */
 	UFUNCTION(BlueprintCallable, meta=(DeprecatedFunction, DeprecationMessage = "Use AttachTo with EAttachLocation::SnapToTarget option instead"), Category="Utilities|Transformation")
-	void SnapTo(class USceneComponent* InParent, FName InSocketName = NAME_None);
+	bool SnapTo(USceneComponent* InParent, FName InSocketName = NAME_None);
 
 	/** 
 	 *	Detach this component from whatever it is attached to. Automatically unwelds components that are welded together (See WeldTo)
@@ -589,7 +591,6 @@ public:
 	UFUNCTION(BlueprintCallable, Category="Utilities|Transformation")
 	virtual void DetachFromParent(bool bMaintainWorldPosition = false, bool bCallModify = true);
 
-
 	/** 
 	 * Gets the names of all the sockets on the component.
 	 * @return Get the names of all the sockets on the component.
@@ -694,7 +695,7 @@ public:
 	}
 	virtual void UpdateComponentToWorld(bool bSkipPhysicsMove = false, ETeleportType Teleport = ETeleportType::None) override final
 	{
-		UpdateComponentToWorldWithParent(AttachParent, AttachSocketName, bSkipPhysicsMove, RelativeRotationCache.RotatorToQuat(RelativeRotation), Teleport);
+		UpdateComponentToWorldWithParent(GetAttachParent(), GetAttachSocketName(), bSkipPhysicsMove, RelativeRotationCache.RotatorToQuat(RelativeRotation), Teleport);
 	}
 	virtual void DestroyComponent(bool bPromoteChildren = false) override;
 	virtual void OnComponentDestroyed(bool bDestroyingHierarchy) override;
@@ -933,14 +934,14 @@ protected:
 		If Parent is not passed in we use the component's AttachParent*/
 	FORCEINLINE FTransform CalcNewComponentToWorld(const FTransform& NewRelativeTransform, const USceneComponent* Parent = NULL, FName SocketName = NAME_None) const
 	{
-		SocketName = Parent ? SocketName : AttachSocketName;
-		Parent = Parent ? Parent : AttachParent;
+		SocketName = Parent ? SocketName : GetAttachSocketName();
+		Parent = Parent ? Parent : GetAttachParent();
 		if (Parent)
 		{
 			const bool bGeneral = bAbsoluteLocation || bAbsoluteRotation || bAbsoluteScale;
 			if (!bGeneral)
 			{
-				return NewRelativeTransform * Parent->GetSocketTransform(AttachSocketName);
+				return NewRelativeTransform * Parent->GetSocketTransform(GetAttachSocketName());
 			}
 			
 			return CalcNewComponentToWorld_GeneralCase(NewRelativeTransform, Parent, SocketName);
@@ -1031,19 +1032,19 @@ public:
 	virtual ECollisionChannel GetCollisionObjectType() const;
 
 	/** Compares the CollisionObjectType of each component against the Response of the other, to see what kind of response we should generate */
-	ECollisionResponse GetCollisionResponseToComponent(class USceneComponent* OtherComponent) const;
+	ECollisionResponse GetCollisionResponseToComponent(USceneComponent* OtherComponent) const;
 
 	/** Set how often this component is allowed to move during runtime. Causes a component re-register if the component is already registered */
 	virtual void SetMobility(EComponentMobility::Type NewMobility);
 
-	/** Walks up the attachment chain from this SceneComponent and returns the SceneComponent at the top. this->Attachparent is NULL, returns this. */
-	class USceneComponent* GetAttachmentRoot() const;
+	/** Walks up the attachment chain from this SceneComponent and returns the SceneComponent at the top. If AttachParent is NULL, returns this. */
+	USceneComponent* GetAttachmentRoot() const;
 	
-	/** Walks up the attachment chain from this SceneComponent and returns the top-level actor it's attached to.  Returns NULL if unattached. */
-	class AActor* GetAttachmentRootActor() const;
+	/** Walks up the attachment chain from this SceneComponent and returns the top-level actor it's attached to.  Returns Owner if unattached. */
+	AActor* GetAttachmentRootActor() const;
 
 	/** Walks up the attachment chain to see if this component is attached to the supplied component. If TestComp == this, returns false.*/
-	bool IsAttachedTo(class USceneComponent* TestComp) const;
+	bool IsAttachedTo(USceneComponent* TestComp) const;
 
 	/**
 	 * Find the world-space location and rotation of the given named socket.
@@ -1100,6 +1101,11 @@ protected:
 //////////////////////////////////////////////////////////////////////////
 // USceneComponent inlines
 
+FORCEINLINE const TArray<USceneComponent*>& USceneComponent::GetAttachChildren() const
+{
+	return AttachChildren;
+}
+
 FORCEINLINE USceneComponent* USceneComponent::GetAttachParent() const
 {
 	return AttachParent;
@@ -1259,7 +1265,7 @@ public:
 	
 	typedef TArray<struct FHitResult, TInlineAllocator<2>> TBlockingHitArray;
 
-	FScopedMovementUpdate( class USceneComponent* Component, EScopedUpdate::Type ScopeBehavior = EScopedUpdate::DeferredUpdates );
+	FScopedMovementUpdate( USceneComponent* Component, EScopedUpdate::Type ScopeBehavior = EScopedUpdate::DeferredUpdates );
 	~FScopedMovementUpdate();
 
 	enum class EHasMovedTransformOption
@@ -1337,7 +1343,7 @@ private:
 
 private:
 
-	class USceneComponent* Owner;
+	USceneComponent* Owner;
 	FScopedMovementUpdate* OuterDeferredScope;
 	uint32 bDeferUpdates:1;
 	uint32 bHasMoved:1;
@@ -1413,8 +1419,6 @@ FORCEINLINE_DEBUGGABLE void FScopedMovementUpdate::ForceOverlapUpdate()
 	FinalOverlapCandidatesIndex = INDEX_NONE;
 }
 
-
-
 FORCEINLINE_DEBUGGABLE class FScopedMovementUpdate* USceneComponent::GetCurrentScopedMovement() const
 {
 	if (ScopedMovementStack.Num() > 0)
diff --git a/Engine/Source/Runtime/Engine/Classes/Components/SkinnedMeshComponent.h b/Engine/Source/Runtime/Engine/Classes/Components/SkinnedMeshComponent.h
--- a/Engine/Source/Runtime/Engine/Classes/Components/SkinnedMeshComponent.h
+++ b/Engine/Source/Runtime/Engine/Classes/Components/SkinnedMeshComponent.h
@@ -954,4 +954,32 @@ public:
 	virtual bool IsPlayingRootMotionFromEverything(){ return false; }
 
 	bool ShouldUseUpdateRateOptimizations() const;
+
+	friend class FRenderStateRecreator;
 };
+
+class FRenderStateRecreator
+{
+	bool bWasRenderStateCreated;
+	USkinnedMeshComponent* Component;
+
+public:
+
+	FRenderStateRecreator(USkinnedMeshComponent* InActorComponent)
+	{
+		Component = InActorComponent;
+		bWasRenderStateCreated = Component->IsRenderStateCreated();
+		if (bWasRenderStateCreated)
+		{
+			Component->DestroyRenderState_Concurrent();
+		}
+	}
+
+	~FRenderStateRecreator()
+	{
+		if (bWasRenderStateCreated)
+		{
+			Component->CreateRenderState_Concurrent();
+		}
+	}
+};
\ No newline at end of file
diff --git a/Engine/Source/Runtime/Engine/Classes/Components/StaticMeshComponent.h b/Engine/Source/Runtime/Engine/Classes/Components/StaticMeshComponent.h
--- a/Engine/Source/Runtime/Engine/Classes/Components/StaticMeshComponent.h
+++ b/Engine/Source/Runtime/Engine/Classes/Components/StaticMeshComponent.h
@@ -184,10 +184,19 @@ class ENGINE_API UStaticMeshComponent : public UMeshComponent
 	UPROPERTY(transient)
 	uint32 bForceNavigationObstacle : 1;
 
+	/** Use the collision profile specified in the StaticMesh asset.*/
+	UPROPERTY(EditAnywhere, Category = Collision)
+	uint32 bUseDefaultCollision : 1;
+
 	/** If true, mesh painting is disallowed on this instance. Set if vertex colors are overridden in a construction script. */
 	UPROPERTY(EditAnywhere, AdvancedDisplay, BlueprintReadWrite, Category=Rendering)
 	uint32 bDisallowMeshPaintPerInstance : 1;
 
+#if !(UE_BUILD_SHIPPING || UE_BUILD_TEST)
+	/** Option to draw mesh collision in wireframe */
+	uint32 bDrawMeshCollisionWireframe : 1;
+#endif
+
 	/**
 	 *	Ignore this instance of this static mesh when calculating streaming information.
 	 *	This can be useful when doing things like applying character textures to static geometry,
@@ -253,6 +262,8 @@ class ENGINE_API UStaticMeshComponent : public UMeshComponent
 	UFUNCTION(BlueprintCallable, Category="Components|StaticMesh")
 	void GetLocalBounds(FVector& Min, FVector& Max) const;
 
+	virtual void SetCollisionProfileName(FName InCollisionProfileName) override;
+
 public:
 
 	//~ Begin UObject Interface.
@@ -436,6 +447,12 @@ public:
 	*	@param	InSectionIndexPreview		New value of SectionIndexPreview.
 	*/
 	void SetSectionPreview(int32 InSectionIndexPreview);
+	
+	/** Sets the BodyInstance to use the mesh's body setup for external collision information*/
+	void UpdateCollisionFromStaticMesh();
+
+	/** Whether or not the component supports default collision from its static mesh asset */
+	virtual bool SupportsDefaultCollision();
 
 private:
 	/** Initializes the resources used by the static mesh component. */
diff --git a/Engine/Source/Runtime/Engine/Classes/Engine/Engine.h b/Engine/Source/Runtime/Engine/Classes/Engine/Engine.h
--- a/Engine/Source/Runtime/Engine/Classes/Engine/Engine.h
+++ b/Engine/Source/Runtime/Engine/Classes/Engine/Engine.h
@@ -1584,47 +1584,7 @@ public:
 	/** Called by internal engine systems after a level actor has been requested to be renamed */
 	void BroadcastLevelComponentRequestRename(const UActorComponent* InComponent) { LevelComponentRequestRenameEvent.Broadcast(InComponent); }
 
-	/** Editor-only event triggered when a HLOD Actor is moved between clusters */
-	DECLARE_EVENT_TwoParams(UEngine, FHLODActorMovedEvent, const AActor*, const AActor*);
-	FHLODActorMovedEvent& OnHLODActorMoved() { return HLODActorMovedEvent; }
-
-	/** Called by internal engine systems after a HLOD Actor is moved between clusters */
-	void BroadcastHLODActorMoved(const AActor* InActor, const AActor* ParentActor ) { HLODActorMovedEvent.Broadcast(InActor, ParentActor); }
-
-	/** Editor-only event triggered when a HLOD Actor's mesh is build */
-	DECLARE_EVENT_OneParam(UEngine, FHLODMeshBuildEvent, const class ALODActor*);
-	FHLODMeshBuildEvent& OnHLODMeshBuild() { return HLODMeshBuildEvent; }
-
-	/** Called by internal engine systems after a HLOD Actor's mesh is build */
-	void BroadcastHLODMeshBuild(const class ALODActor* InActor) { HLODMeshBuildEvent.Broadcast(InActor); }
-
-	/** Editor-only event triggered when a HLOD Actor is added to a cluster */
-	DECLARE_EVENT_TwoParams(UEngine, FHLODActorAddedEvent, const AActor*, const AActor*);
-	FHLODActorAddedEvent& OnHLODActorAdded() { return HLODActorAddedEvent; }
-
-	/** Called by internal engine systems after a HLOD Actor is added to a cluster */
-	void BroadcastHLODActorAdded(const AActor* InActor, const AActor* ParentActor) { HLODActorAddedEvent.Broadcast(InActor, ParentActor); }
-
-	/** Editor-only event triggered when a HLOD Actor is marked dirty */
-	DECLARE_EVENT_OneParam(UEngine, FHLODActorMarkedDirtyEvent, class ALODActor*);
-	FHLODActorMarkedDirtyEvent& OnHLODActorMarkedDirty() { return HLODActorMarkedDirtyEvent; }
-
-	/** Called by internal engine systems after a HLOD Actor is marked dirty */
-	void BroadcastHLODActorMarkedDirty(class ALODActor* InActor) { HLODActorMarkedDirtyEvent.Broadcast(InActor); }
-
-	/** Editor-only event triggered when a HLOD Actor is marked dirty */
-	DECLARE_EVENT(UEngine, FHLODTransitionScreenSizeChangedEvent);
-	FHLODTransitionScreenSizeChangedEvent& OnHLODTransitionScreenSizeChanged() { return HLODTransitionScreenSizeChangedEvent; }
-
-	/** Called by internal engine systems after a HLOD Actor is marked dirty */
-	void BroadcastHLODTransitionScreenSizeChanged() { HLODTransitionScreenSizeChangedEvent.Broadcast(); }
-
-	/** Editor-only event triggered when a HLOD level is added or removed */
-	DECLARE_EVENT(UEngine, FHLODLevelsArrayChangedEvent);
-	FHLODLevelsArrayChangedEvent& OnHLODLevelsArrayChanged() { return HLODLevelsArrayChangedEvent; }
-
-	/** Called by internal engine systems after a HLOD Actor is marked dirty */
-	void BroadcastHLODLevelsArrayChanged() { HLODLevelsArrayChangedEvent.Broadcast(); }
+	
 
 #endif // #if WITH_EDITOR
 
@@ -2361,25 +2321,7 @@ private:
 
 	/** Broadcasts after an actor has been moved, rotated or scaled */
 	FOnActorMovedEvent		OnActorMovedEvent;
-
-	/** Broadcasts after an HLOD actor has been moved between clusters */	
-	FHLODActorMovedEvent HLODActorMovedEvent;
-
-	/** Broadcasts after an HLOD actor's mesh is build*/
-	FHLODMeshBuildEvent HLODMeshBuildEvent;
 	
-	/** Broadcasts after an HLOD actor has added to a cluster */
-	FHLODActorAddedEvent HLODActorAddedEvent;
-
-	/** Broadcasts after an HLOD actor has been marked dirty */
-	FHLODActorMarkedDirtyEvent HLODActorMarkedDirtyEvent;
-
-	/** Broadcasts after a Draw distance value (World settings) is changed */
-	FHLODTransitionScreenSizeChangedEvent HLODTransitionScreenSizeChangedEvent;
-
-	/** Broadcasts after the HLOD levels array is changed */
-	FHLODLevelsArrayChangedEvent HLODLevelsArrayChangedEvent;
-
 #endif // #if WITH_EDITOR
 
 	/** Thread preventing screen saver from kicking. Suspend most of the time. */
@@ -2758,7 +2700,7 @@ public:
 
 	void AddNewPendingStreamingLevel(UWorld *InWorld, FName PackageName, bool bNewShouldBeLoaded, bool bNewShouldBeVisible, int32 LODIndex);
 
-	bool ShouldCommitPendingMapChange(UWorld *InWorld);
+	bool ShouldCommitPendingMapChange(const UWorld *InWorld) const;
 	void SetShouldCommitPendingMapChange(UWorld *InWorld, bool NewShouldCommitPendingMapChange);
 
 	FSeamlessTravelHandler&	SeamlessTravelHandlerForWorld(UWorld *World);
diff --git a/Engine/Source/Runtime/Engine/Classes/Engine/EngineBaseTypes.h b/Engine/Source/Runtime/Engine/Classes/Engine/EngineBaseTypes.h
--- a/Engine/Source/Runtime/Engine/Classes/Engine/EngineBaseTypes.h
+++ b/Engine/Source/Runtime/Engine/Classes/Engine/EngineBaseTypes.h
@@ -872,13 +872,17 @@ enum EViewModeIndex
 	VMI_CollisionVisibility = 16, 
 	VMI_VertexDensities = 17,
 	/** Colored according to the current LOD index. */
-	VMI_LODColoration = 18,
+	VMI_LODColoration = 18,	
 	/** Colored according to the quad coverage. */
 	VMI_QuadComplexity = 19,
 	/** Colored according to the accuracy of the texture streamer wanted mips computation. */
 	VMI_WantedMipsAccuracy = 20,
 	/** Colored according to the texel factor accuracy. */
 	VMI_TexelFactorAccuracy = 21,
+	/** Colored according to the current HLOD index. */
+	VMI_HLODColoration = 22,
+	/** Group item for LOD and HLOD coloration*/
+	VMI_GroupLODColoration = 23,
 
 	VMI_Max UMETA(Hidden),
 
diff --git a/Engine/Source/Runtime/Engine/Classes/Engine/EngineTypes.h b/Engine/Source/Runtime/Engine/Classes/Engine/EngineTypes.h
--- a/Engine/Source/Runtime/Engine/Classes/Engine/EngineTypes.h
+++ b/Engine/Source/Runtime/Engine/Classes/Engine/EngineTypes.h
@@ -2693,6 +2693,11 @@ struct FMeshProxySettings
 	UPROPERTY()
 	FMaterialSimplificationSettings Material_DEPRECATED;
 
+
+	/** Determines whether or not the correct LOD models should be calculated given the source meshes and transition size */
+	UPROPERTY(EditAnywhere, Category = ProxySettings)
+	bool bCalculateLODSourceModels;
+
 	/** Distance at which meshes should be merged together */
 	UPROPERTY(EditAnywhere, Category = ProxySettings)
 	float MergeDistance;
@@ -2709,22 +2714,6 @@ struct FMeshProxySettings
 	UPROPERTY(EditAnywhere, Category=ProxySettings)
 	bool bRecalculateNormals;
 		
-	/** Set to true to cap the mesh with a ground plane */
-	UPROPERTY(EditAnywhere, Category=ProxySettings)
-	bool bUseClippingPlane;
-
-	/* Ground plane level */
-	UPROPERTY(EditAnywhere, Category=ProxySettings)
-	float ClippingLevel;
-
-	/** Set the axis index for the ground plane (0:X-Axis, 1:Y-Axis, 2:Z-Axis) */
-	UPROPERTY(EditAnywhere, Category=ProxySettings)
-	int32 AxisIndex;
-
-	/** Set to true to use negative halfspace for model, and reject the positive halfspace */
-	UPROPERTY(EditAnywhere, Category=ProxySettings)
-	bool bPlaneNegativeHalfspace;
-
 	UPROPERTY()
 	bool bBakeVertexData_DEPRECATED;
 
@@ -2741,10 +2730,6 @@ struct FMeshProxySettings
 		, HardAngleThreshold(80.0f)
 		, LightMapResolution(256)
 		, bRecalculateNormals(true)
-		, bUseClippingPlane(false)
-		, ClippingLevel(0.0)
-		, AxisIndex(0)
-		, bPlaneNegativeHalfspace(false)
 	{ }
 
 	/** Equality operator. */
diff --git a/Engine/Source/Runtime/Engine/Classes/Engine/HLODMeshCullingVolume.h b/Engine/Source/Runtime/Engine/Classes/Engine/HLODMeshCullingVolume.h
new file mode 100644
--- /dev/null
+++ b/Engine/Source/Runtime/Engine/Classes/Engine/HLODMeshCullingVolume.h
@@ -0,0 +1,16 @@
+// Copyright 1998-2015 Epic Games, Inc. All Rights Reserved.
+
+#pragma once
+
+#include "GameFramework/Volume.h"
+#include "HLODMeshCullingVolume.generated.h"
+
+/** A volume that can be added to an HLOD cluster to remove triangles from source meshes before generating HLOD mesh */
+UCLASS(MinimalAPI)
+class AHLODMeshCullingVolume : public AVolume
+{
+	GENERATED_UCLASS_BODY()
+
+public:
+
+};
diff --git a/Engine/Source/Runtime/Engine/Classes/Engine/NetDriver.h b/Engine/Source/Runtime/Engine/Classes/Engine/NetDriver.h
--- a/Engine/Source/Runtime/Engine/Classes/Engine/NetDriver.h
+++ b/Engine/Source/Runtime/Engine/Classes/Engine/NetDriver.h
@@ -583,10 +583,32 @@ public:
 	virtual AActor* GetActorForGUID(FNetworkGUID InGUID) const { return nullptr; }
 
 	/** Returns the object that manages the list of replicated UObjects. */
-	ENGINE_API virtual FNetworkObjectList& GetNetworkObjectList() { return *NetworkObjects; }
+	ENGINE_API FNetworkObjectList& GetNetworkObjectList() { return *NetworkObjects; }
 
 	/** Returns the object that manages the list of replicated UObjects. */
-	ENGINE_API virtual const FNetworkObjectList& GetNetworkObjectList() const { return *NetworkObjects; }
+	ENGINE_API const FNetworkObjectList& GetNetworkObjectList() const { return *NetworkObjects; }
+
+	/** Get the network object matching the given Actor, or null if not found. */
+	ENGINE_API const FNetworkObjectInfo* GetNetworkActor(const AActor* InActor) const;
+
+	/** Get the network object matching the given Actor, or null if not found. */
+	ENGINE_API FNetworkObjectInfo* GetNetworkActor(const AActor* InActor);
+
+	/**
+	 * Returns whether adaptive net frequency is enabled. If enabled, update frequency is allowed to ramp down to MinNetUpdateFrequency for an actor when no replicated properties have changed.
+	 * This is currently controlled by the CVar "net.UseAdaptiveNetUpdateFrequency".
+	 */
+	ENGINE_API static bool IsAdaptiveNetUpdateFrequencyEnabled();
+
+	/** Returns true if adaptive net update frequency is enabled and the given actor is having its update rate lowered from its standard rate. */
+	ENGINE_API bool IsNetworkActorUpdateFrequencyThrottled(const AActor* InActor) const;
+
+	/** Returns true if adaptive net update frequency is enabled and the given actor is having its update rate lowered from its standard rate. */
+	ENGINE_API bool IsNetworkActorUpdateFrequencyThrottled(const FNetworkObjectInfo& InNetworkActor) const;
+
+	/** Stop adaptive replication for the given actor if it's currently throttled. It maybe be allowed to throttle again later. */
+	ENGINE_API void CancelAdaptiveReplication(FNetworkObjectInfo& InNetworkActor);
+
 
 protected:
 
diff --git a/Engine/Source/Runtime/Engine/Classes/Engine/NetworkObjectList.h b/Engine/Source/Runtime/Engine/Classes/Engine/NetworkObjectList.h
--- a/Engine/Source/Runtime/Engine/Classes/Engine/NetworkObjectList.h
+++ b/Engine/Source/Runtime/Engine/Classes/Engine/NetworkObjectList.h
@@ -90,6 +90,9 @@ public:
 	/** Returns a reference to the set of tracked actors. */
 	FNetworkObjectSet& GetObjects() { return NetworkObjects; }
 
+	/** Returns a const reference to the set of tracked actors. */
+	const FNetworkObjectSet& GetObjects() const { return NetworkObjects; }
+
 private:
 	FNetworkObjectSet NetworkObjects;
 };
diff --git a/Engine/Source/Runtime/Engine/Classes/Engine/StaticMesh.h b/Engine/Source/Runtime/Engine/Classes/Engine/StaticMesh.h
--- a/Engine/Source/Runtime/Engine/Classes/Engine/StaticMesh.h
+++ b/Engine/Source/Runtime/Engine/Classes/Engine/StaticMesh.h
@@ -347,6 +347,14 @@ class UStaticMesh : public UObject, public IInterface_CollisionDataProvider, pub
 	UPROPERTY(EditAnywhere, transient, duplicatetransient, Instanced, Category = StaticMesh)
 	class UBodySetup* BodySetup;
 
+	/** 
+	 *	Specifies which mesh LOD to use for complex (per-poly) collision. 
+	 *	Sometimes it can be desirable to use a lower poly representation for collision to reduce memory usage, improve performance and behaviour.
+	 *	Collision representation does not change based on distance to camera.
+	 */
+	UPROPERTY(EditAnywhere,  Category = StaticMesh)
+	int32 LODForCollision;
+
 	/** True if mesh should use a less-conservative method of mip LOD texture factor computation.
 		requires mesh to be resaved to take effect as algorithm is applied on save. */
 	UPROPERTY(EditAnywhere, AdvancedDisplay, Category=StaticMesh, meta=(ToolTip="If true, use a less-conservative method of mip LOD texture factor computation.  Requires mesh to be resaved to take effect as algorithm is applied on save"))
diff --git a/Engine/Source/Runtime/Engine/Classes/Engine/StaticMeshActor.h b/Engine/Source/Runtime/Engine/Classes/Engine/StaticMeshActor.h
--- a/Engine/Source/Runtime/Engine/Classes/Engine/StaticMeshActor.h
+++ b/Engine/Source/Runtime/Engine/Classes/Engine/StaticMeshActor.h
@@ -39,6 +39,8 @@ public:
 	//~ Begin AActor Interface
 	virtual void CheckForErrors() override;
 	virtual bool GetReferencedContentObjects( TArray<UObject*>& Objects ) const override;
+	virtual void Serialize(FArchive& Ar) override;
+	virtual void PostLoad() override;
 	//~ End AActor Interface
 #endif // WITH_EDITOR	
 
diff --git a/Engine/Source/Runtime/Engine/Classes/Engine/TextureDefines.h b/Engine/Source/Runtime/Engine/Classes/Engine/TextureDefines.h
--- a/Engine/Source/Runtime/Engine/Classes/Engine/TextureDefines.h
+++ b/Engine/Source/Runtime/Engine/Classes/Engine/TextureDefines.h
@@ -51,6 +51,8 @@ enum TextureGroup
 	TEXTUREGROUP_IESLightProfile UMETA(DisplayName="IESLightProfile"),
 	/** Non-filtered, useful for 2D rendering. */
 	TEXTUREGROUP_Pixels2D UMETA(DisplayName="2D Pixels (unfiltered)"),
+	/** Hierarchical LOD generated textures*/
+	TEXTUREGROUP_HierarchicalLOD UMETA(DisplayName = "Hierarchical LOD"),
 	TEXTUREGROUP_MAX,
 };
 
diff --git a/Engine/Source/Runtime/Engine/Classes/Engine/TriggerBase.h b/Engine/Source/Runtime/Engine/Classes/Engine/TriggerBase.h
--- a/Engine/Source/Runtime/Engine/Classes/Engine/TriggerBase.h
+++ b/Engine/Source/Runtime/Engine/Classes/Engine/TriggerBase.h
@@ -26,9 +26,9 @@ private_subobject:
 
 public:
 	/** Returns CollisionComponent subobject **/
-	ENGINE_API UShapeComponent* GetCollisionComponent() const;
+	ENGINE_API UShapeComponent* GetCollisionComponent() const { return CollisionComponent; }
 	/** Returns SpriteComponent subobject **/
-	ENGINE_API UBillboardComponent* GetSpriteComponent() const;
+	ENGINE_API UBillboardComponent* GetSpriteComponent() const { return SpriteComponent; }
 };
 
 
diff --git a/Engine/Source/Runtime/Engine/Classes/Engine/World.h b/Engine/Source/Runtime/Engine/Classes/Engine/World.h
--- a/Engine/Source/Runtime/Engine/Classes/Engine/World.h
+++ b/Engine/Source/Runtime/Engine/Classes/Engine/World.h
@@ -2728,13 +2728,13 @@ public:
 	bool Listen( FURL& InURL );
 
 	/** @return true if this level is a client */
-	bool IsClient();
+	bool IsClient() const;
 
 	/** @return true if this level is a server */
-	bool IsServer();
+	bool IsServer() const;
 
 	/** @return true if the world is in the paused state */
-	bool IsPaused();
+	bool IsPaused() const;
 
 	/**
 	 * Wrapper for DestroyActor() that should be called in the editor.
@@ -3375,6 +3375,32 @@ private:
 //////////////////////////////////////////////////////////////////////////
 // UWorld inlines:
 
+FORCEINLINE_DEBUGGABLE float UWorld::GetTimeSeconds() const
+{
+	return TimeSeconds;
+}
+
+FORCEINLINE_DEBUGGABLE float UWorld::GetRealTimeSeconds() const
+{
+	checkSlow(IsInGameThread());
+	return RealTimeSeconds;
+}
+
+FORCEINLINE_DEBUGGABLE float UWorld::GetAudioTimeSeconds() const
+{
+	return AudioTimeSeconds;
+}
+
+FORCEINLINE_DEBUGGABLE float UWorld::GetDeltaSeconds() const
+{
+	return DeltaTimeSeconds;
+}
+
+FORCEINLINE_DEBUGGABLE float UWorld::TimeSince(float Time) const
+{
+	return GetTimeSeconds() - Time;
+}
+
 FORCEINLINE_DEBUGGABLE bool UWorld::ComponentOverlapMulti(TArray<struct FOverlapResult>& OutOverlaps, const class UPrimitiveComponent* PrimComp, const FVector& Pos, const FRotator& Rot, const struct FComponentQueryParams& Params, const struct FCollisionObjectQueryParams& ObjectQueryParams) const
 {
 	// Pass through to FQuat version.
diff --git a/Engine/Source/Runtime/Engine/Classes/GameFramework/Actor.h b/Engine/Source/Runtime/Engine/Classes/GameFramework/Actor.h
--- a/Engine/Source/Runtime/Engine/Classes/GameFramework/Actor.h
+++ b/Engine/Source/Runtime/Engine/Classes/GameFramework/Actor.h
@@ -2813,6 +2813,10 @@ FORCEINLINE_DEBUGGABLE ENetRole AActor::GetRemoteRole() const
 	return RemoteRole;
 }
 
+FORCEINLINE_DEBUGGABLE void AActor::SetNetUpdateTime(float NewUpdateTime)
+{
+	NetUpdateTime = NewUpdateTime;
+}
 
 //////////////////////////////////////////////////////////////////////////
 // Macro to hide common Transform functions in native code for classes where they don't make sense.
diff --git a/Engine/Source/Runtime/Engine/Classes/GameFramework/CharacterMovementComponent.h b/Engine/Source/Runtime/Engine/Classes/GameFramework/CharacterMovementComponent.h
--- a/Engine/Source/Runtime/Engine/Classes/GameFramework/CharacterMovementComponent.h
+++ b/Engine/Source/Runtime/Engine/Classes/GameFramework/CharacterMovementComponent.h
@@ -452,6 +452,10 @@ public:
 	UPROPERTY(Category="Character Movement: Physics Interaction", EditAnywhere, BlueprintReadWrite, meta=(editcondition = "bEnablePhysicsInteraction"))
 	bool bPushForceScaledToMass;
 
+	/** If enabled, the PushForce location is moved using PushForcePointZOffsetFactor. Otherwise simply use the impact point. */
+	UPROPERTY(Category = "Character Movement: Physics Interaction", EditAnywhere, BlueprintReadWrite, meta = (editcondition = "bEnablePhysicsInteraction"))
+	bool bPushForceUsingZOffset;
+
 	/** If enabled, the applied push force will try to get the physics object to the same velocity than the player, not faster. This will only
 		scale the force down, it will never apply more force than defined by PushForceFactor. */
 	UPROPERTY(Category="Character Movement: Physics Interaction", EditAnywhere, BlueprintReadWrite, meta=(editcondition = "bEnablePhysicsInteraction"))
@@ -513,13 +517,19 @@ protected:
 	FVector Acceleration;
 
 	/**
-	 * Location after last PerformMovement update. Used internally to detect changes in position from outside character movement to try to validate the current floor.
+	 * Location after last PerformMovement or SimulateMovement update. Used internally to detect changes in position from outside character movement to try to validate the current floor.
 	 */
 	UPROPERTY()
 	FVector LastUpdateLocation;
 
 	/**
-	 * Velocity after last PerformMovement update. Used internally to detect changes in velocity from external sources.
+	 * Rotation after last PerformMovement or SimulateMovement update.
+	 */
+	UPROPERTY()
+	FQuat LastUpdateRotation;
+
+	/**
+	 * Velocity after last PerformMovement or SimulateMovement update. Used internally to detect changes in velocity from external sources.
 	 */
 	UPROPERTY()
 	FVector LastUpdateVelocity;
@@ -1840,6 +1850,9 @@ protected:
 	/** Ticks the characters pose and accumulates root motion */
 	void TickCharacterPose(float DeltaTime);
 
+	/** On the server if we know we are having our replication rate throttled, this method checks if important replicated properties have changed that should cause us to return to the normal replication rate. */
+	virtual bool ShouldCancelAdaptiveReplication() const;
+
 public:
 
 	/** React to instantaneous change in position. Invalidates cached floor recomputes it if possible if there is a current movement base. */
diff --git a/Engine/Source/Runtime/Engine/Classes/GameFramework/MovementComponent.h b/Engine/Source/Runtime/Engine/Classes/GameFramework/MovementComponent.h
--- a/Engine/Source/Runtime/Engine/Classes/GameFramework/MovementComponent.h
+++ b/Engine/Source/Runtime/Engine/Classes/GameFramework/MovementComponent.h
@@ -77,11 +77,11 @@ class ENGINE_API UMovementComponent : public UActorComponent
 	 * If true, movement will be constrained to a plane.
 	 * @see PlaneConstraintNormal, PlaneConstraintOrigin, PlaneConstraintAxisSetting
 	 */
-	UPROPERTY(EditDefaultsOnly, BlueprintReadOnly, Category=PlanarMovement)
+	UPROPERTY(EditAnywhere, BlueprintReadOnly, Category=PlanarMovement)
 	uint32 bConstrainToPlane:1;
 
 	/** If true and plane constraints are enabled, then the updated component will be snapped to the plane when first attached. */
-	UPROPERTY(EditDefaultsOnly, BlueprintReadOnly, Category=PlanarMovement, meta=(editcondition=bConstrainToPlane))
+	UPROPERTY(EditAnywhere, BlueprintReadOnly, Category=PlanarMovement, meta=(editcondition=bConstrainToPlane))
 	uint32 bSnapToPlaneAtStart:1;
 
 private:
@@ -96,7 +96,7 @@ private:
 	 * so that movement along the locked axis is not be possible.
 	 * @see SetPlaneConstraintAxisSetting
 	 */
-	UPROPERTY(EditDefaultsOnly, Category=PlanarMovement, meta=(editcondition=bConstrainToPlane))
+	UPROPERTY(EditAnywhere, Category=PlanarMovement, meta=(editcondition=bConstrainToPlane))
 	EPlaneConstraintAxisSetting PlaneConstraintAxisSetting;
 
 protected:
@@ -104,10 +104,10 @@ protected:
 	/**
 	 * The normal or axis of the plane that constrains movement, if bConstrainToPlane is enabled.
 	 * If for example you wanted to constrain movement to the X-Z plane (so that Y cannot change), the normal would be set to X=0 Y=1 Z=0.
-	 * This is recalculated whenever PlaneConstraintAxisSetting changes.
+	 * This is recalculated whenever PlaneConstraintAxisSetting changes. It is normalized once the component is registered with the game world.
 	 * @see bConstrainToPlane, SetPlaneConstraintNormal(), SetPlaneConstraintFromVectors()
 	 */
-	UPROPERTY(EditDefaultsOnly, BlueprintReadOnly, Category=PlanarMovement, meta=(editcondition=bConstrainToPlane))
+	UPROPERTY(EditAnywhere, BlueprintReadOnly, Category=PlanarMovement, meta=(editcondition=bConstrainToPlane))
 	FVector PlaneConstraintNormal;
 
 	/**
@@ -115,7 +115,7 @@ protected:
 	 * This defines the behavior of snapping a position to the plane, such as by SnapUpdatedComponentToPlane().
 	 * @see bConstrainToPlane, SetPlaneConstraintOrigin().
 	 */
-	UPROPERTY(EditDefaultsOnly, BlueprintReadOnly, Category=PlanarMovement, meta=(editcondition=bConstrainToPlane))
+	UPROPERTY(EditAnywhere, BlueprintReadOnly, Category=PlanarMovement, meta=(editcondition=bConstrainToPlane))
 	FVector PlaneConstraintOrigin;
 
 	/**
diff --git a/Engine/Source/Runtime/Engine/Classes/Kismet/KismetMathLibrary.h b/Engine/Source/Runtime/Engine/Classes/Kismet/KismetMathLibrary.h
--- a/Engine/Source/Runtime/Engine/Classes/Kismet/KismetMathLibrary.h
+++ b/Engine/Source/Runtime/Engine/Classes/Kismet/KismetMathLibrary.h
@@ -55,6 +55,23 @@ namespace EEasingFunc
 	};
 }
 
+/** Different methods for interpolating rotation between transforms */
+UENUM(BlueprintType)
+namespace ELerpInterpolationMode
+{
+	enum Type
+	{
+		/** Shortest Path or Quaternion interpolation for the rotation. */
+		QuatInterp,
+
+		/** Rotor or Euler Angle interpolation. */
+		EulerInterp,
+
+		/** Dual quaternion interpolation, follows helix or screw-motion path between keyframes.   */
+		DualQuatInterp
+	};
+}
+
 UCLASS()
 class ENGINE_API UKismetMathLibrary : public UBlueprintFunctionLibrary
 {
@@ -1546,8 +1563,8 @@ class ENGINE_API UKismetMathLibrary : public UBlueprintFunctionLibrary
 	static FTransform InvertTransform(const FTransform& T);
 
 	/** Linearly interpolates between A and B based on Alpha (100% of A when Alpha=0 and 100% of B when Alpha=1). */
-	UFUNCTION(BlueprintPure, meta=(DisplayName = "Lerp (Transform)"), Category="Math|Transform")
-	static FTransform TLerp(const FTransform& A, const FTransform& B, float Alpha);
+	UFUNCTION(BlueprintPure, meta=(DisplayName = "Lerp (Transform)", AdvancedDisplay = "3"), Category="Math|Transform")
+	static FTransform TLerp(const FTransform& A, const FTransform& B, float Alpha, TEnumAsByte<ELerpInterpolationMode::Type> InterpMode = ELerpInterpolationMode::QuatInterp);
 
 	/** Ease between A and B using a specified easing function. */
 	UFUNCTION(BlueprintPure, meta = (DisplayName = "Ease (Transform)", BlueprintInternalUseOnly = "true"), Category = "Math|Interpolation")
diff --git a/Engine/Source/Runtime/Engine/Classes/Kismet/KismetSystemLibrary.h b/Engine/Source/Runtime/Engine/Classes/Kismet/KismetSystemLibrary.h
--- a/Engine/Source/Runtime/Engine/Classes/Kismet/KismetSystemLibrary.h
+++ b/Engine/Source/Runtime/Engine/Classes/Kismet/KismetSystemLibrary.h
@@ -1484,10 +1484,10 @@ class ENGINE_API UKismetSystemLibrary : public UBlueprintFunctionLibrary
 	 * Sets the state of the transition message rendered by the viewport. (The blue text displayed when the game is paused and so forth.)
 	 *
 	 * @param WorldContextObject	World context
-	 * @param State					set true to supress transition message
+	 * @param State					set true to suppress transition message
 	 */
 	UFUNCTION(BlueprintCallable, Category = "Utilities", meta = (WorldContext="WorldContextObject"))
-	static void SetSupressViewportTransitionMessage(UObject* WorldContextObject, bool bState);
+	static void SetSuppressViewportTransitionMessage(UObject* WorldContextObject, bool bState);
 
 	/**
 	 * Returns an array of the user's preferred languages in order of preference
diff --git a/Engine/Source/Runtime/Engine/Classes/PhysicsEngine/BodyInstance.h b/Engine/Source/Runtime/Engine/Classes/PhysicsEngine/BodyInstance.h
--- a/Engine/Source/Runtime/Engine/Classes/PhysicsEngine/BodyInstance.h
+++ b/Engine/Source/Runtime/Engine/Classes/PhysicsEngine/BodyInstance.h
@@ -12,9 +12,10 @@ struct FCollisionShape;
 struct FConstraintInstance;
 class UPhysicsConstraintComponent;
 enum class ETeleportType;
+class UBodySetup;
 
 /** Delegate for applying custom physics forces upon the body. Can be passed to "AddCustomPhysics" so 
-  * custom forces and torques can be calculated induvidually for every physics substep.
+  * custom forces and torques can be calculated individually for every physics substep.
   * The function provides delta time for a physics step and pointer to body instance upon which forces must be added.
   * 
   * Do not expect this callback to be called from the main game thread! It may get called from a physics simulation thread. */
@@ -148,7 +149,6 @@ struct ENGINE_API FBodyInstance
 	int16 InstanceBoneIndex;
 
 	/** Current scale of physics - used to know when and how physics must be rescaled to match current transform of OwnerComponent. */
-	UPROPERTY()
 	FVector Scale3D;
 
 	/** Physics scene index for the synchronous scene. */
@@ -284,12 +284,22 @@ protected:
 	UPROPERTY(EditAnywhere, AdvancedDisplay, BlueprintReadOnly, Category = Physics, meta = (editcondition = "bOverrideMaxDepenetrationVelocity", ClampMin = "0.0", UIMin = "0.0"))
 	float MaxDepenetrationVelocity;
 
-public:
+	/** The body setup holding the default body instance and its collision profile. */
+	TWeakObjectPtr<UBodySetup> ExternalCollisionProfileBodySetup;
 
+	
 	/**Mass of the body in KG. By default we compute this based on physical material and mass scale.
 	*@see bOverrideMass to set this directly */
-	UPROPERTY(EditAnywhere, BlueprintReadOnly, Category = Physics, meta = (editcondition = "bOverrideMass", ClampMin = "0.001", UIMin = "0.001"))
-	float MassInKg;
+	UPROPERTY(EditAnywhere, BlueprintReadOnly, Category = Physics, meta = (editcondition = "bOverrideMass", ClampMin = "0.001", UIMin = "0.001", DisplayName = "MassInKg"))
+	float MassInKgOverride;
+
+public:
+
+	/** Returns the mass override. See MassInKgOverride for documentation */
+	float GetMassOverride() const { return MassInKgOverride; }
+
+	/** Sets the mass override */
+	void SetMassOverride(float MassInKG);
 
 	/** 'Drag' force added to reduce linear movement */
 	UPROPERTY(EditAnywhere, BlueprintReadWrite, Category = Physics)
@@ -311,6 +321,11 @@ public:
 	UPROPERTY(EditAnywhere, AdvancedDisplay, BlueprintReadWrite, Category = Physics)
 	float MassScale;
 
+	/** Use the collision profile found in the given BodySetup's default BodyInstance */
+	void UseExternalCollisionProfile(UBodySetup* InExternalCollisionProfileBodySetup);
+
+	void ClearExternalCollisionProfile();
+
 	/** Locks physical movement along axis. */
 	void SetDOFLock(EDOFMode::Type NewDOFMode);
 
@@ -398,7 +413,7 @@ public:
 	TWeakObjectPtr<class UPrimitiveComponent> OwnerComponent;
 
 	/** BodySetup pointer that this instance is initialized from */
-	TWeakObjectPtr<class UBodySetup> BodySetup;
+	TWeakObjectPtr<UBodySetup> BodySetup;
 
 	/** Constructor **/
 	FBodyInstance();
@@ -426,7 +441,7 @@ public:
 	 *	@param InRBScene The physics scene to place the body into
 	 *	@param InAggregate An aggregate to place the body into
 	 */
-	void InitBody(class UBodySetup* Setup, const FTransform& Transform, class UPrimitiveComponent* PrimComp, class FPhysScene* InRBScene, PhysXAggregateType InAggregate = NULL);
+	void InitBody(UBodySetup* Setup, const FTransform& Transform, class UPrimitiveComponent* PrimComp, class FPhysScene* InRBScene, PhysXAggregateType InAggregate = NULL);
 
 	/** Validate a body transform, outputting debug info
 	 *	@param Transform Transform to debug
@@ -443,7 +458,7 @@ public:
 	 *	@param InRBScene
 	 *  @param PhysicsSerializer
 	 */
-	static void InitStaticBodies(const TArray<FBodyInstance*>& Bodies, const TArray<FTransform>& Transforms, class UBodySetup* BodySetup, class UPrimitiveComponent* PrimitiveComp, class FPhysScene* InRBScene, class UPhysicsSerializer* PhysicsSerializer);
+	static void InitStaticBodies(const TArray<FBodyInstance*>& Bodies, const TArray<FTransform>& Transforms, UBodySetup* BodySetup, class UPrimitiveComponent* PrimitiveComp, class FPhysScene* InRBScene, class UPhysicsSerializer* PhysicsSerializer);
 
 	/** Obtains the appropriate PhysX scene lock for READING and executes the passed in lambda. */
 	void ExecuteOnPhysicsReadOnly(TFunctionRef<void()> Func) const;
@@ -813,9 +828,8 @@ public:
 
 	/** Return the ignore mask filter. */
 	FORCEINLINE FMaskFilter GetMaskFilter() const { return MaskFilter; }
-
-	/** Get the current collision profile assigned to this body */
-	FORCEINLINE_DEBUGGABLE FName GetCollisionProfileName() const { return CollisionProfileName; }
+	/** Returns the collision profile name that will be used. */
+	FName GetCollisionProfileName() const;
 
 	/** return true if it uses Collision Profile System. False otherwise*/
 	bool DoesUseCollisionProfile() const;
diff --git a/Engine/Source/Runtime/Engine/Classes/PhysicsEngine/ConstraintInstance.h b/Engine/Source/Runtime/Engine/Classes/PhysicsEngine/ConstraintInstance.h
--- a/Engine/Source/Runtime/Engine/Classes/PhysicsEngine/ConstraintInstance.h
+++ b/Engine/Source/Runtime/Engine/Classes/PhysicsEngine/ConstraintInstance.h
@@ -251,7 +251,12 @@ struct ENGINE_API FConstraintInstance
 	UPROPERTY(EditAnywhere, AdvancedDisplay, Category=Angular, meta=(editcondition = "bAngularBreakable", ClampMin = "0.0"))
 	float AngularBreakThreshold;
 
+	/** If true, linear limits scale using the absolute min of the 3d scale of the owning component */
+	UPROPERTY(EditAnywhere, Category = Linear)
+	uint32 bScaleLinearLimits : 1;
+
 private:
+
 	/** Enables/Disables linear position drive along the x axis. */
 	UPROPERTY(EditAnywhere, Category=LinearMotor)
 	uint32 bLinearXPositionDrive:1;
@@ -369,7 +374,11 @@ public:
 
 	float AverageMass;
 
+private:
+	/** The component scale passed in during initialization*/
+	float LastKnownScale;
 
+public:
 	/** Sets the LinearX Motion Type
 	*	@param MotionType	New Motion Type
 	*/
@@ -515,7 +524,7 @@ public:
 
 private:
 #if WITH_PHYSX 
-	bool CreatePxJoint_AssumesLocked(physx::PxRigidActor* PActor1, physx::PxRigidActor* PActor2, physx::PxScene* PScene, const float Scale);
+	bool CreatePxJoint_AssumesLocked(physx::PxRigidActor* PActor1, physx::PxRigidActor* PActor2, physx::PxScene* PScene);
 	void UpdateConstraintFlags_AssumesLocked();
 	void UpdateAverageMass_AssumesLocked(const physx::PxRigidActor* PActor1, const physx::PxRigidActor* PActor2);
 	physx::PxD6Joint* GetUnbrokenJoint_AssumesLocked() const;
diff --git a/Engine/Source/Runtime/Engine/Classes/PhysicsEngine/PhysicsConstraintComponent.h b/Engine/Source/Runtime/Engine/Classes/PhysicsEngine/PhysicsConstraintComponent.h
--- a/Engine/Source/Runtime/Engine/Classes/PhysicsEngine/PhysicsConstraintComponent.h
+++ b/Engine/Source/Runtime/Engine/Classes/PhysicsEngine/PhysicsConstraintComponent.h
@@ -275,5 +275,8 @@ protected:
 
 	/** Routes the FConstraint callback to the dynamic delegate */
 	void OnConstraintBrokenHandler(FConstraintInstance* BrokenConstraint);
+
+	/** Returns the scale of the constraint as it will be passed into the ConstraintInstance*/
+	float GetConstraintScale() const;
 };
 
diff --git a/Engine/Source/Runtime/Engine/Classes/PhysicsEngine/PhysicsSettings.h b/Engine/Source/Runtime/Engine/Classes/PhysicsEngine/PhysicsSettings.h
--- a/Engine/Source/Runtime/Engine/Classes/PhysicsEngine/PhysicsSettings.h
+++ b/Engine/Source/Runtime/Engine/Classes/PhysicsEngine/PhysicsSettings.h
@@ -162,6 +162,11 @@ class ENGINE_API UPhysicsSettings : public UDeveloperSettings
 	UPROPERTY(config, EditAnywhere, Category = Optimization)
 	bool bSuppressFaceRemapTable;
 
+	/**
+	* If true, physx will not update unreal with any bodies that have moved during the simulation. This should only be used if you have no physx simulation or you are manually updating the unreal data via polling physx.  */
+	UPROPERTY(config, EditAnywhere, Category = Optimization)
+	bool bDisableActiveTransforms;
+
 	/**
 	*  If true CCD will be ignored. This is an optimization when CCD is never used which removes the need for physx to check it internally. */
 	UPROPERTY(config, EditAnywhere, Category = Simulation)
diff --git a/Engine/Source/Runtime/Engine/Classes/Sound/SoundBase.h b/Engine/Source/Runtime/Engine/Classes/Sound/SoundBase.h
--- a/Engine/Source/Runtime/Engine/Classes/Sound/SoundBase.h
+++ b/Engine/Source/Runtime/Engine/Classes/Sound/SoundBase.h
@@ -26,6 +26,10 @@ class ENGINE_API USoundBase : public UObject
 {
 	GENERATED_UCLASS_BODY()
 
+private:
+	static USoundClass* DefaultSoundClassObject;
+	static USoundConcurrency* DefaultSoundConcurrencyObject;
+
 protected:
 	/** Sound class this sound belongs to */
 	UPROPERTY(EditAnywhere, Category=Sound, meta=(DisplayName = "Sound Class"))
@@ -66,7 +70,7 @@ public:
 	USoundAttenuation* AttenuationSettings;
 
 	/** Sound priority (higher value is higher priority) used for concurrency resolution. This priority value is weighted against the final volume of the sound. */
-	UPROPERTY(EditAnywhere, BlueprintReadWrite, Category = Concurrency)
+	UPROPERTY(EditAnywhere, BlueprintReadWrite, Category = Concurrency, meta = (ClampMin = "0.0", UIMin = "0.0", ClampMax = "100.0", UIMax = "100.0") )
 	float Priority;
 
 public:	
@@ -97,6 +101,11 @@ public:
 	virtual float GetVolumeMultiplier();
 	virtual float GetPitchMultiplier();
 
+	/**
+	  * Returns the subtitle priority
+	  */
+	virtual float GetSubtitlePriority() const { return 10000.f; };
+
 	/** Returns whether or not this sound is looping. */
 	bool IsLooping();
 
diff --git a/Engine/Source/Runtime/Engine/Classes/Sound/SoundWave.h b/Engine/Source/Runtime/Engine/Classes/Sound/SoundWave.h
--- a/Engine/Source/Runtime/Engine/Classes/Sound/SoundWave.h
+++ b/Engine/Source/Runtime/Engine/Classes/Sound/SoundWave.h
@@ -163,6 +163,9 @@ class ENGINE_API USoundWave : public USoundBase
 	UPROPERTY(EditAnywhere, Category=Subtitles )
 	FString SpokenText;
 
+	/** The priority of the subtitle. */
+	UPROPERTY(EditAnywhere, Category=Subtitles)
+	float SubtitlePriority;
 
 	/** Playback volume of sound 0 to 1 - Default is 1.0. */
 	UPROPERTY(Category=Sound, meta=(ClampMin = "0.0"), EditAnywhere)
@@ -290,6 +293,7 @@ public:
 	virtual void Parse( class FAudioDevice* AudioDevice, const UPTRINT NodeWaveInstanceHash, FActiveSound& ActiveSound, const FSoundParseParameters& ParseParams, TArray<FWaveInstance*>& WaveInstances ) override;
 	virtual float GetMaxAudibleDistance() override;
 	virtual float GetDuration() override;
+	virtual float GetSubtitlePriority() const override;
 	//~ End USoundBase Interface.
 
 	/**
diff --git a/Engine/Source/Runtime/Engine/Private/Actor.cpp b/Engine/Source/Runtime/Engine/Private/Actor.cpp
--- a/Engine/Source/Runtime/Engine/Private/Actor.cpp
+++ b/Engine/Source/Runtime/Engine/Private/Actor.cpp
@@ -305,7 +305,7 @@ FVector AActor::GetVelocity() const
 
 void AActor::ClearCrossLevelReferences()
 {
-	if(RootComponent && GetRootComponent()->AttachParent && (GetOutermost() != GetRootComponent()->AttachParent->GetOutermost()))
+	if(RootComponent && GetRootComponent()->GetAttachParent() && (GetOutermost() != GetRootComponent()->GetAttachParent()->GetOutermost()))
 	{
 		GetRootComponent()->DetachFromParent();
 	}
@@ -582,9 +582,9 @@ void AActor::PostLoadSubobjects(FObjectInstancingGraph* OuterInstanceGraph)
 		GetRootComponent()->RelativeScale3D = OldScale;
 		
 		// Migrate any attachment to the new root
-		if(OldRoot->AttachParent)
+		if(OldRoot->GetAttachParent())
 		{
-			RootComponent->AttachTo(OldRoot->AttachParent);
+			RootComponent->AttachTo(OldRoot->GetAttachParent());
 			OldRoot->DetachFromParent();
 		}
 
@@ -620,7 +620,7 @@ void AActor::ProcessEvent(UFunction* Function, void* Parameters)
 void AActor::ApplyWorldOffset(const FVector& InOffset, bool bWorldShift)
 {
 	// Attached components will be shifted by parents, will shift only USceneComponents derived components
-	if (RootComponent != nullptr && RootComponent->AttachParent == nullptr)
+	if (RootComponent != nullptr && RootComponent->GetAttachParent() == nullptr)
 	{
 		RootComponent->ApplyWorldOffset(InOffset, bWorldShift);
 	}
@@ -854,7 +854,7 @@ void AActor::PreReplication( IRepChangedPropertyTracker & ChangedPropertyTracker
 	// Attachment replication gets filled in by GatherCurrentMovement(), but in the case of a detached root we need to trigger remote detachment.
 	AttachmentReplication.AttachParent = nullptr;
 
-	if ( bReplicateMovement || (RootComponent && RootComponent->AttachParent) )
+	if ( bReplicateMovement || (RootComponent && RootComponent->GetAttachParent()) )
 	{
 		GatherCurrentMovement();
 	}
@@ -1442,7 +1442,7 @@ AActor* AActor::GetAttachParentActor() const
 {
 	if(RootComponent != NULL)
 	{
-		for( const USceneComponent* Test=GetRootComponent()->AttachParent; Test!=NULL; Test=Test->AttachParent )
+		for( const USceneComponent* Test=GetRootComponent()->GetAttachParent(); Test!=NULL; Test=Test->GetAttachParent())
 		{
 			AActor* TestOwner = Test->GetOwner();
 			if( TestOwner != this )
@@ -1459,12 +1459,12 @@ FName AActor::GetAttachParentSocketName() const
 {
 	if(RootComponent != NULL)
 	{
-		for( const USceneComponent* Test=GetRootComponent(); Test!=NULL; Test=Test->AttachParent )
+		for( const USceneComponent* Test=GetRootComponent(); Test!=NULL; Test=Test->GetAttachParent() )
 		{
 			AActor* TestOwner = Test->GetOwner();
 			if( TestOwner != this )
 			{
-				return Test->AttachSocketName;
+				return Test->GetAttachSocketName();
 			}
 		}
 	}
@@ -1508,10 +1508,8 @@ void AActor::GetAttachedActors(TArray<class AActor*>& OutActors) const
 					else
 					{
 						// This component is owned by us, we need to add its children
-						for (int32 i = 0; i < SceneComp->AttachChildren.Num(); ++i)
+						for (USceneComponent* ChildComp : SceneComp->GetAttachChildren())
 						{
-							USceneComponent* ChildComp = SceneComp->AttachChildren[i];
-
 							// Add any we have not explored yet to the set to check
 							if ((ChildComp != NULL) && !CheckedComps.Contains(ChildComp))
 							{
@@ -1602,11 +1600,6 @@ bool AActor::IsRelevancyOwnerFor(const AActor* ReplicatedActor, const AActor* Ac
 	return (ActorOwner == this);
 }
 
-void AActor::SetNetUpdateTime(float NewUpdateTime)
-{
-	NetUpdateTime = NewUpdateTime;
-}
-
 void AActor::ForceNetUpdate()
 {
 	if (NetDormancy > DORM_Awake)
@@ -2486,7 +2479,7 @@ static USceneComponent* FixupNativeActorComponents(AActor* Actor)
 		for (USceneComponent* Component : SceneComponents)
 		{
 			if ((Component == nullptr) ||
-				(Component->AttachParent != nullptr) ||
+				(Component->GetAttachParent() != nullptr) ||
 				(Component->CreationMethod != EComponentCreationMethod::Native))
 			{
 				continue;
@@ -3378,13 +3371,13 @@ AWorldSettings * AActor::GetWorldSettings() const
 ENetMode AActor::GetNetMode() const
 {
 	UNetDriver *NetDriver = GetNetDriver();
-	if (NetDriver != NULL)
+	if (NetDriver != nullptr)
 	{
 		return NetDriver->GetNetMode();
 	}
 
 	UWorld* World = GetWorld();
-	if (World != NULL && World->DemoNetDriver != NULL)
+	if (World != nullptr && World->DemoNetDriver != nullptr)
 	{
 		return World->DemoNetDriver->GetNetMode();
 	}
@@ -3397,8 +3390,7 @@ UNetDriver* AActor::GetNetDriver() const
 	UWorld *World = GetWorld();
 	if (NetDriverName == NAME_GameNetDriver)
 	{
-		check(World);
-		return World->GetNetDriver();
+		return (World ? World->GetNetDriver() : nullptr);
 	}
 
 	return GEngine->FindNamedNetDriver(World, NetDriverName);
@@ -3713,11 +3705,11 @@ static USceneComponent* GetUnregisteredParent(UActorComponent* Component)
 	USceneComponent* SceneComponent = Cast<USceneComponent>(Component);
 	
 	while (	SceneComponent && 
-			SceneComponent->AttachParent && 
-			SceneComponent->AttachParent->GetOwner() == Component->GetOwner() && 
-			!SceneComponent->AttachParent->IsRegistered())
+			SceneComponent->GetAttachParent() &&
+			SceneComponent->GetAttachParent()->GetOwner() == Component->GetOwner() &&
+			!SceneComponent->GetAttachParent()->IsRegistered())
 	{
-		SceneComponent = SceneComponent->AttachParent;
+		SceneComponent = SceneComponent->GetAttachParent();
 		if (SceneComponent->bAutoRegister && !SceneComponent->IsPendingKill())
 		{
 			// We found unregistered parent that should be registered
@@ -3930,9 +3922,9 @@ void AActor::DrawDebugComponents(FColor const& BaseColor) const
 		DrawDebugCoordinateSystem(MyWorld, Loc, Rot, 10.f);
 
 		// draw line from me to my parent
-		if (Component->AttachParent)
+		if (Component->GetAttachParent())
 		{
-			DrawDebugLine(MyWorld, Component->AttachParent->GetComponentLocation(), Loc, BaseColor);
+			DrawDebugLine(MyWorld, Component->GetAttachParent()->GetComponentLocation(), Loc, BaseColor);
 		}
 
 		// draw component name
diff --git a/Engine/Source/Runtime/Engine/Private/ActorConstruction.cpp b/Engine/Source/Runtime/Engine/Private/ActorConstruction.cpp
--- a/Engine/Source/Runtime/Engine/Private/ActorConstruction.cpp
+++ b/Engine/Source/Runtime/Engine/Private/ActorConstruction.cpp
@@ -104,7 +104,7 @@ void AActor::DestroyConstructedComponents()
 	{
 		if (USceneComponent* BSC = Cast<USceneComponent>(&B))
 		{
-			if (BSC->AttachParent == &A)
+			if (BSC->GetAttachParent() == &A)
 			{
 				return false;
 			}
@@ -266,12 +266,12 @@ void AActor::RerunConstructionScripts()
 				{
 					USceneComponent* EachRoot = AttachedActor->GetRootComponent();
 					// If the component we are attached to is about to go away...
-					if (EachRoot && EachRoot->AttachParent && EachRoot->AttachParent->IsCreatedByConstructionScript())
+					if (EachRoot && EachRoot->GetAttachParent() && EachRoot->GetAttachParent()->IsCreatedByConstructionScript())
 					{
 						// Save info about actor to reattach
 						FAttachedActorInfo Info;
 						Info.AttachedActor = AttachedActor;
-						Info.AttachedToSocket = EachRoot->AttachSocketName;
+						Info.AttachedToSocket = EachRoot->GetAttachSocketName();
 						Info.bSetRelativeTransform = false;
 						AttachedActorInfos.Add(Info);
 
@@ -290,17 +290,17 @@ void AActor::RerunConstructionScripts()
 		if (bUseRootComponentProperties && RootComponent != nullptr)
 		{
 			// Do not need to detach if root component is not going away
-			if (RootComponent->AttachParent != NULL && RootComponent->IsCreatedByConstructionScript())
+			if (RootComponent->GetAttachParent() != NULL && RootComponent->IsCreatedByConstructionScript())
 			{
-				Parent = RootComponent->AttachParent->GetOwner();
+				Parent = RootComponent->GetAttachParent()->GetOwner();
 				// Root component should never be attached to another component in the same actor!
 				if (Parent == this)
 				{
 					UE_LOG(LogActor, Warning, TEXT("RerunConstructionScripts: RootComponent (%s) attached to another component in this Actor (%s)."), *RootComponent->GetPathName(), *Parent->GetPathName());
 					Parent = NULL;
 				}
-				AttachParentComponent = RootComponent->AttachParent;
-				SocketName = RootComponent->AttachSocketName;
+				AttachParentComponent = RootComponent->GetAttachParent();
+				SocketName = RootComponent->GetAttachSocketName();
 				//detach it to remove any scaling 
 				RootComponent->DetachFromParent(true);
 			}
@@ -376,7 +376,7 @@ void AActor::RerunConstructionScripts()
 			if (!Info.AttachedActor->IsPendingKill() && Info.AttachedActor->GetAttachParentActor() == nullptr)
 			{
 				USceneComponent* ChildRoot = Info.AttachedActor->GetRootComponent();
-				if (ChildRoot && ChildRoot->AttachParent != RootComponent)
+				if (ChildRoot && ChildRoot->GetAttachParent() != RootComponent)
 				{
 					ChildRoot->AttachTo(RootComponent, Info.AttachedToSocket, EAttachLocation::KeepWorldPosition);
 					if (Info.bSetRelativeTransform)
@@ -551,7 +551,7 @@ void AActor::ProcessUserConstructionScript()
 	for (auto SceneComponent : SceneComponents)
 	{
 		// A parent component can't be more mobile than its children, so we check for that here and adjust as needed.
-		if(SceneComponent != RootComponent && SceneComponent->AttachParent != nullptr && SceneComponent->AttachParent->Mobility > SceneComponent->Mobility)
+		if(SceneComponent != RootComponent && SceneComponent->GetAttachParent() != nullptr && SceneComponent->GetAttachParent()->Mobility > SceneComponent->Mobility)
 		{
 			if(SceneComponent->IsA<UStaticMeshComponent>())
 			{
@@ -561,7 +561,7 @@ void AActor::ProcessUserConstructionScript()
 			else
 			{
 				// Set the new component (and any children) to be at least as mobile as its parent
-				SceneComponent->SetMobility(SceneComponent->AttachParent->Mobility);
+				SceneComponent->SetMobility(SceneComponent->GetAttachParent()->Mobility);
 			}
 		}
 	}
diff --git a/Engine/Source/Runtime/Engine/Private/ActorEditor.cpp b/Engine/Source/Runtime/Engine/Private/ActorEditor.cpp
--- a/Engine/Source/Runtime/Engine/Private/ActorEditor.cpp
+++ b/Engine/Source/Runtime/Engine/Private/ActorEditor.cpp
@@ -67,7 +67,7 @@ void AActor::PostEditChangeProperty(FPropertyChangedEvent& PropertyChangedEvent)
 				}
 				if (USceneComponent* ASC = Cast<USceneComponent>(&A))
 				{
-					if (ASC->AttachParent == &B)
+					if (ASC->GetAttachParent() == &B)
 					{
 						return false;
 					}
@@ -85,7 +85,7 @@ void AActor::PostEditChangeProperty(FPropertyChangedEvent& PropertyChangedEvent)
 				else if (Component->CreationMethod == EComponentCreationMethod::Instance)
 				{
 					USceneComponent* SC = Cast<USceneComponent>(Component);
-					if (SC == nullptr || SC == RootComponent || (SC->AttachParent && SC->AttachParent->IsRegistered()))
+					if (SC == nullptr || SC == RootComponent || (SC->GetAttachParent() && SC->GetAttachParent()->IsRegistered()))
 					{
 						Component->RegisterComponent();
 					}
@@ -239,7 +239,7 @@ void AActor::DebugShowOneComponentHierarchy( USceneComponent* SceneComp, int32&
 	{
 		UE_LOG(LogActor, Warning, TEXT("%sSceneComp [%x] (%s) No Owner"), *Nest, SceneComp, *SceneComp->GetFName().ToString() );
 	}
-	if( SceneComp->AttachParent )
+	if( SceneComp->GetAttachParent())
 	{
 		if( bShowPosition )
 		{
@@ -251,18 +251,17 @@ void AActor::DebugShowOneComponentHierarchy( USceneComponent* SceneComp, int32&
 		{
 			PosString = "";
 		}
-		UE_LOG(LogActor, Warning, TEXT("%sAttachParent [%x] (%s) %s"), *Nest, SceneComp->AttachParent , *SceneComp->AttachParent->GetFName().ToString(), *PosString );
+		UE_LOG(LogActor, Warning, TEXT("%sAttachParent [%x] (%s) %s"), *Nest, SceneComp->GetAttachParent(), *SceneComp->GetAttachParent()->GetFName().ToString(), *PosString );
 	}
 	else
 	{
 		UE_LOG(LogActor, Warning, TEXT("%s[NO PARENT]"), *Nest );
 	}
 
-	if( SceneComp->AttachChildren.Num() != 0 )
+	if( SceneComp->GetAttachChildren().Num() != 0 )
 	{
-		for (int32 iChild = 0; iChild < SceneComp->AttachChildren.Num() ; iChild++)
+		for (USceneComponent* EachSceneComp : SceneComp->GetAttachChildren())
 		{			
-			USceneComponent* EachSceneComp = Cast<USceneComponent>(SceneComp->AttachChildren[iChild]);			
 			DebugShowOneComponentHierarchy(EachSceneComp,NestLevel, bShowPosition );
 		}
 	}
@@ -282,16 +281,16 @@ AActor::FActorTransactionAnnotation::FActorTransactionAnnotation(const AActor* A
 		RootComponentData.Transform = ActorRootComponent->ComponentToWorld;
 		RootComponentData.Transform.SetTranslation(ActorRootComponent->GetComponentLocation()); // take into account any custom location
 
-		if (ActorRootComponent->AttachParent)
+		if (ActorRootComponent->GetAttachParent())
 		{
-			RootComponentData.AttachedParentInfo.Actor = ActorRootComponent->AttachParent->GetOwner();
-			RootComponentData.AttachedParentInfo.AttachParent = ActorRootComponent->AttachParent;
-			RootComponentData.AttachedParentInfo.AttachParentName = ActorRootComponent->AttachParent->GetFName();
-			RootComponentData.AttachedParentInfo.SocketName = ActorRootComponent->AttachSocketName;
+			RootComponentData.AttachedParentInfo.Actor = ActorRootComponent->GetAttachParent()->GetOwner();
+			RootComponentData.AttachedParentInfo.AttachParent = ActorRootComponent->GetAttachParent();
+			RootComponentData.AttachedParentInfo.AttachParentName = ActorRootComponent->GetAttachParent()->GetFName();
+			RootComponentData.AttachedParentInfo.SocketName = ActorRootComponent->GetAttachSocketName();
 			RootComponentData.AttachedParentInfo.RelativeTransform = ActorRootComponent->GetRelativeTransform();
 		}
 
-		for (USceneComponent* AttachChild : ActorRootComponent->AttachChildren)
+		for (USceneComponent* AttachChild : ActorRootComponent->GetAttachChildren())
 		{
 			AActor* ChildOwner = (AttachChild ? AttachChild->GetOwner() : NULL);
 			if (ChildOwner && ChildOwner != Actor)
@@ -299,7 +298,7 @@ AActor::FActorTransactionAnnotation::FActorTransactionAnnotation(const AActor* A
 				// Save info about actor to reattach
 				FActorRootComponentReconstructionData::FAttachedActorInfo Info;
 				Info.Actor = ChildOwner;
-				Info.SocketName = AttachChild->AttachSocketName;
+				Info.SocketName = AttachChild->GetAttachSocketName();
 				Info.RelativeTransform = AttachChild->GetRelativeTransform();
 				RootComponentData.AttachedToInfo.Add(Info);
 			}
diff --git a/Engine/Source/Runtime/Engine/Private/ActorEditorUtils.cpp b/Engine/Source/Runtime/Engine/Private/ActorEditorUtils.cpp
--- a/Engine/Source/Runtime/Engine/Private/ActorEditorUtils.cpp
+++ b/Engine/Source/Runtime/Engine/Private/ActorEditorUtils.cpp
@@ -41,7 +41,7 @@ namespace FActorEditorUtils
 	{
 		TInlineComponentArray<UActorComponent*> InstanceComponents;
 		InActor->GetComponents(InstanceComponents);
-		for (auto Component : InstanceComponents)
+		for (UActorComponent* Component : InstanceComponents)
 		{
 #if WITH_EDITOR
 			if (Component->CreationMethod == EComponentCreationMethod::Native)
@@ -76,7 +76,7 @@ namespace FActorEditorUtils
 
 		if( InActor->GetRootComponent() )
 		{
-			for(auto* ChildComponent : InActor->GetRootComponent()->AttachChildren)
+			for(USceneComponent* ChildComponent : InActor->GetRootComponent()->GetAttachChildren())
 			{
 				AActor* ChildActor = ChildComponent ? ChildComponent->GetOwner() : nullptr;
 				if(ChildActor && ChildActor != InActor && !TraverseActorTree_ParentFirst(ChildActor, InPredicate))
@@ -96,7 +96,7 @@ namespace FActorEditorUtils
 			return true;
 		}
 
-		for (auto* ChildComponent : InActor->GetRootComponent()->AttachChildren)
+		for (USceneComponent* ChildComponent : InActor->GetRootComponent()->GetAttachChildren())
 		{
 			AActor* ChildActor = ChildComponent ? ChildComponent->GetOwner() : nullptr;
 			if (ChildActor && ChildActor != InActor && !TraverseActorTree_ChildFirst(ChildActor, InPredicate))
diff --git a/Engine/Source/Runtime/Engine/Private/ActorReplication.cpp b/Engine/Source/Runtime/Engine/Private/ActorReplication.cpp
--- a/Engine/Source/Runtime/Engine/Private/ActorReplication.cpp
+++ b/Engine/Source/Runtime/Engine/Private/ActorReplication.cpp
@@ -139,7 +139,7 @@ void AActor::OnRep_ReplicatedMovement()
 		else
 		{
 			// Attachment trumps global position updates, see GatherCurrentMovement().
-			if (!RootComponent->AttachParent)
+			if (!RootComponent->GetAttachParent())
 			{
 				if (Role == ROLE_SimulatedProxy)
 				{
@@ -218,9 +218,9 @@ bool AActor::IsNetRelevantFor(const AActor* RealViewer, const AActor* ViewTarget
 	{
 		return false;
 	}
-	else if ( RootComponent && RootComponent->AttachParent && RootComponent->AttachParent->GetOwner() && (Cast<USkeletalMeshComponent>(RootComponent->AttachParent) || (RootComponent->AttachParent->GetOwner() == Owner)) )
+	else if ( RootComponent && RootComponent->GetAttachParent() && RootComponent->GetAttachParent()->GetOwner() && (Cast<USkeletalMeshComponent>(RootComponent->GetAttachParent()) || (RootComponent->GetAttachParent()->GetOwner() == Owner)) )
 	{
-		return RootComponent->AttachParent->GetOwner()->IsNetRelevantFor(RealViewer, ViewTarget, SrcLocation);
+		return RootComponent->GetAttachParent()->GetOwner()->IsNetRelevantFor(RealViewer, ViewTarget, SrcLocation);
 	}
 	else if( bHidden && (!RootComponent || !RootComponent->IsCollisionEnabled()) )
 	{
@@ -256,18 +256,18 @@ void AActor::GatherCurrentMovement()
 	else if (RootComponent != nullptr)
 	{
 		// If we are attached, don't replicate absolute position, use AttachmentReplication instead.
-		if (RootComponent->AttachParent != nullptr)
+		if (RootComponent->GetAttachParent() != nullptr)
 		{
 			// Networking for attachments assumes the RootComponent of the AttachParent actor. 
 			// If that's not the case, we can't update this, as the client wouldn't be able to resolve the Component and would detach as a result.
-			AttachmentReplication.AttachParent = RootComponent->AttachParent->GetAttachmentRootActor();
+			AttachmentReplication.AttachParent = RootComponent->GetAttachParent()->GetAttachmentRootActor();
 			if (AttachmentReplication.AttachParent != nullptr)
 			{
 				AttachmentReplication.LocationOffset = RootComponent->RelativeLocation;
 				AttachmentReplication.RotationOffset = RootComponent->RelativeRotation;
 				AttachmentReplication.RelativeScale3D = RootComponent->RelativeScale3D;
-				AttachmentReplication.AttachComponent = RootComponent->AttachParent;
-				AttachmentReplication.AttachSocket = RootComponent->AttachSocketName;
+				AttachmentReplication.AttachComponent = RootComponent->GetAttachParent();
+				AttachmentReplication.AttachSocket = RootComponent->GetAttachSocketName();
 			}
 		}
 		else
diff --git a/Engine/Source/Runtime/Engine/Private/Animation/AnimInstance.cpp b/Engine/Source/Runtime/Engine/Private/Animation/AnimInstance.cpp
--- a/Engine/Source/Runtime/Engine/Private/Animation/AnimInstance.cpp
+++ b/Engine/Source/Runtime/Engine/Private/Animation/AnimInstance.cpp
@@ -2592,6 +2592,11 @@ void UAnimInstance::UpdateComponentsMaterialParameters(UPrimitiveComponent* Comp
 	GetProxyOnGameThread<FAnimInstanceProxy>().UpdateComponentsMaterialParameters(Component);
 }
 
+void UAnimInstance::RecordStateWeight(const int32& InMachineClassIndex, const int32& InStateIndex, const float& InStateWeight)
+{
+	GetProxyOnAnyThread<FAnimInstanceProxy>().RecordStateWeight(InMachineClassIndex, InStateIndex, InStateWeight);
+}
+
 TArray<FActiveVertexAnim> UAnimInstance::UpdateActiveVertexAnims(const USkeletalMesh* SkeletalMesh)
 {
 	FAnimInstanceProxy& Proxy = GetProxyOnAnyThread<FAnimInstanceProxy>();	// this can be called from CreateRenderState_Concurrent
diff --git a/Engine/Source/Runtime/Engine/Private/Animation/AnimInstanceProxy.cpp b/Engine/Source/Runtime/Engine/Private/Animation/AnimInstanceProxy.cpp
--- a/Engine/Source/Runtime/Engine/Private/Animation/AnimInstanceProxy.cpp
+++ b/Engine/Source/Runtime/Engine/Private/Animation/AnimInstanceProxy.cpp
@@ -69,6 +69,24 @@ void FAnimInstanceProxy::Initialize(UAnimInstance* InAnimInstance)
 			Skeleton = AnimClassInterface->GetTargetSkeleton();
 		}
 
+		// Initialize state buffers
+		int32 NumStates = 0;
+		if(IAnimClassInterface* Interface = GetAnimClassInterface())
+		{
+			const TArray<FBakedAnimationStateMachine>& BakedMachines = Interface->GetBakedStateMachines();
+			const int32 NumMachines = BakedMachines.Num();
+			for(int32 MachineClassIndex = 0; MachineClassIndex < NumMachines; ++MachineClassIndex)
+			{
+				const FBakedAnimationStateMachine& Machine = BakedMachines[MachineClassIndex];
+				StateMachineClassIndexToWeightOffset.Add(MachineClassIndex, NumStates);
+				NumStates += Machine.States.Num();
+			}
+			StateWeightArrays[0].Reset(NumStates);
+			StateWeightArrays[0].AddZeroed(NumStates);
+			StateWeightArrays[1].Reset(NumStates);
+			StateWeightArrays[1].AddZeroed(NumStates);
+		}
+
 #if WITH_EDITORONLY_DATA
 		if (UAnimBlueprint* Blueprint = Cast<UAnimBlueprint>(InAnimInstance->GetClass()->ClassGeneratedBy))
 		{
@@ -153,6 +171,9 @@ void FAnimInstanceProxy::PreUpdate(UAnimInstance* InAnimInstance, float DeltaSec
 		SyncGroups[GroupIndex].Reset();
 	}
 
+	TArray<float>& StateWeights = StateWeightArrays[GetSyncGroupWriteIndex()];
+	FMemory::Memset(StateWeights.GetData(), 0, StateWeights.Num() * sizeof(float));
+
 #if WITH_EDITORONLY_DATA
 	bIsBeingDebugged = false;
 	if (UAnimBlueprint* AnimBlueprint = GetAnimBlueprint())
@@ -1007,7 +1028,7 @@ float FAnimInstanceProxy::GetInstanceStateWeight(int32 MachineIndex, int32 State
 {
 	if(FAnimNode_StateMachine* MachineInstance = GetStateMachineInstance(MachineIndex))
 	{
-		return MachineInstance->GetStateWeight(StateIndex);
+		return GetRecordedStateWeight(MachineInstance->StateMachineIndexInClass, StateIndex);
 	}
 	
 	return 0.0f;
@@ -1450,4 +1471,28 @@ int32 FAnimInstanceProxy::GetInstanceAssetPlayerIndex(FName MachineName, FName S
 	return INDEX_NONE;
 }
 
-#undef LOCTEXT_NAMESPACE
\ No newline at end of file
+float FAnimInstanceProxy::GetRecordedStateWeight(const int32& InMachineClassIndex, const int32& InStateIndex)
+{
+	const int32* BaseIndexPtr = StateMachineClassIndexToWeightOffset.Find(InMachineClassIndex);
+
+	if(BaseIndexPtr)
+	{
+		const int32 StateIndex = *BaseIndexPtr + InStateIndex;
+		return StateWeightArrays[GetSyncGroupReadIndex()][StateIndex];
+	}
+
+	return 0.0f;
+}
+
+void FAnimInstanceProxy::RecordStateWeight(const int32& InMachineClassIndex, const int32& InStateIndex, const float& InStateWeight)
+{
+	const int32* BaseIndexPtr = StateMachineClassIndexToWeightOffset.Find(InMachineClassIndex);
+
+	if(BaseIndexPtr)
+	{
+		const int32 StateIndex = *BaseIndexPtr + InStateIndex;
+		StateWeightArrays[GetSyncGroupWriteIndex()][StateIndex] = InStateWeight;
+	}
+}
+
+#undef LOCTEXT_NAMESPACE
diff --git a/Engine/Source/Runtime/Engine/Private/Animation/AnimNode_StateMachine.cpp b/Engine/Source/Runtime/Engine/Private/Animation/AnimNode_StateMachine.cpp
--- a/Engine/Source/Runtime/Engine/Private/Animation/AnimNode_StateMachine.cpp
+++ b/Engine/Source/Runtime/Engine/Private/Animation/AnimNode_StateMachine.cpp
@@ -508,6 +508,8 @@ void FAnimNode_StateMachine::Update(const FAnimationUpdateContext& Context)
 		}
 
 		StatePoseLinks[CurrentState].Update(Context);
+
+		Context.AnimInstanceProxy->RecordStateWeight(StateMachineIndexInClass, CurrentState, GetStateWeight(CurrentState));
 	}
 
 	ElapsedTime += Context.GetDeltaTime();
@@ -978,6 +980,8 @@ void FAnimNode_StateMachine::UpdateState(int32 StateIndex, const FAnimationUpdat
 	{
 		StatesUpdated.Add(StateIndex);
 		StatePoseLinks[StateIndex].Update(Context);
+
+		Context.AnimInstanceProxy->RecordStateWeight(StateMachineIndexInClass, StateIndex, GetStateWeight(StateIndex));
 	}
 }
 
diff --git a/Engine/Source/Runtime/Engine/Private/Animation/AnimNotifyState_TimedParticleEffect.cpp b/Engine/Source/Runtime/Engine/Private/Animation/AnimNotifyState_TimedParticleEffect.cpp
--- a/Engine/Source/Runtime/Engine/Private/Animation/AnimNotifyState_TimedParticleEffect.cpp
+++ b/Engine/Source/Runtime/Engine/Private/Animation/AnimNotifyState_TimedParticleEffect.cpp
@@ -38,13 +38,13 @@ void UAnimNotifyState_TimedParticleEffect::NotifyEnd(USkeletalMeshComponent * Me
 	{
 		if(UParticleSystemComponent* ParticleComponent = Cast<UParticleSystemComponent>(Component))
 		{
-			bool bSocketMatch = ParticleComponent->AttachSocketName == SocketName;
+			bool bSocketMatch = ParticleComponent->GetAttachSocketName() == SocketName;
 			bool bTemplateMatch = ParticleComponent->Template == PSTemplate;
 
 #if WITH_EDITORONLY_DATA
 			// In editor someone might have changed our parameters while we're ticking; so check 
 			// previous known parameters too.
-			bSocketMatch |= PreviousSocketNames.Contains(ParticleComponent->AttachSocketName);
+			bSocketMatch |= PreviousSocketNames.Contains(ParticleComponent->GetAttachSocketName());
 			bTemplateMatch |= PreviousPSTemplates.Contains(ParticleComponent->Template);
 #endif
 
diff --git a/Engine/Source/Runtime/Engine/Private/Audio.cpp b/Engine/Source/Runtime/Engine/Private/Audio.cpp
--- a/Engine/Source/Runtime/Engine/Private/Audio.cpp
+++ b/Engine/Source/Runtime/Engine/Private/Audio.cpp
@@ -582,12 +582,23 @@ float FWaveInstance::GetActualVolume() const
 	return Volume * VolumeMultiplier;
 }
 
+bool FWaveInstance::ShouldStopDueToMaxConcurrency() const
+{
+	return ActiveSound->bShouldStopDueToMaxConcurrency;
+}
+
 float FWaveInstance::GetVolumeWeightedPriority() const
 {
-	check(ActiveSound);
-	// If this wave instance's active sound should stop due to max concurrency, return a large negative weighted priority so 
-	// that it will always be sorted to the bottom of the WaveInstance list and stopped.
-	return ActiveSound->bShouldStopDueToMaxConcurrency ? -100.0f : GetActualVolume() * Priority;
+	// This will result in zero-volume sounds still able to be sorted due to priority but give non-zero volumes higher priority than 0 volumes
+	float ActualVolume = GetActualVolume();
+	if (ActualVolume > 0.0f)
+	{
+		return GetActualVolume() * Priority;
+	}
+	else
+	{
+		return Priority - MAX_SOUND_PRIORITY - 1;
+	}
 }
 
 bool FWaveInstance::IsStreaming() const
diff --git a/Engine/Source/Runtime/Engine/Private/AudioDevice.cpp b/Engine/Source/Runtime/Engine/Private/AudioDevice.cpp
--- a/Engine/Source/Runtime/Engine/Private/AudioDevice.cpp
+++ b/Engine/Source/Runtime/Engine/Private/AudioDevice.cpp
@@ -74,9 +74,6 @@ void FDynamicParameter::Update(float DeltaTime)
 /**
 * Number of ticks an inaudible source remains alive before being stopped
 */
-#define AUDIOSOURCE_TICK_LONGEVITY		600
-
-
 
 FAudioDevice::FAudioDevice()
 	: CommonAudioPool(nullptr)
@@ -2084,26 +2081,23 @@ int32 FAudioDevice::GetSortedActiveWaveInstances(TArray<FWaveInstance*>& WaveIns
 		ConcurrencyManager.StopQuietSoundsDueToMaxConcurrency();
 	}
 
-	// Helper function for "Sort" (higher priority sorts last).
-	struct FCompareFWaveInstanceByPlayPriority
+	int FirstActiveIndex = 0;
+	if (WaveInstances.Num() >= MaxChannels)
 	{
-		FORCEINLINE bool operator()( const FWaveInstance& A, const FWaveInstance& B) const
+		// Helper function for "Sort" (higher priority sorts last).
+		struct FCompareFWaveInstanceByPlayPriority
 		{
-			return A.GetVolumeWeightedPriority() < B.GetVolumeWeightedPriority();
-		}
-	};
+			FORCEINLINE bool operator()(const FWaveInstance& A, const FWaveInstance& B) const
+			{
+				return A.GetVolumeWeightedPriority() < B.GetVolumeWeightedPriority();
+			}
+		};
 
-	// Sort by priority (lowest priority first).
-	WaveInstances.Sort( FCompareFWaveInstanceByPlayPriority() );
+		// Sort by priority (lowest priority first).
+		WaveInstances.Sort(FCompareFWaveInstanceByPlayPriority());
 
-	// Return the first audible waveinstance
-	int32 FirstActiveIndex = FMath::Max( WaveInstances.Num() - MaxChannels, 0 );
-	for( ; FirstActiveIndex < WaveInstances.Num(); FirstActiveIndex++ )
-	{
-		if (WaveInstances[FirstActiveIndex]->GetVolumeWeightedPriority() > KINDA_SMALL_NUMBER)
-		{
-			break;
-		}
+		// Get the first index that will result in a active source voice
+		FirstActiveIndex = FMath::Max(WaveInstances.Num() - MaxChannels, 0);
 	}
 	return( FirstActiveIndex );
 }
@@ -2112,7 +2106,7 @@ int32 FAudioDevice::GetSortedActiveWaveInstances(TArray<FWaveInstance*>& WaveIns
 void FAudioDevice::StopSources( TArray<FWaveInstance*>& WaveInstances, int32 FirstActiveIndex )
 {
 	// Touch sources that are high enough priority to play
-	for( int32 InstanceIndex = FirstActiveIndex; InstanceIndex < WaveInstances.Num(); InstanceIndex++ )
+	for (int32 InstanceIndex = FirstActiveIndex; InstanceIndex < WaveInstances.Num(); InstanceIndex++)
 	{
 		FWaveInstance* WaveInstance = WaveInstances[ InstanceIndex ];
 		FSoundSource* Source = WaveInstanceSourceMap.FindRef( WaveInstance );
@@ -2121,8 +2115,8 @@ void FAudioDevice::StopSources( TArray<FWaveInstance*>& WaveInstances, int32 Fir
 			Source->LastUpdate = CurrentTick;
 
 			// If they are still audible, mark them as such
-			float VolumeWeightedPriority = WaveInstance->GetVolumeWeightedPriority();
-			if (VolumeWeightedPriority > KINDA_SMALL_NUMBER)
+			float VolumeWeightedPriority = WaveInstance->GetActualVolume();
+			if (VolumeWeightedPriority > 0.0f)
 			{
 				Source->LastHeardUpdate = CurrentTick;
 			}
@@ -2131,11 +2125,11 @@ void FAudioDevice::StopSources( TArray<FWaveInstance*>& WaveInstances, int32 Fir
 
 	// Stop inactive sources, sources that no longer have a WaveInstance associated
 	// or sources that need to be reset because Stop & Play were called in the same frame.
-	for( int32 SourceIndex = 0; SourceIndex < Sources.Num(); SourceIndex++ )
+	for (int32 SourceIndex = 0; SourceIndex < Sources.Num(); SourceIndex++)
 	{
 		FSoundSource* Source = Sources[ SourceIndex ];
 
-		if( Source->WaveInstance )
+		if (Source->WaveInstance)
 		{
 #if STATS && WITH_EDITORONLY_DATA
 			if( Source->UsesCPUDecompression() )
@@ -2143,20 +2137,20 @@ void FAudioDevice::StopSources( TArray<FWaveInstance*>& WaveInstances, int32 Fir
 				INC_DWORD_STAT( STAT_OggWaveInstances );
 			}
 #endif
-			// Source was not one of the active sounds this tick so needs to be stopped
-			if (Source->LastUpdate != CurrentTick)
+			// If we need to stop this sound due to max concurrency (i.e. it was quietest in a concurrency group)
+			if (Source->WaveInstance->ShouldStopDueToMaxConcurrency())
 			{
 				Source->Stop();
 			}
-			// Source has been inaudible for several ticks
-			else if (Source->LastHeardUpdate + AUDIOSOURCE_TICK_LONGEVITY < CurrentTick)
+			// Source was not one of the active sounds this tick so needs to be stopped
+			else if (Source->LastUpdate != CurrentTick)
 			{
 				Source->Stop();
 			}
-			// Need to update the source still so that it gets any volume settings applied to
-			// otherwise the source may play at a very quiet volume and not actually set to 0.0
 			else
 			{
+				// Need to update the source still so that it gets any volume settings applied to
+				// otherwise the source may play at a very quiet volume and not actually set to 0.0
 				Source->Update();
 			}
 		}
@@ -2169,7 +2163,6 @@ void FAudioDevice::StopSources( TArray<FWaveInstance*>& WaveInstances, int32 Fir
 	{
 		FWaveInstance* WaveInstance = WaveInstances[ InstanceIndex ];
 		WaveInstance->StopWithoutNotification();
-		// UE_LOG(LogAudio, Log, TEXT( "SoundStoppedWithoutNotification due to priority reasons: %s" ), *WaveInstance->WaveData->GetName() );
 	}
 
 #if STATS
@@ -2777,7 +2770,7 @@ void FAudioDevice::PlaySoundAtLocation(USoundBase* Sound, UWorld* World, float V
 		NewActiveSound.Transform.SetRotation(FQuat(Rotation));
 		NewActiveSound.bIsUISound = !bIsInGameWorld;
 		NewActiveSound.bHandleSubtitles = true;
-		NewActiveSound.SubtitlePriority = 10000.f; // Todo: Fix this. Add it to USoundBase
+		NewActiveSound.SubtitlePriority = Sound->GetSubtitlePriority();
 
 		NewActiveSound.bHasAttenuationSettings = (bIsInGameWorld && AttenuationSettingsToApply);
 		if (NewActiveSound.bHasAttenuationSettings)
diff --git a/Engine/Source/Runtime/Engine/Private/Camera/CameraAnimInst.cpp b/Engine/Source/Runtime/Engine/Private/Camera/CameraAnimInst.cpp
--- a/Engine/Source/Runtime/Engine/Private/Camera/CameraAnimInst.cpp
+++ b/Engine/Source/Runtime/Engine/Private/Camera/CameraAnimInst.cpp
@@ -265,9 +265,9 @@ void UCameraAnimInst::Stop(bool bImmediate)
 		MoveInst = nullptr;
 		bFinished = true;
  	}
-	else
+	else if (bBlendingOut == false)
 	{
-		// start blend out
+		// start blend out if not already blending out
 		bBlendingOut = true;
 		CurBlendOutTime = 0.f;
 	}
diff --git a/Engine/Source/Runtime/Engine/Private/Collision/CollisionConversions.cpp b/Engine/Source/Runtime/Engine/Private/Collision/CollisionConversions.cpp
--- a/Engine/Source/Runtime/Engine/Private/Collision/CollisionConversions.cpp
+++ b/Engine/Source/Runtime/Engine/Private/Collision/CollisionConversions.cpp
@@ -431,11 +431,13 @@ EConvertQueryResult ConvertQueryImpactHit(const UWorld* World, const PxLocationH
 	SCOPE_CYCLE_COUNTER(STAT_ConvertQueryImpactHit);
 
 #if WITH_EDITOR
+	bool bIgnoreFaceIndexForPIE = false;
 	if(bReturnFaceIndex && World->IsGameWorld())
 	{
 		if(!ensure(UPhysicsSettings::Get()->bSuppressFaceRemapTable == false))
 		{
 			UE_LOG(LogPhysics, Error, TEXT("A scene query is relying on face indices, but bSuppressFaceRemapTable is false."));
+			bIgnoreFaceIndexForPIE = true;
 		}
 	}
 #endif
@@ -530,6 +532,12 @@ EConvertQueryResult ConvertQueryImpactHit(const UWorld* World, const PxLocationH
 			PHit.faceIndex < PTriMeshGeom.triangleMesh->getNbTriangles() )
 		{
 			OutResult.FaceIndex	= PTriMeshGeom.triangleMesh->getTrianglesRemap()[PHit.faceIndex];
+#if WITH_EDITOR
+			if(bIgnoreFaceIndexForPIE)
+			{
+				OutResult.FaceIndex = INDEX_NONE;	//if we need to ignore face index in PIE it's because at runtime we don't have the info. We want to make sure PIE is consistent with cooked version
+			}
+#endif
 		}
 	}
 
