Copying //UE4/Dev-Framework to Dev-Main (//UE4/Dev-Main) @ 2855699

diff --git a/Engine/Source/Runtime/Engine/Private/ComponentInstanceDataCache.cpp b/Engine/Source/Runtime/Engine/Private/ComponentInstanceDataCache.cpp
--- a/Engine/Source/Runtime/Engine/Private/ComponentInstanceDataCache.cpp
+++ b/Engine/Source/Runtime/Engine/Private/ComponentInstanceDataCache.cpp
@@ -179,7 +179,7 @@ FComponentInstanceDataCache::FComponentInstanceDataCache(const AActor* Actor)
 				// If the instance component is attached to a BP component we have to be prepared for the possibility that it will be deleted
 				if (USceneComponent* SceneComponent = Cast<USceneComponent>(Component))
 				{
-					if (SceneComponent->AttachParent && SceneComponent->AttachParent->IsCreatedByConstructionScript())
+					if (SceneComponent->GetAttachParent() && SceneComponent->GetAttachParent()->IsCreatedByConstructionScript())
 					{
 						InstanceComponentTransformToRootMap.Add(SceneComponent, SceneComponent->GetComponentTransform().GetRelativeTransform(Actor->GetRootComponent()->GetComponentTransform()));
 					}
@@ -244,7 +244,7 @@ void FComponentInstanceDataCache::ApplyToActor(AActor* Actor, const ECacheApplyP
 			check(Actor->GetRootComponent());
 
 			USceneComponent* SceneComponent = InstanceTransformPair.Key;
-			if (SceneComponent && (SceneComponent->AttachParent == nullptr || SceneComponent->AttachParent->IsPendingKill()))
+			if (SceneComponent && (SceneComponent->GetAttachParent() == nullptr || SceneComponent->GetAttachParent()->IsPendingKill()))
 			{
 				SceneComponent->AttachTo(Actor->GetRootComponent());
 				SceneComponent->SetRelativeTransform(InstanceTransformPair.Value);
diff --git a/Engine/Source/Runtime/Engine/Private/ComponentUtils.cpp b/Engine/Source/Runtime/Engine/Private/ComponentUtils.cpp
--- a/Engine/Source/Runtime/Engine/Private/ComponentUtils.cpp
+++ b/Engine/Source/Runtime/Engine/Private/ComponentUtils.cpp
@@ -60,7 +60,7 @@ namespace ComponentUtils
 	 */
 	USceneComponent* GetAttachedParent(USceneComponent const* SceneComponentObject)
 	{
-		USceneComponent* SceneParent = SceneComponentObject->AttachParent;
+		USceneComponent* SceneParent = SceneComponentObject->GetAttachParent();
 		if(SceneParent == nullptr)
 		{
 			USCS_Node* const SCSNode = FindCorrespondingSCSNode(SceneComponentObject);
diff --git a/Engine/Source/Runtime/Engine/Private/Components/ActorComponent.cpp b/Engine/Source/Runtime/Engine/Private/Components/ActorComponent.cpp
--- a/Engine/Source/Runtime/Engine/Private/Components/ActorComponent.cpp
+++ b/Engine/Source/Runtime/Engine/Private/Components/ActorComponent.cpp
@@ -1463,12 +1463,21 @@ bool UActorComponent::IsSupportedForNetworking() const
 
 void UActorComponent::SetIsReplicated(bool ShouldReplicate)
 {
-	check(GetComponentClassCanReplicate()); // Only certain component classes can replicate!
-	bReplicates = ShouldReplicate;
-
-	if (AActor* MyOwner = GetOwner())
+	if (bReplicates != ShouldReplicate)
 	{
-		MyOwner->UpdateReplicatedComponent( this );
+		if (GetComponentClassCanReplicate())
+		{
+			bReplicates = ShouldReplicate;
+
+			if (AActor* MyOwner = GetOwner())
+			{
+				MyOwner->UpdateReplicatedComponent( this );
+			}
+		}
+		else
+		{
+			UE_LOG(LogActorComponent, Error, TEXT("Calling SetIsReplicated on component of Class '%s' which cannot replicate."), *GetClass()->GetName());
+		}
 	}
 }
 
diff --git a/Engine/Source/Runtime/Engine/Private/Components/CharacterMovementComponent.cpp b/Engine/Source/Runtime/Engine/Private/Components/CharacterMovementComponent.cpp
--- a/Engine/Source/Runtime/Engine/Private/Components/CharacterMovementComponent.cpp
+++ b/Engine/Source/Runtime/Engine/Private/Components/CharacterMovementComponent.cpp
@@ -23,6 +23,7 @@
 #include "Components/DestructibleComponent.h"
 
 #include "Engine/DemoNetDriver.h"
+#include "Engine/NetworkObjectList.h"
 
 #include "PerfCountersHelpers.h"
 
@@ -250,6 +251,7 @@ UCharacterMovementComponent::UCharacterMovementComponent(const FObjectInitialize
 	bJustTeleported = true;
 	CrouchedHalfHeight = 40.0f;
 	Buoyancy = 1.0f;
+	LastUpdateRotation = FQuat::Identity;
 	LastUpdateVelocity = FVector::ZeroVector;
 	PendingImpulseToApply = FVector::ZeroVector;
 	PendingLaunchVelocity = FVector::ZeroVector;
@@ -269,6 +271,7 @@ UCharacterMovementComponent::UCharacterMovementComponent(const FObjectInitialize
 	InitialPushForceFactor = 500.0f;
 	PushForceFactor = 750000.0f;
 	PushForcePointZOffsetFactor = -0.75f;
+	bPushForceUsingZOffset = false;
 	bPushForceScaledToMass = false;
 	bScalePushForceToVelocity = true;
 	
@@ -1441,6 +1444,7 @@ void UCharacterMovementComponent::SimulateMovement(float DeltaSeconds)
 	bJustTeleported = false;
 
 	LastUpdateLocation = UpdatedComponent ? UpdatedComponent->GetComponentLocation() : FVector::ZeroVector;
+	LastUpdateRotation = UpdatedComponent ? UpdatedComponent->GetComponentQuat() : FQuat::Identity;
 	LastUpdateVelocity = Velocity;
 }
 
@@ -1875,11 +1879,44 @@ void UCharacterMovementComponent::PerformMovement(float DeltaSeconds)
 	SaveBaseLocation();
 	UpdateComponentVelocity();
 
+	// If we move we want to avoid a long delay before replication catches up to notice this change, especially if it's throttling our rate.
+	if (UNetDriver::IsAdaptiveNetUpdateFrequencyEnabled() && UpdatedComponent && CharacterOwner && CharacterOwner->HasAuthority())
+	{
+		const UWorld* MyWorld = GetWorld();
+		if (MyWorld && MyWorld->GetTimeSeconds() <= CharacterOwner->NetUpdateTime)
+		{
+			UNetDriver* NetDriver = MyWorld->GetNetDriver();
+			if (NetDriver && NetDriver->IsServer())
+			{
+				FNetworkObjectInfo* NetActor = NetDriver->GetNetworkActor(CharacterOwner);
+				if (NetActor && NetDriver->IsNetworkActorUpdateFrequencyThrottled(*NetActor))
+				{
+					if (ShouldCancelAdaptiveReplication())
+					{
+						NetDriver->CancelAdaptiveReplication(*NetActor);
+					}
+				}
+			}
+		}
+	}
+
 	LastUpdateLocation = UpdatedComponent ? UpdatedComponent->GetComponentLocation() : FVector::ZeroVector;
+	LastUpdateRotation = UpdatedComponent ? UpdatedComponent->GetComponentQuat() : FQuat::Identity;
 	LastUpdateVelocity = Velocity;
 }
 
 
+bool UCharacterMovementComponent::ShouldCancelAdaptiveReplication() const
+{
+	// Update sooner if important properties changed.
+	const bool bVelocityChanged = (Velocity != LastUpdateVelocity);
+	const bool bLocationChanged = (UpdatedComponent->GetComponentLocation() != LastUpdateLocation);
+	const bool bRotationChanged = (UpdatedComponent->GetComponentQuat() != LastUpdateRotation);
+
+	return (bVelocityChanged || bLocationChanged || bRotationChanged);
+}
+
+
 void UCharacterMovementComponent::CallMovementUpdateDelegate(float DeltaTime, const FVector& OldLocation, const FVector& OldVelocity)
 {
 	SCOPE_CYCLE_COUNTER(STAT_CharMoveUpdateDelegate);
@@ -5798,14 +5835,17 @@ void UCharacterMovementComponent::ApplyImpactPhysicsForces(const FHitResult& Imp
 			{
 				BodyMass = FMath::Max(BI->GetBodyMass(), 1.0f);
 
-				FBox Bounds = BI->GetBodyBounds();
+				if(bPushForceUsingZOffset)
+				{
+					FBox Bounds = BI->GetBodyBounds();
 
-				FVector Center, Extents;
-				Bounds.GetCenterAndExtents(Center, Extents);
+					FVector Center, Extents;
+					Bounds.GetCenterAndExtents(Center, Extents);
 
-				if (!Extents.IsNearlyZero())
-				{
-					ForcePoint.Z = Center.Z + Extents.Z * PushForcePointZOffsetFactor;
+					if (!Extents.IsNearlyZero())
+					{
+						ForcePoint.Z = Center.Z + Extents.Z * PushForcePointZOffsetFactor;
+					}
 				}
 			}
 
diff --git a/Engine/Source/Runtime/Engine/Private/Components/MaterialBillboardComponent.cpp b/Engine/Source/Runtime/Engine/Private/Components/MaterialBillboardComponent.cpp
--- a/Engine/Source/Runtime/Engine/Private/Components/MaterialBillboardComponent.cpp
+++ b/Engine/Source/Runtime/Engine/Private/Components/MaterialBillboardComponent.cpp
@@ -303,6 +303,16 @@ void UMaterialBillboardComponent::AddElement(
 	MarkRenderStateDirty();
 }
 
+void UMaterialBillboardComponent::SetElements(const TArray<FMaterialSpriteElement>& NewElements)
+{
+	// Replace existing array
+	Elements = NewElements;
+
+	// Indicate scene proxy needs to be updated
+	MarkRenderStateDirty();
+}
+
+
 UMaterialInterface* UMaterialBillboardComponent::GetMaterial(int32 Index) const
 {
 	UMaterialInterface* ResultMI = nullptr;
diff --git a/Engine/Source/Runtime/Engine/Private/Components/MovementComponent.cpp b/Engine/Source/Runtime/Engine/Private/Components/MovementComponent.cpp
--- a/Engine/Source/Runtime/Engine/Private/Components/MovementComponent.cpp
+++ b/Engine/Source/Runtime/Engine/Private/Components/MovementComponent.cpp
@@ -120,7 +120,11 @@ void UMovementComponent::OnRegister()
 		SetPlaneConstraintAxisSetting(PlaneConstraintAxisSetting);
 	}
 
-	PlaneConstraintNormal = PlaneConstraintNormal.GetSafeNormal();
+	const UWorld* MyWorld = GetWorld();
+	if (MyWorld && MyWorld->IsGameWorld())
+	{
+		PlaneConstraintNormal = PlaneConstraintNormal.GetSafeNormal();
+	}
 
 	if (bSnapToPlaneAtStart)
 	{
diff --git a/Engine/Source/Runtime/Engine/Private/Components/PostProcessComponent.cpp b/Engine/Source/Runtime/Engine/Private/Components/PostProcessComponent.cpp
--- a/Engine/Source/Runtime/Engine/Private/Components/PostProcessComponent.cpp
+++ b/Engine/Source/Runtime/Engine/Private/Components/PostProcessComponent.cpp
@@ -15,7 +15,7 @@ UPostProcessComponent::UPostProcessComponent(const FObjectInitializer& ObjectIni
 
 bool UPostProcessComponent::EncompassesPoint(FVector Point, float SphereRadius/*=0.f*/, float* OutDistanceToPoint)
 {
-	UShapeComponent* ParentShape = Cast<UShapeComponent>(AttachParent);
+	UShapeComponent* ParentShape = Cast<UShapeComponent>(GetAttachParent());
 	if (ParentShape != nullptr)
 	{
 #if WITH_PHYSX
diff --git a/Engine/Source/Runtime/Engine/Private/Components/SceneComponent.cpp b/Engine/Source/Runtime/Engine/Private/Components/SceneComponent.cpp
--- a/Engine/Source/Runtime/Engine/Private/Components/SceneComponent.cpp
+++ b/Engine/Source/Runtime/Engine/Private/Components/SceneComponent.cpp
@@ -45,7 +45,7 @@ FOverlapInfo::FOverlapInfo(UPrimitiveComponent* InComponent, int32 InBodyIndex)
 	OverlapInfo.Item = InBodyIndex;
 }
 
-const FName& USceneComponent::GetDefaultSceneRootVariableName()
+FName USceneComponent::GetDefaultSceneRootVariableName()
 {
 	return SceneComponentStatics::DefaultSceneRootVariableName;
 }
@@ -177,15 +177,26 @@ static int32 SetAncestorMobility(USceneComponent const* SceneComponentObject, EC
 		if (ShouldOverrideMobility.Execute(AttachedParent->Mobility))
 		{
 			// USceneComponents shouldn't be set Stationary 
-			if ((NewMobilityType == EComponentMobility::Stationary) && AttachedParent->IsA(UStaticMeshComponent::StaticClass()))
-			{
-				// make it Static (because it is acceptable for Stationary children to have Static parents)
-				AttachedParent->Mobility = EComponentMobility::Static;
-			}
-			else
+			switch(NewMobilityType)
 			{
+			case EComponentMobility::Stationary:
+				if (UStaticMeshComponent* StaticMeshParent = Cast<UStaticMeshComponent>(AttachedParent))
+				{
+					// make it Static (because it is acceptable for Stationary children to have Static parents)
+					StaticMeshParent->Mobility = EComponentMobility::Static;
+					StaticMeshParent->SetSimulatePhysics(false);
+				}
+				break;
+			case EComponentMobility::Static:
+				if (UPrimitiveComponent* PrimitiveComponentParent = Cast<UPrimitiveComponent>(AttachedParent))
+				{
+					PrimitiveComponentParent->SetSimulatePhysics(false);
+				}
+				//FALLTHROUGH: we still want to set mobility
+			default:
 				AttachedParent->Mobility = NewMobilityType;
 			}
+
 			++MobilityAlteredCount;
 		}
 		SceneComponentObject = AttachedParent;
@@ -1338,9 +1349,9 @@ void USceneComponent::AppendDescendants(TArray<USceneComponent*>& Children) cons
 }
 
 //This function is used for giving AttachTo different bWeldSimulatedBodies default, but only when called from BP
-void USceneComponent::K2_AttachTo(class USceneComponent* InParent, FName InSocketName, EAttachLocation::Type AttachLocationType, bool bWeldSimulatedBodies /*= true*/)
+bool USceneComponent::K2_AttachTo(class USceneComponent* InParent, FName InSocketName, EAttachLocation::Type AttachLocationType, bool bWeldSimulatedBodies /*= true*/)
 {
-	AttachTo(InParent, InSocketName, AttachLocationType, bWeldSimulatedBodies);
+	return AttachTo(InParent, InSocketName, AttachLocationType, bWeldSimulatedBodies);
 }
 
 bool USceneComponent::AttachTo(class USceneComponent* Parent, FName InSocketName, EAttachLocation::Type AttachType /*= EAttachLocation::KeepRelativeOffset */, bool bWeldSimulatedBodies /*= false*/)
@@ -1610,9 +1621,9 @@ bool USceneComponent::AttachTo(class USceneComponent* Parent, FName InSocketName
 	return false;
 }
 
-void USceneComponent::SnapTo(class USceneComponent* Parent, FName InSocketName)
+bool USceneComponent::SnapTo(class USceneComponent* Parent, FName InSocketName)
 {
-	AttachTo(Parent, InSocketName, EAttachLocation::SnapToTarget);
+	return AttachTo(Parent, InSocketName, EAttachLocation::SnapToTarget);
 }
 
 void USceneComponent::DetachFromParent(bool bMaintainWorldPosition, bool bCallModify)
diff --git a/Engine/Source/Runtime/Engine/Private/Components/SkeletalMeshComponent.cpp b/Engine/Source/Runtime/Engine/Private/Components/SkeletalMeshComponent.cpp
--- a/Engine/Source/Runtime/Engine/Private/Components/SkeletalMeshComponent.cpp
+++ b/Engine/Source/Runtime/Engine/Private/Components/SkeletalMeshComponent.cpp
@@ -1413,20 +1413,23 @@ void USkeletalMeshComponent::SetSkeletalMesh(USkeletalMesh* InSkelMesh, bool bRe
 
 	UPhysicsAsset* OldPhysAsset = GetPhysicsAsset();
 
-	Super::SetSkeletalMesh(InSkelMesh, bReinitPose);
-
+	{
+		FRenderStateRecreator RenderStateRecreator(this);
+		Super::SetSkeletalMesh(InSkelMesh, bReinitPose);
+		
 #if WITH_EDITOR
-	ValidateAnimation();
+		ValidateAnimation();
 #endif
+	
+		if (GetPhysicsAsset() != OldPhysAsset && IsPhysicsStateCreated())
+		{
+			RecreatePhysicsState();
+		}
 
-	if (GetPhysicsAsset() != OldPhysAsset && IsPhysicsStateCreated())
-	{
-		RecreatePhysicsState();
-	}
-
-	UpdateHasValidBodies();
+		UpdateHasValidBodies();
 
-	InitAnim(bReinitPose);
+		InitAnim(bReinitPose);
+	}
 
 #if WITH_APEX_CLOTHING
 	RecreateClothingActors();
@@ -1723,7 +1726,7 @@ FVector USkeletalMeshComponent::GetClosestCollidingRigidBodyLocation(const FVect
 
 SIZE_T USkeletalMeshComponent::GetResourceSize( EResourceSizeMode::Type Mode )
 {
-	SIZE_T ResSize = 0;
+	SIZE_T ResSize = Super::GetResourceSize(Mode);
 
 	for (int32 i=0; i < Bodies.Num(); ++i)
 	{
diff --git a/Engine/Source/Runtime/Engine/Private/Components/SkinnedMeshComponent.cpp b/Engine/Source/Runtime/Engine/Private/Components/SkinnedMeshComponent.cpp
--- a/Engine/Source/Runtime/Engine/Private/Components/SkinnedMeshComponent.cpp
+++ b/Engine/Source/Runtime/Engine/Private/Components/SkinnedMeshComponent.cpp
@@ -964,27 +964,20 @@ void USkinnedMeshComponent::SetSkeletalMesh(USkeletalMesh* InSkelMesh, bool bRei
 		return;
 	}
 
-	// Destroy render state, swap skeletal mesh and recreate
-	const bool bWasRenderStateCreated = bRenderStateCreated;
-	if (bWasRenderStateCreated)
 	{
-		DestroyRenderState_Concurrent();
-	}
-
-	SkeletalMesh = InSkelMesh;
+		//Handle destroying and recreating the renderstate
+		FRenderStateRecreator RenderStateRecreator(this);
 
-	// Don't init anim state if not registered
-	if (IsRegistered())
-	{
-		AllocateTransformData();
-		UpdateMasterBoneMap();
-		UpdateLODStatus();
-		InvalidateCachedBounds();
-	}
+		SkeletalMesh = InSkelMesh;
 
-	if (bWasRenderStateCreated)
-	{
-		CreateRenderState_Concurrent();
+		// Don't init anim state if not registered
+		if (IsRegistered())
+		{
+			AllocateTransformData();
+			UpdateMasterBoneMap();
+			UpdateLODStatus();
+			InvalidateCachedBounds();
+		}
 	}
 	
 	// TODO: (LH) find better way to call this 
diff --git a/Engine/Source/Runtime/Engine/Private/Components/StaticMeshComponent.cpp b/Engine/Source/Runtime/Engine/Private/Components/StaticMeshComponent.cpp
--- a/Engine/Source/Runtime/Engine/Private/Components/StaticMeshComponent.cpp
+++ b/Engine/Source/Runtime/Engine/Private/Components/StaticMeshComponent.cpp
@@ -468,6 +468,8 @@ void UStaticMeshComponent::RemoveSpeedTreeWind()
 
 void UStaticMeshComponent::OnRegister()
 {
+	UpdateCollisionFromStaticMesh();
+	
 	if(StaticMesh != NULL && StaticMesh->RenderData)
 	{
 		// Check that the static-mesh hasn't been changed to be incompatible with the cached light-map.
@@ -1209,6 +1211,21 @@ void UStaticMeshComponent::PostEditChangeProperty(FPropertyChangedEvent& Propert
 }
 #endif // WITH_EDITOR
 
+bool UStaticMeshComponent::SupportsDefaultCollision()
+{
+	return StaticMesh && GetBodySetup() == StaticMesh->BodySetup;
+}
+
+void UStaticMeshComponent::UpdateCollisionFromStaticMesh()
+{
+	if(bUseDefaultCollision && SupportsDefaultCollision())
+	{
+		if (UBodySetup* BodySetup = GetBodySetup())
+		{
+			BodyInstance.UseExternalCollisionProfile(BodySetup);	//static mesh component by default uses the same collision profile as its static mesh
+		}
+	}
+}
 
 void UStaticMeshComponent::PostLoad()
 {
@@ -1337,6 +1354,12 @@ void UStaticMeshComponent::GetLocalBounds(FVector& Min, FVector& Max) const
 	}
 }
 
+void UStaticMeshComponent::SetCollisionProfileName(FName InCollisionProfileName)
+{
+	Super::SetCollisionProfileName(InCollisionProfileName);
+	bUseDefaultCollision = false;
+}
+
 bool UStaticMeshComponent::UsesOnlyUnlitMaterials() const
 {
 	if (StaticMesh && StaticMesh->RenderData)
diff --git a/Engine/Source/Runtime/Engine/Private/Controller.cpp b/Engine/Source/Runtime/Engine/Private/Controller.cpp
--- a/Engine/Source/Runtime/Engine/Private/Controller.cpp
+++ b/Engine/Source/Runtime/Engine/Private/Controller.cpp
@@ -579,7 +579,8 @@ APawn* AController::K2_GetPawn() const
 
 const FNavAgentProperties& AController::GetNavAgentPropertiesRef() const
 {
-	return Pawn && Pawn->GetMovementComponent() ? Pawn->GetMovementComponent()->GetNavAgentPropertiesRef() : FNavAgentProperties::DefaultProperties;
+	UPawnMovementComponent* MovementComponent = (Pawn ? Pawn->GetMovementComponent() : nullptr);
+	return MovementComponent ? MovementComponent->GetNavAgentPropertiesRef() : FNavAgentProperties::DefaultProperties;
 }
 
 FVector AController::GetNavAgentLocation() const
diff --git a/Engine/Source/Runtime/Engine/Private/GameplayStatics.cpp b/Engine/Source/Runtime/Engine/Private/GameplayStatics.cpp
--- a/Engine/Source/Runtime/Engine/Private/GameplayStatics.cpp
+++ b/Engine/Source/Runtime/Engine/Private/GameplayStatics.cpp
@@ -905,7 +905,7 @@ UAudioComponent* UGameplayStatics::CreateSound2D(UObject* WorldContextObject, cl
 		AudioComponent->bAllowSpatialization = false;
 		AudioComponent->bIsUISound = true;
 		AudioComponent->bAutoDestroy = true;
-		AudioComponent->SubtitlePriority = 10000.f; // Fixme: pass in? Do we want that exposed to blueprints though?		
+		AudioComponent->SubtitlePriority = Sound->GetSubtitlePriority();
 	}
 	return AudioComponent;
 }
@@ -988,7 +988,7 @@ class UAudioComponent* UGameplayStatics::SpawnSoundAttached(class USoundBase* So
 	FVector TestLocation = Location;
 	if (LocationType != EAttachLocation::KeepWorldPosition)
 	{
-		TestLocation = AttachToComponent->GetRelativeTransform().TransformPosition(Location);
+		TestLocation = AttachToComponent->GetComponentTransform().TransformPosition(Location);
 	}
 
 	UAudioComponent* AudioComponent = FAudioDevice::CreateComponent(Sound, AttachToComponent->GetWorld(), AttachToComponent->GetOwner(), false, bStopWhenAttachedToDestroyed, &TestLocation, AttenuationSettings, ConcurrencySettings);
@@ -1779,4 +1779,4 @@ int32 UGameplayStatics::GetIntOption( const FString& Options, const FString& Key
 	return DefaultValue;
 }
 
-#undef LOCTEXT_NAMESPACE
\ No newline at end of file
+#undef LOCTEXT_NAMESPACE
diff --git a/Engine/Source/Runtime/Engine/Private/HLODMeshCullingVolume.cpp b/Engine/Source/Runtime/Engine/Private/HLODMeshCullingVolume.cpp
new file mode 100644
--- /dev/null
+++ b/Engine/Source/Runtime/Engine/Private/HLODMeshCullingVolume.cpp
@@ -0,0 +1,16 @@
+// Copyright 1998-2016 Epic Games, Inc. All Rights Reserved.
+
+#include "EnginePrivate.h"
+#include "Components/BrushComponent.h"
+
+#include "../Classes/Engine/HLODMeshCullingVolume.h"
+
+AHLODMeshCullingVolume::AHLODMeshCullingVolume(const FObjectInitializer& ObjectInitializer)
+:Super(ObjectInitializer)
+{
+	bNotForClientOrServer = true;
+
+	bColored = true;
+	BrushColor = FColor(45, 225, 45);
+}
+
diff --git a/Engine/Source/Runtime/Engine/Private/KismetMathLibrary.cpp b/Engine/Source/Runtime/Engine/Private/KismetMathLibrary.cpp
--- a/Engine/Source/Runtime/Engine/Private/KismetMathLibrary.cpp
+++ b/Engine/Source/Runtime/Engine/Private/KismetMathLibrary.cpp
@@ -3,6 +3,7 @@
 #include "EnginePrivate.h"
 #include "Kismet/KismetMathLibrary.h"
 #include "CoreStats.h"
+#include "Math/DualQuat.h"
 
 /** Interpolate a linear alpha value using an ease mode and function. */
 float EaseAlpha(float InAlpha, uint8 EasingFunc, float BlendExp, int32 Steps)
@@ -1138,16 +1139,40 @@ FTransform UKismetMathLibrary::InvertTransform(const FTransform& T)
 	return T.Inverse();
 }
 
-FTransform UKismetMathLibrary::TLerp(const FTransform& A, const FTransform& B, float Alpha)
+
+
+FTransform UKismetMathLibrary::TLerp(const FTransform& A, const FTransform& B, float Alpha, TEnumAsByte<ELerpInterpolationMode::Type> LerpInterpolationMode)
 {
 	FTransform Result;
-	
+
 	FTransform NA = A;
 	FTransform NB = B;
 	NA.NormalizeRotation();
 	NB.NormalizeRotation();
-	Result.Blend(NA, NB, Alpha);
-	return Result;
+
+	// Quaternion interpolation
+	if (LerpInterpolationMode == ELerpInterpolationMode::QuatInterp)
+	{
+		Result.Blend(NA, NB, Alpha);
+		return Result;
+	}
+	// Euler Angle interpolation
+	else if (LerpInterpolationMode == ELerpInterpolationMode::EulerInterp)
+	{
+		Result.SetTranslation(FMath::Lerp(NA.GetTranslation(), NB.GetTranslation(), Alpha));
+		Result.SetScale3D(FMath::Lerp(NA.GetScale3D(), NB.GetScale3D(), Alpha));
+		Result.SetRotation(FQuat(RLerp(NA.Rotator(), NB.Rotator(), Alpha, false)));
+		return Result;
+	}
+	// Dual quaternion interpolation
+	else
+	{
+		if ((NB.GetRotation() | NA.GetRotation()) < 0.0f)
+		{
+			NB.SetRotation(NB.GetRotation()*-1.0f);
+		}
+		return (FDualQuat(NA)*(1 - Alpha) + FDualQuat(NB)*Alpha).Normalized().AsFTransform(FMath::Lerp(NA.GetScale3D(), NB.GetScale3D(), Alpha));
+	}
 }
 
 FTransform UKismetMathLibrary::TEase(const FTransform& A, const FTransform& B, float Alpha, TEnumAsByte<EEasingFunc::Type> EasingFunc, float BlendExp, int32 Steps)
diff --git a/Engine/Source/Runtime/Engine/Private/KismetSystemLibrary.cpp b/Engine/Source/Runtime/Engine/Private/KismetSystemLibrary.cpp
--- a/Engine/Source/Runtime/Engine/Private/KismetSystemLibrary.cpp
+++ b/Engine/Source/Runtime/Engine/Private/KismetSystemLibrary.cpp
@@ -3231,7 +3231,7 @@ bool UKismetSystemLibrary::IsControllerAssignedToGamepad(int32 ControllerId)
 	return FPlatformMisc::IsControllerAssignedToGamepad(ControllerId);
 }
 
-void UKismetSystemLibrary::SetSupressViewportTransitionMessage(UObject* WorldContextObject, bool bState)
+void UKismetSystemLibrary::SetSuppressViewportTransitionMessage(UObject* WorldContextObject, bool bState)
 {
 	UWorld* World = GEngine->GetWorldFromContextObject(WorldContextObject);
 	if (World && World->GetFirstLocalPlayerFromController() != nullptr && World->GetFirstLocalPlayerFromController()->ViewportClient != nullptr )
diff --git a/Engine/Source/Runtime/Engine/Private/LODActor.cpp b/Engine/Source/Runtime/Engine/Private/LODActor.cpp
--- a/Engine/Source/Runtime/Engine/Private/LODActor.cpp
+++ b/Engine/Source/Runtime/Engine/Private/LODActor.cpp
@@ -6,48 +6,89 @@
 
 #include "EnginePrivate.h"
 #include "Engine/LODActor.h"
+#include "Engine/HLODMeshCullingVolume.h"
 #include "MapErrors.h"
 #include "MessageLog.h"
 #include "UObjectToken.h"
 
 #include "StaticMeshResources.h"
 
+#if WITH_EDITOR
+#include "Editor.h"
+#endif
+
 #define LOCTEXT_NAMESPACE "LODActor"
 
 #if !(UE_BUILD_SHIPPING || UE_BUILD_TEST)
-static void ToggleHLODEnabled(UWorld* InWorld)
+static void HLODConsoleCommand(const TArray<FString>& Args, UWorld* World)
 {
-	static bool bHLODEnabled = true;
-
-	FlushRenderingCommands();
-
-	auto Levels = InWorld->GetLevels();
-
-	for (auto Level : Levels)
+	if (Args.Num() == 1)
 	{
-		for (AActor* Actor : Level->Actors)
+		const int32 State = FCString::Atoi(*Args[0]);
+
+		if (State == 0 || State == 1)
 		{
-			ALODActor* LODActor = Cast<ALODActor>(Actor);
-			if (LODActor)
+			const bool bHLODEnabled = (State == 1) ?  true : false;
+			FlushRenderingCommands();
+			const TArray<ULevel*>& Levels = World->GetLevels();
+			for (ULevel* Level : Levels)
 			{
-				LODActor->SetActorHiddenInGame(bHLODEnabled);
-	#if WITH_EDITOR
-				LODActor->SetIsTemporarilyHiddenInEditor(bHLODEnabled);
-	#endif // WITH_EDITOR
-				LODActor->MarkComponentsRenderStateDirty();
+				for (AActor* Actor : Level->Actors)
+				{
+					ALODActor* LODActor = Cast<ALODActor>(Actor);
+					if (LODActor)
+					{
+						LODActor->SetActorHiddenInGame(!bHLODEnabled);
+#if WITH_EDITOR
+						LODActor->SetIsTemporarilyHiddenInEditor(!bHLODEnabled);
+#endif // WITH_EDITOR
+						LODActor->MarkComponentsRenderStateDirty();
+					}
+				}
 			}
 		}
 	}
+	else if (Args.Num() == 2)
+	{
+#if WITH_EDITOR
+		if (Args[0] == "force")
+		{
+			const int32 ForcedLevel = FCString::Atoi(*Args[1]);
 
-	bHLODEnabled = !bHLODEnabled;
+			if (ForcedLevel >= -1 && ForcedLevel < World->GetWorldSettings()->HierarchicalLODSetup.Num())
+			{
+				const TArray<ULevel*>& Levels = World->GetLevels();
+				for (ULevel* Level : Levels)
+				{
+					for (AActor* Actor : Level->Actors)
+					{
+						ALODActor* LODActor = Cast<ALODActor>(Actor);
+
+						if (LODActor)
+						{
+							if (LODActor->LODLevel == ForcedLevel + 1)
+							{
+								LODActor->SetForcedView(true);
+							}
+							else
+							{
+								LODActor->SetHiddenFromEditorView(true, ForcedLevel + 1);
+							}
+						}
+					}
+				}
+			}
+		}
+#endif // WITH_EDITOR
+	}
 }
 
-static FAutoConsoleCommandWithWorld GToggleHLODEnabledCmd(
-	TEXT("ToggleHLODEnabled"),
-	TEXT("Toggles whether or not the HLOD system is enabled."),
-	FConsoleCommandWithWorldDelegate::CreateStatic(ToggleHLODEnabled),
-	ECVF_Cheat
+static FAutoConsoleCommandWithWorldAndArgs GHLODCmd(
+	TEXT("r.HLOD"),
+	TEXT("Single argument: 0 or 1 to Disable/Enable HLOD System\nMultiple arguments: force X where X is the HLOD level that should be forced into view"),
+	FConsoleCommandWithWorldAndArgsDelegate::CreateStatic(HLODConsoleCommand)
 	);
+
 #endif // !(UE_BUILD_SHIPPING || UE_BUILD_TEST)
 
 ALODActor::ALODActor(const FObjectInitializer& ObjectInitializer)
@@ -366,6 +407,7 @@ void ALODActor::SetIsDirty(const bool bNewState)
 			ALODActor* LODParentActor = Cast<ALODActor>(ParentComponent->GetOwner());
 			if (LODParentActor)
 			{
+				LODParentActor->Modify();
 				LODParentActor->SetIsDirty(true);
 			}
 		}
@@ -374,19 +416,13 @@ void ALODActor::SetIsDirty(const bool bNewState)
 		StaticMeshComponent->StaticMesh = nullptr;
 		// Mark render state dirty to update viewport
 		StaticMeshComponent->MarkRenderStateDirty();
-
-		// Propagate to sub actors that we no longer have a static mesh
-		for (auto& SubActor : SubActors)
-		{
-			SubActor->SetLODParent(nullptr, LODDrawDistance);
-		}
-
-
+#if WITH_EDITOR
 		// Broadcast actor marked dirty event
-		if (GEngine)
+		if (GEditor)
 		{
-			GEngine->BroadcastHLODActorMarkedDirty(this);
-		}		
+			GEditor->BroadcastHLODActorMarkedDirty(this);
+		}
+#endif // WITH_EDITOR
 	}	
 	else
 	{
@@ -400,30 +436,7 @@ void ALODActor::SetIsDirty(const bool bNewState)
 
 const bool ALODActor::HasValidSubActors()
 {
-	if (SubActors.Num() > 1)
-	{
-		// More than one sub actor
-		return true;
-	}
-	else if (SubActors.Num() == 0)
-	{
-		// No sub actors
-		return false;
-	}	
-	else
-	{
-		if (SubActors[0]->IsA<ALODActor>())
-		{
-			// LODActor as only sub actor (which doesn't make sense for the system)
-			return false;
-		}
-		else
-		{
-			// StaticMeshActor as only sub actor
-			return true;
-		}
-	}
-	
+	return (SubActors.Num() != 0);	
 }
 
 void ALODActor::ToggleForceView()
@@ -531,6 +544,7 @@ void ALODActor::RecalculateDrawingDistance(const float InTransitionScreenSize)
 	UpdateSubActorLODParents();
 }
 
+
 #endif // WITH_EDITOR
 
 FBox ALODActor::GetComponentsBoundingBox(bool bNonColliding) const 
@@ -566,4 +580,8 @@ FBox ALODActor::GetComponentsBoundingBox(bool bNonColliding) const
 	return BoundBox;	
 }
 
+//////////////////////////////////////////////////////////////////////////
+// AHLODMeshCullingVolume
+
+
 #undef LOCTEXT_NAMESPACE
diff --git a/Engine/Source/Runtime/Engine/Private/LevelActor.cpp b/Engine/Source/Runtime/Engine/Private/LevelActor.cpp
--- a/Engine/Source/Runtime/Engine/Private/LevelActor.cpp
+++ b/Engine/Source/Runtime/Engine/Private/LevelActor.cpp
@@ -579,9 +579,9 @@ bool UWorld::DestroyActor( AActor* ThisActor, bool bNetForce, bool bShouldModify
 
 	// Detach from anything we were attached to
 	USceneComponent* RootComp = ThisActor->GetRootComponent();
-	if( RootComp != NULL && RootComp->AttachParent != NULL)
+	if( RootComp != nullptr && RootComp->GetAttachParent() != nullptr)
 	{
-		AActor* OldParentActor = RootComp->AttachParent->GetOwner();
+		AActor* OldParentActor = RootComp->GetAttachParent()->GetOwner();
 		if (OldParentActor)
 		{
 			OldParentActor->Modify();
diff --git a/Engine/Source/Runtime/Engine/Private/LevelTick.cpp b/Engine/Source/Runtime/Engine/Private/LevelTick.cpp
--- a/Engine/Source/Runtime/Engine/Private/LevelTick.cpp
+++ b/Engine/Source/Runtime/Engine/Private/LevelTick.cpp
@@ -389,10 +389,10 @@ void UWorld::TickNetClient( float DeltaSeconds )
 -----------------------------------------------------------------------------*/
 
 
-bool UWorld::IsPaused()
+bool UWorld::IsPaused() const
 {
 	// pause if specifically set or if we're waiting for the end of the tick to perform streaming level loads (so actors don't fall through the world in the meantime, etc)
-	AWorldSettings* Info = GetWorldSettings();
+	const AWorldSettings* Info = GetWorldSettings();
 	return ( (Info && Info->Pauser != NULL && TimeSeconds >= PauseDelay) ||
 				(bRequestedBlockOnAsyncLoading && GetNetMode() == NM_Client) ||
 				(GEngine->ShouldCommitPendingMapChange(this)) ||
diff --git a/Engine/Source/Runtime/Engine/Private/Materials/MaterialInstance.cpp b/Engine/Source/Runtime/Engine/Private/Materials/MaterialInstance.cpp
--- a/Engine/Source/Runtime/Engine/Private/Materials/MaterialInstance.cpp
+++ b/Engine/Source/Runtime/Engine/Private/Materials/MaterialInstance.cpp
@@ -2640,7 +2640,7 @@ bool UMaterialInstance::IsPropertyActive(EMaterialProperty InProperty) const
 {
 	if(InProperty == MP_DiffuseColor || InProperty == MP_SpecularColor)
 	{
-		// to supress some CompilePropertyEx calls
+		// to suppress some CompilePropertyEx calls
 		return false;
 	}
 
diff --git a/Engine/Source/Runtime/Engine/Private/Materials/MaterialUniformExpressions.h b/Engine/Source/Runtime/Engine/Private/Materials/MaterialUniformExpressions.h
--- a/Engine/Source/Runtime/Engine/Private/Materials/MaterialUniformExpressions.h
+++ b/Engine/Source/Runtime/Engine/Private/Materials/MaterialUniformExpressions.h
@@ -1282,10 +1282,10 @@ public:
 		A->GetNumberValue(Context, ValueA);
 		B->GetNumberValue(Context, ValueB);
 
-		OutValue.R = fmod(ValueA.R, ValueB.R);
-		OutValue.G = fmod(ValueA.G, ValueB.G);
-		OutValue.B = fmod(ValueA.B, ValueB.B);
-		OutValue.A = fmod(ValueA.A, ValueB.A);
+		OutValue.R = FMath::Fmod(ValueA.R, ValueB.R);
+		OutValue.G = FMath::Fmod(ValueA.G, ValueB.G);
+		OutValue.B = FMath::Fmod(ValueA.B, ValueB.B);
+		OutValue.A = FMath::Fmod(ValueA.A, ValueB.A);
 	}
 	virtual bool IsConstant() const
 	{
diff --git a/Engine/Source/Runtime/Engine/Private/NetworkDriver.cpp b/Engine/Source/Runtime/Engine/Private/NetworkDriver.cpp
--- a/Engine/Source/Runtime/Engine/Private/NetworkDriver.cpp
+++ b/Engine/Source/Runtime/Engine/Private/NetworkDriver.cpp
@@ -192,6 +192,81 @@ void UNetDriver::AssertValid()
 {
 }
 
+/*static*/ bool UNetDriver::IsAdaptiveNetUpdateFrequencyEnabled()
+{
+	const bool bUseAdapativeNetFrequency = CVarUseAdaptiveNetUpdateFrequency.GetValueOnGameThread() > 0;
+	return bUseAdapativeNetFrequency;
+}
+
+FNetworkObjectInfo* UNetDriver::GetNetworkActor(const AActor* InActor)
+{
+	FNetworkObjectInfo* NetActor = nullptr;
+	if (InActor)
+	{
+		// TODO: this Find() operation is implicitly creating a TSharedPtr and causing a malloc.
+		// note: const_cast<> required because the set keys are not const.
+		TSharedPtr<FNetworkObjectInfo>* NetActorInfo = GetNetworkObjectList().GetObjects().Find(const_cast<AActor*>(InActor));
+		NetActor = (NetActorInfo ? NetActorInfo->Get() : nullptr);
+	}
+	return NetActor;
+}
+
+const FNetworkObjectInfo* UNetDriver::GetNetworkActor(const AActor* InActor) const
+{
+	return const_cast<UNetDriver*>(this)->GetNetworkActor(InActor);
+}
+
+bool UNetDriver::IsNetworkActorUpdateFrequencyThrottled(const FNetworkObjectInfo& InNetworkActor) const
+{
+	bool bThrottled = false;
+	if (IsAdaptiveNetUpdateFrequencyEnabled())
+	{
+		// Must have been replicated once for this to happen (and for OptimalNetUpdateDelta to have been set)
+		const AActor* Actor = InNetworkActor.Actor;
+		if (Actor && InNetworkActor.LastNetReplicateTime != 0)
+		{
+			const float ExpectedNetDelay = (1.0f / Actor->NetUpdateFrequency);
+			if (InNetworkActor.OptimalNetUpdateDelta > ExpectedNetDelay)
+			{
+				bThrottled = true;
+			}
+		}
+	}
+
+	return bThrottled;
+}
+
+bool UNetDriver::IsNetworkActorUpdateFrequencyThrottled(const AActor* InActor) const
+{
+	bool bThrottled = false;
+	if (InActor && IsAdaptiveNetUpdateFrequencyEnabled())
+	{
+		if (const FNetworkObjectInfo* NetActor = GetNetworkActor(InActor))
+		{
+			bThrottled = IsNetworkActorUpdateFrequencyThrottled(*NetActor);
+		}
+	}
+
+	return bThrottled;
+}
+
+void UNetDriver::CancelAdaptiveReplication(FNetworkObjectInfo& InNetworkActor)
+{
+	if (IsAdaptiveNetUpdateFrequencyEnabled())
+	{
+		if (AActor* Actor = InNetworkActor.Actor)
+		{
+			if (UWorld* World = Actor->GetWorld())
+			{
+				const float ExpectedNetDelay = (1.0f / Actor->NetUpdateFrequency);
+				Actor->SetNetUpdateTime(FMath::Min(Actor->NetUpdateTime, World->GetTimeSeconds() + FMath::FRandRange(0.5f, 1.0f) * ExpectedNetDelay));
+				InNetworkActor.OptimalNetUpdateDelta = ExpectedNetDelay;
+				// TODO: we really need a way to cancel the throttling completely. OptimalNetUpdateDelta is going to be recalculated based on LastNetReplicateTime.
+			}
+		}
+	}
+}
+
 void UNetDriver::TickFlush(float DeltaSeconds)
 {
 #if USE_SERVER_PERF_COUNTERS
@@ -2140,7 +2215,7 @@ void UNetDriver::ServerReplicateActors_BuildConsiderList( TArray<FNetworkObjectI
 
 	int32 NumInitiallyDormant = 0;
 
-	const bool bUseAdapativeNetFrequency = CVarUseAdaptiveNetUpdateFrequency.GetValueOnGameThread() > 0;
+	const bool bUseAdapativeNetFrequency = IsAdaptiveNetUpdateFrequencyEnabled();
 
 	for ( auto ActorIt = GetNetworkObjectList().GetObjects().CreateIterator(); ActorIt; ++ActorIt )
 	{
diff --git a/Engine/Source/Runtime/Engine/Private/Particles/ParticleComponents.cpp b/Engine/Source/Runtime/Engine/Private/Particles/ParticleComponents.cpp
--- a/Engine/Source/Runtime/Engine/Private/Particles/ParticleComponents.cpp
+++ b/Engine/Source/Runtime/Engine/Private/Particles/ParticleComponents.cpp
@@ -3303,24 +3303,24 @@ void UParticleSystemComponent::OnRegister()
 	if (bAutoManageAttachment && !IsActive())
 	{
 		// Detach from current parent, we are supposed to wait for activation.
-		if (AttachParent)
+		if (GetAttachParent())
 		{
 			// If no auto attach parent override, use the current parent when we activate
 			if (!AutoAttachParent.IsValid())
 			{
-				AutoAttachParent = AttachParent;
+				AutoAttachParent = GetAttachParent();
 			}
 			// If no auto attach socket override, use current socket when we activate
 			if (AutoAttachSocketName == NAME_None)
 			{
-				AutoAttachSocketName = AttachSocketName;
+				AutoAttachSocketName = GetAttachSocketName();
 			}
 
 			// Prevent attachment before Super::OnRegister() tries to attach us, since we only attach when activated.
-			if (AttachParent->AttachChildren.Contains(this))
+			if (GetAttachParent()->GetAttachChildren().Contains(this))
 			{
 				// Only detach if we are not about to auto attach to the same target, that would be wasteful.
-				if (!bAutoActivate || (AutoAttachLocationType != EAttachLocation::KeepRelativeOffset) || (AutoAttachSocketName != AttachSocketName) || (AutoAttachParent != AttachParent))
+				if (!bAutoActivate || (AutoAttachLocationType != EAttachLocation::KeepRelativeOffset) || (AutoAttachSocketName != GetAttachSocketName()) || (AutoAttachParent != GetAttachParent()))
 				{
 					DetachFromParent(/*bMaintainWorldPosition=*/ false, /*bCallModify=*/ false);
 				}
@@ -5003,7 +5003,7 @@ void UParticleSystemComponent::ActivateSystem(bool bFlagAsJustAttached)
 			USceneComponent* NewParent = AutoAttachParent.Get();
 			if (NewParent)
 			{
-				const bool bAlreadyAttached = GetAttachParent() && (GetAttachParent() == NewParent) && (GetAttachSocketName() == AutoAttachSocketName) && GetAttachParent()->AttachChildren.Contains(this);
+				const bool bAlreadyAttached = GetAttachParent() && (GetAttachParent() == NewParent) && (GetAttachSocketName() == AutoAttachSocketName) && GetAttachParent()->GetAttachChildren().Contains(this);
 				if (!bAlreadyAttached)
 				{
 					bDidAutoAttach = bWasAutoAttached;
diff --git a/Engine/Source/Runtime/Engine/Private/Particles/ParticleEmitterInstances.cpp b/Engine/Source/Runtime/Engine/Private/Particles/ParticleEmitterInstances.cpp
--- a/Engine/Source/Runtime/Engine/Private/Particles/ParticleEmitterInstances.cpp
+++ b/Engine/Source/Runtime/Engine/Private/Particles/ParticleEmitterInstances.cpp
@@ -910,7 +910,7 @@ float FParticleEmitterInstance::Tick_EmitterTimeSetup(float DeltaTime, UParticle
  *
  *	@param	DeltaTime			The current time slice
  *	@param	CurrentLODLevel		The current LOD level for the instance
- *	@param	bSuppressSpawning	true if spawning has been supressed on the owning particle system component
+ *	@param	bSuppressSpawning	true if spawning has been suppressed on the owning particle system component
  *	@param	bFirstTime			true if this is the first time the instance has been ticked
  *
  *	@return	float				The SpawnFraction remaining
diff --git a/Engine/Source/Runtime/Engine/Private/Particles/ParticleTrail2EmitterInstance.cpp b/Engine/Source/Runtime/Engine/Private/Particles/ParticleTrail2EmitterInstance.cpp
--- a/Engine/Source/Runtime/Engine/Private/Particles/ParticleTrail2EmitterInstance.cpp
+++ b/Engine/Source/Runtime/Engine/Private/Particles/ParticleTrail2EmitterInstance.cpp
@@ -3136,7 +3136,7 @@ void FParticleAnimTrailEmitterInstance::SpawnParticle( int32& StartParticleIndex
 //	}
 //	else
 	{
-		UMeshComponent* MeshComp = Cast<UMeshComponent>(Component->AttachParent);
+		UMeshComponent* MeshComp = Cast<UMeshComponent>(Component->GetAttachParent());
 		check(TrailTypeData);
 		check(MeshComp);
 
@@ -3427,7 +3427,7 @@ float FParticleAnimTrailEmitterInstance::Spawn(float DeltaTime)
 	bool bTilingTrail = !FMath::IsNearlyZero(TrailTypeData->TilingDistance);
 
 	//The mesh we're sampline socket locations from.
-	UMeshComponent* MeshComp = Cast<UMeshComponent>(Component->AttachParent);
+	UMeshComponent* MeshComp = Cast<UMeshComponent>(Component->GetAttachParent());
 
 	check(TrailTypeData);
 	// Don't allow new spawning if the emitter is finished
diff --git a/Engine/Source/Runtime/Engine/Private/Pawn.cpp b/Engine/Source/Runtime/Engine/Private/Pawn.cpp
--- a/Engine/Source/Runtime/Engine/Private/Pawn.cpp
+++ b/Engine/Source/Runtime/Engine/Private/Pawn.cpp
@@ -324,10 +324,11 @@ void APawn::TurnOff()
 	// do not block anything, just ignore
 	SetActorEnableCollision(false);
 
-	if (GetMovementComponent())
+	UPawnMovementComponent* MovementComponent = GetMovementComponent();
+	if (MovementComponent)
 	{
-		GetMovementComponent()->StopMovementImmediately();
-		GetMovementComponent()->SetComponentTickEnabled(false);
+		MovementComponent->StopMovementImmediately();
+		MovementComponent->SetComponentTickEnabled(false);
 	}
 
 	DisableComponentsSimulatePhysics();
@@ -671,9 +672,10 @@ void APawn::AddControllerRollInput(float Val)
 
 void APawn::Restart()
 {
-	if (GetMovementComponent())
+	UPawnMovementComponent* MovementComponent = GetMovementComponent();
+	if (MovementComponent)
 	{
-		GetMovementComponent()->StopMovementImmediately();
+		MovementComponent->StopMovementImmediately();
 	}
 	ConsumeMovementInputVector();
 	RecalculateBaseEyeHeight();
@@ -937,10 +939,10 @@ void APawn::TeleportSucceeded(bool bIsATest)
 {
 	if (bIsATest == false)
 	{
-		UMovementComponent* const MoveComponent = GetMovementComponent();
-		if (MoveComponent)
+		UPawnMovementComponent* MovementComponent = GetMovementComponent();
+		if (MovementComponent)
 		{
-			MoveComponent->OnTeleported();
+			MovementComponent->OnTeleported();
 		}
 	}
 
@@ -969,9 +971,10 @@ void APawn::PostNetReceiveVelocity(const FVector& NewVelocity)
 {
 	if (Role == ROLE_SimulatedProxy)
 	{
-		if ( GetMovementComponent() )
+		UMovementComponent* const MoveComponent = GetMovementComponent();
+		if ( MoveComponent )
 		{
-			GetMovementComponent()->Velocity = NewVelocity;
+			MoveComponent->Velocity = NewVelocity;
 		}
 	}
 }
@@ -1088,5 +1091,6 @@ void APawn::PawnMakeNoise(float Loudness, FVector NoiseLocation, bool bUseNoiseM
 
 const FNavAgentProperties& APawn::GetNavAgentPropertiesRef() const
 {
-	return GetMovementComponent() ? GetMovementComponent()->GetNavAgentPropertiesRef() : FNavAgentProperties::DefaultProperties;
+	UPawnMovementComponent* MovementComponent = GetMovementComponent();
+	return MovementComponent ? MovementComponent->GetNavAgentPropertiesRef() : FNavAgentProperties::DefaultProperties;
 }
diff --git a/Engine/Source/Runtime/Engine/Private/PhysicsEngine/BodyInstance.cpp b/Engine/Source/Runtime/Engine/Private/PhysicsEngine/BodyInstance.cpp
--- a/Engine/Source/Runtime/Engine/Private/PhysicsEngine/BodyInstance.cpp
+++ b/Engine/Source/Runtime/Engine/Private/PhysicsEngine/BodyInstance.cpp
@@ -245,7 +245,8 @@ FBodyInstance::FBodyInstance()
 	, bOverrideMaxDepenetrationVelocity(false)
 	, bOverrideWalkableSlopeOnInstance(false)
 	, MaxDepenetrationVelocity(0.f)
-	, MassInKg(100.f)
+	, ExternalCollisionProfileBodySetup(nullptr)
+	, MassInKgOverride(100.f)
 	, LinearDamping(0.01)
 	, AngularDamping(0.0)
 	, CustomDOFPlaneNormal(FVector::ZeroVector)
@@ -435,6 +436,7 @@ void FBodyInstance::UpdatePhysicalMaterials()
 void FBodyInstance::InvalidateCollisionProfileName()
 {
 	CollisionProfileName = UCollisionProfile::CustomCollisionProfileName;
+	ExternalCollisionProfileBodySetup = nullptr;
 }
 
 void FBodyInstance::SetResponseToChannel(ECollisionChannel Channel, ECollisionResponse NewResponse)
@@ -479,18 +481,36 @@ void FBodyInstance::SetCollisionProfileName(FName InCollisionProfileName)
 {
 	SCOPE_CYCLE_COUNTER(STAT_BodyInst_SetCollisionProfileName);
 
-	if (CollisionProfileName != InCollisionProfileName)
+	//Note that GetCollisionProfileName will use the external profile if one is set.
+	//GetCollisionProfileName will be consistent with the values set by LoadProfileData.
+	//This is why we can't use CollisionProfileName directly during the equality check
+	if (GetCollisionProfileName() != InCollisionProfileName)
 	{
+		//LoadProfileData uses GetCollisionProfileName internally so we must now set the external collision data to null.
+		ExternalCollisionProfileBodySetup = nullptr;
 		CollisionProfileName = InCollisionProfileName;
 		// now Load ProfileData
 		LoadProfileData(false);
 	}
+	
+	ExternalCollisionProfileBodySetup = nullptr;	//Even if incoming is the same as GetCollisionProfileName we turn it into "manual mode"
+}
+
+FName FBodyInstance::GetCollisionProfileName() const
+{
+	FName ReturnProfileName = CollisionProfileName;
+	if (UBodySetup* BodySetupPtr = ExternalCollisionProfileBodySetup.Get(true))
+	{
+		ReturnProfileName = BodySetupPtr->DefaultInstance.CollisionProfileName;
+	}
+	
+	return ReturnProfileName;
 }
 
 
 bool FBodyInstance::DoesUseCollisionProfile() const
 {
-	return IsValidCollisionProfileName(CollisionProfileName);
+	return IsValidCollisionProfileName(GetCollisionProfileName());
 }
 
 void FBodyInstance::SetMassScale(float InMassScale)
@@ -561,6 +581,19 @@ FVector FBodyInstance::GetLockedAxis() const
 	return FVector::ZeroVector;
 }
 
+void FBodyInstance::UseExternalCollisionProfile(UBodySetup* InExternalCollisionProfileBodySetup)
+{
+	ensureAlways(InExternalCollisionProfileBodySetup);
+	ExternalCollisionProfileBodySetup = InExternalCollisionProfileBodySetup;
+	LoadProfileData(false);
+}
+
+void FBodyInstance::ClearExternalCollisionProfile()
+{
+	ExternalCollisionProfileBodySetup = nullptr;
+	LoadProfileData(false);
+}
+
 void FBodyInstance::SetDOFLock(EDOFMode::Type NewAxisMode)
 {
 	DOFMode = NewAxisMode;
@@ -1060,7 +1093,7 @@ struct FInitBodiesHelper
 		PhysXName = GetDebugDebugName(PrimitiveComp, BodySetup, DebugName);
 
 		bStatic = bCompileStatic || PrimitiveComp == nullptr || PrimitiveComp->Mobility != EComponentMobility::Movable;
-		SkelMeshComp = bCompileStatic ? nullptr : GetSkeletalMeshComponentAndProperties(PrimitiveComp, BodySetup, InstanceBlendWeight, bInstanceSimulatePhysics, bComponentAwake);
+		SkelMeshComp = bCompileStatic ? nullptr : GetSkeletalMeshComponentAndProperties(PrimitiveComp, BodySetup, InstanceBlendWeight, bInstanceSimulatePhysics);
 
 		const AActor* OwningActor = PrimitiveComp ? PrimitiveComp->GetOwner() : nullptr;
 		InitialLinVel = GetInitialLinearVelocity(OwningActor, bComponentAwake);
@@ -1083,6 +1116,8 @@ struct FInitBodiesHelper
 #endif
 	}
 
+	FORCEINLINE bool IsStatic() const { return bCompileStatic || bStatic; }
+
 	//The arguments passed into InitBodies
 	TArray<FBodyInstance*>& Bodies;   
 	TArray<FTransform>& Transforms;
@@ -1125,7 +1160,7 @@ struct FInitBodiesHelper
 		const ECollisionEnabled::Type CollisionType = Instance->GetCollisionEnabled();
 		const bool bDisableSim = CollisionType == ECollisionEnabled::QueryOnly || CollisionType == ECollisionEnabled::NoCollision;
 
-		if (bCompileStatic || bStatic)
+		if (IsStatic())
 		{
 			Instance->RigidActorSync = GPhysXSDK->createRigidStatic(PTransform);
 
@@ -1189,7 +1224,7 @@ struct FInitBodiesHelper
 		Instance->GetFilterData_AssumesLocked(ShapeData);
 		Instance->GetShapeFlags_AssumesLocked(ShapeData, ShapeData.CollisionEnabled, BodySetup->GetCollisionTraceFlag() == CTF_UseComplexAsSimple);
 
-		if (!bCompileStatic && PNewDynamic)
+		if (!IsStatic() && PNewDynamic)
 		{
 			if (!Instance->ShouldInstanceSimulatingPhysics())
 			{
@@ -1265,22 +1300,22 @@ struct FInitBodiesHelper
 			Instance->BodySetup = BodySetup;
 			Instance->Scale3D = Transform.GetScale3D();
 			Instance->CharDebugName = PhysXName;
-			Instance->bHasSharedShapes = bStatic && PhysScene->HasAsyncScene() && UPhysicsSettings::Get()->bEnableShapeSharing;
+			Instance->bHasSharedShapes = IsStatic() && PhysScene->HasAsyncScene() && UPhysicsSettings::Get()->bEnableShapeSharing;
 			Instance->bEnableGravity = Instance->bEnableGravity && (SkelMeshComp ? SkelMeshComp->BodyInstance.bEnableGravity : true);	//In the case of skeletal mesh component we AND bodies with the parent body
 
 			// Handle autowelding here to avoid extra work
-			if (!bCompileStatic && Instance->bAutoWeld)
+			if (!IsStatic() && Instance->bAutoWeld)
 			{
 				ECollisionEnabled::Type CollisionType = Instance->GetCollisionEnabled();
 				if (CollisionType != ECollisionEnabled::QueryOnly)
 				{
-					if (UPrimitiveComponent * ParentPrimComponent = PrimitiveComp ? Cast<UPrimitiveComponent>(PrimitiveComp->AttachParent) : NULL)
+					if (UPrimitiveComponent * ParentPrimComponent = PrimitiveComp ? Cast<UPrimitiveComponent>(PrimitiveComp->GetAttachParent()) : NULL)
 					{
 						UWorld* World = PrimitiveComp->GetWorld();
 						if (World && World->IsGameWorld())
 						{
 							//if we have a parent we will now do the weld and exit any further initialization
-							if (PrimitiveComp->WeldToImplementation(ParentPrimComponent, PrimitiveComp->AttachSocketName, false))	//welded new simulated body so initialization is done
+							if (PrimitiveComp->WeldToImplementation(ParentPrimComponent, PrimitiveComp->GetAttachSocketName(), false))	//welded new simulated body so initialization is done
 							{
 								return false;
 							}
@@ -1300,7 +1335,7 @@ struct FInitBodiesHelper
 			}
 
 			// Set sim parameters for bodies from skeletal mesh components
-			if (!bCompileStatic && SkelMeshComp)
+			if (!IsStatic() && SkelMeshComp)
 			{
 				Instance->bSimulatePhysics = bInstanceSimulatePhysics;
 				if (InstanceBlendWeight != -1.0f)
@@ -1327,7 +1362,7 @@ struct FInitBodiesHelper
 
 				if (bInitFail)
 				{
-					UE_LOG(LogPhysics, Log, TEXT("Init Instance %d of Primitive Component %s failed"), BodyIdx, *PrimitiveComp->GetName());
+					UE_LOG(LogPhysics, Warning, TEXT("Init Instance %d of Primitive Component %s failed. Does it have collision data available?"), BodyIdx, *PrimitiveComp->GetReadableName());
 					if (Instance->RigidActorSync)
 					{
 						Instance->RigidActorSync->release();
@@ -1342,6 +1377,7 @@ struct FInitBodiesHelper
 
 					Instance->OwnerComponent = nullptr;
 					Instance->BodySetup = nullptr;
+					Instance->ExternalCollisionProfileBodySetup = nullptr;
 
 					continue;
 				}
@@ -1362,7 +1398,7 @@ struct FInitBodiesHelper
 			}
 
 			//handle special stuff only dynamic actors care about
-			if (!bCompileStatic && PNewDynamic)
+			if (!IsStatic() && PNewDynamic)
 			{
 				// turn off gravity if desired
 				if (!Instance->bEnableGravity)
@@ -1373,9 +1409,9 @@ struct FInitBodiesHelper
 				bDynamicsUseAsyncScene = Instance->UseAsyncScene(PhysScene);
 			}
 
-			if (bCompileStatic || bCanDefer)
+			if (IsStatic() || bCanDefer)
 			{
-				if (!bCompileStatic && PNewDynamic)
+				if (!IsStatic() && PNewDynamic)
 				{
 					PhysScene->DeferAddActor(Instance, PNewDynamic, Instance->UseAsyncScene(PhysScene) ? PST_Async : PST_Sync);
 				}
@@ -1457,7 +1493,7 @@ struct FInitBodiesHelper
 		}
 
 		// Set up dynamic instance data
-		if (!bCompileStatic && !bStatic)
+		if (!IsStatic())
 		{
 			SCOPE_CYCLE_COUNTER(STAT_InitBodyPostAdd);
 			for (int32 BodyIdx = 0, NumBodies = Bodies.Num(); BodyIdx < NumBodies; ++BodyIdx)
@@ -1479,11 +1515,11 @@ struct FInitBodiesHelper
 		check(PDynamicActors.Num() == 0);
 
 		// Only static objects qualify for deferred addition
-		const bool bCanDefer = bCompileStatic;
+		const bool bCanDefer = IsStatic();
 		bool bDynamicsUseAsync = false;
 		if (CreateShapesAndActors_PhysX(PSyncActors, PAsyncActors, PDynamicActors, bCanDefer, bDynamicsUseAsync))
 		{
-			if (!bCompileStatic && !bCanDefer)
+			if (!IsStatic() && !bCanDefer)
 			{
 				const bool bAddingToSyncScene = (PSyncActors.Num() || (PDynamicActors.Num() && !bDynamicsUseAsync)) && PSyncScene;
 				const bool bAddingToAsyncScene = (PAsyncActors.Num() || (PDynamicActors.Num() && bDynamicsUseAsync)) && PAsyncScene;
@@ -1523,7 +1559,7 @@ struct FInitBodiesHelper
 
 					// Create the body definition
 					b2BodyDef BodyDefinition;
-					if (bStatic)
+					if (IsStatic())
 					{
 						BodyDefinition.type = b2_staticBody;
 					}
@@ -1628,6 +1664,7 @@ struct FInitBodiesHelper
 						//clear Owner and Setup info as well to properly clean up the BodyInstance.
 						Instance->OwnerComponent = NULL;
 						Instance->BodySetup = NULL;
+						Instance->ExternalCollisionProfileBodySetup = nullptr;
 
 						return;
 					}
@@ -1643,7 +1680,7 @@ struct FInitBodiesHelper
 					// Set the filter data on the shapes (call this after setting BodyData because it uses that pointer)
 					Instance->UpdatePhysicsFilterData();
 
-					if (!bStatic)
+					if (!IsStatic())
 					{
 						// Compute mass (call this after setting BodyData because it uses that pointer)
 						Instance->UpdateMassProperties();
@@ -1692,8 +1729,17 @@ void FBodyInstance::InitBody(class UBodySetup* Setup, const FTransform& Transfor
 	Bodies.Add(this);
 	Transforms.Add(Transform);
 
-	FInitBodiesHelper<false> InitBodiesHelper(Bodies, Transforms, Setup, PrimComp, InRBScene, InAggregate);
-	InitBodiesHelper.InitBodies();
+	bool bIsStatic = PrimComp == nullptr || PrimComp->Mobility != EComponentMobility::Movable;
+	if(bIsStatic)
+	{
+		FInitBodiesHelper<true> InitBodiesHelper(Bodies, Transforms, Setup, PrimComp, InRBScene, InAggregate);
+		InitBodiesHelper.InitBodies();
+	}
+	else
+	{
+		FInitBodiesHelper<false> InitBodiesHelper(Bodies, Transforms, Setup, PrimComp, InRBScene, InAggregate);
+		InitBodiesHelper.InitBodies();
+	}
 
 	Bodies.Reset();
 	Transforms.Reset();
@@ -1724,7 +1770,7 @@ TSharedPtr<TArray<ANSICHAR>> GetDebugDebugName(const UPrimitiveComponent* Primit
 	return PhysXName;
 }
 
-const USkeletalMeshComponent* GetSkeletalMeshComponentAndProperties(const UPrimitiveComponent* PrimitiveComp, const UBodySetup* BodySetup, float& InstanceBlendWeight, bool& bInstanceSimulatePhysics, bool& bComponentAwake)
+const USkeletalMeshComponent* GetSkeletalMeshComponentAndProperties(const UPrimitiveComponent* PrimitiveComp, const UBodySetup* BodySetup, float& InstanceBlendWeight, bool& bInstanceSimulatePhysics)
 {
 	const USkeletalMeshComponent* SkelMeshComp = Cast<USkeletalMeshComponent>(PrimitiveComp);
 	if (SkelMeshComp)
@@ -1752,8 +1798,6 @@ const USkeletalMeshComponent* GetSkeletalMeshComponentAndProperties(const UPrimi
 				}
 			}
 		}
-
-		bComponentAwake = SkelMeshComp->BodyInstance.bStartAwake;
 	}
 
 	return SkelMeshComp;
@@ -1928,8 +1972,10 @@ void FBodyInstance::TermBody()
 
 	// @TODO UE4: Release spring body here
 
+	CurrentSceneState = BodyInstanceSceneState::NotAdded;
 	BodySetup = NULL;
 	OwnerComponent = NULL;
+	ExternalCollisionProfileBodySetup = nullptr;
 
 	if (DOFConstraint)
 	{
@@ -2009,6 +2055,25 @@ bool FBodyInstance::Weld(FBodyInstance* TheirBody, const FTransform& TheirTM)
 				PxShape* PShape = PNewShapes[ShapeIdx];
 				ShapeToBodiesMap->Add(PShape, FWeldInfo(TheirBody, RelativeTM));
 			}
+
+			if(TheirBody->ShapeToBodiesMap.IsValid())
+			{
+				TSet<FBodyInstance*> Bodies;
+				//If the body that is welding to us has things welded to it, make sure to weld those things to us as well
+				TMap<physx::PxShape*, FWeldInfo>& TheirWeldInfo = *TheirBody->ShapeToBodiesMap.Get();
+				for(auto Itr = TheirWeldInfo.CreateIterator(); Itr; ++Itr)
+				{
+					const FWeldInfo& WeldInfo = Itr->Value;
+					if(!Bodies.Contains(WeldInfo.ChildBI))
+					{
+						Bodies.Add(WeldInfo.ChildBI);	//only want to weld once per body and can have multiple shapes
+						const FTransform ChildWorldTM = WeldInfo.RelativeTM * TheirTM;
+						Weld(WeldInfo.ChildBI, ChildWorldTM);
+					}
+				}
+
+				TheirWeldInfo.Empty();	//They are no longer root so empty this
+			}
 		}
 
 		PostShapeChange();
@@ -2177,6 +2242,12 @@ EScaleMode::Type ComputeScaleMode(const TArray<PxShape*>& PShapes)
 	return ScaleMode;
 }
 
+void FBodyInstance::SetMassOverride(float MassInKG)
+{
+	bOverrideMass = true;
+	MassInKgOverride = MassInKG;
+}
+
 bool FBodyInstance::UpdateBodyScale(const FVector& InScale3D, bool bForceUpdate)
 {
 	SCOPE_CYCLE_COUNTER(STAT_BodyInstanceUpdateBodyScale);
@@ -2549,7 +2620,7 @@ void FBodyInstance::SetInstanceSimulatePhysics(bool bSimulate, bool bMaintainPhy
 		// If we are enabling simulation, and we are the root body of our component, we detach the component 
 		if (OwnerComponentInst && OwnerComponentInst->IsRegistered() && OwnerComponentInst->GetBodyInstance() == this)
 		{
-			if (OwnerComponentInst->AttachParent)
+			if (OwnerComponentInst->GetAttachParent())
 			{
 				OwnerComponentInst->DetachFromParent(true);
 			}
@@ -3279,17 +3350,15 @@ void FBodyInstance::UpdateMassProperties()
 
 			// Apply user-defined mass scaling.
 			NewMass *= FMath::Clamp<float>(MassScale, 0.01f, 100.0f);
-
-			MassInKg = NewMass;	//update the override mass to be the default mass
 		}
 		else
 		{
-			NewMass = FMath::Max(MassInKg, 0.001f);	//min weight of 1g
+			NewMass = FMath::Max(MassInKgOverride, 0.001f);	//min weight of 1g
 		}
 
-			check(NewMass > 0.f);
+		check(NewMass > 0.f);
 
-			float MassRatio = NewMass / OldMass;
+		float MassRatio = NewMass / OldMass;
 		PxVec3 InertiaTensor = PRigidBody->getMassSpaceInertiaTensor();
 
 		PRigidBody->setMassSpaceInertiaTensor(InertiaTensor * MassRatio);
@@ -3324,7 +3393,7 @@ void FBodyInstance::UpdateMassProperties()
 		}
 		else
 		{
-			MassScaledDensity = FMath::Max(MassInKg, 0.001f);	//min weight of 1g	//TODO: this is actually wrong because we're assuming mass and density are the same thing, but good enough for now
+			MassScaledDensity = FMath::Max(MassInKgOverride, 0.001f);	//min weight of 1g	//TODO: this is actually wrong because we're assuming mass and density are the same thing, but good enough for now
 		}
 
 		check(MassScaledDensity > 0.f);
@@ -4390,6 +4459,7 @@ bool FBodyInstance::IsValidCollisionProfileName(FName InCollisionProfileName)
 
 void FBodyInstance::LoadProfileData(bool bVerifyProfile)
 {
+	FName UseCollisionProfileName = GetCollisionProfileName();
 	if ( bVerifyProfile )
 	{
 		// if collision profile name exists, 
@@ -4397,10 +4467,10 @@ void FBodyInstance::LoadProfileData(bool bVerifyProfile)
 		// if same, then keep the profile name
 		// if not same, that means it has been modified from default
 		// leave it as it is, and clear profile name
-		if ( IsValidCollisionProfileName(CollisionProfileName) )
+		if ( IsValidCollisionProfileName(UseCollisionProfileName) )
 		{
 			FCollisionResponseTemplate Template;
-			if ( UCollisionProfile::Get()->GetProfileTemplate(CollisionProfileName, Template) ) 
+			if ( UCollisionProfile::Get()->GetProfileTemplate(UseCollisionProfileName, Template) ) 
 			{
 				// this function is only used for old code that did require verification of using profile or not
 				// so that means it will have valid ResponsetoChannels value, so this is okay to access. 
@@ -4411,7 +4481,7 @@ void FBodyInstance::LoadProfileData(bool bVerifyProfile)
 			}
 			else
 			{
-				UE_LOG(LogPhysics, Warning, TEXT("COLLISION PROFILE [%s] is not found"), *CollisionProfileName.ToString());
+				UE_LOG(LogPhysics, Warning, TEXT("COLLISION PROFILE [%s] is not found"), *UseCollisionProfileName.ToString());
 				// if not nothing to do
 				InvalidateCollisionProfileName(); 
 			}
@@ -4420,9 +4490,9 @@ void FBodyInstance::LoadProfileData(bool bVerifyProfile)
 	}
 	else
 	{
-		if ( IsValidCollisionProfileName(CollisionProfileName) )
+		if ( IsValidCollisionProfileName(UseCollisionProfileName) )
 		{
-			if ( UCollisionProfile::Get()->ReadConfig(CollisionProfileName, *this) == false)
+			if ( UCollisionProfile::Get()->ReadConfig(UseCollisionProfileName, *this) == false)
 			{
 				// clear the name
 				InvalidateCollisionProfileName();
@@ -4432,7 +4502,18 @@ void FBodyInstance::LoadProfileData(bool bVerifyProfile)
 		// no profile, so it just needs to update container from array data
 		if ( DoesUseCollisionProfile() == false )
 		{
-			CollisionResponses.UpdateResponseContainerFromArray();
+			// if external profile copy the data over
+			if (ExternalCollisionProfileBodySetup.IsValid(true))
+			{
+				const FBodyInstance& ExternalBodyInstance = ExternalCollisionProfileBodySetup->DefaultInstance;
+				CollisionProfileName = ExternalBodyInstance.CollisionProfileName;
+				ObjectType = ExternalBodyInstance.ObjectType;
+				CollisionResponses.SetCollisionResponseContainer(ExternalBodyInstance.CollisionResponses.ResponseToChannels);
+			}
+			else
+			{
+				CollisionResponses.UpdateResponseContainerFromArray();
+			}
 		}
 	}
 }
@@ -4661,19 +4742,12 @@ void FBodyInstance::InitDynamicProperties_AssumesLocked()
 		int32 VelocityIterCount = FMath::Clamp(VelocitySolverIterationCount, 1, 255);
 		RigidActor->setSolverIterationCounts(PositionIterCount, VelocityIterCount);
 
-		if(RigidActor->getScene())
+		CreateDOFLock();
+		if(!IsRigidBodyKinematic_AssumesLocked(RigidActor))
 		{
-			CreateDOFLock();
-			if(!IsRigidBodyKinematic_AssumesLocked(RigidActor))
+			if(!bStartAwake && !bWokenExternally)
 			{
-				if(bStartAwake || bWokenExternally)
-				{
-					RigidActor->wakeUp();
-				}
-				else
-				{
-					RigidActor->putToSleep();
-				}
+				RigidActor->setWakeCounter(0.f);
 			}
 		}
 	}
@@ -5026,16 +5100,16 @@ void FBodyInstanceEditorHelpers::EnsureConsistentMobilitySimulationSettingsOnPos
 				{
 					if (Component->BodyInstance.bSimulatePhysics)
 					{
-						Component->BodyInstance.SetCollisionProfileName(UCollisionProfile::PhysicsActor_ProfileName);
+						Component->SetCollisionProfileName(UCollisionProfile::PhysicsActor_ProfileName);
 					}
 					else
 					{
-						Component->BodyInstance.SetCollisionProfileName(UCollisionProfile::BlockAllDynamic_ProfileName);
+						Component->SetCollisionProfileName(UCollisionProfile::BlockAllDynamic_ProfileName);
 					}
 				}
 				else
 				{
-					Component->BodyInstance.SetCollisionProfileName(UCollisionProfile::BlockAll_ProfileName);
+					Component->SetCollisionProfileName(UCollisionProfile::BlockAll_ProfileName);
 				}
 			}
 		}
diff --git a/Engine/Source/Runtime/Engine/Private/PhysicsEngine/BodySetup.cpp b/Engine/Source/Runtime/Engine/Private/PhysicsEngine/BodySetup.cpp
--- a/Engine/Source/Runtime/Engine/Private/PhysicsEngine/BodySetup.cpp
+++ b/Engine/Source/Runtime/Engine/Private/PhysicsEngine/BodySetup.cpp
@@ -1510,7 +1510,7 @@ class UPhysicalMaterial* UBodySetup::GetPhysMaterial() const
 float UBodySetup::CalculateMass(const UPrimitiveComponent* Component) const
 {
 	FVector ComponentScale(1.0f, 1.0f, 1.0f);
-	const FBodyInstance* BodyInstance = NULL;
+	const FBodyInstance* BodyInstance = &DefaultInstance;
 	float MassScale = DefaultInstance.MassScale;
 
 	const UPrimitiveComponent* OuterComp = Component != NULL ? Component : Cast<UPrimitiveComponent>(GetOuter());
@@ -1531,9 +1531,10 @@ float UBodySetup::CalculateMass(const UPrimitiveComponent* Component) const
 			}
 		}
 	}
-	else
+
+	if(BodyInstance->bOverrideMass)
 	{
-		BodyInstance = &DefaultInstance;
+		return BodyInstance->GetMassOverride();
 	}
 
 	UPhysicalMaterial* PhysMat = BodyInstance->GetSimplePhysicalMaterial();
diff --git a/Engine/Source/Runtime/Engine/Private/PhysicsEngine/ConstraintInstance.cpp b/Engine/Source/Runtime/Engine/Private/PhysicsEngine/ConstraintInstance.cpp
--- a/Engine/Source/Runtime/Engine/Private/PhysicsEngine/ConstraintInstance.cpp
+++ b/Engine/Source/Runtime/Engine/Private/PhysicsEngine/ConstraintInstance.cpp
@@ -127,6 +127,7 @@ bool FConstraintInstance::ExecuteOnUnbrokenJointReadWrite(TFunctionRef<void(phys
 
 void FConstraintInstance::UpdateLinearLimit()
 {
+	const float UseScale = bScaleLinearLimits ? LastKnownScale : 1.f;
 #if WITH_PHYSX
 	ExecuteOnUnbrokenJointReadWrite([&] (PxD6Joint* Joint)
 	{
@@ -140,7 +141,7 @@ void FConstraintInstance::UpdateLinearLimit()
 		if (LinearXMotion != LCM_Free || LinearYMotion != LCM_Free || LinearZMotion != LCM_Free)
 		{
 			// If limit drops below RB_MinSizeToLockDOF, just pass RB_MinSizeToLockDOF to physics - that axis will be locked anyway, and PhysX dislikes 0 here
-			float LinearLimit = FMath::Max<float>(LinearLimitSize, RB_MinSizeToLockDOF);
+			float LinearLimit = FMath::Max<float>(LinearLimitSize * UseScale, RB_MinSizeToLockDOF);
 			PxJointLinearLimit PLinearLimit(GPhysXSDK->getTolerancesScale(), LinearLimit, 0.05f * GPhysXSDK->getTolerancesScale().length);
 			SetSoftLimitParams_AssumesLocked(&PLinearLimit, bLinearLimitSoft, LinearLimitStiffness*AverageMass, LinearLimitDamping*AverageMass);
 			Joint->setLinearLimit(PLinearLimit);
@@ -235,6 +236,7 @@ FConstraintInstance::FConstraintInstance()
 	, SwingLimitDamping(5)
 	, TwistLimitStiffness(50)
 	, TwistLimitDamping(5)
+	, bScaleLinearLimits(true)
 	, bLinearPositionDrive(false)
 	, bLinearVelocityDrive(false)
 	, LinearPositionTarget(ForceInit)
@@ -384,16 +386,24 @@ bool GetPActors_AssumesLocked(const FBodyInstance* Body1, const FBodyInstance* B
 	return true;
 }
 
-bool FConstraintInstance::CreatePxJoint_AssumesLocked(physx::PxRigidActor* PActor1, physx::PxRigidActor* PActor2, physx::PxScene* PScene, const float Scale)
+bool FConstraintInstance::CreatePxJoint_AssumesLocked(physx::PxRigidActor* PActor1, physx::PxRigidActor* PActor2, physx::PxScene* PScene)
 {
 	ConstraintData = nullptr;
 
 	FTransform Local1 = GetRefFrame(EConstraintFrame::Frame1);
-	Local1.ScaleTranslation(FVector(Scale));
+	if(PActor1)
+	{
+		Local1.ScaleTranslation(FVector(LastKnownScale));
+	}
+
 	checkf(Local1.IsValid() && !Local1.ContainsNaN(), TEXT("%s"), *Local1.ToString());
 
 	FTransform Local2 = GetRefFrame(EConstraintFrame::Frame2);
-	Local2.ScaleTranslation(FVector(Scale));
+	if(PActor2)
+	{
+		Local2.ScaleTranslation(FVector(LastKnownScale));
+	}
+	
 	checkf(Local2.IsValid() && !Local2.ContainsNaN(), TEXT("%s"), *Local2.ToString());
 
 	SCOPED_SCENE_WRITE_LOCK(PScene);
@@ -487,9 +497,10 @@ void EnsureSleepingActorsStaySleeping_AssumesLocked(PxRigidActor* PActor1, PxRig
 /** 
  *	Create physics engine constraint.
  */
-void FConstraintInstance::InitConstraint(USceneComponent* Owner, FBodyInstance* Body1, FBodyInstance* Body2, float Scale)
+void FConstraintInstance::InitConstraint(USceneComponent* Owner, FBodyInstance* Body1, FBodyInstance* Body2, float InScale)
 {
 	OwnerComponent = Owner;
+	LastKnownScale = InScale;
 
 #if WITH_PHYSX
 	PhysxUserData = FPhysxUserData(this);
@@ -505,13 +516,12 @@ void FConstraintInstance::InitConstraint(USceneComponent* Owner, FBodyInstance*
 	PxScene* PScene = GetPScene_LockFree(Body1, Body2);
 	SCOPED_SCENE_WRITE_LOCK(PScene);
 
-	const bool bValidConstraintSetup = PScene && GetPActors_AssumesLocked(Body1, Body2, &PActor1, &PActor2) && CreatePxJoint_AssumesLocked(PActor1, PActor2, PScene, Scale);
+	const bool bValidConstraintSetup = PScene && GetPActors_AssumesLocked(Body1, Body2, &PActor1, &PActor2) && CreatePxJoint_AssumesLocked(PActor1, PActor2, PScene);
 	if (!bValidConstraintSetup)
 	{
 		return;
 	}
 
-
 	// update mass
 	UpdateAverageMass_AssumesLocked(PActor1, PActor2);
 	
diff --git a/Engine/Source/Runtime/Engine/Private/PhysicsEngine/PhysScene.cpp b/Engine/Source/Runtime/Engine/Private/PhysicsEngine/PhysScene.cpp
--- a/Engine/Source/Runtime/Engine/Private/PhysicsEngine/PhysScene.cpp
+++ b/Engine/Source/Runtime/Engine/Private/PhysicsEngine/PhysScene.cpp
@@ -266,10 +266,8 @@ FPhysScene::FPhysScene()
 	FrameTimeSmoothingFactor[PST_Sync] = PhysSetting->SyncSceneSmoothingFactor;
 	FrameTimeSmoothingFactor[PST_Async] = PhysSetting->AsyncSceneSmoothingFactor;
 
-#if WITH_SUBSTEPPING
 	bSubstepping = PhysSetting->bSubstepping;
 	bSubsteppingAsync = PhysSetting->bSubsteppingAsync;
-#endif
 	bAsyncSceneEnabled = PhysSetting->bEnableAsyncScene;
 	NumPhysScenes = bAsyncSceneEnabled ? PST_Async + 1 : PST_Cloth + 1;
 
@@ -348,7 +346,6 @@ bool FPhysScene::GetKinematicTarget_AssumesLocked(const FBodyInstance* BodyInsta
 #if WITH_PHYSX
 	if (PxRigidDynamic * PRigidDynamic = BodyInstance->GetPxRigidDynamic_AssumesLocked())
 	{
-#if WITH_SUBSTEPPING
 		uint32 BodySceneType = SceneType_AssumesLocked(BodyInstance);
 		if (IsSubstepping(BodySceneType))
 		{
@@ -356,7 +353,6 @@ bool FPhysScene::GetKinematicTarget_AssumesLocked(const FBodyInstance* BodyInsta
 			return PhysSubStepper->GetKinematicTarget_AssumesLocked(BodyInstance, OutTM);
 		}
 		else
-#endif
 		{
 			PxTransform POutTM;
 			bool validTM = PRigidDynamic->getKinematicTarget(POutTM);
@@ -382,7 +378,6 @@ void FPhysScene::SetKinematicTarget_AssumesLocked(FBodyInstance* BodyInstance, c
 		const bool bIsKinematicTarget = IsRigidBodyKinematicAndInSimulationScene_AssumesLocked(PRigidDynamic);
 		if(bIsKinematicTarget)
 		{
-#if WITH_SUBSTEPPING
 			uint32 BodySceneType = SceneType_AssumesLocked(BodyInstance);
 			if (bAllowSubstepping && IsSubstepping(BodySceneType))
 			{
@@ -390,7 +385,6 @@ void FPhysScene::SetKinematicTarget_AssumesLocked(FBodyInstance* BodyInstance, c
 				PhysSubStepper->SetKinematicTarget_AssumesLocked(BodyInstance, TargetTransform);
 			}
 			else
-#endif
 			{
 				const PxTransform PNewPose = U2PTransform(TargetTransform);
 				PRigidDynamic->setKinematicTarget(PNewPose);
@@ -409,7 +403,6 @@ void FPhysScene::SetKinematicTarget_AssumesLocked(FBodyInstance* BodyInstance, c
 void FPhysScene::AddCustomPhysics_AssumesLocked(FBodyInstance* BodyInstance, FCalculateCustomPhysics& CalculateCustomPhysics)
 {
 #if WITH_PHYSX
-#if WITH_SUBSTEPPING
 	uint32 BodySceneType = SceneType_AssumesLocked(BodyInstance);
 	if (IsSubstepping(BodySceneType))
 	{
@@ -417,7 +410,6 @@ void FPhysScene::AddCustomPhysics_AssumesLocked(FBodyInstance* BodyInstance, FCa
 		PhysSubStepper->AddCustomPhysics_AssumesLocked(BodyInstance, CalculateCustomPhysics);
 	}
 	else
-#endif
 	{
 		// Since physics frame is set up before "pre-physics" tick group is called, can just fetch delta time from there
 		CalculateCustomPhysics.ExecuteIfBound(this->DeltaSeconds, BodyInstance);
@@ -431,7 +423,6 @@ void FPhysScene::AddForce_AssumesLocked(FBodyInstance* BodyInstance, const FVect
 
 	if (PxRigidBody * PRigidBody = BodyInstance->GetPxRigidBody_AssumesLocked())
 	{
-#if WITH_SUBSTEPPING
 		uint32 BodySceneType = SceneType_AssumesLocked(BodyInstance);
 		if (bAllowSubstepping && IsSubstepping(BodySceneType))
 		{
@@ -439,7 +430,6 @@ void FPhysScene::AddForce_AssumesLocked(FBodyInstance* BodyInstance, const FVect
 			PhysSubStepper->AddForce_AssumesLocked(BodyInstance, Force, bAccelChange);
 		}
 		else
-#endif
 		{
 			PRigidBody->addForce(U2PVector(Force), bAccelChange ? PxForceMode::eACCELERATION : PxForceMode::eFORCE, true);
 		}
@@ -453,7 +443,6 @@ void FPhysScene::AddForceAtPosition_AssumesLocked(FBodyInstance* BodyInstance, c
 
 	if (PxRigidBody * PRigidBody = BodyInstance->GetPxRigidBody_AssumesLocked())
 	{
-#if WITH_SUBSTEPPING
 		uint32 BodySceneType = SceneType_AssumesLocked(BodyInstance);
 		if (bAllowSubstepping && IsSubstepping(BodySceneType))
 		{
@@ -461,7 +450,6 @@ void FPhysScene::AddForceAtPosition_AssumesLocked(FBodyInstance* BodyInstance, c
 			PhysSubStepper->AddForceAtPosition_AssumesLocked(BodyInstance, Force, Position);
 		}
 		else
-#endif
 		{
 			PxRigidBodyExt::addForceAtPos(*PRigidBody, U2PVector(Force), U2PVector(Position), PxForceMode::eFORCE, true);
 		}
@@ -475,7 +463,6 @@ void FPhysScene::AddRadialForceToBody_AssumesLocked(FBodyInstance* BodyInstance,
 
 	if (PxRigidBody * PRigidBody = BodyInstance->GetPxRigidBody_AssumesLocked())
 	{
-#if WITH_SUBSTEPPING
 		uint32 BodySceneType = SceneType_AssumesLocked(BodyInstance);
 		if (bAllowSubstepping && IsSubstepping(BodySceneType))
 		{
@@ -483,7 +470,6 @@ void FPhysScene::AddRadialForceToBody_AssumesLocked(FBodyInstance* BodyInstance,
 			PhysSubStepper->AddRadialForceToBody_AssumesLocked(BodyInstance, Origin, Radius, Strength, Falloff, bAccelChange);
 		}
 		else
-#endif
 		{
 			AddRadialForceToPxRigidBody_AssumesLocked(*PRigidBody, Origin, Radius, Strength, Falloff, bAccelChange);
 		}
@@ -497,7 +483,6 @@ void FPhysScene::AddTorque_AssumesLocked(FBodyInstance* BodyInstance, const FVec
 
 	if (PxRigidBody * PRigidBody = BodyInstance->GetPxRigidBody_AssumesLocked())
 	{
-#if WITH_SUBSTEPPING
 		uint32 BodySceneType = SceneType_AssumesLocked(BodyInstance);
 		if (bAllowSubstepping && IsSubstepping(BodySceneType))
 		{
@@ -505,7 +490,6 @@ void FPhysScene::AddTorque_AssumesLocked(FBodyInstance* BodyInstance, const FVec
 			PhysSubStepper->AddTorque_AssumesLocked(BodyInstance, Torque, bAccelChange);
 		}
 		else
-#endif
 		{
 			PRigidBody->addTorque(U2PVector(Torque), bAccelChange ? PxForceMode::eACCELERATION : PxForceMode::eFORCE, true);
 		}
@@ -529,10 +513,8 @@ void FPhysScene::TermBody_AssumesLocked(FBodyInstance* BodyInstance)
 {
 	if (PxRigidBody* PRigidBody = BodyInstance->GetPxRigidBody_AssumesLocked())
 	{
-#if WITH_SUBSTEPPING
 		FPhysSubstepTask* PhysSubStepper = PhysSubSteppers[SceneType_AssumesLocked(BodyInstance)];
 		PhysSubStepper->RemoveBodyInstance_AssumesLocked(BodyInstance);
-#endif
 	}
 
 	// Remove body from any pending deferred addition / removal
@@ -552,8 +534,6 @@ void FPhysScene::TermBody_AssumesLocked(FBodyInstance* BodyInstance)
 #endif
 }
 
-#if WITH_SUBSTEPPING
-
 #if WITH_APEX
 void FPhysScene::AddPendingDamageEvent(UDestructibleComponent* DestructibleComponent, const NxApexDamageEventReportData& DamageEvent)
 {
@@ -600,8 +580,6 @@ bool FPhysScene::SubstepSimulation(uint32 SceneType, FGraphEventRef &InOutComple
 
 }
 
-#endif //#if WITH_SUBSTEPPING
-
 /** Adds to queue of skelmesh we want to add to collision disable table */
 void FPhysScene::DeferredAddCollisionDisableTable(uint32 SkelMeshCompID, TMap<struct FRigidBodyIndexPair, bool> * CollisionDisableTable)
 {
@@ -852,13 +830,11 @@ void FPhysScene::TickPhysScene(uint32 SceneType, FGraphEventRef& InOutCompletion
 		return;
 	}
 
-#if WITH_SUBSTEPPING
 	if (IsSubstepping(SceneType))	//we don't bother sub-stepping cloth
 	{
 		//We're about to start stepping so swap buffers. Might want to find a better place for this?
 		PhysSubSteppers[SceneType]->SwapBuffers();
 	}
-#endif
 
 	/**
 	* clamp down... if this happens we are simming physics slower than real-time, so be careful with it.
@@ -907,16 +883,13 @@ void FPhysScene::TickPhysScene(uint32 SceneType, FGraphEventRef& InOutCompletion
 	if (VehicleManager && SceneType == PST_Sync)
 	{
 		float TickTime = AveragedFrameTime[SceneType];
-#if WITH_SUBSTEPPING
 		if (IsSubstepping(SceneType))
 		{
 			TickTime = UseSyncTime(SceneType) ? SyncDeltaSeconds : DeltaSeconds;
 		}
-#endif
 		VehicleManager->PreTick(TickTime);
-#if WITH_SUBSTEPPING
+
 		if (IsSubstepping(SceneType) == false)
-#endif
 		{
 			VehicleManager->Update(AveragedFrameTime[SceneType]);
 		}
@@ -931,13 +904,11 @@ void FPhysScene::TickPhysScene(uint32 SceneType, FGraphEventRef& InOutCompletion
 	if (ApexScene && UseDelta > 0.f)
 #endif
 	{
-#if WITH_SUBSTEPPING
 		if (IsSubstepping(SceneType)) //we don't bother sub-stepping cloth
 		{
 			bTaskOutstanding = SubstepSimulation(SceneType, InOutCompletionEvent);
 		}
 		else
-#endif
 		{
 #if !WITH_APEX
 			PhysXCompletionTask* Task = new PhysXCompletionTask(InOutCompletionEvent, SceneType, PScene->getTaskManager());
@@ -967,10 +938,8 @@ void FPhysScene::TickPhysScene(uint32 SceneType, FGraphEventRef& InOutCompletion
 		TArray<FBaseGraphTask*> NewTasks;
 		InOutCompletionEvent->DispatchSubsequents(NewTasks, ENamedThreads::AnyThread); // nothing to do, so nothing to wait for
 	}
-#if WITH_SUBSTEPPING
 	bSubstepping = UPhysicsSettings::Get()->bSubstepping;
 	bSubsteppingAsync = UPhysicsSettings::Get()->bSubsteppingAsync;
-#endif
 }
 
 void FPhysScene::KillVisualDebugger()
@@ -1663,8 +1632,11 @@ void FPhysScene::InitPhysScene(uint32 SceneType)
 	
 #endif
 
-	// We want to use 'active transforms'
-	PSceneDesc.flags |= PxSceneFlag::eENABLE_ACTIVETRANSFORMS;
+	if(!UPhysicsSettings::Get()->bDisableActiveTransforms)
+	{
+		// We want to use 'active transforms'
+		PSceneDesc.flags |= PxSceneFlag::eENABLE_ACTIVETRANSFORMS;
+	}
 
 	// @TODO Should we set up PSceneDesc.limits? How?
 
@@ -1741,7 +1713,6 @@ void FPhysScene::InitPhysScene(uint32 SceneType)
 	}
 #endif
 
-#if WITH_SUBSTEPPING
 	//Initialize substeppers
 	//we don't bother sub-stepping cloth
 #if WITH_PHYSX
@@ -1758,8 +1729,6 @@ void FPhysScene::InitPhysScene(uint32 SceneType)
 	}
 #endif
 	
-#endif
-
 #endif // WITH_PHYSX
 }
 
@@ -1787,8 +1756,6 @@ void FPhysScene::TermPhysScene(uint32 SceneType)
 		}
 #endif
 
-#if WITH_SUBSTEPPING
-
 		if (SceneType == PST_Sync && PhysSubSteppers[SceneType])
 		{
 			PhysSubSteppers[SceneType]->SetVehicleManager(NULL);
@@ -1796,7 +1763,6 @@ void FPhysScene::TermPhysScene(uint32 SceneType)
 
 		delete PhysSubSteppers[SceneType];
 		PhysSubSteppers[SceneType] = NULL;
-#endif
 
 		// @todo block on any running scene before calling this
 		GPhysCommandHandler->DeferredRelease(PScene);
diff --git a/Engine/Source/Runtime/Engine/Private/PhysicsEngine/PhysSubstepTasks.cpp b/Engine/Source/Runtime/Engine/Private/PhysicsEngine/PhysSubstepTasks.cpp
--- a/Engine/Source/Runtime/Engine/Private/PhysicsEngine/PhysSubstepTasks.cpp
+++ b/Engine/Source/Runtime/Engine/Private/PhysicsEngine/PhysSubstepTasks.cpp
@@ -4,8 +4,6 @@
 #include "PhysicsPublic.h"
 #include "PhysicsEngine/PhysicsSettings.h"
 
-#if WITH_SUBSTEPPING
-
 #if WITH_PHYSX
 #include "PhysXSupport.h"
 #include "../Vehicles/PhysXVehicleManager.h"
@@ -456,6 +454,4 @@ void FPhysSubstepTask::SubstepSimulationEnd(ENamedThreads::Type CurrentThread, c
 void FPhysSubstepTask::SetVehicleManager(FPhysXVehicleManager * InVehicleManager)
 {
 	VehicleManager = InVehicleManager;
-}
-
-#endif //if WITH_SUBSTEPPING
\ No newline at end of file
+}
\ No newline at end of file
diff --git a/Engine/Source/Runtime/Engine/Private/PhysicsEngine/PhysSubstepTasks.h b/Engine/Source/Runtime/Engine/Private/PhysicsEngine/PhysSubstepTasks.h
--- a/Engine/Source/Runtime/Engine/Private/PhysicsEngine/PhysSubstepTasks.h
+++ b/Engine/Source/Runtime/Engine/Private/PhysicsEngine/PhysSubstepTasks.h
@@ -40,9 +40,6 @@ public:
 #endif	//#if WITH_PHYSX
 
 
-#if WITH_SUBSTEPPING
-
-
 /** Hold information about kinematic target */
 struct FKinematicTarget
 {
@@ -186,6 +183,4 @@ private:
 	PxApexScene * PAScene;
 #endif
 
-};
-
-#endif //if WITH_SUBSTEPPING
\ No newline at end of file
+};
\ No newline at end of file
diff --git a/Engine/Source/Runtime/Engine/Private/PhysicsEngine/PhysicsConstraintComponent.cpp b/Engine/Source/Runtime/Engine/Private/PhysicsEngine/PhysicsConstraintComponent.cpp
--- a/Engine/Source/Runtime/Engine/Private/PhysicsEngine/PhysicsConstraintComponent.cpp
+++ b/Engine/Source/Runtime/Engine/Private/PhysicsEngine/PhysicsConstraintComponent.cpp
@@ -191,7 +191,7 @@ void UPhysicsConstraintComponent::InitComponentConstraint()
 	// Then we init the constraint
 	FBodyInstance* Body1 = GetBodyInstance(EConstraintFrame::Frame1);
 	FBodyInstance* Body2 = GetBodyInstance(EConstraintFrame::Frame2);
-	ConstraintInstance.InitConstraint(this, Body1, Body2, 1.0f);
+	ConstraintInstance.InitConstraint(this, Body1, Body2, GetConstraintScale());
 }
 
 void UPhysicsConstraintComponent::TermComponentConstraint()
@@ -204,6 +204,11 @@ void UPhysicsConstraintComponent::OnConstraintBrokenHandler(FConstraintInstance*
 	OnConstraintBroken.Broadcast(BrokenConstraint->ConstraintIndex);
 }
 
+float UPhysicsConstraintComponent::GetConstraintScale() const
+{
+	return GetComponentScale().GetAbsMin();
+}
+
 void UPhysicsConstraintComponent::SetConstrainedComponents(UPrimitiveComponent* Component1, FName BoneName1, UPrimitiveComponent* Component2, FName BoneName2)
 {
 	if(Component1 != NULL)
@@ -385,9 +390,10 @@ void UPhysicsConstraintComponent::CheckForErrors()
 void UPhysicsConstraintComponent::UpdateConstraintFrames()
 {
 	FTransform A1Transform = GetBodyTransform(EConstraintFrame::Frame1);
-	A1Transform.SetScale3D(FVector(1.f));
+	A1Transform.RemoveScaling();
+
 	FTransform A2Transform = GetBodyTransform(EConstraintFrame::Frame2);
-	A2Transform.SetScale3D(FVector(1.f));
+	A2Transform.RemoveScaling();
 
 	// World ref frame
 	const FVector WPos = GetComponentLocation();
@@ -400,11 +406,28 @@ void UPhysicsConstraintComponent::UpdateConstraintFrames()
 
 	const FVector RotatedX = ConstraintInstance.AngularRotationOffset.RotateVector(FVector(1,0,0));
 	const FVector RotatedY = ConstraintInstance.AngularRotationOffset.RotateVector(FVector(0,1,0));
-	const FVector WPri2 = ComponentToWorld.TransformVector(RotatedX);
-	const FVector WOrth2 = ComponentToWorld.TransformVector(RotatedY);
+	const FVector WPri2 = ComponentToWorld.TransformVectorNoScale(RotatedX);
+	const FVector WOrth2 = ComponentToWorld.TransformVectorNoScale(RotatedY);
+	
+	
 	ConstraintInstance.Pos2 = A2Transform.InverseTransformPosition(WPos);
 	ConstraintInstance.PriAxis2 = A2Transform.InverseTransformVectorNoScale(WPri2);
 	ConstraintInstance.SecAxis2 = A2Transform.InverseTransformVectorNoScale(WOrth2);
+
+	//Constraint instance is given our reference frame scale and uses it to scale position.
+	//Note that the scale passed in is also used for limits, so we first undo the position scale so that it's consistent.
+
+	//Note that in the case where there is no body instance, the position is given in world space and there is no scaling.
+	const float RefScale = FMath::Max(GetConstraintScale(), 0.01f);
+	if(GetBodyInstance(EConstraintFrame::Frame1))
+	{
+		ConstraintInstance.Pos1 /= RefScale;
+	}
+
+	if (GetBodyInstance(EConstraintFrame::Frame2))
+	{
+		ConstraintInstance.Pos2 /= RefScale;
+	}
 }
 
 void UPhysicsConstraintComponent::SetConstraintReferenceFrame(EConstraintFrame::Type Frame, const FTransform& RefFrame)
diff --git a/Engine/Source/Runtime/Engine/Private/PhysicsEngine/PhysicsSettings.cpp b/Engine/Source/Runtime/Engine/Private/PhysicsEngine/PhysicsSettings.cpp
--- a/Engine/Source/Runtime/Engine/Private/PhysicsEngine/PhysicsSettings.cpp
+++ b/Engine/Source/Runtime/Engine/Private/PhysicsEngine/PhysicsSettings.cpp
@@ -23,6 +23,7 @@ UPhysicsSettings::UPhysicsSettings(const FObjectInitializer& ObjectInitializer)
 	, bSimulateSkeletalMeshOnDedicatedServer(true)
 	, bDefaultHasComplexCollision(true)
 	, bSuppressFaceRemapTable(false)
+	, bDisableActiveTransforms(false)
 	, MaxPhysicsDeltaTime(1.f / 30.f)
 	, bSubstepping(false)
 	, bSubsteppingAsync(false)
diff --git a/Engine/Source/Runtime/Engine/Private/PhysicsEngine2D/Box2DIntegration.cpp b/Engine/Source/Runtime/Engine/Private/PhysicsEngine2D/Box2DIntegration.cpp
--- a/Engine/Source/Runtime/Engine/Private/PhysicsEngine2D/Box2DIntegration.cpp
+++ b/Engine/Source/Runtime/Engine/Private/PhysicsEngine2D/Box2DIntegration.cpp
@@ -135,7 +135,7 @@ bool FPhysicsScene2D::ShouldCollide(b2Fixture* FixtureA, b2Fixture* FixtureB)
 {
 	//@TODO: BOX2D: May need to extend Box2D to be able to indicate that we don't want them to collide now but keep tracking the interaction, not sure!
 	const bool KillInteraction = false;
-	const bool SupressInteraction = false;
+	const bool SuppressInteraction = false;
 	const bool NormalInteraction = true;
 
 	const b2Filter& FilterDataA = FixtureA->GetFilterData();
@@ -160,7 +160,7 @@ bool FPhysicsScene2D::ShouldCollide(b2Fixture* FixtureA, b2Fixture* FixtureB)
 	if ((bIsKinematicA || bIsKinematicB) && (bIsStaticA || bIsStaticB))
 	{
 		// should return eSUPPRESS here instead eKILL so that kinematics vs statics will still be considered once kinematics become dynamic (dying ragdoll case)
-		return SupressInteraction;
+		return SuppressInteraction;
 	}
 
 	//@TODO: BOX2D: Skeletal collision filtering
@@ -188,7 +188,7 @@ bool FPhysicsScene2D::ShouldCollide(b2Fixture* FixtureA, b2Fixture* FixtureB)
 	const bool bDoesWantToBlock = BlockFlagToB && BlockFlagToA;
 	if (!bDoesWantToBlock)
 	{
-		return SupressInteraction;
+		return SuppressInteraction;
 	}
 
 	//@TODO: BOX2D: CCD / contact notification flags
diff --git a/Engine/Source/Runtime/Engine/Private/PlayerController.cpp b/Engine/Source/Runtime/Engine/Private/PlayerController.cpp
--- a/Engine/Source/Runtime/Engine/Private/PlayerController.cpp
+++ b/Engine/Source/Runtime/Engine/Private/PlayerController.cpp
@@ -1951,10 +1951,11 @@ bool APlayerController::GetHitResultAtScreenPosition(const FVector2D ScreenPosit
 	return false;
 }
 
+static const FName NAME_ClickableTrace("ClickableTrace");
 
 bool APlayerController::GetHitResultAtScreenPosition(const FVector2D ScreenPosition, const ECollisionChannel TraceChannel, bool bTraceComplex, FHitResult& HitResult) const
 {
-	FCollisionQueryParams CollisionQueryParams( "ClickableTrace", bTraceComplex );
+	FCollisionQueryParams CollisionQueryParams( NAME_ClickableTrace, bTraceComplex );
 	return GetHitResultAtScreenPosition( ScreenPosition, TraceChannel, CollisionQueryParams, HitResult );
 }
 
@@ -4608,9 +4609,9 @@ void APlayerController::BuildHiddenComponentList(const FVector& ViewLocation, TS
 				{
 					HiddenComponents.Add(PrimitiveComponent->ComponentId);
 
-					for (int32 AttachChildrenIndex = 0; AttachChildrenIndex < PrimitiveComponent->AttachChildren.Num(); AttachChildrenIndex++)
+					for (USceneComponent* AttachedChild : PrimitiveComponent->GetAttachChildren())
 					{						
-						UPrimitiveComponent* AttachChildPC = Cast<UPrimitiveComponent>(PrimitiveComponent->AttachChildren[AttachChildrenIndex]);
+						UPrimitiveComponent* AttachChildPC = Cast<UPrimitiveComponent>(AttachedChild);
 						if (AttachChildPC && AttachChildPC->IsRegistered())
 						{
 							HiddenComponents.Add(AttachChildPC->ComponentId);
diff --git a/Engine/Source/Runtime/Engine/Private/PrimitiveComponentPhysics.cpp b/Engine/Source/Runtime/Engine/Private/PrimitiveComponentPhysics.cpp
--- a/Engine/Source/Runtime/Engine/Private/PrimitiveComponentPhysics.cpp
+++ b/Engine/Source/Runtime/Engine/Private/PrimitiveComponentPhysics.cpp
@@ -496,8 +496,7 @@ void UPrimitiveComponent::SetMassOverrideInKg(FName BoneName, float MassInKg, bo
 	if (FBodyInstance* BI = GetBodyInstance(BoneName))
 	{
 		WarnInvalidPhysicsOperations(LOCTEXT("SetCenterOfMass", "SetCenterOfMass"), nullptr);
-		BI->bOverrideMass = bOverrideMass;
-		BI->MassInKg = MassInKg;
+		BI->SetMassOverride(MassInKg);
 		BI->UpdateMassProperties();
 	}
 }
@@ -536,7 +535,7 @@ float UPrimitiveComponent::CalculateMass(FName)
 {
 	if (BodyInstance.bOverrideMass)
 	{
-		return BodyInstance.MassInKg;
+		return BodyInstance.GetMassOverride();
 	}
 
 	if (BodyInstance.BodySetup.IsValid())
@@ -646,7 +645,7 @@ void UPrimitiveComponent::SyncComponentToRBPhysics()
 UPrimitiveComponent * GetRootWelded(const UPrimitiveComponent* PrimComponent, FName ParentSocketName = NAME_None, FName* OutSocketName = NULL, bool bAboutToWeld = false)
 {
 	UPrimitiveComponent * Result = NULL;
-	UPrimitiveComponent * RootComponent = Cast<UPrimitiveComponent>(PrimComponent->AttachParent);	//we must find the root component along hierarchy that has bWelded set to true
+	UPrimitiveComponent * RootComponent = Cast<UPrimitiveComponent>(PrimComponent->GetAttachParent());	//we must find the root component along hierarchy that has bWelded set to true
 
 	//check that body itself is welded
 	if (FBodyInstance* BI = PrimComponent->GetBodyInstance(ParentSocketName, false))
@@ -660,11 +659,11 @@ UPrimitiveComponent * GetRootWelded(const UPrimitiveComponent* PrimComponent, FN
 	FName PrevSocketName = ParentSocketName;
 	FName SocketName = NAME_None; //because of skeletal mesh it's important that we check along the bones that we attached
 	FBodyInstance* RootBI = NULL;
-	for (; RootComponent; RootComponent = Cast<UPrimitiveComponent>(RootComponent->AttachParent))
+	for (; RootComponent; RootComponent = Cast<UPrimitiveComponent>(RootComponent->GetAttachParent()))
 	{
 		Result = RootComponent;
 		SocketName = PrevSocketName;
-		PrevSocketName = RootComponent->AttachSocketName;
+		PrevSocketName = RootComponent->GetAttachSocketName();
 
 		RootBI = RootComponent->GetBodyInstance(SocketName, false);
 		if (RootBI && RootBI->bWelded)
@@ -688,7 +687,7 @@ void UPrimitiveComponent::GetWeldedBodies(TArray<FBodyInstance*> & OutWeldedBodi
 	OutWeldedBodies.Add(&BodyInstance);
 	OutLabels.Add(NAME_None);
 
-	for (USceneComponent * Child : AttachChildren)
+	for (USceneComponent * Child : GetAttachChildren())
 	{
 		if (UPrimitiveComponent * PrimChild = Cast<UPrimitiveComponent>(Child))
 		{
@@ -708,7 +707,7 @@ bool UPrimitiveComponent::WeldToImplementation(USceneComponent * InParent, FName
 	SCOPE_CYCLE_COUNTER(STAT_WeldPhysics);
 
 	//WeldToInternal assumes attachment is already done
-	if (AttachParent != InParent || AttachSocketName != ParentSocketName)
+	if (GetAttachParent() != InParent || GetAttachSocketName() != ParentSocketName)
 	{
 		return false;
 	}
@@ -761,7 +760,7 @@ bool UPrimitiveComponent::WeldToImplementation(USceneComponent * InParent, FName
 			}
 
 			//root is simulated so we actually weld the body
-			FTransform RelativeTM = RootComponent == AttachParent ? GetRelativeTransform() : GetComponentToWorld().GetRelativeTransform(RootComponent->GetComponentToWorld());	//if direct parent we already have relative. Otherwise compute it
+			FTransform RelativeTM = RootComponent == GetAttachParent() ? GetRelativeTransform() : GetComponentToWorld().GetRelativeTransform(RootComponent->GetComponentToWorld());	//if direct parent we already have relative. Otherwise compute it
 			RootBI->Weld(BI, GetComponentToWorld());
 
 			return true;
@@ -774,7 +773,7 @@ bool UPrimitiveComponent::WeldToImplementation(USceneComponent * InParent, FName
 void UPrimitiveComponent::WeldTo(USceneComponent* InParent, FName InSocketName /* = NAME_None */)
 {
 	//automatically attach if needed
-	if (AttachParent != InParent || AttachSocketName != InSocketName)
+	if (GetAttachParent() != InParent || GetAttachSocketName() != InSocketName)
 	{
 		AttachTo(InParent, InSocketName, EAttachLocation::KeepWorldPosition);
 	}
@@ -794,7 +793,7 @@ void UPrimitiveComponent::UnWeldFromParent()
 	}
 
 	FName SocketName;
-	UPrimitiveComponent * RootComponent = GetRootWelded(this, AttachSocketName, &SocketName);
+	UPrimitiveComponent * RootComponent = GetRootWelded(this, GetAttachSocketName(), &SocketName);
 
 	if (RootComponent)
 	{
@@ -852,7 +851,7 @@ void UPrimitiveComponent::UnWeldFromParent()
 
 void UPrimitiveComponent::UnWeldChildren()
 {
-	for (USceneComponent* ChildComponent : AttachChildren)
+	for (USceneComponent* ChildComponent : GetAttachChildren())
 	{
 		if (UPrimitiveComponent* PrimComp = Cast<UPrimitiveComponent>(ChildComponent))
 		{
@@ -866,7 +865,7 @@ FBodyInstance* UPrimitiveComponent::GetBodyInstance(FName BoneName, bool bGetWel
 	if (bGetWelded && BodyInstance.bWelded)
 	{
 		FName OutSocket;
-		if (UPrimitiveComponent * RootComponentWelded = GetRootWelded(this, AttachSocketName, &OutSocket))
+		if (UPrimitiveComponent * RootComponentWelded = GetRootWelded(this, GetAttachSocketName(), &OutSocket))
 		{
 			if (FBodyInstance* BI = RootComponentWelded->GetBodyInstance(OutSocket, bGetWelded))
 			{
diff --git a/Engine/Source/Runtime/Engine/Private/PrimitiveDrawingUtils.cpp b/Engine/Source/Runtime/Engine/Private/PrimitiveDrawingUtils.cpp
--- a/Engine/Source/Runtime/Engine/Private/PrimitiveDrawingUtils.cpp
+++ b/Engine/Source/Runtime/Engine/Private/PrimitiveDrawingUtils.cpp
@@ -1287,7 +1287,8 @@ bool IsRichView(const FSceneViewFamily& ViewFamily)
 		ViewFamily.EngineShowFlags.LightInfluences ||
 		ViewFamily.EngineShowFlags.Wireframe ||
 		ViewFamily.EngineShowFlags.LevelColoration ||
-		ViewFamily.EngineShowFlags.LODColoration)
+		ViewFamily.EngineShowFlags.LODColoration ||
+		ViewFamily.EngineShowFlags.HLODColoration)
 	{
 		return true;
 	}
@@ -1369,6 +1370,24 @@ void ApplyViewModeOverrides(
 			Collector.RegisterOneFrameMaterialProxy(LODColorationMaterialInstance);
 		}
 	}
+	else if (EngineShowFlags.HLODColoration)
+	{
+		if (!Mesh.IsTranslucent(FeatureLevel) && GEngine->HLODColorationColors.Num() > 0)
+		{
+			int32 hlodColorationIndex = FMath::Clamp((int32)Mesh.VisualizeHLODIndex, 0, GEngine->HLODColorationColors.Num() - 1);
+
+			bool bLit = Mesh.MaterialRenderProxy->GetMaterial(FeatureLevel)->GetShadingModel() != MSM_Unlit;
+			const UMaterial* HLODColorationMaterial = (bLit && EngineShowFlags.Lighting) ? GEngine->LevelColorationLitMaterial : GEngine->LevelColorationUnlitMaterial;
+
+			auto HLODColorationMaterialInstance = new FColoredMaterialRenderProxy(
+				HLODColorationMaterial->GetRenderProxy(Mesh.MaterialRenderProxy->IsSelected(), Mesh.MaterialRenderProxy->IsHovered()),
+				GetSelectionColor(GEngine->HLODColorationColors[hlodColorationIndex], bSelected, Mesh.MaterialRenderProxy->IsHovered())
+				);
+
+			Mesh.MaterialRenderProxy = HLODColorationMaterialInstance;
+			Collector.RegisterOneFrameMaterialProxy(HLODColorationMaterialInstance);
+		}
+	}
 	else if (!EngineShowFlags.Materials)
 	{
 		// Don't render unlit translucency when in 'lighting only' viewmode.
diff --git a/Engine/Source/Runtime/Engine/Private/ShowFlags.cpp b/Engine/Source/Runtime/Engine/Private/ShowFlags.cpp
--- a/Engine/Source/Runtime/Engine/Private/ShowFlags.cpp
+++ b/Engine/Source/Runtime/Engine/Private/ShowFlags.cpp
@@ -283,6 +283,7 @@ void ApplyViewMode(EViewModeIndex ViewModeIndex, bool bPerspective, FEngineShowF
 			bPostProcessing = false;
 			break;
 		case VMI_LODColoration:
+		case VMI_HLODColoration:
 			bPostProcessing = true;
 			break;
 	}
@@ -315,6 +316,7 @@ void ApplyViewMode(EViewModeIndex ViewModeIndex, bool bPerspective, FEngineShowF
 	EngineShowFlags.SetCollisionVisibility(ViewModeIndex == VMI_CollisionVisibility);
 	EngineShowFlags.SetVertexDensities(ViewModeIndex == VMI_VertexDensities);
 	EngineShowFlags.SetLODColoration(ViewModeIndex == VMI_LODColoration);
+	EngineShowFlags.SetHLODColoration(ViewModeIndex == VMI_HLODColoration);
 }
 
 void EngineShowFlagOverride(EShowFlagInitMode ShowFlagInitMode, EViewModeIndex ViewModeIndex, FEngineShowFlags& EngineShowFlags, FName CurrentBufferVisualizationMode, bool bIsSplitScreen)
@@ -442,7 +444,7 @@ void EngineShowFlagOverride(EShowFlagInitMode ShowFlagInitMode, EViewModeIndex V
 			EngineShowFlags.AtmosphericFog = 0;
 		}
 
-		if( ViewModeIndex == VMI_LODColoration )
+		if (ViewModeIndex == VMI_LODColoration || ViewModeIndex == VMI_HLODColoration)
 		{
 			EngineShowFlags.SetLighting(true);	// Best currently otherwise the image becomes hard to read.
 			EngineShowFlags.Fog = 0;			// Removed fog to improve color readability.
@@ -616,6 +618,10 @@ EViewModeIndex FindViewMode(const FEngineShowFlags& EngineShowFlags)
 	{
 		return VMI_LODColoration;
 	}
+	else if (EngineShowFlags.HLODColoration)
+	{
+		return VMI_HLODColoration;
+	}
 
 	return EngineShowFlags.Lighting ? VMI_Lit : VMI_Unlit;
 }
@@ -645,6 +651,7 @@ const TCHAR* GetViewModeName(EViewModeIndex ViewModeIndex)
 		case VMI_CollisionVisibility:		return TEXT("CollisionVis");
 		case VMI_VertexDensities:			return TEXT("VertexDensity");
 		case VMI_LODColoration:				return TEXT("LODColoration");
+		case VMI_HLODColoration:				return TEXT("HLODColoration");
 	}
 	return TEXT("");
 }
diff --git a/Engine/Source/Runtime/Engine/Private/SimpleConstructionScript.cpp b/Engine/Source/Runtime/Engine/Private/SimpleConstructionScript.cpp
--- a/Engine/Source/Runtime/Engine/Private/SimpleConstructionScript.cpp
+++ b/Engine/Source/Runtime/Engine/Private/SimpleConstructionScript.cpp
@@ -30,8 +30,8 @@ namespace
 			if(InstancedSceneComponent != nullptr)
 			{
 				// Make a copy of the array here in case something alters the AttachChildren array during registration (e.g. Physics)
-				TArray<USceneComponent*> AttachChildren = InstancedSceneComponent->AttachChildren;
-				for(auto InstancedChildComponent : AttachChildren)
+				TArray<USceneComponent*> AttachChildren = InstancedSceneComponent->GetAttachChildren();
+				for (USceneComponent* InstancedChildComponent : AttachChildren)
 				{
 					RegisterInstancedComponent(InstancedChildComponent);
 				}
diff --git a/Engine/Source/Runtime/Engine/Private/SkeletalMeshComponentPhysics.cpp b/Engine/Source/Runtime/Engine/Private/SkeletalMeshComponentPhysics.cpp
--- a/Engine/Source/Runtime/Engine/Private/SkeletalMeshComponentPhysics.cpp
+++ b/Engine/Source/Runtime/Engine/Private/SkeletalMeshComponentPhysics.cpp
@@ -798,7 +798,7 @@ void USkeletalMeshComponent::InitArticulated(FPhysScene* PhysScene)
 	}
 
 	FVector Scale3D = ComponentToWorld.GetScale3D();
-	float Scale = Scale3D.X;
+	const float Scale = Scale3D.GetAbsMin();
 
 	// Find root physics body
 	int32 RootBodyIndex = INDEX_NONE;
@@ -860,6 +860,8 @@ void USkeletalMeshComponent::InitArticulated(FPhysScene* PhysScene)
 			BodyInst->InstanceBodyIndex = i; // Set body index 
 			BodyInst->InstanceBoneIndex = BoneIndex; // Set bone index
 
+			BodyInst->bStartAwake = BodyInstance.bStartAwake;	//We don't allow customization here. Just use whatever the component is set to
+
 			if (i == RootBodyIndex)
 			{
 				BodyInst->DOFMode = BodyInstance.DOFMode;
@@ -906,7 +908,6 @@ void USkeletalMeshComponent::InitArticulated(FPhysScene* PhysScene)
 	SetRootBodyIndex(RootBodyIndex);
 
 #if WITH_PHYSX
-
 	if (PhysScene)
 	{
 		// Get the scene type from the SkeletalMeshComponent's BodyInstance
@@ -917,23 +918,6 @@ void USkeletalMeshComponent::InitArticulated(FPhysScene* PhysScene)
 		if (Aggregate && Aggregate->getNbActors() > 0)
 		{
 			PScene->addAggregate(*Aggregate);
-
-			// If we've used an aggregate, InitBody would not be able to set awake status as we *must* have a scene
-			// to do that, so we reconcile this here.
-			AActor* Owner = GetOwner();
-			bool bShouldSleep = !BodyInstance.bStartAwake && (Owner && Owner->GetVelocity().SizeSquared() <= KINDA_SMALL_NUMBER);
-
-			for (FBodyInstance* Body : Bodies)
-			{
-				// Creates a DOF constraint if necessary for the body - also requires the scene to exist within the actor
-				Body->CreateDOFLock();
-
-				// Set to sleep if necessary
-				if (bShouldSleep  && Body->GetPxRigidDynamic_AssumesLocked())
-				{
-					Body->GetPxRigidDynamic_AssumesLocked()->putToSleep();
-				}
-			}
 		}
 	}
 
