Copying //UE4/Dev-Framework to Dev-Main (//UE4/Dev-Main) @ 2855699

diff --git a/Engine/Source/Runtime/Engine/Private/SkeletalRenderGPUSkin.cpp b/Engine/Source/Runtime/Engine/Private/SkeletalRenderGPUSkin.cpp
--- a/Engine/Source/Runtime/Engine/Private/SkeletalRenderGPUSkin.cpp
+++ b/Engine/Source/Runtime/Engine/Private/SkeletalRenderGPUSkin.cpp
@@ -67,14 +67,7 @@ void FMorphVertexBuffer::InitDynamicRHI()
 
 	// Lock the buffer.
 	FMorphGPUSkinVertex* Buffer = (FMorphGPUSkinVertex*)BufferData;
-
-	// zero all deltas (NOTE: DeltaTangentZ is FPackedNormal, so we can't just FMemory::Memzero)
-	for (uint32 VertIndex=0; VertIndex < LodModel.NumVertices; ++VertIndex)
-	{
-		Buffer[VertIndex].DeltaPosition = FVector::ZeroVector;
-		Buffer[VertIndex].DeltaTangentZ = FPackedNormal::ZeroNormal;
-	}
-
+	FMemory::Memzero(&Buffer[0], sizeof(FMorphGPUSkinVertex)*LodModel.NumVertices);
 	// Unlock the buffer.
 	RHIUnlockVertexBuffer(VertexBufferRHI);
 	
@@ -411,15 +404,12 @@ void FSkeletalMeshObjectGPUSkin::FSkeletalMeshObjectLOD::UpdateMorphVertexBuffer
 			SCOPE_CYCLE_COUNTER(STAT_MorphVertexBuffer_Init);
 
 			// zero everything
-			int32 vertsToAdd = static_cast<int32>(LodModel.NumVertices) - MorphDeltaTangentZAccumulationArray.Num();
+			int32 vertsToAdd = static_cast<int32>(LodModel.NumVertices) - MorphAccumulatedWeightArray.Num();
 			if(vertsToAdd > 0) 
 			{
-				// we're memzero-ing afterwards anyway, so add uninitalized
-				MorphDeltaTangentZAccumulationArray.AddUninitialized(vertsToAdd);
 				MorphAccumulatedWeightArray.AddUninitialized(vertsToAdd);
 			}
 
-			FMemory::Memzero(MorphDeltaTangentZAccumulationArray.GetData(), sizeof(FVector)*LodModel.NumVertices);
 			FMemory::Memzero(MorphAccumulatedWeightArray.GetData(), sizeof(float)*LodModel.NumVertices);
 
 			// PackedNormals will be wrong init with 0, but they'll be overwritten later
@@ -457,8 +447,7 @@ void FSkeletalMeshObjectGPUSkin::FSkeletalMeshObjectLOD::UpdateMorphVertexBuffer
 						FMorphGPUSkinVertex& DestVertex = Buffer[MorphVertex.SourceIdx];
 
 						DestVertex.DeltaPosition += MorphVertex.PositionDelta * VertAnim.Weight;
-						// add to accumulated tangent Z
-						MorphDeltaTangentZAccumulationArray[MorphVertex.SourceIdx] += MorphVertex.TangentZDelta * VertAnim.Weight;
+						DestVertex.DeltaTangentZ += MorphVertex.TangentZDelta * VertAnim.Weight;
 						// accumulate the weight so we can normalized it later
 						MorphAccumulatedWeightArray[MorphVertex.SourceIdx] += VertAnimAbsWeight;
 					}
@@ -473,20 +462,13 @@ void FSkeletalMeshObjectGPUSkin::FSkeletalMeshObjectLOD::UpdateMorphVertexBuffer
 				FMorphGPUSkinVertex& DestVertex = Buffer[iVertex];
 				float AccumulatedWeight = MorphAccumulatedWeightArray[iVertex];
 
-				if (AccumulatedWeight > MinVertexAnimBlendWeight)
-				{
-					// when copy back, make sure to normalize by accumulated weight
-					// since delta diff of normal is (-2, 2), we divide by 2 to packed into packed normal
-					// when we unpack, we'll apply *2. 
-					DestVertex.DeltaTangentZ = (MorphDeltaTangentZAccumulationArray[iVertex] / AccumulatedWeight)/2;
-					// we now add W as how much alpha of DeltaTangentZ we're apply to the original tangent
-					DestVertex.DeltaTangentZ.Vector.W = FMath::Min(1.0f, AccumulatedWeight) * 255.9999f;
-				}
-				else
+				// if accumualated weight is >1.f
+				// previous code was applying the weight again in GPU if less than 1, but it doesn't make sense to do so
+				// so instead, we just divide by AccumulatedWeight if it's more than 1.
+				// now DeltaTangentZ isn't FPackedNormal, so you can apply any value to it. 
+				if (AccumulatedWeight > 1.f)
 				{
-					DestVertex.DeltaTangentZ = FPackedNormal::ZeroNormal;
-					// we now add W as how much alpha of DeltaTangentZ we're apply to the original tangent					
-					DestVertex.DeltaTangentZ.Vector.W = 0;
+					DestVertex.DeltaTangentZ /= AccumulatedWeight;
 				}
 			}
 		} // ApplyDelta
@@ -626,7 +608,7 @@ void InitMorphVertexFactoryComponents(typename VertexFactoryType::FDataType* Ver
 		VertexBuffers.MorphVertexBuffer,STRUCT_OFFSET(FMorphGPUSkinVertex,DeltaPosition),sizeof(FMorphGPUSkinVertex),VET_Float3);
 	// delta normals
 	VertexFactoryData->DeltaTangentZComponent = FVertexStreamComponent(
-		VertexBuffers.MorphVertexBuffer,STRUCT_OFFSET(FMorphGPUSkinVertex,DeltaTangentZ),sizeof(FMorphGPUSkinVertex),VET_PackedNormal);
+		VertexBuffers.MorphVertexBuffer, STRUCT_OFFSET(FMorphGPUSkinVertex, DeltaTangentZ), sizeof(FMorphGPUSkinVertex), VET_Float3);
 }
 
 /** 
diff --git a/Engine/Source/Runtime/Engine/Private/SkeletalRenderGPUSkin.h b/Engine/Source/Runtime/Engine/Private/SkeletalRenderGPUSkin.h
--- a/Engine/Source/Runtime/Engine/Private/SkeletalRenderGPUSkin.h
+++ b/Engine/Source/Runtime/Engine/Private/SkeletalRenderGPUSkin.h
@@ -104,12 +104,12 @@ public:
 struct FMorphGPUSkinVertex
 {
 	FVector			DeltaPosition;
-	FPackedNormal	DeltaTangentZ;
+	FVector			DeltaTangentZ;
 
 	FMorphGPUSkinVertex() {};
 	
 	/** Construct for special case **/
-	FMorphGPUSkinVertex(const FVector& InDeltaPosition, const FPackedNormal& InDeltaTangentZ)
+	FMorphGPUSkinVertex(const FVector& InDeltaPosition, const FVector& InDeltaTangentZ)
 	{
 		DeltaPosition = InDeltaPosition;
 		DeltaTangentZ = InDeltaTangentZ;
diff --git a/Engine/Source/Runtime/Engine/Private/Slate/SceneViewport.cpp b/Engine/Source/Runtime/Engine/Private/Slate/SceneViewport.cpp
--- a/Engine/Source/Runtime/Engine/Private/Slate/SceneViewport.cpp
+++ b/Engine/Source/Runtime/Engine/Private/Slate/SceneViewport.cpp
@@ -740,19 +740,26 @@ FReply FSceneViewport::OnKeyDown( const FGeometry& InGeometry, const FKeyEvent&
 	CurrentReplyState = FReply::Handled(); 
 
 	FKey Key = InKeyEvent.GetKey();
-	KeyStateMap.Add( Key, true );
-
-	//@todo Slate Viewports: FWindowsViewport checks for Alt+Enter or F11 and toggles fullscreen.  Unknown if fullscreen via this method will be needed for slate viewports. 
-	if( ViewportClient && GetSizeXY() != FIntPoint::ZeroValue  )
+	if (Key.IsValid())
 	{
-		// Switch to the viewport clients world before processing input
-		FScopedConditionalWorldSwitcher WorldSwitcher( ViewportClient );
+		KeyStateMap.Add(Key, true);
 
-		if (!ViewportClient->InputKey(this, InKeyEvent.GetUserIndex(), Key, InKeyEvent.IsRepeat() ? IE_Repeat : IE_Pressed, 1.0f, Key.IsGamepadKey()))
+		//@todo Slate Viewports: FWindowsViewport checks for Alt+Enter or F11 and toggles fullscreen.  Unknown if fullscreen via this method will be needed for slate viewports. 
+		if (ViewportClient && GetSizeXY() != FIntPoint::ZeroValue)
 		{
-			CurrentReplyState = FReply::Unhandled();
+			// Switch to the viewport clients world before processing input
+			FScopedConditionalWorldSwitcher WorldSwitcher(ViewportClient);
+
+			if (!ViewportClient->InputKey(this, InKeyEvent.GetUserIndex(), Key, InKeyEvent.IsRepeat() ? IE_Repeat : IE_Pressed, 1.0f, Key.IsGamepadKey()))
+			{
+				CurrentReplyState = FReply::Unhandled();
+			}
 		}
 	}
+	else
+	{
+		CurrentReplyState = FReply::Unhandled();
+	}
 	return CurrentReplyState;
 }
 
@@ -762,18 +769,25 @@ FReply FSceneViewport::OnKeyUp( const FGeometry& InGeometry, const FKeyEvent& In
 	CurrentReplyState = FReply::Handled(); 
 
 	FKey Key = InKeyEvent.GetKey();
-	KeyStateMap.Add( Key, false );
-	
-	if( ViewportClient && GetSizeXY() != FIntPoint::ZeroValue  )
+	if (Key.IsValid())
 	{
-		// Switch to the viewport clients world before processing input
-		FScopedConditionalWorldSwitcher WorldSwitcher( ViewportClient );
+		KeyStateMap.Add(Key, false);
 
-		if (!ViewportClient->InputKey(this, InKeyEvent.GetUserIndex(), Key, IE_Released, 1.0f, Key.IsGamepadKey()))
+		if (ViewportClient && GetSizeXY() != FIntPoint::ZeroValue)
 		{
-			CurrentReplyState = FReply::Unhandled();
+			// Switch to the viewport clients world before processing input
+			FScopedConditionalWorldSwitcher WorldSwitcher(ViewportClient);
+
+			if (!ViewportClient->InputKey(this, InKeyEvent.GetUserIndex(), Key, IE_Released, 1.0f, Key.IsGamepadKey()))
+			{
+				CurrentReplyState = FReply::Unhandled();
+			}
 		}
 	}
+	else
+	{
+		CurrentReplyState = FReply::Unhandled();
+	}
 
 	return CurrentReplyState;
 }
@@ -784,18 +798,25 @@ FReply FSceneViewport::OnAnalogValueChanged(const FGeometry& MyGeometry, const F
 	CurrentReplyState = FReply::Handled();
 
 	FKey Key = InAnalogInputEvent.GetKey();
-	KeyStateMap.Add(Key, true);
-
-	if (ViewportClient)
+	if (Key.IsValid())
 	{
-		// Switch to the viewport clients world before processing input
-		FScopedConditionalWorldSwitcher WorldSwitcher(ViewportClient);
+		KeyStateMap.Add(Key, true);
 
-		if (!ViewportClient->InputAxis(this, InAnalogInputEvent.GetUserIndex(), Key, Key == EKeys::Gamepad_RightY ? -InAnalogInputEvent.GetAnalogValue() : InAnalogInputEvent.GetAnalogValue(), FApp::GetDeltaTime(), 1, Key.IsGamepadKey()))
+		if (ViewportClient)
 		{
-			CurrentReplyState = FReply::Unhandled();
+			// Switch to the viewport clients world before processing input
+			FScopedConditionalWorldSwitcher WorldSwitcher(ViewportClient);
+
+			if (!ViewportClient->InputAxis(this, InAnalogInputEvent.GetUserIndex(), Key, Key == EKeys::Gamepad_RightY ? -InAnalogInputEvent.GetAnalogValue() : InAnalogInputEvent.GetAnalogValue(), FApp::GetDeltaTime(), 1, Key.IsGamepadKey()))
+			{
+				CurrentReplyState = FReply::Unhandled();
+			}
 		}
 	}
+	else
+	{
+		CurrentReplyState = FReply::Unhandled();
+	}
 
 	return CurrentReplyState;
 }
diff --git a/Engine/Source/Runtime/Engine/Private/SoundBase.cpp b/Engine/Source/Runtime/Engine/Private/SoundBase.cpp
--- a/Engine/Source/Runtime/Engine/Private/SoundBase.cpp
+++ b/Engine/Source/Runtime/Engine/Private/SoundBase.cpp
@@ -5,6 +5,9 @@
 #include "Sound/AudioSettings.h"
 #include "AudioDevice.h"
 
+USoundClass* USoundBase::DefaultSoundClassObject = nullptr;
+USoundConcurrency* USoundBase::DefaultSoundConcurrencyObject = nullptr;
+
 USoundBase::USoundBase(const FObjectInitializer& ObjectInitializer)
 	: Super(ObjectInitializer)
 	, bIgnoreFocus(false)
@@ -17,17 +20,26 @@ void USoundBase::PostInitProperties()
 {
 	Super::PostInitProperties();
 
-	const FStringAssetReference DefaultSoundClassName = GetDefault<UAudioSettings>()->DefaultSoundClassName;
-	if (DefaultSoundClassName.IsValid())
+	if (USoundBase::DefaultSoundClassObject == nullptr)
 	{
-		SoundClassObject = LoadObject<USoundClass>(nullptr, *DefaultSoundClassName.ToString());
+		const FStringAssetReference DefaultSoundClassName = GetDefault<UAudioSettings>()->DefaultSoundClassName;
+		if (DefaultSoundClassName.IsValid())
+		{
+			USoundBase::DefaultSoundClassObject = LoadObject<USoundClass>(nullptr, *DefaultSoundClassName.ToString());
+		}
 	}
+	SoundClassObject = USoundBase::DefaultSoundClassObject;
 
-	const FStringAssetReference DefaultSoundConcurrencyName = GetDefault<UAudioSettings>()->DefaultSoundConcurrencyName;
-	if (DefaultSoundConcurrencyName.IsValid())
+	if (USoundBase::DefaultSoundConcurrencyObject == nullptr)
 	{
-		SoundConcurrencySettings = LoadObject<USoundConcurrency>(nullptr, *DefaultSoundConcurrencyName.ToString());
+		const FStringAssetReference DefaultSoundConcurrencyName = GetDefault<UAudioSettings>()->DefaultSoundConcurrencyName;
+		if (DefaultSoundConcurrencyName.IsValid())
+		{
+			USoundBase::DefaultSoundConcurrencyObject = LoadObject<USoundConcurrency>(nullptr, *DefaultSoundConcurrencyName.ToString());
+		}
 	}
+	SoundConcurrencySettings = USoundBase::DefaultSoundConcurrencyObject;
+
 }
 
 bool USoundBase::IsPlayable() const
@@ -90,7 +102,7 @@ const FSoundConcurrencySettings* USoundBase::GetSoundConcurrencySettingsToApply(
 
 float USoundBase::GetPriority() const
 {
-	return Priority;
+	return FMath::Clamp(Priority, MIN_SOUND_PRIORITY, MAX_SOUND_PRIORITY);
 }
 
 uint32 USoundBase::GetSoundConcurrencyObjectID() const
diff --git a/Engine/Source/Runtime/Engine/Private/SoundNodeLooping.cpp b/Engine/Source/Runtime/Engine/Private/SoundNodeLooping.cpp
--- a/Engine/Source/Runtime/Engine/Private/SoundNodeLooping.cpp
+++ b/Engine/Source/Runtime/Engine/Private/SoundNodeLooping.cpp
@@ -124,6 +124,29 @@ bool USoundNodeLooping::NotifyWaveInstanceFinished( FWaveInstance* InWaveInstanc
 
 float USoundNodeLooping::GetDuration()
 {
-	return INDEFINITELY_LOOPING_DURATION;
+	// Assume no duration (i.e. no input node)
+	float Duration = 0.0f;
+
+	// If we have any child nodes
+	if (ChildNodes.Num() > 0)
+	{
+		// If we're told to loop indefinitely, then the duration will be "infinite"
+		if (bLoopIndefinitely)
+		{
+			Duration = INDEFINITELY_LOOPING_DURATION;
+		}
+		else
+		{
+			// Looping nodes can only have one child node
+			check(ChildNodes.Num() == 1);
+			if (USoundNode* Child = ChildNodes[0])
+			{
+				// Duration will be the loop count times the child node duration
+				Duration = LoopCount * Child->GetDuration();
+			}
+		}
+	}
+	
+	return Duration;
 }
 
diff --git a/Engine/Source/Runtime/Engine/Private/SoundWave.cpp b/Engine/Source/Runtime/Engine/Private/SoundWave.cpp
--- a/Engine/Source/Runtime/Engine/Private/SoundWave.cpp
+++ b/Engine/Source/Runtime/Engine/Private/SoundWave.cpp
@@ -61,44 +61,34 @@ USoundWave::USoundWave(const FObjectInitializer& ObjectInitializer)
 	Volume = 1.0;
 	Pitch = 1.0;
 	CompressionQuality = 40;
+	SubtitlePriority = 10000.f;
 }
 
 SIZE_T USoundWave::GetResourceSize(EResourceSizeMode::Type Mode)
 {
-	int32 CalculatedResourceSize = 0;
-	
-	FAudioDevice* AudioDevice = nullptr;
-	
-	if (GEngine)
+	if (!GEngine)
 	{
-		AudioDevice = GEngine->GetMainAudioDevice();
+		return 0;
 	}
 
-	if (AudioDevice && AudioDevice->HasCompressedAudioInfoClass(this) && DecompressionType == DTYPE_Native)
-	{
-		// If we've been decompressed, need to account for decompressed and also compressed
-		CalculatedResourceSize += RawPCMDataSize;
-	}
-	else if (DecompressionType == DTYPE_RealTime)
+	if (FAudioDevice* LocalAudioDevice = GEngine->GetMainAudioDevice())
 	{
-		if (CachedRealtimeFirstBuffer)
+		if (LocalAudioDevice->HasCompressedAudioInfoClass(this) && DecompressionType == DTYPE_Native)
 		{
-			CalculatedResourceSize += MONO_PCM_BUFFER_SIZE * NumChannels;
+			check(ResourceSize == 0);
+			return RawPCMDataSize;
+		}
+		else if (DecompressionType == DTYPE_RealTime && CachedRealtimeFirstBuffer)
+		{
+			return MONO_PCM_BUFFER_SIZE * NumChannels;
+		}
+		else if ((!FPlatformProperties::SupportsAudioStreaming() || !IsStreaming()))
+		{
+			return GetCompressedDataSize(LocalAudioDevice->GetRuntimeFormat(this));
 		}
 	}
 
-	// Don't add compressed data to size of streaming sounds
-	if (AudioDevice && (!FPlatformProperties::SupportsAudioStreaming() || !IsStreaming()))
-	{
-		// Only add the compressed data size of the asset from DDC in editor builds otherwise use ResourceSize (which may be 0)
-#if WITH_EDITOR
-		CalculatedResourceSize += GetCompressedDataSize(AudioDevice->GetRuntimeFormat(this));
-#else
-		CalculatedResourceSize += ResourceSize;
-#endif
-	}
-
-	return CalculatedResourceSize;
+	return 0;
 }
 
 int32 USoundWave::GetResourceSizeForFormat(FName Format)
@@ -270,6 +260,11 @@ void USoundWave::LogSubtitle( FOutputDevice& Ar )
 	Ar.Logf( bMature ? TEXT( "Mature:    Yes" ) : TEXT( "Mature:    No" ) );
 }
 
+float USoundWave::GetSubtitlePriority() const
+{ 
+	return SubtitlePriority;
+};
+
 void USoundWave::PostInitProperties()
 {
 	Super::PostInitProperties();
@@ -687,10 +682,7 @@ void USoundWave::Parse( FAudioDevice* AudioDevice, const UPTRINT NodeWaveInstanc
 		}
 
 		// Don't add wave instances that are not going to be played at this point.
-		if( WaveInstance->Volume > KINDA_SMALL_NUMBER )
-		{
-			WaveInstances.Add( WaveInstance );
-		}
+		WaveInstances.Add( WaveInstance );
 
 		// We're still alive.
 		ActiveSound.bFinished = false;
diff --git a/Engine/Source/Runtime/Engine/Private/StaticMesh.cpp b/Engine/Source/Runtime/Engine/Private/StaticMesh.cpp
--- a/Engine/Source/Runtime/Engine/Private/StaticMesh.cpp
+++ b/Engine/Source/Runtime/Engine/Private/StaticMesh.cpp
@@ -14,6 +14,7 @@
 #include "TargetPlatform.h"
 #include "SpeedTreeWind.h"
 #include "DistanceFieldAtlas.h"
+#include "UObject/DevObjectVersion.h"
 
 #if WITH_EDITOR
 #include "RawMesh.h"
@@ -1576,12 +1577,14 @@ void UStaticMesh::GetAssetRegistryTags(TArray<FAssetRegistryTag>& OutTags) const
 	int32 NumTriangles = 0;
 	int32 NumVertices = 0;
 	int32 NumUVChannels = 0;
+	int32 NumLODs = 0;
 	if (RenderData && RenderData->LODResources.Num() > 0)
 	{
 		const FStaticMeshLODResources& LOD = RenderData->LODResources[0];
 		NumTriangles = LOD.IndexBuffer.GetNumIndices() / 3;
 		NumVertices = LOD.VertexBuffer.GetNumVertices();
 		NumUVChannels = LOD.VertexBuffer.GetNumTexCoords();
+		NumLODs = RenderData->LODResources.Num();
 	}
 
 	int32 NumCollisionPrims = 0;
@@ -1602,7 +1605,8 @@ void UStaticMesh::GetAssetRegistryTags(TArray<FAssetRegistryTag>& OutTags) const
 	OutTags.Add( FAssetRegistryTag("UVChannels", FString::FromInt(NumUVChannels), FAssetRegistryTag::TT_Numerical) );
 	OutTags.Add( FAssetRegistryTag("Materials", FString::FromInt(Materials.Num()), FAssetRegistryTag::TT_Numerical) );
 	OutTags.Add( FAssetRegistryTag("ApproxSize", ApproxSizeStr, FAssetRegistryTag::TT_Dimensional) );
-	OutTags.Add( FAssetRegistryTag("CollisionPrims", FString::FromInt(NumCollisionPrims), FAssetRegistryTag::TT_Numerical) );
+	OutTags.Add( FAssetRegistryTag("CollisionPrims", FString::FromInt(NumCollisionPrims), FAssetRegistryTag::TT_Numerical));
+	OutTags.Add( FAssetRegistryTag("LODs", FString::FromInt(NumLODs), FAssetRegistryTag::TT_Numerical));
 
 #if WITH_EDITORONLY_DATA
 	if (AssetImportData)
@@ -1902,6 +1906,13 @@ void UStaticMesh::Serialize(FArchive& Ar)
 	}
 #endif
 
+	Ar.UsingCustomVersion(FFrameworkObjectVersion::GUID);
+
+	if(Ar.IsLoading() && Ar.CustomVer(FFrameworkObjectVersion::GUID) < FFrameworkObjectVersion::UseBodySetupCollisionProfile && BodySetup)
+	{
+		BodySetup->DefaultInstance.SetCollisionProfileName(UCollisionProfile::BlockAll_ProfileName);
+	}
+
 #if WITH_EDITORONLY_DATA
 	if( !StripFlags.IsEditorDataStripped() )
 	{
@@ -2165,7 +2176,11 @@ bool UStaticMesh::GetPhysicsTriMeshData(struct FTriMeshCollisionData* CollisionD
 #if WITH_EDITORONLY_DATA
 	check(HasValidRenderData());
 
-	FStaticMeshLODResources& LOD = RenderData->LODResources[0];
+	// Get the LOD level to use for collision
+	// Always use 0 if asking for 'all tri data'
+	const int32 UseLODIndex = bInUseAllTriData ? 0 : FMath::Clamp(LODForCollision, 0, RenderData->LODResources.Num()-1);
+
+	FStaticMeshLODResources& LOD = RenderData->LODResources[UseLODIndex];
 
 	// Scale all verts into temporary vertex buffer.
 	const uint32 NumVerts = LOD.PositionVertexBuffer.GetNumVertices();
@@ -2213,16 +2228,24 @@ bool UStaticMesh::GetPhysicsTriMeshData(struct FTriMeshCollisionData* CollisionD
 
 bool UStaticMesh::ContainsPhysicsTriMeshData(bool bInUseAllTriData) const 
 {
-	if (RenderData
-		&& RenderData->LODResources.Num() > 0
-		&& RenderData->LODResources[0].PositionVertexBuffer.GetNumVertices() > 0)
+	if(RenderData == nullptr || RenderData->LODResources.Num() == 0)
+	{
+		return false;
+	}
+
+	// Get the LOD level to use for collision
+	// Always use 0 if asking for 'all tri data'
+	const int32 UseLODIndex = bInUseAllTriData ? 0 : FMath::Clamp(LODForCollision, 0, RenderData->LODResources.Num() - 1);
+
+	if (RenderData->LODResources[UseLODIndex].PositionVertexBuffer.GetNumVertices() > 0)
 	{
 		// In non-cooked builds we need to look at the section info map to get
 		// accurate per-section info.
 #if WITH_EDITORONLY_DATA
 		return bInUseAllTriData || SectionInfoMap.AnySectionHasCollision();
 #else
-		FStaticMeshLODResources& LOD = RenderData->LODResources[0];
+		// Get the LOD level to use for collision
+		FStaticMeshLODResources& LOD = RenderData->LODResources[UseLODIndex];
 		for (int32 SectionIndex = 0; SectionIndex < LOD.Sections.Num(); ++SectionIndex)
 		{
 			const FStaticMeshSection& Section = LOD.Sections[SectionIndex];
@@ -2298,6 +2321,7 @@ void UStaticMesh::CreateBodySetup()
 	if (BodySetup==NULL)
 	{
 		BodySetup = NewObject<UBodySetup>(this);
+		BodySetup->DefaultInstance.SetCollisionProfileName(UCollisionProfile::BlockAll_ProfileName);
 	}
 }
 
diff --git a/Engine/Source/Runtime/Engine/Private/StaticMeshActor.cpp b/Engine/Source/Runtime/Engine/Private/StaticMeshActor.cpp
--- a/Engine/Source/Runtime/Engine/Private/StaticMeshActor.cpp
+++ b/Engine/Source/Runtime/Engine/Private/StaticMeshActor.cpp
@@ -10,6 +10,7 @@
 #include "UObjectToken.h"
 #include "MapErrors.h"
 #include "Engine/StaticMeshActor.h"
+#include "UObject/DevObjectVersion.h"
 
 #define LOCTEXT_NAMESPACE "StaticMeshActor"
 
@@ -22,6 +23,7 @@ AStaticMeshActor::AStaticMeshActor(const FObjectInitializer& ObjectInitializer)
 	StaticMeshComponent->SetCollisionProfileName(UCollisionProfile::BlockAll_ProfileName);
 	StaticMeshComponent->Mobility = EComponentMobility::Static;
 	StaticMeshComponent->bGenerateOverlapEvents = false;
+	StaticMeshComponent->bUseDefaultCollision = true;
 
 	RootComponent = StaticMeshComponent;
 }
@@ -120,6 +122,30 @@ bool AStaticMeshActor::GetReferencedContentObjects( TArray<UObject*>& Objects )
 	return true;
 }
 
+void AStaticMeshActor::Serialize(FArchive& Ar)
+{
+	Super::Serialize(Ar);
+	Ar.UsingCustomVersion(FFrameworkObjectVersion::GUID);
+}
+
+void AStaticMeshActor::PostLoad()
+{
+	Super::PostLoad();
+
+	if (StaticMeshComponent && GetLinkerCustomVersion(FFrameworkObjectVersion::GUID) < FFrameworkObjectVersion::UseBodySetupCollisionProfile)
+	{
+		//For all existing content we check whether we need to mark collision profile for override.
+		if (UBodySetup* BodySetup = StaticMeshComponent->GetBodySetup())
+		{
+			if (BodySetup->DefaultInstance.GetCollisionProfileName() != StaticMeshComponent->GetCollisionProfileName())
+			{
+				StaticMeshComponent->bUseDefaultCollision = false;
+			}
+		}
+	}
+	
+}
+
 void AStaticMeshActor::CheckForErrors()
 {
 	Super::CheckForErrors();
diff --git a/Engine/Source/Runtime/Engine/Private/StaticMeshRender.cpp b/Engine/Source/Runtime/Engine/Private/StaticMeshRender.cpp
--- a/Engine/Source/Runtime/Engine/Private/StaticMeshRender.cpp
+++ b/Engine/Source/Runtime/Engine/Private/StaticMeshRender.cpp
@@ -18,6 +18,8 @@
 #include "AI/Navigation/NavCollision.h"
 #include "ComponentRecreateRenderStateContext.h"
 
+#include "Classes/Engine/LODActor.h"
+
 /** If true, optimized depth-only index buffers are used for shadow rendering. */
 static bool GUseShadowIndexBuffer = true;
 
@@ -99,6 +101,8 @@ FStaticMeshSceneProxy::FStaticMeshSceneProxy(UStaticMeshComponent* InComponent):
 #endif
 #if !(UE_BUILD_SHIPPING || UE_BUILD_TEST)
 	, WorldLightmapFactor(0)
+	, LODForCollision(InComponent->StaticMesh->LODForCollision)
+	, bDrawMeshCollisionWireframe(InComponent->bDrawMeshCollisionWireframe)
 #endif
 {
 	check(RenderData);
@@ -171,6 +175,17 @@ FStaticMeshSceneProxy::FStaticMeshSceneProxy(UStaticMeshComponent* InComponent):
 			PropertyColor = TempPropertyColor;
 		}
 	}
+
+	// Setup Hierarchical LOD index
+	if (Owner && Owner->IsA<ALODActor>())
+	{
+		ALODActor* LODActor = Cast<ALODActor>(Owner);
+		HierarchicalLODIndex = LODActor->LODLevel + 1;
+	}
+	else if (InComponent->GetLODParentPrimitive())
+	{		
+		HierarchicalLODIndex = 1;
+	}
 #endif
 
 #if !(UE_BUILD_SHIPPING || UE_BUILD_TEST)
@@ -236,6 +251,7 @@ bool FStaticMeshSceneProxy::GetShadowMeshElement(int32 LODIndex, int32 BatchInde
 	OutMeshBatch.LODIndex = LODIndex;
 #if !(UE_BUILD_SHIPPING || UE_BUILD_TEST)
 	OutMeshBatch.VisualizeLODIndex = LODIndex;
+	OutMeshBatch.VisualizeHLODIndex = HierarchicalLODIndex;
 #endif
 	OutMeshBatch.LCI = &ProxyLODInfo;
 	if (ForcedLodModel > 0) 
@@ -319,6 +335,7 @@ bool FStaticMeshSceneProxy::GetMeshElement(
 		OutMeshBatch.LODIndex = LODIndex;
 #if !(UE_BUILD_SHIPPING || UE_BUILD_TEST)
 		OutMeshBatch.VisualizeLODIndex = LODIndex;
+		OutMeshBatch.VisualizeHLODIndex = HierarchicalLODIndex;
 #endif
 		OutMeshBatch.UseDynamicData = false;
 		OutMeshBatch.ReverseCulling = IsLocalToWorldDeterminantNegative() && !bUseReversedIndices;
@@ -737,10 +754,8 @@ void FStaticMeshSceneProxy::GetDynamicMeshElements(const TArray<const FSceneView
 	bool bDrawSimpleCollision = false, bDrawComplexCollision = false;
 	const bool bInCollisionView = IsCollisionView(EngineShowFlags, bDrawSimpleCollision, bDrawComplexCollision);
 	
-	// Should we draw the mesh wireframe to indicate we are using the mesh as collision
-	const bool bDrawComplexWireframeCollision = (EngineShowFlags.Collision && IsCollisionEnabled() && CollisionTraceFlag == ECollisionTraceFlag::CTF_UseComplexAsSimple);
-
-	const bool bDrawMesh = (bInCollisionView) ? (bDrawComplexCollision) : 
+	// Skip drawing mesh normally if in a collision view, will rely on collision drawing code below
+	const bool bDrawMesh = !bInCollisionView && 
 	(	IsRichView(ViewFamily) || HasViewDependentDPG()
 		|| EngineShowFlags.Collision
 		|| EngineShowFlags.Bounds
@@ -832,8 +847,6 @@ void FStaticMeshSceneProxy::GetDynamicMeshElements(const TArray<const FSceneView
 								
 									if (GetMeshElement(LODIndex, BatchIndex, SectionIndex, SDPG_World, bSectionIsSelected, IsHovered(), true, MeshElement))
 									{
-										bool bAddMesh = true;
-
 	#if !(UE_BUILD_SHIPPING || UE_BUILD_TEST)
 										if (EngineShowFlags.VertexColors && AllowDebugViewmodes())
 										{
@@ -871,43 +884,6 @@ void FStaticMeshSceneProxy::GetDynamicMeshElements(const TArray<const FSceneView
 											Collector.RegisterOneFrameMaterialProxy(VertexColorVisualizationMaterialInstance);
 											MeshElement.MaterialRenderProxy = VertexColorVisualizationMaterialInstance;
 										}
-										else if (AllowDebugViewmodes() && (bDrawComplexWireframeCollision || (bInCollisionView && bDrawComplexCollision)))
-										{
-											if (LODModel.Sections[SectionIndex].bEnableCollision)
-											{
-												UMaterial* MaterialToUse = (bDrawComplexWireframeCollision) ? GEngine->WireframeMaterial : GEngine->ShadedLevelColorationUnlitMaterial;
-
-												if (MaterialToUse)
-												{
-													// Override the mesh's material with our material that draws the collision color
-													auto CollisionMaterialInstance = new FColoredMaterialRenderProxy(
-														MaterialToUse ? MaterialToUse->GetRenderProxy(bProxyIsSelected, IsHovered()) : NULL,
-														WireframeColor
-														);
-
-													Collector.RegisterOneFrameMaterialProxy(CollisionMaterialInstance);
-
-													// If drawing the complex collision in wireframe, we do that in _addition_ to the normal mesh
-													if (bDrawComplexWireframeCollision)
-													{
-														FMeshBatch& CollisionElement = Collector.AllocateMesh();
-														GetMeshElement(LODIndex, BatchIndex, SectionIndex, SDPG_World, bSectionIsSelected, IsHovered(), true, CollisionElement);
-														CollisionElement.MaterialRenderProxy = CollisionMaterialInstance;
-														Collector.AddMesh(ViewIndex, CollisionElement);
-														INC_DWORD_STAT_BY(STAT_StaticMeshTriangles, CollisionElement.GetNumPrimitives());
-													}
-													// Otherwise just change material on existing mesh
-													else
-													{
-														MeshElement.MaterialRenderProxy = CollisionMaterialInstance;
-													}
-												}
-											}
-											else
-											{
-												bAddMesh = false;
-											}
-										}
 										else
 	#endif
 	#if WITH_EDITOR
@@ -941,11 +917,8 @@ void FStaticMeshSceneProxy::GetDynamicMeshElements(const TArray<const FSceneView
 										MeshElement.bCanApplyViewModeOverrides = true;
 										MeshElement.bUseWireframeSelectionColoring = bSectionIsSelected;
 
-										if (bAddMesh)
-										{
-											Collector.AddMesh(ViewIndex, MeshElement);
-											INC_DWORD_STAT_BY(STAT_StaticMeshTriangles,MeshElement.GetNumPrimitives());
-										}
+										Collector.AddMesh(ViewIndex, MeshElement);
+										INC_DWORD_STAT_BY(STAT_StaticMeshTriangles,MeshElement.GetNumPrimitives());
 									}
 								}
 							}
@@ -956,17 +929,75 @@ void FStaticMeshSceneProxy::GetDynamicMeshElements(const TArray<const FSceneView
 		}
 	}
 	
+#if !(UE_BUILD_SHIPPING || UE_BUILD_TEST)
+	// Collision and bounds drawing
+	FColor SimpleCollisionColor = FColor(157, 149, 223, 255);
+	FColor ComplexCollisionColor = FColor(0, 255, 255, 255);
+
 	for (int32 ViewIndex = 0; ViewIndex < Views.Num(); ViewIndex++)
 	{
 		if (VisibilityMap & (1 << ViewIndex))
 		{
-			// Draw simple collision as wireframe if 'show collision', collision is enabled, and we are not using the complex as the simple
-			const bool bDrawSimpleWireframeCollision = (EngineShowFlags.Collision && IsCollisionEnabled() && CollisionTraceFlag != ECollisionTraceFlag::CTF_UseComplexAsSimple);
+			if(AllowDebugViewmodes())
+			{
+				// Should we draw the mesh wireframe to indicate we are using the mesh as collision
+				const bool bDrawComplexWireframeCollision = (EngineShowFlags.Collision && IsCollisionEnabled() && CollisionTraceFlag == ECollisionTraceFlag::CTF_UseComplexAsSimple);
+
+				// If drawing complex collision as solid or wireframe
+				if(bDrawMeshCollisionWireframe || bDrawComplexWireframeCollision || (bInCollisionView && bDrawComplexCollision))
+				{
+					// If we have at least one valid LOD to draw
+					if(RenderData->LODResources.Num() > 0)
+					{
+						// Get LOD used for collision
+						int32 DrawLOD = FMath::Clamp(LODForCollision, 0, RenderData->LODResources.Num() - 1);
+						const FStaticMeshLODResources& LODModel = RenderData->LODResources[DrawLOD];
+
+						UMaterial* MaterialToUse = UMaterial::GetDefaultMaterial(MD_Surface);
+						FLinearColor DrawCollisionColor = WireframeColor;
+						// Collision view modes draw collision mesh as solid
+						if(bInCollisionView)
+						{
+							MaterialToUse = GEngine->ShadedLevelColorationUnlitMaterial;
+						}
+						// Wireframe, choose color based on complex or simple
+						else
+						{
+							MaterialToUse = GEngine->WireframeMaterial;
+							DrawCollisionColor = (CollisionTraceFlag == ECollisionTraceFlag::CTF_UseComplexAsSimple) ? SimpleCollisionColor : ComplexCollisionColor;
+						}
 
+						// Iterate over sections of that LOD
+						for (int32 SectionIndex = 0; SectionIndex < LODModel.Sections.Num(); SectionIndex++)
+						{
+							// If this section has collision enabled
+							if (LODModel.Sections[SectionIndex].bEnableCollision)
+							{
+								// See if we are selected
+								const bool bSectionIsSelected = LODs[DrawLOD].Sections[SectionIndex].bSelected;
 
-			if((bDrawSimpleCollision || bDrawSimpleWireframeCollision) && AllowDebugViewmodes())
-			{
-				if(BodySetup)
+								// Create colored proxy
+								FColoredMaterialRenderProxy* CollisionMaterialInstance = new FColoredMaterialRenderProxy(MaterialToUse->GetRenderProxy(bSectionIsSelected, false), DrawCollisionColor);
+								Collector.RegisterOneFrameMaterialProxy(CollisionMaterialInstance);
+
+								// Iterate over batches
+								for (int32 BatchIndex = 0; BatchIndex < GetNumMeshBatches(); BatchIndex++)
+								{
+									FMeshBatch& CollisionElement = Collector.AllocateMesh();
+									GetMeshElement(DrawLOD, BatchIndex, SectionIndex, SDPG_World, bSectionIsSelected, false, true, CollisionElement);
+									CollisionElement.MaterialRenderProxy = CollisionMaterialInstance;
+									Collector.AddMesh(ViewIndex, CollisionElement);
+									INC_DWORD_STAT_BY(STAT_StaticMeshTriangles, CollisionElement.GetNumPrimitives());
+								}
+							}
+						}
+					}
+				}
+
+				// Draw simple collision as wireframe if 'show collision', collision is enabled, and we are not using the complex as the simple
+				const bool bDrawSimpleWireframeCollision = (EngineShowFlags.Collision && IsCollisionEnabled() && CollisionTraceFlag != ECollisionTraceFlag::CTF_UseComplexAsSimple);
+
+				if((bDrawSimpleCollision || bDrawSimpleWireframeCollision) && BodySetup)
 				{
 					if(FMath::Abs(GetLocalToWorld().Determinant()) < SMALL_NUMBER)
 					{
@@ -994,9 +1025,8 @@ void FStaticMeshSceneProxy::GetDynamicMeshElements(const TArray<const FSceneView
 						// wireframe
 						else
 						{
-							FColor CollisionColor = FColor(157,149,223,255);
 							FTransform GeomTransform(GetLocalToWorld());
-							BodySetup->AggGeom.GetAggGeom(GeomTransform, GetSelectionColor(CollisionColor, bProxyIsSelected, IsHovered()).ToFColor(true), NULL, ( Owner == NULL ), false, UseEditorDepthTest(), ViewIndex, Collector);
+							BodySetup->AggGeom.GetAggGeom(GeomTransform, GetSelectionColor(SimpleCollisionColor, bProxyIsSelected, IsHovered()).ToFColor(true), NULL, ( Owner == NULL ), false, UseEditorDepthTest(), ViewIndex, Collector);
 						}
 
 
@@ -1012,14 +1042,14 @@ void FStaticMeshSceneProxy::GetDynamicMeshElements(const TArray<const FSceneView
 				}
 			}
 	
-#if !(UE_BUILD_SHIPPING || UE_BUILD_TEST)
 			if (EngineShowFlags.StaticMeshes)
 			{
 				RenderBounds(Collector.GetPDI(ViewIndex), EngineShowFlags, GetBounds(), !Owner || IsSelected());
 			}
-#endif
 		}
 	}
+#endif // !(UE_BUILD_SHIPPING || UE_BUILD_TEST)
+
 }
 
 void FStaticMeshSceneProxy::GetLCIs(FLCIArray& LCIs)
diff --git a/Engine/Source/Runtime/Engine/Private/TriggerBase.cpp b/Engine/Source/Runtime/Engine/Private/TriggerBase.cpp
--- a/Engine/Source/Runtime/Engine/Private/TriggerBase.cpp
+++ b/Engine/Source/Runtime/Engine/Private/TriggerBase.cpp
@@ -47,9 +47,4 @@ ATriggerBase::ATriggerBase(const FObjectInitializer& ObjectInitializer)
 #endif
 		SpriteComponent->bIsScreenSizeScaled = true;
 	}
-}
-
-/** Returns CollisionComponent subobject **/
-UShapeComponent* ATriggerBase::GetCollisionComponent() const { return CollisionComponent; }
-/** Returns SpriteComponent subobject **/
-UBillboardComponent* ATriggerBase::GetSpriteComponent() const { return SpriteComponent; }
+}
\ No newline at end of file
diff --git a/Engine/Source/Runtime/Engine/Private/UnrealEngine.cpp b/Engine/Source/Runtime/Engine/Private/UnrealEngine.cpp
--- a/Engine/Source/Runtime/Engine/Private/UnrealEngine.cpp
+++ b/Engine/Source/Runtime/Engine/Private/UnrealEngine.cpp
@@ -10947,9 +10947,9 @@ void UEngine::AddNewPendingStreamingLevel(UWorld *InWorld, FName PackageName, bo
 	new(Context.PendingLevelStreamingStatusUpdates) FLevelStreamingStatus(PackageName, bNewShouldBeLoaded, bNewShouldBeVisible, LODIndex);
 }
 
-bool UEngine::ShouldCommitPendingMapChange(UWorld *InWorld)
+bool UEngine::ShouldCommitPendingMapChange(const UWorld *InWorld) const
 {
-	FWorldContext* WorldContext = GetWorldContextFromWorld(InWorld);
+	const FWorldContext* WorldContext = GetWorldContextFromWorld(InWorld);
 	return (WorldContext ? WorldContext->bShouldCommitPendingMapChange : false);
 }
 
diff --git a/Engine/Source/Runtime/Engine/Private/World.cpp b/Engine/Source/Runtime/Engine/Private/World.cpp
--- a/Engine/Source/Runtime/Engine/Private/World.cpp
+++ b/Engine/Source/Runtime/Engine/Private/World.cpp
@@ -3500,35 +3500,6 @@ bool UWorld::AreActorsInitialized() const
 	return PersistentLevel && PersistentLevel->Actors.Num() && bActorsInitialized;
 }
 
-float UWorld::GetTimeSeconds() const
-{
-	return TimeSeconds;
-}
-
-
-float UWorld::GetRealTimeSeconds() const
-{
-	checkSlow(IsInGameThread());
-	return RealTimeSeconds;
-}
-
-
-float UWorld::GetAudioTimeSeconds() const
-{
-	return AudioTimeSeconds;
-}
-
-
-float UWorld::GetDeltaSeconds() const
-{
-	return DeltaTimeSeconds;
-}
-
-float UWorld::TimeSince( float Time ) const
-{
-	return GetTimeSeconds() - Time;
-}
-
 void UWorld::CreatePhysicsScene()
 {
 	SetPhysicsScene(new FPhysScene());
@@ -4262,12 +4233,12 @@ bool UWorld::Listen( FURL& InURL )
 #endif // WITH_SERVER_CODE
 }
 
-bool UWorld::IsClient()
+bool UWorld::IsClient() const
 {
 	return GIsClient;
 }
 
-bool UWorld::IsServer()
+bool UWorld::IsServer() const
 {
 	if ( NetDriver != NULL )
 	{
diff --git a/Engine/Source/Runtime/Engine/Public/Animation/AnimInstanceProxy.h b/Engine/Source/Runtime/Engine/Public/Animation/AnimInstanceProxy.h
--- a/Engine/Source/Runtime/Engine/Public/Animation/AnimInstanceProxy.h
+++ b/Engine/Source/Runtime/Engine/Public/Animation/AnimInstanceProxy.h
@@ -208,7 +208,7 @@ public:
 
 	/** Register a named slot */
 	void RegisterSlotNodeWithAnimInstance(FName SlotNodeName);
-	
+
 	/** Check whether we have a valid root node */
 	bool HasRootNode() const
 	{ 
@@ -250,7 +250,8 @@ public:
 	 */
 	int32 GetInstanceAssetPlayerIndex(FName MachineName, FName StateName, FName InstanceName = NAME_None);
 
-
+	float GetRecordedStateWeight(const int32& InMachineClassIndex, const int32& InStateIndex);
+	void RecordStateWeight(const int32& InMachineClassIndex, const int32& InStateIndex, const float& InStateWeight);
 
 	/** Only restricted classes can access the protected interface */
 	friend class UAnimInstance;
@@ -501,6 +502,12 @@ private:
 	/** The set of tick groups for this anim instance */
 	TArray<FAnimGroupInstance> SyncGroupArrays[2];
 
+	/** Buffers containing read/write buffers for all current state weights */
+	TArray<float> StateWeightArrays[2];
+
+	/** Map that transforms state class indices to base offsets into the weight array */
+	TMap<int32, int32> StateMachineClassIndexToWeightOffset;
+
 	// Current sync group buffer index
 	int32 SyncGroupWriteIndex;
 
diff --git a/Engine/Source/Runtime/Engine/Public/Audio.h b/Engine/Source/Runtime/Engine/Public/Audio.h
--- a/Engine/Source/Runtime/Engine/Public/Audio.h
+++ b/Engine/Source/Runtime/Engine/Public/Audio.h
@@ -39,6 +39,9 @@ ENGINE_API DECLARE_LOG_CATEGORY_EXTERN(LogAudioDebug, Display, All);
 #define MIN_PITCH						0.4f
 #define MAX_PITCH						2.0f
 
+#define MIN_SOUND_PRIORITY				0.0f
+#define MAX_SOUND_PRIORITY				100.0f
+
 /**
  * Some filters don't work properly with extreme values, so these are the limits 
  */
@@ -292,6 +295,9 @@ struct ENGINE_API FWaveInstance
 	 */
 	void AddReferencedObjects( FReferenceCollector& Collector );
 
+	/** Returns the actual volume the wave instance will play at */
+	bool ShouldStopDueToMaxConcurrency() const;
+
 	/** Returns the actual volume the wave instance will play at */
 	float GetActualVolume() const;
 
diff --git a/Engine/Source/Runtime/Engine/Public/EngineDefines.h b/Engine/Source/Runtime/Engine/Public/EngineDefines.h
--- a/Engine/Source/Runtime/Engine/Public/EngineDefines.h
+++ b/Engine/Source/Runtime/Engine/Public/EngineDefines.h
@@ -33,10 +33,6 @@
 	#define WITH_PHYSICS_COOKING (WITH_EDITOR || WITH_APEX)		//APEX currently relies on cooking even at runtime
 #endif
 
-#ifndef WITH_SUBSTEPPING
-#define WITH_SUBSTEPPING 1
-#endif
-
 #if WITH_APEX
 #ifndef WITH_APEX_CLOTHING
 	#define WITH_APEX_CLOTHING	(1 && WITH_APEX)
diff --git a/Engine/Source/Runtime/Engine/Public/MeshBatch.h b/Engine/Source/Runtime/Engine/Public/MeshBatch.h
--- a/Engine/Source/Runtime/Engine/Public/MeshBatch.h
+++ b/Engine/Source/Runtime/Engine/Public/MeshBatch.h
@@ -78,6 +78,9 @@ struct FMeshBatch
 	/** Conceptual LOD index used for the LOD Coloration visualization. */
 	int8 VisualizeLODIndex;
 
+	/** Conceptual HLOD index used for the HLOD Coloration visualization. */
+	int8 VisualizeHLODIndex;
+
 	uint32 UseDynamicData : 1;
 	uint32 ReverseCulling : 1;
 	uint32 bDisableBackfaceCulling : 1;
@@ -203,6 +206,7 @@ struct FMeshBatch
 	:	DynamicVertexStride(0)
 	,	LODIndex(INDEX_NONE)
 	,	VisualizeLODIndex(INDEX_NONE)
+	,	VisualizeHLODIndex(INDEX_NONE)
 	,	UseDynamicData(false)
 	,	ReverseCulling(false)
 	,	bDisableBackfaceCulling(false)
diff --git a/Engine/Source/Runtime/Engine/Public/ParticleEmitterInstances.h b/Engine/Source/Runtime/Engine/Public/ParticleEmitterInstances.h
--- a/Engine/Source/Runtime/Engine/Public/ParticleEmitterInstances.h
+++ b/Engine/Source/Runtime/Engine/Public/ParticleEmitterInstances.h
@@ -351,7 +351,7 @@ public:
 	 *
 	 *	@param	DeltaTime			The current time slice
 	 *	@param	CurrentLODLevel		The current LOD level for the instance
-	 *	@param	bSuppressSpawning	true if spawning has been supressed on the owning particle system component
+	 *	@param	bSuppressSpawning	true if spawning has been suppressed on the owning particle system component
 	 *	@param	bFirstTime			true if this is the first time the instance has been ticked
 	 *
 	 *	@return	float				The SpawnFraction remaining
diff --git a/Engine/Source/Runtime/Engine/Public/ShowFlags.h b/Engine/Source/Runtime/Engine/Public/ShowFlags.h
--- a/Engine/Source/Runtime/Engine/Public/ShowFlags.h
+++ b/Engine/Source/Runtime/Engine/Public/ShowFlags.h
@@ -295,6 +295,7 @@ private:
 		SetLightMapDensity(false);
 		SetVertexDensities(false);
 		SetLODColoration(false);
+		SetHLODColoration(false);
 		SetVisualizeLPV(false);
 		SetStreamingBounds(false);
 		SetConstraints(false);
diff --git a/Engine/Source/Runtime/Engine/Public/ShowFlagsValues.inl b/Engine/Source/Runtime/Engine/Public/ShowFlagsValues.inl
--- a/Engine/Source/Runtime/Engine/Public/ShowFlagsValues.inl
+++ b/Engine/Source/Runtime/Engine/Public/ShowFlagsValues.inl
@@ -292,6 +292,8 @@ SHOWFLAG_FIXED_IN_SHIPPING(0, VisualizeBloom, SFG_Visualize, LOCTEXT("VisualizeB
 SHOWFLAG_FIXED_IN_SHIPPING(0, VertexDensities, SFG_Hidden, LOCTEXT("VisualizeVertexDensities", "Visualize Vertex Densities"))
 /** Visualize LOD Coloration */
 SHOWFLAG_FIXED_IN_SHIPPING(0, LODColoration, SFG_Hidden, LOCTEXT("VisualizeLODColoration", "Visualize LOD Coloration"))
+/** Visualize HLOD Coloration */
+SHOWFLAG_FIXED_IN_SHIPPING(0, HLODColoration, SFG_Hidden, LOCTEXT("VisualizeHLODColoration", "Visualize HLOD Coloration"))
 /** Visualize screen quads */
 SHOWFLAG_FIXED_IN_SHIPPING(0, QuadComplexity, SFG_Hidden, LOCTEXT("QuadComplexitySF", "Quad Complexity"))
 /** Visualize the overhead of material quads */
diff --git a/Engine/Source/Runtime/Engine/Public/StaticMeshResources.h b/Engine/Source/Runtime/Engine/Public/StaticMeshResources.h
--- a/Engine/Source/Runtime/Engine/Public/StaticMeshResources.h
+++ b/Engine/Source/Runtime/Engine/Public/StaticMeshResources.h
@@ -883,6 +883,9 @@ protected:
 
 	const FDistanceFieldVolumeData* DistanceFieldData;	
 
+	/** Hierarchical LOD Index used for rendering */
+	uint8 HierarchicalLODIndex;
+
 	/**
 	 * The forcedLOD set in the static mesh editor, copied from the mesh component
 	 */
@@ -910,6 +913,11 @@ protected:
 #if !(UE_BUILD_SHIPPING || UE_BUILD_TEST)
 	/** Texture streaming factor used for lightmap textures */
 	float WorldLightmapFactor;
+
+	/** LOD used for collision */
+	int32 LODForCollision;
+	/** If we want to draw the mesh collision for debugging */
+	uint32 bDrawMeshCollisionWireframe : 1;
 #endif
 
 	/**
diff --git a/Engine/Source/Runtime/GeometryCache/Private/GeometryCacheTrack.cpp b/Engine/Source/Runtime/GeometryCache/Private/GeometryCacheTrack.cpp
--- a/Engine/Source/Runtime/GeometryCache/Private/GeometryCacheTrack.cpp
+++ b/Engine/Source/Runtime/GeometryCache/Private/GeometryCacheTrack.cpp
@@ -73,7 +73,7 @@ const uint32 UGeometryCacheTrack::FindSampleIndexFromTime(const TArray<float>& S
 	float SampleTime = Time;
 	if (bLooping)
 	{
-		SampleTime = fmod(Time, SampleTimes[SampleTimes.Num() - 1]);
+		SampleTime = FMath::Fmod(Time, SampleTimes[SampleTimes.Num() - 1]);
 	}
 
 	// Binary searching for closest (floored) SampleIndex 
diff --git a/Engine/Source/Runtime/Landscape/Classes/LandscapeProxy.h b/Engine/Source/Runtime/Landscape/Classes/LandscapeProxy.h
--- a/Engine/Source/Runtime/Landscape/Classes/LandscapeProxy.h
+++ b/Engine/Source/Runtime/Landscape/Classes/LandscapeProxy.h
@@ -493,6 +493,10 @@ public:
 	int32 MaxPaintedLayersPerComponent; // 0 = disabled
 #endif
 
+	/** Flag whether or not this Landscape's surface can be used for culling hidden triangles **/
+	UPROPERTY(EditAnywhere, Category = HLOD)
+	bool bUseLandscapeForCullingInvisibleHLODVertices;
+
 public:
 
 #if WITH_EDITOR
diff --git a/Engine/Source/Runtime/Landscape/Private/LandscapeEdit.cpp b/Engine/Source/Runtime/Landscape/Private/LandscapeEdit.cpp
--- a/Engine/Source/Runtime/Landscape/Private/LandscapeEdit.cpp
+++ b/Engine/Source/Runtime/Landscape/Private/LandscapeEdit.cpp
@@ -3316,7 +3316,7 @@ void ALandscapeProxy::RecreateCollisionComponents()
 		Component->DestroyComponent();
 	}
 
-	TArray<USceneComponent*> AttachedCollisionComponents = RootComponent->AttachChildren.FilterByPredicate(
+	TArray<USceneComponent*> AttachedCollisionComponents = RootComponent->GetAttachChildren().FilterByPredicate(
 		[](USceneComponent* Component)
 	{
 		return Cast<ULandscapeHeightfieldCollisionComponent>(Component);
diff --git a/Engine/Source/Runtime/Landscape/Private/LandscapeGrass.cpp b/Engine/Source/Runtime/Landscape/Private/LandscapeGrass.cpp
--- a/Engine/Source/Runtime/Landscape/Private/LandscapeGrass.cpp
+++ b/Engine/Source/Runtime/Landscape/Private/LandscapeGrass.cpp
@@ -1447,7 +1447,7 @@ void ALandscapeProxy::FlushGrassComponents(const TSet<ULandscapeComponent*>* Onl
 			Component->DestroyComponent();
 		}
 
-		TArray<USceneComponent*> AttachedFoliageComponents = RootComponent->AttachChildren.FilterByPredicate(
+		TArray<USceneComponent*> AttachedFoliageComponents = RootComponent->GetAttachChildren().FilterByPredicate(
 			[](USceneComponent* Component)
 		{
 			return Cast<UHierarchicalInstancedStaticMeshComponent>(Component);
diff --git a/Engine/Source/Runtime/Landscape/Private/LandscapeSplines.cpp b/Engine/Source/Runtime/Landscape/Private/LandscapeSplines.cpp
--- a/Engine/Source/Runtime/Landscape/Private/LandscapeSplines.cpp
+++ b/Engine/Source/Runtime/Landscape/Private/LandscapeSplines.cpp
@@ -425,9 +425,9 @@ FBoxSphereBounds ULandscapeSplinesComponent::CalcBounds(const FTransform& LocalT
 	{
 		// There's no such thing as an "invalid" FBoxSphereBounds (unlike FBox)
 		// try to return something that won't modify the parent bounds
-		if (AttachParent)
+		if (GetAttachParent())
 		{
-			NewBounds = FBoxSphereBounds(AttachParent->Bounds.Origin, FVector::ZeroVector, 0.0f);
+			NewBounds = FBoxSphereBounds(GetAttachParent()->Bounds.Origin, FVector::ZeroVector, 0.0f);
 		}
 		else
 		{
@@ -2289,9 +2289,9 @@ void ULandscapeSplineSegment::UpdateSplinePoints(bool bUpdateCollision)
 			MeshComponent->SplineUpDir = FVector(0,0,1); // Up, to be consistent between joined meshes. We rotate it to horizontal using roll if using Z Forward X Up or X Forward Y Up
 			MeshComponent->ForwardAxis = MeshEntry->ForwardAxis;
 
-			if (MeshComponent->AttachParent != OuterSplines)
+			if (MeshComponent->GetAttachParent() != OuterSplines)
 			{
-				FTransform RelativeTransform = OuterSplines->ComponentToWorld.GetRelativeTransform(MeshComponent->AttachParent->ComponentToWorld);
+				FTransform RelativeTransform = OuterSplines->ComponentToWorld.GetRelativeTransform(MeshComponent->GetAttachParent()->ComponentToWorld);
 				MeshComponent->SplineParams.StartPos = RelativeTransform.TransformPosition(MeshComponent->SplineParams.StartPos);
 				MeshComponent->SplineParams.EndPos   = RelativeTransform.TransformPosition(MeshComponent->SplineParams.EndPos);
 			}
diff --git a/Engine/Source/Runtime/RHI/Private/RHICommandList.cpp b/Engine/Source/Runtime/RHI/Private/RHICommandList.cpp
--- a/Engine/Source/Runtime/RHI/Private/RHICommandList.cpp
+++ b/Engine/Source/Runtime/RHI/Private/RHICommandList.cpp
@@ -22,7 +22,7 @@ static TAutoConsoleVariable<int32> CVarRHICmdBypass(
 	FRHICommandListExecutor::DefaultBypass,
 	TEXT("Whether to bypass the rhi command list and send the rhi commands immediately.\n")
 	TEXT("0: Disable (required for the multithreaded renderer)\n")
-	TEXT("1: Enable (convenient for debugging low level graphics API calls, can supress artifacts from multithreaded renderer code)"));
+	TEXT("1: Enable (convenient for debugging low level graphics API calls, can suppress artifacts from multithreaded renderer code)"));
 
 static TAutoConsoleVariable<int32> CVarRHICmdUseParallelAlgorithms(
 	TEXT("r.RHICmdUseParallelAlgorithms"),
diff --git a/Engine/Source/Runtime/Renderer/Private/RendererScene.cpp b/Engine/Source/Runtime/Renderer/Private/RendererScene.cpp
--- a/Engine/Source/Runtime/Renderer/Private/RendererScene.cpp
+++ b/Engine/Source/Runtime/Renderer/Private/RendererScene.cpp
@@ -743,7 +743,7 @@ void FScene::UpdatePrimitiveAttachment(UPrimitiveComponent* Primitive)
 				UpdatePrimitiveLightingAttachmentRoot(CurrentPrimitive);
 			}
 
-			ProcessStack.Append(Current->AttachChildren);
+			ProcessStack.Append(Current->GetAttachChildren());
 		}
 	}
 }
diff --git a/Engine/Source/Runtime/WebBrowser/Private/WebBrowserHandler.cpp b/Engine/Source/Runtime/WebBrowser/Private/WebBrowserHandler.cpp
--- a/Engine/Source/Runtime/WebBrowser/Private/WebBrowserHandler.cpp
+++ b/Engine/Source/Runtime/WebBrowser/Private/WebBrowserHandler.cpp
@@ -147,16 +147,16 @@ bool FWebBrowserHandler::OnBeforePopup( CefRefPtr<CefBrowser> Browser,
 	bool* NoJavascriptAccess )
 {
 	// By default, we ignore any popup requests unless they are handled by us in some way.
-	bool bSupressCEFWindowCreation = true;
+	bool bSuppressCEFWindowCreation = true;
 	if (CefCurrentlyOn(TID_UI))
 	{
 		TSharedPtr<FWebBrowserWindow> BrowserWindow = BrowserWindowPtr.Pin();
 
 		if (BrowserWindow.IsValid())
 		{
-			bSupressCEFWindowCreation = BrowserWindow->OnCefBeforePopup(TargetUrl, TargetFrameName);
+			bSuppressCEFWindowCreation = BrowserWindow->OnCefBeforePopup(TargetUrl, TargetFrameName);
 
-			if(!bSupressCEFWindowCreation)
+			if(!bSuppressCEFWindowCreation)
 			{
 				if(BrowserWindow->SupportsNewWindows())
 				{
@@ -174,11 +174,11 @@ bool FWebBrowserHandler::OnBeforePopup( CefRefPtr<CefBrowser> Browser,
 					WindowInfo.SetAsWindowless(nullptr, bUseTransparency);
 
 					// We need to rely on CEF to create our window so we set the WindowInfo, BrowserSettings, Client, and then return false
-					bSupressCEFWindowCreation = false;
+					bSuppressCEFWindowCreation = false;
 				}
 				else
 				{
-					bSupressCEFWindowCreation = true;
+					bSuppressCEFWindowCreation = true;
 				}
 			}
 		}
@@ -189,14 +189,14 @@ bool FWebBrowserHandler::OnBeforePopup( CefRefPtr<CefBrowser> Browser,
 		FEvent* Done = FPlatformProcess::GetSynchEventFromPool(true);
 		CefPostTask(TID_UI, new FWebBrowserClosureTask(this, [&]()
 		{
-			bSupressCEFWindowCreation = OnBeforePopup(Browser, Frame, TargetUrl, TargetFrameName, PopupFeatures, WindowInfo, Client, Settings, NoJavascriptAccess);
+			bSuppressCEFWindowCreation = OnBeforePopup(Browser, Frame, TargetUrl, TargetFrameName, PopupFeatures, WindowInfo, Client, Settings, NoJavascriptAccess);
 			Done->Trigger();
 		}));
 		Done->Wait();
 		FPlatformProcess::ReturnSynchEventToPool(Done);
 	}
 
-	return bSupressCEFWindowCreation; 
+	return bSuppressCEFWindowCreation;
 }
 
 bool FWebBrowserHandler::OnCertificateError(CefRefPtr<CefBrowser> Browser,
