Fix shadow variable warning
Fix unknown type in HierarchicalLODProxyProcessor.h

diff --git a/Engine/Source/Developer/HierarchicalLODUtilities/Public/HierarchicalLODProxyProcessor.h b/Engine/Source/Developer/HierarchicalLODUtilities/Public/HierarchicalLODProxyProcessor.h
--- a/Engine/Source/Developer/HierarchicalLODUtilities/Public/HierarchicalLODProxyProcessor.h
+++ b/Engine/Source/Developer/HierarchicalLODUtilities/Public/HierarchicalLODProxyProcessor.h
@@ -6,6 +6,8 @@
 #include "MeshUtilities.h"
 #include "GameFramework/WorldSettings.h"
 
+class ALODActor;
+
 class HIERARCHICALLODUTILITIES_API FHierarchicalLODProxyProcessor : public FTickerObjectBase
 {
 public:		
@@ -46,7 +48,7 @@ protected:
 	struct FProcessData
 	{
 		/** LODActor instance for which a proxy is generated */
-		class ALODActor* LODActor;
+		ALODActor* LODActor;
 		/** Array with resulting asset objects from proxy generation (StaticMesh/Materials/Textures) */
 		TArray<class UObject*> AssetObjects;
 		/** HLOD settings structure used for creating the proxy */
diff --git a/Engine/Source/Runtime/Engine/Private/NetworkDriver.cpp b/Engine/Source/Runtime/Engine/Private/NetworkDriver.cpp
--- a/Engine/Source/Runtime/Engine/Private/NetworkDriver.cpp
+++ b/Engine/Source/Runtime/Engine/Private/NetworkDriver.cpp
@@ -256,10 +256,10 @@ void UNetDriver::CancelAdaptiveReplication(FNetworkObjectInfo& InNetworkActor)
 	{
 		if (AActor* Actor = InNetworkActor.Actor)
 		{
-			if (UWorld* World = Actor->GetWorld())
+			if (UWorld* ActorWorld = Actor->GetWorld())
 			{
 				const float ExpectedNetDelay = (1.0f / Actor->NetUpdateFrequency);
-				Actor->SetNetUpdateTime(FMath::Min(Actor->NetUpdateTime, World->GetTimeSeconds() + FMath::FRandRange(0.5f, 1.0f) * ExpectedNetDelay));
+				Actor->SetNetUpdateTime(FMath::Min(Actor->NetUpdateTime, ActorWorld->GetTimeSeconds() + FMath::FRandRange(0.5f, 1.0f) * ExpectedNetDelay));
 				InNetworkActor.OptimalNetUpdateDelta = ExpectedNetDelay;
 				// TODO: we really need a way to cancel the throttling completely. OptimalNetUpdateDelta is going to be recalculated based on LastNetReplicateTime.
 			}
