Fix cook warnings due to missing collision data. (Remaking CL 2858049 in //UE4/Main).

diff --git a/Engine/Source/Runtime/Engine/Private/PhysicsEngine/BodyInstance.cpp b/Engine/Source/Runtime/Engine/Private/PhysicsEngine/BodyInstance.cpp
--- a/Engine/Source/Runtime/Engine/Private/PhysicsEngine/BodyInstance.cpp
+++ b/Engine/Source/Runtime/Engine/Private/PhysicsEngine/BodyInstance.cpp
@@ -1362,7 +1362,14 @@ struct FInitBodiesHelper
 
 				if (bInitFail)
 				{
-					UE_LOG(LogPhysics, Warning, TEXT("Init Instance %d of Primitive Component %s failed. Does it have collision data available?"), BodyIdx, *PrimitiveComp->GetReadableName());
+#if WITH_EDITOR
+					//In the editor we may have ended up here because of world trace ignoring our EnableCollision. Since we can't get at the data in that function we check for it here
+					if(!PrimitiveComp || PrimitiveComp->IsCollisionEnabled())
+#endif
+					{
+						UE_LOG(LogPhysics, Log, TEXT("Init Instance %d of Primitive Component %s failed. Does it have collision data available?"), BodyIdx, *PrimitiveComp->GetReadableName());
+					}
+
 					if (Instance->RigidActorSync)
 					{
 						Instance->RigidActorSync->release();
