Fixed crash in merged skeletal meshes when composed ref pose matrices are accessed before mesh is finalized

Added call to CalculateInvRefMatrices() to FSkeletalMeshMerge::MergeSkeleton.

#tests PIE in Ocean level NewCleaner_Test_01

diff --git a/Engine/Source/Runtime/Engine/Private/SkeletalMeshMerge.cpp b/Engine/Source/Runtime/Engine/Private/SkeletalMeshMerge.cpp
--- a/Engine/Source/Runtime/Engine/Private/SkeletalMeshMerge.cpp
+++ b/Engine/Source/Runtime/Engine/Private/SkeletalMeshMerge.cpp
@@ -91,6 +91,12 @@ void FSkeletalMeshMerge::MergeSkeleton(const TArray<FRefPoseOverride>* RefPoseOv
 	// Assign new referencer skeleton.
 
 	MergeMesh->RefSkeleton = NewRefSkeleton;
+
+	// Rebuild inverse ref pose matrices here as some access patterns 
+	// may need to access these matrices before FinalizeMesh is called
+	// (which would *normally* rebuild the inv ref matrices).
+	MergeMesh->RefBasesInvMatrix.Empty();
+	MergeMesh->CalculateInvRefMatrices();
 }
 
 bool FSkeletalMeshMerge::FinalizeMesh()
