Copying //UE4/Orion-Staging to //UE4/Main (Origin //Orion/Dev-General @ 2861092)

diff --git a/Engine/Source/Runtime/Core/Public/Windows/WIndowsPlatform.h b/Engine/Source/Runtime/Core/Public/Windows/WIndowsPlatform.h
--- a/Engine/Source/Runtime/Core/Public/Windows/WIndowsPlatform.h
+++ b/Engine/Source/Runtime/Core/Public/Windows/WIndowsPlatform.h
@@ -78,7 +78,9 @@ typedef FWindowsPlatformTypes FPlatformTypes;
 
 #define PLATFORM_SUPPORTS_STACK_SYMBOLS						1
 
-// Intrinsics for 128-bit atomics on Windows platform requires Windows 8 or higher (WINVER>0x0602)
+// Intrinsics for 64-bit atomics on Windows platform requires Windows Vista or higher (WINVER>=0x0600)
+#define PLATFORM_HAS_64BIT_ATOMICS							(WINVER >= 0x600)
+// Intrinsics for 128-bit atomics on Windows platform requires Windows 8 or higher (WINVER>=0x0602)
 // http://msdn.microsoft.com/en-us/library/windows/desktop/hh972640.aspx
 #define PLATFORM_HAS_128BIT_ATOMICS							(!HACK_HEADER_GENERATOR && PLATFORM_64BITS && (WINVER >= 0x602))
 #define PLATFORM_USES_ANSI_STRING_FOR_EXTERNAL_PROFILING	0
diff --git a/Engine/Source/Runtime/Core/Public/Windows/WindowsPlatformAtomics.h b/Engine/Source/Runtime/Core/Public/Windows/WindowsPlatformAtomics.h
--- a/Engine/Source/Runtime/Core/Public/Windows/WindowsPlatformAtomics.h
+++ b/Engine/Source/Runtime/Core/Public/Windows/WindowsPlatformAtomics.h
@@ -17,7 +17,7 @@ struct CORE_API FWindowsPlatformAtomics
 		return (int32)::InterlockedIncrement((LPLONG)Value);
 	}
 
-#if WINVER >= 0x0600 // Interlock...64 functions are only available from Vista onwards
+#if PLATFORM_HAS_64BIT_ATOMICS
 	static FORCEINLINE int64 InterlockedIncrement( volatile int64* Value )
 	{
 		return (int64)::InterlockedIncrement64((LONGLONG*)Value);
@@ -29,7 +29,7 @@ struct CORE_API FWindowsPlatformAtomics
 		return (int32)::InterlockedDecrement((LPLONG)Value);
 	}
 
-#if WINVER >= 0x0600	// Interlock...64 functions are only available from Vista onwards
+#if PLATFORM_HAS_64BIT_ATOMICS
 	static FORCEINLINE int64 InterlockedDecrement( volatile int64* Value )
 	{
 		return (int64)::InterlockedDecrement64((LONGLONG*)Value);
@@ -41,7 +41,7 @@ struct CORE_API FWindowsPlatformAtomics
 		return (int32)::InterlockedExchangeAdd((LPLONG)Value, (LONG)Amount);
 	}
 
-#if WINVER >= 0x0600	// Interlock...64 functions are only available from Vista onwards
+#if PLATFORM_HAS_64BIT_ATOMICS
 	static FORCEINLINE int64 InterlockedAdd( volatile int64* Value, int64 Amount )
 	{
 		return (int64)::InterlockedExchangeAdd64((LONGLONG*)Value, (LONGLONG)Amount);
@@ -53,7 +53,7 @@ struct CORE_API FWindowsPlatformAtomics
 		return (int32)::InterlockedExchange((LPLONG)Value, (LONG)Exchange);
 	}
 
-#if WINVER >= 0x0600	// Interlock...64 functions are only available from Vista onwards
+#if PLATFORM_HAS_64BIT_ATOMICS
 	static FORCEINLINE int64 InterlockedExchange( volatile int64* Value, int64 Exchange )
 	{
 		return (int64)::InterlockedExchange64((LONGLONG*)Value, (LONGLONG)Exchange);
@@ -69,7 +69,7 @@ struct CORE_API FWindowsPlatformAtomics
 		}
 #endif
 
-#if WINVER >= 0x0600	// Interlock...64 functions are only available from Vista onwards
+#if PLATFORM_HAS_64BIT_ATOMICS
 		return ::InterlockedExchangePointer(Dest, Exchange);
 #else
 		return (void*)::InterlockedExchange((PLONG)(Dest), (LONG)(Exchange));
@@ -81,7 +81,7 @@ struct CORE_API FWindowsPlatformAtomics
 		return (int32)::InterlockedCompareExchange((LPLONG)Dest, (LONG)Exchange, (LONG)Comparand);
 	}
 
-#if WINVER >= 0x0600	// Interlock...64 functions are only available from Vista onwards
+#if PLATFORM_HAS_64BIT_ATOMICS
 	static FORCEINLINE int64 InterlockedCompareExchange( volatile int64* Dest, int64 Exchange, int64 Comparand )
 	{
 #if !(UE_BUILD_SHIPPING || UE_BUILD_TEST)
diff --git a/Engine/Source/Runtime/CoreUObject/CoreUObject.Build.cs b/Engine/Source/Runtime/CoreUObject/CoreUObject.Build.cs
--- a/Engine/Source/Runtime/CoreUObject/CoreUObject.Build.cs
+++ b/Engine/Source/Runtime/CoreUObject/CoreUObject.Build.cs
@@ -14,7 +14,6 @@ public class CoreUObject : ModuleRules
                 new string[] 
 			    {
 				    "TargetPlatform",
-				    "CookingStats",
 			    }
             );
 
diff --git a/Engine/Source/Runtime/CoreUObject/Private/UObject/SavePackage.cpp b/Engine/Source/Runtime/CoreUObject/Private/UObject/SavePackage.cpp
--- a/Engine/Source/Runtime/CoreUObject/Private/UObject/SavePackage.cpp
+++ b/Engine/Source/Runtime/CoreUObject/Private/UObject/SavePackage.cpp
@@ -18,90 +18,10 @@ static const FName WorldClassName = FName("World");
 #define VALIDATE_INITIALIZECORECLASSES 0
 #define EXPORT_SORTING_DETAILED_LOGGING 0
 
-#define UE_PROFILE_COOKSAVE 1
-#if UE_PROFILE_COOKSAVE
-
-#include "CookingStatsModule.h"
-
-#define UE_OUTPUTSTATSTOLOG 0
-
-namespace SavePackageStats
-{
-static const FName CookingStatsName("CookingStats");
-static const FName GSavePackageTransactionId(*FString::Printf(TEXT("SavePackageTransactionId%s"), *FGuid::NewGuid().ToString()));
-static uint32 GSavePackageTransactionNumber = 0;
-
-	ICookingStats* GetCookingStats()
-	{
-		static ICookingStats* CookingStats = nullptr;
-		static bool bInitialized = false;
-		if (bInitialized == false)
-		{
-			FCookingStatsModule* CookingStatsModule = FModuleManager::LoadModulePtr<FCookingStatsModule>(TEXT("CookingStats"));
-			if (CookingStatsModule)
-			{
-				CookingStats = CookingStatsModule->Get();
-			}
-			bInitialized = true;
-		}
-		return CookingStats;
-	}
-
-
-	// static ICookingStats* CookingStats = GetCookingStats();
-
-	void CookStatsStartStat(const FName& Key, const TCHAR *Filename)
-	{
-#if UE_OUTPUTSTATSTOLOG
-		UE_LOG(LogSavePackage, Log, TEXT("Starting save package for %s"), Filename);
-#endif
-
-		ICookingStats* CookingStats = GetCookingStats();
-		if (CookingStats)
-		{
-			CookingStats->AddTagValue(Key, TEXT("Filename"), FString(Filename));
-		}
-	}
-
-	void CookStatsAddStat(const FName& Key, const FName& Tag, const float TimeMilliseconds)
-	{
-#if UE_OUTPUTSTATSTOLOG
-		UE_LOG(LogSavePackage, Log, TEXT("TIMING: %s took %.2fms"), Tag, TimeMilliseconds);
-#endif
-
-		ICookingStats* CookingStats = GetCookingStats();
-		if (CookingStats)
-		{
-			CookingStats->AddTagValue(Key, Tag, TimeMilliseconds);
-		}
-	}
-}; // namespace SavePackageStats
-
-
-#define UE_START_LOG_COOK_TIME(InFilename) double PreviousTime; double StartTime; PreviousTime = StartTime = FPlatformTime::Seconds(); const FName CookStatsKey = FName(SavePackageStats::GSavePackageTransactionId,++SavePackageStats::GSavePackageTransactionNumber); \
-	SavePackageStats::CookStatsStartStat(CookStatsKey, InFilename);
-
-#define UE_LOG_COOK_TIME(TimeType) \
-	{\
-		double CurrentTime = FPlatformTime::Seconds(); \
-	const FName TagName = FName(TimeType); \
-	SavePackageStats::CookStatsAddStat(CookStatsKey, TagName, (CurrentTime - PreviousTime) * 1000.0f); \
-		PreviousTime = CurrentTime; \
-	}
-
-#define UE_FINISH_LOG_COOK_TIME()\
-	UE_LOG_COOK_TIME(TEXT("UnaccountedTime")); \
-	{\
-		double CurrentTime = FPlatformTime::Seconds(); \
-		UE_LOG(LogSavePackage, Log, TEXT("Total save time: %.2fms"), (CurrentTime - StartTime) * 1000.0f); \
-	}
-
-#else
 // don't do anything
 #define UE_START_LOG_COOK_TIME(InFilename)
 #define UE_LOG_COOK_TIME(TimeType) 
 #define UE_FINISH_LOG_COOK_TIME()
-#endif
 
 
 static bool HasDeprecatedOrPendingKillOuter(UObject* InObj, UPackage* InSavingPackage)
diff --git a/Engine/Source/Runtime/CoreUObject/Public/UObject/CoreNet.h b/Engine/Source/Runtime/CoreUObject/Public/UObject/CoreNet.h
--- a/Engine/Source/Runtime/CoreUObject/Public/UObject/CoreNet.h
+++ b/Engine/Source/Runtime/CoreUObject/Public/UObject/CoreNet.h
@@ -218,7 +218,8 @@ enum ELifetimeCondition
 	COND_InitialOrOwner		= 7,		// This property will send on the initial packet, or to the actors owner
 	COND_Custom				= 8,		// This property has no particular condition, but wants the ability to toggle on/off via SetCustomIsActiveOverride
 	COND_ReplayOrOwner		= 9,		// This property will only send to the replay connection, or to the actors owner
-	COND_Max				= 10,
+	COND_ReplayOnly			= 10,		// This property will only send to the replay connection
+	COND_Max				= 11,
 };
 
 
diff --git a/Engine/Source/Runtime/Engine/Classes/Engine/World.h b/Engine/Source/Runtime/Engine/Classes/Engine/World.h
--- a/Engine/Source/Runtime/Engine/Classes/Engine/World.h
+++ b/Engine/Source/Runtime/Engine/Classes/Engine/World.h
@@ -30,6 +30,7 @@ class AMatineeActor;
 class AWorldSettings;
 class ACameraActor;
 class FUniqueNetId;
+class FWorldInGamePerformanceTrackers;
 
 template<typename,typename> class TOctree;
 
@@ -1033,6 +1034,9 @@ public:
 	TArray<TSharedPtr<FAsyncPreRegisterDDCRequest>> AsyncPreRegisterDDCRequests;
 #endif
 
+	//Experimental: In game performance tracking.
+	FWorldInGamePerformanceTrackers* PerfTrackers;
+
 	/**
 	 * UWorld default constructor
 	 */
diff --git a/Engine/Source/Runtime/Engine/Classes/GameFramework/GameUserSettings.h b/Engine/Source/Runtime/Engine/Classes/GameFramework/GameUserSettings.h
--- a/Engine/Source/Runtime/Engine/Classes/GameFramework/GameUserSettings.h
+++ b/Engine/Source/Runtime/Engine/Classes/GameFramework/GameUserSettings.h
@@ -26,6 +26,8 @@ namespace EWindowMode
 
 #endif
 
+DECLARE_DYNAMIC_MULTICAST_DELEGATE(FOnGameUserSettingsUINeedsUpdate);
+
 /**
  * Stores user settings for a game (for example graphics and sound settings), with the ability to save and load to and from a file.
  */
@@ -314,6 +316,10 @@ protected:
 	/** Min resolution scale we allow in current display mode */
 	int32 MinResolutionScale;
 
+	/** Desired screen height used to calculate the resolution scale when user changes display mode */
+	UPROPERTY(config)
+	int32 DesiredScreenHeight;
+
 	/**
 	 * Check if the current version of the game user settings is valid. Sub-classes can override this to provide game-specific versioning as necessary.
 	 * @return True if the current version is valid, false if it is not
@@ -323,9 +329,14 @@ protected:
 	/** Update the version of the game user settings to the current version */
 	virtual void UpdateVersion();
 
+	int32 FindResolutionQualityForScreenHeight(int32 ScreenHeight);
+
 private:
 
-	void UpdateMinResolutionScaling();
+	UPROPERTY(BlueprintAssignable, meta = (AllowPrivateAccess = "true"))
+	FOnGameUserSettingsUINeedsUpdate OnGameUserSettingsUINeedsUpdate;
+
+	void UpdateResolutionQuality();
 
 	void SetPreferredFullscreenMode(int32 Mode);
 };
diff --git a/Engine/Source/Runtime/Engine/Classes/GameFramework/HUD.h b/Engine/Source/Runtime/Engine/Classes/GameFramework/HUD.h
--- a/Engine/Source/Runtime/Engine/Classes/GameFramework/HUD.h
+++ b/Engine/Source/Runtime/Engine/Classes/GameFramework/HUD.h
@@ -432,6 +432,9 @@ public:
 	/** draw overlays for actors that were rendered this tick and have added themselves to the PostRenderedActors array	*/
 	virtual void DrawActorOverlays(FVector Viewpoint, FRotator ViewRotation);
 
+	/** Draw the safe zone debugging overlay when enabled */
+	virtual void DrawSafeZoneOverlay();
+
 	/** Called in PostInitializeComponents or postprocessing chain has changed (happens because of the worldproperties can define it's own chain and this one is set late). */
 	virtual void NotifyBindPostProcessEffects();
 
diff --git a/Engine/Source/Runtime/Engine/Classes/Kismet/KismetStringLibrary.h b/Engine/Source/Runtime/Engine/Classes/Kismet/KismetStringLibrary.h
--- a/Engine/Source/Runtime/Engine/Classes/Kismet/KismetStringLibrary.h
+++ b/Engine/Source/Runtime/Engine/Classes/Kismet/KismetStringLibrary.h
@@ -260,6 +260,17 @@ class ENGINE_API UKismetStringLibrary : public UBlueprintFunctionLibrary
 	UFUNCTION(BlueprintPure, Category = "Utilities|String")
 	static int32 FindSubstring(const FString& SearchIn, const FString& Substring, bool bUseCase = false, bool bSearchFromEnd = false, int32 StartPosition = -1);
 
+	/**
+	* Returns whether this string contains the specified substring.
+	*
+	* @param SubStr			Find to search for
+	* @param SearchCase		Indicates whether the search is case sensitive or not ( defaults to ESearchCase::IgnoreCase )
+	* @param SearchDir			Indicates whether the search starts at the begining or at the end ( defaults to ESearchDir::FromStart )
+	* @return					Returns whether the string contains the substring
+	**/
+	UFUNCTION(BlueprintPure, Category = "Utilities|String")
+	static bool Contains(const FString& SearchIn, const FString& Substring, bool bUseCase = false, bool bSearchFromEnd = false);
+
 	/** 
 	 * Gets a single character from the string (as an integer)
 	 * @param SourceString - The string to convert
diff --git a/Engine/Source/Runtime/Engine/Classes/Particles/ParticleEmitter.h b/Engine/Source/Runtime/Engine/Classes/Particles/ParticleEmitter.h
--- a/Engine/Source/Runtime/Engine/Classes/Particles/ParticleEmitter.h
+++ b/Engine/Source/Runtime/Engine/Classes/Particles/ParticleEmitter.h
@@ -394,6 +394,9 @@ class UParticleEmitter : public UObject
 #if STATS
 	void CreateStatID() const;
 #endif
+
+	/** Returns if this emitter is considered significant for the passed requirement. */
+	ENGINE_API bool IsSignificant(EParticleSignificanceLevel RequiredSignificance);
 };
 
 
diff --git a/Engine/Source/Runtime/Engine/Classes/Particles/ParticleSystem.h b/Engine/Source/Runtime/Engine/Classes/Particles/ParticleSystem.h
--- a/Engine/Source/Runtime/Engine/Classes/Particles/ParticleSystem.h
+++ b/Engine/Source/Runtime/Engine/Classes/Particles/ParticleSystem.h
@@ -106,6 +106,12 @@ class UParticleSystem : public UObject
 {
 	GENERATED_UCLASS_BODY()
 
+	UParticleSystem()
+	:OcclusionBoundsMethod(EPSOBM_ParticleBounds)
+	{
+
+	}
+
 	UPROPERTY(EditAnywhere, Category=ParticleSystem, AssetRegistrySearchable)
 	TEnumAsByte<enum EParticleSystemUpdateMode> SystemUpdateMode;
 
@@ -298,8 +304,13 @@ public:
 	UPROPERTY(EditAnywhere, Category = Performance, meta = (ToolTip = "Minimum duration between ticks; 33=tick at max. 30FPS, 16=60FPS, 8=120FPS"))
 	uint32 MinTimeBetweenTicks;
 
-	UPROPERTY(EditAnywhere, Category = Performance, meta = (ToolTip = "Time this system can be deemed insignificant before it deactivates. "))
-	float InsignificantTimeTillDeactivate;
+	/** The reaction this system takes when all emitters are insignificant. */
+	UPROPERTY(EditAnywhere, Category = Performance)
+	EParticleSystemInsignificanceReaction InsignificantReaction;
+
+	/** Time delay between all emitters becoming insignificant and the systems insignificant reaction. */
+	UPROPERTY(EditAnywhere, Category = Performance)
+	float InsignificanceDelay;
 
 	/** The maximum level of significance for emitters in this system. Any emitters with a higher significance will be capped at this significance level. */
 	UPROPERTY(EditAnywhere, Category = Performance)
@@ -482,9 +493,20 @@ public:
 	/** Returns true if the particle system is looping (contains one or more looping emitters) */
 	bool IsLooping() const { return bAnyEmitterLoopsForever; }
 
+	EParticleSignificanceLevel GetHighestSignificance()const { return HighestSignificance; }
+	EParticleSignificanceLevel GetLowestSignificance()const { return LowestSignificance; }
+	bool ShouldManageSignificance()const { return bShouldManageSignificance; }
 private:
+
+	/** The highest significance of any emitter. Clamped by MaxSignificanceLevel.*/
+	EParticleSignificanceLevel HighestSignificance;
+	/** The lowest significance of any emitter. Clamped by MaxSignificanceLevel.*/
+	EParticleSignificanceLevel LowestSignificance;
+	
+	uint32 bShouldManageSignificance : 1;
+
 	/** Does any emitter loop forever (computed during load and when emitters are built) */
-	bool bAnyEmitterLoopsForever;
+	uint32 bAnyEmitterLoopsForever : 1;
 };
 
 
diff --git a/Engine/Source/Runtime/Engine/Classes/Particles/ParticleSystemComponent.h b/Engine/Source/Runtime/Engine/Classes/Particles/ParticleSystemComponent.h
--- a/Engine/Source/Runtime/Engine/Classes/Particles/ParticleSystemComponent.h
+++ b/Engine/Source/Runtime/Engine/Classes/Particles/ParticleSystemComponent.h
@@ -74,7 +74,8 @@ enum EParticleEventType
 // Called when the particle system is done
 DECLARE_DYNAMIC_MULTICAST_DELEGATE_OneParam( FOnSystemFinished, class UParticleSystemComponent*, PSystem );
 
-DECLARE_DYNAMIC_MULTICAST_DELEGATE_OneParam(FOnSystemPreActivate, class UParticleSystemComponent*, PSC);
+//Called just before the activation of a component changes.
+DECLARE_MULTICAST_DELEGATE_TwoParams(FOnSystemPreActivationChange, class UParticleSystemComponent*, bool);
 
 /** Struct used for a particular named instance parameter for this ParticleSystemComponent. */
 USTRUCT(BlueprintType)
@@ -375,14 +376,11 @@ public:
 	UPROPERTY()
 	EParticleSignificanceLevel RequiredSignificance;
 
-	/** True if this component has emitters that weren't created or inited due to significance. If set we need to re-init when significance increases. */
-	bool bMissingEmittersDueToSignificance;
-
-	/** The significance level we ran InitParticles() with. If we change to require less significance than this we need to re-run InitParticles() to be sure we have created all emitters. */
-	EParticleSignificanceLevel InitRequiredSignificanceLevel;
-	
 	/** Time in seconds since we were last considered significant. */
 	float LastSignificantTime;
+
+	/** If this component is considered significant or not. */
+	uint32 bIsSignificant : 1;
 	
 private:
 	/** Did we auto attach during activation? Used to determine if we should restore the relative transform during detachment. */
@@ -440,6 +438,8 @@ public:
 	 */
 	UPROPERTY(EditAnywhere, BlueprintReadWrite, Category=Particles)
 	float SecondsBeforeInactive;
+	/** Allow the above or not. This feature can interfere with some other systems. */
+	bool bAllowSecondsBeforeInactive;
 
 private:
 	/** Tracks the time since the last forced UpdateTransform. */
@@ -543,9 +543,6 @@ public:
 	UPROPERTY(BlueprintAssignable)
 	FOnSystemFinished OnSystemFinished;
 
-	UPROPERTY(BlueprintAssignable)
-	FOnSystemPreActivate OnSystemPreActivate;
-
 public:
 	/**
 	 * Component we automatically attach to when activated, if bAutoManageAttachment is true.
@@ -609,7 +606,14 @@ private:
 
 public:
 
-	void SetSignificance(float OldSignificance, float NewSignificance);
+	/** Called from game code when the significance required for a component changes. */
+	void SetRequiredSignificance(EParticleSignificanceLevel NewRequiredSignificance);
+	/** Whether this component should have it's significance managed by game code. */
+	bool ShouldManageSignificance()const;
+	/** When the overall significance for the component is changed. */
+	void OnSignificanceChanged(bool bSignificant, bool bApplyToEmitters);
+	
+	void Complete();
 
 	FORCEINLINE const FTransform& GetAsyncComponentToWorld()
 	{
@@ -903,6 +907,9 @@ public:
 	/** Command fence used to shut down properly */
 	class FRenderCommandFence* ReleaseResourcesFence;
 
+	/** Static delegate called for all systems on an activation change. */
+	static FOnSystemPreActivationChange OnSystemPreActivationChange;
+
 private:
 	/** Task ref for parallel portion */
 	FGraphEventRef AsyncWork;
@@ -914,6 +921,8 @@ private:
 	bool bNeedsFinalize;
 	/** If true, it means the Async work is in process and not yet completed */
 	bool bAsyncWorkOutstanding;
+	/** Number of significant emitters. When this is 0, the system can either be deactivated or stopped ticking. */
+	uint32 NumSignificantEmitters;
 	/** Time in ms since a tick was last performed; used with MinTimeBetweenTicks (on UParticleSystem) to control tick rate */
 	uint32 TimeSinceLastTick;
 
@@ -1082,7 +1091,7 @@ public:
 	}
 
 	// If particles have not already been initialised (ie. EmitterInstances created) do it now.
-	virtual void InitParticles(bool bReInitExistingEmitters=true);
+	virtual void InitParticles();
 
 
 	// @todo document
diff --git a/Engine/Source/Runtime/Engine/Classes/Sound/DialogueWave.h b/Engine/Source/Runtime/Engine/Classes/Sound/DialogueWave.h
--- a/Engine/Source/Runtime/Engine/Classes/Sound/DialogueWave.h
+++ b/Engine/Source/Runtime/Engine/Classes/Sound/DialogueWave.h
@@ -36,6 +36,8 @@ struct FDialogueContextMapping
 bool operator==(const FDialogueContextMapping& LHS, const FDialogueContextMapping& RHS);
 bool operator!=(const FDialogueContextMapping& LHS, const FDialogueContextMapping& RHS);
 
+class UDialogueWaveFactory;
+
 UCLASS(hidecategories=Object, editinlinenew, MinimalAPI, BlueprintType)
 class UDialogueWave : public UObject
 {
@@ -84,6 +86,10 @@ public:
 	ENGINE_API FString GetContextLocalizationKey( const FDialogueContext& Context ) const;
 	//~ End UDialogueWave Interface.
 
+	ENGINE_API void UpdateContext(FDialogueContextMapping& ContextMapping, USoundWave* SoundWave, UDialogueVoice* Speaker, const TArray<UDialogueVoice*>& Targets);
+
 protected:
+
 	void UpdateMappingProxy(FDialogueContextMapping& ContextMapping);
+
 };
\ No newline at end of file
diff --git a/Engine/Source/Runtime/Engine/Classes/Sound/SoundNodeDialoguePlayer.h b/Engine/Source/Runtime/Engine/Classes/Sound/SoundNodeDialoguePlayer.h
--- a/Engine/Source/Runtime/Engine/Classes/Sound/SoundNodeDialoguePlayer.h
+++ b/Engine/Source/Runtime/Engine/Classes/Sound/SoundNodeDialoguePlayer.h
@@ -1,7 +1,6 @@
 // Copyright 1998-2016 Epic Games, Inc. All Rights Reserved.
-
-
 #pragma once
+
 #include "Sound/SoundNode.h"
 #include "SoundNodeDialoguePlayer.generated.h"
 
@@ -21,11 +20,18 @@ class USoundNodeDialoguePlayer : public USoundNode
 	uint32 bLooping:1;
 
 public:	
+
+	ENGINE_API UDialogueWave* GetDialogueWave() const;
+	ENGINE_API void SetDialogueWave(UDialogueWave* Value);
+
 	//~ Begin USoundNode Interface
 	virtual int32 GetMaxChildNodes() const override;
 	virtual float GetDuration() override;
+	virtual int32 GetNumSounds(const UPTRINT NodeWaveInstanceHash, FActiveSound& ActiveSound) const { return 1; }
 	virtual void ParseNodes( FAudioDevice* AudioDevice, const UPTRINT NodeWaveInstanceHash, FActiveSound& ActiveSound, const FSoundParseParameters& ParseParams, TArray<FWaveInstance*>& WaveInstances ) override;
+#if WITH_EDITOR
+	virtual FText GetTitle() const override;
+#endif
 	//~ End USoundNode Interface
-
 };
 
diff --git a/Engine/Source/Runtime/Engine/Private/Actor.cpp b/Engine/Source/Runtime/Engine/Private/Actor.cpp
--- a/Engine/Source/Runtime/Engine/Private/Actor.cpp
+++ b/Engine/Source/Runtime/Engine/Private/Actor.cpp
@@ -3122,13 +3122,12 @@ bool AActor::SetActorTransform(const FTransform& NewTransform, bool bSweep, FHit
 	// happens, and if so, something else is giving NAN as output
 	if (RootComponent)
 	{
-		if (ensure(!NewTransform.ContainsNaN()))
+		if (ensureMsgf(!NewTransform.ContainsNaN(), TEXT("SetActorTransform: Get NAN Transform data for %s: %s"), *GetNameSafe(this), *NewTransform.ToString()))
 		{
 			RootComponent->SetWorldTransform(NewTransform, bSweep, OutSweepHitResult, Teleport);
 		}
 		else
 		{
-			UE_LOG(LogScript, Warning, TEXT("SetActorTransform: Get NAN Transform data for %s: %s"), *GetNameSafe(this), *NewTransform.ToString());
 			if (OutSweepHitResult)
 			{
 				*OutSweepHitResult = FHitResult();
diff --git a/Engine/Source/Runtime/Engine/Private/Animation/AnimInstanceProxy.cpp b/Engine/Source/Runtime/Engine/Private/Animation/AnimInstanceProxy.cpp
--- a/Engine/Source/Runtime/Engine/Private/Animation/AnimInstanceProxy.cpp
+++ b/Engine/Source/Runtime/Engine/Private/Animation/AnimInstanceProxy.cpp
@@ -312,8 +312,8 @@ void FAnimInstanceProxy::TickAssetPlayerInstances(float DeltaSeconds)
 					// if this leader contains correct position, break
 					SyncGroup.MarkerTickContext = TickContext.MarkerTickContext;
 					SyncGroup.GroupLeaderIndex = GroupLeaderIndex;
-					UE_LOG(LogAnimMarkerSync, Log, TEXT("Previous Sync Group Makrer Tick Context :\n%s"), *SyncGroup.MarkerTickContext.ToString());
-					UE_LOG(LogAnimMarkerSync, Log, TEXT("New Sync Group Makrer Tick Context :\n%s"), *TickContext.MarkerTickContext.ToString());
+					UE_LOG(LogAnimMarkerSync, Log, TEXT("Previous Sync Group Marker Tick Context :\n%s"), *SyncGroup.MarkerTickContext.ToString());
+					UE_LOG(LogAnimMarkerSync, Log, TEXT("New Sync Group Marker Tick Context :\n%s"), *TickContext.MarkerTickContext.ToString());
 					break;
 				}
 				else
@@ -484,7 +484,7 @@ FMarkerSyncAnimPosition FAnimInstanceProxy::GetSyncGroupPosition(FName InSyncGro
 	if (SyncGroups.IsValidIndex(SyncGroupIndex))
 	{
 		const FAnimGroupInstance& SyncGroupInstance = SyncGroups[SyncGroupIndex];
-		if (SyncGroupInstance.bCanUseMarkerSync)
+		if (SyncGroupInstance.bCanUseMarkerSync && SyncGroupInstance.MarkerTickContext.IsMarkerSyncEndValid())
 		{
 			return SyncGroupInstance.MarkerTickContext.GetMarkerSyncEndPosition();
 		}
diff --git a/Engine/Source/Runtime/Engine/Private/Animation/AnimMontage.cpp b/Engine/Source/Runtime/Engine/Private/Animation/AnimMontage.cpp
--- a/Engine/Source/Runtime/Engine/Private/Animation/AnimMontage.cpp
+++ b/Engine/Source/Runtime/Engine/Private/Animation/AnimMontage.cpp
@@ -351,7 +351,7 @@ void UAnimMontage::PostLoad()
 
 		if(CurrentCalculatedLength != SequenceLength)
 		{
-			UE_LOG(LogAnimMontage, Warning, TEXT("UAnimMontage::PostLoad: The actual sequence length for montage %s does not match the length stored in the asset, please resave the montage asset."), *GetName());
+			UE_LOG(LogAnimMontage, Warning, TEXT("UAnimMontage::PostLoad: The actual sequence length for %s does not match the length stored in the asset, please resave the asset."), *GetFullName());
 			SequenceLength = CurrentCalculatedLength;
 		}
 	}
diff --git a/Engine/Source/Runtime/Engine/Private/AudioDerivedData.cpp b/Engine/Source/Runtime/Engine/Private/AudioDerivedData.cpp
--- a/Engine/Source/Runtime/Engine/Private/AudioDerivedData.cpp
+++ b/Engine/Source/Runtime/Engine/Private/AudioDerivedData.cpp
@@ -12,13 +12,6 @@ DEFINE_LOG_CATEGORY_STATIC(LogAudioDerivedData, Log, All);
 ------------------------------------------------------------------------------*/
 
 
-// enable audio building stats these appear in the Saved\Stats\stats.csv file
-#define BUILD_AUDIO_STATS 1 && ENABLE_DDC_STATS
-
-#if BUILD_AUDIO_STATS
-#include "DDCStatsHelper.h"
-#endif
-
 #if WITH_EDITORONLY_DATA
 
 // If you want to bump this version, generate a new guid using
@@ -235,14 +228,6 @@ class FStreamedAudioCacheDerivedDataWorker : public FNonAbandonableTask
 	/** Build the streamed audio. This function is safe to call from any thread. */
 	void BuildStreamedAudio()
 	{
-#if BUILD_AUDIO_STATS
-		FString DDCKey;
-		GetStreamedAudioDerivedDataKeyFromSuffix(KeySuffix, DDCKey);
-		static const FName NAME_BuildStreamedAudio(TEXT("BuildStreamedAudio"));
-		FDDCScopeStatHelper ScopeTimer(*DDCKey, NAME_BuildStreamedAudio);
-#endif
-
-
 		GetStreamedAudioDerivedDataKeySuffix(SoundWave, AudioFormatName, KeySuffix);
 
 		DerivedData->Chunks.Empty();
diff --git a/Engine/Source/Runtime/Engine/Private/ContentStreaming.cpp b/Engine/Source/Runtime/Engine/Private/ContentStreaming.cpp
--- a/Engine/Source/Runtime/Engine/Private/ContentStreaming.cpp
+++ b/Engine/Source/Runtime/Engine/Private/ContentStreaming.cpp
@@ -187,13 +187,6 @@ static TAutoConsoleVariable<int32> CVarStreamingHLODStrategy(
 	TEXT("2: disable streaming"),
 	ECVF_Default);
 
-static TAutoConsoleVariable<float> CVarStreamingHLODPriorityScale(
-	TEXT("r.Streaming.HLODPriorityScale"),
-	1.5,
-	TEXT("Define the streaming priority to apply to HLOD textures\n"),
-	ECVF_Default
-	);
-
 static TAutoConsoleVariable<float> CVarStreamingHiddenPrimitiveScale(
 	TEXT("r.Streaming.HiddenPrimitiveScale"),
 	0.5,
@@ -203,6 +196,12 @@ static TAutoConsoleVariable<float> CVarStreamingHiddenPrimitiveScale(
 	ECVF_Default
 	);
 
+static TAutoConsoleVariable<int32> CVarStreamingAllowRetentionLogic(
+	TEXT("r.Streaming.AllowRetentionLogic"),
+	1,
+	TEXT("Define if to handle overbudget situations with retention logic.\n"),
+	ECVF_Default
+	);
 
 /** Streaming priority: Linear distance factor from 0 to MAX_STREAMINGDISTANCE. */
 #define MAX_STREAMINGDISTANCE	10000.0f
@@ -440,33 +439,54 @@ struct FStreamingTexture
 	}
 
 	/**
-	 * Calculates a priority value for the texture. Higher value means more important.
+	 * Calculates a retention priority value for the textures. Higher value means more important.
+	 * Retention is used determine which texture to drop when out of texture pool memory.
+	 * The retention value must be stable, meaning it must not depend on the current state of the streamer.
+	 * Not doing so could make the streamer go into a loop where is never stops dropping and loading different textures when out of budget.
 	 * @return		Priority value
 	 */
-	float CalcPriority( )
+	float CalcRetentionPriority( )
 	{
-		if (LODGroup != TEXTUREGROUP_Terrain_Heightmap)
-		{
-	//		float DistanceFactor = 1.0f - FMath::Clamp( MinDistance, 1.0f, MAX_STREAMINGDISTANCE ) / MAX_STREAMINGDISTANCE;
-			float DistanceFactor = 1 - FMath::Sqrt( FMath::Clamp( MinDistance, 1.0f, MAX_STREAMINGDISTANCE ) / MAX_STREAMINGDISTANCE );
-	//		float MipFactor = WantedMips - ResidentMips;
-			float MipFactor = float(WantedMips) / float(MAX_TEXTURE_MIP_COUNT);
-			float TimeFactor = GStreamWithTimeFactor ? (FMath::Clamp( LastRenderTime, 1.0f, MAX_LASTRENDERTIME ) / MAX_LASTRENDERTIME) : 0.0f;
-			float Priority = MipFactor + DistanceFactor * (1.0f - TimeFactor * 0.5f) + bForceFullyLoad * 100.0f;
+		bool bCanBeDropped = (LODGroup != TEXTUREGROUP_Terrain_Heightmap) && WantedMips > MinAllowedMips && !bForceFullyLoad;
 
-			// Custom priority control for HLOD textures.
-			if (bIsHLODTexture)
-			{
-				Priority *= CVarStreamingHLODPriorityScale.GetValueOnAnyThread();
-			}
+		bool bIsVisible = GStreamWithTimeFactor && LastRenderTime < 5.f && LastRenderTime != MAX_FLT;
 
-			return Priority;
-		}
-		else
-		{
-			// Because heightmaps are technically geometry, it is absolutely imperative that they are streamed in
-			return 1000.0f;
-		}
+		// Usually HLODTexture and Lightmaps are already low rez.
+		bool bAlreadyLowRez = /*(PixelPerTexel > 1.5 && Distance > 1) ||*/ bIsHLODTexture || bIsLightmap;
+		
+		// Don't consider dropping mips that give nothing. This is because streaming time has a high cost per texture.
+		// Dropping many textures for not so much memory would affect the time it takes the streamer to recover.
+		// We also don't want to prioritize by size, since bigger wanted mips are expected to have move visual impact.
+		int32 Threshold = 256;
+		bool bGivesOffMemory = WantedMips > MinAllowedMips && GetSize(WantedMips) - GetSize(WantedMips - 1) >= 256 * 1024; 
+
+		float Priority = 0;
+		if (!bCanBeDropped)		Priority += 1024.f;
+		if (bAlreadyLowRez)		Priority += 512.f;
+		if (!bGivesOffMemory)	Priority += 256.f;
+		if (bIsVisible)			Priority += 128.f;
+
+		return Priority;
+	}
+
+	/**
+	 * Calculates a load order priority value for the texture. Higher value means more important.
+	 * Load Order can depend on the current state of resident mips, because it will not affect the streamer stability.
+	 * @return		Priority value
+	 */
+	float CalcLoadOrderPriority()
+	{
+		bool bMustLoadFirst = LODGroup == TEXTUREGROUP_Terrain_Heightmap;
+		bool bIsVisible = GStreamWithTimeFactor && LastRenderTime < 5.f && LastRenderTime != MAX_FLT;
+		bool bFarFromTargetMips = WantedMips - ResidentMips > 2.f;
+		bool bBigOnScreen = WantedMips > 9.f || bIsHLODTexture;
+
+		float Priority = 0;
+		if (bMustLoadFirst)		Priority += 1024.f;
+		if (bIsVisible)			Priority += 512.f;
+		if (bFarFromTargetMips)	Priority += 256.f;
+		if (bBigOnScreen)		Priority += 128.f;
+		return Priority;
 	}
 
 	/**
@@ -1126,7 +1146,8 @@ public:
 		bAbort = false;
 		ThreadContext.Reset( false, NULL, bCollectTextureStats );
 		ThreadStats.Reset();
-		PrioritizedTextures.Empty( StreamingManager.StreamingTextures.Num() );
+		PrioritizedTexturesByRetention.Empty( StreamingManager.StreamingTextures.Num() );
+		PrioritizedTexturesByLoadOrder.Empty( StreamingManager.StreamingTextures.Num() );
 	}
 
 	/** Notifies the async work that it should abort the thread ASAP. */
@@ -1179,9 +1200,15 @@ public:
 	};
 
 	/** Returns the resulting priorities, matching the FStreamingManagerTexture::StreamingTextures array. */
-	const TArray<FTexturePriority>& GetPrioritizedTextures() const
+	const TArray<FTexturePriority>& GetPrioritizedTexturesByRentention() const
 	{
-		return PrioritizedTextures;
+		return PrioritizedTexturesByRetention;
+	}
+
+		/** Returns the resulting priorities, matching the FStreamingManagerTexture::StreamingTextures array. */
+	const TArray<FTexturePriority>& GetPrioritizedTexturesByLoadOrder() const
+	{
+		return PrioritizedTexturesByLoadOrder;
 	}
 
 	/** Returns the context (temporary info) used for the async work. */
@@ -1202,7 +1229,8 @@ private:
 	void DoWork()
 	{
 		DECLARE_SCOPE_CYCLE_COUNTER(TEXT("FAsyncTextureStreaming::DoWork"), STAT_AsyncTextureStreaming_DoWork, STATGROUP_StreamingDetails);
-		PrioritizedTextures.Empty( StreamingManager.StreamingTextures.Num() );
+		PrioritizedTexturesByRetention.Empty( StreamingManager.StreamingTextures.Num() );
+		PrioritizedTexturesByLoadOrder.Empty( StreamingManager.StreamingTextures.Num() );
 
 		// Calculate DynamicWantedMips and DynamicMinDistanceSq for all dynamic textures.
 		//@TODO: This is not thread-safe because it looks up UTexture2D to get to the FStreamingTexture...
@@ -1240,12 +1268,13 @@ private:
 				bool bTrackedTexture = TrackTextureEvent( &StreamingTexture, StreamingTexture.Texture, StreamingTexture.bForceFullyLoad, &StreamingManager );
 
 				// Add to sort list, if it wants to stream in or could potentially stream out.
-				if ( StreamingTexture.WantedMips > StreamingTexture.ResidentMips || StreamingTexture.ResidentMips > StreamingTexture.MinAllowedMips )
+				if ( StreamingTexture.WantedMips != StreamingTexture.MinAllowedMips || StreamingTexture.RequestedMips != StreamingTexture.MinAllowedMips || StreamingTexture.ResidentMips != StreamingTexture.MinAllowedMips)
 				{
 					const UTexture2D* Texture = StreamingTexture.Texture; // This needs to go on the stack in case the main thread updates it at the same time (for the condition consistency).
 					if (Texture)
 					{
-						FTexturePriority* TexturePriority = new (PrioritizedTextures) FTexturePriority( StreamingTexture.CalcPriority(), Index, Texture );
+						new (PrioritizedTexturesByRetention) FTexturePriority( StreamingTexture.CalcRetentionPriority(), Index, Texture );
+						new (PrioritizedTexturesByLoadOrder) FTexturePriority( StreamingTexture.CalcLoadOrderPriority(), Index, Texture );
 					}
 				}
 
@@ -1302,7 +1331,8 @@ private:
 				return false;
 			}
 		};
-		PrioritizedTextures.Sort( FCompareTexturePriority() );
+		PrioritizedTexturesByRetention.Sort( FCompareTexturePriority() );
+		PrioritizedTexturesByLoadOrder.Sort( FCompareTexturePriority() );
 	}
 
 	FORCEINLINE TStatId GetStatId() const
@@ -1312,8 +1342,10 @@ private:
 
 	/** Reference to the owning streaming manager, for accessing the thread-safe data. */
 	FStreamingManagerTexture&	StreamingManager;
-	/** Resulting priorities, matching the FStreamingManagerTexture::StreamingTextures array. */
-	TArray<FTexturePriority>	PrioritizedTextures;
+	/** Resulting priorities for keeping textures in memory, matching the FStreamingManagerTexture::StreamingTextures array. */
+	TArray<FTexturePriority>	PrioritizedTexturesByRetention;
+	/** Resulting priorities for loading order, matching the FStreamingManagerTexture::StreamingTextures array. */
+	TArray<FTexturePriority>	PrioritizedTexturesByLoadOrder;
 	/** Context (temporary info) used for the async work. */
 	FStreamingContext			ThreadContext;
 	/** Thread statistics. */
@@ -3414,7 +3446,7 @@ void FStreamingManagerTexture::StreamTextures( bool bProcessEverything )
 	// Setup a context for this run.
 	FStreamingContext Context( bProcessEverything, IndividualStreamingTexture, bCollectTextureStats );
 
-	const TArray<FTexturePriority>& PrioritizedTextures = AsyncWork->GetTask().GetPrioritizedTextures();
+	const TArray<FTexturePriority>& PrioritizedTextures = AsyncWork->GetTask().GetPrioritizedTexturesByLoadOrder();
 	FAsyncTextureStreaming::FAsyncStats ThreadStats = AsyncWork->GetTask().GetStats();
 	Context.AddStats( AsyncWork->GetTask().GetContext() );
 
@@ -3542,6 +3574,49 @@ void FStreamingManagerTexture::StreamTextures( bool bProcessEverything )
 
 	if ( Stats.IsUsingLimitedPoolSize() && !bPauseTextureStreaming )
 	{
+		bool bAllowStreamout = true;
+
+		// Check if the current wanted mips fit within the budget. Otherwise drop on wanted mip for some texture.
+		// This must be done before the next "while" loop because the retention texture lists is not sorted in the same order.
+		if (CVarStreamingAllowRetentionLogic.GetValueOnGameThread())
+		{
+			const TArray<FTexturePriority>& PrioritizedTexturesByRetention = AsyncWork->GetTask().GetPrioritizedTexturesByRentention();
+			int64 ExpectedAvailableLater = AvailableLater;
+
+			// At this point AvailableLater represents the available space when all requested mips will be done.
+			for (int32 RetentionIndex = 0; RetentionIndex < PrioritizedTexturesByRetention.Num(); ++RetentionIndex)
+			{
+				const FTexturePriority& TexturePriority = PrioritizedTexturesByRetention[ RetentionIndex ];
+				if (TexturePriority.TextureIndex < StreamingTextures.Num())
+				{
+					FStreamingTexture& StreamingTexture = StreamingTextures[TexturePriority.TextureIndex];
+					if (TexturePriority.Texture == StreamingTexture.Texture && StreamingTexture.WantedMips != StreamingTexture.RequestedMips)
+					{
+						ExpectedAvailableLater -= StreamingTexture.GetSize( StreamingTexture.WantedMips ) - StreamingTexture.GetSize( StreamingTexture.RequestedMips );
+					}
+				}
+			}
+			// Here ExpectedAvailableLater contains the memory cost for have all WantedMips, but also only the WantedMips.
+			// To have this memory, StreamoutTextures( StreamOut_UnwantedMips ) would be required.
+			// If the streamout is not enough, then we need to drop the last mips for some of the textures.
+			for (int32 RetentionIndex = PrioritizedTexturesByRetention.Num() - 1; RetentionIndex > 0 && ExpectedAvailableLater < 0; --RetentionIndex)
+			{
+				const FTexturePriority& TexturePriority = PrioritizedTexturesByRetention[ RetentionIndex ];
+				if (TexturePriority.TextureIndex < StreamingTextures.Num())
+				{
+					FStreamingTexture& StreamingTexture = StreamingTextures[TexturePriority.TextureIndex];
+					if (TexturePriority.Texture == StreamingTexture.Texture && StreamingTexture.WantedMips > StreamingTexture.MinAllowedMips)
+					{
+						// Drop last mip.
+						ExpectedAvailableLater += StreamingTexture.GetSize( StreamingTexture.WantedMips ) - StreamingTexture.GetSize( StreamingTexture.WantedMips - 1);
+						--StreamingTexture.WantedMips;
+					}
+				}
+			}
+
+			bAllowStreamout = ExpectedAvailableLater < 0;
+		}
+
 		// Stream in high-priority textures if possible, and stream out as many low-priority textures as we need.
 		// Keep track of how much temp memory we use separately, not letting it go over a certain threshold.
 		int32 HighPrioIndex = 0;
@@ -3605,7 +3680,7 @@ void FStreamingManagerTexture::StreamTextures( bool bProcessEverything )
 			}
 
 			// Do we STILL need to stream something out?
-			if ( AvailableLater < 0 )
+			if ( AvailableLater < 0  && bAllowStreamout)
 			{
 				// Try streaming out all mips, from LowPrioIndex to HighPrioIndex.
 				// Note: AvailableLater, TempMemoryUsed, LowPriorityIndex and StreamingRequests may be altered by this call.
@@ -4216,7 +4291,7 @@ void FStreamingManagerTexture::UpdateFrameStats( FStreamingContext& Context, FSt
 		FString TextureName = StreamingTexture.Texture->GetFullName();
 		if ( CollectTextureStatsName.Len() == 0 || TextureName.Contains( CollectTextureStatsName) )
 		{
-			new (Context.TextureStats) FTextureStreamingStats( StreamingTexture.Texture, StreamType, StreamingTexture.ResidentMips, PerfectWantedMips, StreamingTexture.MostResidentMips, ResidentSize, PerfectWantedSize, PotentialSize, MostResidentSize, StreamingTexture.BoostFactor, StreamingTexture.CalcPriority(), TextureIndex );
+			new (Context.TextureStats) FTextureStreamingStats( StreamingTexture.Texture, StreamType, StreamingTexture.ResidentMips, PerfectWantedMips, StreamingTexture.MostResidentMips, ResidentSize, PerfectWantedSize, PotentialSize, MostResidentSize, StreamingTexture.BoostFactor, StreamingTexture.CalcLoadOrderPriority(), TextureIndex );
 		}
 	}
 	Context.ThisFrameNumStreamingTextures++;
@@ -5221,7 +5296,8 @@ void FStreamingManagerTexture::InvestigateTexture( const FString& InInvestigateT
 			UE_LOG(LogContentStreaming, Log,  TEXT("  Current size:    %dx%d"), Texture2D->PlatformData->Mips[CurrentMipIndex].SizeX, Texture2D->PlatformData->Mips[CurrentMipIndex].SizeY );
 			UE_LOG(LogContentStreaming, Log,  TEXT("  Wanted size:     %dx%d"), Texture2D->PlatformData->Mips[WantedMipIndex].SizeX, Texture2D->PlatformData->Mips[WantedMipIndex].SizeY );
 			UE_LOG(LogContentStreaming, Log,  TEXT("  Max size:        %dx%d"), Texture2D->PlatformData->Mips[MaxMipIndex].SizeX, Texture2D->PlatformData->Mips[MaxMipIndex].SizeY );
-			UE_LOG(LogContentStreaming, Log,  TEXT("  Priority:        %.3f"), StreamingTexture.CalcPriority() );
+			UE_LOG(LogContentStreaming, Log,  TEXT("  LoadOrder Priority: %.3f"), StreamingTexture.CalcLoadOrderPriority() );
+			UE_LOG(LogContentStreaming, Log,  TEXT("  Retention Priority: %.3f"), StreamingTexture.CalcRetentionPriority() );
 			UE_LOG(LogContentStreaming, Log,  TEXT("  Boost factor:    %.1f"), StreamingTexture.BoostFactor );
 			UE_LOG(LogContentStreaming, Log,  TEXT("  Allowed mips:    %d-%d"), StreamingTexture.MinAllowedMips, StreamingTexture.MaxAllowedMips );
 			UE_LOG(LogContentStreaming, Log,  TEXT("  Global mip bias: %.1f"), ThreadSettings.MipBias );
diff --git a/Engine/Source/Runtime/Engine/Private/DataReplication.cpp b/Engine/Source/Runtime/Engine/Private/DataReplication.cpp
--- a/Engine/Source/Runtime/Engine/Private/DataReplication.cpp
+++ b/Engine/Source/Runtime/Engine/Private/DataReplication.cpp
@@ -831,6 +831,7 @@ void FObjectReplicator::ReplicateCustomDeltaProperties( FOutBunch & Bunch, FRepl
 	ConditionMap[COND_InitialOrOwner] = bIsInitial || bIsOwner;
 	ConditionMap[COND_Custom] = true;
 	ConditionMap[COND_ReplayOrOwner] = bIsReplay || bIsOwner;
+	ConditionMap[COND_ReplayOnly] = bIsReplay;
 
 	// Replicate those properties.
 	for ( int32 i = 0; i < LifetimeCustomDeltaProperties.Num(); i++ )
diff --git a/Engine/Source/Runtime/Engine/Private/DialogueTypes.cpp b/Engine/Source/Runtime/Engine/Private/DialogueTypes.cpp
--- a/Engine/Source/Runtime/Engine/Private/DialogueTypes.cpp
+++ b/Engine/Source/Runtime/Engine/Private/DialogueTypes.cpp
@@ -44,7 +44,7 @@ FString FDialogueContext::GetLocalizationKey() const
 	return Key;
 }
 
-FDialogueWaveParameter::FDialogueWaveParameter() : DialogueWave(NULL)
+FDialogueWaveParameter::FDialogueWaveParameter()
 {
 
 }
diff --git a/Engine/Source/Runtime/Engine/Private/DialogueWave.cpp b/Engine/Source/Runtime/Engine/Private/DialogueWave.cpp
--- a/Engine/Source/Runtime/Engine/Private/DialogueWave.cpp
+++ b/Engine/Source/Runtime/Engine/Private/DialogueWave.cpp
@@ -632,6 +632,15 @@ FString UDialogueWave::GetContextLocalizationKey( const FDialogueContext& Contex
 
 // End UDialogueWave interface.
 
+void UDialogueWave::UpdateContext(FDialogueContextMapping& ContextMapping, USoundWave* SoundWave, UDialogueVoice* Speaker, const TArray<UDialogueVoice*>& Targets)
+{
+	ContextMapping.SoundWave = SoundWave;
+	ContextMapping.Context.Speaker = Speaker;
+	ContextMapping.Context.Targets = Targets;
+
+	UpdateMappingProxy(ContextMapping);
+}
+
 void UDialogueWave::UpdateMappingProxy(FDialogueContextMapping& ContextMapping)
 {
 	if (ContextMapping.SoundWave)
diff --git a/Engine/Source/Runtime/Engine/Private/GameUserSettings.cpp b/Engine/Source/Runtime/Engine/Private/GameUserSettings.cpp
--- a/Engine/Source/Runtime/Engine/Private/GameUserSettings.cpp
+++ b/Engine/Source/Runtime/Engine/Private/GameUserSettings.cpp
@@ -58,9 +58,12 @@ FIntPoint UGameUserSettings::GetDesktopResolution() const
 
 void UGameUserSettings::SetScreenResolution(FIntPoint Resolution)
 {
-	ResolutionSizeX = Resolution.X;
-	ResolutionSizeY = Resolution.Y;
-	UpdateMinResolutionScaling();
+	if (ResolutionSizeX != Resolution.X || ResolutionSizeY != Resolution.Y)
+	{
+		ResolutionSizeX = Resolution.X;
+		ResolutionSizeY = Resolution.Y;
+		UpdateResolutionQuality();
+	}
 }
 
 EWindowMode::Type UGameUserSettings::GetFullscreenMode() const
@@ -92,7 +95,7 @@ void UGameUserSettings::SetFullscreenMode( EWindowMode::Type InFullscreenMode )
 			break;
 	}
 
-	UpdateMinResolutionScaling();
+	UpdateResolutionQuality();
 }
 
 void UGameUserSettings::SetVSyncEnabled( bool bEnable )
@@ -167,6 +170,7 @@ void UGameUserSettings::SetToDefaults()
 	FullscreenMode = GetDefaultWindowMode();
 	FrameRateLimit = 0.0f;
 	MinResolutionScale = Scalability::MinResolutionScale;
+	DesiredScreenHeight = 720;
 
 	static const auto CVar = IConsoleManager::Get().FindTConsoleVariableDataInt(TEXT("r.FullScreenMode"));
 	PreferredFullscreenMode = CVar->GetValueOnGameThread();
@@ -175,7 +179,7 @@ void UGameUserSettings::SetToDefaults()
 
 	if (!IsRunningDedicatedServer())
 	{
-		UpdateMinResolutionScaling();
+		UpdateResolutionQuality();
 	}
 }
 
@@ -189,12 +193,50 @@ void UGameUserSettings::UpdateVersion()
 	Version = UE_GAMEUSERSETTINGS_VERSION;
 }
 
-void UGameUserSettings::UpdateMinResolutionScaling()
+void UGameUserSettings::UpdateResolutionQuality()
 {
 	const int32 MinHeight = UKismetSystemLibrary::GetMinYResolutionFor3DView();
 	const int32 ScreenHeight = FullscreenMode == EWindowMode::WindowedFullscreen ? GetDesktopResolution().Y : ResolutionSizeY;
 	MinResolutionScale = FMath::Max<int32>(Scalability::MinResolutionScale, ((float)MinHeight / (float)ScreenHeight) * 100);
+	if (DesiredScreenHeight > ScreenHeight && ScreenHeight > 0)
+	{
+		DesiredScreenHeight = ScreenHeight;
+	}
+	ScalabilityQuality.ResolutionQuality = FindResolutionQualityForScreenHeight(DesiredScreenHeight);
 	ScalabilityQuality.ResolutionQuality = FMath::Max(ScalabilityQuality.ResolutionQuality, MinResolutionScale);
+
+	OnGameUserSettingsUINeedsUpdate.Broadcast();
+}
+
+int32 UGameUserSettings::FindResolutionQualityForScreenHeight(int32 ScreenHeight)
+{
+	int32 ResolutionQuality = 100;
+
+	const FIntPoint ScreenSize = FullscreenMode == EWindowMode::WindowedFullscreen ? GetDesktopResolution() : FIntPoint(ResolutionSizeX, ResolutionSizeY);
+	const float AspectRatio = (float)ScreenSize.X / (float)ScreenSize.Y;
+	const float Ratio16by9 = 16.0f / 9.0f;
+
+	// ScreenHeight values are based on 16:9 screens. If the screen aspect ratio is different, we need to adjust the target height
+	if (!FMath::IsNearlyEqual(AspectRatio, Ratio16by9))
+	{
+		if (AspectRatio < Ratio16by9)
+		{
+			// For smaller aspect ratios (16:10, 4:3) we allow more vertical space so that the screen width matches the width 16:9 mode would have
+			ScreenHeight = (ScreenHeight * Ratio16by9) / AspectRatio;
+		}
+		else
+		{
+			// For wider screens we try to choose a screen size that'd have similar total number of pixels as the target 16:9 mode
+			ScreenHeight = FMath::Sqrt((ScreenHeight * ScreenHeight * Ratio16by9) / AspectRatio);
+		}
+	}
+
+	if (ScreenHeight < ScreenSize.Y)
+	{
+		ResolutionQuality = ((float)ScreenHeight / (float)ScreenSize.Y) * 100;
+	}
+
+	return ResolutionQuality;
 }
 
 void UGameUserSettings::SetPreferredFullscreenMode(int32 Mode)
@@ -447,7 +489,7 @@ void UGameUserSettings::ResetToCurrentSettings()
 		// Reset the quality settings to the current levels
 		ScalabilityQuality = Scalability::GetQualityLevels();
 
-		UpdateMinResolutionScaling();
+		UpdateResolutionQuality();
 	}
 }
 
@@ -505,6 +547,8 @@ void UGameUserSettings::GetResolutionScaleInformation(float& CurrentScaleNormali
 void UGameUserSettings::SetResolutionScaleValue(int32 NewScaleValue)
 {
 	ScalabilityQuality.ResolutionQuality = FMath::Clamp(NewScaleValue, MinResolutionScale, Scalability::MaxResolutionScale);
+	const int32 ScreenHeight = FullscreenMode == EWindowMode::WindowedFullscreen ? GetDesktopResolution().Y : ResolutionSizeY;
+	DesiredScreenHeight = ScreenHeight * ScalabilityQuality.ResolutionQuality / 100;
 }
 
 void UGameUserSettings::SetResolutionScaleNormalized(float NewScaleNormalized)
diff --git a/Engine/Source/Runtime/Engine/Private/HUD.cpp b/Engine/Source/Runtime/Engine/Private/HUD.cpp
--- a/Engine/Source/Runtime/Engine/Private/HUD.cpp
+++ b/Engine/Source/Runtime/Engine/Private/HUD.cpp
@@ -20,6 +20,21 @@ DEFINE_LOG_CATEGORY_STATIC(LogHUD, Log, All);
 
 FOnShowDebugInfo AHUD::OnShowDebugInfo;
 
+// How opaque should the safe zone visualization be?
+TAutoConsoleVariable<int32> GSafeZoneVisualizationModeCVar(
+	TEXT("r.DebugSafeZone.Mode"),
+	0,
+	TEXT("The safe zone visualization mode (0..2)\n")
+	TEXT(" 0: Disabled (default)\n")
+	TEXT(" 1: Show Title Safe Zone\n")
+	TEXT(" 2: Show Action Safe Zone"));
+
+// How opaque should the safe zone visualization be?
+TAutoConsoleVariable<float> GSafeZoneVisualizationAlphaCVar(
+	TEXT("r.DebugSafeZone.OverlayAlpha"),
+	0.2f,
+	TEXT("The alpha value of the safe zone overlay (0..1)\n")
+	TEXT(" default: 0.2"));
 
 AHUD::AHUD(const FObjectInitializer& ObjectInitializer)
 	: Super(ObjectInitializer)
@@ -197,6 +212,8 @@ void AHUD::PostRender()
 		RenderHitBoxes( Canvas->Canvas );
 	}
 
+	DrawSafeZoneOverlay();
+
 	LastHUDRenderTime = GetWorld()->TimeSeconds;
 }
 
@@ -219,6 +236,52 @@ void AHUD::DrawActorOverlays(FVector Viewpoint, FRotator ViewRotation)
 	}
 }
 
+void AHUD::DrawSafeZoneOverlay()
+{
+	const int32 DebugSafeZoneMode = GSafeZoneVisualizationModeCVar.GetValueOnGameThread();
+
+	if ((DebugSafeZoneMode > 0) && (DebugCanvas != nullptr))
+	{
+		FDisplayMetrics Metrics;
+		FSlateApplication::Get().GetDisplayMetrics(Metrics);
+
+		const FMargin SafeMargin = (DebugSafeZoneMode == 1) ?
+			FMargin(Metrics.TitleSafePaddingSize.X, Metrics.TitleSafePaddingSize.Y) :
+			FMargin(Metrics.ActionSafePaddingSize.X, Metrics.ActionSafePaddingSize.Y);
+
+		const float UnsafeZoneAlpha = GSafeZoneVisualizationAlphaCVar.GetValueOnGameThread();
+		const FLinearColor UnsafeZoneColor(1.0f, 0.5f, 0.5f, UnsafeZoneAlpha);
+
+		const float Width = DebugCanvas->SizeX;
+		const float Height = DebugCanvas->SizeY;
+
+		const float HeightOfSides = Height - SafeMargin.GetTotalSpaceAlong<Orient_Vertical>();
+
+		FCanvasTileItem TileItem(FVector2D::ZeroVector, GWhiteTexture, UnsafeZoneColor);
+		TileItem.BlendMode = SE_BLEND_Translucent;
+
+		// Top bar
+		TileItem.Position = FVector2D::ZeroVector;
+		TileItem.Size = FVector2D(Width, SafeMargin.Top);
+		DebugCanvas->DrawItem(TileItem);
+
+		// Bottom bar
+		TileItem.Position = FVector2D(0.0f, Height - SafeMargin.Bottom);
+		TileItem.Size = FVector2D(Width, SafeMargin.Bottom);
+		DebugCanvas->DrawItem(TileItem);
+
+		// Left bar
+		TileItem.Position = FVector2D(0.0f, SafeMargin.Top);
+		TileItem.Size = FVector2D(SafeMargin.Left, HeightOfSides);
+		DebugCanvas->DrawItem(TileItem);
+
+		// Right bar
+		TileItem.Position = FVector2D(Width - SafeMargin.Right, SafeMargin.Top);
+		TileItem.Size = FVector2D(SafeMargin.Right, HeightOfSides);
+		DebugCanvas->DrawItem(TileItem);
+	}
+}
+
 void AHUD::RemovePostRenderedActor(AActor* A)
 {
 	for ( int32 i=0; i<PostRenderedActors.Num(); i++ )
diff --git a/Engine/Source/Runtime/Engine/Private/KismetStringLibrary.cpp b/Engine/Source/Runtime/Engine/Private/KismetStringLibrary.cpp
--- a/Engine/Source/Runtime/Engine/Private/KismetStringLibrary.cpp
+++ b/Engine/Source/Runtime/Engine/Private/KismetStringLibrary.cpp
@@ -285,6 +285,11 @@ int32 UKismetStringLibrary::FindSubstring(const FString& SearchIn, const FString
 	return SearchIn.Find(Substring, Case, Dir, StartPosition);
 }
 
+bool UKismetStringLibrary::Contains(const FString& SearchIn, const FString& Substring, bool bUseCase, bool bSearchFromEnd)
+{
+	return FindSubstring(SearchIn, Substring, bUseCase, bSearchFromEnd) != -1;
+}
+
 int32 UKismetStringLibrary::GetCharacterAsNumber(const FString& SourceString, int32 Index)
 {
 	if ((Index >= 0) && (Index < SourceString.Len()))
diff --git a/Engine/Source/Runtime/Engine/Private/Materials/MaterialShader.cpp b/Engine/Source/Runtime/Engine/Private/Materials/MaterialShader.cpp
--- a/Engine/Source/Runtime/Engine/Private/Materials/MaterialShader.cpp
+++ b/Engine/Source/Runtime/Engine/Private/Materials/MaterialShader.cpp
@@ -16,144 +16,6 @@ static FAutoConsoleVariableRef CVarCreateShadersOnLoad(
 	TEXT("Whether to create shaders on load, which can reduce hitching, but use more memory.  Otherwise they will be created as needed.")
 	);
 
-
-
-//
-// Cooking Stats
-// 
-#define MATERIAL_COOKING_STATS (1 && WITH_EDITORONLY_DATA)
-
-#if MATERIAL_COOKING_STATS
-
-#include "CookingStatsModule.h"
-#include "ModuleManager.h"
-
-
-static FString GetMaterialShaderMapKeyString(const FMaterialShaderMapId& ShaderMapId, EShaderPlatform Platform);
-
-namespace MaterialStats
-{
-
-	ICookingStats* GetCookingStats()
-	{
-		static ICookingStats* CookingStats = nullptr;
-		static bool bInitialized = false;
-		if (bInitialized == false)
-		{
-			FCookingStatsModule* CookingStatsModule = FModuleManager::LoadModulePtr<FCookingStatsModule>(TEXT("CookingStats"));
-			if (CookingStatsModule)
-			{
-				CookingStats = CookingStatsModule->Get();
-			}
-			bInitialized = true;
-		}
-		return CookingStats;
-	}
-
-	FCriticalSection CookingStatsSyncObject;
-
-	struct FMaterialTimingInfo
-	{
-	public:
-		FName TransactionGuid;
-		TMap<FName, double> TagStartTime;
-	};
-
-	TMap<FMaterialShaderMapId, FMaterialTimingInfo> CurrentCookingStats;
-	int32 CurrentIndex;
-	FName TransactionGuid;
-
-	FName GenerateNewTransactionName()
-	{
-		FScopeLock ScopeSync(&CookingStatsSyncObject);
-		++CurrentIndex;
-		if (CurrentIndex <= 1)
-		{
-			FString TransactionGuidString = TEXT("MaterialTransactionId");
-			TransactionGuidString += FGuid::NewGuid().ToString();
-			TransactionGuid = FName(*TransactionGuidString);
-			CurrentIndex = 1;
-		}
-		check(TransactionGuid != NAME_None);
-
-		return FName(TransactionGuid, CurrentIndex);
-	}
-
-	void StartStat(const FMaterialShaderMapId& ShaderMapId, const EShaderPlatform Platform, const FName& TagName)
-	{
-		ICookingStats* CookingStats = GetCookingStats();
-		if (CookingStats)
-		{
-			FScopeLock ScopeSync(&CookingStatsSyncObject);
-
-			FMaterialTimingInfo* TimingInfo = CurrentCookingStats.Find(ShaderMapId);
-
-			if (TimingInfo == nullptr)
-			{
-				TimingInfo = &CurrentCookingStats.Add(ShaderMapId);
-				TimingInfo->TransactionGuid = GenerateNewTransactionName();
-				static const FName NAME_CacheKey(TEXT("CacheKey"));
-
-				FString CacheKeyString = GetMaterialShaderMapKeyString(ShaderMapId, Platform);
-
-				CookingStats->AddTagValue(TimingInfo->TransactionGuid, NAME_CacheKey, CacheKeyString);
-			}
-
-			double* StartTime = TimingInfo->TagStartTime.Find(TagName);
-			if (!StartTime)
-			{
-				TimingInfo->TagStartTime.Add(TagName, FPlatformTime::Seconds());
-			}
-		}
-	}
-
-	void StopStat(const FMaterialShaderMapId& ShaderMapId, const FName& TagName)
-	{
-		ICookingStats* CookingStats = GetCookingStats();
-		if (CookingStats)
-		{
-			FScopeLock ScopeSync(&CookingStatsSyncObject);
-
-			FMaterialTimingInfo* TimingInfo = CurrentCookingStats.Find(ShaderMapId);
-			if (TimingInfo == nullptr)
-			{
-				UE_LOG(LogShaders, Warning, TEXT("Stat tag wasn't found for stat %s, probably caused by calling COOKING_STAT_STOP before COOKING_STAT_START"), *TagName.ToString());
-				return;
-			}
-			check(TimingInfo); // did you call stop stat before calling start
-			double* StartTime = TimingInfo->TagStartTime.Find(TagName);
-			if (StartTime == nullptr)
-			{
-				UE_LOG(LogShaders, Warning, TEXT("Stat start time wasn't found for stat %s, probably caused by calling COOKING_STAT_STOP before COOKING_STAT_START"), *TagName.ToString());
-				return;
-			}
-			check(StartTime); // did you call stop stat before calling start for that tag
-			
-			double Duration = FPlatformTime::Seconds() - *StartTime;
-			TimingInfo->TagStartTime.Remove(TagName);
-
-			CookingStats->AddTagValue(TimingInfo->TransactionGuid, TagName, (float)(Duration*1000.0f));
-		}
-	}
-
-};
-
-#define COOKING_STAT_START(TagName, ShaderMapID, Platform) \
-	static const FName NAME_##TagName(TEXT(#TagName)); \
-	MaterialStats::StartStat(ShaderMapID, Platform, NAME_##TagName);
-
-#define COOKING_STAT_STOP(TagName, ShaderMapID) \
-	static const FName NAME_##TagName(TEXT(#TagName)); \
-	MaterialStats::StopStat(ShaderMapID, NAME_##TagName);
-
-
-#else
-
-#define COOKING_STAT_START(...)
-#define COOKING_STAT_STOP(...)
-
-#endif // MATERIAL_COOKING_STATS
-
 //
 // Globals
 //
@@ -1289,8 +1151,6 @@ void FMaterialShaderMap::Compile(
 		}
 		else
 		{
-			COOKING_STAT_START(ShaderCompilation, InShaderMapId, InPlatform);
-
 			// Assign a unique identifier so that shaders from this shader map can be associated with it after a deferred compile
 			CompilingId = NextCompilingId;
 			check(NextCompilingId < UINT_MAX);
@@ -1723,8 +1583,6 @@ bool FMaterialShaderMap::ProcessCompilationResults(const TArray<FShaderCommonCom
 		// The shader map can now be used on the rendering thread
 		bCompilationFinalized = true;
 
-		COOKING_STAT_STOP(ShaderCompilation, ShaderMapId)
-
 		return true;
 	}
 
diff --git a/Engine/Source/Runtime/Engine/Private/Particles/ParticleBeam2EmitterInstance.cpp b/Engine/Source/Runtime/Engine/Private/Particles/ParticleBeam2EmitterInstance.cpp
--- a/Engine/Source/Runtime/Engine/Private/Particles/ParticleBeam2EmitterInstance.cpp
+++ b/Engine/Source/Runtime/Engine/Private/Particles/ParticleBeam2EmitterInstance.cpp
@@ -334,79 +334,89 @@ void FParticleBeam2EmitterInstance::Tick(float DeltaTime, bool bSuppressSpawning
 		// Handle EmitterTime setup, looping, etc.
 		float EmitterDelay = Tick_EmitterTimeSetup(DeltaTime, LODLevel);
 
-		// Kill before the spawn... Otherwise, we can get 'flashing'
-		KillParticles();
-
-		// If not suppressing spawning...
-		if (!bHaltSpawning && !bSuppressSpawning && (EmitterTime >= 0.0f))
+		if (bEnabled)
 		{
-			if ((LODLevel->RequiredModule->EmitterLoops == 0) || 
-				(LoopCount < LODLevel->RequiredModule->EmitterLoops) ||
-				(SecondsSinceCreation < (EmitterDuration * LODLevel->RequiredModule->EmitterLoops)))
+			// Kill before the spawn... Otherwise, we can get 'flashing'
+			KillParticles();
+
+			// If not suppressing spawning...
+			if (!bHaltSpawning && !bSuppressSpawning && (EmitterTime >= 0.0f))
 			{
-				// For beams, we probably want to ignore the SpawnRate distribution,
-				// and focus strictly on the BurstList...
-				float SpawnRate = 0.0f;
-				// Figure out spawn rate for this tick.
-				SpawnRate = LODLevel->SpawnModule->Rate.GetValue(EmitterTime, Component);
-				// Take Bursts into account as well...
-				int32		Burst		= 0;
-				float	BurstTime	= GetCurrentBurstRateOffset(DeltaTime, Burst);
-				SpawnRate += BurstTime;
-
-				// Spawn new particles...
-
-				//@todo. Fix the issue of 'blanking' beams when the count drops...
-				// This is a temporary hack!
-				const float InvDeltaTime = (DeltaTime > 0.0f) ? 1.0f / DeltaTime : 0.0f;
-				if ((ActiveParticles < BeamCount) && (SpawnRate <= 0.0f))
+				if ((LODLevel->RequiredModule->EmitterLoops == 0) ||
+					(LoopCount < LODLevel->RequiredModule->EmitterLoops) ||
+					(SecondsSinceCreation < (EmitterDuration * LODLevel->RequiredModule->EmitterLoops)))
 				{
-					// Force the spawn of a single beam...
-					SpawnRate = 1.0f * InvDeltaTime;
-				}
+					// For beams, we probably want to ignore the SpawnRate distribution,
+					// and focus strictly on the BurstList...
+					float SpawnRate = 0.0f;
+					// Figure out spawn rate for this tick.
+					SpawnRate = LODLevel->SpawnModule->Rate.GetValue(EmitterTime, Component);
+					// Take Bursts into account as well...
+					int32		Burst = 0;
+					float	BurstTime = GetCurrentBurstRateOffset(DeltaTime, Burst);
+					SpawnRate += BurstTime;
+
+					// Spawn new particles...
+
+					//@todo. Fix the issue of 'blanking' beams when the count drops...
+					// This is a temporary hack!
+					const float InvDeltaTime = (DeltaTime > 0.0f) ? 1.0f / DeltaTime : 0.0f;
+					if ((ActiveParticles < BeamCount) && (SpawnRate <= 0.0f))
+					{
+						// Force the spawn of a single beam...
+						SpawnRate = 1.0f * InvDeltaTime;
+					}
 
-				// Force beams if the emitter is marked "AlwaysOn"
-				if ((ActiveParticles < BeamCount) && BeamTypeData->bAlwaysOn)
-				{
-					Burst		= BeamCount;
-					if (DeltaTime > KINDA_SMALL_NUMBER)
+					// Force beams if the emitter is marked "AlwaysOn"
+					if ((ActiveParticles < BeamCount) && BeamTypeData->bAlwaysOn)
 					{
-						BurstTime	 = Burst * InvDeltaTime;
-						SpawnRate	+= BurstTime;
+						Burst = BeamCount;
+						if (DeltaTime > KINDA_SMALL_NUMBER)
+						{
+							BurstTime = Burst * InvDeltaTime;
+							SpawnRate += BurstTime;
+						}
 					}
-				}
 
-				if (SpawnRate > 0.f)
-				{
-					SpawnFraction = SpawnBeamParticles(SpawnFraction, SpawnRate, DeltaTime, Burst, BurstTime);
+					if (SpawnRate > 0.f)
+					{
+						SpawnFraction = SpawnBeamParticles(SpawnFraction, SpawnRate, DeltaTime, Burst, BurstTime);
+					}
 				}
 			}
-		}
-
-		// Reset particle data
-		ResetParticleParameters(DeltaTime);
+			else if (bFakeBurstsWhenSpawningSupressed)
+			{
+				FakeBursts();
+			}
 
-		// Not really necessary as beams do not LOD at the moment, but for consistency...
-		CurrentMaterial = LODLevel->RequiredModule->Material;
+			// Reset particle data
+			ResetParticleParameters(DeltaTime);
 
-		Tick_ModuleUpdate(DeltaTime, LODLevel);
-		Tick_ModulePostUpdate(DeltaTime, LODLevel);
+			// Not really necessary as beams do not LOD at the moment, but for consistency...
+			CurrentMaterial = LODLevel->RequiredModule->Material;
 
-		// Calculate bounding box and simulate velocity.
-		UpdateBoundingBox(DeltaTime);
+			Tick_ModuleUpdate(DeltaTime, LODLevel);
+			Tick_ModulePostUpdate(DeltaTime, LODLevel);
 
-		if (!bSuppressSpawning)
-		{
-			// Ensure that we flip the 'FirstEmission' flag
-			FirstEmission = false;
-		}
+			// Calculate bounding box and simulate velocity.
+			UpdateBoundingBox(DeltaTime);
 
-		// Invalidate the contents of the vertex/index buffer.
-		IsRenderDataDirty = 1;
+			if (!bSuppressSpawning)
+			{
+				// Ensure that we flip the 'FirstEmission' flag
+				FirstEmission = false;
+			}
 
-		// Bump the tick count
-		TickCount++;
+			// Invalidate the contents of the vertex/index buffer.
+			IsRenderDataDirty = 1;
 
+			// Bump the tick count
+			TickCount++;
+		}
+		else
+		{
+			FakeBursts();
+		}
 		// 'Reset' the emitter time so that the delay functions correctly
 		EmitterTime += CurrentDelay;
 		
@@ -1005,7 +1015,7 @@ FDynamicEmitterDataBase* FParticleBeam2EmitterInstance::GetDynamicData(bool bSel
 	QUICK_SCOPE_CYCLE_COUNTER(STAT_ParticleBeam2EmitterInstance_GetDynamicData);
 
 	UParticleLODLevel* LODLevel = SpriteTemplate->GetCurrentLODLevel(this);
-	if (IsDynamicDataRequired(LODLevel) == false)
+	if (IsDynamicDataRequired(LODLevel) == false || !bEnabled)
 	{
 		return NULL;
 	}
@@ -1069,7 +1079,7 @@ bool FParticleBeam2EmitterInstance::UpdateDynamicData(FDynamicEmitterDataBase* D
 {
 	QUICK_SCOPE_CYCLE_COUNTER(STAT_ParticleBeam2EmitterInstance_UpdateDynamicData);
 
-	if (ActiveParticles <= 0)
+	if (ActiveParticles <= 0 || !bEnabled)
 	{
 		return false;
 	}
diff --git a/Engine/Source/Runtime/Engine/Private/Particles/ParticleComponents.cpp b/Engine/Source/Runtime/Engine/Private/Particles/ParticleComponents.cpp
--- a/Engine/Source/Runtime/Engine/Private/Particles/ParticleComponents.cpp
+++ b/Engine/Source/Runtime/Engine/Private/Particles/ParticleComponents.cpp
@@ -81,6 +81,7 @@ DECLARE_CYCLE_STAT(TEXT("ParticleComponent QueueAsync"), STAT_UParticleSystemCom
 DECLARE_CYCLE_STAT(TEXT("ParticleComponent WaitForAsyncAndFinalize"), STAT_UParticleSystemComponent_WaitForAsyncAndFinalize, STATGROUP_Particles);
 DECLARE_CYCLE_STAT(TEXT("ParticleComponent CreateRenderState Concurrent"), STAT_ParticleSystemComponent_CreateRenderState_Concurrent, STATGROUP_Particles);
 
+#include "InGamePerformanceTracker.h"
 
 #define LOCTEXT_NAMESPACE "ParticleComponents"
 
@@ -1733,6 +1734,12 @@ void UParticleEmitter::CreateStatID() const
 }
 #endif
 
+bool UParticleEmitter::IsSignificant(EParticleSignificanceLevel RequiredSignificance)
+{
+	UParticleSystem* PSysOuter = CastChecked<UParticleSystem>(GetOuter());
+	EParticleSignificanceLevel Significance = FMath::Min(PSysOuter->MaxSignificanceLevel, SignificanceLevel);
+	return Significance >= RequiredSignificance;
+}
 
 /*-----------------------------------------------------------------------------
 	UParticleSpriteEmitter implementation.
@@ -1908,6 +1915,9 @@ void UParticleSpriteEmitter::SetToSensibleDefaults()
 
 UParticleSystem::UParticleSystem(const FObjectInitializer& ObjectInitializer)
 	: Super(ObjectInitializer)
+	, OcclusionBoundsMethod(EPSOBM_ParticleBounds)
+	, HighestSignificance(EParticleSignificanceLevel::Critical)
+	, LowestSignificance(EParticleSignificanceLevel::Low)
 	, bAnyEmitterLoopsForever(false)
 {
 #if WITH_EDITORONLY_DATA
@@ -1942,8 +1952,10 @@ UParticleSystem::UParticleSystem(const FObjectInitializer& ObjectInitializer)
 	MacroUVRadius = 200.0f;
 	bAutoDeactivate = true;
 	MinTimeBetweenTicks = 0;
-	InsignificantTimeTillDeactivate = 5.0f;
+	InsignificantReaction = EParticleSystemInsignificanceReaction::Auto;
+	InsignificanceDelay = 0.0f;
 	MaxSignificanceLevel = EParticleSignificanceLevel::Critical;
+	bShouldManageSignificance = false;
 }
 
 
@@ -2024,6 +2036,8 @@ void UParticleSystem::PostEditChangeProperty(FPropertyChangedEvent& PropertyChan
 
 	// Recompute the looping flag
 	bAnyEmitterLoopsForever = false;
+	HighestSignificance = EParticleSignificanceLevel::Low;
+	LowestSignificance = EParticleSignificanceLevel::Critical;
 	for (UParticleEmitter* Emitter : Emitters)
 	{
 		if (Emitter != nullptr)
@@ -2042,9 +2056,21 @@ void UParticleSystem::PostEditChangeProperty(FPropertyChangedEvent& PropertyChan
 					}
 				}
 			}
+
+			EParticleSignificanceLevel EmitterSignificance = FMath::Min(MaxSignificanceLevel, Emitter->SignificanceLevel);
+			if (EmitterSignificance > HighestSignificance)
+			{
+				HighestSignificance = EmitterSignificance;
+			}
+			if (EmitterSignificance < LowestSignificance)
+			{
+				LowestSignificance = EmitterSignificance;
+			}
 		}
 	}
 
+	bShouldManageSignificance = GetLowestSignificance() != EParticleSignificanceLevel::Critical && !ContainsEmitterType(UParticleModuleTypeDataBeam2::StaticClass());
+
 	//cap the WarmupTickRate to realistic values
 	if (WarmupTickRate <= 0)
 	{
@@ -2092,6 +2118,8 @@ void UParticleSystem::PostLoad()
 	// Run thru all of the emitters, load them up and compute some flags based on them
 	bHasPhysics = false;
 	bAnyEmitterLoopsForever = false;
+	HighestSignificance = EParticleSignificanceLevel::Low;
+	LowestSignificance = EParticleSignificanceLevel::Critical;
 	for (int32 i = Emitters.Num() - 1; i >= 0; i--)
 	{
 		// Remove any old emitters
@@ -2151,9 +2179,21 @@ void UParticleSystem::PostLoad()
 					}
 				}
 			}
+
+			EParticleSignificanceLevel EmitterSignificance = FMath::Min(MaxSignificanceLevel, Emitter->SignificanceLevel);
+			if (EmitterSignificance > HighestSignificance)
+			{
+				HighestSignificance = EmitterSignificance;
+			}
+			if (EmitterSignificance < LowestSignificance)
+			{
+				LowestSignificance = EmitterSignificance;
+			}
 		}
 	}
 
+	bShouldManageSignificance = GetLowestSignificance() != EParticleSignificanceLevel::Critical && !ContainsEmitterType(UParticleModuleTypeDataBeam2::StaticClass());
+	
 	if (LODSettings.Num() == 0)
 	{
 		if (Emitters.Num() > 0)
@@ -2340,6 +2380,20 @@ void UParticleSystem::GetAssetRegistryTags(TArray<FAssetRegistryTag>& OutTags) c
 	OutTags.Add(FAssetRegistryTag("CPUCollision", UsesCPUCollision() ? TEXT("True") : TEXT("False"), FAssetRegistryTag::TT_Alphabetical));
 	OutTags.Add(FAssetRegistryTag("Looping", bAnyEmitterLoopsForever ? TEXT("True") : TEXT("False"), FAssetRegistryTag::TT_Alphabetical));
 	OutTags.Add(FAssetRegistryTag("Immortal", IsPotentiallyImmortal() ? TEXT("True") : TEXT("False"), FAssetRegistryTag::TT_Alphabetical));
+	OutTags.Add(FAssetRegistryTag("CanBeOccluded", OcclusionBoundsMethod == EParticleSystemOcclusionBoundsMethod::EPSOBM_None ? TEXT("False") : TEXT("True"), FAssetRegistryTag::TT_Alphabetical));
+
+	uint32 NumEmittersAtEachSig[(int32)EParticleSignificanceLevel::Num] = { 0, 0, 0, 0 };
+	for (UParticleEmitter* Emitter : Emitters)
+	{
+		if (Emitter)
+		{
+			++NumEmittersAtEachSig[(int32)Emitter->SignificanceLevel];			
+		}
+	}
+	OutTags.Add(FAssetRegistryTag("Critical Emitters", LexicalConversion::ToString(NumEmittersAtEachSig[(int32)EParticleSignificanceLevel::Critical]), FAssetRegistryTag::TT_Numerical));
+	OutTags.Add(FAssetRegistryTag("High Emitters", LexicalConversion::ToString(NumEmittersAtEachSig[(int32)EParticleSignificanceLevel::High]), FAssetRegistryTag::TT_Numerical));
+	OutTags.Add(FAssetRegistryTag("Medium Emitters", LexicalConversion::ToString(NumEmittersAtEachSig[(int32)EParticleSignificanceLevel::Medium]), FAssetRegistryTag::TT_Numerical));
+	OutTags.Add(FAssetRegistryTag("Low Emitters", LexicalConversion::ToString(NumEmittersAtEachSig[(int32)EParticleSignificanceLevel::Low]), FAssetRegistryTag::TT_Numerical));
 
 	Super::GetAssetRegistryTags(OutTags);
 }
@@ -3018,6 +3072,8 @@ bool UParticleSystem::HasGPUEmitter() const
 	return false;
 }
 
+FOnSystemPreActivationChange UParticleSystemComponent::OnSystemPreActivationChange;
+
 UParticleSystemComponent::UParticleSystemComponent(const FObjectInitializer& ObjectInitializer)
 	: Super(ObjectInitializer), FXSystem(NULL), ReleaseResourcesFence(NULL)
 {
@@ -3034,6 +3090,7 @@ UParticleSystemComponent::UParticleSystemComponent(const FObjectInitializer& Obj
 
 	WarmupTime = 0.0f;
 	SecondsBeforeInactive = 1.0f;
+	bAllowSecondsBeforeInactive = true;
 	bIsTransformDirty = false;
 	bSkipUpdateDynamicDataDuringTick = false;
 	bIsViewRelevanceDirty = true;
@@ -3060,104 +3117,129 @@ UParticleSystemComponent::UParticleSystemComponent(const FObjectInitializer& Obj
 	TimeSinceLastTick = 0;
 
 	RequiredSignificance = EParticleSignificanceLevel::Low;
-	bMissingEmittersDueToSignificance = false;
-	InitRequiredSignificanceLevel = EParticleSignificanceLevel::Critical;
 	LastSignificantTime = 0.0f;
+	bIsSignificant = 1;
 }
 
-static float GParticleEmitterSignificanceThreshold_Low = .6666f;
-static FAutoConsoleVariableRef ParticleEmitterSignificanceThreshold_Low(
-	TEXT("fx.Significance.EmitterLevel_Low"),
-	GParticleEmitterSignificanceThreshold_Low,
-	TEXT("The significance value (0.0 to 1.0) to use as the Low Significance Level in particle emitter. \n"),
-	ECVF_Default
-	);
-static float GParticleEmitterSignificanceThreshold_Medium = .3333f;
-static FAutoConsoleVariableRef ParticleEmitterSignificanceThreshold_Medium(
-	TEXT("fx.Significance.EmitterLevel_Medium"),
-	GParticleEmitterSignificanceThreshold_Medium,
-	TEXT("The significance value (0.0 to 1.0) to use as the Medium Significance Level in particle emitter. \n"),
-	ECVF_Default
-	);
-static float GParticleEmitterSignificanceThreshold_High = 0.0f;
-static FAutoConsoleVariableRef ParticleEmitterSignificanceThreshold_High(
-	TEXT("fx.Significance.EmitterLevel_High"),
-	GParticleEmitterSignificanceThreshold_High,
-	TEXT("The significance value (0.0 to 1.0) to use as the High Significance Level in particle emitter. \n"),
-	ECVF_Default
-	);
-
-void UParticleSystemComponent::SetSignificance(float OldSignificance, float NewSignificance)
+void UParticleSystemComponent::SetRequiredSignificance(EParticleSignificanceLevel NewRequiredSignificance)
 {
 	check(Template);
-	check(!bAutoDestroy);//We can't use significance with systems that will auto destroy.
 
-	float SignificanceLevels[(int32)EParticleSignificanceLevel::Num] =
-	{
-		GParticleEmitterSignificanceThreshold_Low/*EParticleRequiredSignificanceLevel::Low*/,
-		GParticleEmitterSignificanceThreshold_Medium/*EParticleRequiredSignificanceLevel::Medium*/,
-		GParticleEmitterSignificanceThreshold_High/*EParticleRequiredSignificanceLevel::High*/,
-		0.0f
-	};
+	RequiredSignificance = NewRequiredSignificance;
 
-	if (NewSignificance > SignificanceLevels[(int32)EParticleSignificanceLevel::Low])
-	{
-		RequiredSignificance = EParticleSignificanceLevel::Low;
-	}
-	else if (NewSignificance > SignificanceLevels[(int32)EParticleSignificanceLevel::Medium])
-	{
-		RequiredSignificance = EParticleSignificanceLevel::Medium;
-	}
-	else if (NewSignificance > SignificanceLevels[(int32)EParticleSignificanceLevel::High])
+	EParticleSystemInsignificanceReaction Reaction = Template->InsignificantReaction;
+	if (Template->InsignificantReaction == EParticleSystemInsignificanceReaction::Auto)
 	{
-		RequiredSignificance = EParticleSignificanceLevel::High;
-	}
-	else
-	{
-		RequiredSignificance = EParticleSignificanceLevel::Critical;
+		Reaction = Template->IsLooping() ? EParticleSystemInsignificanceReaction::DisableTick : EParticleSystemInsignificanceReaction::Complete;
 	}
 
-	bool bNeedsInit = false;
-	// We need to re-init. Some emitters are missing due to significance that may now be significant.
-	if (bIsActive && ((int32)RequiredSignificance < (int32)InitRequiredSignificanceLevel))
+	//If our tick is disabled we need to work out if we should re-enable it based on this new significance
+	if (!IsComponentTickEnabled() && Reaction == EParticleSystemInsignificanceReaction::DisableTick && Template->GetHighestSignificance() >= NewRequiredSignificance)
 	{
-		InitParticles(false);
+		//Set us to be significant again.
+		OnSignificanceChanged(true, true);
 	}
+}
 
-	bool bHasAnyEmittersEnabled = false;
-	for (FParticleEmitterInstance* Inst : EmitterInstances)
+void UParticleSystemComponent::OnSignificanceChanged(bool bSignificant, bool bApplyToEmitters)
+{
+	int32 LocalNumSignificantEmitters = 0;
+	if (bSignificant)
 	{
-		if (Inst)
+		SetComponentTickEnabled(true);
+		bIsSignificant = true;
+		if (bApplyToEmitters && EmitterInstances.Num() > 0)
 		{
-			if ((int32)Inst->SpriteTemplate->SignificanceLevel >= (int32)RequiredSignificance)
+			//Mark any emitters as significant if needed.
+			for (FParticleEmitterInstance* Inst : EmitterInstances)
 			{
-				bHasAnyEmittersEnabled = true;
-				Inst->bEnabled = true;
-				Inst->SetHaltSpawning(false);
-				//TODO: Warm up if was previously disabled emitters? How does this all interact with LODs?
+				if (Inst)					
+				{
+					if (Inst->SpriteTemplate->IsSignificant(RequiredSignificance))
+					{
+						Inst->bEnabled = true;
+						Inst->SetHaltSpawning(false);
+						Inst->SetFakeBurstWhenSpawningSupressed(false);
+						++LocalNumSignificantEmitters;
+					}
+				}
+				else
+				{
+					++LocalNumSignificantEmitters;//Set significant for missing emitters due to other reasons such as detail mode.
+				}
 			}
-			else
+
+			if (LocalNumSignificantEmitters == 0)
 			{
-				Inst->bEnabled = Inst->SpriteTemplate->bDisableWhenInsignficant ? false : true;
-				Inst->SetHaltSpawning(true);
-				//Inst->DisabledSince = GetWorld()->GetTimeSeconds();//Allows us to warm back up if we ever re-enable?
+				UE_LOG(LogParticles, Warning, TEXT("Setting PSC as significant but it has no significant emitters. %s Template: %s"), *GetFullName(), *Template->GetFullName());
 			}
+			NumSignificantEmitters = LocalNumSignificantEmitters;
 		}
 	}
-
-	if (bHasAnyEmittersEnabled)
-	{
-		LastSignificantTime = GetWorld()->GetTimeSeconds();
-		Activate();
-	}
 	else
 	{
-		float CurrentTime = GetWorld()->GetTimeSeconds();
-		if (Template->InsignificantTimeTillDeactivate > 0.0f && CurrentTime >= LastSignificantTime + Template->InsignificantTimeTillDeactivate)
+		SetComponentTickEnabled(false);
+		bIsSignificant = false;
+		if (bApplyToEmitters && EmitterInstances.Num() > 0)
+		{
+			//Mark any emitters as significant if needed.
+			for (FParticleEmitterInstance* Inst : EmitterInstances)
+			{
+				if (Inst)
+				{
+					if(Inst->SpriteTemplate->IsSignificant(RequiredSignificance))
+					{
+						++LocalNumSignificantEmitters;
+					}
+					else
+					{
+						Inst->bEnabled = false;
+						Inst->SetHaltSpawning(true);
+						Inst->SetFakeBurstWhenSpawningSupressed(true);
+					}
+				}
+			}
+
+			if (LocalNumSignificantEmitters > 0)
+			{
+				UE_LOG(LogParticles, Warning, TEXT("Setting PSC as not significant but it has some significant emitters. %s Template: %s"), *GetFullName(), *Template->GetFullName());
+			}
+
+			NumSignificantEmitters = LocalNumSignificantEmitters;
+		}
+
+		EParticleSystemInsignificanceReaction Reaction = Template->InsignificantReaction;
+		if (Template->InsignificantReaction == EParticleSystemInsignificanceReaction::Auto)
+		{
+			Reaction = Template->IsLooping() ? EParticleSystemInsignificanceReaction::DisableTick : EParticleSystemInsignificanceReaction::Complete;
+		}
+
+		switch (Reaction)
+		{
+		case EParticleSystemInsignificanceReaction::Complete:
+		{
+			Complete();
+		}
+			break;
+		case EParticleSystemInsignificanceReaction::DisableTick:
 		{
-			Deactivate();
+			SetComponentTickEnabled(false);
+		}
+			break;
+		case EParticleSystemInsignificanceReaction::DisableTickAndKill:
+		{
+			KillParticlesForced();//TODO: Make this actually free memory.
+			SetComponentTickEnabled(false);
+		}
+			break;
 		}
 	}
+	MarkRenderStateDirty();
+}
+
+bool UParticleSystemComponent::ShouldManageSignificance() const
+{
+	return Template ? Template->ShouldManageSignificance() : false;
 }
 
 #if WITH_EDITOR
@@ -3857,13 +3939,6 @@ void UParticleSystemComponent::UpdateDynamicData(FParticleSystemSceneProxy* Prox
 	{
 		// Create the dynamic data for rendering this particle system
 		FParticleDynamicData* ParticleDynamicData = CreateDynamicData();
-		if (ParticleDynamicData)
-		{
-			for (int32 Index = 0; Index < ParticleDynamicData->DynamicEmitterDataArray.Num(); Index++)
-			{
-				static_cast<FParticleSystemSceneProxy*>(SceneProxy)->QueueVertexFactoryCreation(ParticleDynamicData->DynamicEmitterDataArray[Index]);
-			}
-		}
 
 		// Render the particles
 #if !(UE_BUILD_SHIPPING || UE_BUILD_TEST)
@@ -4149,6 +4224,9 @@ DECLARE_CYCLE_STAT(TEXT("PSys Comp Marshall Time"),STAT_UParticleSystemComponent
 
 void UParticleSystemComponent::TickComponent(float DeltaTime, enum ELevelTick TickType, FActorComponentTickFunction *ThisTickFunction)
 {
+	bool bManagingSignificance = ShouldManageSignificance();
+	FInGameScopedCycleCounter InGameCycleCounter(GetWorld(), EInGamePerfTrackers::VFXSignificance, EInGamePerfTrackerThreads::GameThread, bManagingSignificance);
+
 	if( Template == nullptr)
 	{
 		return;
@@ -4255,7 +4333,7 @@ void UParticleSystemComponent::TickComponent(float DeltaTime, enum ELevelTick Ti
 	{
 		//For now, we're only allowing the SecondsBeforeInactive optimization on looping emitters as it can cause leaks with non-looping effects.
 		//Longer term, there is likely a better solution.
-		if (Template->IsLooping())
+		if (bAllowSecondsBeforeInactive && Template->IsLooping())
 		{
 			const float MaxSecondsBeforeInactive = FMath::Max(SecondsBeforeInactive, Template->SecondsBeforeInactive);
 
@@ -4442,10 +4520,14 @@ int32 UParticleSystemComponent::GetCurrentDetailMode() const
 
 void UParticleSystemComponent::ComputeTickComponent_Concurrent()
 {
+	bool bLocalShouldManageSignificance = ShouldManageSignificance();
+	FInGameScopedCycleCounter InGameCycleCounter(GetWorld(), EInGamePerfTrackers::VFXSignificance, IsInGameThread() ? EInGamePerfTrackerThreads::GameThread : EInGamePerfTrackerThreads::OtherThread, bLocalShouldManageSignificance);
+
 	SCOPE_CYCLE_COUNTER(STAT_ParticleComputeTickTime);
 	FScopeCycleCounterUObject AdditionalScope(AdditionalStatObject(), GET_STATID(STAT_ParticleComputeTickTime));
 	// Tick Subemitters.
 	int32 EmitterIndex;
+	NumSignificantEmitters = 0;
 	for (EmitterIndex = 0; EmitterIndex < EmitterInstances.Num(); EmitterIndex++)
 	{
 		FParticleEmitterInstance* Instance = EmitterInstances[EmitterIndex];
@@ -4457,13 +4539,38 @@ void UParticleSystemComponent::ComputeTickComponent_Concurrent()
 			FPlatformMisc::Prefetch(NextInstance);
 		}
 
-		if (Instance && Instance->SpriteTemplate && Instance->bEnabled)
+		if (Instance && Instance->SpriteTemplate)
 		{
 			check(Instance->SpriteTemplate->LODLevels.Num() > 0);
 
 			UParticleLODLevel* SpriteLODLevel = Instance->SpriteTemplate->GetCurrentLODLevel(Instance);
 			if (SpriteLODLevel && SpriteLODLevel->bEnabled)
 			{
+				if (bLocalShouldManageSignificance)
+				{
+					bool bEmitterIsSignificant = Instance->SpriteTemplate->IsSignificant(RequiredSignificance);
+					if (bEmitterIsSignificant)
+					{
+						++NumSignificantEmitters;
+						Instance->SetHaltSpawning(false);
+						Instance->SetFakeBurstWhenSpawningSupressed(false);
+						Instance->bEnabled = true;
+					}
+					else
+					{
+						Instance->SetHaltSpawning(true);
+						Instance->SetFakeBurstWhenSpawningSupressed(true);
+						if (Instance->SpriteTemplate->bDisableWhenInsignficant)
+						{
+							Instance->bEnabled = false;
+						}
+					}
+				}
+				else
+				{
+					++NumSignificantEmitters;
+				}
+
 				Instance->Tick(DeltaTimeTick, bSuppressSpawning);
 
 				if (!Instance->Tick_MaterialOverrides())
@@ -4489,6 +4596,9 @@ void UParticleSystemComponent::ComputeTickComponent_Concurrent()
 
 void UParticleSystemComponent::FinalizeTickComponent()
 {
+	bool bLocalShouldManageSignificance = ShouldManageSignificance();
+	FInGameScopedCycleCounter InGameCycleCounter(GetWorld(), EInGamePerfTrackers::VFXSignificance, EInGamePerfTrackerThreads::GameThread, bLocalShouldManageSignificance);
+
 	SCOPE_CYCLE_COUNTER(STAT_ParticleFinalizeTickTime);
 
 	if(bAsyncDataCopyIsValid)
@@ -4548,35 +4658,28 @@ void UParticleSystemComponent::FinalizeTickComponent()
 
 	// Indicate that we have been ticked since being registered.
 	bJustRegistered = false;
+	
+	float CurrTime = GetWorld()->GetTimeSeconds();
 
-	// If component has just totally finished, call script event.
-	const bool bIsCompleted = HasCompleted(); 
-	if (bIsCompleted && !bWasCompleted)
+	//Are we still significant?
+	if ((bIsActive && !bWasDeactivated) && bLocalShouldManageSignificance && NumSignificantEmitters == 0 && CurrTime >= LastSignificantTime + Template->InsignificanceDelay)
 	{
-		UE_LOG(LogParticles,Verbose,
-			TEXT("HasCompleted()==true @ %fs %s"), World->TimeSeconds,
-			Template != NULL ? *Template->GetName() : TEXT("NULL"));
-
-		OnSystemFinished.Broadcast(this);
-
-		// When system is done - destroy all subemitters etc. We don't need them any more.
-		ResetParticles();
-		bIsActive = false;
-		SetComponentTickEnabled(false);
-
-		if (bAutoDestroy)
-		{
-			DestroyComponent();
-		}
-		else if (bAutoManageAttachment)
+		OnSignificanceChanged(false, false);
+	}
+	else
+	{
+		LastSignificantTime = CurrTime;
+		// If component has just totally finished, call script event.
+		const bool bIsCompleted = HasCompleted();
+		if (bIsCompleted && !bWasCompleted)
 		{
-			CancelAutoAttachment(/*bDetachFromParent=*/ true);
+			Complete();
 		}
+		bWasCompleted = bIsCompleted;
 	}
-	bWasCompleted = bIsCompleted;
 
 	// Update bounding box.
-	if (!bWarmingUp && !bIsCompleted && !Template->bUseFixedRelativeBoundingBox && !bIsTransformDirty)
+	if (!bWarmingUp && !bWasCompleted && !Template->bUseFixedRelativeBoundingBox && !bIsTransformDirty)
 	{
 		// Force an update every once in a while to shrink the bounds.
 		TimeSinceLastForceUpdateTransform += DeltaTimeTick;
@@ -4663,7 +4766,7 @@ void UParticleSystemComponent::WaitForAsyncAndFinalize(EForceAsyncWorkCompletion
 	}
 }
 
-void UParticleSystemComponent::InitParticles(bool bReInitExistingEmitters)
+void UParticleSystemComponent::InitParticles()
 {
 	SCOPE_CYCLE_COUNTER(STAT_ParticleSystemComponent_InitParticles);
 
@@ -4680,12 +4783,10 @@ void UParticleSystemComponent::InitParticles(bool bReInitExistingEmitters)
 
 	if (Template != NULL)
 	{
-		InitRequiredSignificanceLevel = RequiredSignificance;
-
 		WarmupTime = Template->WarmupTime;
 		WarmupTickRate = Template->WarmupTickRate;
 		bIsViewRelevanceDirty = true;
-		const int32 GlobalDetailMode = GetCurrentDetailMode();;
+		const int32 GlobalDetailMode = GetCurrentDetailMode();
 		const bool bCanEverRender = CanEverRender();
 
 		//simplified version.
@@ -4706,15 +4807,8 @@ void UParticleSystemComponent::InitParticles(bool bReInitExistingEmitters)
 			UParticleEmitter* Emitter = Template->Emitters[Idx];
 			FParticleEmitterInstance* Instance = NumInstances == 0 ? NULL : EmitterInstances[Idx];
 			
-			const bool bDetailModeAllowsRendering	= DetailMode <= GlobalDetailMode;
-			bool bShouldCreateAndOrInit = bDetailModeAllowsRendering && Emitter->HasAnyEnabledLODs() && bCanEverRender;
-			const bool bIsSignificant = (int32)Emitter->SignificanceLevel >= (int32)RequiredSignificance;
-			bool bFailedOnlyOnSignificance = false;
-			if (bShouldCreateAndOrInit && !bIsSignificant)
-			{
-				bFailedOnlyOnSignificance = true;
-				bShouldCreateAndOrInit = false;
-			}
+			const bool bDetailModeAllowsRendering = DetailMode <= GlobalDetailMode && Emitter->DetailMode <= GlobalDetailMode;
+			const bool bShouldCreateAndOrInit = bDetailModeAllowsRendering && Emitter->HasAnyEnabledLODs() && bCanEverRender;
 
 			if (bShouldCreateAndOrInit)
 			{
@@ -4731,37 +4825,24 @@ void UParticleSystemComponent::InitParticles(bool bReInitExistingEmitters)
 				if (Instance)
 				{
 					Instance->bEnabled = true;
-					if (bReInitExistingEmitters)
-					{
-						Instance->InitParameters(Emitter, this);
-						Instance->Init();
+					Instance->InitParameters(Emitter, this);
+					Instance->Init();
 
-						PreferredLODLevel = FMath::Min(PreferredLODLevel, Emitter->LODLevels.Num());
-						bSetLodLevels |= !bIsFirstCreate;//Only set lod levels if we init any instances and it's not the first creation time.
-					}
+					PreferredLODLevel = FMath::Min(PreferredLODLevel, Emitter->LODLevels.Num());
+					bSetLodLevels |= !bIsFirstCreate;//Only set lod levels if we init any instances and it's not the first creation time.
 				}
 			}
 			else
 			{
 				if (Instance)
 				{
-					if(!bFailedOnlyOnSignificance)
-					{
-						//If the instance exists already and we failed for something other than significance then destroy this emitter.
 #if STATS
-						Instance->PreDestructorCall();
+					Instance->PreDestructorCall();
 #endif
-						delete Instance;
-						EmitterInstances[Idx] = NULL;
-						bClearDynamicData = true;
-					}
-					else
-					{
-						//otherwise, we just disable it.
-						Instance->bEnabled = false;
-					}
+					delete Instance;
+					EmitterInstances[Idx] = NULL;
+					bClearDynamicData = true;
 				}
-				bMissingEmittersDueToSignificance |= bFailedOnlyOnSignificance;
 			}
 		}
 
@@ -4793,7 +4874,7 @@ void UParticleSystemComponent::InitParticles(bool bReInitExistingEmitters)
 						Instance->CurrentLODLevelIndex = Instance->SpriteTemplate->LODLevels.Num()-1;
 						ensureMsgf(false, TEXT("LOD access out of bounds (OR-11322). Please let olaf.piesche or simon.tovey know."));
 					}
-					Instance->CurrentLODLevel		= Instance->SpriteTemplate->LODLevels[Instance->CurrentLODLevelIndex];
+					Instance->CurrentLODLevel = Instance->SpriteTemplate->LODLevels[Instance->CurrentLODLevelIndex];
 				}
 			}
 		}
@@ -4810,6 +4891,12 @@ void UParticleSystemComponent::ResetParticles(bool bEmptyInstances)
 
 	UWorld* OwningWorld = GetWorld();
 
+	//Also consider this deactivation.
+	if (bIsActive)
+	{
+		OnSystemPreActivationChange.Broadcast(this, false);
+	}
+
 	const bool bIsGameWorld = OwningWorld ? OwningWorld->IsGameWorld() : !GIsEditor;
 
 	// Remove instances from scene.
@@ -4832,9 +4919,6 @@ void UParticleSystemComponent::ResetParticles(bool bEmptyInstances)
 	// Set the system as inactive
 	bIsActive	= false;
 
-	
-	bMissingEmittersDueToSignificance = false;
-
 	// Remove instances if we're not running gameplay.ww
 	if (!bIsGameWorld || bEmptyInstances)
 	{
@@ -4960,6 +5044,9 @@ void UParticleSystemComponent::SetTemplate(class UParticleSystem* NewTemplate)
 
 void UParticleSystemComponent::ActivateSystem(bool bFlagAsJustAttached)
 {
+	bool bManagingSignificance = ShouldManageSignificance();
+	FInGameScopedCycleCounter InGameCycleCounter(GetWorld(), EInGamePerfTrackers::VFXSignificance, EInGamePerfTrackerThreads::GameThread, bManagingSignificance);
+
 	SCOPE_CYCLE_COUNTER(STAT_ParticleActivateTime);
 	ForceAsyncWorkCompletion(STALL);
 
@@ -5002,11 +5089,8 @@ void UParticleSystemComponent::ActivateSystem(bool bFlagAsJustAttached)
 	}
 
 	// System settings may have been lowered. Support late deactivation.
-	const bool bDetailModeAllowsRendering	= DetailMode <= GetCurrentDetailMode();
-
-	//Assume significant initially.
-	LastSignificantTime = GetWorld()->GetTimeSeconds();
-
+	const bool bDetailModeAllowsRendering = DetailMode <= GetCurrentDetailMode();
+	
 	if( GIsAllowingParticles && bDetailModeAllowsRendering )
 	{
 		// Auto attach if requested
@@ -5037,8 +5121,15 @@ void UParticleSystemComponent::ActivateSystem(bool bFlagAsJustAttached)
 			}
 		}
 
-		//Call this now after any attachment has happened.
-		OnSystemPreActivate.Broadcast(this);
+		if (!bIsActive)
+		{
+			LastSignificantTime = GetWorld()->GetTimeSeconds();
+			bIsSignificant = true;
+			RequiredSignificance = EParticleSignificanceLevel::Low;
+
+			//Call this now after any attachment has happened.
+			OnSystemPreActivationChange.Broadcast(this, true);
+		}
 
 		if (bFlagAsJustAttached)
 		{
@@ -5152,11 +5243,42 @@ void UParticleSystemComponent::ActivateSystem(bool bFlagAsJustAttached)
 	{
 		LastRenderTime = GetWorld()->GetTimeSeconds();
 	}
+
+	//We are definitely insignificant already so set insignificant before we ever begin ticking.
+	if (bManagingSignificance && Template->GetHighestSignificance() < RequiredSignificance && Template->InsignificanceDelay == 0.0f)
+	{
+		OnSignificanceChanged(false, true);
+	}
 }
 
+void UParticleSystemComponent::Complete()
+{
+	UE_LOG(LogParticles, Verbose,
+		TEXT("HasCompleted()==true @ %fs %s"), World->TimeSeconds,
+		Template != NULL ? *Template->GetName() : TEXT("NULL"));
+
+	OnSystemFinished.Broadcast(this);
+
+	// When system is done - destroy all subemitters etc. We don't need them any more.
+	ResetParticles();
+	bIsActive = false;
+	SetComponentTickEnabled(false);
+
+	if (bAutoDestroy)
+	{
+		DestroyComponent();
+	}
+	else if (bAutoManageAttachment)
+	{
+		CancelAutoAttachment(/*bDetachFromParent=*/ true);
+	}
+}
 
 void UParticleSystemComponent::DeactivateSystem()
 {
+	bool bManagingSignificance = ShouldManageSignificance();
+	FInGameScopedCycleCounter InGameCycleCounter(GetWorld(), EInGamePerfTrackers::VFXSignificance, EInGamePerfTrackerThreads::GameThread, bManagingSignificance);
+
 	if (IsTemplate() == true)
 	{
 		return;
@@ -5168,6 +5290,11 @@ void UParticleSystemComponent::DeactivateSystem()
 		TEXT("DeactivateSystem @ %fs %s"), GetWorld()->TimeSeconds,
 		Template != NULL ? *Template->GetName() : TEXT("NULL"));
 
+	if (bIsActive)
+	{
+		OnSystemPreActivationChange.Broadcast(this, false);
+	}
+
 	bSuppressSpawning = true;
 	bWasDeactivated = true;
 
@@ -5208,6 +5335,10 @@ void UParticleSystemComponent::DeactivateSystem()
 		MarkRenderStateDirty();
 	}
 
+	//We have to ensure ticking is enabled so that this component completes and is can be destroyed etc correctly.
+	//TODO: What if there are immortal particles but bKillOnDeactivate is false? Need to mark emitters with currently immortal particles, kill them and warn the user.
+	SetComponentTickEnabled(true);
+
 	LastRenderTime = GetWorld()->GetTimeSeconds();
 }
 
@@ -5250,7 +5381,8 @@ bool UParticleSystemComponent::ShouldActivate() const
 void UParticleSystemComponent::Activate(bool bReset) 
 {
 	// If the particle system can't ever render (ie on dedicated server or in a commandlet) than do not activate...
-	if (FApp::CanEverRender())
+	// Occasionaly we can arrive here with no template so check that here too.
+	if (FApp::CanEverRender() && Template)
 	{
 		ForceAsyncWorkCompletion(STALL);
 		if (bReset || ShouldActivate()==true)
@@ -5425,7 +5557,7 @@ bool UParticleSystemComponent::HasCompleted()
 	{
 		FParticleEmitterInstance* Instance = EmitterInstances[i];
 
-		if (Instance && Instance->CurrentLODLevel)
+		if (Instance && Instance->CurrentLODLevel && Instance->bEnabled)
 		{
 			if (!Instance->bEmitterIsDone)
 			{
@@ -5537,7 +5669,7 @@ void UParticleSystemComponent::InitializeSystem()
 		Template != NULL ? *Template->GetName() : TEXT("NULL"), this, FXSystem);
 
 	// System settings may have been lowered. Support late deactivation.
-	const bool bDetailModeAllowsRendering	= DetailMode <= GetCurrentDetailMode();
+	const bool bDetailModeAllowsRendering = DetailMode <= GetCurrentDetailMode();
 
 	if( GIsAllowingParticles && bDetailModeAllowsRendering )
 	{
