Copying //UE4/Dev-Rendering to Dev-Main (//UE4/Dev-Main)

diff --git a/Engine/Source/Runtime/Engine/Private/Components/SplineMeshComponent.cpp b/Engine/Source/Runtime/Engine/Private/Components/SplineMeshComponent.cpp
--- a/Engine/Source/Runtime/Engine/Private/Components/SplineMeshComponent.cpp
+++ b/Engine/Source/Runtime/Engine/Private/Components/SplineMeshComponent.cpp
@@ -39,10 +39,16 @@ void FSplineMeshVertexFactoryShaderParameters::Bind(const FShaderParameterMap& P
 
 void FSplineMeshVertexFactoryShaderParameters::SetMesh(FRHICommandList& RHICmdList, FShader* Shader, const FVertexFactory* VertexFactory, const FSceneView& View, const FMeshBatchElement& BatchElement, uint32 DataFlags) const
 {
+	if (BatchElement.bUserDataIsColorVertexBuffer)
+	{
+		FColorVertexBuffer* OverrideColorVertexBuffer = (FColorVertexBuffer*)BatchElement.UserData;
+		check(OverrideColorVertexBuffer);
+		static_cast<const FLocalVertexFactory*>(VertexFactory)->SetColorOverrideStream(RHICmdList, OverrideColorVertexBuffer);
+	}
 	if (Shader->GetVertexShader())
 	{
-		FSplineMeshVertexFactory* SplineVertexFactory = (FSplineMeshVertexFactory*)VertexFactory;
-		FSplineMeshSceneProxy* SplineProxy = SplineVertexFactory->SplineSceneProxy;
+		checkSlow(BatchElement.bIsSplineProxy);
+		FSplineMeshSceneProxy* SplineProxy = BatchElement.SplineMeshSceneProxy;
 		FSplineMeshParams& SplineParams = SplineProxy->SplineParams;
 
 		SetShaderValue(RHICmdList, Shader->GetVertexShader(), SplineStartPosParam, SplineParams.StartPos);
@@ -89,16 +95,30 @@ FVertexFactoryShaderParameters* FSplineMeshVertexFactory::ConstructShaderParamet
 //////////////////////////////////////////////////////////////////////////
 // SplineMeshSceneProxy
 
-void FSplineMeshSceneProxy::InitResources( USplineMeshComponent* InComponent, int32 InLODIndex, FColorVertexBuffer* InOverrideColorVertexBuffer )
+void FSplineMeshSceneProxy::InitVertexFactory(USplineMeshComponent* InComponent, int32 InLODIndex, FColorVertexBuffer* InOverrideColorVertexBuffer)
 {
 	// Initialize the static mesh's vertex factory.
-	ENQUEUE_UNIQUE_RENDER_COMMAND_FOURPARAMETER(
+	ENQUEUE_UNIQUE_RENDER_COMMAND_THREEPARAMETER(
 		InitSplineMeshVertexFactory,
-		FSplineMeshVertexFactory*, VertexFactory, LODResources[InLODIndex].VertexFactory,
 		FStaticMeshLODResources*, RenderData, &InComponent->StaticMesh->RenderData->LODResources[InLODIndex],
 		UStaticMesh*, Parent, InComponent->StaticMesh,
-		FColorVertexBuffer*, OverridenColorVertexBuffer, InOverrideColorVertexBuffer,
+		bool, bOverrideColorVertexBuffer, !!InOverrideColorVertexBuffer,
+		{
+
+		if ((RenderData->SplineVertexFactory && !bOverrideColorVertexBuffer) || (RenderData->SplineVertexFactoryOverrideColorVertexBuffer && bOverrideColorVertexBuffer))
+		{
+			// we already have it
+			return;
+		}
+		FSplineMeshVertexFactory* VertexFactory = new FSplineMeshVertexFactory;
+		if (bOverrideColorVertexBuffer)
 		{
+			RenderData->SplineVertexFactoryOverrideColorVertexBuffer = VertexFactory;
+		}
+		else
+		{
+			RenderData->SplineVertexFactory = VertexFactory;
+		}
 		FLocalVertexFactory::FDataType Data;
 
 		Data.PositionComponent = FVertexStreamComponent(
@@ -119,14 +139,21 @@ void FSplineMeshSceneProxy::InitResources( USplineMeshComponent* InComponent, in
 			RenderData->VertexBuffer.GetStride(),
 			VET_PackedNormal
 			);
-		
-		FColorVertexBuffer* LODColorVertexBuffer = &RenderData->ColorVertexBuffer;
-		// Override the buffer if it has been changed
-		if ( OverridenColorVertexBuffer != NULL )
+		if (bOverrideColorVertexBuffer)
 		{
-			LODColorVertexBuffer = OverridenColorVertexBuffer;
+			Data.ColorComponent = FVertexStreamComponent(
+				&GNullColorVertexBuffer,
+				0,	// Struct offset to color
+				sizeof(FColor), //asserted elsewhere
+				VET_Color,
+				false, // not instanced
+				true // set in SetMesh
+				);
 		}
-		if ( LODColorVertexBuffer->GetNumVertices() > 0 )
+		else
+		{
+			FColorVertexBuffer* LODColorVertexBuffer = &RenderData->ColorVertexBuffer;
+			if (LODColorVertexBuffer->GetNumVertices() > 0)
 		{
 			Data.ColorComponent = FVertexStreamComponent(
 				LODColorVertexBuffer,
@@ -135,6 +162,7 @@ void FSplineMeshSceneProxy::InitResources( USplineMeshComponent* InComponent, in
 				VET_Color
 				);
 		}
+		}
 
 		Data.TextureCoordinates.Empty();
 
@@ -212,14 +240,6 @@ void FSplineMeshSceneProxy::InitResources( USplineMeshComponent* InComponent, in
 }
 
 
-void FSplineMeshSceneProxy::ReleaseResources()
-{
-	for (FSplineMeshSceneProxy::FLODResources& LODResource : LODResources)
-	{
-		LODResource.VertexFactory->ReleaseResource();
-	}
-}
-
 //////////////////////////////////////////////////////////////////////////
 // SplineMeshComponent
 
diff --git a/Engine/Source/Runtime/Engine/Private/Components/StaticMeshComponent.cpp b/Engine/Source/Runtime/Engine/Private/Components/StaticMeshComponent.cpp
--- a/Engine/Source/Runtime/Engine/Private/Components/StaticMeshComponent.cpp
+++ b/Engine/Source/Runtime/Engine/Private/Components/StaticMeshComponent.cpp
@@ -1471,15 +1471,23 @@ int32 UStaticMeshComponent::GetStaticLightMapResolution() const
 	return FMath::Max<int32>(Width, Height);
 }
 
-bool UStaticMeshComponent::HasValidSettingsForStaticLighting() const
+bool UStaticMeshComponent::HasValidSettingsForStaticLighting(bool bOverlookInvalidComponents) const
 {
-	int32 LightMapWidth = 0;
-	int32 LightMapHeight = 0;
-	GetLightMapResolution(LightMapWidth, LightMapHeight);
+	if (bOverlookInvalidComponents && !StaticMesh)
+	{
+		// Return true for invalid components, this is used during the map check where those invalid components will be warned about separately
+		return true;
+	}
+	else
+	{
+		int32 LightMapWidth = 0;
+		int32 LightMapHeight = 0;
+		GetLightMapResolution(LightMapWidth, LightMapHeight);
 
-	return Super::HasValidSettingsForStaticLighting() 
-		&& StaticMesh
-		&& UsesTextureLightmaps(LightMapWidth, LightMapHeight);
+		return Super::HasValidSettingsForStaticLighting(bOverlookInvalidComponents) 
+			&& StaticMesh
+			&& UsesTextureLightmaps(LightMapWidth, LightMapHeight);
+	}
 }
 
 bool UStaticMeshComponent::UsesTextureLightmaps(int32 InWidth, int32 InHeight) const
diff --git a/Engine/Source/Runtime/Engine/Private/DebugViewModeMaterialProxy.cpp b/Engine/Source/Runtime/Engine/Private/DebugViewModeMaterialProxy.cpp
new file mode 100644
--- /dev/null
+++ b/Engine/Source/Runtime/Engine/Private/DebugViewModeMaterialProxy.cpp
@@ -0,0 +1,185 @@
+// Copyright 1998-2016 Epic Games, Inc. All Rights Reserved.
+
+/*=============================================================================
+DebugViewModeMaterialProxy.cpp : Contains definitions the debug view mode material shaders.
+=============================================================================*/
+
+#include "EnginePrivate.h"
+#include "DebugViewModeMaterialProxy.h"
+
+ENGINE_API const FMaterial* GetDebugViewMaterialPS(EDebugViewShaderMode DebugViewShaderMode, const FMaterial* Material)
+{
+#if WITH_EDITORONLY_DATA
+	return FDebugViewModeMaterialProxy::GetShader(DebugViewShaderMode, Material);
+#else
+	return nullptr;
+#endif
+}
+
+ENGINE_API void ClearAllDebugViewMaterials()
+{
+#if WITH_EDITORONLY_DATA
+	FDebugViewModeMaterialProxy::ClearAllShaders();
+#endif
+}
+
+#if WITH_EDITORONLY_DATA
+
+TMap<const FMaterial*, FDebugViewModeMaterialProxy*> FDebugViewModeMaterialProxy::DebugMaterialShaderMap;
+volatile bool FDebugViewModeMaterialProxy::bReentrantCall = false;
+
+void FDebugViewModeMaterialProxy::AddShader(UMaterialInterface* InMaterialInterface, EMaterialShaderMapUsage::Type InUsage)
+{
+	if (!InMaterialInterface) return;
+
+	const FMaterial* Material = InMaterialInterface->GetMaterialResource(GMaxRHIFeatureLevel);
+	if (!Material) return;
+
+	if (!DebugMaterialShaderMap.Contains(Material))
+	{
+		DebugMaterialShaderMap.Add(Material, new FDebugViewModeMaterialProxy(InMaterialInterface, InUsage));
+	}
+}
+
+const FMaterial* FDebugViewModeMaterialProxy::GetShader(EDebugViewShaderMode DebugViewShaderMode, const FMaterial* Material)
+{
+	if (DebugViewShaderMode == DVSM_TexCoordScaleAnalysis)
+	{
+		FDebugViewModeMaterialProxy** BoundMaterial = DebugMaterialShaderMap.Find(Material);
+		if (BoundMaterial)
+		{
+			return *BoundMaterial;
+		}
+	}
+	return nullptr;
+}
+
+void FDebugViewModeMaterialProxy::ClearAllShaders()
+{
+	if (bReentrantCall) return;
+
+	bReentrantCall = true;
+	for (TMap<const FMaterial*, FDebugViewModeMaterialProxy*>::TIterator It(DebugMaterialShaderMap); It; ++It)
+	{
+		FDebugViewModeMaterialProxy* Proxy = It.Value();
+		delete Proxy;
+	}
+	DebugMaterialShaderMap.Empty();
+
+	bReentrantCall = false;
+}
+
+FDebugViewModeMaterialProxy::FDebugViewModeMaterialProxy(UMaterialInterface* InMaterialInterface, EMaterialShaderMapUsage::Type InUsage)
+	: FMaterial()
+	, MaterialInterface(InMaterialInterface)
+	, Material(nullptr)
+	, Usage(InUsage)
+{
+	SetQualityLevelProperties(EMaterialQualityLevel::High, false, GMaxRHIFeatureLevel);
+	Material = InMaterialInterface->GetMaterial();
+	Material->AppendReferencedTextures(ReferencedTextures);
+
+	FMaterialResource* Resource = InMaterialInterface->GetMaterialResource(GMaxRHIFeatureLevel);
+
+	FMaterialShaderMapId ResourceId;
+	Resource->GetShaderMapId(GMaxRHIShaderPlatform, ResourceId);
+
+	{
+		TArray<FShaderType*> ShaderTypes;
+		TArray<FVertexFactoryType*> VFTypes;
+		TArray<const FShaderPipelineType*> ShaderPipelineTypes;
+		GetDependentShaderAndVFTypes(GMaxRHIShaderPlatform, ShaderTypes, ShaderPipelineTypes, VFTypes);
+
+		// Overwrite the shader map Id's dependencies with ones that came from the FMaterial actually being compiled (this)
+		// This is necessary as we change FMaterial attributes like GetShadingModel(), which factor into the ShouldCache functions that determine dependent shader types
+		ResourceId.SetShaderDependencies(ShaderTypes, ShaderPipelineTypes, VFTypes);
+	}
+
+	ResourceId.Usage = InUsage;
+
+	CacheShaders(ResourceId, GMaxRHIShaderPlatform, true);
+}
+
+const FMaterial* FDebugViewModeMaterialProxy::GetMaterial(ERHIFeatureLevel::Type FeatureLevel) const
+{
+	if (GetRenderingThreadShaderMap())
+	{
+		return this;
+	}
+	else
+	{
+		return UMaterial::GetDefaultMaterial(MD_Surface)->GetRenderProxy(false)->GetMaterial(FeatureLevel);
+	}
+}
+
+bool FDebugViewModeMaterialProxy::GetVectorValue(const FName ParameterName, FLinearColor* OutValue, const FMaterialRenderContext& Context) const
+{
+	return MaterialInterface->GetRenderProxy(0)->GetVectorValue(ParameterName, OutValue, Context);
+}
+
+bool FDebugViewModeMaterialProxy::GetScalarValue(const FName ParameterName, float* OutValue, const FMaterialRenderContext& Context) const
+{
+	return MaterialInterface->GetRenderProxy(0)->GetScalarValue(ParameterName, OutValue, Context);
+}
+
+bool FDebugViewModeMaterialProxy::GetTextureValue(const FName ParameterName,const UTexture** OutValue, const FMaterialRenderContext& Context) const
+{
+	return MaterialInterface->GetRenderProxy(0)->GetTextureValue(ParameterName,OutValue,Context);
+}
+
+int32 FDebugViewModeMaterialProxy::GetMaterialDomain() const
+{
+	return Material ? (int32)Material->MaterialDomain : (int32)MD_Surface;
+}
+
+bool FDebugViewModeMaterialProxy::IsTwoSided() const
+{ 
+	return MaterialInterface && MaterialInterface->IsTwoSided();
+}
+
+bool FDebugViewModeMaterialProxy::IsDitheredLODTransition() const
+{ 
+	return MaterialInterface && MaterialInterface->IsDitheredLODTransition();
+}
+
+bool FDebugViewModeMaterialProxy::IsLightFunction() const
+{
+	return Material && Material->MaterialDomain == MD_LightFunction;
+}
+
+bool FDebugViewModeMaterialProxy::IsUsedWithDeferredDecal() const
+{
+	return	Material && Material->MaterialDomain == MD_DeferredDecal;
+}
+
+bool FDebugViewModeMaterialProxy::IsSpecialEngineMaterial() const
+{
+	return Material && Material->bUsedAsSpecialEngineMaterial;
+}
+
+bool FDebugViewModeMaterialProxy::IsWireframe() const
+{
+	return Material && Material->Wireframe;
+}
+
+bool FDebugViewModeMaterialProxy::IsMasked() const
+{ 
+	return Material && Material->IsMasked(); 
+}
+
+enum EBlendMode FDebugViewModeMaterialProxy::GetBlendMode() const
+{ 
+	return Material ? Material->GetBlendMode() : BLEND_Opaque; 
+}
+
+enum EMaterialShadingModel FDebugViewModeMaterialProxy::GetShadingModel() const
+{ 
+	return Material ? Material->GetShadingModel() : MSM_Unlit;
+}
+
+float FDebugViewModeMaterialProxy::GetOpacityMaskClipValue() const
+{ 
+	return Material ? Material->GetOpacityMaskClipValue() : .5f;
+}
+
+#endif
\ No newline at end of file
diff --git a/Engine/Source/Runtime/Engine/Private/DebugViewModeMaterialProxy.h b/Engine/Source/Runtime/Engine/Private/DebugViewModeMaterialProxy.h
new file mode 100644
--- /dev/null
+++ b/Engine/Source/Runtime/Engine/Private/DebugViewModeMaterialProxy.h
@@ -0,0 +1,111 @@
+// Copyright 1998-2016 Epic Games, Inc. All Rights Reserved.
+
+/*=============================================================================
+DebugViewModeMaterialProxy.h : Contains definitions the debug view mode material shaders.
+=============================================================================*/
+
+#pragma once
+
+#if WITH_EDITORONLY_DATA
+
+/**
+ * Material proxy for debug viewmode. Used to prevent debug view mode shaders from being stored in the default material map.
+ */
+class FDebugViewModeMaterialProxy : public FMaterial, public FMaterialRenderProxy
+{
+public:
+	FDebugViewModeMaterialProxy()
+		: FMaterial()
+		, MaterialInterface(nullptr)
+		, Material(nullptr)
+	{
+		SetQualityLevelProperties(EMaterialQualityLevel::High, false, GMaxRHIFeatureLevel);
+	}
+
+	FDebugViewModeMaterialProxy(UMaterialInterface* InMaterialInterface, EMaterialShaderMapUsage::Type InUsage);
+
+	virtual bool RequiresSynchronousCompilation() const override { return false; };
+
+	/**
+	* Should shaders compiled for this material be saved to disk?
+	*/
+	virtual bool IsPersistent() const override { return false; }
+
+	// Normally this would cause a bug as the shader map would try to be shared by both, but GetShaderMapUsage() allows this to work
+	virtual FGuid GetMaterialId() const override { return Material->StateId; }
+
+	virtual EMaterialShaderMapUsage::Type GetShaderMapUsage() const override { return Usage; }
+
+
+	virtual bool ShouldCache(EShaderPlatform Platform, const FShaderType* ShaderType, const FVertexFactoryType* VertexFactoryType) const override
+	{
+		return FString(ShaderType->GetName()).Contains(TEXT("FTexCoordScaleAnalysisPS"));
+	}
+
+	virtual const TArray<UTexture*>& GetReferencedTextures() const override
+	{
+		return ReferencedTextures;
+	}
+
+	// Material properties.
+	/** Entry point for compiling a specific material property.  This must call SetMaterialProperty. */
+	virtual int32 CompilePropertyAndSetMaterialProperty(EMaterialProperty Property, FMaterialCompiler* Compiler, EShaderFrequency OverrideShaderFrequency, bool bUsePreviousFrameTime) const override
+	{
+		return MaterialInterface ? MaterialInterface->GetMaterialResource(GMaxRHIFeatureLevel)->CompilePropertyAndSetMaterialProperty(Property, Compiler, OverrideShaderFrequency, bUsePreviousFrameTime) : INDEX_NONE;
+	}
+
+	virtual FString GetMaterialUsageDescription() const override
+	{
+		return FString::Printf(TEXT("FDebugViewModeMaterialProxy %s"), MaterialInterface ? *MaterialInterface->GetName() : TEXT("NULL"));
+	}
+
+	virtual FString GetFriendlyName() const override { return FString::Printf(TEXT("FDebugViewModeMaterialProxy %s"), MaterialInterface ? *MaterialInterface->GetName() : TEXT("NULL")); }
+
+	const UMaterialInterface* GetMaterialInterface() const
+	{
+		return MaterialInterface;
+	}
+
+	friend FArchive& operator<< ( FArchive& Ar, FDebugViewModeMaterialProxy& V )
+	{
+		return Ar << V.MaterialInterface;
+	}
+
+	////////////////
+	// FMaterialRenderProxy interface.
+	virtual const FMaterial* GetMaterial(ERHIFeatureLevel::Type FeatureLevel) const override;
+	virtual bool GetVectorValue(const FName ParameterName, FLinearColor* OutValue, const FMaterialRenderContext& Context) const override;
+	virtual bool GetScalarValue(const FName ParameterName, float* OutValue, const FMaterialRenderContext& Context) const override;
+	virtual bool GetTextureValue(const FName ParameterName,const UTexture** OutValue, const FMaterialRenderContext& Context) const override;
+
+	virtual int32 GetMaterialDomain() const override;
+	virtual bool IsTwoSided() const  override;
+	virtual bool IsDitheredLODTransition() const  override;
+	virtual bool IsLightFunction() const override;
+	virtual bool IsUsedWithDeferredDecal() const override;
+	virtual bool IsSpecialEngineMaterial() const override;
+	virtual bool IsWireframe() const override;
+	virtual bool IsMasked() const override;
+	virtual enum EBlendMode GetBlendMode() const override;
+	virtual enum EMaterialShadingModel GetShadingModel() const override;
+	virtual float GetOpacityMaskClipValue() const override;
+
+	////////////////
+
+	static void AddShader(UMaterialInterface* InMaterialInterface, EMaterialShaderMapUsage::Type InUsage);
+	static const FMaterial* GetShader(EDebugViewShaderMode DebugViewShaderMode, const FMaterial* Material);
+	static void ClearAllShaders();
+
+private:
+
+	/** The material interface for this proxy */
+	UMaterialInterface* MaterialInterface;
+	UMaterial* Material;	
+	TArray<UTexture*> ReferencedTextures;
+	EMaterialShaderMapUsage::Type Usage;
+
+	static volatile bool bReentrantCall;
+	static TMap<const FMaterial*, FDebugViewModeMaterialProxy*> DebugMaterialShaderMap;
+};
+
+#endif
\ No newline at end of file
diff --git a/Engine/Source/Runtime/Engine/Private/GPUSkinCache.cpp b/Engine/Source/Runtime/Engine/Private/GPUSkinCache.cpp
--- a/Engine/Source/Runtime/Engine/Private/GPUSkinCache.cpp
+++ b/Engine/Source/Runtime/Engine/Private/GPUSkinCache.cpp
@@ -21,7 +21,6 @@ DEFINE_STAT(STAT_GPUSkinCache_SkippedForMaxChunksPerLOD);
 DEFINE_STAT(STAT_GPUSkinCache_SkippedForChunksPerFrame);
 DEFINE_STAT(STAT_GPUSkinCache_SkippedForZeroInfluences);
 DEFINE_STAT(STAT_GPUSkinCache_SkippedForCloth);
-DEFINE_STAT(STAT_GPUSkinCache_SkippedForMorphs);
 DEFINE_STAT(STAT_GPUSkinCache_SkippedForMemory);
 DEFINE_STAT(STAT_GPUSkinCache_SkippedForAlreadyCached);
 DEFINE_STAT(STAT_GPUSkinCache_SkippedForOutOfECS);
@@ -92,9 +91,16 @@ public:
 		SkinInputStream.Bind(Initializer.ParameterMap, TEXT("SkinStreamInputBuffer"));
 
 		SkinCacheBuffer.Bind(Initializer.ParameterMap, TEXT("SkinCacheBuffer"));
+
+		MorphInputStreamBuffer.Bind(Initializer.ParameterMap, TEXT("MorphInputStreamBuffer"));
+		MorphInputStreamFloatOffset.Bind(Initializer.ParameterMap, TEXT("MorphInputStreamFloatOffset"));
 	}
 
-	void SetParameters(FRHICommandList& RHICmdList, uint32 InputStreamStride, uint32 InputStreamFloatOffset, uint32 InputStreamVertexCount, uint32 OutputBufferFloatOffset, const FVertexBufferAndSRV& BoneBuffer, FUniformBufferRHIRef UniformBuffer, FShaderResourceViewRHIRef VertexBufferSRV, FRWBuffer& SkinBuffer, const FVector& MeshOrigin, const FVector& MeshExtension)
+	void SetParameters(FRHICommandList& RHICmdList, uint32 InputStreamStride,
+		uint32 InputStreamFloatOffset, uint32 InputStreamVertexCount, uint32 OutputBufferFloatOffset,
+		const FVertexBufferAndSRV& BoneBuffer, FUniformBufferRHIRef UniformBuffer,
+		FShaderResourceViewRHIRef VertexBufferSRV, FRWBuffer& SkinBuffer, const FVector& MeshOrigin, const FVector& MeshExtension,
+		uint32 SkinType, uint32 InMorphInputStreamFloatOffset, FShaderResourceViewRHIParamRef MorphSRV)
 	{
 		FComputeShaderRHIParamRef ShaderRHI = GetComputeShader();
 		SetShaderValue(RHICmdList, ShaderRHI, SkinMeshOriginParameter, MeshOrigin);
@@ -116,8 +122,14 @@ public:
 		}
 
 		RHICmdList.SetShaderResourceViewParameter(ShaderRHI, SkinInputStream.GetBaseIndex(), VertexBufferSRV);
-
 		RHICmdList.SetUAVParameter(ShaderRHI, SkinCacheBuffer.GetBaseIndex(), SkinBuffer.UAV);
+
+		const bool bMorph = SkinType == 1;
+		if(bMorph)
+		{
+			SetShaderValue(RHICmdList, ShaderRHI, MorphInputStreamFloatOffset, InMorphInputStreamFloatOffset);
+			RHICmdList.SetShaderResourceViewParameter(ShaderRHI, MorphInputStreamBuffer.GetBaseIndex(), MorphSRV);
+		}
 	}
 
 	void UnsetParameters(FRHICommandList& RHICmdList)
@@ -136,7 +148,7 @@ public:
 		bool bShaderHasOutdatedParameters = FGlobalShader::Serialize(Ar);
 		Ar  << SkinMeshOriginParameter << SkinMeshExtensionParameter << SkinInputStreamStride
 			<< SkinInputStreamVertexCount << SkinInputStreamFloatOffset << SkinOutputBufferFloatOffset
-			<< SkinInputStream << SkinCacheBuffer << BoneMatrices;
+			<< SkinInputStream << SkinCacheBuffer << BoneMatrices << MorphInputStreamBuffer << MorphInputStreamFloatOffset;
 		//Ar << DebugParameter;
 		return bShaderHasOutdatedParameters;
 	}
@@ -158,13 +170,17 @@ private:
 	FShaderResourceParameter BoneMatrices;
 	FShaderResourceParameter SkinInputStream;
 	FShaderResourceParameter SkinCacheBuffer;
+
+	FShaderResourceParameter MorphInputStreamBuffer;
+	FShaderParameter MorphInputStreamFloatOffset;
 };
 
 /** Compute shader that skins a batch of vertices. */
-template <bool bUseExtraBoneInfluencesT>
+// @param SkinType 0:normal, 1:with morph target, 2:with APEX cloth (not yet implemented)
+template <bool bUseExtraBoneInfluencesT, uint32 SkinType>
 class TGPUSkinCacheCS : public FBaseGPUSkinCacheCS
 {
-	DECLARE_SHADER_TYPE(TGPUSkinCacheCS<bUseExtraBoneInfluencesT>,Global)
+	DECLARE_SHADER_TYPE(TGPUSkinCacheCS, Global)
 public:
 
 	static bool ShouldCache(EShaderPlatform Platform)
@@ -177,6 +193,7 @@ public:
 		FGlobalShader::ModifyCompilationEnvironment(Platform,OutEnvironment);
 		const uint32 UseExtraBoneInfluences = bUseExtraBoneInfluencesT;
 		OutEnvironment.SetDefine(TEXT("GPUSKIN_USE_EXTRA_INFLUENCES"), UseExtraBoneInfluences);
+		OutEnvironment.SetDefine(TEXT("GPUSKIN_MORPH_BLEND"), (uint32)(SkinType == 1));
 
 		OutEnvironment.SetDefine(TEXT("GPUSKIN_RWBUFFER_OFFSET_POSITION"), FGPUSkinCache::RWPositionOffsetInFloats);
 		OutEnvironment.SetDefine(TEXT("GPUSKIN_RWBUFFER_OFFSET_TANGENT_X"), FGPUSkinCache::RWTangentXOffsetInFloats);
@@ -192,8 +209,26 @@ public:
 	TGPUSkinCacheCS()
 	{
 	}
+	
+	static const TCHAR* GetSourceFilename()
+	{
+		return TEXT("GpuSkinCacheComputeShader");
+	}
+
+	static const TCHAR* GetFunctionName()
+	{
+		return TEXT("SkinCacheUpdateBatchCS");
+	}
 };
 
+// #define avoids a lot of code duplication
+#define VARIATION1(A)		VARIATION2(A,0)			VARIATION2(A,1)
+#define VARIATION2(A, B) typedef TGPUSkinCacheCS<A, B> TGPUSkinCacheCS##A##B; \
+	IMPLEMENT_SHADER_TYPE2(TGPUSkinCacheCS##A##B, SF_Compute);
+	VARIATION1(0) VARIATION1(1)
+#undef VARIATION1
+#undef VARIATION2
+
 FGPUSkinCache::FGPUSkinCache()
 	: bInitialized(false)
 	, FrameCounter(0)
@@ -253,7 +288,9 @@ void FGPUSkinCache::Cleanup()
 }
 
 // Kick off compute shader skinning for incoming chunk, return key for fast lookup in draw passes
-int32 FGPUSkinCache::StartCacheMesh(FRHICommandListImmediate& RHICmdList, int32 Key, FGPUBaseSkinVertexFactory* VertexFactory, FGPUSkinPassthroughVertexFactory* TargetVertexFactory, const FSkelMeshChunk& BatchElement, const FSkeletalMeshObjectGPUSkin* Skin)
+int32 FGPUSkinCache::StartCacheMesh(FRHICommandListImmediate& RHICmdList, int32 Key, FGPUBaseSkinVertexFactory* VertexFactory,
+	FGPUSkinPassthroughVertexFactory* TargetVertexFactory, const FSkelMeshChunk& BatchElement, const FSkeletalMeshObjectGPUSkin* Skin,
+	const FMorphVertexBuffer* MorphVertexBuffer)
 {
 	if (CachedChunksThisFrameCount >= GMaxGPUSkinCacheElementsPerFrame && GFrameNumberRenderThread <= FrameCounter)
 	{
@@ -322,7 +359,6 @@ int32 FGPUSkinCache::StartCacheMesh(FRHICommandListImmediate& RHICmdList, int32
 			Info.TargetVertexFactory = TargetVertexFactory;
 			Info.Skin = Skin;
 			Info.Key = CachedElements.Num();
-			Info.VertexBufferSRVIndex = -1;
 			Info.FrameUpdated = GFrameNumberRenderThread;
 			Info.bExtraBoneInfluences = bExtraBoneInfluences;
 			CachedElements.Add(Info);
@@ -343,7 +379,6 @@ int32 FGPUSkinCache::StartCacheMesh(FRHICommandListImmediate& RHICmdList, int32
 			Info->VertexFactory = VertexFactory;
 			Info->TargetVertexFactory = TargetVertexFactory;
 			Info->Skin = Skin;
-			Info->VertexBufferSRVIndex = -1;
 			Info->FrameUpdated = GFrameNumberRenderThread;
 			Info->bExtraBoneInfluences = bExtraBoneInfluences;
 
@@ -386,41 +421,38 @@ int32 FGPUSkinCache::StartCacheMesh(FRHICommandListImmediate& RHICmdList, int32
 	}
 
 	const FVertexBuffer* SkinVertexBuffer = GPUVertexFactory->GetSkinVertexBuffer();
-	FVertexBufferInfo* VBInfo = nullptr;
 
-	if (ECSInfo->VertexBufferSRVIndex >= 0)
-	{
-		VBInfo = &(CachedVertexBuffers.GetData()[ECSInfo->VertexBufferSRVIndex]);
-	}
-	else
+	FVertexBufferInfo* VBInfo = ReuseOrFindAndCacheSRV(ECSInfo->VertexBufferSRVIndex, *SkinVertexBuffer);
+	
+	FShaderResourceViewRHIRef MorphVertexBufferSRV;
+	uint32 MorphInputStreamFloatOffset = 0;
+
+	if(MorphVertexBuffer != 0)
 	{
-		// time complexity: O(n) 
-		VBInfo = CachedVertexBuffers.FindByKey(*SkinVertexBuffer);
+		FVertexBufferInfo* MorphVBInfo = ReuseOrFindAndCacheSRV(ECSInfo->MorphVertexBufferSRVIndex, *MorphVertexBuffer);
+		MorphVertexBufferSRV = MorphVBInfo->VertexBufferSRV;
 
-		if (VBInfo)
-		{
-			ECSInfo->VertexBufferSRVIndex = VBInfo->Index;
-		}
-	}
+		// in bytes
+		const uint32 MorphStride = sizeof(FMorphGPUSkinVertex);
 
-	if (VBInfo == nullptr)	// Add new cache entry for vertex buffer and SRV
-	{
-		FVertexBufferInfo	Info;
-		Info.VertexBuffer = SkinVertexBuffer;
-		Info.VertexBufferSRV = RHICreateShaderResourceView(SkinVertexBuffer->VertexBufferRHI, sizeof(float), PF_R32_FLOAT);
-		int32 Index = CachedVertexBuffers.Add(Info);
-		VBInfo = &(CachedVertexBuffers.GetData()[Index]);
-		VBInfo->Index = Index;
-		ECSInfo->VertexBufferSRVIndex = Index;
-	}
+		// see GPU code "check(MorphStride == sizeof(float) * 6);"
+		check(MorphStride == sizeof(float) * 6);
 
+		MorphInputStreamFloatOffset = (MorphStride * BatchElement.BaseVertexIndex) / sizeof(float);
+	}
 
 	INC_DWORD_STAT(STAT_GPUSkinCache_TotalNumChunks);
 	INC_DWORD_STAT_BY(STAT_GPUSkinCache_TotalNumVertices, NumVertices);
 
 	auto& ShaderData = GPUVertexFactory->GetShaderData();
 	
-	DispatchSkinCacheProcess(RHICmdList, InputStreamFloatOffset, ECSInfo->StreamOffset, ShaderData.GetBoneBufferForWriting(false, GFrameNumberRenderThread), ShaderData.GetUniformBuffer(), VBInfo, StreamStrides[0], NumVertices, ShaderData.MeshOrigin, ShaderData.MeshExtension, bExtraBoneInfluences, Skin->GetFeatureLevel());
+	// SkinType 0:normal, 1:with morph target, 2:with APEX cloth (not yet implemented)
+	uint32 SkinType = MorphVertexBuffer ? 1 : 0;
+
+	DispatchSkinCacheProcess(RHICmdList, InputStreamFloatOffset, ECSInfo->StreamOffset,
+		ShaderData.GetBoneBufferForWriting(false, GFrameNumberRenderThread), ShaderData.GetUniformBuffer(),
+		VBInfo, StreamStrides[0], NumVertices, ShaderData.MeshOrigin, ShaderData.MeshExtension, bExtraBoneInfluences, Skin->GetFeatureLevel(),
+		SkinType, MorphInputStreamFloatOffset, MorphVertexBufferSRV);
 
 	TargetVertexFactory->UpdateVertexDeclaration(VertexFactory, &SkinCacheBuffer[FrameCounter % GPUSKINCACHE_FRAMES]);
 
@@ -459,6 +491,39 @@ void FGPUSkinCache::TransitionAllToWriteable(FRHICommandList& RHICmdList)
 	}
 }
 
+// @param InOutIndex index -1 if not yet set
+FGPUSkinCache::FVertexBufferInfo* FGPUSkinCache::ReuseOrFindAndCacheSRV(int32& InOutIndex, const FVertexBuffer& VertexBuffer)
+{
+	FVertexBufferInfo* Ret = 0;
+	if (InOutIndex >= 0)
+	{
+		Ret = &(CachedVertexBuffers.GetData()[InOutIndex]);
+	}
+	else
+	{
+		// time complexity: O(n) 
+		Ret = CachedVertexBuffers.FindByKey(VertexBuffer);
+
+		if (Ret)
+		{
+			InOutIndex = Ret - CachedVertexBuffers.GetData();
+		}
+	}
+
+	if (!Ret)
+	{
+		// Add new cache entry for vertex buffer and SRV
+		FVertexBufferInfo	Info;
+		Info.VertexBuffer = &VertexBuffer;
+		Info.VertexBufferSRV = RHICreateShaderResourceView(VertexBuffer.VertexBufferRHI, 4, PF_R32_FLOAT);
+
+		InOutIndex = CachedVertexBuffers.Add(Info);
+		Ret = &(CachedVertexBuffers.GetData()[InOutIndex]);
+	}
+
+	return Ret;
+}
+
 bool FGPUSkinCache::InternalIsElementProcessed(int32 Key) const
 {
 	if (Key >= 0 && Key < CachedElements.Num())
@@ -554,7 +619,10 @@ FGPUSkinCache::FElementCacheStatusInfo* FGPUSkinCache::FindEvictableCacheStatusI
 	}
 }
 
-void FGPUSkinCache::DispatchSkinCacheProcess(FRHICommandListImmediate& RHICmdList, uint32 InputStreamFloatOffset, uint32 OutputBufferFloatOffset, const FVertexBufferAndSRV& BoneBuffer, FUniformBufferRHIRef UniformBuffer, const FVertexBufferInfo* VBInfo, uint32 VertexStride, uint32 VertexCount, const FVector& MeshOrigin, const FVector& MeshExtension, bool bUseExtraBoneInfluences, ERHIFeatureLevel::Type FeatureLevel)
+void FGPUSkinCache::DispatchSkinCacheProcess(FRHICommandListImmediate& RHICmdList, uint32 InputStreamFloatOffset, uint32 OutputBufferFloatOffset,
+	const FVertexBufferAndSRV& BoneBuffer, FUniformBufferRHIRef UniformBuffer, const FGPUSkinCache::FVertexBufferInfo* VBInfo, uint32 VertexStride,
+	uint32 VertexCount, const FVector& MeshOrigin, const FVector& MeshExtension, bool bUseExtraBoneInfluences, ERHIFeatureLevel::Type FeatureLevel,
+	uint32 SkinType, uint32 InMorphInputStreamFloatOffset, FShaderResourceViewRHIParamRef MorphVertexBufferSRV)
 {
 	check(FeatureLevel >= ERHIFeatureLevel::SM5);
 
@@ -564,12 +632,26 @@ void FGPUSkinCache::DispatchSkinCacheProcess(FRHICommandListImmediate& RHICmdLis
 
 	int32 UAVIndex = FrameCounter % GPUSKINCACHE_FRAMES;
 
-	TShaderMapRef<TGPUSkinCacheCS<true> > SkinCacheCS1(GetGlobalShaderMap(FeatureLevel));
-	TShaderMapRef<TGPUSkinCacheCS<false> > SkinCacheCS0(GetGlobalShaderMap(FeatureLevel));
-	FBaseGPUSkinCacheCS* Shader = bUseExtraBoneInfluences ? (FBaseGPUSkinCacheCS*)*SkinCacheCS1 : (FBaseGPUSkinCacheCS*)*SkinCacheCS0;
+	TShaderMapRef<TGPUSkinCacheCS<true,0> > SkinCacheCS10(GetGlobalShaderMap(FeatureLevel));
+	TShaderMapRef<TGPUSkinCacheCS<false,0> > SkinCacheCS00(GetGlobalShaderMap(FeatureLevel));
+	TShaderMapRef<TGPUSkinCacheCS<true,1> > SkinCacheCS11(GetGlobalShaderMap(FeatureLevel));
+	TShaderMapRef<TGPUSkinCacheCS<false,1> > SkinCacheCS01(GetGlobalShaderMap(FeatureLevel));
+
+	FBaseGPUSkinCacheCS* Shader = 0;
+	
+	switch(SkinType)
+	{
+		case 0: Shader = bUseExtraBoneInfluences ? (FBaseGPUSkinCacheCS*)*SkinCacheCS10 : (FBaseGPUSkinCacheCS*)*SkinCacheCS00; break;
+		case 1: Shader = bUseExtraBoneInfluences ? (FBaseGPUSkinCacheCS*)*SkinCacheCS11 : (FBaseGPUSkinCacheCS*)*SkinCacheCS01; break;
+	}
+	
+	check(Shader);
 
 	RHICmdList.SetComputeShader(Shader->GetComputeShader());
-	Shader->SetParameters(RHICmdList, VertexStride, InputStreamFloatOffset, VertexCount, OutputBufferFloatOffset, BoneBuffer, UniformBuffer, VBInfo->VertexBufferSRV, SkinCacheBuffer[UAVIndex], MeshOrigin, MeshExtension);
+	Shader->SetParameters(RHICmdList, VertexStride, InputStreamFloatOffset, VertexCount,
+		OutputBufferFloatOffset, BoneBuffer, UniformBuffer, VBInfo->VertexBufferSRV,
+		SkinCacheBuffer[UAVIndex], MeshOrigin, MeshExtension,
+		SkinType, InMorphInputStreamFloatOffset, MorphVertexBufferSRV);
 
 	RHICmdList.DispatchComputeShader(VertexCountAlign64, 1, 1);
 	Shader->UnsetParameters(RHICmdList);
@@ -598,6 +680,3 @@ void FGPUSkinCache::CVarSinkFunction()
 }
 
 FAutoConsoleVariableSink FGPUSkinCache::CVarSink(FConsoleCommandDelegate::CreateStatic(&CVarSinkFunction));
-
-IMPLEMENT_SHADER_TYPE(template<>,TGPUSkinCacheCS<false>,TEXT("GpuSkinCacheComputeShader"),TEXT("SkinCacheUpdateBatchCS"),SF_Compute);
-IMPLEMENT_SHADER_TYPE(template<>,TGPUSkinCacheCS<true>,TEXT("GpuSkinCacheComputeShader"),TEXT("SkinCacheUpdateBatchCS"),SF_Compute);
diff --git a/Engine/Source/Runtime/Engine/Private/InstancedStaticMesh.cpp b/Engine/Source/Runtime/Engine/Private/InstancedStaticMesh.cpp
--- a/Engine/Source/Runtime/Engine/Private/InstancedStaticMesh.cpp
+++ b/Engine/Source/Runtime/Engine/Private/InstancedStaticMesh.cpp
@@ -1186,7 +1186,7 @@ FBoxSphereBounds UInstancedStaticMeshComponent::CalcBounds(const FTransform& Bou
 #if WITH_EDITOR
 void UInstancedStaticMeshComponent::GetStaticLightingInfo(FStaticLightingPrimitiveInfo& OutPrimitiveInfo, const TArray<ULightComponent*>& InRelevantLights, const FLightingBuildOptions& Options)
 {
-	if (HasValidSettingsForStaticLighting())
+	if (HasValidSettingsForStaticLighting(false))
 	{
 		// create static lighting for LOD 0
 		int32 LightMapWidth = 0;
diff --git a/Engine/Source/Runtime/Engine/Private/Materials/HLSLMaterialTranslator.h b/Engine/Source/Runtime/Engine/Private/Materials/HLSLMaterialTranslator.h
--- a/Engine/Source/Runtime/Engine/Private/Materials/HLSLMaterialTranslator.h
+++ b/Engine/Source/Runtime/Engine/Private/Materials/HLSLMaterialTranslator.h
@@ -741,6 +741,7 @@ public:
 			const FString CollectionName = FString::Printf(TEXT("MaterialCollection%u"), CollectionIndex);
 			FShaderUniformBufferParameter::ModifyCompilationEnvironment(*CollectionName, ParameterCollections[CollectionIndex]->GetUniformBufferStruct(), InPlatform, OutEnvironment);
 		}
+		OutEnvironment.SetDefine(TEXT("IS_MATERIAL_SHADER"), TEXT("1"));
 	}
 
 	void GetSharedInputsMaterialCode(FString& PixelMembersDeclaration, FString& NormalAssignment, FString& PixelMembersInitializationEpilog)
@@ -2128,14 +2129,24 @@ protected:
 		return AddInlinedCodeChunk(MCT_Float3,TEXT("Parameters.LightVector"));
 	}
 
-	virtual int32 ScreenPosition() override
+	virtual int32 ScreenPosition(EMaterialExpressionScreenPositionMapping Mapping) override
 	{
 		if (ShaderFrequency != SF_Pixel && ShaderFrequency != SF_Compute && ShaderFrequency != SF_Vertex)
 		{
 			return Errorf(TEXT("Invalid node used in hull/domain shader input!"));
 		}
 
-		return AddCodeChunk(MCT_Float2,TEXT("ScreenAlignedPosition(GetScreenPosition(Parameters))"));		
+		switch (Mapping)
+		{
+		case MESP_SceneTextureUV:
+			return AddCodeChunk(MCT_Float2, TEXT("ScreenAlignedPosition(GetScreenPosition(Parameters))"));
+		case MESP_ViewportUV:
+			// Works for pixel and vertex shader.  The following commented line would optimize for pixel shader but doesnt work for vertex shader
+			return AddCodeChunk(MCT_Float2, TEXT("BufferUVToViewportUV(ScreenAlignedPosition(GetScreenPosition(Parameters)))"));
+			//return AddCodeChunk(MCT_Float2, TEXT("SvPositionToViewportUV(Parameters.SvPosition)"));
+		default:
+			return Errorf(TEXT("Invalid UV mapping!"));
+		}		
 	}
 
 	virtual int32 ParticleMacroUV() override 
@@ -2507,7 +2518,8 @@ protected:
 		int32 MipValue0Index=INDEX_NONE,
 		int32 MipValue1Index=INDEX_NONE,
 		ETextureMipValueMode MipValueMode=TMVM_None,
-		ESamplerSourceMode SamplerSource=SSM_FromTextureAsset
+		ESamplerSourceMode SamplerSource=SSM_FromTextureAsset,
+		int32 TextureReferenceIndex=INDEX_NONE
 		) override
 	{
 		if(TextureIndex == INDEX_NONE || CoordinateIndex == INDEX_NONE)
@@ -2661,7 +2673,13 @@ protected:
 
 		FString UVs = CoerceParameter(CoordinateIndex, UVsType);
 
-		return AddCodeChunk(
+		if (TextureReferenceIndex != INDEX_NONE)
+		{
+			// Output GPU coordinate analysis. Stubbed unless in TexCoordScaleAnalysis pixel shader.
+			AddCodeChunk(MCT_Float, TEXT("StoreTexCoordScale(%s, %d)"), *UVs, (int)TextureReferenceIndex);
+		}
+
+		int32 SamplingCodeIndex = AddCodeChunk(
 			MCT_Float4,
 			*SampleCode,
 			*TextureName,
@@ -2670,6 +2688,15 @@ protected:
 			*MipValue0Code,
 			*MipValue1Code
 			);
+
+		if (TextureReferenceIndex != INDEX_NONE)
+		{
+			// Output sampling result for the TexCoordScaleAccuracy debug view mode.
+			FString SamplingCode = CoerceParameter(SamplingCodeIndex, MCT_Float4);
+			AddCodeChunk(MCT_Float, TEXT("StoreTexSample(%s, %d)"), *SamplingCode, (int)TextureReferenceIndex);
+		}
+
+		return SamplingCodeIndex;
 	}
 
 	virtual int32 TextureProperty(int32 TextureIndex, EMaterialExposedTextureProperty Property) override
@@ -3010,7 +3037,7 @@ protected:
 			);
 	}
 
-	virtual int32 Texture(UTexture* InTexture,ESamplerSourceMode SamplerSource=SSM_FromTextureAsset) override
+	virtual int32 Texture(UTexture* InTexture,int32& TextureReferenceIndex,ESamplerSourceMode SamplerSource=SSM_FromTextureAsset) override
 	{
 		if (ShaderFrequency != SF_Pixel
 			&& ErrorUnlessFeatureLevelSupported(ERHIFeatureLevel::SM4) == INDEX_NONE)
@@ -3019,7 +3046,7 @@ protected:
 		}
 
 		EMaterialValueType ShaderType = InTexture->GetMaterialType();
-		const int32 TextureReferenceIndex = Material->GetReferencedTextures().Find(InTexture);
+		TextureReferenceIndex = Material->GetReferencedTextures().Find(InTexture);
 
 #if DO_CHECK
 		// UE-3518: Additional pre-assert logging to help determine the cause of this failure.
@@ -3038,7 +3065,7 @@ protected:
 		return AddUniformExpression(new FMaterialUniformExpressionTexture(TextureReferenceIndex, SamplerSource),ShaderType,TEXT(""));
 	}
 
-	virtual int32 TextureParameter(FName ParameterName,UTexture* DefaultValue,ESamplerSourceMode SamplerSource=SSM_FromTextureAsset) override
+	virtual int32 TextureParameter(FName ParameterName,UTexture* DefaultValue,int32& TextureReferenceIndex,ESamplerSourceMode SamplerSource=SSM_FromTextureAsset) override
 	{
 		if (ShaderFrequency != SF_Pixel
 			&& ErrorUnlessFeatureLevelSupported(ERHIFeatureLevel::SM4) == INDEX_NONE)
@@ -3047,7 +3074,7 @@ protected:
 		}
 
 		EMaterialValueType ShaderType = DefaultValue->GetMaterialType();
-		const int32 TextureReferenceIndex = Material->GetReferencedTextures().Find(DefaultValue);
+		TextureReferenceIndex = Material->GetReferencedTextures().Find(DefaultValue);
 		checkf(TextureReferenceIndex != INDEX_NONE, TEXT("Material expression called Compiler->TextureParameter() without implementing UMaterialExpression::GetReferencedTexture properly"));
 		return AddUniformExpression(new FMaterialUniformExpressionTextureParameter(ParameterName, TextureReferenceIndex, SamplerSource),ShaderType,TEXT(""));
 	}
@@ -3147,7 +3174,8 @@ protected:
 		else if (GetFeatureLevel() >= ERHIFeatureLevel::SM4)
 		{
 			FString WeightmapName = FString::Printf(TEXT("Weightmap%d"),WeightmapIndex);
-			int32 TextureCodeIndex = TextureParameter(FName(*WeightmapName),  GEngine->WeightMapPlaceholderTexture);
+			int32 TextureReferenceIndex = INDEX_NONE;
+			int32 TextureCodeIndex = TextureParameter(FName(*WeightmapName), GEngine->WeightMapPlaceholderTexture, TextureReferenceIndex);
 			int32 WeightmapCode = TextureSample(TextureCodeIndex, TextureCoordinate(3, false, false), SAMPLERTYPE_Masks);
 			FString LayerMaskName = FString::Printf(TEXT("LayerMask_%s"),*ParameterName.ToString());
 			return Dot(WeightmapCode,VectorParameter(FName(*LayerMaskName), FLinearColor(1.f,0.f,0.f,0.f)));
diff --git a/Engine/Source/Runtime/Engine/Private/Materials/MaterialExpressions.cpp b/Engine/Source/Runtime/Engine/Private/Materials/MaterialExpressions.cpp
--- a/Engine/Source/Runtime/Engine/Private/Materials/MaterialExpressions.cpp
+++ b/Engine/Source/Runtime/Engine/Private/Materials/MaterialExpressions.cpp
@@ -1317,7 +1317,8 @@ int32 UMaterialExpressionTextureSample::Compile(class FMaterialCompiler* Compile
 {
 	if (Texture || TextureObject.Expression)
 	{
-		int32 TextureCodeIndex = TextureObject.Expression ? TextureObject.Compile(Compiler) : Compiler->Texture(Texture, SamplerSource);
+		int32 TextureReferenceIndex = INDEX_NONE;
+		int32 TextureCodeIndex = TextureObject.Expression ? TextureObject.Compile(Compiler) : Compiler->Texture(Texture, TextureReferenceIndex, SamplerSource);
 
 		UTexture* EffectiveTexture = Texture;
 		EMaterialSamplerType EffectiveSamplerType = (EMaterialSamplerType)SamplerType;
@@ -1360,8 +1361,8 @@ int32 UMaterialExpressionTextureSample::Compile(class FMaterialCompiler* Compile
 				CompileMipValue0(Compiler),
 				CompileMipValue1(Compiler),
 				MipValueMode,
-				SamplerSource
-				);
+				SamplerSource,
+				TextureReferenceIndex);
 		}
 		else
 		{
@@ -1530,16 +1531,17 @@ int32 UMaterialExpressionTextureSampleParameter::Compile(class FMaterialCompiler
 
 	int32 MipValue0Index = CompileMipValue0(Compiler);
 	int32 MipValue1Index = CompileMipValue1(Compiler);
+	int32 TextureReferenceIndex = INDEX_NONE;
 
 	return Compiler->TextureSample(
-					Compiler->TextureParameter(ParameterName, Texture, SamplerSource),
+					Compiler->TextureParameter(ParameterName, Texture, TextureReferenceIndex, SamplerSource),
 					Coordinates.Expression ? Coordinates.Compile(Compiler) : Compiler->TextureCoordinate(ConstCoordinate, false, false),
 					(EMaterialSamplerType)SamplerType,
 					MipValue0Index,
 					MipValue1Index,
 					MipValueMode,
-					SamplerSource
-					);
+					SamplerSource,
+					TextureReferenceIndex);
 }
 
 void UMaterialExpressionTextureSampleParameter::GetCaption(TArray<FString>& OutCaptions) const
@@ -1713,7 +1715,6 @@ int32 UMaterialExpressionTextureObject::Compile(class FMaterialCompiler* Compile
 	{
 		return CompilerError(Compiler, TEXT("Requires valid texture"));
 	}
-
 	return Compiler->Texture(Texture);
 }
 
@@ -4725,11 +4726,19 @@ UMaterialExpressionScreenPosition::UMaterialExpressionScreenPosition(const FObje
 
 int32 UMaterialExpressionScreenPosition::Compile(class FMaterialCompiler* Compiler, int32 OutputIndex, int32 MultiplexIndex)
 {
-	return Compiler->ScreenPosition();
+	return Compiler->ScreenPosition(Mapping);
 }
 
 void UMaterialExpressionScreenPosition::GetCaption(TArray<FString>& OutCaptions) const
 {
+
+#if WITH_EDITOR
+	const UEnum* ScreenPositionMappingEnum = FindObject<UEnum>(NULL, TEXT("Engine.EMaterialExpressionScreenPositionMapping"));
+	check(ScreenPositionMappingEnum);
+
+	const FString MappingDisplayName = ScreenPositionMappingEnum->GetDisplayNameText(Mapping).ToString();
+	OutCaptions.Add(MappingDisplayName);
+#endif
 	OutCaptions.Add(TEXT("ScreenPosition"));
 }
 
diff --git a/Engine/Source/Runtime/Engine/Private/Materials/MaterialShared.cpp b/Engine/Source/Runtime/Engine/Private/Materials/MaterialShared.cpp
--- a/Engine/Source/Runtime/Engine/Private/Materials/MaterialShared.cpp
+++ b/Engine/Source/Runtime/Engine/Private/Materials/MaterialShared.cpp
@@ -1129,6 +1129,10 @@ FMaterial::~FMaterial()
 	}
 
 	FMaterialShaderMap::RemovePendingMaterial(this);
+
+	// If the material becomes invalid, then the debug view material will also be invalid
+	void ClearAllDebugViewMaterials();
+	ClearAllDebugViewMaterials();
 }
 
 // could be more to a more central DBuffer file
diff --git a/Engine/Source/Runtime/Engine/Private/MeshParticleVertexFactory.cpp b/Engine/Source/Runtime/Engine/Private/MeshParticleVertexFactory.cpp
--- a/Engine/Source/Runtime/Engine/Private/MeshParticleVertexFactory.cpp
+++ b/Engine/Source/Runtime/Engine/Private/MeshParticleVertexFactory.cpp
@@ -24,7 +24,11 @@ public:
 		ParticleDirection.Bind(ParameterMap, TEXT("ParticleDirection"));
 		RelativeTime.Bind(ParameterMap, TEXT("RelativeTime"));
 		DynamicParameter.Bind(ParameterMap, TEXT("DynamicParameter"));
-		ParticleColor.Bind(ParameterMap,TEXT("ParticleColor"));
+		ParticleColor.Bind(ParameterMap, TEXT("ParticleColor"));
+		PrevTransform0.Bind(ParameterMap, TEXT("PrevTransform0"));
+		PrevTransform1.Bind(ParameterMap, TEXT("PrevTransform1"));
+		PrevTransform2.Bind(ParameterMap, TEXT("PrevTransform2"));
+		PrevTransformBuffer.Bind(ParameterMap, TEXT("PrevTransformBuffer"));
 	}
 
 	virtual void Serialize(FArchive& Ar) override
@@ -38,6 +42,10 @@ public:
 		Ar << RelativeTime;
 		Ar << DynamicParameter;
 		Ar << ParticleColor;
+		Ar << PrevTransform0;
+		Ar << PrevTransform1;
+		Ar << PrevTransform2;
+		Ar << PrevTransformBuffer;
 	}
 
 	virtual void SetMesh(FRHICommandList& RHICmdList, FShader* Shader,const FVertexFactory* VertexFactory,const FSceneView& View,const FMeshBatchElement& BatchElement,uint32 DataFlags) const override
@@ -52,6 +60,7 @@ public:
 			const FMeshParticleVertexFactory::FBatchParametersCPU* BatchParameters = (const FMeshParticleVertexFactory::FBatchParametersCPU*)BatchElement.UserData;
 			const FMeshParticleInstanceVertex* Vertex = BatchParameters->InstanceBuffer + BatchElement.UserIndex;
 			const FMeshParticleInstanceVertexDynamicParameter* DynamicVertex = BatchParameters->DynamicParameterBuffer + BatchElement.UserIndex;
+			const FMeshParticleInstanceVertexPrevTransform* PrevTransformVertex = BatchParameters->PrevTransformBuffer + BatchElement.UserIndex;
 
 			SetShaderValue(RHICmdList, VertexShaderRHI, Transform1, Vertex->Transform[0]);
 			SetShaderValue(RHICmdList, VertexShaderRHI, Transform2, Vertex->Transform[1]);
@@ -60,12 +69,25 @@ public:
 			SetShaderValue(RHICmdList, VertexShaderRHI, SubUVLerp, Vertex->SubUVLerp);
 			SetShaderValue(RHICmdList, VertexShaderRHI, ParticleDirection, Vertex->Velocity);
 			SetShaderValue(RHICmdList, VertexShaderRHI, RelativeTime, Vertex->RelativeTime);
+
 			if (BatchParameters->DynamicParameterBuffer)
 			{
 				SetShaderValue(RHICmdList, VertexShaderRHI, DynamicParameter, FVector4(DynamicVertex->DynamicValue[0], DynamicVertex->DynamicValue[1], DynamicVertex->DynamicValue[2], DynamicVertex->DynamicValue[3]));
 			}
+
+			if (BatchParameters->PrevTransformBuffer)
+			{
+				SetShaderValue(RHICmdList, VertexShaderRHI, PrevTransform0, PrevTransformVertex->PrevTransform0);
+				SetShaderValue(RHICmdList, VertexShaderRHI, PrevTransform1, PrevTransformVertex->PrevTransform1);
+				SetShaderValue(RHICmdList, VertexShaderRHI, PrevTransform2, PrevTransformVertex->PrevTransform2);
+			}
+
 			SetShaderValue(RHICmdList, VertexShaderRHI, ParticleColor, FVector4(Vertex->Color.Component(0), Vertex->Color.Component(1), Vertex->Color.Component(2), Vertex->Color.Component(3)));
 		}
+		else
+		{
+			SetSRVParameter(RHICmdList, VertexShaderRHI, PrevTransformBuffer, MeshParticleVF->GetPreviousTransformBufferSRV());
+		}
 	}
 
 private:
@@ -79,6 +101,10 @@ private:
 	FShaderParameter RelativeTime;
 	FShaderParameter DynamicParameter;
 	FShaderParameter ParticleColor;
+	FShaderParameter PrevTransform0;
+	FShaderParameter PrevTransform1;
+	FShaderParameter PrevTransform2;
+	FShaderResourceParameter PrevTransformBuffer;
 };
 
 
@@ -206,6 +232,37 @@ void FMeshParticleVertexFactory::SetDynamicParameterBuffer(const FVertexBuffer*
 	}
 }
 
+uint8* FMeshParticleVertexFactory::LockPreviousTransformBuffer(uint32 ParticleCount)
+{
+	const static uint32 ElementSize = sizeof(FVector4);
+	const static uint32 ParticleSize = ElementSize * 3;
+	const uint32 AllocationRequest = ParticleCount * ParticleSize;
+
+	check(!PrevTransformBuffer.MappedBuffer);
+
+	if (AllocationRequest > PrevTransformBuffer.NumBytes)
+	{
+		PrevTransformBuffer.Release();
+		PrevTransformBuffer.Initialize(ElementSize, ParticleCount * 3, PF_A32B32G32R32F, BUF_Dynamic);
+	}
+
+	PrevTransformBuffer.Lock();
+
+	return PrevTransformBuffer.MappedBuffer;
+}
+
+void FMeshParticleVertexFactory::UnlockPreviousTransformBuffer()
+{
+	check(PrevTransformBuffer.MappedBuffer);
+
+	PrevTransformBuffer.Unlock();
+}
+
+FShaderResourceViewRHIParamRef FMeshParticleVertexFactory::GetPreviousTransformBufferSRV() const
+{
+	return PrevTransformBuffer.SRV;
+}
+
 bool FMeshParticleVertexFactory::ShouldCache(EShaderPlatform Platform, const class FMaterial* Material, const class FShaderType* ShaderType)
 {
 	return (Material->IsUsedWithMeshParticles() || Material->IsSpecialEngineMaterial());
diff --git a/Engine/Source/Runtime/Engine/Private/Particles/FXSystem.cpp b/Engine/Source/Runtime/Engine/Private/Particles/FXSystem.cpp
--- a/Engine/Source/Runtime/Engine/Private/Particles/FXSystem.cpp
+++ b/Engine/Source/Runtime/Engine/Private/Particles/FXSystem.cpp
@@ -328,20 +328,37 @@ bool FFXSystem::UsesGlobalDistanceField() const
 	return false;
 }
 
+DECLARE_STATS_GROUP(TEXT("Command List Markers"), STATGROUP_CommandListMarkers, STATCAT_Advanced);
+
+DECLARE_CYCLE_STAT(TEXT("FXPreRender_Prepare"), STAT_CLM_FXPreRender_Prepare, STATGROUP_CommandListMarkers);
+DECLARE_CYCLE_STAT(TEXT("FXPreRender_Simulate"), STAT_CLM_FXPreRender_Simulate, STATGROUP_CommandListMarkers);
+DECLARE_CYCLE_STAT(TEXT("FXPreRender_Finalize"), STAT_CLM_FXPreRender_Finalize, STATGROUP_CommandListMarkers);
+DECLARE_CYCLE_STAT(TEXT("FXPreRender_PrepareCDF"), STAT_CLM_FXPreRender_PrepareCDF, STATGROUP_CommandListMarkers);
+DECLARE_CYCLE_STAT(TEXT("FXPreRender_SimulateCDF"), STAT_CLM_FXPreRender_SimulateCDF, STATGROUP_CommandListMarkers);
+DECLARE_CYCLE_STAT(TEXT("FXPreRender_FinalizeCDF"), STAT_CLM_FXPreRender_FinalizeCDF, STATGROUP_CommandListMarkers);
+
+
 void FFXSystem::PreRender(FRHICommandListImmediate& RHICmdList, const FGlobalDistanceFieldParameterData* GlobalDistanceFieldParameterData)
 {
 	if (RHISupportsGPUParticles(FeatureLevel))
 	{
+		RHICmdList.SetCurrentStat(GET_STATID(STAT_CLM_FXPreRender_Prepare));
 		PrepareGPUSimulation(RHICmdList);
 
+		RHICmdList.SetCurrentStat(GET_STATID(STAT_CLM_FXPreRender_Simulate));
 		SimulateGPUParticles(RHICmdList, EParticleSimulatePhase::Main, NULL, NULL, FTexture2DRHIParamRef(), FTexture2DRHIParamRef());
 
+		RHICmdList.SetCurrentStat(GET_STATID(STAT_CLM_FXPreRender_Finalize));
 		FinalizeGPUSimulation(RHICmdList);
+
+		RHICmdList.SetCurrentStat(GET_STATID(STAT_CLM_FXPreRender_PrepareCDF));
 		PrepareGPUSimulation(RHICmdList);
 
+		RHICmdList.SetCurrentStat(GET_STATID(STAT_CLM_FXPreRender_SimulateCDF));
 		SimulateGPUParticles(RHICmdList, EParticleSimulatePhase::CollisionDistanceField, NULL, GlobalDistanceFieldParameterData, FTexture2DRHIParamRef(), FTexture2DRHIParamRef());
 
-		//particles rendered during basepass may need to read pos/velocity buffers.  must finalize unless we know for sure that nothingin base pass will read.
+		//particles rendered during basepass may need to read pos/velocity buffers.  must finalize unless we know for sure that nothing in base pass will read.
+		RHICmdList.SetCurrentStat(GET_STATID(STAT_CLM_FXPreRender_FinalizeCDF));
 		FinalizeGPUSimulation(RHICmdList);
 	}
 }
diff --git a/Engine/Source/Runtime/Engine/Private/Particles/ParticleComponents.cpp b/Engine/Source/Runtime/Engine/Private/Particles/ParticleComponents.cpp
--- a/Engine/Source/Runtime/Engine/Private/Particles/ParticleComponents.cpp
+++ b/Engine/Source/Runtime/Engine/Private/Particles/ParticleComponents.cpp
@@ -2001,6 +2001,25 @@ bool UParticleSystem::SetLODDistance(int32 LODLevelIndex, float InDistance)
 	return true;
 }
 
+bool UParticleSystem::DoesAnyEmitterHaveMotionBlur(int32 LODLevelIndex) const
+{
+	for (auto& EmitterIter : Emitters)
+	{
+		auto* EmitterLOD = EmitterIter->GetLODLevel(LODLevelIndex);
+		if (!EmitterLOD)
+		{
+			continue;
+		}
+
+		if (EmitterLOD->TypeDataModule && EmitterLOD->TypeDataModule->IsMotionBlurEnabled())
+		{
+			return true;
+		}
+	}
+
+	return false;
+}
+
 #if WITH_EDITOR
 void UParticleSystem::PostEditChangeProperty(FPropertyChangedEvent& PropertyChangedEvent)
 {
@@ -2188,12 +2207,12 @@ void UParticleSystem::PostLoad()
 			if (EmitterSignificance < LowestSignificance)
 			{
 				LowestSignificance = EmitterSignificance;
-			}
 		}
 	}
+	}
 
 	bShouldManageSignificance = GetLowestSignificance() != EParticleSignificanceLevel::Critical && !ContainsEmitterType(UParticleModuleTypeDataBeam2::StaticClass());
-	
+
 	if (LODSettings.Num() == 0)
 	{
 		if (Emitters.Num() > 0)
@@ -3139,7 +3158,7 @@ void UParticleSystemComponent::SetRequiredSignificance(EParticleSignificanceLeve
 		//Set us to be significant again.
 		OnSignificanceChanged(true, true);
 	}
-}
+	}
 
 void UParticleSystemComponent::OnSignificanceChanged(bool bSignificant, bool bApplyToEmitters)
 {
@@ -3156,25 +3175,25 @@ void UParticleSystemComponent::OnSignificanceChanged(bool bSignificant, bool bAp
 				if (Inst)					
 				{
 					if (Inst->SpriteTemplate->IsSignificant(RequiredSignificance))
-					{
+	{
 						Inst->bEnabled = true;
 						Inst->SetHaltSpawning(false);
 						Inst->SetFakeBurstWhenSpawningSupressed(false);
 						++LocalNumSignificantEmitters;
-					}
-				}
-				else
-				{
+	}
+	}
+	else
+	{
 					++LocalNumSignificantEmitters;//Set significant for missing emitters due to other reasons such as detail mode.
 				}
-			}
+	}
 
 			if (LocalNumSignificantEmitters == 0)
-			{
+	{
 				UE_LOG(LogParticles, Warning, TEXT("Setting PSC as significant but it has no significant emitters. %s Template: %s"), *GetFullName(), *Template->GetFullName());
 			}
 			NumSignificantEmitters = LocalNumSignificantEmitters;
-		}
+	}
 	}
 	else
 	{
@@ -3183,25 +3202,25 @@ void UParticleSystemComponent::OnSignificanceChanged(bool bSignificant, bool bAp
 		if (bApplyToEmitters && EmitterInstances.Num() > 0)
 		{
 			//Mark any emitters as significant if needed.
-			for (FParticleEmitterInstance* Inst : EmitterInstances)
-			{
-				if (Inst)
-				{
+	for (FParticleEmitterInstance* Inst : EmitterInstances)
+	{
+		if (Inst)
+		{
 					if(Inst->SpriteTemplate->IsSignificant(RequiredSignificance))
-					{
+			{
 						++LocalNumSignificantEmitters;
-					}
-					else
-					{
+			}
+			else
+			{
 						Inst->bEnabled = false;
-						Inst->SetHaltSpawning(true);
+				Inst->SetHaltSpawning(true);
 						Inst->SetFakeBurstWhenSpawningSupressed(true);
-					}
-				}
 			}
+		}
+	}
 
 			if (LocalNumSignificantEmitters > 0)
-			{
+	{
 				UE_LOG(LogParticles, Warning, TEXT("Setting PSC as not significant but it has some significant emitters. %s Template: %s"), *GetFullName(), *Template->GetFullName());
 			}
 
@@ -3212,10 +3231,10 @@ void UParticleSystemComponent::OnSignificanceChanged(bool bSignificant, bool bAp
 		if (Template->InsignificantReaction == EParticleSystemInsignificanceReaction::Auto)
 		{
 			Reaction = Template->IsLooping() ? EParticleSystemInsignificanceReaction::DisableTick : EParticleSystemInsignificanceReaction::Complete;
-		}
+	}
 
 		switch (Reaction)
-		{
+	{
 		case EParticleSystemInsignificanceReaction::Complete:
 		{
 			Complete();
@@ -3224,7 +3243,7 @@ void UParticleSystemComponent::OnSignificanceChanged(bool bSignificant, bool bAp
 		case EParticleSystemInsignificanceReaction::DisableTick:
 		{
 			SetComponentTickEnabled(false);
-		}
+	}
 			break;
 		case EParticleSystemInsignificanceReaction::DisableTickAndKill:
 		{
@@ -3939,7 +3958,6 @@ void UParticleSystemComponent::UpdateDynamicData(FParticleSystemSceneProxy* Prox
 	{
 		// Create the dynamic data for rendering this particle system
 		FParticleDynamicData* ParticleDynamicData = CreateDynamicData();
-
 		// Render the particles
 #if !(UE_BUILD_SHIPPING || UE_BUILD_TEST)
 		[REDACTED_EMAIL]. Remove thisline  - it is used for debugging purposes...
@@ -4227,7 +4245,7 @@ void UParticleSystemComponent::TickComponent(float DeltaTime, enum ELevelTick Ti
 	bool bManagingSignificance = ShouldManageSignificance();
 	FInGameScopedCycleCounter InGameCycleCounter(GetWorld(), EInGamePerfTrackers::VFXSignificance, EInGamePerfTrackerThreads::GameThread, bManagingSignificance);
 
-	if( Template == nullptr)
+	if (Template == nullptr || Template->Emitters.Num() == 0)
 	{
 		return;
 	}
@@ -4259,11 +4277,6 @@ void UParticleSystemComponent::TickComponent(float DeltaTime, enum ELevelTick Ti
 
 	UWorld* World = GetWorld();
 	check(World);
-	// Bail out if there is no template or there are no instances, or we're running a dedicated server and we don't update on those
-	if ((Template == NULL) || (EmitterInstances.Num() == 0) || (Template->Emitters.Num() == 0))
-	{
-		return;
-	}
 
 	// Bail out if we are running on a dedicated server and we don't want to update on those
 	if ((bUpdateOnDedicatedServer == false) && (GetNetMode() == NM_DedicatedServer))
@@ -4658,7 +4671,7 @@ void UParticleSystemComponent::FinalizeTickComponent()
 
 	// Indicate that we have been ticked since being registered.
 	bJustRegistered = false;
-	
+
 	float CurrTime = GetWorld()->GetTimeSeconds();
 
 	//Are we still significant?
@@ -4669,14 +4682,14 @@ void UParticleSystemComponent::FinalizeTickComponent()
 	else
 	{
 		LastSignificantTime = CurrTime;
-		// If component has just totally finished, call script event.
-		const bool bIsCompleted = HasCompleted();
-		if (bIsCompleted && !bWasCompleted)
-		{
+	// If component has just totally finished, call script event.
+	const bool bIsCompleted = HasCompleted(); 
+	if (bIsCompleted && !bWasCompleted)
+	{
 			Complete();
 		}
 		bWasCompleted = bIsCompleted;
-	}
+		}
 
 	// Update bounding box.
 	if (!bWarmingUp && !bWasCompleted && !Template->bUseFixedRelativeBoundingBox && !bIsTransformDirty)
@@ -4825,24 +4838,24 @@ void UParticleSystemComponent::InitParticles()
 				if (Instance)
 				{
 					Instance->bEnabled = true;
-					Instance->InitParameters(Emitter, this);
-					Instance->Init();
+						Instance->InitParameters(Emitter, this);
+						Instance->Init();
 
-					PreferredLODLevel = FMath::Min(PreferredLODLevel, Emitter->LODLevels.Num());
-					bSetLodLevels |= !bIsFirstCreate;//Only set lod levels if we init any instances and it's not the first creation time.
+						PreferredLODLevel = FMath::Min(PreferredLODLevel, Emitter->LODLevels.Num());
+						bSetLodLevels |= !bIsFirstCreate;//Only set lod levels if we init any instances and it's not the first creation time.
+					}
 				}
-			}
 			else
 			{
 				if (Instance)
 				{
 #if STATS
-					Instance->PreDestructorCall();
+						Instance->PreDestructorCall();
 #endif
-					delete Instance;
-					EmitterInstances[Idx] = NULL;
-					bClearDynamicData = true;
-				}
+						delete Instance;
+						EmitterInstances[Idx] = NULL;
+						bClearDynamicData = true;
+					}
 			}
 		}
 
@@ -5090,7 +5103,7 @@ void UParticleSystemComponent::ActivateSystem(bool bFlagAsJustAttached)
 
 	// System settings may have been lowered. Support late deactivation.
 	const bool bDetailModeAllowsRendering = DetailMode <= GetCurrentDetailMode();
-	
+
 	if( GIsAllowingParticles && bDetailModeAllowsRendering )
 	{
 		// Auto attach if requested
@@ -5127,7 +5140,7 @@ void UParticleSystemComponent::ActivateSystem(bool bFlagAsJustAttached)
 			bIsSignificant = true;
 			RequiredSignificance = EParticleSignificanceLevel::Low;
 
-			//Call this now after any attachment has happened.
+		//Call this now after any attachment has happened.
 			OnSystemPreActivationChange.Broadcast(this, true);
 		}
 
diff --git a/Engine/Source/Runtime/Engine/Private/Particles/ParticleEmitterInstances.cpp b/Engine/Source/Runtime/Engine/Private/Particles/ParticleEmitterInstances.cpp
--- a/Engine/Source/Runtime/Engine/Private/Particles/ParticleEmitterInstances.cpp
+++ b/Engine/Source/Runtime/Engine/Private/Particles/ParticleEmitterInstances.cpp
@@ -2965,6 +2965,7 @@ FParticleMeshEmitterInstance::FParticleMeshEmitterInstance() :
 	, MeshTypeData(NULL)
 	, MeshRotationActive(false)
 	, MeshRotationOffset(0)
+	, MeshMotionBlurOffset(0)
 {
 }
 
@@ -3016,7 +3017,7 @@ bool FParticleMeshEmitterInstance::Resize(int32 NewMaxActiveParticles, bool bSet
 				PayloadData->RotationRateBase			= FVector::ZeroVector;
 			}
 		}
-
+		
 		return true;
 	}
 
@@ -3225,6 +3226,36 @@ void FParticleMeshEmitterInstance::UpdateBoundingBox(float DeltaTime)
 			FPlatformMisc::Prefetch(ParticleData, ParticleStride * ParticleIndices[i+1]);
 			FPlatformMisc::Prefetch(ParticleData, (ParticleIndices[i+1] * ParticleStride) + PLATFORM_CACHE_LINE_SIZE);
 
+			if (MeshMotionBlurOffset)
+			{
+				FMeshRotationPayloadData* RotationPayloadData = (FMeshRotationPayloadData*)((uint8*)&Particle + MeshRotationOffset);
+				FMeshMotionBlurPayloadData* MotionBlurPayloadData = (FMeshMotionBlurPayloadData*)((uint8*)&Particle + MeshRotationOffset);
+				MotionBlurPayloadData->BaseParticlePrevRotation = Particle.Rotation;
+				MotionBlurPayloadData->BaseParticlePrevVelocity = Particle.Velocity;
+				MotionBlurPayloadData->BaseParticlePrevSize = Particle.Size;
+				MotionBlurPayloadData->PayloadPrevRotation = RotationPayloadData->Rotation;
+
+				if (CameraPayloadOffset)
+				{
+					const FCameraOffsetParticlePayload* CameraPayload = (const FCameraOffsetParticlePayload*)((const uint8*)&Particle + CameraPayloadOffset);
+					MotionBlurPayloadData->PayloadPrevCameraOffset = CameraPayload->Offset;
+				}
+				else
+				{
+					MotionBlurPayloadData->PayloadPrevCameraOffset = 0.0f;
+				}
+
+				if (OrbitModuleOffset)
+				{
+					const FOrbitChainModuleInstancePayload* OrbitPayload = (const FOrbitChainModuleInstancePayload*)((const uint8*)&Particle + OrbitModuleOffset);
+					MotionBlurPayloadData->PayloadPrevOrbitOffset = OrbitPayload->Offset;
+				}
+				else
+				{
+					MotionBlurPayloadData->PayloadPrevOrbitOffset = FVector::ZeroVector;
+				}
+			}
+
 			// Do linear integrator and update bounding box
 			Particle.OldLocation = Particle.Location;
 			if ((Particle.Flags & STATE_Particle_Freeze) == 0)
@@ -3296,8 +3327,20 @@ void FParticleMeshEmitterInstance::UpdateBoundingBox(float DeltaTime)
 uint32 FParticleMeshEmitterInstance::RequiredBytes()
 {
 	uint32 uiBytes = FParticleEmitterInstance::RequiredBytes();
+
 	MeshRotationOffset	= PayloadOffset + uiBytes;
 	uiBytes += sizeof(FMeshRotationPayloadData);
+
+	if (MeshTypeData)
+	{
+		const auto* MeshTD = static_cast<const UParticleModuleTypeDataMesh*>(MeshTypeData);
+		if (MeshTypeData->bEnableMotionBlur)
+		{
+			MeshMotionBlurOffset = PayloadOffset + uiBytes;
+			uiBytes += sizeof(FMeshMotionBlurPayloadData);
+		}
+	}
+
 	return uiBytes;
 }
 
@@ -3338,6 +3381,37 @@ void FParticleMeshEmitterInstance::PostSpawn(FBaseParticle* Particle, float Inte
 
 	FVector InitialOrient = MeshTypeData->RollPitchYawRange.GetValue(SpawnTime, 0, 0, &MeshTypeData->RandomStream);
 	PayloadData->InitialOrientation = InitialOrient;
+
+	if (MeshMotionBlurOffset)
+	{
+		FMeshRotationPayloadData* RotationPayloadData = (FMeshRotationPayloadData*)((uint8*)Particle + MeshRotationOffset);
+		FMeshMotionBlurPayloadData* MotionBlurPayloadData = (FMeshMotionBlurPayloadData*)((uint8*)Particle + MeshMotionBlurOffset);
+
+		MotionBlurPayloadData->BaseParticlePrevRotation = Particle->Rotation;
+		MotionBlurPayloadData->BaseParticlePrevVelocity = Particle->Velocity;
+		MotionBlurPayloadData->BaseParticlePrevSize = Particle->Size;
+		MotionBlurPayloadData->PayloadPrevRotation = RotationPayloadData->Rotation;
+
+		if (CameraPayloadOffset)
+		{
+			const FCameraOffsetParticlePayload* CameraPayload = (const FCameraOffsetParticlePayload*)((const uint8*)Particle + CameraPayloadOffset);
+			MotionBlurPayloadData->PayloadPrevCameraOffset = CameraPayload->Offset;
+		}
+		else
+		{
+			MotionBlurPayloadData->PayloadPrevCameraOffset = 0.0f;
+		}
+
+		if (OrbitModuleOffset)
+		{
+			const FOrbitChainModuleInstancePayload* OrbitPayload = (const FOrbitChainModuleInstancePayload*)((const uint8*)Particle + OrbitModuleOffset);
+			MotionBlurPayloadData->PayloadPrevOrbitOffset = OrbitPayload->Offset;
+		}
+		else
+		{
+			MotionBlurPayloadData->PayloadPrevOrbitOffset = FVector::ZeroVector;
+		}
+	}
 }
 
 bool FParticleMeshEmitterInstance::IsDynamicDataRequired(UParticleLODLevel* CurrentLODLevel)
@@ -3622,6 +3696,7 @@ bool FParticleMeshEmitterInstance::FillReplayData( FDynamicEmitterReplayDataBase
 	NewReplayData->SubImages_Horizontal = LODLevel->RequiredModule->SubImages_Horizontal;
 	NewReplayData->SubImages_Vertical = LODLevel->RequiredModule->SubImages_Vertical;
 	NewReplayData->MeshRotationOffset = MeshRotationOffset;
+	NewReplayData->MeshMotionBlurOffset = MeshMotionBlurOffset;
 	NewReplayData->bMeshRotationActive = MeshRotationActive;
 	NewReplayData->MeshAlignment = MeshTypeData->MeshAlignment;
 
diff --git a/Engine/Source/Runtime/Engine/Private/Particles/ParticleSystemRender.cpp b/Engine/Source/Runtime/Engine/Private/Particles/ParticleSystemRender.cpp
--- a/Engine/Source/Runtime/Engine/Private/Particles/ParticleSystemRender.cpp
+++ b/Engine/Source/Runtime/Engine/Private/Particles/ParticleSystemRender.cpp
@@ -21,6 +21,7 @@
 #include "ParticleBeamTrailVertexFactory.h"
 #include "MeshBatch.h"
 #include "Particles/SubUVAnimation.h"
+#include "../../Renderer/Private/ScenePrivate.h"
 
 DECLARE_CYCLE_STAT(TEXT("ParticleSystemSceneProxy GetMeshElements"), STAT_FParticleSystemSceneProxy_GetMeshElements, STATGROUP_Particles);
 DECLARE_CYCLE_STAT(TEXT("DynamicSpriteEmitterData GetDynamicMeshElementsEmitter GetParticleOrderData"), STAT_FDynamicSpriteEmitterData_GetDynamicMeshElementsEmitter_GetParticleOrderData, STATGROUP_Particles);
@@ -318,6 +319,28 @@ void ComputeLockedAxes(EParticleAxisLock LockAxisFlag, const FMatrix& LocalToWor
 	}
 }
 
+FORCEINLINE FVector GetCameraOffset(
+	float CameraPayloadOffset,
+	FVector DirToCamera
+	)
+{
+	float CheckSize = DirToCamera.SizeSquared();
+	DirToCamera.Normalize();
+
+	if (CheckSize > (CameraPayloadOffset * CameraPayloadOffset))
+	{
+		return DirToCamera * CameraPayloadOffset;
+	}
+	else
+	{
+		// If the offset will push the particle behind the camera, then push it 
+		// WAY behind the camera. This is a hack... but in the case of 
+		// PSA_Velocity, it is required to ensure that the particle doesn't 
+		// 'spin' flat and come into view.
+		return DirToCamera * CameraPayloadOffset * HALF_WORLD_MAX;
+	}
+}
+
 /**
  *	Helper function for retrieving the camera offset payload of a particle.
  *
@@ -330,28 +353,17 @@ void ComputeLockedAxes(EParticleAxisLock LockAxisFlag, const FMatrix& LocalToWor
  */
 FORCEINLINE FVector GetCameraOffsetFromPayload(
 	int32 InCameraPayloadOffset,
-	const FBaseParticle& InParticle, 
-	const FVector& InPosition,
-	const FVector& InCameraPosition )
+	const FBaseParticle& InParticle,
+	const FVector& InParticlePosition,
+	const FVector& InCameraPosition 
+	)
 {
 	checkSlow(InCameraPayloadOffset > 0);
 
-	FVector DirToCamera = InCameraPosition - InPosition;
-	float CheckSize = DirToCamera.SizeSquared();
-	DirToCamera.Normalize();
+	FVector DirToCamera = InCameraPosition - InParticlePosition;
 	FCameraOffsetParticlePayload* CameraPayload = ((FCameraOffsetParticlePayload*)((uint8*)(&InParticle) + InCameraPayloadOffset));
-	if (CheckSize > (CameraPayload->Offset * CameraPayload->Offset))
-	{
-		return DirToCamera * CameraPayload->Offset;
-	}
-	else
-	{
-		// If the offset will push the particle behind the camera, then push it 
-		// WAY behind the camera. This is a hack... but in the case of 
-		// PSA_Velocity, it is required to ensure that the particle doesn't 
-		// 'spin' flat and come into view.
-		return DirToCamera * CameraPayload->Offset * HALF_WORLD_MAX;
-	}
+	
+	return GetCameraOffset(CameraPayload->Offset, DirToCamera);
 }
 
 void FDynamicSpriteEmitterDataBase::SortSpriteParticles(int32 SortMode, bool bLocalSpace, 
@@ -1422,14 +1434,17 @@ void FDynamicMeshEmitterData::UpdateRenderThreadResourcesEmitter(const FParticle
 	{
 		FMeshParticleUniformParameters UniformParameters;
 		UniformParameters.SubImageSize = FVector4(
-		1.0f / SourceData->SubImages_Horizontal,
-		1.0f / SourceData->SubImages_Vertical,
-		0,0);
+			1.0f / SourceData->SubImages_Horizontal,
+			1.0f / SourceData->SubImages_Vertical,
+			0,0);
 
 		// A weight is used to determine whether the mesh texture coordinates or SubUVs are passed from the vertex shader to the pixel shader.
 		const uint32 TexCoordWeight = (SourceData->SubUVDataOffset > 0) ? 1 : 0;
 		UniformParameters.TexCoordWeightA = TexCoordWeight;
 		UniformParameters.TexCoordWeightB = 1 - TexCoordWeight;
+
+		UniformParameters.PrevTransformAvailable = Source.MeshMotionBlurOffset ? 1 : 0;
+
 		UniformBuffer = FMeshParticleUniformBufferRef::CreateUniformBufferImmediate( UniformParameters, UniformBuffer_SingleFrame );
 	}
 }
@@ -1443,7 +1458,7 @@ void FDynamicMeshEmitterData::UpdateRenderThreadResourcesEmitter(const FParticle
  */
 void FDynamicMeshEmitterData::ReleaseRenderThreadResources(const FParticleSystemSceneProxy* InOwnerProxy)
 {
-	return FDynamicSpriteEmitterDataBase::ReleaseRenderThreadResources( InOwnerProxy );
+	return FDynamicSpriteEmitterDataBase::ReleaseRenderThreadResources(InOwnerProxy);
 	UniformBuffer.SafeRelease();
 }
 
@@ -1463,20 +1478,23 @@ class FMeshParticleInstanceVertices : public FOneFrameResource
 public:
 	TArray<FMeshParticleInstanceVertex, SceneRenderingAllocator> InstanceDataAllocationsCPU;
 	TArray<FMeshParticleInstanceVertexDynamicParameter, SceneRenderingAllocator> DynamicParameterDataAllocationsCPU;
+	TArray<FMeshParticleInstanceVertexPrevTransform, SceneRenderingAllocator> PrevTransformDataAllocationsCPU;
 };
 
 
 FParticleVertexFactoryBase *FDynamicMeshEmitterData::CreateVertexFactory()
 {
-		FMeshParticleVertexFactory *VertexFactory = ConstructMeshParticleVertexFactory();
-		VertexFactory->SetParticleFactoryType(PVFT_Mesh);
-		SetupVertexFactory(VertexFactory, StaticMesh->RenderData->LODResources[0]);
-		const int32 InstanceVertexStride = GetDynamicVertexStride(ERHIFeatureLevel::Type::SM5);	// featurelevel is ignored
-		const int32 DynamicParameterVertexStride = GetDynamicParameterVertexStride();
-		VertexFactory->SetStrides(InstanceVertexStride, GetDynamicParameterVertexStride());
-		VertexFactory->InitResource();
-		return VertexFactory;
+	FMeshParticleVertexFactory *VertexFactory = ConstructMeshParticleVertexFactory();
+
+	VertexFactory->SetParticleFactoryType(PVFT_Mesh);
+	SetupVertexFactory(VertexFactory, StaticMesh->RenderData->LODResources[0]);
 
+	const int32 InstanceVertexStride = GetDynamicVertexStride(ERHIFeatureLevel::Type::SM5);	// featurelevel is ignored
+	const int32 DynamicParameterVertexStride = GetDynamicParameterVertexStride();
+	VertexFactory->SetStrides(InstanceVertexStride, DynamicParameterVertexStride);
+	VertexFactory->InitResource();
+
+	return VertexFactory;
 }
 
 
@@ -1517,30 +1535,41 @@ void FDynamicMeshEmitterData::GetDynamicMeshElementsEmitter(const FParticleSyste
 			{
 				FGlobalDynamicVertexBuffer::FAllocation Allocation = FGlobalDynamicVertexBuffer::Get().Allocate( ParticleCount * InstanceVertexStride );
 				FGlobalDynamicVertexBuffer::FAllocation DynamicParameterAllocation;
+				uint8* PrevTransformBuffer = nullptr;
 
 				if (bUsesDynamicParameter)
 				{
 					DynamicParameterAllocation = FGlobalDynamicVertexBuffer::Get().Allocate( ParticleCount * DynamicParameterVertexStride );
 				}
 
+				if (Source.MeshMotionBlurOffset)
+				{
+					PrevTransformBuffer = MeshVertexFactory->LockPreviousTransformBuffer(ParticleCount);
+				}
+
 				if(Allocation.IsValid() && (!bUsesDynamicParameter || DynamicParameterAllocation.IsValid()))
 				{
 					// Fill instance buffer.
 					if (Collector.ShouldUseTasks())
 					{
 						Collector.AddTask(
-							[this, View, Proxy, Allocation, DynamicParameterAllocation]()
+							[this, View, Proxy, Allocation, DynamicParameterAllocation, PrevTransformBuffer]()
 							{
-					GetInstanceData(Allocation.Buffer, DynamicParameterAllocation.Buffer, Proxy, View);
-				}
+								GetInstanceData(Allocation.Buffer, DynamicParameterAllocation.Buffer, PrevTransformBuffer, Proxy, View);
+							}
 						);
 					}
 					else
 					{
-						GetInstanceData(Allocation.Buffer, DynamicParameterAllocation.Buffer, Proxy, View);
+						GetInstanceData(Allocation.Buffer, DynamicParameterAllocation.Buffer, PrevTransformBuffer, Proxy, View);
 					}
 				}
 
+				if (Source.MeshMotionBlurOffset)
+				{
+					MeshVertexFactory->UnlockPreviousTransformBuffer();
+				}
+
 				MeshVertexFactory->SetInstanceBuffer(Allocation.VertexBuffer, Allocation.VertexOffset, InstanceVertexStride);
 				MeshVertexFactory->SetDynamicParameterBuffer(DynamicParameterAllocation.VertexBuffer, DynamicParameterAllocation.VertexOffset , GetDynamicParameterVertexStride());
 			}
@@ -1556,20 +1585,38 @@ void FDynamicMeshEmitterData::GetDynamicMeshElementsEmitter(const FParticleSyste
 					InstanceVerticesCPU->DynamicParameterDataAllocationsCPU.AddUninitialized(ParticleCount);
 				}
 
+				if (Source.MeshMotionBlurOffset)
+				{
+					InstanceVerticesCPU->PrevTransformDataAllocationsCPU.Reset(ParticleCount);
+					InstanceVerticesCPU->PrevTransformDataAllocationsCPU.AddUninitialized(ParticleCount);
+				}
+
 				// Fill instance buffer.
 				if (Collector.ShouldUseTasks())
 				{
 					Collector.AddTask(
 						[this, View, Proxy, InstanceVerticesCPU]()
 						{
-							GetInstanceData((void*)InstanceVerticesCPU->InstanceDataAllocationsCPU.GetData(), (void*)InstanceVerticesCPU->DynamicParameterDataAllocationsCPU.GetData(), Proxy, View);
+							GetInstanceData(
+								(void*)InstanceVerticesCPU->InstanceDataAllocationsCPU.GetData(), 
+								(void*)InstanceVerticesCPU->DynamicParameterDataAllocationsCPU.GetData(), 
+								(void*)InstanceVerticesCPU->PrevTransformDataAllocationsCPU.GetData(), 
+								Proxy, 
+								View
+								);
 						}
 					);
 				}
 				else
 				{
-				GetInstanceData((void*)InstanceVerticesCPU->InstanceDataAllocationsCPU.GetData(), (void*)InstanceVerticesCPU->DynamicParameterDataAllocationsCPU.GetData(), Proxy, View);
-			}
+					GetInstanceData(
+						(void*)InstanceVerticesCPU->InstanceDataAllocationsCPU.GetData(), 
+						(void*)InstanceVerticesCPU->DynamicParameterDataAllocationsCPU.GetData(), 
+						(void*)InstanceVerticesCPU->PrevTransformDataAllocationsCPU.GetData(), 
+						Proxy, 
+						View
+						);
+				}
 			}
 
 			Proxy->UpdateWorldSpacePrimitiveUniformBuffer();
@@ -1644,6 +1691,7 @@ void FDynamicMeshEmitterData::GetDynamicMeshElementsEmitter(const FParticleSyste
 						FMeshParticleVertexFactory::FBatchParametersCPU& BatchParameters = Collector.AllocateOneFrameResource<FMeshParticleVertexFactory::FBatchParametersCPU>();
 						BatchParameters.InstanceBuffer = InstanceVerticesCPU->InstanceDataAllocationsCPU.GetData();
 						BatchParameters.DynamicParameterBuffer = InstanceVerticesCPU->DynamicParameterDataAllocationsCPU.GetData();
+						BatchParameters.PrevTransformBuffer = InstanceVerticesCPU->PrevTransformDataAllocationsCPU.GetData();
 						BatchElement.UserData = &BatchParameters;
 						BatchElement.UserIndex = 0;
 
@@ -1683,49 +1731,191 @@ void FDynamicMeshEmitterData::GatherSimpleLights(const FParticleSystemSceneProxy
 	GatherParticleLightData(Source, Proxy->GetLocalToWorld(), ViewFamily, OutParticleLights);
 }
 
-void FDynamicMeshEmitterData::GetParticleTransform(const FBaseParticle& InParticle, const FVector& CameraPosition, const FVector& CameraFacingOpVector, const FQuat& PointToLockedAxis, 
-	const FParticleSystemSceneProxy* Proxy, const FSceneView* View, FMatrix& OutTransformMat) const
-	
+void FDynamicMeshEmitterData::GetParticleTransform(
+	const FBaseParticle& InParticle,
+	const FParticleSystemSceneProxy* Proxy,
+	const FSceneView* View,
+	FMatrix& OutTransformMat
+	) const
 {
 	const uint8* ParticleBase = (const uint8*)&InParticle;
-	OutTransformMat = FMatrix::Identity;
-	
-	FMatrix LocalToWorld;
-	FTranslationMatrix kTransMat(FVector::ZeroVector);
-	FScaleMatrix kScaleMat(FVector(1.0f));
-	FQuat kLockedAxisQuat = FQuat::Identity;
-	FVector Location;
-	FVector ScaledSize;
-	FVector	DirToCamera;
-	FVector	LocalSpaceFacingAxis;
-	FVector	LocalSpaceUpAxis;
-	FQuat PointTo = PointToLockedAxis;
 
-	// Initialize particle position and scale.
-	FVector ParticlePosition(InParticle.Location);
+	const FMeshRotationPayloadData* RotationPayload = (const FMeshRotationPayloadData*)((const uint8*)&InParticle + Source.MeshRotationOffset);
+	FVector RotationPayloadInitialOrientation = RotationPayload->InitialOrientation;
+	FVector RotationPayloadRotation = RotationPayload->Rotation;
+
+	FVector CameraPayloadCameraOffset = FVector::ZeroVector;
 	if (Source.CameraPayloadOffset != 0)
 	{
-		const FVector CameraOffset = GetCameraOffsetFromPayload(Source.CameraPayloadOffset, InParticle, ParticlePosition, CameraPosition);
-		ParticlePosition += CameraOffset;
+		// Put the camera origin in the appropriate coordinate space.
+		FVector CameraPosition = View->ViewMatrices.ViewOrigin;
+		if (Source.bUseLocalSpace)
+		{
+			const FMatrix InvLocalToWorld = Proxy->GetLocalToWorld().Inverse();
+			CameraPosition = InvLocalToWorld.TransformPosition(CameraPosition);
+		}
+
+		CameraPayloadCameraOffset = GetCameraOffsetFromPayload(Source.CameraPayloadOffset, InParticle, InParticle.Location, CameraPosition);
+	}
+
+	FVector OrbitPayloadOrbitOffset = FVector::ZeroVector;
+	if (Source.OrbitModuleOffset != 0)
+	{
+		int32 CurrentOffset = Source.OrbitModuleOffset;
+		PARTICLE_ELEMENT(FOrbitChainModuleInstancePayload, OrbitPayload);
+		OrbitPayloadOrbitOffset = OrbitPayload.Offset;
 	}
 
+	CalculateParticleTransform(
+		Proxy->GetLocalToWorld(),
+		InParticle.Location,
+		InParticle.Rotation,
+		InParticle.Velocity,
+		InParticle.Size,
+		RotationPayloadInitialOrientation,
+		RotationPayloadRotation,
+		CameraPayloadCameraOffset,
+		OrbitPayloadOrbitOffset,
+		View->ViewMatrices.ViewOrigin,
+		View->GetViewDirection(),
+		OutTransformMat
+		);
+}
+
+void FDynamicMeshEmitterData::GetParticlePrevTransform(
+	const FBaseParticle& InParticle,
+	const FParticleSystemSceneProxy* Proxy,
+	const FSceneView* View,
+	FMatrix& OutTransformMat
+	) const
+{
+	const FMeshRotationPayloadData* RotationPayload = (const FMeshRotationPayloadData*)((const uint8*)&InParticle + Source.MeshRotationOffset);
+	const FMeshMotionBlurPayloadData* MotionBlurPayload = (const FMeshMotionBlurPayloadData*)((const uint8*)&InParticle + Source.MeshMotionBlurOffset);
+
+	const auto* ViewInfo = static_cast<const FViewInfo*>(View);
+
+	FVector CameraPayloadCameraOffset = FVector::ZeroVector;
+	if (Source.CameraPayloadOffset != 0)
+	{
+		// Put the camera origin in the appropriate coordinate space.
+		FVector CameraPosition = ViewInfo->PrevViewMatrices.ViewOrigin;
+		if (Source.bUseLocalSpace)
+		{
+			const FMatrix InvLocalToWorld = Proxy->GetLocalToWorld().Inverse();
+			CameraPosition = InvLocalToWorld.TransformPosition(CameraPosition);
+		}
+
+		CameraPayloadCameraOffset = GetCameraOffset(MotionBlurPayload->PayloadPrevCameraOffset, CameraPosition - InParticle.OldLocation);
+	}
+
+	CalculateParticleTransform(
+		Proxy->GetLocalToWorld(),
+		InParticle.OldLocation,
+		MotionBlurPayload->BaseParticlePrevRotation,
+		MotionBlurPayload->BaseParticlePrevVelocity,
+		MotionBlurPayload->BaseParticlePrevSize,
+		RotationPayload->InitialOrientation,
+		MotionBlurPayload->PayloadPrevRotation,
+		CameraPayloadCameraOffset,
+		MotionBlurPayload->PayloadPrevOrbitOffset,
+		ViewInfo->PrevViewMatrices.ViewOrigin,
+		ViewInfo->GetPrevViewDirection(),
+		OutTransformMat
+		);
+}
+
+void FDynamicMeshEmitterData::CalculateParticleTransform(
+	const FMatrix& ProxyLocalToWorld,
+	const FVector& ParticleLocation,
+		  float    ParticleRotation,
+	const FVector& ParticleVelocity,
+	const FVector& ParticleSize,
+	const FVector& ParticlePayloadInitialOrientation,
+	const FVector& ParticlePayloadRotation,
+	const FVector& ParticlePayloadCameraOffset,
+	const FVector& ParticlePayloadOrbitOffset,
+	const FVector& ViewOrigin,
+	const FVector& ViewDirection,
+	FMatrix& OutTransformMat
+	) const
+{
+	FVector CameraFacingOpVector = FVector::ZeroVector;
+	if (CameraFacingOption != XAxisFacing_NoUp)
+	{
+		switch (CameraFacingOption)
+		{
+		case XAxisFacing_ZUp:
+			CameraFacingOpVector = FVector(0.0f, 0.0f, 1.0f);
+			break;
+		case XAxisFacing_NegativeZUp:
+			CameraFacingOpVector = FVector(0.0f, 0.0f, -1.0f);
+			break;
+		case XAxisFacing_YUp:
+			CameraFacingOpVector = FVector(0.0f, 1.0f, 0.0f);
+			break;
+		case XAxisFacing_NegativeYUp:
+			CameraFacingOpVector = FVector(0.0f, -1.0f, 0.0f);
+			break;
+		case LockedAxis_YAxisFacing:
+		case VelocityAligned_YAxisFacing:
+			CameraFacingOpVector = FVector(0.0f, 1.0f, 0.0f);
+			break;
+		case LockedAxis_NegativeYAxisFacing:
+		case VelocityAligned_NegativeYAxisFacing:
+			CameraFacingOpVector = FVector(0.0f, -1.0f, 0.0f);
+			break;
+		case LockedAxis_ZAxisFacing:
+		case VelocityAligned_ZAxisFacing:
+			CameraFacingOpVector = FVector(0.0f, 0.0f, 1.0f);
+			break;
+		case LockedAxis_NegativeZAxisFacing:
+		case VelocityAligned_NegativeZAxisFacing:
+			CameraFacingOpVector = FVector(0.0f, 0.0f, -1.0f);
+			break;
+		}
+	}
+
+	FQuat PointToLockedAxis;
+	if (bUseMeshLockedAxis == true)
+	{
+		// facing axis is taken to be the local x axis.	
+		PointToLockedAxis = FQuat::FindBetweenNormals(FVector(1, 0, 0), Source.LockedAxis);
+	}
+
+	OutTransformMat = FMatrix::Identity;
+
+	FTranslationMatrix kTransMat(FVector::ZeroVector);
+	FScaleMatrix kScaleMat(FVector(1.0f));
+	FQuat kLockedAxisQuat = FQuat::Identity;
+
+	FVector ParticlePosition(ParticleLocation + ParticlePayloadCameraOffset);
 	kTransMat.M[3][0] = ParticlePosition.X;
 	kTransMat.M[3][1] = ParticlePosition.Y;
 	kTransMat.M[3][2] = ParticlePosition.Z;
-	ScaledSize = InParticle.Size * Source.Scale;
+
+	FVector ScaledSize = ParticleSize * Source.Scale;
 	kScaleMat.M[0][0] = ScaledSize.X;
 	kScaleMat.M[1][1] = ScaledSize.Y;
 	kScaleMat.M[2][2] = ScaledSize.Z;
 
 	FMatrix kRotMat(FMatrix::Identity);
-	LocalToWorld = Proxy->GetLocalToWorld();
-	if (bUseCameraFacing == true)
+	FMatrix LocalToWorld = ProxyLocalToWorld;
+
+	FVector	LocalSpaceFacingAxis;
+	FVector	LocalSpaceUpAxis;
+	FVector Location;
+	FVector	DirToCamera;
+	FQuat PointTo = PointToLockedAxis;
+
+	if (bUseCameraFacing)
 	{
 		Location = ParticlePosition;
-		FVector	VelocityDirection = InParticle.Velocity;
+		FVector	VelocityDirection = ParticleVelocity;
+
 		if (Source.bUseLocalSpace)
 		{
 			bool bClearLocal2World = false;
+
 			// Transform the location to world space
 			Location = LocalToWorld.TransformPosition(Location);
 			if (CameraFacingOption <= XAxisFacing_NegativeYUp)
@@ -1748,20 +1938,20 @@ void FDynamicMeshEmitterData::GetParticleTransform(const FBaseParticle& InPartic
 		}
 		VelocityDirection.Normalize();
 
-		if( bFaceCameraDirectionRatherThanPosition )
+		if (bFaceCameraDirectionRatherThanPosition)
 		{
-			DirToCamera = -View->GetViewDirection();
+			DirToCamera = -ViewDirection;
 		}
 		else
 		{
-			DirToCamera	= View->ViewMatrices.ViewOrigin - Location;
+			DirToCamera = ViewOrigin - Location;
 		}
-		
+
 		DirToCamera.Normalize();
-		if (DirToCamera.SizeSquared() <	0.5f)
+		if (DirToCamera.SizeSquared() < 0.5f)
 		{
 			// Assert possible if DirToCamera is not normalized
-			DirToCamera	= FVector(1,0,0);
+			DirToCamera = FVector(1, 0, 0);
 		}
 
 		bool bFacingDirectionIsValid = true;
@@ -1778,6 +1968,7 @@ void FDynamicMeshEmitterData::GetParticleTransform(const FBaseParticle& InPartic
 					// We have to fudge it
 					bFacingDirectionIsValid = false;
 				}
+
 				// Velocity align the X-axis, and camera face the selected axis
 				PointTo = FQuat::FindBetweenNormals(FVector(1.0f, 0.0f, 0.0f), VelocityDirection);
 				FacingDir = VelocityDirection;
@@ -1786,9 +1977,9 @@ void FDynamicMeshEmitterData::GetParticleTransform(const FBaseParticle& InPartic
 			else if (CameraFacingOption <= XAxisFacing_NegativeYUp)
 			{
 				// Camera face the X-axis, and point the selected axis towards the world up
-				PointTo = FQuat::FindBetweenNormals(FVector(1,0,0), DirToCamera);
+				PointTo = FQuat::FindBetweenNormals(FVector(1, 0, 0), DirToCamera);
 				FacingDir = DirToCamera;
-				DesiredDir = FVector(0,0,1);
+				DesiredDir = FVector(0, 0, 1);
 			}
 			else
 			{
@@ -1796,14 +1987,14 @@ void FDynamicMeshEmitterData::GetParticleTransform(const FBaseParticle& InPartic
 				// PointTo will contain quaternion for locked axis rotation.
 				FacingDir = Source.LockedAxis;
 
-				if(Source.bUseLocalSpace)
+				if (Source.bUseLocalSpace)
 				{
 					//Transform the direction vector into local space.
 					DesiredDir = LocalToWorld.GetTransposed().TransformVector(DirToCamera);
-				}	
+				}
 				else
 				{
-					DesiredDir = DirToCamera;	
+					DesiredDir = DirToCamera;
 				}
 			}
 
@@ -1817,28 +2008,28 @@ void FDynamicMeshEmitterData::GetParticleTransform(const FBaseParticle& InPartic
 			{
 				if (bFacingDirectionIsValid)
 				{
-					FQuat AddedRotation = FQuat(FacingDir, InParticle.Rotation);
+					FQuat AddedRotation = FQuat(FacingDir, ParticleRotation);
 					kLockedAxisQuat = (AddedRotation * PointTo);
 				}
 			}
 			else
 			{
-				FQuat AddedRotation = FQuat(DirToCamera, InParticle.Rotation);
+				FQuat AddedRotation = FQuat(DirToCamera, ParticleRotation);
 				kLockedAxisQuat = (AddedRotation * PointTo);
 			}
 		}
 		else
 		{
-			PointTo = FQuat::FindBetweenNormals(FVector(1,0,0), DirToCamera);
+			PointTo = FQuat::FindBetweenNormals(FVector(1, 0, 0), DirToCamera);
 			// Add in additional rotation about facing axis
-			FQuat AddedRotation = FQuat(DirToCamera, InParticle.Rotation);
+			FQuat AddedRotation = FQuat(DirToCamera, ParticleRotation);
 			kLockedAxisQuat = (AddedRotation * PointTo);
 		}
 	}
-	else if (bUseMeshLockedAxis == true)
+	else if (bUseMeshLockedAxis)
 	{
 		// Add any 'sprite rotation' about the locked axis
-		FQuat AddedRotation = FQuat(Source.LockedAxis, InParticle.Rotation);
+		FQuat AddedRotation = FQuat(Source.LockedAxis, ParticleRotation);
 		kLockedAxisQuat = (AddedRotation * PointTo);
 	}
 	else if (Source.ScreenAlignment == PSA_TypeSpecific)
@@ -1851,16 +2042,17 @@ void FDynamicMeshEmitterData::GetParticleTransform(const FBaseParticle& InPartic
 			kTransMat.SetOrigin(Location);
 			LocalToWorld.SetIdentity();
 		}
-		DirToCamera	= View->ViewMatrices.ViewOrigin - Location;
+
+		DirToCamera = ViewOrigin - Location;
 		DirToCamera.Normalize();
-		if (DirToCamera.SizeSquared() <	0.5f)
+		if (DirToCamera.SizeSquared() < 0.5f)
 		{
 			// Assert possible if DirToCamera is not normalized
-			DirToCamera	= FVector(1,0,0);
+			DirToCamera = FVector(1, 0, 0);
 		}
 
-		LocalSpaceFacingAxis = FVector(1,0,0); // facing axis is taken to be the local x axis.	
-		LocalSpaceUpAxis = FVector(0,0,1); // up axis is taken to be the local z axis
+		LocalSpaceFacingAxis = FVector(1, 0, 0); // facing axis is taken to be the local x axis.	
+		LocalSpaceUpAxis = FVector(0, 0, 1); // up axis is taken to be the local z axis
 
 		if (Source.MeshAlignment == PSMA_MeshFaceCameraWithLockedAxis)
 		{
@@ -1869,12 +2061,12 @@ void FDynamicMeshEmitterData::GetParticleTransform(const FBaseParticle& InPartic
 			// For the locked axis behavior, only rotate to	face the camera	about the
 			// locked direction, and maintain the up vector	pointing towards the locked	direction
 			// Find	the	rotation that points the localupaxis towards the targetupaxis
-			FQuat PointToUp	= FQuat::FindBetweenNormals(LocalSpaceUpAxis, Source.LockedAxis);
+			FQuat PointToUp = FQuat::FindBetweenNormals(LocalSpaceUpAxis, Source.LockedAxis);
 
 			// Add in rotation about the TargetUpAxis to point the facing vector towards the camera
 			FVector	DirToCameraInRotationPlane = DirToCamera - ((DirToCamera | Source.LockedAxis)*Source.LockedAxis);
 			DirToCameraInRotationPlane.Normalize();
-			FQuat PointToCamera	= FQuat::FindBetweenNormals(PointToUp.RotateVector(LocalSpaceFacingAxis), DirToCameraInRotationPlane);
+			FQuat PointToCamera = FQuat::FindBetweenNormals(PointToUp.RotateVector(LocalSpaceFacingAxis), DirToCameraInRotationPlane);
 
 			// Set kRotMat to the composed rotation
 			FQuat MeshRotation = PointToCamera*PointToUp;
@@ -1882,68 +2074,66 @@ void FDynamicMeshEmitterData::GetParticleTransform(const FBaseParticle& InPartic
 		}
 		else if (Source.MeshAlignment == PSMA_MeshFaceCameraWithSpin)
 		{
-				// Implement a tangent-rotation	version	of point-to-camera.	 The facing	direction points to	the	camera,
-				// with	no roll, and has addtional sprite-particle rotation	about the tangential axis
-				// (c.f. the roll rotation is about	the	radial axis)
+			// Implement a tangent-rotation	version	of point-to-camera.	 The facing	direction points to	the	camera,
+			// with	no roll, and has addtional sprite-particle rotation	about the tangential axis
+			// (c.f. the roll rotation is about	the	radial axis)
 
-				// Find	the	rotation that points the facing	axis towards the camera
-				FRotator PointToRotation = FRotator(FQuat::FindBetweenNormals(LocalSpaceFacingAxis, DirToCamera));
+			// Find	the	rotation that points the facing	axis towards the camera
+			FRotator PointToRotation = FRotator(FQuat::FindBetweenNormals(LocalSpaceFacingAxis, DirToCamera));
 
-				// When	constructing the rotation, we need to eliminate	roll around	the	dirtocamera	axis,
-				// otherwise the particle appears to rotate	around the dircamera axis when it or the camera	moves
-				PointToRotation.Roll = 0;
+			// When	constructing the rotation, we need to eliminate	roll around	the	dirtocamera	axis,
+			// otherwise the particle appears to rotate	around the dircamera axis when it or the camera	moves
+			PointToRotation.Roll = 0;
 
-				// Add in the tangential rotation we do	want.
-				FVector	vPositivePitch = FVector(0,0,1); //	this is	set	by the rotator's yaw/pitch/roll	reference frame
-				FVector	vTangentAxis = vPositivePitch^DirToCamera;
-				vTangentAxis.Normalize();
-				if (vTangentAxis.SizeSquared() < 0.5f)
-				{
-					vTangentAxis = FVector(1,0,0); // assert is	possible if	FQuat axis/angle constructor is	passed zero-vector
-				}
+			// Add in the tangential rotation we do	want.
+			FVector	vPositivePitch = FVector(0, 0, 1); //	this is	set	by the rotator's yaw/pitch/roll	reference frame
+			FVector	vTangentAxis = vPositivePitch^DirToCamera;
+			vTangentAxis.Normalize();
+			if (vTangentAxis.SizeSquared() < 0.5f)
+			{
+				vTangentAxis = FVector(1, 0, 0); // assert is	possible if	FQuat axis/angle constructor is	passed zero-vector
+			}
 
-				FQuat AddedTangentialRotation =	FQuat(vTangentAxis,	InParticle.Rotation);
+			FQuat AddedTangentialRotation = FQuat(vTangentAxis, ParticleRotation);
 
-				// Set kRotMat to the composed rotation
-				FQuat MeshRotation = AddedTangentialRotation*PointToRotation.Quaternion();
-				kRotMat = FQuatRotationMatrix(MeshRotation);
+			// Set kRotMat to the composed rotation
+			FQuat MeshRotation = AddedTangentialRotation*PointToRotation.Quaternion();
+			kRotMat = FQuatRotationMatrix(MeshRotation);
 		}
-		else 
+		else
 		{
-				// Implement a roll-rotation version of	point-to-camera.  The facing direction points to the camera,
-				// with	no roll, and then rotates about	the	direction_to_camera	by the spriteparticle rotation.
+			// Implement a roll-rotation version of	point-to-camera.  The facing direction points to the camera,
+			// with	no roll, and then rotates about	the	direction_to_camera	by the spriteparticle rotation.
 
-				// Find	the	rotation that points the facing	axis towards the camera
-				FRotator PointToRotation = FRotator(FQuat::FindBetweenNormals(LocalSpaceFacingAxis, DirToCamera));
+			// Find	the	rotation that points the facing	axis towards the camera
+			FRotator PointToRotation = FRotator(FQuat::FindBetweenNormals(LocalSpaceFacingAxis, DirToCamera));
 
-				// When	constructing the rotation, we need to eliminate	roll around	the	dirtocamera	axis,
-				// otherwise the particle appears to rotate	around the dircamera axis when it or the camera	moves
-				PointToRotation.Roll = 0;
+			// When	constructing the rotation, we need to eliminate	roll around	the	dirtocamera	axis,
+			// otherwise the particle appears to rotate	around the dircamera axis when it or the camera	moves
+			PointToRotation.Roll = 0;
 
-				// Add in the roll we do want.
-				FQuat AddedRollRotation	= FQuat(DirToCamera, InParticle.Rotation);
+			// Add in the roll we do want.
+			FQuat AddedRollRotation = FQuat(DirToCamera, ParticleRotation);
 
-				// Set kRotMat to the composed	rotation
-				FQuat MeshRotation = AddedRollRotation*PointToRotation.Quaternion();
-				kRotMat = FQuatRotationMatrix(MeshRotation);
+			// Set kRotMat to the composed	rotation
+			FQuat MeshRotation = AddedRollRotation*PointToRotation.Quaternion();
+			kRotMat = FQuatRotationMatrix(MeshRotation);
 		}
 	}
 	else
 	{
-		float fRot = InParticle.Rotation * 180.0f / PI;
+		float fRot = ParticleRotation * 180.0f / PI;
 		FVector kRotVec = FVector(fRot, fRot, fRot);
 		FRotator kRotator = FRotator::MakeFromEuler(kRotVec);
 
-		const FMeshRotationPayloadData* PayloadData = (const FMeshRotationPayloadData*)((const uint8*)&InParticle + Source.MeshRotationOffset);
-		kRotator += FRotator::MakeFromEuler(PayloadData->Rotation);
+		kRotator += FRotator::MakeFromEuler(ParticlePayloadRotation);
 
 		kRotMat = FRotationMatrix(kRotator);
 	}
 
 	if (bApplyPreRotation == true)
 	{
-		const FMeshRotationPayloadData* PayloadData = (const FMeshRotationPayloadData*)((const uint8*)&InParticle + Source.MeshRotationOffset);
-		FRotator MeshOrient = FRotator::MakeFromEuler(PayloadData->InitialOrientation);
+		FRotator MeshOrient = FRotator::MakeFromEuler(ParticlePayloadInitialOrientation);
 		FRotationMatrix OrientMat(MeshOrient);
 
 		if ((bUseCameraFacing == true) || (bUseMeshLockedAxis == true))
@@ -1964,83 +2154,25 @@ void FDynamicMeshEmitterData::GetParticleTransform(const FBaseParticle& InPartic
 		OutTransformMat = kScaleMat * kRotMat * kTransMat;
 	}
 
-	FVector OrbitOffset(0.0f, 0.0f, 0.0f);
-	if (Source.OrbitModuleOffset != 0)
+	FVector OrbitOffset = ParticlePayloadOrbitOffset;
+	if (Source.bUseLocalSpace == false)
 	{
-		int32 CurrentOffset = Source.OrbitModuleOffset;
-		PARTICLE_ELEMENT(FOrbitChainModuleInstancePayload, OrbitPayload);
-		OrbitOffset = OrbitPayload.Offset;
-		if (Source.bUseLocalSpace == false)
-		{
-			OrbitOffset = LocalToWorld.TransformVector(OrbitOffset);
-		}
-
-		FTranslationMatrix OrbitMatrix(OrbitOffset);
-		OutTransformMat *= OrbitMatrix;
+		OrbitOffset = LocalToWorld.TransformVector(OrbitOffset);
 	}
 
+	FTranslationMatrix OrbitMatrix(OrbitOffset);
+	OutTransformMat *= OrbitMatrix;
+
 	if (Source.bUseLocalSpace)
 	{
 		OutTransformMat *= LocalToWorld;
 	}
 }
 
-
-void FDynamicMeshEmitterData::GetInstanceData(void* InstanceData, void* DynamicParameterData, const FParticleSystemSceneProxy* Proxy, const FSceneView* View) const
+void FDynamicMeshEmitterData::GetInstanceData(void* InstanceData, void* DynamicParameterData, void* PrevTransformBuffer, const FParticleSystemSceneProxy* Proxy, const FSceneView* View) const
 {
 	SCOPE_CYCLE_COUNTER(STAT_ParticlePackingTime);
 
-	FQuat PointToLockedAxis;
-	if (bUseMeshLockedAxis == true)
-	{
-		// facing axis is taken to be the local x axis.	
-		PointToLockedAxis = FQuat::FindBetweenNormals(FVector(1,0,0), Source.LockedAxis);
-	}
-
-	FVector CameraFacingOpVector = FVector::ZeroVector;
-	if (CameraFacingOption != XAxisFacing_NoUp)
-	{
-		switch (CameraFacingOption)
-		{
-		case XAxisFacing_ZUp:
-			CameraFacingOpVector = FVector( 0.0f, 0.0f, 1.0f);
-			break;
-		case XAxisFacing_NegativeZUp:
-			CameraFacingOpVector = FVector( 0.0f, 0.0f,-1.0f);
-			break;
-		case XAxisFacing_YUp:
-			CameraFacingOpVector = FVector( 0.0f, 1.0f, 0.0f);
-			break;
-		case XAxisFacing_NegativeYUp:
-			CameraFacingOpVector = FVector( 0.0f,-1.0f, 0.0f);
-			break;
-		case LockedAxis_YAxisFacing:
-		case VelocityAligned_YAxisFacing:
-			CameraFacingOpVector = FVector(0.0f, 1.0f, 0.0f);
-			break;
-		case LockedAxis_NegativeYAxisFacing:
-		case VelocityAligned_NegativeYAxisFacing:
-			CameraFacingOpVector = FVector(0.0f,-1.0f, 0.0f);
-			break;
-		case LockedAxis_ZAxisFacing:
-		case VelocityAligned_ZAxisFacing:
-			CameraFacingOpVector = FVector(0.0f, 0.0f, 1.0f);
-			break;
-		case LockedAxis_NegativeZAxisFacing:
-		case VelocityAligned_NegativeZAxisFacing:
-			CameraFacingOpVector = FVector(0.0f, 0.0f,-1.0f);
-			break;
-		}
-	}
-
-	// Put the camera origin in the appropriate coordinate space.
-	FVector CameraPosition = View->ViewMatrices.ViewOrigin;
-	if (Source.bUseLocalSpace)
-	{
-		const FMatrix InvLocalToWorld = Proxy->GetLocalToWorld().Inverse();
-		CameraPosition = InvLocalToWorld.TransformPosition(CameraPosition);
-	}
-
 	int32 SubImagesX = Source.SubImages_Horizontal;
 	int32 SubImagesY = Source.SubImages_Vertical;
 	float SubImageSizeX = 1.0f / SubImagesX;
@@ -2054,8 +2186,11 @@ void FDynamicMeshEmitterData::GetInstanceData(void* InstanceData, void* DynamicP
 
 	int32 InstanceVertexStride = sizeof(FMeshParticleInstanceVertex);
 	int32 DynamicParameterVertexStride = bUsesDynamicParameter ? sizeof(FMeshParticleInstanceVertexDynamicParameter) : 0;
+	int32 PrevTransformVertexStride = sizeof(FVector4) * 3;
+
 	uint8* TempVert = (uint8*)InstanceData;
 	uint8* TempDynamicParameterVert = (uint8*)DynamicParameterData;
+	uint8* TempPrevTranformVert = (uint8*)PrevTransformBuffer;
 
 	for (int32 i = ParticleCount - 1; i >= 0; i--)
 	{
@@ -2067,10 +2202,10 @@ void FDynamicMeshEmitterData::GetInstanceData(void* InstanceData, void* DynamicP
 		// Populate instance buffer;
 		// The particle color.
 		CurrentInstanceVertex->Color = Particle.Color;
-
+		
 		// Instance to world transformation. Translation (Instance world position) is packed into W
 		FMatrix TransMat(FMatrix::Identity);
-		GetParticleTransform(Particle, CameraPosition, CameraFacingOpVector, PointToLockedAxis, Proxy, View, TransMat);
+		GetParticleTransform(Particle, Proxy, View, TransMat);
 		
 		// Transpose on CPU to allow for simpler shader code to perform the transform. 
 		const FMatrix Transpose = TransMat.GetTransposed();
@@ -2078,6 +2213,23 @@ void FDynamicMeshEmitterData::GetInstanceData(void* InstanceData, void* DynamicP
 		CurrentInstanceVertex->Transform[1] = FVector4(Transpose.M[1][0], Transpose.M[1][1], Transpose.M[1][2], Transpose.M[1][3]);
 		CurrentInstanceVertex->Transform[2] = FVector4(Transpose.M[2][0], Transpose.M[2][1], Transpose.M[2][2], Transpose.M[2][3]);
 
+		if (Source.MeshMotionBlurOffset)
+		{
+			FVector4* PrevTransformVertex = (FVector4*)TempPrevTranformVert;
+
+			// Instance to world transformation. Translation (Instance world position) is packed into W
+			FMatrix PrevTransMat(FMatrix::Identity);
+			GetParticlePrevTransform(Particle, Proxy, View, PrevTransMat);
+
+			// Transpose on CPU to allow for simpler shader code to perform the transform. 
+			const FMatrix PrevTranspose = PrevTransMat.GetTransposed();
+			PrevTransformVertex[0] = FVector4(PrevTranspose.M[0][0], PrevTranspose.M[0][1], PrevTranspose.M[0][2], PrevTranspose.M[0][3]);
+			PrevTransformVertex[1] = FVector4(PrevTranspose.M[1][0], PrevTranspose.M[1][1], PrevTranspose.M[1][2], PrevTranspose.M[1][3]);
+			PrevTransformVertex[2] = FVector4(PrevTranspose.M[2][0], PrevTranspose.M[2][1], PrevTranspose.M[2][2], PrevTranspose.M[2][3]);
+
+			TempPrevTranformVert += PrevTransformVertexStride;
+		}
+
 		// Particle velocity. Calculate on CPU to avoid computing in vertex shader.
 		// Note: It would be preferred if we could check whether the material makes use of the 'Particle Direction' node to avoid this work.
 		FVector DeltaPosition = Particle.Location - Particle.OldLocation;
@@ -2100,7 +2252,7 @@ void FDynamicMeshEmitterData::GetInstanceData(void* InstanceData, void* DynamicP
 			DeltaPosition.ToDirectionAndLength(Direction, Speed);
 
 			// Pack direction and speed.
-			CurrentInstanceVertex->Velocity = FVector4(Direction,Speed);
+			CurrentInstanceVertex->Velocity = FVector4(Direction, Speed);
 		}
 		else
 		{
@@ -2148,7 +2300,7 @@ void FDynamicMeshEmitterData::GetInstanceData(void* InstanceData, void* DynamicP
 		// The particle's relative time
 		CurrentInstanceVertex->RelativeTime = Particle.RelativeTime;
 
-		TempVert += InstanceVertexStride; 
+		TempVert += InstanceVertexStride;
 	}
 }
 
@@ -6459,6 +6611,8 @@ FParticleSystemSceneProxy::FParticleSystemSceneProxy(const UParticleSystemCompon
 
 	LODMethod = Component->LODMethod;
 
+	// Particle systems intrinsically always have motion, but is this motion relevant to systems external to particle systems?
+	bAlwaysHasVelocity = Component->Template->DoesAnyEmitterHaveMotionBlur(Component->GetCurrentLODIndex());
 }
 
 FParticleSystemSceneProxy::~FParticleSystemSceneProxy()
