Copying //UE4/Dev-Rendering to Dev-Main (//UE4/Dev-Main)

diff --git a/Engine/Source/Runtime/Engine/Private/Particles/ParticleTrail2EmitterInstance.cpp b/Engine/Source/Runtime/Engine/Private/Particles/ParticleTrail2EmitterInstance.cpp
--- a/Engine/Source/Runtime/Engine/Private/Particles/ParticleTrail2EmitterInstance.cpp
+++ b/Engine/Source/Runtime/Engine/Private/Particles/ParticleTrail2EmitterInstance.cpp
@@ -510,16 +510,16 @@ void FParticleTrailsEmitterInstance_Base::KillParticles(int32 InTrailIdx, int32
 	if (ActiveParticles)
 	{
 		int32 KilledCount = 0;
-		// Loop over the active particles...
-		for (int32 ParticleIdx = ActiveParticles - 1; ParticleIdx >= 0 && (KilledCount < InKillCount); ParticleIdx--)
+
 		{
 			// Find the end particle
-			int32 CurrentIndex = ParticleIndices[ParticleIdx];
-			DECLARE_PARTICLE_PTR(Particle, ParticleData + ParticleStride * CurrentIndex);
-			FTrailsBaseTypeDataPayload* TrailData = ((FTrailsBaseTypeDataPayload*)((uint8*)Particle + TypeDataOffset));
-			if (TRAIL_EMITTER_IS_END(TrailData->Flags))
+			FTrailsBaseTypeDataPayload* TrailData = nullptr;
+			FBaseParticle *Particle = nullptr;
+			int32 EndIndex;
+			GetTrailEnd<FTrailsBaseTypeDataPayload>(InTrailIdx, EndIndex, TrailData, Particle);
+
 			{
-				if (TrailData->TrailIndex == InTrailIdx)
+				if (TrailData && TrailData->TrailIndex == InTrailIdx)
 				{
 					while ((TrailData != NULL) && (KilledCount < InKillCount))
 					{
@@ -865,21 +865,8 @@ void FParticleRibbonEmitterInstance::Tick_RecalculateTangents(float DeltaTime, U
 			// Find the Start particle of the current trail...
 			FBaseParticle* StartParticle = NULL;
 			FRibbonTypeDataPayload* StartTrailData = NULL;
-			for (int32 FindTrailIdx = 0; FindTrailIdx < ActiveParticles; FindTrailIdx++)
-			{
-				int32 CheckStartIndex = ParticleIndices[FindTrailIdx];
-				DECLARE_PARTICLE_PTR(CheckParticle, ParticleData + ParticleStride * CheckStartIndex);
-				FRibbonTypeDataPayload* CheckTrailData = ((FRibbonTypeDataPayload*)((uint8*)CheckParticle + TypeDataOffset));
-				if (TRAIL_EMITTER_IS_START(CheckTrailData->Flags))
-				{
-					if (CheckTrailData->TrailIndex == TrailIdx)
-					{
-						StartParticle = CheckParticle;
-						StartTrailData = CheckTrailData;
-						break;
-					}
-				}
-			}
+			int32 StartIndex = -1;
+			GetTrailStart<FRibbonTypeDataPayload>(TrailIdx, StartIndex, StartTrailData, StartParticle);
 
 			// Recalculate tangents at each particle to properly handle moving particles...
 			if ((StartParticle != NULL) && (TRAIL_EMITTER_IS_ONLY(StartTrailData->Flags) == 0))
@@ -1282,22 +1269,7 @@ float FParticleRibbonEmitterInstance::Spawn(float DeltaTime)
 	FBaseParticle* StartParticle = NULL;
 	int32 StartIndex = -1;
 	FRibbonTypeDataPayload* StartTrailData = NULL;
-	for (int32 FindTrailIdx = 0; FindTrailIdx < ActiveParticles; FindTrailIdx++)
-	{
-		int32 CheckStartIndex = ParticleIndices[FindTrailIdx];
-		DECLARE_PARTICLE_PTR(CheckParticle, ParticleData + ParticleStride * CheckStartIndex);
-		FRibbonTypeDataPayload* CheckTrailData = ((FRibbonTypeDataPayload*)((uint8*)CheckParticle + TypeDataOffset));
-		if (CheckTrailData->TrailIndex == TrailIdx)
-		{
-			if (TRAIL_EMITTER_IS_START(CheckTrailData->Flags))
-			{
-				StartParticle = CheckParticle;
-				StartIndex = CheckStartIndex;
-				StartTrailData = CheckTrailData;
-				break;
-			}
-		}
-	}
+	GetTrailStart<FRibbonTypeDataPayload>(TrailIdx, StartIndex, StartTrailData, StartParticle);
 
 	bNoLivingParticles = (StartParticle == NULL);
 	bool bTilingTrail = !FMath::IsNearlyZero(TrailTypeData->TilingDistance);
@@ -1619,6 +1591,9 @@ bool FParticleRibbonEmitterInstance::Spawn_Source(float DeltaTime)
 		FBaseParticle* StartParticle = NULL;
 		int32 StartIndex = -1;
 		FRibbonTypeDataPayload* StartTrailData = NULL;
+
+		// temporarily not using index cache here, as it causes problems later
+
 		for (int32 FindTrailIdx = 0; FindTrailIdx < ActiveParticles; FindTrailIdx++)
 		{
 			int32 CheckStartIndex = ParticleIndices[FindTrailIdx];
@@ -1826,6 +1801,8 @@ bool FParticleRibbonEmitterInstance::Spawn_Source(float DeltaTime)
 
 				// Standard spawn setup
 				PreSpawn(Particle, CurrPosition, FVector::ZeroVector);
+				SetDeadIndex(TrailData->TrailIndex, ParticleIndex);
+
 				for (int32 SpawnModuleIdx = 0; SpawnModuleIdx < LODLevel->SpawnModules.Num(); SpawnModuleIdx++)
 				{
 					UParticleModule* SpawnModule = LODLevel->SpawnModules[SpawnModuleIdx];
@@ -2982,24 +2959,11 @@ void FParticleAnimTrailEmitterInstance::Tick_RecalculateTangents(float DeltaTime
 	{
 		//@todo. Multiple trails, single emitter
 		int32 TrailIdx = 0;
+		int32 StartIndex = 0;
 		// Find the Start particle of the current trail...
 		FBaseParticle* StartParticle = NULL;
 		FAnimTrailTypeDataPayload* StartTrailData = NULL;
-		for (int32 FindTrailIdx = 0; FindTrailIdx < ActiveParticles; FindTrailIdx++)
-		{
-			int32 CheckStartIndex = ParticleIndices[FindTrailIdx];
-			DECLARE_PARTICLE_PTR(CheckParticle, ParticleData + ParticleStride * CheckStartIndex);
-			FAnimTrailTypeDataPayload* CheckTrailData = ((FAnimTrailTypeDataPayload*)((uint8*)CheckParticle + TypeDataOffset));
-			if (TRAIL_EMITTER_IS_START(CheckTrailData->Flags))
-			{
-				if (CheckTrailData->TrailIndex == TrailIdx)
-				{
-					StartParticle = CheckParticle;
-					StartTrailData = CheckTrailData;
-					break;
-				}
-			}
-		}
+		GetTrailStart<FAnimTrailTypeDataPayload>(TrailIdx, StartIndex, StartTrailData, StartParticle);
 
 		// Recalculate tangents at each particle to properly handle moving particles...
 		if ((StartParticle != NULL) && (TRAIL_EMITTER_IS_ONLY(StartTrailData->Flags) == 0))
@@ -3418,19 +3382,9 @@ float FParticleAnimTrailEmitterInstance::Spawn(float DeltaTime)
 	//@todo. Support multiple trails per emitter
 	// Find the start particle of the current trail...
 	int32 StartIndex = -1;
-	for (int32 FindTrailIdx = 0; FindTrailIdx < ActiveParticles; FindTrailIdx++)
+	if (TrailIdx != INDEX_NONE)
 	{
-		int32 CheckStartIndex = ParticleIndices[FindTrailIdx];
-		DECLARE_PARTICLE_PTR(CheckParticle, ParticleData + ParticleStride * CheckStartIndex);
-		FAnimTrailTypeDataPayload* CheckTrailData = ((FAnimTrailTypeDataPayload*)((uint8*)CheckParticle + TypeDataOffset));
-		if (CheckTrailData->TrailIndex == TrailIdx)
-		{
-			if (TRAIL_EMITTER_IS_START(CheckTrailData->Flags))
-			{
-				StartIndex = CheckStartIndex;
-				break;
-			}
-		}
+		StartIndex = CurrentStartIndices[TrailIdx];
 	}
 
 	bool bTilingTrail = !FMath::IsNearlyZero(TrailTypeData->TilingDistance);
@@ -3493,13 +3447,8 @@ float FParticleAnimTrailEmitterInstance::Spawn(float DeltaTime)
 			int32 CheckStartIndex = ParticleIndices[FindTrailIdx];
 			DECLARE_PARTICLE_PTR(CheckParticle, ParticleData + ParticleStride * CheckStartIndex);
 			FAnimTrailTypeDataPayload* CheckTrailData = ((FAnimTrailTypeDataPayload*)((uint8*)CheckParticle + TypeDataOffset));
-			if (CheckTrailData->TrailIndex == TrailIdx)
-			{
-				if (TRAIL_EMITTER_IS_START(CheckTrailData->Flags))
-				{
-					CheckTrailData->Flags = TRAIL_EMITTER_SET_DEADTRAIL(CheckTrailData->Flags);
-				}
-			}
+			CheckTrailData->Flags = TRAIL_EMITTER_SET_DEADTRAIL(CheckTrailData->Flags);
+			SetDeadIndex(CheckTrailData->TrailIndex, CheckStartIndex);
 		}
 		bTagTrailAsDead = false;
 	}
diff --git a/Engine/Source/Runtime/Engine/Private/PrimitiveSceneProxy.cpp b/Engine/Source/Runtime/Engine/Private/PrimitiveSceneProxy.cpp
--- a/Engine/Source/Runtime/Engine/Private/PrimitiveSceneProxy.cpp
+++ b/Engine/Source/Runtime/Engine/Private/PrimitiveSceneProxy.cpp
@@ -35,7 +35,7 @@ FPrimitiveSceneProxy::FPrimitiveSceneProxy(const UPrimitiveComponent* InComponen
 ,	bTreatAsBackgroundForOcclusion(InComponent->bTreatAsBackgroundForOcclusion)
 ,	bDisableStaticPath(false)
 ,	bNeedsUnbuiltPreviewLighting(InComponent->HasStaticLighting() && !InComponent->bHasCachedStaticLighting)
-,	bHasValidSettingsForStaticLighting(InComponent->HasValidSettingsForStaticLighting())
+,	bHasValidSettingsForStaticLighting(InComponent->HasValidSettingsForStaticLighting(false))
 ,	bWillEverBeLit(true)
 	// Disable dynamic shadow casting if the primitive only casts indirect shadows, since dynamic shadows are always shadowing direct lighting
 ,	bCastDynamicShadow(InComponent->bCastDynamicShadow && InComponent->CastShadow && !InComponent->GetShadowIndirectOnly())
diff --git a/Engine/Source/Runtime/Engine/Private/SceneManagement.cpp b/Engine/Source/Runtime/Engine/Private/SceneManagement.cpp
--- a/Engine/Source/Runtime/Engine/Private/SceneManagement.cpp
+++ b/Engine/Source/Runtime/Engine/Private/SceneManagement.cpp
@@ -450,11 +450,47 @@ FLODMask ComputeLODForMeshes( const TIndirectArray<class FStaticMesh>& StaticMes
 }
 
 FFrameUniformShaderParameters::FFrameUniformShaderParameters()
-	: DirectionalLightShadowTexture(GWhiteTexture->TextureRHI)
-	, DirectionalLightShadowSampler(TStaticSamplerState<SF_Point, AM_Clamp, AM_Clamp, AM_Clamp>::GetRHI())
 {
+	FMemory::Memzero(*this);
+	FTextureRHIParamRef BlackVolume = (GBlackVolumeTexture &&  GBlackVolumeTexture->TextureRHI) ? GBlackVolumeTexture->TextureRHI : GBlackTexture->TextureRHI; // for es2, this might need to be 2d
+	check(GBlackVolumeTexture);
+	DirectionalLightShadowTexture = GWhiteTexture->TextureRHI;
+	DirectionalLightShadowSampler = TStaticSamplerState<SF_Point, AM_Clamp, AM_Clamp, AM_Clamp>::GetRHI();
+
+	AtmosphereTransmittanceTexture_UB = GWhiteTexture->TextureRHI;
+	AtmosphereTransmittanceTextureSampler_UB = TStaticSamplerState<SF_Bilinear>::GetRHI();
+	AtmosphereIrradianceTexture_UB = GWhiteTexture->TextureRHI;
+	AtmosphereIrradianceTextureSampler_UB = TStaticSamplerState<SF_Bilinear>::GetRHI();
+	AtmosphereInscatterTexture_UB = BlackVolume;
+	AtmosphereInscatterTextureSampler_UB = TStaticSamplerState<SF_Bilinear>::GetRHI();
+
+	PerlinNoiseGradientTexture = GWhiteTexture->TextureRHI;
+	PerlinNoiseGradientTextureSampler = TStaticSamplerState<SF_Point, AM_Wrap, AM_Wrap, AM_Wrap>::GetRHI();
+
+	PerlinNoise3DTexture = BlackVolume;
+	PerlinNoise3DTextureSampler = TStaticSamplerState<SF_Bilinear, AM_Wrap, AM_Wrap, AM_Wrap>::GetRHI();
+
+	GlobalDistanceFieldTexture0_UB = BlackVolume;
+	GlobalDistanceFieldSampler0_UB = TStaticSamplerState<SF_Bilinear, AM_Wrap, AM_Wrap, AM_Wrap>::GetRHI();
+	GlobalDistanceFieldTexture1_UB = BlackVolume;
+	GlobalDistanceFieldSampler1_UB = TStaticSamplerState<SF_Bilinear, AM_Wrap, AM_Wrap, AM_Wrap>::GetRHI();
+	GlobalDistanceFieldTexture2_UB = BlackVolume;
+	GlobalDistanceFieldSampler2_UB = TStaticSamplerState<SF_Bilinear, AM_Wrap, AM_Wrap, AM_Wrap>::GetRHI();
+	GlobalDistanceFieldTexture3_UB = BlackVolume;
+	GlobalDistanceFieldSampler3_UB = TStaticSamplerState<SF_Bilinear, AM_Wrap, AM_Wrap, AM_Wrap>::GetRHI();
 }
 
+FViewUniformShaderParameters::FViewUniformShaderParameters()
+{
+	FMemory::Memzero(*this);
+}
+
+FInstancedViewUniformShaderParameters::FInstancedViewUniformShaderParameters()
+{
+	FMemory::Memzero(*this);
+}
+
+
 void FSharedSamplerState::InitRHI()
 {
 	const float MipMapBias = UTexture2D::GetGlobalMipMapLODBias();
diff --git a/Engine/Source/Runtime/Engine/Private/SceneView.cpp b/Engine/Source/Runtime/Engine/Private/SceneView.cpp
--- a/Engine/Source/Runtime/Engine/Private/SceneView.cpp
+++ b/Engine/Source/Runtime/Engine/Private/SceneView.cpp
@@ -1796,6 +1796,8 @@ FSceneViewFamily::FSceneViewFamily( const ConstructionValues& CVS )
 		CurrentWorldTime = 0;
 		CurrentRealTime = 0; 
 	}
+
+	DebugViewShaderMode = ChooseDebugViewShaderMode();
 #endif
 
 #if !WITH_EDITOR
@@ -1897,7 +1899,7 @@ ERHIFeatureLevel::Type FSceneViewFamily::GetFeatureLevel() const
 }
 
 #if !(UE_BUILD_SHIPPING || UE_BUILD_TEST)
-EDebugViewShaderMode FSceneViewFamily::GetDebugViewShaderMode() const
+EDebugViewShaderMode FSceneViewFamily::ChooseDebugViewShaderMode() const
 {
 	const bool bDebugViewShaderExists = AllowDebugViewModeShader(GetFeatureLevelShaderPlatform(GetFeatureLevel()));
 
@@ -1905,13 +1907,11 @@ EDebugViewShaderMode FSceneViewFamily::GetDebugViewShaderMode() const
 	{
 		if (EngineShowFlags.ShaderComplexity)
 		{
-			if (EngineShowFlags.QuadComplexity)
+			if (EngineShowFlags.QuadOverdraw)
 			{
-				// This can be used to visualize the QOM_ShaderComplexityBleeding mode.
-				// return EngineShowFlags.QuadOverhead ? DVSM_ShaderComplexityBleedingQuadOverhead : DVSM_QuadComplexity;
 				return DVSM_QuadComplexity;
 			}
-			else if (EngineShowFlags.QuadOverhead)
+			else if (EngineShowFlags.ShaderComplexityWithQuadOverdraw)
 			{
 				return DVSM_ShaderComplexityContainedQuadOverhead;
 			}
@@ -1928,8 +1928,12 @@ EDebugViewShaderMode FSceneViewFamily::GetDebugViewShaderMode() const
 		{
 			return DVSM_TexelFactorAccuracy;
 		}
+		else if (EngineShowFlags.TexCoordScaleAccuracy)
+		{
+			return DVSM_TexCoordScaleAnalysis;
+		}
 	}
-	else if (EngineShowFlags.ShaderComplexity && !EngineShowFlags.QuadComplexity)
+	else if (EngineShowFlags.ShaderComplexity && !EngineShowFlags.QuadOverdraw)
 	{
 		return DVSM_ShaderComplexity;
 	}
diff --git a/Engine/Source/Runtime/Engine/Private/ShowFlags.cpp b/Engine/Source/Runtime/Engine/Private/ShowFlags.cpp
--- a/Engine/Source/Runtime/Engine/Private/ShowFlags.cpp
+++ b/Engine/Source/Runtime/Engine/Private/ShowFlags.cpp
@@ -225,6 +225,7 @@ void FEngineShowFlags::AddNameByIndex(uint32 InIndex, FString& Out)
 	}
 }
 
+RENDERER_API void BuildDebugViewModeShader(EViewModeIndex);
 
 void ApplyViewMode(EViewModeIndex ViewModeIndex, bool bPerspective, FEngineShowFlags& EngineShowFlags)
 {
@@ -255,9 +256,11 @@ void ApplyViewMode(EViewModeIndex ViewModeIndex, bool bPerspective, FEngineShowF
 			bPostProcessing = false;
 			break;
 		case VMI_ShaderComplexity:
-		case VMI_QuadComplexity:
+		case VMI_QuadOverdraw:
+		case VMI_ShaderComplexityWithQuadOverdraw:
 		case VMI_WantedMipsAccuracy:
 		case VMI_TexelFactorAccuracy:
+		case VMI_TexCoordScaleAccuracy:
 			bPostProcessing = false;
 			break;
 		case VMI_StationaryLightOverlap:
@@ -302,10 +305,12 @@ void ApplyViewMode(EViewModeIndex ViewModeIndex, bool bPerspective, FEngineShowF
 	EngineShowFlags.SetReflectionOverride(ViewModeIndex == VMI_ReflectionOverride);
 	EngineShowFlags.SetVisualizeBuffer(ViewModeIndex == VMI_VisualizeBuffer);
 	EngineShowFlags.SetVisualizeLightCulling(ViewModeIndex == VMI_LightComplexity);
-	EngineShowFlags.SetShaderComplexity(ViewModeIndex == VMI_ShaderComplexity || ViewModeIndex == VMI_QuadComplexity);
-	EngineShowFlags.SetQuadComplexity(ViewModeIndex == VMI_QuadComplexity);
+	EngineShowFlags.SetShaderComplexity(ViewModeIndex == VMI_ShaderComplexity || ViewModeIndex == VMI_QuadOverdraw || ViewModeIndex == VMI_ShaderComplexityWithQuadOverdraw);
+	EngineShowFlags.SetQuadOverdraw(ViewModeIndex == VMI_QuadOverdraw);
+	EngineShowFlags.SetShaderComplexityWithQuadOverdraw(ViewModeIndex == VMI_ShaderComplexityWithQuadOverdraw);
 	EngineShowFlags.SetWantedMipsAccuracy(ViewModeIndex == VMI_WantedMipsAccuracy);
 	EngineShowFlags.SetTexelFactorAccuracy(ViewModeIndex == VMI_TexelFactorAccuracy);
+	EngineShowFlags.SetTexCoordScaleAccuracy(ViewModeIndex == VMI_TexCoordScaleAccuracy);
 	EngineShowFlags.SetStationaryLightOverlap(ViewModeIndex == VMI_StationaryLightOverlap);
 	EngineShowFlags.SetLightMapDensity(ViewModeIndex == VMI_LightmapDensity || ViewModeIndex == VMI_LitLightmapDensity);
 	EngineShowFlags.SetPostProcessing(bPostProcessing);
@@ -393,9 +398,11 @@ void EngineShowFlagOverride(EShowFlagInitMode ShowFlagInitMode, EViewModeIndex V
 			ViewModeIndex == VMI_Wireframe ||
 			ViewModeIndex == VMI_Unlit ||
 			ViewModeIndex == VMI_ShaderComplexity ||
-			ViewModeIndex == VMI_QuadComplexity ||
+			ViewModeIndex == VMI_QuadOverdraw ||
+			ViewModeIndex == VMI_ShaderComplexityWithQuadOverdraw ||
 			ViewModeIndex == VMI_WantedMipsAccuracy ||
 			ViewModeIndex == VMI_TexelFactorAccuracy ||
+			ViewModeIndex == VMI_TexCoordScaleAccuracy ||
 			ViewModeIndex == VMI_LightmapDensity ||
 			ViewModeIndex == VMI_VertexDensities ||
 			ViewModeIndex == VMI_LitLightmapDensity)
@@ -413,9 +420,11 @@ void EngineShowFlagOverride(EShowFlagInitMode ShowFlagInitMode, EViewModeIndex V
 			ViewModeIndex == VMI_Unlit ||
 			ViewModeIndex == VMI_StationaryLightOverlap ||
 			ViewModeIndex == VMI_ShaderComplexity ||
-			ViewModeIndex == VMI_QuadComplexity ||
+			ViewModeIndex == VMI_QuadOverdraw ||
+			ViewModeIndex == VMI_ShaderComplexityWithQuadOverdraw ||
 			ViewModeIndex == VMI_WantedMipsAccuracy ||
 			ViewModeIndex == VMI_TexelFactorAccuracy ||
+			ViewModeIndex == VMI_TexCoordScaleAccuracy ||
 			ViewModeIndex == VMI_VertexDensities ||
 			ViewModeIndex == VMI_LightmapDensity)
 		{
@@ -453,7 +462,8 @@ void EngineShowFlagOverride(EShowFlagInitMode ShowFlagInitMode, EViewModeIndex V
 		}
 
 		if (ViewModeIndex == VMI_WantedMipsAccuracy ||
-			ViewModeIndex == VMI_TexelFactorAccuracy)
+			ViewModeIndex == VMI_TexelFactorAccuracy ||
+			ViewModeIndex == VMI_TexCoordScaleAccuracy)
 		{
 			EngineShowFlags.Decals = 0; // Decals require the use of FDebugPSInLean.
 			EngineShowFlags.Particles = 0; // FX are fully streamed.
@@ -548,9 +558,13 @@ EViewModeIndex FindViewMode(const FEngineShowFlags& EngineShowFlags)
 		return VMI_ShaderComplexity;
 	}
 	// Test QuadComplexity before ShaderComplexity because QuadComplexity also use ShaderComplexity
-	else if(EngineShowFlags.QuadComplexity)
+	else if(EngineShowFlags.QuadOverdraw)
+	{
+		return VMI_QuadOverdraw;
+	}
+	else if(EngineShowFlags.ShaderComplexityWithQuadOverdraw)
 	{
-		return VMI_QuadComplexity;
+		return VMI_ShaderComplexityWithQuadOverdraw;
 	}
 	else if(EngineShowFlags.WantedMipsAccuracy)
 	{
@@ -560,6 +574,10 @@ EViewModeIndex FindViewMode(const FEngineShowFlags& EngineShowFlags)
 	{
 		return VMI_TexelFactorAccuracy;
 	}
+	else if(EngineShowFlags.TexCoordScaleAccuracy)
+	{
+		return VMI_TexCoordScaleAccuracy;
+	}
 	else if(EngineShowFlags.ShaderComplexity)
 	{
 		return VMI_ShaderComplexity;
@@ -639,9 +657,11 @@ const TCHAR* GetViewModeName(EViewModeIndex ViewModeIndex)
 		case VMI_LightingOnly:				return TEXT("LightingOnly");
 		case VMI_LightComplexity:			return TEXT("LightComplexity");
 		case VMI_ShaderComplexity:			return TEXT("ShaderComplexity");
-		case VMI_QuadComplexity:			return TEXT("QuadOverdraw");
+		case VMI_QuadOverdraw:				return TEXT("QuadOverdraw");
+		case VMI_ShaderComplexityWithQuadOverdraw: return TEXT("ShaderComplexityWithQuadOverdraw");
 		case VMI_WantedMipsAccuracy:		return TEXT("WantedMipsAccuracy");
 		case VMI_TexelFactorAccuracy:		return TEXT("TexelFactorAccuracy");
+		case VMI_TexCoordScaleAccuracy:		return TEXT("TexCoordScaleAccuracy");
 		case VMI_StationaryLightOverlap:	return TEXT("StationaryLightOverlap");
 		case VMI_LightmapDensity:			return TEXT("LightmapDensity");
 		case VMI_LitLightmapDensity:		return TEXT("LitLightmapDensity");
diff --git a/Engine/Source/Runtime/Engine/Private/SkeletalMesh.cpp b/Engine/Source/Runtime/Engine/Private/SkeletalMesh.cpp
--- a/Engine/Source/Runtime/Engine/Private/SkeletalMesh.cpp
+++ b/Engine/Source/Runtime/Engine/Private/SkeletalMesh.cpp
@@ -1148,7 +1148,7 @@ FArchive& operator<<(FArchive& Ar,FSkelMeshSection& S)
 		Ar << S.NumTriangles;
 	}
 		
-		Ar << S.TriangleSorting;
+	Ar << S.TriangleSorting;
 
 	// for clothing info
 	if( Ar.UE4Ver() >= VER_UE4_APEX_CLOTH )
@@ -1169,14 +1169,6 @@ FArchive& operator<<(FArchive& Ar,FSkelMeshSection& S)
 	FStaticLODModel
 -----------------------------------------------------------------------------*/
 
-/**
-* Special serialize function passing the owning UObject along as required by FUnytpedBulkData
-* serialization.
-*
-* @param	Ar		Archive to serialize with
-* @param	Owner	UObject this structure is serialized within
-* @param	Idx		Index of current array entry being serialized
-*/
 void FStaticLODModel::Serialize( FArchive& Ar, UObject* Owner, int32 Idx )
 {
 	DECLARE_SCOPE_CYCLE_COUNTER( TEXT("FStaticLODModel::Serialize"), STAT_StaticLODModel_Serialize, STATGROUP_LoadTime );
@@ -1196,7 +1188,13 @@ void FStaticLODModel::Serialize( FArchive& Ar, UObject* Owner, int32 Idx )
 	MultiSizeIndexContainer.Serialize(Ar, bKeepBuffersInCPUMemory);
 	Ar << ActiveBoneIndices;
 	Ar << Chunks;
-	Ar << Size;
+	
+	// no longer in use
+	{
+		uint32 LegacySize = 0;
+		Ar << LegacySize;
+	}
+
 	if (!StripFlags.IsDataStrippedForServer())
 	{
 		Ar << NumVertices;
@@ -1277,11 +1275,6 @@ void FStaticLODModel::Serialize( FArchive& Ar, UObject* Owner, int32 Idx )
 	}
 }
 
-/**
-* Initialize the LOD's render resources.
-*
-* @param Parent Parent mesh
-*/
 void FStaticLODModel::InitResources(bool bNeedsVertexColors)
 {
 	INC_DWORD_STAT_BY( STAT_SkeletalMeshIndexMemory, MultiSizeIndexContainer.IsIndexBufferValid() ? (MultiSizeIndexContainer.GetIndexBuffer()->Num() * MultiSizeIndexContainer.GetDataTypeSize()) : 0 );
@@ -1312,9 +1305,6 @@ void FStaticLODModel::InitResources(bool bNeedsVertexColors)
 	}
 }
 
-/**
-* Releases the LOD's render resources.
-*/
 void FStaticLODModel::ReleaseResources()
 {
 	DEC_DWORD_STAT_BY( STAT_SkeletalMeshIndexMemory, MultiSizeIndexContainer.IsIndexBufferValid() ? (MultiSizeIndexContainer.GetIndexBuffer()->Num() * MultiSizeIndexContainer.GetDataTypeSize()) : 0 );
@@ -1329,12 +1319,8 @@ void FStaticLODModel::ReleaseResources()
 	BeginReleaseResource(&VertexBufferGPUSkin);
 	BeginReleaseResource(&ColorVertexBuffer);
 	BeginReleaseResource(&APEXClothVertexBuffer);
-
 }
 
-/**
-* Utility function for returning total number of faces in this LOD. 
-*/
 int32 FStaticLODModel::GetTotalFaces() const
 {
 	int32 TotalFaces = 0;
@@ -1346,9 +1332,6 @@ int32 FStaticLODModel::GetTotalFaces() const
 	return TotalFaces;
 }
 
-/** 
- *	Utility for finding the chunk that a particular vertex is in.
- */
 void FStaticLODModel::GetChunkAndSkinType(int32 InVertIndex, int32& OutChunkIndex, int32& OutVertIndex, bool& bOutSoftVert, bool& bOutHasExtraBoneInfluences) const
 {
 	OutChunkIndex = 0;
@@ -1387,15 +1370,8 @@ void FStaticLODModel::GetChunkAndSkinType(int32 InVertIndex, int32& OutChunkInde
 
 	// InVertIndex should always be in some chunk!
 	//check(false);
-	return;
 }
 
-
-/**
-* Fill array with vertex position and tangent data from skel mesh chunks.
-*
-* @param Vertices Array to fill.
-*/
 void FStaticLODModel::GetVertices(TArray<FSoftSkinVertex>& Vertices) const
 {
 	Vertices.Empty(NumVertices);
@@ -1437,11 +1413,6 @@ void FStaticLODModel::GetVertices(TArray<FSoftSkinVertex>& Vertices) const
 	}
 }
 
-/**
-* Fill array with APEX cloth mapping data.
-*
-* @param MappingData Array to fill.
-*/
 void FStaticLODModel::GetApexClothMappingData(TArray<FApexClothPhysToRenderVertData>& MappingData) const
 {
 	for( int32 ChunkIndex = 0; ChunkIndex < Chunks.Num(); ChunkIndex++ )
@@ -1469,11 +1440,6 @@ void FStaticLODModel::GetApexClothMappingData(TArray<FApexClothPhysToRenderVertD
 	}
 }
 
-/**
-* Initialize position and tangent vertex buffers from skel mesh chunks
-*
-* @param Mesh Parent mesh
-*/
 void FStaticLODModel::BuildVertexBuffers(uint32 BuildFlags)
 {
 	bool bUseFullPrecisionUVs = (BuildFlags & EVertexFlags::UseFullPrecisionUVs) != 0;
@@ -1508,11 +1474,6 @@ void FStaticLODModel::BuildVertexBuffers(uint32 BuildFlags)
 	}
 }
 
-/**
-* Sort the triangles in the LODmodel
-*
-* @param ETriangleSortOption NewTriangleSorting new sorting method
-*/
 void FStaticLODModel::SortTriangles( FVector SortCenter, bool bUseSortCenter, int32 SectionIndex, ETriangleSortOption NewTriangleSorting )
 {
 #if WITH_EDITOR
diff --git a/Engine/Source/Runtime/Engine/Private/SkeletalRenderGPUSkin.cpp b/Engine/Source/Runtime/Engine/Private/SkeletalRenderGPUSkin.cpp
--- a/Engine/Source/Runtime/Engine/Private/SkeletalRenderGPUSkin.cpp
+++ b/Engine/Source/Runtime/Engine/Private/SkeletalRenderGPUSkin.cpp
@@ -63,7 +63,9 @@ void FMorphVertexBuffer::InitDynamicRHI()
 	uint32 Size = LodModel.NumVertices * sizeof(FMorphGPUSkinVertex);
 	FRHIResourceCreateInfo CreateInfo;
 	void* BufferData = nullptr;
-	VertexBufferRHI = RHICreateAndLockVertexBuffer(Size, BUF_Dynamic, CreateInfo, BufferData);
+
+	// BUF_ShaderResource is needed for Morph support of the SkinCache
+	VertexBufferRHI = RHICreateAndLockVertexBuffer(Size, BUF_Dynamic | BUF_ShaderResource, CreateInfo, BufferData);
 
 	// Lock the buffer.
 	FMorphGPUSkinVertex* Buffer = (FMorphGPUSkinVertex*)BufferData;
@@ -283,7 +285,9 @@ void FSkeletalMeshObjectGPUSkin::ProcessUpdatedDynamicData(FRHICommandListImmedi
 	const FSkelMeshObjectLODInfo& MeshLODInfo = LODInfo[DynamicData->LODIndex];
 
 	// if hasn't been updated, force update again
-	bMorphNeedsUpdate = LOD.MorphVertexBuffer.bHasBeenUpdated? bMorphNeedsUpdate:true;
+	bMorphNeedsUpdate = LOD.MorphVertexBuffer.bHasBeenUpdated ? bMorphNeedsUpdate : true;
+	
+	bool bMorph = DynamicData->NumWeightedActiveVertexAnims > 0;
 
 	const FStaticLODModel& LODModel = SkeletalMeshResource->LODModels[DynamicData->LODIndex];
 	const TArray<FSkelMeshChunk>& Chunks = GetRenderChunks(DynamicData->LODIndex);
@@ -293,10 +297,17 @@ void FSkeletalMeshObjectGPUSkin::ProcessUpdatedDynamicData(FRHICommandListImmedi
 
 	bool DataPresent = false;
 
-	if(DynamicData->NumWeightedActiveVertexAnims > 0) 
+	if(bMorph) 
 	{
 		DataPresent = true;
 		checkSlow((VertexFactoryData.MorphVertexFactories.Num() == Chunks.Num()));
+		
+		// only update if the morph data changed and there are weighted morph targets
+		if(bMorphNeedsUpdate)
+		{
+			// update the morph data for the lod (before SkinCache)
+			LOD.UpdateMorphVertexBuffer( DynamicData->ActiveVertexAnims );
+		}
 	}
 	else
 	{
@@ -335,11 +346,6 @@ void FSkeletalMeshObjectGPUSkin::ProcessUpdatedDynamicData(FRHICommandListImmedi
 					INC_DWORD_STAT(STAT_GPUSkinCache_SkippedForZeroInfluences);
 					bUseSkinCache = false;
 				}
-				else if (DynamicData->NumWeightedActiveVertexAnims > 0)
-				{
-					INC_DWORD_STAT(STAT_GPUSkinCache_SkippedForMorphs);
-					bUseSkinCache = false;
-				}
 			}
 
 			// Create a uniform buffer from the bone transforms.
@@ -349,7 +355,8 @@ void FSkeletalMeshObjectGPUSkin::ProcessUpdatedDynamicData(FRHICommandListImmedi
 			// Try to use the GPU skinning cache if possible
 			if (bUseSkinCache)
 			{
-				int32 Key = GGPUSkinCache.StartCacheMesh(RHICmdList, GPUSkinCacheKeys[ChunkIdx], VertexFactoryData.VertexFactories[ChunkIdx].Get(), VertexFactoryData.PassthroughVertexFactories[ChunkIdx].Get(), Chunk, this);
+				int32 Key = GGPUSkinCache.StartCacheMesh(RHICmdList, GPUSkinCacheKeys[ChunkIdx], VertexFactoryData.VertexFactories[ChunkIdx].Get(),
+					VertexFactoryData.PassthroughVertexFactories[ChunkIdx].Get(), Chunk, this, bMorph ? &LOD.MorphVertexBuffer : 0);
 				if(Key >= 0)
 				{
 					GPUSkinCacheKeys[ChunkIdx] = (int16)Key;
@@ -375,14 +382,6 @@ void FSkeletalMeshObjectGPUSkin::ProcessUpdatedDynamicData(FRHICommandListImmedi
 			}
 		}
 	}
-
-	// only update if the morph data changed and there are weighted morph targets
-	if( bMorphNeedsUpdate &&
-		DynamicData->NumWeightedActiveVertexAnims > 0 )
-	{
-		// update the morph data for the lod
-		LOD.UpdateMorphVertexBuffer( DynamicData->ActiveVertexAnims );
-	}
 }
 
 TArray<FVector> FSkeletalMeshObjectGPUSkin::FSkeletalMeshObjectLOD::MorphDeltaTangentZAccumulationArray;
@@ -500,24 +499,21 @@ const FVertexFactory* FSkeletalMeshObjectGPUSkin::GetVertexFactory(int32 LODInde
 	{
 		return LOD.GPUSkinVertexFactories.ClothVertexFactories[ChunkIdx]->GetVertexFactory();
 	}
+
+	// If the GPU skinning cache was used, return the passthrough vertex factory
+	if (GGPUSkinCache.IsElementProcessed(GPUSkinCacheKeys[ChunkIdx]))
+	{
+		return LOD.GPUSkinVertexFactories.PassthroughVertexFactories[ChunkIdx].Get();
+	}
+
 	// use the morph enabled vertex factory if any active morphs are set
-	else if( DynamicData->NumWeightedActiveVertexAnims > 0 )
+	if( DynamicData->NumWeightedActiveVertexAnims > 0 )
 	{
 		return LOD.GPUSkinVertexFactories.MorphVertexFactories[ChunkIdx].Get();
 	}
+
 	// use the default gpu skin vertex factory
-	else
-	{
-		// If the GPU skinning cache was used, return the passthrough vertex factory
-		if (GGPUSkinCache.IsElementProcessed(GPUSkinCacheKeys[ChunkIdx]))
-		{
-			return LOD.GPUSkinVertexFactories.PassthroughVertexFactories[ChunkIdx].Get();
-		}
-		else
-		{
-			return LOD.GPUSkinVertexFactories.VertexFactories[ChunkIdx].Get();
-		}
-	}
+	return LOD.GPUSkinVertexFactories.VertexFactories[ChunkIdx].Get();
 }
 
 /** 
diff --git a/Engine/Source/Runtime/Engine/Private/SplineMeshSceneProxy.cpp b/Engine/Source/Runtime/Engine/Private/SplineMeshSceneProxy.cpp
--- a/Engine/Source/Runtime/Engine/Private/SplineMeshSceneProxy.cpp
+++ b/Engine/Source/Runtime/Engine/Private/SplineMeshSceneProxy.cpp
@@ -39,33 +39,14 @@ FSplineMeshSceneProxy::FSplineMeshSceneProxy(USplineMeshComponent* InComponent)
 		SplineMeshMinZ = InComponent->SplineBoundaryMin * SplineMeshScaleZ;
 	}
 
-	LODResources.Reset(InComponent->StaticMesh->RenderData->LODResources.Num());
-
 	for (int32 LODIndex = 0; LODIndex < LODs.Num(); LODIndex++)
 	{
-		FSplineMeshVertexFactory* VertexFactory = new FSplineMeshVertexFactory(this);
-
-		LODResources.Add(VertexFactory);
-
-		if( InComponent->LODData.IsValidIndex( LODIndex ))
+		InitVertexFactory(InComponent, LODIndex, nullptr); // we always need this one for shadows etc
+		if (InComponent->LODData.IsValidIndex(LODIndex) && InComponent->LODData[LODIndex].OverrideVertexColors)
 		{
-			const FStaticMeshComponentLODInfo& ComponentLODInfo = InComponent->LODData[ LODIndex ];
-			InitResources( InComponent, LODIndex, ComponentLODInfo.OverrideVertexColors );
+			InitVertexFactory(InComponent, LODIndex, InComponent->LODData[LODIndex].OverrideVertexColors);
 		}
-		else
-			InitResources( InComponent, LODIndex, NULL );
-	}
-}
-
-FSplineMeshSceneProxy::~FSplineMeshSceneProxy()
-{
-	ReleaseResources();
-
-	for (FSplineMeshSceneProxy::FLODResources& LODResource : LODResources)
-	{
-		delete LODResource.VertexFactory;
 	}
-	LODResources.Empty();
 }
 
 bool FSplineMeshSceneProxy::GetShadowMeshElement(int32 LODIndex, int32 BatchIndex, uint8 InDepthPriorityGroup, FMeshBatch& OutMeshBatch, bool bDitheredLODTransition) const
@@ -74,7 +55,12 @@ bool FSplineMeshSceneProxy::GetShadowMeshElement(int32 LODIndex, int32 BatchInde
 
 	if (FStaticMeshSceneProxy::GetShadowMeshElement(LODIndex, BatchIndex, InDepthPriorityGroup, OutMeshBatch, bDitheredLODTransition))
 	{
-		OutMeshBatch.VertexFactory = LODResources[LODIndex].VertexFactory;
+		const FStaticMeshLODResources& LOD = RenderData->LODResources[LODIndex];
+		check(OutMeshBatch.Elements.Num() == 1);
+		OutMeshBatch.VertexFactory = OutMeshBatch.Elements[0].bUserDataIsColorVertexBuffer ? LOD.SplineVertexFactoryOverrideColorVertexBuffer : LOD.SplineVertexFactory;
+		check(OutMeshBatch.VertexFactory);
+		OutMeshBatch.Elements[0].SplineMeshSceneProxy = const_cast<FSplineMeshSceneProxy*>(this);
+		OutMeshBatch.Elements[0].bIsSplineProxy = true;
 		OutMeshBatch.ReverseCulling ^= (SplineParams.StartScale.X < 0) ^ (SplineParams.StartScale.Y < 0);
 		return true;
 	}
@@ -87,7 +73,12 @@ bool FSplineMeshSceneProxy::GetMeshElement(int32 LODIndex, int32 BatchIndex, int
 
 	if (FStaticMeshSceneProxy::GetMeshElement(LODIndex, BatchIndex, SectionIndex, InDepthPriorityGroup, bUseSelectedMaterial, bUseHoveredMaterial, bAllowPreCulledIndices, OutMeshBatch))
 	{
-		OutMeshBatch.VertexFactory = LODResources[LODIndex].VertexFactory;
+		const FStaticMeshLODResources& LOD = RenderData->LODResources[LODIndex];
+		check(OutMeshBatch.Elements.Num() == 1);
+		OutMeshBatch.VertexFactory = OutMeshBatch.Elements[0].bUserDataIsColorVertexBuffer ? LOD.SplineVertexFactoryOverrideColorVertexBuffer : LOD.SplineVertexFactory;
+		check(OutMeshBatch.VertexFactory);
+		OutMeshBatch.Elements[0].SplineMeshSceneProxy = const_cast<FSplineMeshSceneProxy*>(this);
+		OutMeshBatch.Elements[0].bIsSplineProxy = true;
 		OutMeshBatch.ReverseCulling ^= (SplineParams.StartScale.X < 0) ^ (SplineParams.StartScale.Y < 0);
 		return true;
 	}
@@ -100,7 +91,12 @@ bool FSplineMeshSceneProxy::GetWireframeMeshElement(int32 LODIndex, int32 BatchI
 
 	if (FStaticMeshSceneProxy::GetWireframeMeshElement(LODIndex, BatchIndex, WireframeRenderProxy, InDepthPriorityGroup, bAllowPreCulledIndices, OutMeshBatch))
 	{
-		OutMeshBatch.VertexFactory = LODResources[LODIndex].VertexFactory;
+		const FStaticMeshLODResources& LOD = RenderData->LODResources[LODIndex];
+		check(OutMeshBatch.Elements.Num() == 1);
+		OutMeshBatch.VertexFactory = OutMeshBatch.Elements[0].bUserDataIsColorVertexBuffer ? LOD.SplineVertexFactoryOverrideColorVertexBuffer : LOD.SplineVertexFactory;
+		check(OutMeshBatch.VertexFactory);
+		OutMeshBatch.Elements[0].SplineMeshSceneProxy = const_cast<FSplineMeshSceneProxy*>(this);
+		OutMeshBatch.Elements[0].bIsSplineProxy = true;
 		OutMeshBatch.ReverseCulling ^= (SplineParams.StartScale.X < 0) ^ (SplineParams.StartScale.Y < 0);
 		return true;
 	}
diff --git a/Engine/Source/Runtime/Engine/Private/SplineMeshSceneProxy.h b/Engine/Source/Runtime/Engine/Private/SplineMeshSceneProxy.h
--- a/Engine/Source/Runtime/Engine/Private/SplineMeshSceneProxy.h
+++ b/Engine/Source/Runtime/Engine/Private/SplineMeshSceneProxy.h
@@ -17,11 +17,6 @@ struct FSplineMeshVertexFactory : public FLocalVertexFactory
 	DECLARE_VERTEX_FACTORY_TYPE(FSplineMeshVertexFactory);
 public:
 
-	FSplineMeshVertexFactory(class FSplineMeshSceneProxy* InSplineProxy) :
-		SplineSceneProxy(InSplineProxy)
-	{
-	}
-
 	/** Should we cache the material's shadertype on this platform with this vertex factory? */
 	static bool ShouldCache(EShaderPlatform Platform, const class FMaterial* Material, const class FShaderType* ShaderType)
 	{
@@ -49,17 +44,7 @@ public:
 		BeginUpdateResourceRHI(this);
 	}
 
-	// FRenderResource interface.
-	virtual void InitRHI() override
-	{
-		FLocalVertexFactory::InitRHI();
-	}
-
 	static FVertexFactoryShaderParameters* ConstructShaderParameters(EShaderFrequency ShaderFrequency);
-
-	class FSplineMeshSceneProxy* SplineSceneProxy;
-
-private:
 };
 
 //////////////////////////////////////////////////////////////////////////
@@ -147,11 +132,7 @@ public:
 
 	FSplineMeshSceneProxy(USplineMeshComponent* InComponent);
 
-	virtual ~FSplineMeshSceneProxy() override;
-
-	void InitResources(USplineMeshComponent* InComponent, int32 InLODIndex, FColorVertexBuffer*);
-
-	void ReleaseResources();
+	void InitVertexFactory(USplineMeshComponent* InComponent, int32 InLODIndex, FColorVertexBuffer*);
 
 	/** Sets up a shadow FMeshBatch for a specific LOD. */
 	virtual bool GetShadowMeshElement(int32 LODIndex, int32 BatchIndex, uint8 InDepthPriorityGroup, FMeshBatch& OutMeshBatch, bool bDitheredLODTransition) const override;
diff --git a/Engine/Source/Runtime/Engine/Private/StaticMesh.cpp b/Engine/Source/Runtime/Engine/Private/StaticMesh.cpp
--- a/Engine/Source/Runtime/Engine/Private/StaticMesh.cpp
+++ b/Engine/Source/Runtime/Engine/Private/StaticMesh.cpp
@@ -513,11 +513,15 @@ FStaticMeshLODResources::FStaticMeshLODResources()
 	, bHasReversedIndices(false)
 	, bHasReversedDepthOnlyIndices(false)
 	, DepthOnlyNumTriangles(0)
+	, SplineVertexFactory(nullptr)
+	, SplineVertexFactoryOverrideColorVertexBuffer(nullptr)
 {
 }
 
 FStaticMeshLODResources::~FStaticMeshLODResources()
 {
+	delete (FRenderResource*)SplineVertexFactory;
+	delete (FRenderResource*)SplineVertexFactoryOverrideColorVertexBuffer;
 	delete DistanceFieldData;
 }
 
@@ -633,6 +637,15 @@ void FStaticMeshLODResources::ReleaseResources()
 	BeginReleaseResource(&VertexFactory);
 	BeginReleaseResource(&VertexFactoryOverrideColorVertexBuffer);
 
+	if (SplineVertexFactory)
+	{
+		BeginReleaseResource((FRenderResource *)SplineVertexFactory);
+	}
+	if (SplineVertexFactoryOverrideColorVertexBuffer)
+	{
+		BeginReleaseResource((FRenderResource *)SplineVertexFactoryOverrideColorVertexBuffer);
+	}
+
 	if (DistanceFieldData)
 	{
 		DEC_DWORD_STAT_BY( STAT_StaticMeshDistanceFieldMemory, DistanceFieldData->GetResourceSize() );
diff --git a/Engine/Source/Runtime/Engine/Private/StaticMeshLight.cpp b/Engine/Source/Runtime/Engine/Private/StaticMeshLight.cpp
--- a/Engine/Source/Runtime/Engine/Private/StaticMeshLight.cpp
+++ b/Engine/Source/Runtime/Engine/Private/StaticMeshLight.cpp
@@ -287,7 +287,7 @@ void FStaticMeshStaticLightingTextureMapping::Apply(FQuantizedLightmapData* Quan
 #if WITH_EDITOR
 void UStaticMeshComponent::GetStaticLightingInfo(FStaticLightingPrimitiveInfo& OutPrimitiveInfo,const TArray<ULightComponent*>& InRelevantLights,const FLightingBuildOptions& Options)
 {
-	if( HasValidSettingsForStaticLighting() )
+	if( HasValidSettingsForStaticLighting(false) )
 	{
 		int32		BaseLightMapWidth	= 0;
 		int32		BaseLightMapHeight	= 0;
@@ -347,7 +347,7 @@ void UStaticMeshComponent::GetStaticLightingInfo(FStaticLightingPrimitiveInfo& O
 ELightMapInteractionType UStaticMeshComponent::GetStaticLightingType() const
 {
 	bool bUseTextureMap = false;
-	if( HasValidSettingsForStaticLighting() )
+	if( HasValidSettingsForStaticLighting(false) )
 	{
 		// Process each LOD separately.
 		TArray<FStaticMeshStaticLightingMesh*> StaticLightingMeshes;
diff --git a/Engine/Source/Runtime/Engine/Private/TextureStreamingBuild.cpp b/Engine/Source/Runtime/Engine/Private/TextureStreamingBuild.cpp
new file mode 100644
--- /dev/null
+++ b/Engine/Source/Runtime/Engine/Private/TextureStreamingBuild.cpp
@@ -0,0 +1,53 @@
+// Copyright 1998-2016 Epic Games, Inc. All Rights Reserved.
+
+/*=============================================================================
+TextureStreamingBuild.cpp : Contains definitions to build texture streaming data.
+=============================================================================*/
+
+#include "EnginePrivate.h"
+#include "DebugVieWModeMaterialProxy.h"
+#include "ShaderCompiler.h"
+
+ENGINE_API void BuildTextureStreamingData(UWorld* InWorld)
+{
+#if WITH_EDITORONLY_DATA
+	FlushRenderingCommands();
+
+	double StartTime = FPlatformTime::Seconds();
+
+	FDebugViewModeMaterialProxy::ClearAllShaders();
+
+	if (InWorld)
+	{
+		for (int32 LevelIndex=0; LevelIndex < InWorld->GetNumLevels(); LevelIndex++)
+		{
+			ULevel* Level = InWorld->GetLevel(LevelIndex);
+
+			TArray<UObject*> ObjectsInOuter;
+			GetObjectsWithOuter(Level, ObjectsInOuter);
+
+			for (int32 Index = 0; Index < ObjectsInOuter.Num(); Index++)
+			{
+				UPrimitiveComponent* Primitive = Cast<UPrimitiveComponent>(ObjectsInOuter[Index]);
+				if (!Primitive) continue;
+
+				TArray<UMaterialInterface*> Materials;
+				Primitive->GetUsedMaterials(Materials);
+
+				for (int32 MaterialIndex = 0; MaterialIndex < Materials.Num(); ++MaterialIndex)
+				{
+					UMaterialInterface* MaterialInterface = Materials[MaterialIndex];
+					if (!MaterialInterface) continue;
+
+					FDebugViewModeMaterialProxy::AddShader(MaterialInterface, EMaterialShaderMapUsage::DebugViewModeTexCoordScale);
+				}
+			}
+		}
+	}
+	GShaderCompilingManager->FinishAllCompilation();
+
+	UE_LOG(LogLevel, Display, TEXT("Build TextureStreaming took %.3f seconds."), FPlatformTime::Seconds() - StartTime);
+#else
+	UE_LOG(LogLevel, Fatal,TEXT("Build TextureStreaming should not be called on a console"));
+#endif
+}
diff --git a/Engine/Source/Runtime/Engine/Public/BlendableManager.h b/Engine/Source/Runtime/Engine/Public/BlendableManager.h
--- a/Engine/Source/Runtime/Engine/Public/BlendableManager.h
+++ b/Engine/Source/Runtime/Engine/Public/BlendableManager.h
@@ -13,9 +13,11 @@ private:
 	FName BlendableType;
 	// to be able to jump over data
 	uint32 DataSize;
-
+	// to align the data
+	uint32 PrePadding;
+	 
 	// @return pointer to the next object or end (can be compared with container end pointer)
-	uint8* GetDataPtr() { check(this); return (uint8*)(this + 1); }
+	uint8* GetDataPtr(){ check(this); return ((uint8*)(this + 1)) + PrePadding; }
 	// @return next or end of the array
 	FBlendableEntry* GetNext() { return (FBlendableEntry*)(GetDataPtr() + DataSize); }
 
@@ -35,15 +37,19 @@ public:
 
 	// @param InWeight 0..1, excluding 0 as this is used to disable entries
 	// @param InData is copied with a memcpy
-	// @return pointer to the newly added entry
 	template <class T>
 	T* PushBlendableData(float InWeight, const T& InData)
 	{
 		FName BlendableType = T::GetFName();
 
-		FBlendableEntry* Ret = PushBlendableDataPtr(InWeight, BlendableType, (const uint8*)&InData, sizeof(T));
+		// at least 4 byte alignment
+		uint32 Alignment = FMath::Max((uint32)4, (uint32)ALIGNOF(T));
+
+		FBlendableEntry* Entry = PushBlendableDataPtr(InWeight, BlendableType, (const uint8*)&InData, sizeof(T), Alignment);
+
+		T* Ret = (T*)Entry->GetDataPtr();
 
-		return (T*)Ret;
+		return Ret;
 	}
 
 	// used to blend multiple blendables of the given type with lerp into one final one
@@ -96,8 +102,8 @@ public:
 		do
 		{
 			InIterator = GetNextBlendableEntryPtr(InIterator);
-
-		} while (InIterator && InIterator->Weight <= 0.0f && InIterator->BlendableType != BlendableType);
+			
+		} while (InIterator && InIterator->BlendableType != BlendableType);		
 
 		if (InIterator)
 		{
@@ -140,21 +146,29 @@ private:
 	// @param InData is copied
 	// @param InDataSize >0
 	// @return pointer to the newly added entry
-	FBlendableEntry* PushBlendableDataPtr(float InWeight, FName InBlendableType, const uint8* InData, uint32 InDataSize)
+	FBlendableEntry* PushBlendableDataPtr(float InWeight, FName InBlendableType, const uint8* InData, uint32 InDataSize, uint32 Alignment)
 	{
 		check(InWeight > 0.0f && InWeight <= 1.0f);
 		check(InData);
 		check(InDataSize);
 
-		uint32 OldSize = Scratch.AddUninitialized(sizeof(FBlendableEntry) + InDataSize);
+		uint32 PrePadding;
+		{
+			uint8* DataStart = Scratch.GetData() + Scratch.Num() + sizeof(FBlendableEntry);
+			PrePadding = static_cast<uint32>(Alignment - reinterpret_cast<ptrdiff_t>(DataStart) % Alignment);
+			PrePadding = (PrePadding == Alignment) ? 0 : PrePadding;	
+		}
 
+		uint32 OldSize = Scratch.AddUninitialized(sizeof(FBlendableEntry) + InDataSize + PrePadding);
+		
 		FBlendableEntry* Dst = (FBlendableEntry*)&Scratch[OldSize];
 
 		Dst->Weight = InWeight;
 		Dst->BlendableType = InBlendableType;
 		Dst->DataSize = InDataSize;
+		Dst->PrePadding = PrePadding;
 		memcpy(Dst->GetDataPtr(), InData, InDataSize);
 
-		return (FBlendableEntry*)Dst->GetDataPtr();
+		return Dst;
 	}
 };
\ No newline at end of file
diff --git a/Engine/Source/Runtime/Engine/Public/GPUSkinCache.h b/Engine/Source/Runtime/Engine/Public/GPUSkinCache.h
--- a/Engine/Source/Runtime/Engine/Public/GPUSkinCache.h
+++ b/Engine/Source/Runtime/Engine/Public/GPUSkinCache.h
@@ -12,6 +12,7 @@
 #include "RenderResource.h"
 #include "UniformBuffer.h"
 
+typedef FRHIShaderResourceView*              FShaderResourceViewRHIParamRef;
 
 extern ENGINE_API int32 GEnableGPUSkinCacheShaders;
 extern ENGINE_API int32 GEnableGPUSkinCache;
@@ -61,8 +62,11 @@ public:
 	//	For each Chunk:
 	//		Call Add*
 	//	Call End*
+	// @param MorphVertexBuffer if no moprh targets are used
 	// @returns -1 if failed, otherwise index into CachedElements[]
-	int32 StartCacheMesh(FRHICommandListImmediate& RHICmdList, int32 Key, class FGPUBaseSkinVertexFactory* VertexFactory, class FGPUSkinPassthroughVertexFactory* TargetVertexFactory, const struct FSkelMeshChunk& BatchElement, const FSkeletalMeshObjectGPUSkin* Skin);
+	int32 StartCacheMesh(FRHICommandListImmediate& RHICmdList, int32 Key, class FGPUBaseSkinVertexFactory* VertexFactory,
+		class FGPUSkinPassthroughVertexFactory* TargetVertexFactory, const struct FSkelMeshChunk& BatchElement, const FSkeletalMeshObjectGPUSkin* Skin,
+		const class FMorphVertexBuffer* MorphVertexBuffer);
 
 	inline bool SetVertexStreamFromCache(FRHICommandList& RHICmdList, int32 Key, FShader* Shader, const FGPUSkinPassthroughVertexFactory* VertexFactory, uint32 BaseVertexIndex, FShaderParameter PreviousStreamFloatOffset, FShaderResourceParameter PreviousStreamBuffer)
 	{
@@ -84,6 +88,7 @@ public:
 		int32	Key;
 		// index into CachedVertexBuffers[] or -1 if it has no assignment yet
 		int32	VertexBufferSRVIndex;
+		int32	MorphVertexBufferSRVIndex;
 
 		uint32	FrameUpdated;
 
@@ -95,6 +100,12 @@ public:
 
 		bool	bExtraBoneInfluences;
 
+		FElementCacheStatusInfo()
+			: VertexBufferSRVIndex(-1)
+			, MorphVertexBufferSRVIndex(-1)
+		{
+		}
+
 		struct FSearchInfo
 		{
 			const FSkeletalMeshObjectGPUSkin* Skin;
@@ -117,18 +128,6 @@ public:
 		}
 	};
 
-	struct FVertexBufferInfo
-	{
-		const FVertexBuffer* VertexBuffer;
-		FShaderResourceViewRHIRef	VertexBufferSRV;
-		int32	Index;
-
-		bool operator == (const FVertexBuffer& OtherVertexBuffer) const 
-		{
-			return (VertexBuffer == &OtherVertexBuffer);
-		}
-	};
-
 	ENGINE_API void TransitionToReadable(FRHICommandList& RHICmdList);
 	ENGINE_API void TransitionToWriteable(FRHICommandList& RHICmdList);
 
@@ -136,7 +135,14 @@ public:
 private:
 	FElementCacheStatusInfo* FindEvictableCacheStatusInfo();
 
-	void	DispatchSkinCacheProcess(FRHICommandListImmediate& RHICmdList, uint32 InputStreamFloatOffset, uint32 OutputBufferFloatOffset, const struct FVertexBufferAndSRV& BoneBuffer, FUniformBufferRHIRef UniformBuffer, const FVertexBufferInfo* VBInfo, uint32 VertexStride, uint32 VertexCount, const FVector& MeshOrigin, const FVector& MeshExtension, bool bUseExtraBoneInfluences, ERHIFeatureLevel::Type FeatureLevel);
+	struct FVertexBufferInfo;
+
+	// @param SkinType 0:normal, 1:with morph target, 2:with APEX cloth (not yet implemented)
+	void DispatchSkinCacheProcess(FRHICommandListImmediate& RHICmdList, uint32 InputStreamFloatOffset, uint32 OutputBufferFloatOffset,
+		const struct FVertexBufferAndSRV& BoneBuffer, FUniformBufferRHIRef UniformBuffer, const FVertexBufferInfo* VBInfo,
+		uint32 VertexStride, uint32 VertexCount, const FVector& MeshOrigin, const FVector& MeshExtension, bool bUseExtraBoneInfluences, ERHIFeatureLevel::Type FeatureLevel,
+		uint32 SkinType, uint32 InMorphInputStreamFloatOffset, FShaderResourceViewRHIParamRef MorphVertexBufferSRV);
+
 	bool InternalSetVertexStreamFromCache(FRHICommandList& RHICmdList, int32 Key, FShader* Shader, const class FGPUSkinPassthroughVertexFactory* VertexFactory, uint32 BaseVertexIndex, FShaderParameter PreviousStreamFloatOffset, FShaderResourceParameter PreviousStreamBuffer);
 
 	bool	bInitialized : 1;
@@ -150,8 +156,6 @@ private:
 	int		CacheCurrentFloatOffset;
 	int		CachedChunksThisFrameCount;
 
-	TArray<FVertexBufferInfo>		CachedVertexBuffers;
-
 	TArray<FElementCacheStatusInfo>	CachedElements;
 
 	FRWBuffer	SkinCacheBuffer[GPUSKINCACHE_FRAMES];
@@ -159,6 +163,29 @@ private:
 	static void CVarSinkFunction();
 
 	static FAutoConsoleVariableSink CVarSink;
+
+	
+	// -------------------------------------------------
+	
+	// to cache the SRV needed to access the input vertex buffers
+	struct FVertexBufferInfo
+	{
+		// key: input vertex buffer from which we need an SRV
+		const FVertexBuffer* VertexBuffer;
+		// value: SRV created from the VertexBuffer
+		FShaderResourceViewRHIRef VertexBufferSRV;
+		
+		// to support FindByKey()
+		bool operator == (const FVertexBuffer& OtherVertexBuffer) const 
+		{
+			return (VertexBuffer == &OtherVertexBuffer);
+		}
+	};
+	
+	// to cache the SRV needed to access the input vertex buffers
+	TArray<FVertexBufferInfo> CachedVertexBuffers;
+	
+	FVertexBufferInfo* ReuseOrFindAndCacheSRV(int32& InOutIndex, const FVertexBuffer& VertexBuffer);
 };
 
 BEGIN_UNIFORM_BUFFER_STRUCT(GPUSkinCacheBonesUniformShaderParameters,)
@@ -176,7 +203,6 @@ DECLARE_DWORD_COUNTER_STAT_EXTERN(TEXT("Skipped for Max Chunks Per LOD"), STAT_G
 DECLARE_DWORD_COUNTER_STAT_EXTERN(TEXT("Skipped for Num Chunks Per Frame"), STAT_GPUSkinCache_SkippedForChunksPerFrame, STATGROUP_GPUSkinCache,);
 DECLARE_DWORD_COUNTER_STAT_EXTERN(TEXT("Skipped for zero influences"), STAT_GPUSkinCache_SkippedForZeroInfluences, STATGROUP_GPUSkinCache,);
 DECLARE_DWORD_COUNTER_STAT_EXTERN(TEXT("Skipped for Cloth"), STAT_GPUSkinCache_SkippedForCloth, STATGROUP_GPUSkinCache,);
-DECLARE_DWORD_COUNTER_STAT_EXTERN(TEXT("Skipped for Morphs"), STAT_GPUSkinCache_SkippedForMorphs, STATGROUP_GPUSkinCache,);
 DECLARE_DWORD_COUNTER_STAT_EXTERN(TEXT("Skipped for Memory"), STAT_GPUSkinCache_SkippedForMemory, STATGROUP_GPUSkinCache,);
 DECLARE_DWORD_COUNTER_STAT_EXTERN(TEXT("Skipped for Already Cached"), STAT_GPUSkinCache_SkippedForAlreadyCached, STATGROUP_GPUSkinCache,);
 DECLARE_DWORD_COUNTER_STAT_EXTERN(TEXT("Skipped for out of ECS"), STAT_GPUSkinCache_SkippedForOutOfECS, STATGROUP_GPUSkinCache,);
diff --git a/Engine/Source/Runtime/Engine/Public/GPUSkinPublicDefs.h b/Engine/Source/Runtime/Engine/Public/GPUSkinPublicDefs.h
--- a/Engine/Source/Runtime/Engine/Public/GPUSkinPublicDefs.h
+++ b/Engine/Source/Runtime/Engine/Public/GPUSkinPublicDefs.h
@@ -36,5 +36,5 @@ namespace SkinningTools
 	}
 }
 
-// Number of frames buffered
+// Number of frames buffered (2 would be enough for normal rendering as GPU produces and consumes but we need one more frame for velocity)
 #define GPUSKINCACHE_FRAMES 3
diff --git a/Engine/Source/Runtime/Engine/Public/GlobalShader.h b/Engine/Source/Runtime/Engine/Public/GlobalShader.h
--- a/Engine/Source/Runtime/Engine/Public/GlobalShader.h
+++ b/Engine/Source/Runtime/Engine/Public/GlobalShader.h
@@ -131,7 +131,9 @@ public:
 		CheckShaderIsValid();
 		SetUniformBufferParameter(RHICmdList, ShaderRHI, ViewUniformBufferParameter, View.ViewUniformBuffer);
 		SetUniformBufferParameter(RHICmdList, ShaderRHI, FrameUniformBufferParameter, View.FrameUniformBuffer);
+#if USE_GBuiltinSamplersUniformBuffer
 		SetUniformBufferParameter(RHICmdList, ShaderRHI, BuiltinSamplersUBParameter, GBuiltinSamplersUniformBuffer.GetUniformBufferRHI());
+#endif
 	}
 
 	typedef void (*ModifyCompilationEnvironmentType)(EShaderPlatform, FShaderCompilerEnvironment&);
diff --git a/Engine/Source/Runtime/Engine/Public/MaterialCompiler.h b/Engine/Source/Runtime/Engine/Public/MaterialCompiler.h
--- a/Engine/Source/Runtime/Engine/Public/MaterialCompiler.h
+++ b/Engine/Source/Runtime/Engine/Public/MaterialCompiler.h
@@ -7,6 +7,7 @@
 #pragma once
 
 #include "MaterialShared.h"
+#include "Materials/MaterialExpressionScreenPosition.h"
 #include "Materials/MaterialExpressionSpeedTree.h"
 #include "Materials/MaterialExpressionTextureSample.h"
 #include "Materials/MaterialExpressionWorldPosition.h"
@@ -79,7 +80,7 @@ public:
 	virtual int32 CameraVector() = 0;
 	virtual int32 LightVector() = 0;
 
-	virtual int32 ScreenPosition() = 0;
+	virtual int32 ScreenPosition(EMaterialExpressionScreenPositionMapping Mapping) = 0;
 	virtual int32 WorldPosition(EWorldPositionIncludedOffsets WorldPositionIncludedOffsets) = 0;
 	virtual int32 ObjectWorldPosition() = 0;
 	virtual int32 ObjectRadius() = 0;
@@ -102,15 +103,27 @@ public:
 	virtual int32 If(int32 A,int32 B,int32 AGreaterThanB,int32 AEqualsB,int32 ALessThanB,int32 Threshold) = 0;
 
 	virtual int32 TextureCoordinate(uint32 CoordinateIndex, bool UnMirrorU, bool UnMirrorV) = 0;
-	virtual int32 TextureSample(int32 Texture,int32 Coordinate,enum EMaterialSamplerType SamplerType,int32 MipValue0Index=INDEX_NONE,int32 MipValue1Index=INDEX_NONE,ETextureMipValueMode MipValueMode=TMVM_None,ESamplerSourceMode SamplerSource=SSM_FromTextureAsset) = 0;
+	virtual int32 TextureSample(int32 Texture,int32 Coordinate,enum EMaterialSamplerType SamplerType,int32 MipValue0Index=INDEX_NONE,int32 MipValue1Index=INDEX_NONE,ETextureMipValueMode MipValueMode=TMVM_None,ESamplerSourceMode SamplerSource=SSM_FromTextureAsset,int32 TextureReferenceIndex=INDEX_NONE) = 0;
 	virtual int32 TextureProperty(int32 InTexture, EMaterialExposedTextureProperty Property) = 0;
 
 	virtual int32 TextureDecalMipmapLevel(int32 TextureSizeInput) = 0;
 	virtual int32 TextureDecalDerivative(bool bDDY) = 0;
 	virtual int32 DecalLifetimeOpacity() = 0;
 
-	virtual int32 Texture(UTexture* Texture,ESamplerSourceMode SamplerSource=SSM_FromTextureAsset) = 0;
-	virtual int32 TextureParameter(FName ParameterName,UTexture* DefaultTexture,ESamplerSourceMode SamplerSource=SSM_FromTextureAsset) = 0;
+	virtual int32 Texture(UTexture* Texture,int32& TextureReferenceIndex,ESamplerSourceMode SamplerSource=SSM_FromTextureAsset) = 0;
+	virtual int32 TextureParameter(FName ParameterName,UTexture* DefaultTexture,int32& TextureReferenceIndex,ESamplerSourceMode SamplerSource=SSM_FromTextureAsset) = 0;
+
+	int32 Texture(UTexture* InTexture,ESamplerSourceMode SamplerSource=SSM_FromTextureAsset)
+	{
+		int32 TextureReferenceIndex = INDEX_NONE;
+		return  Texture(InTexture, TextureReferenceIndex, SamplerSource);
+	}
+
+	int32 TextureParameter(FName ParameterName,UTexture* DefaultTexture,ESamplerSourceMode SamplerSource=SSM_FromTextureAsset)
+	{
+		int32 TextureReferenceIndex = INDEX_NONE;
+		return  TextureParameter(ParameterName, DefaultTexture, TextureReferenceIndex, SamplerSource);
+	}
 
 	virtual	int32 PixelDepth()=0;
 	virtual int32 SceneDepth(int32 Offset, int32 UV, bool bUseOffset) = 0;
@@ -256,7 +269,7 @@ public:
 	virtual int32 CameraVector() override { return Compiler->CameraVector(); }
 	virtual int32 LightVector() override { return Compiler->LightVector(); }
 
-	virtual int32 ScreenPosition() override { return Compiler->ScreenPosition(); }
+	virtual int32 ScreenPosition(EMaterialExpressionScreenPositionMapping Mapping = MESP_SceneTextureUV) override { return Compiler->ScreenPosition(Mapping); }
 	virtual int32 WorldPosition(EWorldPositionIncludedOffsets WorldPositionIncludedOffsets) override { return Compiler->WorldPosition(WorldPositionIncludedOffsets); }
 	virtual int32 ObjectWorldPosition() override { return Compiler->ObjectWorldPosition(); }
 	virtual int32 ObjectRadius() override { return Compiler->ObjectRadius(); }
@@ -272,8 +285,8 @@ public:
 
 	virtual int32 If(int32 A,int32 B,int32 AGreaterThanB,int32 AEqualsB,int32 ALessThanB,int32 Threshold) override { return Compiler->If(A,B,AGreaterThanB,AEqualsB,ALessThanB,Threshold); }
 
-	virtual int32 TextureSample(int32 InTexture,int32 Coordinate,enum EMaterialSamplerType SamplerType,int32 MipValue0Index,int32 MipValue1Index,ETextureMipValueMode MipValueMode,ESamplerSourceMode SamplerSource) override 
-		{ return Compiler->TextureSample(InTexture,Coordinate,SamplerType,MipValue0Index,MipValue1Index,MipValueMode,SamplerSource); }
+	virtual int32 TextureSample(int32 InTexture,int32 Coordinate,enum EMaterialSamplerType SamplerType,int32 MipValue0Index,int32 MipValue1Index,ETextureMipValueMode MipValueMode,ESamplerSourceMode SamplerSource,int32 TextureReferenceIndex) override 
+		{ return Compiler->TextureSample(InTexture,Coordinate,SamplerType,MipValue0Index,MipValue1Index,MipValueMode,SamplerSource,TextureReferenceIndex); }
 	virtual int32 TextureProperty(int32 InTexture, EMaterialExposedTextureProperty Property) override 
 		{ return Compiler->TextureProperty(InTexture, Property); }
 
@@ -283,8 +296,8 @@ public:
 	virtual int32 TextureDecalDerivative(bool bDDY) override { return Compiler->TextureDecalDerivative(bDDY); }
 	virtual int32 DecalLifetimeOpacity() override { return Compiler->DecalLifetimeOpacity(); }
 
-	virtual int32 Texture(UTexture* InTexture,ESamplerSourceMode SamplerSource=SSM_FromTextureAsset) override { return Compiler->Texture(InTexture,SamplerSource); }
-	virtual int32 TextureParameter(FName ParameterName,UTexture* DefaultValue,ESamplerSourceMode SamplerSource=SSM_FromTextureAsset) override { return Compiler->TextureParameter(ParameterName,DefaultValue,SamplerSource); }
+	virtual int32 Texture(UTexture* InTexture,int32& TextureReferenceIndex,ESamplerSourceMode SamplerSource=SSM_FromTextureAsset) override { return Compiler->Texture(InTexture,TextureReferenceIndex,SamplerSource); }
+	virtual int32 TextureParameter(FName ParameterName,UTexture* DefaultValue,int32& TextureReferenceIndex,ESamplerSourceMode SamplerSource=SSM_FromTextureAsset) override { return Compiler->TextureParameter(ParameterName,DefaultValue,TextureReferenceIndex,SamplerSource); }
 
 	virtual	int32 PixelDepth() override { return Compiler->PixelDepth();	}
 	virtual int32 SceneDepth(int32 Offset, int32 UV, bool bUseOffset) override { return Compiler->SceneDepth(Offset, UV, bUseOffset); }
diff --git a/Engine/Source/Runtime/Engine/Public/MaterialShared.h b/Engine/Source/Runtime/Engine/Public/MaterialShared.h
--- a/Engine/Source/Runtime/Engine/Public/MaterialShared.h
+++ b/Engine/Source/Runtime/Engine/Public/MaterialShared.h
@@ -360,7 +360,8 @@ namespace EMaterialShaderMapUsage
 		MaterialExportAO,
 		MaterialExportEmissive,
 		MaterialExportOpacity,
-		MaterialExportSubSurfaceColor
+		MaterialExportSubSurfaceColor,
+		DebugViewModeTexCoordScale
 	};
 }
 
@@ -1625,6 +1626,8 @@ protected:
 	ENGINE_API virtual bool HasMaterialAttributesConnected() const override;
 	/** Useful for debugging. */
 	ENGINE_API virtual FString GetBaseMaterialPathName() const override;
+
+	friend class FDebugViewModeMaterialProxy; // Needed to redirect compilation
 };
 
 /**
diff --git a/Engine/Source/Runtime/Engine/Public/MeshBatch.h b/Engine/Source/Runtime/Engine/Public/MeshBatch.h
--- a/Engine/Source/Runtime/Engine/Public/MeshBatch.h
+++ b/Engine/Source/Runtime/Engine/Public/MeshBatch.h
@@ -19,10 +19,15 @@ struct FMeshBatchElement
 
 	const FIndexBuffer* IndexBuffer;
 
-	/** Instance runs, where number of runs is specified by NumInstances.  Run structure is [StartInstanceIndex, EndInstanceIndex]. */
-	uint32* InstanceRuns;
+	union 
+	{
+		/** If bIsSplineProxy, Instance runs, where number of runs is specified by NumInstances.  Run structure is [StartInstanceIndex, EndInstanceIndex]. */
+		uint32* InstanceRuns;
+		/** If bIsSplineProxy, a pointer back to the proxy */
+		class FSplineMeshSceneProxy* SplineMeshSceneProxy;
+	};
 	const void* UserData;
-	/** 
+	/**
 	 *	DynamicIndexData - pointer to user memory containing the index data.
 	 *	Used for rendering dynamic data directly.
 	 */
@@ -43,7 +48,9 @@ struct FMeshBatchElement
 	uint8 InstancedLODRange : 4;
 	uint8 bUserDataIsColorVertexBuffer : 1;
 	uint8 bIsInstancedMesh : 1;
-	
+	uint8 bIsSplineProxy : 1;
+	uint8 bIsInstanceRuns : 1;
+
 	FMeshBatchElement()
 	:	PrimitiveUniformBufferResource(nullptr)
 	,	IndexBuffer(nullptr)
@@ -58,6 +65,8 @@ struct FMeshBatchElement
 	,	InstancedLODRange(0)
 	,	bUserDataIsColorVertexBuffer(false)
 	,   bIsInstancedMesh(false)
+	,	bIsSplineProxy(false)
+	,	bIsInstanceRuns(false)
 	{
 	}
 };
@@ -176,7 +185,7 @@ struct FMeshBatch
 		int32 Count=0;
 		for( int32 ElementIdx=0;ElementIdx<Elements.Num();ElementIdx++ )
 		{
-			if (Elements[ElementIdx].InstanceRuns)
+			if (Elements[ElementIdx].bIsInstanceRuns && Elements[ElementIdx].InstanceRuns)
 			{
 				for (uint32 Run = 0; Run < Elements[ElementIdx].NumInstances; Run++)
 				{
diff --git a/Engine/Source/Runtime/Engine/Public/MeshParticleVertexFactory.h b/Engine/Source/Runtime/Engine/Public/MeshParticleVertexFactory.h
--- a/Engine/Source/Runtime/Engine/Public/MeshParticleVertexFactory.h
+++ b/Engine/Source/Runtime/Engine/Public/MeshParticleVertexFactory.h
@@ -17,6 +17,7 @@ BEGIN_UNIFORM_BUFFER_STRUCT( FMeshParticleUniformParameters, )
 	DECLARE_UNIFORM_BUFFER_STRUCT_MEMBER( FVector4, SubImageSize )
 	DECLARE_UNIFORM_BUFFER_STRUCT_MEMBER( uint32, TexCoordWeightA )
 	DECLARE_UNIFORM_BUFFER_STRUCT_MEMBER( uint32, TexCoordWeightB )
+	DECLARE_UNIFORM_BUFFER_STRUCT_MEMBER( uint32, PrevTransformAvailable )
 END_UNIFORM_BUFFER_STRUCT( FMeshParticleUniformParameters )
 typedef TUniformBufferRef<FMeshParticleUniformParameters> FMeshParticleUniformBufferRef;
 
@@ -63,7 +64,6 @@ public:
 		FDataType()
 		: bInitialized(false)
 		{
-
 		}
 	};
 
@@ -72,6 +72,7 @@ public:
 	public:
 		const struct FMeshParticleInstanceVertex* InstanceBuffer;
 		const struct FMeshParticleInstanceVertexDynamicParameter* DynamicParameterBuffer;
+		const struct FMeshParticleInstanceVertexPrevTransform* PrevTransformBuffer;
 	};
 
 	/** Default constructor. */
@@ -103,9 +104,9 @@ public:
 
 		// Set a define so we can tell in MaterialTemplate.usf when we are compiling a mesh particle vertex factory
 		OutEnvironment.SetDefine(TEXT("PARTICLE_MESH_FACTORY"),TEXT("1"));
-		OutEnvironment.SetDefine(TEXT("PARTICLE_MESH_INSTANCED"),(uint32) 1);
+		OutEnvironment.SetDefine(TEXT("PARTICLE_MESH_INSTANCED"), (GetMaxSupportedFeatureLevel(Platform) > ERHIFeatureLevel::ES2) ? 1 : 0);
 	}
-
+	
 	/**
 	 * An implementation of the interface used by TSynchronizedResource to update the resource with new data from the game thread.
 	 */
@@ -126,7 +127,7 @@ public:
 	{
 		return MeshParticleUniformBuffer;
 	}
-
+	
 	/**
 	 * Update the data strides (MUST HAPPEN BEFORE InitRHI is called)
 	 */
@@ -146,6 +147,10 @@ public:
 	 */
 	void SetDynamicParameterBuffer(const FVertexBuffer* InDynamicParameterBuffer, uint32 StreamOffset, uint32 Stride);
 
+	uint8* LockPreviousTransformBuffer(uint32 ParticleCount);
+	void UnlockPreviousTransformBuffer();
+	FShaderResourceViewRHIParamRef GetPreviousTransformBufferSRV() const;
+
 	/**
 	* Copy the data from another vertex factory
 	* @param Other - factory to copy from
@@ -166,9 +171,10 @@ protected:
 	int32 DynamicVertexStride;
 	int32 DynamicParameterVertexStride;
 	
-
 	/** Uniform buffer with mesh particle parameters. */
 	FUniformBufferRHIParamRef MeshParticleUniformBuffer;
+
+	FDynamicReadBuffer PrevTransformBuffer;
 };
 
 
@@ -187,7 +193,7 @@ public:
 
 	static bool ShouldCache(EShaderPlatform Platform, const class FMaterial* Material, const class FShaderType* ShaderType)
 	{
-		return (Platform == SP_OPENGL_ES2_ANDROID || Platform == SP_OPENGL_ES2_WEBGL) // Those are only platforms that might not support hardware instancing
+		return (GetMaxSupportedFeatureLevel(Platform) <= ERHIFeatureLevel::ES2)
 				&& FMeshParticleVertexFactory::ShouldCache(Platform, Material, ShaderType);
 	}
 
diff --git a/Engine/Source/Runtime/Engine/Public/ParticleEmitterInstances.h b/Engine/Source/Runtime/Engine/Public/ParticleEmitterInstances.h
--- a/Engine/Source/Runtime/Engine/Public/ParticleEmitterInstances.h
+++ b/Engine/Source/Runtime/Engine/Public/ParticleEmitterInstances.h
@@ -846,6 +846,7 @@ struct ENGINE_API FParticleMeshEmitterInstance : public FParticleEmitterInstance
 	UParticleModuleTypeDataMesh* MeshTypeData;
 	bool MeshRotationActive;
 	int32 MeshRotationOffset;
+	int32 MeshMotionBlurOffset;
 
 	/** The materials to render this instance with.	*/
 	TArray<UMaterialInterface*> CurrentMaterials;
@@ -1178,6 +1179,39 @@ struct FParticleTrailsEmitterInstance_Base : public FParticleEmitterInstance
 		CheckIndices(TrailIndex);
 	}
 
+	template<typename TrailDataType> void GetTrailStart(const int32 TrailIdx, int32 &OutStartIndex, TrailDataType *&OutTrailData, FBaseParticle *&OutParticle)
+	{
+		if (TrailIdx != INDEX_NONE)
+		{
+			OutStartIndex = CurrentStartIndices[TrailIdx];
+			if (OutStartIndex != INDEX_NONE)
+			{
+				DECLARE_PARTICLE_PTR(CheckParticle, ParticleData + ParticleStride * OutStartIndex);
+				TrailDataType* CheckTrailData = ((TrailDataType*)((uint8*)CheckParticle + TypeDataOffset));
+				check(TRAIL_EMITTER_IS_START(CheckTrailData->Flags));
+				OutTrailData = CheckTrailData;
+				OutParticle = CheckParticle;
+			}
+		}
+	}
+
+	template<typename TrailDataType> void GetTrailEnd(const int32 TrailIdx, int32 &OutEndIndex, TrailDataType *&OutTrailData, FBaseParticle *&OutParticle)
+	{
+		if (TrailIdx != INDEX_NONE)
+		{
+			OutEndIndex = CurrentEndIndices[TrailIdx];
+			if (OutEndIndex != INDEX_NONE)
+			{
+				DECLARE_PARTICLE_PTR(CheckParticle, ParticleData + ParticleStride * OutEndIndex);
+				TrailDataType* CheckTrailData = ((TrailDataType*)((uint8*)CheckParticle + TypeDataOffset));
+				check(TRAIL_EMITTER_IS_END(CheckTrailData->Flags));
+				OutTrailData = CheckTrailData;
+				OutParticle = CheckParticle;
+			}
+		}
+	}
+
+
 	void SetEndIndex(int32 TrailIndex, int32 ParticleIndex)
 	{
 		CurrentEndIndices[TrailIndex] = ParticleIndex;
@@ -1210,23 +1244,43 @@ struct FParticleTrailsEmitterInstance_Base : public FParticleEmitterInstance
 
 	void CheckIndices(int32 TrailIdx)
 	{
-		//Disabled until bug with this optimization is found. Usage of the optimization is alredy disabled but these check were still being hit.
-
-// #if !(UE_BUILD_SHIPPING || UE_BUILD_TEST)
-// 		if (CurrentEndIndices[TrailIdx] != INDEX_NONE)
-// 		{
-// 			DECLARE_PARTICLE_PTR(EndParticle, ParticleData + ParticleStride * CurrentEndIndices[TrailIdx]);
-// 			FRibbonTypeDataPayload* EndTrailData = ((FRibbonTypeDataPayload*)((uint8*)EndParticle + TypeDataOffset));
-// 			check(TRAIL_EMITTER_IS_END(EndTrailData->Flags));
-// 		}
-// 
-// 		if (CurrentStartIndices[TrailIdx] != INDEX_NONE)
-// 		{
-// 			DECLARE_PARTICLE_PTR(StartParticle, ParticleData + ParticleStride * CurrentStartIndices[TrailIdx]);
-// 			FRibbonTypeDataPayload* StartTrailData = ((FRibbonTypeDataPayload*)((uint8*)StartParticle + TypeDataOffset));
-// 			check(TRAIL_EMITTER_IS_START(StartTrailData->Flags));
-// 		}
-// #endif
+#if !(UE_BUILD_SHIPPING || UE_BUILD_TEST)
+ 		if (CurrentEndIndices[TrailIdx] != INDEX_NONE)
+ 		{
+ 			DECLARE_PARTICLE_PTR(EndParticle, ParticleData + ParticleStride * CurrentEndIndices[TrailIdx]);
+ 			FRibbonTypeDataPayload* EndTrailData = ((FRibbonTypeDataPayload*)((uint8*)EndParticle + TypeDataOffset));
+ 			check(TRAIL_EMITTER_IS_END(EndTrailData->Flags));
+ 		}
+ 
+ 		if (CurrentStartIndices[TrailIdx] != INDEX_NONE)
+ 		{
+ 			DECLARE_PARTICLE_PTR(StartParticle, ParticleData + ParticleStride * CurrentStartIndices[TrailIdx]);
+ 			FRibbonTypeDataPayload* StartTrailData = ((FRibbonTypeDataPayload*)((uint8*)StartParticle + TypeDataOffset));
+ 			check(TRAIL_EMITTER_IS_START(StartTrailData->Flags));
+ 		}
+#endif
+	}
+
+	void CheckAllIndices()
+	{
+#if !(UE_BUILD_SHIPPING || UE_BUILD_TEST)
+		for (uint32 TrailIdx = 0; TrailIdx < 128; TrailIdx++)
+		{
+			if (CurrentEndIndices[TrailIdx] != INDEX_NONE)
+			{
+				DECLARE_PARTICLE_PTR(EndParticle, ParticleData + ParticleStride * CurrentEndIndices[TrailIdx]);
+				FRibbonTypeDataPayload* EndTrailData = ((FRibbonTypeDataPayload*)((uint8*)EndParticle + TypeDataOffset));
+				check(TRAIL_EMITTER_IS_END(EndTrailData->Flags));
+			}
+
+			if (CurrentStartIndices[TrailIdx] != INDEX_NONE)
+			{
+				DECLARE_PARTICLE_PTR(StartParticle, ParticleData + ParticleStride * CurrentStartIndices[TrailIdx]);
+				FRibbonTypeDataPayload* StartTrailData = ((FRibbonTypeDataPayload*)((uint8*)StartParticle + TypeDataOffset));
+				check(TRAIL_EMITTER_IS_START(StartTrailData->Flags));
+			}
+		}
+#endif
 	}
 
 
diff --git a/Engine/Source/Runtime/Engine/Public/ParticleHelper.h b/Engine/Source/Runtime/Engine/Public/ParticleHelper.h
--- a/Engine/Source/Runtime/Engine/Public/ParticleHelper.h
+++ b/Engine/Source/Runtime/Engine/Public/ParticleHelper.h
@@ -387,7 +387,7 @@ struct FMeshParticleInstanceVertex
 	FVector4 Transform[3];
 
 	/** The velocity of the particle, XYZ: direction, W: speed. */
-	FVector Velocity;
+	FVector4 Velocity;
 
 	/** The sub-image texture offsets for the particle. */
 	int16 SubUVParams[4];
@@ -405,6 +405,13 @@ struct FMeshParticleInstanceVertexDynamicParameter
 	float DynamicValue[4];
 };
 
+struct FMeshParticleInstanceVertexPrevTransform
+{
+	FVector4 PrevTransform0;
+	FVector4 PrevTransform1;
+	FVector4 PrevTransform2;
+};
+
 //
 //  Trail emitter flags and macros
 //
@@ -762,6 +769,16 @@ struct FMeshRotationPayloadData
 	FVector  RotationRateBase;
 };
 
+struct FMeshMotionBlurPayloadData
+{
+	FVector BaseParticlePrevVelocity;
+	FVector BaseParticlePrevSize;
+	FVector PayloadPrevRotation;
+	FVector PayloadPrevOrbitOffset;
+	float   BaseParticlePrevRotation;
+	float   PayloadPrevCameraOffset;
+};
+
 /** ModuleLocationEmitter instance payload							*/
 struct FLocationEmitterInstancePayload
 {
@@ -1743,15 +1760,16 @@ struct FDynamicSpriteEmitterData : public FDynamicSpriteEmitterDataBase
 struct FDynamicMeshEmitterReplayData
 	: public FDynamicSpriteEmitterReplayDataBase
 {
-	int32					SubUVInterpMethod;
-	int32					SubUVDataOffset;
-	int32					SubImages_Horizontal;
-	int32					SubImages_Vertical;
-	bool				bScaleUV;
-	int32					MeshRotationOffset;
-	uint8				MeshAlignment;
-	bool				bMeshRotationActive;
-	FVector				LockedAxis;	
+	int32	SubUVInterpMethod;
+	int32	SubUVDataOffset;
+	int32	SubImages_Horizontal;
+	int32	SubImages_Vertical;
+	bool	bScaleUV;
+	int32	MeshRotationOffset;
+	int32	MeshMotionBlurOffset;
+	uint8	MeshAlignment;
+	bool	bMeshRotationActive;
+	FVector	LockedAxis;	
 
 	/** Constructor */
 	FDynamicMeshEmitterReplayData() : 
@@ -1761,6 +1779,7 @@ struct FDynamicMeshEmitterReplayData
 		SubImages_Vertical( 0 ),
 		bScaleUV( false ),
 		MeshRotationOffset( 0 ),
+		MeshMotionBlurOffset( 0 ),
 		MeshAlignment( 0 ),
 		bMeshRotationActive( false ),
 		LockedAxis(1.0f, 0.0f, 0.0f)
@@ -1773,13 +1792,14 @@ struct FDynamicMeshEmitterReplayData
 	{
 		// Call parent implementation
 		FDynamicSpriteEmitterReplayDataBase::Serialize( Ar );
-
+		
 		Ar << SubUVInterpMethod;
 		Ar << SubUVDataOffset;
 		Ar << SubImages_Horizontal;
 		Ar << SubImages_Vertical;
 		Ar << bScaleUV;
 		Ar << MeshRotationOffset;
+		Ar << MeshMotionBlurOffset;
 		Ar << MeshAlignment;
 		Ar << bMeshRotationActive;
 		Ar << LockedAxis;
@@ -1829,25 +1849,38 @@ struct FDynamicMeshEmitterData : public FDynamicSpriteEmitterDataBase
 	 *
 	 *	@param	InstanceData            The memory to fill the vertex data into
 	 *	@param	DynamicParameterData    The memory to fill the vertex dynamic parameter data into
+	 *	@param	PrevTransformBuffer     The memory to fill the vertex prev transform data into. May be null
 	 *	@param	Proxy                   The scene proxy for the particle system that owns this emitter
 	 *	@param	View                    The scene view being rendered
 	 */
-	void GetInstanceData(void* InstanceData, void* DynamicParameterData, const FParticleSystemSceneProxy* Proxy, const FSceneView* View) const;
+	void GetInstanceData(void* InstanceData, void* DynamicParameterData, void* PrevTransformBuffer, const FParticleSystemSceneProxy* Proxy, const FSceneView* View) const;
 
 	/**
 	 *	Helper function for retrieving the particle transform.
 	 *
 	 *	@param	InParticle					The particle being processed
-	 *  @param	ParticleBase				The scene proxy for the particle system that owns this emitter
-	 *  @param  CameraPosition				The position of the camera
-	 *  @param	CameraFacingOpVector		The facing option for the camera.
-	 *  @param  PointToLockedAxis			A quaternion for locked axis rotation
 	 *  @param	Proxy					    The scene proxy for the particle system that owns this emitter
 	 *	@param	View						The scene view being rendered
 	 *	@param	OutTransform				The InstanceToWorld transform matrix for the particle
 	 */
-	void GetParticleTransform(const FBaseParticle& InParticle, const FVector& CameraPosition, const FVector& CameraFacingOpVector, 
-		const FQuat& PointToLockedAxis, const FParticleSystemSceneProxy* Proxy, const FSceneView* View, FMatrix& OutTransformMat) const;
+	void GetParticleTransform(const FBaseParticle& InParticle, const FParticleSystemSceneProxy* Proxy, const FSceneView* View, FMatrix& OutTransformMat) const;
+
+	void GetParticlePrevTransform(const FBaseParticle& InParticle, const FParticleSystemSceneProxy* Proxy, const FSceneView* View, FMatrix& OutTransformMat) const;
+
+	void CalculateParticleTransform(
+		const FMatrix& ProxyLocalToWorld,
+		const FVector& ParticleLocation,
+			  float    ParticleRotation,
+		const FVector& ParticleVelocity,
+		const FVector& ParticleSize,
+		const FVector& ParticlePayloadInitialOrientation,
+		const FVector& ParticlePayloadRotation,
+		const FVector& ParticlePayloadCameraOffset,
+		const FVector& ParticlePayloadOrbitOffset,
+		const FVector& ViewOrigin,
+		const FVector& ViewDirection,
+		FMatrix& OutTransformMat
+		) const;
 
 	/** Gathers simple lights for this emitter. */
 	virtual void GatherSimpleLights(const FParticleSystemSceneProxy* Proxy, const FSceneViewFamily& ViewFamily, FSimpleLightArray& OutParticleLights) const override;
@@ -1889,7 +1922,6 @@ struct FDynamicMeshEmitterData : public FDynamicSpriteEmitterDataBase
 		non-simulating 'replay' particle systems, this data may have come straight from disk! */
 	FDynamicMeshEmitterReplayData Source;
 
-
 	int32					LastFramePreRendered;
 
 	UStaticMesh*		StaticMesh;
diff --git a/Engine/Source/Runtime/Engine/Public/SceneView.h b/Engine/Source/Runtime/Engine/Public/SceneView.h
--- a/Engine/Source/Runtime/Engine/Public/SceneView.h
+++ b/Engine/Source/Runtime/Engine/Public/SceneView.h
@@ -11,6 +11,8 @@
 #include "SceneTypes.h"
 #include "ShaderParameters.h"
 #include "RendererInterface.h"
+#include "RHIStaticStates.h"
+#include "GlobalDistanceFieldParameters.h"
 
 class FSceneViewStateInterface;
 class FViewUniformShaderParameters;
@@ -365,6 +367,7 @@ enum EDebugViewShaderMode
 	DVSM_QuadComplexity,			// Show quad overdraw only.
 	DVSM_WantedMipsAccuracy,		// Accuraty of the wanted mips computed by the texture streamer.
 	DVSM_TexelFactorAccuracy,		// Accuraty of the texel factor computed on each mesh.
+	DVSM_TexCoordScaleAnalysis,		// To view the material texture coordinate scale used to sample each texture.
 	DVSM_MAX
 };
 
@@ -375,7 +378,7 @@ FORCEINLINE bool AllowDebugViewModeShader(EShaderPlatform Platform)
 }
 
 /** The view dependent uniform shader parameters associated with a view. */
-BEGIN_UNIFORM_BUFFER_STRUCT(FViewUniformShaderParameters,ENGINE_API)
+BEGIN_UNIFORM_BUFFER_STRUCT_WITH_CONSTRUCTOR(FViewUniformShaderParameters, ENGINE_API)
 	DECLARE_UNIFORM_BUFFER_STRUCT_MEMBER(FMatrix, TranslatedWorldToClip)
 	DECLARE_UNIFORM_BUFFER_STRUCT_MEMBER(FMatrix, WorldToClip)
 	DECLARE_UNIFORM_BUFFER_STRUCT_MEMBER(FMatrix, TranslatedWorldToView)
@@ -416,7 +419,7 @@ BEGIN_UNIFORM_BUFFER_STRUCT(FViewUniformShaderParameters,ENGINE_API)
 END_UNIFORM_BUFFER_STRUCT(FViewUniformShaderParameters)
 
 /** Copy of the view dependent uniform shader parameters associated with a view for instanced stereo. */
-BEGIN_UNIFORM_BUFFER_STRUCT(FInstancedViewUniformShaderParameters, ENGINE_API)
+BEGIN_UNIFORM_BUFFER_STRUCT_WITH_CONSTRUCTOR(FInstancedViewUniformShaderParameters, ENGINE_API)
 	DECLARE_UNIFORM_BUFFER_STRUCT_MEMBER(FMatrix, TranslatedWorldToClip)
 	DECLARE_UNIFORM_BUFFER_STRUCT_MEMBER(FMatrix, WorldToClip)
 	DECLARE_UNIFORM_BUFFER_STRUCT_MEMBER(FMatrix, TranslatedWorldToView)
@@ -532,8 +535,35 @@ BEGIN_UNIFORM_BUFFER_STRUCT_WITH_CONSTRUCTOR(FFrameUniformShaderParameters, ENGI
 	DECLARE_UNIFORM_BUFFER_STRUCT_MEMBER_ARRAY(FVector4, SkyIrradianceEnvironmentMap, [7])
 	DECLARE_UNIFORM_BUFFER_STRUCT_MEMBER(float, MobilePreviewMode)
 	DECLARE_UNIFORM_BUFFER_STRUCT_MEMBER(float, HMDEyePaddingOffset)
+
+	DECLARE_UNIFORM_BUFFER_STRUCT_MEMBER_ARRAY(FVector4, GlobalVolumeCenterAndExtent_UB, [GMaxGlobalDistanceFieldClipmaps])
+	DECLARE_UNIFORM_BUFFER_STRUCT_MEMBER_ARRAY(FVector4, GlobalVolumeWorldToUVAddAndMul_UB, [GMaxGlobalDistanceFieldClipmaps])
+	DECLARE_UNIFORM_BUFFER_STRUCT_MEMBER(float, GlobalVolumeDimension_UB)
+	DECLARE_UNIFORM_BUFFER_STRUCT_MEMBER(float, GlobalVolumeTexelSize_UB)
+	DECLARE_UNIFORM_BUFFER_STRUCT_MEMBER(float, MaxGlobalDistance_UB)
+	DECLARE_UNIFORM_BUFFER_STRUCT_MEMBER_TEXTURE(Texture3D, GlobalDistanceFieldTexture0_UB)
+	DECLARE_UNIFORM_BUFFER_STRUCT_MEMBER_SAMPLER(SamplerState, GlobalDistanceFieldSampler0_UB)
+	DECLARE_UNIFORM_BUFFER_STRUCT_MEMBER_TEXTURE(Texture3D, GlobalDistanceFieldTexture1_UB)
+	DECLARE_UNIFORM_BUFFER_STRUCT_MEMBER_SAMPLER(SamplerState, GlobalDistanceFieldSampler1_UB)
+	DECLARE_UNIFORM_BUFFER_STRUCT_MEMBER_TEXTURE(Texture3D, GlobalDistanceFieldTexture2_UB)
+	DECLARE_UNIFORM_BUFFER_STRUCT_MEMBER_SAMPLER(SamplerState, GlobalDistanceFieldSampler2_UB)
+	DECLARE_UNIFORM_BUFFER_STRUCT_MEMBER_TEXTURE(Texture3D, GlobalDistanceFieldTexture3_UB)
+	DECLARE_UNIFORM_BUFFER_STRUCT_MEMBER_SAMPLER(SamplerState, GlobalDistanceFieldSampler3_UB)
+
+
 	DECLARE_UNIFORM_BUFFER_STRUCT_MEMBER_TEXTURE(Texture2D, DirectionalLightShadowTexture)
 	DECLARE_UNIFORM_BUFFER_STRUCT_MEMBER_SAMPLER(SamplerState, DirectionalLightShadowSampler)
+	DECLARE_UNIFORM_BUFFER_STRUCT_MEMBER_TEXTURE(Texture2D, AtmosphereTransmittanceTexture_UB)
+	DECLARE_UNIFORM_BUFFER_STRUCT_MEMBER_SAMPLER(SamplerState, AtmosphereTransmittanceTextureSampler_UB)
+	DECLARE_UNIFORM_BUFFER_STRUCT_MEMBER_TEXTURE(Texture2D, AtmosphereIrradianceTexture_UB)
+	DECLARE_UNIFORM_BUFFER_STRUCT_MEMBER_SAMPLER(SamplerState, AtmosphereIrradianceTextureSampler_UB)
+	DECLARE_UNIFORM_BUFFER_STRUCT_MEMBER_TEXTURE(Texture3D, AtmosphereInscatterTexture_UB)
+	DECLARE_UNIFORM_BUFFER_STRUCT_MEMBER_SAMPLER(SamplerState, AtmosphereInscatterTextureSampler_UB)
+	DECLARE_UNIFORM_BUFFER_STRUCT_MEMBER_TEXTURE(Texture2D, PerlinNoiseGradientTexture)
+	DECLARE_UNIFORM_BUFFER_STRUCT_MEMBER_SAMPLER(SamplerState, PerlinNoiseGradientTextureSampler)
+	DECLARE_UNIFORM_BUFFER_STRUCT_MEMBER_TEXTURE(Texture3D, PerlinNoise3DTexture)
+	DECLARE_UNIFORM_BUFFER_STRUCT_MEMBER_SAMPLER(SamplerState, PerlinNoise3DTextureSampler)
+
 END_UNIFORM_BUFFER_STRUCT(FFrameUniformShaderParameters)
 
 BEGIN_UNIFORM_BUFFER_STRUCT(FBuiltinSamplersParameters, ENGINE_API)
@@ -554,6 +584,7 @@ public:
 	virtual void ReleaseDynamicRHI() override;
 };
 
+#define USE_GBuiltinSamplersUniformBuffer (0)
 extern ENGINE_API TGlobalResource<FBuiltinSamplersUniformBuffer> GBuiltinSamplersUniformBuffer;	
 
 namespace EDrawDynamicFlags
@@ -1070,7 +1101,9 @@ public:
 	EShaderPlatform GetShaderPlatform() const { return GShaderPlatformForFeatureLevel[GetFeatureLevel()]; }
 
 #if !(UE_BUILD_SHIPPING || UE_BUILD_TEST)
-	EDebugViewShaderMode GetDebugViewShaderMode() const;
+	EDebugViewShaderMode DebugViewShaderMode;
+	FORCEINLINE EDebugViewShaderMode GetDebugViewShaderMode() const { return DebugViewShaderMode; }
+	EDebugViewShaderMode ChooseDebugViewShaderMode() const;
 #else
 	FORCEINLINE EDebugViewShaderMode GetDebugViewShaderMode() const { return DVSM_None; }
 #endif
diff --git a/Engine/Source/Runtime/Engine/Public/ShowFlags.h b/Engine/Source/Runtime/Engine/Public/ShowFlags.h
--- a/Engine/Source/Runtime/Engine/Public/ShowFlags.h
+++ b/Engine/Source/Runtime/Engine/Public/ShowFlags.h
@@ -289,8 +289,8 @@ private:
 		SetNavigation(false);
 		SetLightComplexity(false);
 		SetShaderComplexity(false);
-		SetQuadComplexity(false);
-		SetQuadOverhead(false);
+		SetQuadOverdraw(false);
+		SetShaderComplexityWithQuadOverdraw(false);
 		SetStationaryLightOverlap(false);
 		SetLightMapDensity(false);
 		SetVertexDensities(false);
@@ -345,6 +345,7 @@ private:
 		SetVisualizeBloom(false);
 		SetWantedMipsAccuracy(false);
 		SetTexelFactorAccuracy(false);
+		SetTexCoordScaleAccuracy(false);
 	}
 
 
diff --git a/Engine/Source/Runtime/Engine/Public/ShowFlagsValues.inl b/Engine/Source/Runtime/Engine/Public/ShowFlagsValues.inl
--- a/Engine/Source/Runtime/Engine/Public/ShowFlagsValues.inl
+++ b/Engine/Source/Runtime/Engine/Public/ShowFlagsValues.inl
@@ -295,12 +295,14 @@ SHOWFLAG_FIXED_IN_SHIPPING(0, LODColoration, SFG_Hidden, LOCTEXT("VisualizeLODCo
 /** Visualize HLOD Coloration */
 SHOWFLAG_FIXED_IN_SHIPPING(0, HLODColoration, SFG_Hidden, LOCTEXT("VisualizeHLODColoration", "Visualize HLOD Coloration"))
 /** Visualize screen quads */
-SHOWFLAG_FIXED_IN_SHIPPING(0, QuadComplexity, SFG_Hidden, LOCTEXT("QuadComplexitySF", "Quad Complexity"))
+SHOWFLAG_FIXED_IN_SHIPPING(0, QuadOverdraw, SFG_Hidden, LOCTEXT("QuadOverdrawSF", "Quad Overdraw"))
 /** Visualize the overhead of material quads */
-SHOWFLAG_FIXED_IN_SHIPPING(0, QuadOverhead, SFG_Hidden, LOCTEXT("QuadOverheadSF", "Quad Overhead"))
+SHOWFLAG_FIXED_IN_SHIPPING(0, ShaderComplexityWithQuadOverdraw, SFG_Hidden, LOCTEXT("ShaderComplexityWithQuadOverdraw", "Shader Complexity With Quad Overdraw"))
 /** Visualize the streaming texture scale in the mesh editors */
 SHOWFLAG_FIXED_IN_SHIPPING(0, WantedMipsAccuracy, SFG_Hidden, LOCTEXT("WantedMipsAccuracy", "Streaming Wanted Mips Accuracy"))
 /** Visualize how accurate is the texel factor is computed on each mesh */
 SHOWFLAG_FIXED_IN_SHIPPING(0, TexelFactorAccuracy, SFG_Hidden, LOCTEXT("TexelFactorAccuracy", "Texel Factor Accuracy"))
+/** Visualize the accuracy of the material texture coordinate scale factors */
+SHOWFLAG_FIXED_IN_SHIPPING(0, TexCoordScaleAccuracy, SFG_Hidden, LOCTEXT("TexCoordScaleAccuray", "Texture Coordinate Scale Accuracy"))
 #undef SHOWFLAG_ALWAYS_ACCESSIBLE
 #undef SHOWFLAG_FIXED_IN_SHIPPING
diff --git a/Engine/Source/Runtime/Engine/Public/SkeletalMeshTypes.h b/Engine/Source/Runtime/Engine/Public/SkeletalMeshTypes.h
--- a/Engine/Source/Runtime/Engine/Public/SkeletalMeshTypes.h
+++ b/Engine/Source/Runtime/Engine/Public/SkeletalMeshTypes.h
@@ -1488,7 +1488,6 @@ public:
 	* Rendering data.
 	*/
 	FMultiSizeIndexContainer	MultiSizeIndexContainer; 
-	uint32						Size;
 	uint32						NumVertices;
 	/** The number of unique texture coordinate sets in this lod */
 	uint32						NumTexCoords;
@@ -1533,8 +1532,9 @@ public:
 
 	/** Constructor (default) */
 	FStaticLODModel()
-	:	Size(0)
-	,	NumVertices(0)
+		: NumVertices(0)
+		, NumTexCoords(0)
+		, MaxImportVertex(-1)
 	{
 	}
 
@@ -1577,7 +1577,7 @@ public:
 	 * Initialize vertex buffers from skel mesh chunks.
 	 * @param BuildFlags See EVertexFlags.
 	 */
-	ENGINE_API void BuildVertexBuffers(uint32 VertexFlags);
+	ENGINE_API void BuildVertexBuffers(uint32 BuildFlags);
 
 	/** Utility function for returning total number of faces in this LOD. */
 	ENGINE_API int32 GetTotalFaces() const;
@@ -1585,7 +1585,10 @@ public:
 	/** Utility for finding the chunk that a particular vertex is in. */
 	ENGINE_API void GetChunkAndSkinType(int32 InVertIndex, int32& OutChunkIndex, int32& OutVertIndex, bool& bOutSoftVert, bool& bOutHasExtraBoneInfluences) const;
 
-	/** Sort the triangles with the specified sorting method */
+	/**
+	 * Sort the triangles with the specified sorting method
+	 * @param ETriangleSortOption NewTriangleSorting new sorting method
+	 */
 	ENGINE_API void SortTriangles( FVector SortCenter, bool bUseSortCenter, int32 SectionIndex, ETriangleSortOption NewTriangleSorting );
 
 	/**
diff --git a/Engine/Source/Runtime/Engine/Public/SkeletalRenderPublic.h b/Engine/Source/Runtime/Engine/Public/SkeletalRenderPublic.h
--- a/Engine/Source/Runtime/Engine/Public/SkeletalRenderPublic.h
+++ b/Engine/Source/Runtime/Engine/Public/SkeletalRenderPublic.h
@@ -175,9 +175,6 @@ public:
 	{
 		/** Hidden Material Section Flags for rendering - That is Material Index, not Section Index  */
 		TArray<bool>	HiddenMaterials;
-
-		FSkelMeshObjectLODInfo()
-		{}
 	};	
 
 	TArray<FSkelMeshObjectLODInfo> LODInfo;
@@ -203,14 +200,14 @@ public:
 	float WorkingMaxDistanceFactor;
 
 	/** This is set to true when we have sent our Mesh data to the rendering thread at least once as it needs to have have a datastructure created there for each MeshObject **/
-	bool           bHasBeenUpdatedAtLeastOnce;
+	bool bHasBeenUpdatedAtLeastOnce;
 
 #if WITH_EDITORONLY_DATA
 	/** Index of the chunk to preview... If set to -1, all chunks will be rendered */
-	int32				ChunkIndexPreview;
+	int32 ChunkIndexPreview;
 	
 	/** Index of the section to preview... If set to -1, all section will be rendered */
-	int32				SectionIndexPreview;
+	int32 SectionIndexPreview;
 #endif
 
 	/** returns the feature level this FSkeletalMeshObject was created with */
@@ -227,15 +224,15 @@ protected:
 	TArray<FSkeletalMeshLODInfo> SkeletalMeshLODInfo;
 
 	/** GPU Skin Cache Keys per chunk; -1 means not using GPU Skin Cache **/
-	int16	GPUSkinCacheKeys[MAX_GPUSKINCACHE_CHUNKS_PER_LOD];
+	int16 GPUSkinCacheKeys[MAX_GPUSKINCACHE_CHUNKS_PER_LOD];
 
 	/** Used to keep track of the first call to UpdateMinDesiredLODLevel each frame. */
-	uint32			LastFrameNumber;
+	uint32 LastFrameNumber;
 
 	/** Editor only. Used for visualizing drawing order in Animset Viewer. If < 1.0,
 	 * only the specified fraction of triangles will be rendered
 	 */
-	float			ProgressiveDrawingFraction;
+	float ProgressiveDrawingFraction;
 
 	/** Use the 2nd copy of indices for separate left/right sort order (when TRISORT_CustomLeftRight) 
 	 * Set manually by the AnimSetViewer when editing sort order, or based on viewing angle otherwise.
@@ -244,7 +241,7 @@ protected:
 	/** 
 	 *	If true, per-bone motion blur is enabled for this object. This includes is the system overwrites the skeletal mesh setting.
 	 */
-	bool			bUsePerBoneMotionBlur;
+	bool bUsePerBoneMotionBlur;
 
 	/** Used for dynamic stats */
 	TStatId StatId;
diff --git a/Engine/Source/Runtime/Engine/Public/StaticMeshResources.h b/Engine/Source/Runtime/Engine/Public/StaticMeshResources.h
--- a/Engine/Source/Runtime/Engine/Public/StaticMeshResources.h
+++ b/Engine/Source/Runtime/Engine/Public/StaticMeshResources.h
@@ -580,6 +580,9 @@ struct FStaticMeshLODResources
 
 	uint32 DepthOnlyNumTriangles;
 
+	struct FSplineMeshVertexFactory* SplineVertexFactory;
+	struct FSplineMeshVertexFactory* SplineVertexFactoryOverrideColorVertexBuffer;
+
 	/** Default constructor. */
 	FStaticMeshLODResources();
 
diff --git a/Engine/Source/Runtime/NullDrv/Public/NullRHI.h b/Engine/Source/Runtime/NullDrv/Public/NullRHI.h
--- a/Engine/Source/Runtime/NullDrv/Public/NullRHI.h
+++ b/Engine/Source/Runtime/NullDrv/Public/NullRHI.h
@@ -202,6 +202,10 @@ public:
 		return new FRHIShaderResourceView(); 
 	}
 
+	virtual FShaderResourceViewRHIRef RHICreateShaderResourceView(FIndexBufferRHIParamRef Buffer) final override
+	{ 
+		return new FRHIShaderResourceView(); 
+	}
 
 	virtual void RHIClearUAV(FUnorderedAccessViewRHIParamRef UnorderedAccessViewRHI, const uint32* Values) final override
 	{
diff --git a/Engine/Source/Runtime/OpenGLDrv/Private/OpenGLUAV.cpp b/Engine/Source/Runtime/OpenGLDrv/Private/OpenGLUAV.cpp
--- a/Engine/Source/Runtime/OpenGLDrv/Private/OpenGLUAV.cpp
+++ b/Engine/Source/Runtime/OpenGLDrv/Private/OpenGLUAV.cpp
@@ -38,6 +38,13 @@ FShaderResourceViewRHIRef FOpenGLDynamicRHI::RHICreateShaderResourceView(FVertex
 	return Result;
 }
 
+FShaderResourceViewRHIRef FOpenGLDynamicRHI::RHICreateShaderResourceView(FIndexBufferRHIParamRef BufferRHI)
+{
+	UE_LOG(LogRHI, Fatal, TEXT("OpenGL RHI doesn't support RHICreateShaderResourceView with FIndexBufferRHIParamRef yet!"));
+	
+	return FShaderResourceViewRHIRef();
+}
+
 FOpenGLShaderResourceView::~FOpenGLShaderResourceView()
 {
 	FShaderCache::RemoveSRV(this);
diff --git a/Engine/Source/Runtime/OpenGLDrv/Public/OpenGLDrv.h b/Engine/Source/Runtime/OpenGLDrv/Public/OpenGLDrv.h
--- a/Engine/Source/Runtime/OpenGLDrv/Public/OpenGLDrv.h
+++ b/Engine/Source/Runtime/OpenGLDrv/Public/OpenGLDrv.h
@@ -339,6 +339,7 @@ public:
 	virtual FUnorderedAccessViewRHIRef RHICreateUnorderedAccessView(FVertexBufferRHIParamRef VertexBuffer, uint8 Format) final override;
 	virtual FShaderResourceViewRHIRef RHICreateShaderResourceView(FStructuredBufferRHIParamRef StructuredBuffer) final override;
 	virtual FShaderResourceViewRHIRef RHICreateShaderResourceView(FVertexBufferRHIParamRef VertexBuffer, uint32 Stride, uint8 Format) final override;
+	virtual FShaderResourceViewRHIRef RHICreateShaderResourceView(FIndexBufferRHIParamRef Buffer) final override;
 	virtual uint64 RHICalcTexture2DPlatformSize(uint32 SizeX, uint32 SizeY, uint8 Format, uint32 NumMips, uint32 NumSamples, uint32 Flags, uint32& OutAlign) final override;
 	virtual uint64 RHICalcTexture3DPlatformSize(uint32 SizeX, uint32 SizeY, uint32 SizeZ, uint8 Format, uint32 NumMips, uint32 Flags, uint32& OutAlign) final override;
 	virtual uint64 RHICalcTextureCubePlatformSize(uint32 Size, uint8 Format, uint32 NumMips, uint32 Flags, uint32& OutAlign) final override;
diff --git a/Engine/Source/Runtime/RHI/Private/RHICommandList.cpp b/Engine/Source/Runtime/RHI/Private/RHICommandList.cpp
--- a/Engine/Source/Runtime/RHI/Private/RHICommandList.cpp
+++ b/Engine/Source/Runtime/RHI/Private/RHICommandList.cpp
@@ -16,6 +16,7 @@ DECLARE_DWORD_COUNTER_STAT(TEXT("Immed. Command count"), STAT_ImmedCmdListCount,
 #include "RHICommandListCommandExecutes.inl"
 #endif
 
+int32 FRHICommandListBase::StateCacheEnabled = 1;
 
 static TAutoConsoleVariable<int32> CVarRHICmdBypass(
 	TEXT("r.RHICmdBypass"),
@@ -99,6 +100,12 @@ static TAutoConsoleVariable<int32> CVarRHICmdMinCmdlistSizeForParallelTranslate(
 	32,
 	TEXT("In kilobytes. Cmdlists are merged into one parallel translate until we have at least this much memory to process. For a given pass, we won't do more translates than we have task threads. Only relevant if r.RHICmdBalanceTranslatesAfterTasks is on."));
 
+static FAutoConsoleVariableRef CVarRHICmdListStateCache(
+	TEXT("r.RHICmdStateCacheEnable"),
+	FRHICommandListBase::StateCacheEnabled,
+	TEXT("If > 0, then enable a minor state cache on the from of cmdlist recording.")
+	);
+
 
 RHI_API bool GEnableAsyncCompute = true;
 RHI_API FRHICommandListExecutor GRHICommandList;
@@ -642,6 +649,12 @@ struct FRHIAsyncComputeSubmitList : public FRHICommand<FRHIAsyncComputeSubmitLis
 
 void FRHICommandListImmediate::QueueAsyncCompute(FRHIAsyncComputeCommandList& RHIComputeCmdList)
 {
+	if (Bypass())
+	{
+		SCOPE_CYCLE_COUNTER(STAT_AsyncComputeExecute);
+		delete &RHIComputeCmdList;
+		return;
+	}
 	new (AllocCommand<FRHIAsyncComputeSubmitList>()) FRHIAsyncComputeSubmitList(&RHIComputeCmdList);
 }
 	
@@ -711,6 +724,7 @@ void FRHICommandListBase::Reset()
 		RenderThreadContexts[Index] = nullptr;
 	}
 	ExecuteStat = TStatId();
+	FlushStateCache();
 }
 
 
@@ -982,7 +996,9 @@ public:
 
 void FRHICommandListBase::QueueParallelAsyncCommandListSubmit(FGraphEventRef* AnyThreadCompletionEvents, bool bIsPrepass, FRHICommandList** CmdLists, int32* NumDrawsIfKnown, int32 Num, int32 MinDrawsPerTranslate, bool bSpewMerge)
 {
+	QUICK_SCOPE_CYCLE_COUNTER(STAT_FRHICommandListBase_QueueParallelAsyncCommandListSubmit);
 	check(IsInRenderingThread() && IsImmediate() && Num);
+	FlushStateCache();
 	if (GRHIThread)
 	{
 		FRHICommandListExecutor::GetImmediateCommandList().ImmediateFlush(EImmediateFlushType::DispatchToRHIThread); // we should start on the stuff before this async list
@@ -1165,6 +1181,8 @@ void FRHICommandListBase::QueueParallelAsyncCommandListSubmit(FGraphEventRef* An
 void FRHICommandListBase::QueueAsyncCommandListSubmit(FGraphEventRef& AnyThreadCompletionEvent, class FRHICommandList* CmdList)
 {
 	check(IsInRenderingThread() && IsImmediate());
+	FlushStateCache();
+
 	if (GRHIThread)
 	{
 		FRHICommandListExecutor::GetImmediateCommandList().ImmediateFlush(EImmediateFlushType::DispatchToRHIThread); // we should start on the stuff before this async list
@@ -1231,6 +1249,7 @@ struct FRHICommandWaitForAndSubmitRTSubList : public FRHICommand<FRHICommandWait
 void FRHICommandListBase::QueueRenderThreadCommandListSubmit(FGraphEventRef& RenderThreadCompletionEvent, class FRHICommandList* CmdList)
 {
 	check(!IsInRHIThread());
+	FlushStateCache();
 	if (RenderThreadCompletionEvent.GetReference())
 	{
 		check(!IsInActualRenderingThread() && !IsInGameThread() && !IsImmediate());
@@ -1257,6 +1276,7 @@ struct FRHICommandSubmitSubList : public FRHICommand<FRHICommandSubmitSubList>
 
 void FRHICommandListBase::QueueCommandListSubmit(class FRHICommandList* CmdList)
 {
+	FlushStateCache();
 	new (AllocCommand<FRHICommandSubmitSubList>()) FRHICommandSubmitSubList(CmdList);
 }
 
@@ -1708,6 +1728,12 @@ FShaderResourceViewRHIRef FDynamicRHI::CreateShaderResourceView_RenderThread(cla
 	return GDynamicRHI->RHICreateShaderResourceView(VertexBuffer, Stride, Format);
 }
 
+FShaderResourceViewRHIRef FDynamicRHI::CreateShaderResourceView_RenderThread(class FRHICommandListImmediate& RHICmdList, FIndexBufferRHIParamRef Buffer)
+{
+	FScopedRHIThreadStaller StallRHIThread(RHICmdList);
+	return GDynamicRHI->RHICreateShaderResourceView(Buffer);
+}
+
 struct FRHICommandUpdateVertexBuffer : public FRHICommand<FRHICommandUpdateVertexBuffer>
 {
 	FVertexBufferRHIParamRef VertexBuffer;
@@ -2037,6 +2063,12 @@ FShaderResourceViewRHIRef FDynamicRHI::RHICreateShaderResourceView_RenderThread(
 	return GDynamicRHI->RHICreateShaderResourceView(VertexBuffer, Stride, Format);
 }
 
+FShaderResourceViewRHIRef FDynamicRHI::RHICreateShaderResourceView_RenderThread(class FRHICommandListImmediate& RHICmdList, FIndexBufferRHIParamRef Buffer)
+{
+	FScopedRHIThreadStaller StallRHIThread(RHICmdList);
+	return GDynamicRHI->RHICreateShaderResourceView(Buffer);
+}
+
 FShaderResourceViewRHIRef FDynamicRHI::RHICreateShaderResourceView_RenderThread(class FRHICommandListImmediate& RHICmdList, FStructuredBufferRHIParamRef StructuredBuffer)
 {
 	FScopedRHIThreadStaller StallRHIThread(RHICmdList);
diff --git a/Engine/Source/Runtime/RHI/Public/DynamicRHI.h b/Engine/Source/Runtime/RHI/Public/DynamicRHI.h
--- a/Engine/Source/Runtime/RHI/Public/DynamicRHI.h
+++ b/Engine/Source/Runtime/RHI/Public/DynamicRHI.h
@@ -617,6 +617,10 @@ public:
 	// FlushType: Wait RHI Thread
 	virtual FShaderResourceViewRHIRef RHICreateShaderResourceView(FVertexBufferRHIParamRef VertexBuffer, uint32 Stride, uint8 Format) = 0;
 
+	/** Creates a shader resource view of the given index buffer. */
+	// FlushType: Wait RHI Thread
+	virtual FShaderResourceViewRHIRef RHICreateShaderResourceView(FIndexBufferRHIParamRef Buffer) = 0;
+
 	/**
 	 * Computes the total size of a 2D texture with the specified parameters.
 	 *
@@ -1107,6 +1111,7 @@ public:
 
 	virtual FVertexBufferRHIRef CreateVertexBuffer_RenderThread(class FRHICommandListImmediate& RHICmdList, uint32 Size, uint32 InUsage, FRHIResourceCreateInfo& CreateInfo);
 	virtual FShaderResourceViewRHIRef CreateShaderResourceView_RenderThread(class FRHICommandListImmediate& RHICmdList, FVertexBufferRHIParamRef VertexBuffer, uint32 Stride, uint8 Format);
+	virtual FShaderResourceViewRHIRef CreateShaderResourceView_RenderThread(class FRHICommandListImmediate& RHICmdList, FIndexBufferRHIParamRef Buffer);
 	virtual void* LockVertexBuffer_RenderThread(class FRHICommandListImmediate& RHICmdList, FVertexBufferRHIParamRef VertexBuffer, uint32 Offset, uint32 SizeRHI, EResourceLockMode LockMode);
 	virtual void UnlockVertexBuffer_RenderThread(class FRHICommandListImmediate& RHICmdList, FVertexBufferRHIParamRef VertexBuffer);
 	virtual FTexture2DRHIRef AsyncReallocateTexture2D_RenderThread(class FRHICommandListImmediate& RHICmdList, FTexture2DRHIParamRef Texture2D, int32 NewMipCount, int32 NewSizeX, int32 NewSizeY, FThreadSafeCounter* RequestStatus);
@@ -1132,6 +1137,7 @@ public:
 	virtual FShaderResourceViewRHIRef RHICreateShaderResourceView_RenderThread(class FRHICommandListImmediate& RHICmdList, FTexture2DArrayRHIParamRef Texture2DArrayRHI, uint8 MipLevel);
 	virtual FShaderResourceViewRHIRef RHICreateShaderResourceView_RenderThread(class FRHICommandListImmediate& RHICmdList, FTextureCubeRHIParamRef TextureCubeRHI, uint8 MipLevel);
 	virtual FShaderResourceViewRHIRef RHICreateShaderResourceView_RenderThread(class FRHICommandListImmediate& RHICmdList, FVertexBufferRHIParamRef VertexBuffer, uint32 Stride, uint8 Format);
+	virtual FShaderResourceViewRHIRef RHICreateShaderResourceView_RenderThread(class FRHICommandListImmediate& RHICmdList, FIndexBufferRHIParamRef Buffer);
 	virtual FShaderResourceViewRHIRef RHICreateShaderResourceView_RenderThread(class FRHICommandListImmediate& RHICmdList, FStructuredBufferRHIParamRef StructuredBuffer);
 	virtual FTextureCubeRHIRef RHICreateTextureCube_RenderThread(class FRHICommandListImmediate& RHICmdList, uint32 Size, uint8 Format, uint32 NumMips, uint32 Flags, FRHIResourceCreateInfo& CreateInfo);
 	virtual FTextureCubeRHIRef RHICreateTextureCubeArray_RenderThread(class FRHICommandListImmediate& RHICmdList, uint32 Size, uint32 ArraySize, uint8 Format, uint32 NumMips, uint32 Flags, FRHIResourceCreateInfo& CreateInfo);
diff --git a/Engine/Source/Runtime/RHI/Public/RHICommandList.h b/Engine/Source/Runtime/RHI/Public/RHICommandList.h
--- a/Engine/Source/Runtime/RHI/Public/RHICommandList.h
+++ b/Engine/Source/Runtime/RHI/Public/RHICommandList.h
@@ -259,8 +259,18 @@ public:
 		Num
 	};
 	void *RenderThreadContexts[(int32)ERenderThreadContext::Num];
-private:
+	static int32 StateCacheEnabled;
+protected:
+	struct FRHICommandSetRasterizerState* CachedRasterizerState;
+	struct FRHICommandSetDepthStencilState* CachedDepthStencilState;
 public:
+	void FORCEINLINE FlushStateCache()
+	{
+		CachedRasterizerState = nullptr;
+		CachedDepthStencilState = nullptr;
+	}
+
+
 	void CopyRenderThreadContexts(const FRHICommandListBase& ParentCommandList)
 	{
 		for (int32 Index = 0; ERenderThreadContext(Index) < ERenderThreadContext::Num; Index++)
@@ -1578,7 +1588,11 @@ public:
 			CMD_CONTEXT(SetRasterizerState)(State);
 			return;
 		}
-		new (AllocCommand<FRHICommandSetRasterizerState>()) FRHICommandSetRasterizerState(State);
+		if (StateCacheEnabled && CachedRasterizerState && CachedRasterizerState->State == State)
+		{
+			return;
+		}
+		CachedRasterizerState = new(AllocCommand<FRHICommandSetRasterizerState>()) FRHICommandSetRasterizerState(State);
 	}
 
 	FORCEINLINE_DEBUGGABLE void SetBlendState(FBlendStateRHIParamRef State, const FLinearColor& BlendFactor = FLinearColor::White)
@@ -1628,7 +1642,11 @@ public:
 			CMD_CONTEXT(SetDepthStencilState)(NewStateRHI, StencilRef);
 			return;
 		}
-		new (AllocCommand<FRHICommandSetDepthStencilState>()) FRHICommandSetDepthStencilState(NewStateRHI, StencilRef);
+		if (StateCacheEnabled && CachedDepthStencilState && CachedDepthStencilState->State == NewStateRHI && CachedDepthStencilState->StencilRef == StencilRef)
+		{
+			return;
+		}
+		CachedDepthStencilState = new(AllocCommand<FRHICommandSetDepthStencilState>()) FRHICommandSetDepthStencilState(NewStateRHI, StencilRef);
 	}
 
 	FORCEINLINE_DEBUGGABLE void SetViewport(uint32 MinX, uint32 MinY, float MinZ, uint32 MaxX, uint32 MaxY, float MaxZ)
diff --git a/Engine/Source/Runtime/RHI/Public/RHIResources.h b/Engine/Source/Runtime/RHI/Public/RHIResources.h
--- a/Engine/Source/Runtime/RHI/Public/RHIResources.h
+++ b/Engine/Source/Runtime/RHI/Public/RHIResources.h
@@ -27,13 +27,13 @@ public:
 			check(NumRefs.GetValue() == 0 && (CurrentlyDeleting == this || bDoNotDeferDelete || Bypass())); // this should not have any outstanding refs
 		}
 	}
-	uint32 AddRef() const
+	FORCEINLINE_DEBUGGABLE uint32 AddRef() const
 	{
 		int32 NewValue = NumRefs.Increment();
-		check(NewValue > 0); 
+		checkSlow(NewValue > 0); 
 		return uint32(NewValue);
 	}
-	uint32 Release() const
+	FORCEINLINE_DEBUGGABLE uint32 Release() const
 	{
 		int32 NewValue = NumRefs.Decrement();
 		if (NewValue == 0)
@@ -50,13 +50,13 @@ public:
 				}
 			}
 		}
-		check(NewValue >= 0);
+		checkSlow(NewValue >= 0);
 		return uint32(NewValue);
 	}
-	uint32 GetRefCount() const
+	FORCEINLINE_DEBUGGABLE uint32 GetRefCount() const
 	{
 		int32 CurrentValue = NumRefs.GetValue();
-		check(CurrentValue >= 0); 
+		checkSlow(CurrentValue >= 0); 
 		return uint32(CurrentValue);
 	}
 	void DoNoDeferDelete()
@@ -421,6 +421,16 @@ public:
 		return nullptr;
 	}
 
+	/**
+	* Returns access to the platform-specific RHI texture baseclass.  This is designed to provide the RHI with fast access to its base classes in the face of multiple inheritance.
+	* @return	The pointer to the platform-specific RHI texture baseclass or NULL if it not initialized or not supported for this RHI
+	*/
+	virtual void* GetTextureBaseRHI()
+	{
+		// Override this in derived classes to expose access to the native texture resource
+		return nullptr;
+	}
+
 	/** @return The number of mip-maps in the texture. */
 	uint32 GetNumMips() const { return NumMips; }
 
