Copying //UE4/Dev-Framework to Dev-Main (//UE4/Dev-Main) @ 2868852

==========================
MAJOR FEATURES + CHANGES
==========================

Change 2821607 on 2016/01/08 by [REDACTED_PERSON]

	Added a way to limit amount of information logged by vlog by discarding logs from classes from outside of class whitelist #UE4

	This feature was followed by refactoring of functions taking FVisualLogEntry pointers to use references instead.

Change 2828384 on 2016/01/14 by [REDACTED_PERSON]

	Back out of visual log refactor done as part of CL#2821607 #UE4

Change 2855722 on 2016/02/04 by [REDACTED_PERSON]

	#ue4 - Expose PrimitiveComponent GetCollisionEnabled(), IsCollisionEnabled(), IsQueryCollisionEnabled(), IsPhysicsCollisionEnabled() functions to blueprints.

Change 2855737 on 2016/02/04 by [REDACTED_PERSON]

	Add test AssetUserData C++ class to QAGame

Change 2855739 on 2016/02/04 by [REDACTED_PERSON]

	PR #1858 : Add abstract and editinlinenew to AssetUserData class, so derived classes can be added to assets via details panel
	https://github.com/EpicGames/UnrealEngine/pull/1858
	#github 1858

Change 2855842 on 2016/02/04 by [REDACTED_PERSON]

	UE-23968 : Add clamping to Friction, Restitution, Density and RaiseMassToPower properties of PhysicalMaterial

Change 2855843 on 2016/02/04 by [REDACTED_PERSON]

	PR #1984 : Add 'AssetUserData' support to AnimationAsset/Texture
	https://github.com/EpicGames/UnrealEngine/pull/1984
	#github 1984

Change 2855844 on 2016/02/04 by [REDACTED_PERSON]

	UE-23176 Fix incorrect clamping in GenerateKDopAsSimpleCollision() which could cause degenerate hulls
	Also don't 'nudge' kdops if there are other primitives present, will offset from mesh we were fitting around

Change 2855850 on 2016/02/04 by [REDACTED_PERSON]

	Update required LOD bones when associating a cloth chunk that uses bones that aren't weighted to that LOD

Change 2856409 on 2016/02/04 by [REDACTED_PERSON]

	#ue4 - Change LogSpawn warnings to normal logs for cases that are common and valid during gameplay (failed due to collision at the target location, or failed because it is destroyed in the process of spawning, ie projectile spawned in a target).

Change 2857018 on 2016/02/05 by [REDACTED_PERSON]

	Fixes for various anim graph related subgraph issues/crashes

Change 2857193 on 2016/02/05 by [REDACTED_PERSON]

	extended crowd agent interface to handle avoidance groups
	#ue4 UE-24823
	#1896

Change 2857200 on 2016/02/05 by [REDACTED_PERSON]

	fixed memory leak in CrowdManager
	#ue UE-26516
	#2026

Change 2857417 on 2016/02/05 by [REDACTED_PERSON]

	Add keyword for animation functions

Change 2858854 on 2016/02/08 by [REDACTED_PERSON]

	Fixed pose evaluator customizations being processed at the wrong time which caused properties to always be re-added even if removed from the details panel, causing duplicated entries.

Change 2858855 on 2016/02/08 by [REDACTED_PERSON]

	Fixed blendspace players only reporting normalized time when using a time getter in an anim graph transition.

Change 2859159 on 2016/02/08 by [REDACTED_PERSON]

	UE-25586 Fix for reporting audio memory usage

	- Realtime decompressed sounds that have CachedRealtimeFirstBuffers need to report the compressed asset size

Change 2859192 on 2016/02/08 by [REDACTED_PERSON]

	Removed check disallowing cooked additive animations from using AnimationRelative and AnimationScale retargeting modes (on the root bone or editor).
	Skip AnimationRelative retargeting for baked additive animations since it gets cancelled out during the additive creation process.

Change 2859238 on 2016/02/08 by [REDACTED_PERSON]

	Git pull request #1895

Change 2859239 on 2016/02/08 by [REDACTED_PERSON]

	Git pull request #1813

Change 2859453 on 2016/02/08 by [REDACTED_PERSON]

diff --git a/Engine/Config/BaseEngine.ini b/Engine/Config/BaseEngine.ini
--- a/Engine/Config/BaseEngine.ini
+++ b/Engine/Config/BaseEngine.ini
@@ -299,6 +299,14 @@ bUseStreamingPause=false
 
 +TaggedPropertyRedirects=(ClassName="BodyInstance",OldPropertyName="MassInKg",NewPropertyName="MassInKgOverride")
 
+; Redirects for EQSettings 
++TaggedPropertyRedirects=(ClassName="AudioEQEffect",OldPropertyName="LFFrequency",NewPropertyName="FrequencyCenter0")
++TaggedPropertyRedirects=(ClassName="AudioEQEffect",OldPropertyName="LFGain",NewPropertyName="Gain0")
++TaggedPropertyRedirects=(ClassName="AudioEQEffect",OldPropertyName="MFCutoffFrequency",NewPropertyName="FrequencyCenter1")
++TaggedPropertyRedirects=(ClassName="AudioEQEffect",OldPropertyName="MFGain",NewPropertyName="Gain1")
++TaggedPropertyRedirects=(ClassName="AudioEQEffect",OldPropertyName="MFBandwidth",NewPropertyName="Bandwidth1")
++TaggedPropertyRedirects=(ClassName="AudioEQEffect",OldPropertyName="HFFrequency",NewPropertyName="FrequencyCenter2")
++TaggedPropertyRedirects=(ClassName="AudioEQEffect",OldPropertyName="HFGain",NewPropertyName="Gain2")
 
 +K2FieldRedirects=(OldFieldName="SkeletalMeshComponent.SetAnimBlueprint",NewFieldName="SkeletalMeshComponent.SetAnimInstanceClass")
 +K2ParamRedirects=(NodeName="/Script/BlueprintGraph.K2Node_CallFunction", OldParamName="SkeletalMeshComponent.SetAnimInstanceClass.NewBlueprint", NewParamName="NewClass")
@@ -2545,3 +2553,10 @@ GameHUDBind=(Key=Tilde,bShift=False,bCtrl=True,bAlt=False,bCmd=False)
 ; These are the defaults for XboxOne settings, and they need to be in the .ini since UBT reads the .ini settings, without instantiating the class
 [/Script/XboxOneRuntimeSettings.XboxOneRuntimeSettings]
 bBuildForD3D12=false
+
+[/Script/Engine.SkeletalMeshReductionSettings]
++Settings=(ScreenSize=1.0,OptimizationSettings=(NumOfTrianglesPercentage=.5))
++Settings=(ScreenSize=.3,OptimizationSettings=(NumOfTrianglesPercentage=.25))
++Settings=(ScreenSize=.1,OptimizationSettings=(NumOfTrianglesPercentage=.125))
++Settings=(ScreenSize=.06,OptimizationSettings=(NumOfTrianglesPercentage=.06))
+
diff --git a/Engine/Source/Editor/AnimGraph/Classes/AnimGraphNode_TransitionPoseEvaluator.h b/Engine/Source/Editor/AnimGraph/Classes/AnimGraphNode_TransitionPoseEvaluator.h
--- a/Engine/Source/Editor/AnimGraph/Classes/AnimGraphNode_TransitionPoseEvaluator.h
+++ b/Engine/Source/Editor/AnimGraph/Classes/AnimGraphNode_TransitionPoseEvaluator.h
@@ -24,9 +24,17 @@ class UAnimGraphNode_TransitionPoseEvaluator : public UAnimGraphNode_Base
 
 	// UAnimGraphNode_Base interface
 	virtual FString GetNodeCategory() const override;
+	virtual void CustomizeDetails(IDetailLayoutBuilder& DetailBuilder) override;
 	// End of UAnimGraphNode_Base interface
 
 	// UK2Node interface.
 	virtual void GetMenuActions(FBlueprintActionDatabaseRegistrar& ActionRegistrar) const override;
+
+private:
+
+	// Details customization helpers
+	EVisibility GetCacheFramesVisibility() const;
+
+
 	// End of UK2Node interface
 };
\ No newline at end of file
diff --git a/Engine/Source/Editor/AnimGraph/Classes/AnimStateConduitNode.h b/Engine/Source/Editor/AnimGraph/Classes/AnimStateConduitNode.h
--- a/Engine/Source/Editor/AnimGraph/Classes/AnimStateConduitNode.h
+++ b/Engine/Source/Editor/AnimGraph/Classes/AnimStateConduitNode.h
@@ -33,5 +33,6 @@ public:
 	virtual void GetTransitionList(TArray<class UAnimStateTransitionNode*>& OutTransitions, bool bWantSortedList = true) override;
 	virtual FString GetDesiredNewNodeName() const override;
 	virtual UEdGraph* GetBoundGraph() const override { return BoundGraph; }
+	virtual void ClearBoundGraph() override { BoundGraph = nullptr; }
 	//~ End UAnimStateNodeBase Interface
 };
diff --git a/Engine/Source/Editor/AnimGraph/Classes/AnimStateNode.h b/Engine/Source/Editor/AnimGraph/Classes/AnimStateNode.h
--- a/Engine/Source/Editor/AnimGraph/Classes/AnimStateNode.h
+++ b/Engine/Source/Editor/AnimGraph/Classes/AnimStateNode.h
@@ -66,4 +66,5 @@ public:
 
 public:
 	virtual UEdGraph* GetBoundGraph() const override { return BoundGraph; }
+	virtual void ClearBoundGraph() override { BoundGraph = nullptr; }
 };
diff --git a/Engine/Source/Editor/AnimGraph/Classes/AnimStateNodeBase.h b/Engine/Source/Editor/AnimGraph/Classes/AnimStateNodeBase.h
--- a/Engine/Source/Editor/AnimGraph/Classes/AnimStateNodeBase.h
+++ b/Engine/Source/Editor/AnimGraph/Classes/AnimStateNodeBase.h
@@ -29,8 +29,8 @@ class UAnimStateNodeBase : public UEdGraphNode
 	// Populates the OutTransitions array with a list of transition nodes connected to this state
 	ANIMGRAPH_API virtual void GetTransitionList(TArray<class UAnimStateTransitionNode*>& OutTransitions, bool bWantSortedList = false) { }
 
-	//
 	virtual UEdGraph* GetBoundGraph() const { return NULL; }
+	virtual void ClearBoundGraph() {}
 
 	ANIMGRAPH_API UAnimBlueprint* GetAnimBlueprint() const;
 
diff --git a/Engine/Source/Editor/AnimGraph/Classes/AnimStateTransitionNode.h b/Engine/Source/Editor/AnimGraph/Classes/AnimStateTransitionNode.h
--- a/Engine/Source/Editor/AnimGraph/Classes/AnimStateTransitionNode.h
+++ b/Engine/Source/Editor/AnimGraph/Classes/AnimStateTransitionNode.h
@@ -115,6 +115,7 @@ class UAnimStateTransitionNode : public UAnimStateNodeBase
 
 	//~ Begin UAnimStateNodeBase Interface
 	virtual UEdGraph* GetBoundGraph() const override { return BoundGraph; }
+	virtual void ClearBoundGraph() override { BoundGraph = nullptr; }
 	virtual UEdGraphPin* GetInputPin() const override { return Pins[0]; }
 	virtual UEdGraphPin* GetOutputPin() const override { return Pins[1]; }
 	//~ End UAnimStateNodeBase Interface
diff --git a/Engine/Source/Editor/AnimGraph/Classes/AnimationConduitGraphSchema.h b/Engine/Source/Editor/AnimGraph/Classes/AnimationConduitGraphSchema.h
--- a/Engine/Source/Editor/AnimGraph/Classes/AnimationConduitGraphSchema.h
+++ b/Engine/Source/Editor/AnimGraph/Classes/AnimationConduitGraphSchema.h
@@ -17,5 +17,6 @@ class UAnimationConduitGraphSchema : public UEdGraphSchema_K2
 	virtual void GetGraphDisplayInformation(const UEdGraph& Graph, /*out*/ FGraphDisplayInfo& DisplayInfo) const override;
 	virtual bool DoesSupportEventDispatcher() const	override { return false; }
 	virtual bool ShouldAlwaysPurgeOnModification() const override { return true; }
+	virtual void HandleGraphBeingDeleted(UEdGraph& GraphBeingRemoved) const override;
 	// End of UEdGraphSchema interface
 };
diff --git a/Engine/Source/Editor/AnimGraph/Classes/AnimationCustomTransitionSchema.h b/Engine/Source/Editor/AnimGraph/Classes/AnimationCustomTransitionSchema.h
--- a/Engine/Source/Editor/AnimGraph/Classes/AnimationCustomTransitionSchema.h
+++ b/Engine/Source/Editor/AnimGraph/Classes/AnimationCustomTransitionSchema.h
@@ -12,5 +12,6 @@ class UAnimationCustomTransitionSchema : public UAnimationGraphSchema
 	//~ Begin UEdGraphSchema Interface.
 	virtual void CreateDefaultNodesForGraph(UEdGraph& Graph) const override;
 	void GetGraphDisplayInformation(const UEdGraph& Graph, /*out*/ FGraphDisplayInfo& DisplayInfo) const override;
+	virtual void HandleGraphBeingDeleted(UEdGraph& GraphBeingRemoved) const override;
 	//~ End UEdGraphSchema Interface.
 };
diff --git a/Engine/Source/Editor/AnimGraph/Classes/AnimationTransitionSchema.h b/Engine/Source/Editor/AnimGraph/Classes/AnimationTransitionSchema.h
--- a/Engine/Source/Editor/AnimGraph/Classes/AnimationTransitionSchema.h
+++ b/Engine/Source/Editor/AnimGraph/Classes/AnimationTransitionSchema.h
@@ -15,6 +15,7 @@ class UAnimationTransitionSchema : public UEdGraphSchema_K2
 	virtual void GetGraphDisplayInformation(const UEdGraph& Graph, /*out*/ FGraphDisplayInfo& DisplayInfo) const override;
 	virtual bool DoesSupportEventDispatcher() const	override { return false; }
 	virtual bool ShouldAlwaysPurgeOnModification() const override { return true; }
+	virtual void HandleGraphBeingDeleted(UEdGraph& GraphBeingRemoved) const override;
 	//~ End UEdGraphSchema Interface.
 
 private:
diff --git a/Engine/Source/Editor/AnimGraph/Private/AnimGraphNode_StateMachineBase.cpp b/Engine/Source/Editor/AnimGraph/Private/AnimGraphNode_StateMachineBase.cpp
--- a/Engine/Source/Editor/AnimGraph/Private/AnimGraphNode_StateMachineBase.cpp
+++ b/Engine/Source/Editor/AnimGraph/Private/AnimGraphNode_StateMachineBase.cpp
@@ -111,8 +111,13 @@ void UAnimGraphNode_StateMachineBase::PostPlacedNewNode()
 	Schema->CreateDefaultNodesForGraph(*EditorStateMachineGraph);
 
 	// Add the new graph as a child of our parent graph
-	GetGraph()->Modify();
-	GetGraph()->SubGraphs.Add(EditorStateMachineGraph);
+	UEdGraph* ParentGraph = GetGraph();
+	
+	if(ParentGraph->SubGraphs.Find(EditorStateMachineGraph) == INDEX_NONE)
+	{
+		ParentGraph->Modify();
+		ParentGraph->SubGraphs.Add(EditorStateMachineGraph);
+	}
 }
 
 UObject* UAnimGraphNode_StateMachineBase::GetJumpTargetForDoubleClick() const
@@ -142,7 +147,11 @@ void UAnimGraphNode_StateMachineBase::PostPasteNode()
 
 	// Add the new graph as a child of our parent graph
 	UEdGraph* ParentGraph = CastChecked<UEdGraph>(GetGraph());
-	ParentGraph->SubGraphs.Add(EditorStateMachineGraph);
+
+	if(ParentGraph->SubGraphs.Find(EditorStateMachineGraph) == INDEX_NONE)
+	{
+		ParentGraph->SubGraphs.Add(EditorStateMachineGraph);
+	}
 
 	// Find an interesting name
 	TSharedPtr<INameValidatorInterface> NameValidator = FNameValidatorFactory::MakeValidator(this);
diff --git a/Engine/Source/Editor/AnimGraph/Private/AnimGraphNode_TransitionPoseEvaluator.cpp b/Engine/Source/Editor/AnimGraph/Private/AnimGraphNode_TransitionPoseEvaluator.cpp
--- a/Engine/Source/Editor/AnimGraph/Private/AnimGraphNode_TransitionPoseEvaluator.cpp
+++ b/Engine/Source/Editor/AnimGraph/Private/AnimGraphNode_TransitionPoseEvaluator.cpp
@@ -7,6 +7,11 @@
 #include "CompilerResultsLog.h"
 #include "GraphEditorActions.h"
 
+#include "DetailCategoryBuilder.h"
+#include "DetailLayoutBuilder.h"
+#include "PropertyHandle.h"
+#include "IDetailPropertyRow.h"
+
 /////////////////////////////////////////////////////
 // UAnimGraphNode_TransitionPoseEvaluator
 
@@ -63,4 +68,29 @@ void UAnimGraphNode_TransitionPoseEvaluator::GetMenuActions(FBlueprintActionData
 	// Intentionally empty: Don't allow an option to create them, as they're auto-created when custom blend graphs are made
 }
 
+void UAnimGraphNode_TransitionPoseEvaluator::CustomizeDetails(IDetailLayoutBuilder& DetailBuilder)
+{
+	IDetailCategoryBuilder& PoseCategory = DetailBuilder.EditCategory("Pose", LOCTEXT("PoseCategoryName", "Pose"));
+
+	FString CacheFramesPropertyName = FString::Printf(TEXT("Node.%s"), GET_MEMBER_NAME_STRING_CHECKED(FAnimNode_TransitionPoseEvaluator, FramesToCachePose));
+	TSharedPtr<IPropertyHandle> FramesToCachePoseProperty = DetailBuilder.GetProperty(*CacheFramesPropertyName, GetClass());
+
+	// Hide this property, we only want this to appear when using delayed freeze
+	FramesToCachePoseProperty->MarkHiddenByCustomization();
+	
+	// Bind visibility helper for the property
+	TAttribute<EVisibility> VisibilityAttr;
+	VisibilityAttr.BindUObject(this, &UAnimGraphNode_TransitionPoseEvaluator::GetCacheFramesVisibility);
+	PoseCategory.AddProperty(FramesToCachePoseProperty).Visibility(VisibilityAttr);
+}
+
+EVisibility UAnimGraphNode_TransitionPoseEvaluator::GetCacheFramesVisibility() const
+{
+	if(Node.EvaluatorMode != EEvaluatorMode::EM_DelayedFreeze)
+	{
+		return EVisibility::Hidden;
+	}
+	return EVisibility::Visible;
+}
+
 #undef LOCTEXT_NAMESPACE
diff --git a/Engine/Source/Editor/AnimGraph/Private/AnimStateConduitNode.cpp b/Engine/Source/Editor/AnimGraph/Private/AnimStateConduitNode.cpp
--- a/Engine/Source/Editor/AnimGraph/Private/AnimStateConduitNode.cpp
+++ b/Engine/Source/Editor/AnimGraph/Private/AnimStateConduitNode.cpp
@@ -111,7 +111,12 @@ void UAnimStateConduitNode::PostPlacedNewNode()
 	Schema->CreateDefaultNodesForGraph(*BoundGraph);
 
 	// Add the new graph as a child of our parent graph
-	GetGraph()->SubGraphs.Add(BoundGraph);
+	UEdGraph* ParentGraph = GetGraph();
+
+	if(ParentGraph->SubGraphs.Find(BoundGraph) == INDEX_NONE)
+	{
+		ParentGraph->SubGraphs.Add(BoundGraph);
+	}
 }
 
 void UAnimStateConduitNode::DestroyNode()
diff --git a/Engine/Source/Editor/AnimGraph/Private/AnimStateNode.cpp b/Engine/Source/Editor/AnimGraph/Private/AnimStateNode.cpp
--- a/Engine/Source/Editor/AnimGraph/Private/AnimStateNode.cpp
+++ b/Engine/Source/Editor/AnimGraph/Private/AnimStateNode.cpp
@@ -162,7 +162,12 @@ void UAnimStateNode::PostPlacedNewNode()
 	Schema->CreateDefaultNodesForGraph(*BoundGraph);
 
 	// Add the new graph as a child of our parent graph
-	GetGraph()->SubGraphs.Add(BoundGraph);
+	UEdGraph* ParentGraph = GetGraph();
+
+	if(ParentGraph->SubGraphs.Find(BoundGraph) == INDEX_NONE)
+	{
+		ParentGraph->SubGraphs.Add(BoundGraph);
+	}
 }
 
 void UAnimStateNode::DestroyNode()
diff --git a/Engine/Source/Editor/AnimGraph/Private/AnimStateNodeBase.cpp b/Engine/Source/Editor/AnimGraph/Private/AnimStateNodeBase.cpp
--- a/Engine/Source/Editor/AnimGraph/Private/AnimStateNodeBase.cpp
+++ b/Engine/Source/Editor/AnimGraph/Private/AnimStateNodeBase.cpp
@@ -50,7 +50,11 @@ void UAnimStateNodeBase::PostPasteNode()
 	{
 		// Add the new graph as a child of our parent graph
 		UEdGraph* ParentGraph = GetGraph();
-		ParentGraph->SubGraphs.Add(BoundGraph);
+
+		if(ParentGraph->SubGraphs.Find(BoundGraph) == INDEX_NONE)
+		{
+			ParentGraph->SubGraphs.Add(BoundGraph);
+		}
 
 		//@TODO: CONDUIT: Merge conflict - May no longer be necessary due to other changes?
 //		FBlueprintEditorUtils::RenameGraphWithSuggestion(BoundGraph, NameValidator, GetDesiredNewNodeName());
diff --git a/Engine/Source/Editor/AnimGraph/Private/AnimStateTransitionNode.cpp b/Engine/Source/Editor/AnimGraph/Private/AnimStateTransitionNode.cpp
--- a/Engine/Source/Editor/AnimGraph/Private/AnimStateTransitionNode.cpp
+++ b/Engine/Source/Editor/AnimGraph/Private/AnimStateTransitionNode.cpp
@@ -215,6 +215,13 @@ void UAnimStateTransitionNode::PinConnectionListChanged(UEdGraphPin* Pin)
 	{
 		// Commit suicide; transitions must always have an input and output connection
 		Modify();
+
+		// Our parent graph will have our graph in SubGraphs so needs to be modified to record that.
+		if(UEdGraph* ParentGraph = GetGraph())
+		{
+			ParentGraph->Modify();
+		}
+
 		DestroyNode();
 	}
 }
@@ -410,7 +417,12 @@ void UAnimStateTransitionNode::CreateBoundGraph()
 	Schema->CreateDefaultNodesForGraph(*BoundGraph);
 
 	// Add the new graph as a child of our parent graph
-	GetGraph()->SubGraphs.Add(BoundGraph);
+	UEdGraph* ParentGraph = GetGraph();
+
+	if(ParentGraph->SubGraphs.Find(BoundGraph) == INDEX_NONE)
+	{
+		ParentGraph->SubGraphs.Add(BoundGraph);
+	}
 }
 
 void UAnimStateTransitionNode::CreateCustomTransitionGraph()
@@ -432,8 +444,13 @@ void UAnimStateTransitionNode::CreateCustomTransitionGraph()
 	Schema->CreateDefaultNodesForGraph(*CustomTransitionGraph);
 
 	// Add the new graph as a child of our parent graph
-	GetGraph()->Modify();
-	GetGraph()->SubGraphs.Add(CustomTransitionGraph);
+	UEdGraph* ParentGraph = GetGraph();
+
+	if(ParentGraph->SubGraphs.Find(CustomTransitionGraph) == INDEX_NONE)
+	{
+		ParentGraph->Modify();
+		ParentGraph->SubGraphs.Add(CustomTransitionGraph);
+	}
 }
 
 void UAnimStateTransitionNode::DestroyNode()
diff --git a/Engine/Source/Editor/AnimGraph/Private/AnimationConduitGraphSchema.cpp b/Engine/Source/Editor/AnimGraph/Private/AnimationConduitGraphSchema.cpp
--- a/Engine/Source/Editor/AnimGraph/Private/AnimationConduitGraphSchema.cpp
+++ b/Engine/Source/Editor/AnimGraph/Private/AnimationConduitGraphSchema.cpp
@@ -40,3 +40,37 @@ void UAnimationConduitGraphSchema::GetGraphDisplayInformation(const UEdGraph& Gr
 
 	DisplayInfo.DisplayName = DisplayInfo.PlainName;
 }
+
+void UAnimationConduitGraphSchema::HandleGraphBeingDeleted(UEdGraph& GraphBeingRemoved) const
+{
+	Super::HandleGraphBeingDeleted(GraphBeingRemoved);
+
+	if(UBlueprint* Blueprint = FBlueprintEditorUtils::FindBlueprintForGraph(&GraphBeingRemoved))
+	{
+		// Handle composite anim graph nodes
+		TArray<UAnimStateNodeBase*> StateNodes;
+		FBlueprintEditorUtils::GetAllNodesOfClassEx<UAnimStateConduitNode>(Blueprint, StateNodes);
+
+		TSet<UAnimStateNodeBase*> NodesToDelete;
+		for(int32 i = 0; i < StateNodes.Num(); ++i)
+		{
+			UAnimStateNodeBase* StateNode = StateNodes[i];
+			if(StateNode->GetBoundGraph() == &GraphBeingRemoved)
+			{
+				NodesToDelete.Add(StateNode);
+			}
+		}
+
+		// Delete the node that owns us
+		ensure(NodesToDelete.Num() <= 1);
+		for(TSet<UAnimStateNodeBase*>::TIterator It(NodesToDelete); It; ++It)
+		{
+			UAnimStateNodeBase* NodeToDelete = *It;
+
+			FBlueprintEditorUtils::RemoveNode(Blueprint, NodeToDelete, true);
+
+			// Prevent re-entrancy here
+			NodeToDelete->ClearBoundGraph();
+		}
+	}
+}
diff --git a/Engine/Source/Editor/AnimGraph/Private/AnimationCustomTransitionSchema.cpp b/Engine/Source/Editor/AnimGraph/Private/AnimationCustomTransitionSchema.cpp
--- a/Engine/Source/Editor/AnimGraph/Private/AnimationCustomTransitionSchema.cpp
+++ b/Engine/Source/Editor/AnimGraph/Private/AnimationCustomTransitionSchema.cpp
@@ -61,3 +61,24 @@ void UAnimationCustomTransitionSchema::GetGraphDisplayInformation(const UEdGraph
 
 	DisplayInfo.DisplayName = DisplayInfo.PlainName;
 }
+
+void UAnimationCustomTransitionSchema::HandleGraphBeingDeleted(UEdGraph& GraphBeingRemoved) const
+{
+	if(UBlueprint* Blueprint = FBlueprintEditorUtils::FindBlueprintForGraph(&GraphBeingRemoved))
+	{
+		// Look for state nodes that reference this graph
+		TArray<UAnimStateTransitionNode*> TransitionNodes;
+		FBlueprintEditorUtils::GetAllNodesOfClass<UAnimStateTransitionNode>(Blueprint, TransitionNodes);
+
+		for(UAnimStateTransitionNode* Node : TransitionNodes)
+		{
+			if(Node->CustomTransitionGraph == &GraphBeingRemoved)
+			{
+				// Clear out the logic for this node as we're removing the graph
+				Node->Modify();
+				Node->LogicType = ETransitionLogicType::TLT_StandardBlend;
+				Node->CustomTransitionGraph = nullptr;
+			}
+		}
+	}
+}
diff --git a/Engine/Source/Editor/AnimGraph/Private/AnimationGraphSchema.cpp b/Engine/Source/Editor/AnimGraph/Private/AnimationGraphSchema.cpp
--- a/Engine/Source/Editor/AnimGraph/Private/AnimationGraphSchema.cpp
+++ b/Engine/Source/Editor/AnimGraph/Private/AnimationGraphSchema.cpp
@@ -87,14 +87,14 @@ void UAnimationGraphSchema::HandleGraphBeingDeleted(UEdGraph& GraphBeingRemoved)
 	if (UBlueprint* Blueprint = FBlueprintEditorUtils::FindBlueprintForGraph(&GraphBeingRemoved))
 	{
 		// Look for state nodes that reference this graph
-		TArray<UAnimStateNode*> StateNodes;
-		FBlueprintEditorUtils::GetAllNodesOfClass<UAnimStateNode>(Blueprint, /*out*/ StateNodes);
+		TArray<UAnimStateNodeBase*> StateNodes;
+		FBlueprintEditorUtils::GetAllNodesOfClassEx<UAnimStateNode>(Blueprint, StateNodes);
 
-		TSet<UAnimStateNode*> NodesToDelete;
+		TSet<UAnimStateNodeBase*> NodesToDelete;
 		for (int32 i = 0; i < StateNodes.Num(); ++i)
 		{
-			UAnimStateNode* StateNode = StateNodes[i];
-			if (StateNode->BoundGraph == &GraphBeingRemoved)
+			UAnimStateNodeBase* StateNode = StateNodes[i];
+			if (StateNode->GetBoundGraph() == &GraphBeingRemoved)
 			{
 				NodesToDelete.Add(StateNode);
 			}
@@ -102,15 +102,14 @@ void UAnimationGraphSchema::HandleGraphBeingDeleted(UEdGraph& GraphBeingRemoved)
 
 		// Delete the node that owns us
 		ensure(NodesToDelete.Num() <= 1);
-		for (TSet<UAnimStateNode*>::TIterator It(NodesToDelete); It; ++It)
+		for (TSet<UAnimStateNodeBase*>::TIterator It(NodesToDelete); It; ++It)
 		{
-			UAnimStateNode* NodeToDelete = *It;
+			UAnimStateNodeBase* NodeToDelete = *It;
 
-			// Prevent re-entrancy here
-			NodeToDelete->BoundGraph = NULL;
+			FBlueprintEditorUtils::RemoveNode(Blueprint, NodeToDelete, true);
 
-			NodeToDelete->Modify();
-			NodeToDelete->DestroyNode();
+			// Prevent re-entrancy here
+			NodeToDelete->ClearBoundGraph();
 		}
 	}
 }
diff --git a/Engine/Source/Editor/AnimGraph/Private/AnimationTransitionSchema.cpp b/Engine/Source/Editor/AnimGraph/Private/AnimationTransitionSchema.cpp
--- a/Engine/Source/Editor/AnimGraph/Private/AnimationTransitionSchema.cpp
+++ b/Engine/Source/Editor/AnimGraph/Private/AnimationTransitionSchema.cpp
@@ -92,4 +92,38 @@ UAnimStateNode* UAnimationTransitionSchema::GetStateNodeFromGraph(const FAnimBlu
 	return NULL;
 }
 
+void UAnimationTransitionSchema::HandleGraphBeingDeleted(UEdGraph& GraphBeingRemoved) const
+{
+	Super::HandleGraphBeingDeleted(GraphBeingRemoved);
+
+	if(UBlueprint* Blueprint = FBlueprintEditorUtils::FindBlueprintForGraph(&GraphBeingRemoved))
+	{
+		// Handle composite anim graph nodes
+		TArray<UAnimStateNodeBase*> StateNodes;
+		FBlueprintEditorUtils::GetAllNodesOfClassEx<UAnimStateTransitionNode>(Blueprint, StateNodes);
+
+		TSet<UAnimStateNodeBase*> NodesToDelete;
+		for(int32 i = 0; i < StateNodes.Num(); ++i)
+		{
+			UAnimStateNodeBase* StateNode = StateNodes[i];
+			if(StateNode->GetBoundGraph() == &GraphBeingRemoved)
+			{
+				NodesToDelete.Add(StateNode);
+			}
+		}
+
+		// Delete the node that owns us
+		ensure(NodesToDelete.Num() <= 1);
+		for(TSet<UAnimStateNodeBase*>::TIterator It(NodesToDelete); It; ++It)
+		{
+			UAnimStateNodeBase* NodeToDelete = *It;
+
+			FBlueprintEditorUtils::RemoveNode(Blueprint, NodeToDelete, true);
+
+			// Prevent re-entrancy here
+			NodeToDelete->ClearBoundGraph();
+		}
+	}
+}
+
 #undef LOCTEXT_NAMESPACE
\ No newline at end of file
diff --git a/Engine/Source/Editor/DetailCustomizations/Private/DetailCustomizations.cpp b/Engine/Source/Editor/DetailCustomizations/Private/DetailCustomizations.cpp
--- a/Engine/Source/Editor/DetailCustomizations/Private/DetailCustomizations.cpp
+++ b/Engine/Source/Editor/DetailCustomizations/Private/DetailCustomizations.cpp
@@ -27,7 +27,6 @@
 #include "ModuleManager.h"
 #include "AnimGraphDefinitions.h"
 #include "SoundDefinitions.h"
-#include "TransitionPoseEvaluatorNodeDetails.h"
 #include "MathStructCustomizations.h"
 #include "MathStructProxyCustomizations.h"
 #include "RangeStructCustomization.h"
@@ -270,7 +269,6 @@ void FDetailCustomizationsModule::RegisterObjectCustomizations()
 	RegisterCustomClassLayout("EditorSkeletonNotifyObj", FOnGetDetailCustomizationInstance::CreateStatic(&FSkeletonNotifyDetails::MakeInstance));
 	RegisterCustomClassLayout("AnimStateNode", FOnGetDetailCustomizationInstance::CreateStatic(&FAnimStateNodeDetails::MakeInstance));
 	RegisterCustomClassLayout("AnimStateTransitionNode", FOnGetDetailCustomizationInstance::CreateStatic(&FAnimTransitionNodeDetails::MakeInstance));
-	RegisterCustomClassLayout("AnimGraphNode_TransitionPoseEvaluator", FOnGetDetailCustomizationInstance::CreateStatic(&FTransitionPoseEvaluatorNodeDetails::MakeInstance));
 	RegisterCustomClassLayout("AnimGraphNode_Trail", FOnGetDetailCustomizationInstance::CreateStatic(&FAnimTrailNodeDetails::MakeInstance));
 
 	RegisterCustomClassLayout("SoundWave", FOnGetDetailCustomizationInstance::CreateStatic(&FSoundWaveDetails::MakeInstance));
diff --git a/Engine/Source/Editor/DetailCustomizations/Private/TransitionPoseEvaluatorNodeDetails.cpp b/Engine/Source/Editor/DetailCustomizations/Private/TransitionPoseEvaluatorNodeDetails.cpp
deleted file mode 100644
diff --git a/Engine/Source/Editor/DetailCustomizations/Private/TransitionPoseEvaluatorNodeDetails.h b/Engine/Source/Editor/DetailCustomizations/Private/TransitionPoseEvaluatorNodeDetails.h
deleted file mode 100644
diff --git a/Engine/Source/Editor/Persona/Private/PersonaMeshDetails.cpp b/Engine/Source/Editor/Persona/Private/PersonaMeshDetails.cpp
--- a/Engine/Source/Editor/Persona/Private/PersonaMeshDetails.cpp
+++ b/Engine/Source/Editor/Persona/Private/PersonaMeshDetails.cpp
@@ -22,6 +22,8 @@
 #include "AnimGraphNodeDetails.h"
 #include "STextComboBox.h"
 
+#include "Engine/SkeletalMeshReductionSettings.h"
+
 #define LOCTEXT_NAMESPACE "PersonaMeshDetails"
 
 /** Returns true if automatic mesh reduction is available. */
@@ -776,30 +778,51 @@ void FPersonaMeshDetails::ApplyChanges()
 	// we need to add more
 	else if (LODCount > CurrentNumLODs)
 	{
-		// creating new ones only
+		// Retrieve default skeletal mesh reduction settings
+		USkeletalMeshReductionSettings* ReductionSettings = USkeletalMeshReductionSettings::Get();		
+
+		// Only create new skeletal mesh LOD level entries
 		for (int32 LODIdx = CurrentNumLODs; LODIdx < LODCount; LODIdx++)
 		{
-			FSkeletalMeshOptimizationSettings Setting;
+			FSkeletalMeshOptimizationSettings Settings;
+			
+			const int32 SettingsIndex = LODIdx - 1;
+			// If there are valid default settings use those for the new LOD
 
-			// find whatever latest that was using mesh reduction, and 
-			// make it 50 % of that. 
-			for (int32 SubLOD = LODIdx - 1; SubLOD >= 0; --SubLOD)
+			const bool bHasValidUserSetting = ReductionSettings->HasValidSettings() && ReductionSettings->GetNumberOfSettings() > SettingsIndex;
+			if (bHasValidUserSetting)
+			{
+				const FSkeletalMeshLODGroupSettings& GroupSettings = ReductionSettings->GetDefaultSettingsForLODLevel(SettingsIndex);
+				Settings = GroupSettings.GetSettings();
+			}
+			else
 			{
-				if (SkelMesh->LODInfo[SubLOD].bHasBeenSimplified)
+				// Otherwise find whatever latest that was using mesh reduction, and make it 50% of that. 	
+				for (int32 SubLOD = LODIdx - 1; SubLOD >= 0; --SubLOD)
 				{
-					// copy whatever latest LOD info reduction setting
-					Setting = SkelMesh->LODInfo[SubLOD].ReductionSettings;
-					// and make it 50 % of that
-					Setting.NumOfTrianglesPercentage *= 0.5f;
-					break;
+					if (SkelMesh->LODInfo[SubLOD].bHasBeenSimplified)
+					{
+						// copy whatever latest LOD info reduction setting
+						Settings = SkelMesh->LODInfo[SubLOD].ReductionSettings;
+						// and make it 50 % of that
+						Settings.NumOfTrianglesPercentage *= 0.5f;
+						break;
+					}
 				}
 			}
 
 			// if no previous setting found, it will use default setting. 
-			FLODUtilities::SimplifySkeletalMeshLOD(UpdateContext, Setting, LODIdx);
+			FLODUtilities::SimplifySkeletalMeshLOD(UpdateContext, Settings, LODIdx);
 
 			FSkeletalMeshLODInfo& Info = SkelMesh->LODInfo[LODIdx];
-			Info.ReductionSettings = Setting;
+			Info.ReductionSettings = Settings;
+
+			// If there is a valid screensize value use that one for this new LOD
+			if (bHasValidUserSetting)
+			{
+				const FSkeletalMeshLODGroupSettings& GroupSettings = ReductionSettings->GetDefaultSettingsForLODLevel(SettingsIndex);
+				Info.ScreenSize = GroupSettings.GetScreenSize();
+			}
 		}
 	}
 	else if (IsApplyNeeded())
diff --git a/Engine/Source/Editor/StaticMeshEditor/Private/StaticMeshEditor.cpp b/Engine/Source/Editor/StaticMeshEditor/Private/StaticMeshEditor.cpp
--- a/Engine/Source/Editor/StaticMeshEditor/Private/StaticMeshEditor.cpp
+++ b/Engine/Source/Editor/StaticMeshEditor/Private/StaticMeshEditor.cpp
@@ -1202,10 +1202,7 @@ void FStaticMeshEditor::GenerateKDop(const FVector* Directions, uint32 NumDirect
 		const FPrimData PrimData = FPrimData(KPT_Convex, PrimIndex);
 		ClearSelectedPrims();
 		AddSelectedPrim(PrimData, true);
-		while( OverlapsExistingPrim(PrimData) )
-		{
-			TranslateSelectedPrims(OverlapNudge);
-		}
+		// Don't 'nudge' KDop prims, as they are fitted specifically around the geometry
 	}
 
 	Viewport->RefreshViewport();
diff --git a/Engine/Source/Editor/UnrealEd/Private/AnimationEditorUtils.cpp b/Engine/Source/Editor/UnrealEd/Private/AnimationEditorUtils.cpp
--- a/Engine/Source/Editor/UnrealEd/Private/AnimationEditorUtils.cpp
+++ b/Engine/Source/Editor/UnrealEd/Private/AnimationEditorUtils.cpp
@@ -11,6 +11,14 @@
 #include "Animation/AimOffsetBlendSpace1D.h"
 #include "Animation/AnimCompress.h"
 
+#include "AnimationGraph.h"
+#include "AnimStateNodeBase.h"
+#include "AnimStateTransitionNode.h"
+#include "AnimGraphNode_StateMachineBase.h"
+#include "AnimationStateMachineGraph.h"
+#include "K2Node_Composite.h"
+#include "AssertionMacros.h"
+
 #define LOCTEXT_NAMESPACE "AnimationEditorUtils"
 
 /////////////////////////////////////////////////////
@@ -436,6 +444,77 @@ namespace AnimationEditorUtils
 
 		return false;
 	}
+
+	void RegenerateSubGraphArrays(UAnimBlueprint* Blueprint)
+	{
+		// The anim graph should be the first function graph on the blueprint
+		if(Blueprint->FunctionGraphs.Num() > 0)
+		{
+			if(UAnimationGraph* AnimGraph = Cast<UAnimationGraph>(Blueprint->FunctionGraphs[0]))
+			{
+				RegenerateGraphSubGraphs(Blueprint, AnimGraph);
+			}
+		}
+	}
+
+	void RegenerateGraphSubGraphs(UAnimBlueprint* OwningBlueprint, UEdGraph* GraphToFix)
+	{
+		TArray<UEdGraph*> ChildGraphs;
+		FindChildGraphsFromNodes(GraphToFix, ChildGraphs);
+
+		for(UEdGraph* Child : ChildGraphs)
+		{
+			RegenerateGraphSubGraphs(OwningBlueprint, Child);
+		}
+
+		if(ChildGraphs != GraphToFix->SubGraphs)
+		{
+			UE_LOG(LogAnimation, Log, TEXT("Fixed missing or duplicated graph entries in SubGraph array for graph %s in AnimBP %s"), *GraphToFix->GetName(), *OwningBlueprint->GetName());
+			GraphToFix->SubGraphs = ChildGraphs;
+		}
+	}
+
+	void RemoveDuplicateSubGraphs(UEdGraph* GraphToClean)
+	{
+		TArray<UEdGraph*> NewSubGraphArray;
+
+		for(UEdGraph* SubGraph : GraphToClean->SubGraphs)
+		{
+			NewSubGraphArray.AddUnique(SubGraph);
+		}
+
+		if(NewSubGraphArray.Num() != GraphToClean->SubGraphs.Num())
+		{
+			GraphToClean->SubGraphs = NewSubGraphArray;
+		}
+	}
+
+	void FindChildGraphsFromNodes(UEdGraph* GraphToSearch, TArray<UEdGraph*>& ChildGraphs)
+	{
+		for(UEdGraphNode* CurrentNode : GraphToSearch->Nodes)
+		{
+			if(UAnimGraphNode_StateMachineBase* StateMachine = Cast<UAnimGraphNode_StateMachineBase>(CurrentNode))
+			{
+				ChildGraphs.AddUnique(StateMachine->EditorStateMachineGraph);
+			}
+			else if(UAnimStateNodeBase* StateNode = Cast<UAnimStateNodeBase>(CurrentNode))
+			{
+				ChildGraphs.AddUnique(StateNode->GetBoundGraph());
+
+				if(UAnimStateTransitionNode* TransitionNode = Cast<UAnimStateTransitionNode>(StateNode))
+				{
+					if(TransitionNode->CustomTransitionGraph)
+					{
+						ChildGraphs.AddUnique(TransitionNode->CustomTransitionGraph);
+					}
+				}
+			}
+			else if(UK2Node_Composite* CompositeNode = Cast<UK2Node_Composite>(CurrentNode))
+			{
+				ChildGraphs.AddUnique(CompositeNode->BoundGraph);
+			}
+		}
+	}
 }
 
 #undef LOCTEXT_NAMESPACE
\ No newline at end of file
diff --git a/Engine/Source/Editor/UnrealEd/Private/ApexClothingUtils.cpp b/Engine/Source/Editor/UnrealEd/Private/ApexClothingUtils.cpp
--- a/Engine/Source/Editor/UnrealEd/Private/ApexClothingUtils.cpp
+++ b/Engine/Source/Editor/UnrealEd/Private/ApexClothingUtils.cpp
@@ -929,6 +929,34 @@ bool AssociateClothingAssetWithSkeletalMesh(USkeletalMesh* SkelMesh, int32 LODIn
 	// set asset submesh info to new chunk 
 	ClothChunk.SetClothSubmeshIndex(AssetIndex, AssetSubmeshIndex);
 
+	// It's possible that the required/active bones for the new cloth chunk differ from
+	// the original chunk - here we inspect the bone map and activate any new bones.
+	bool bRequiredBonesChanged = false;
+	for(FBoneIndexType BoneIdx : ClothChunk.BoneMap)
+	{
+		if(!LODModel.RequiredBones.Contains(BoneIdx))
+		{
+			bRequiredBonesChanged = true;
+
+			FName MissingBoneName = NAME_None;
+			if(SkelMesh->RefSkeleton.IsValidIndex(BoneIdx))
+			{
+				MissingBoneName = SkelMesh->RefSkeleton.GetBoneName(BoneIdx);
+
+				UE_LOG(LogApexClothingUtils, Log, TEXT("Found a bone used by cloth (%s) that is not in the required bone list for this LOD. Adding to required bones."), *MissingBoneName.ToString());
+
+				LODModel.RequiredBones.Add(BoneIdx);
+				LODModel.ActiveBoneIndices.AddUnique(BoneIdx);
+			}
+		}
+	}
+
+	if(bRequiredBonesChanged)
+	{
+		LODModel.RequiredBones.Sort();
+		LODModel.ActiveBoneIndices.Sort();
+	}
+
 	// New Chunk is already added so needs to update VertexFactories before Render codes are called
 	ReregisterSkelMeshComponents(SkelMesh);
 
diff --git a/Engine/Source/Editor/UnrealEd/Private/Fbx/FbxStaticMeshImport.cpp b/Engine/Source/Editor/UnrealEd/Private/Fbx/FbxStaticMeshImport.cpp
--- a/Engine/Source/Editor/UnrealEd/Private/Fbx/FbxStaticMeshImport.cpp
+++ b/Engine/Source/Editor/UnrealEd/Private/Fbx/FbxStaticMeshImport.cpp
@@ -1665,7 +1665,7 @@ bool UnFbx::FFbxImporter::ImportCollisionModels(UStaticMesh* StaticMesh, const F
 
 				// This function cooks the given data, so we cannot test for duplicates based on the position data
 				// before we call it
-				AddConvexGeomFromVertices( CollisionVertices, &AggGeo, (const TCHAR*)Node->GetName() );
+				AddConvexGeomFromVertices( CollisionVertices, &AggGeo, ANSI_TO_TCHAR(Node->GetName()) );
 
 				// Now test the late element in the AggGeo list and remove it if its a duplicate
 				if(AggGeo.ConvexElems.Num() > 1)
@@ -1703,7 +1703,7 @@ bool UnFbx::FFbxImporter::ImportCollisionModels(UStaticMesh* StaticMesh, const F
 		{
 			FKAggregateGeom& AggGeo = StaticMesh->BodySetup->AggGeom;
 
-			AddBoxGeomFromTris( CollisionTriangles, &AggGeo, (const TCHAR*)Node->GetName() );
+			AddBoxGeomFromTris( CollisionTriangles, &AggGeo, ANSI_TO_TCHAR(Node->GetName()) );
 
 			// Now test the last element in the AggGeo list and remove it if its a duplicate
 			if(AggGeo.BoxElems.Num() > 1)
@@ -1727,7 +1727,7 @@ bool UnFbx::FFbxImporter::ImportCollisionModels(UStaticMesh* StaticMesh, const F
 		{
 			FKAggregateGeom& AggGeo = StaticMesh->BodySetup->AggGeom;
 
-			AddSphereGeomFromVerts( CollisionVertices, &AggGeo, (const TCHAR*)Node->GetName() );
+			AddSphereGeomFromVerts( CollisionVertices, &AggGeo, ANSI_TO_TCHAR(Node->GetName()) );
 
 			// Now test the late element in the AggGeo list and remove it if its a duplicate
 			if(AggGeo.SphereElems.Num() > 1)
@@ -1769,4 +1769,4 @@ bool UnFbx::FFbxImporter::ImportCollisionModels(UStaticMesh* StaticMesh, const F
 		
 	return true;
 }
-#undef LOCTEXT_NAMESPACE
\ No newline at end of file
+#undef LOCTEXT_NAMESPACE
diff --git a/Engine/Source/Editor/UnrealEd/Private/GeomFitUtils.cpp b/Engine/Source/Editor/UnrealEd/Private/GeomFitUtils.cpp
--- a/Engine/Source/Editor/UnrealEd/Private/GeomFitUtils.cpp
+++ b/Engine/Source/Editor/UnrealEd/Private/GeomFitUtils.cpp
@@ -98,15 +98,11 @@ int32 GenerateKDopAsSimpleCollision(UStaticMesh* StaticMesh, const TArray<FVecto
 		}
 	}
 
-	// Make sure that the min size is reached
-	const float MinSize = 0.5f;
-
+	// Inflate kdop to ensure it is no degenerate
+	const float MinSize = 0.1f;
 	for(int32 i=0; i<kCount; i++)
 	{
-		if (FMath::Abs(maxDist[i]) < MinSize)
-		{
-			maxDist[i] = MinSize;
-		}
+		maxDist[i] += MinSize;
 	}
 
 	// Now we have the planes of the kdop, we work out the face polygons.
diff --git a/Engine/Source/Editor/UnrealEd/Private/Kismet2/BlueprintEditorUtils.cpp b/Engine/Source/Editor/UnrealEd/Private/Kismet2/BlueprintEditorUtils.cpp
--- a/Engine/Source/Editor/UnrealEd/Private/Kismet2/BlueprintEditorUtils.cpp
+++ b/Engine/Source/Editor/UnrealEd/Private/Kismet2/BlueprintEditorUtils.cpp
@@ -2529,6 +2529,7 @@ void FBlueprintEditorUtils::RemoveGraph(UBlueprint* Blueprint, class UEdGraph* G
 		else if (UEdGraph* OuterGraph = Cast<UEdGraph>(TestOuter))
 		{
 			// remove ourselves
+			OuterGraph->Modify();
 			OuterGraph->SubGraphs.Remove(GraphToRemove);
 		}
 		else if (! (Cast<UK2Node_Composite>(TestOuter)	|| 
diff --git a/Engine/Source/Editor/UnrealEd/Public/AnimationEditorUtils.h b/Engine/Source/Editor/UnrealEd/Public/AnimationEditorUtils.h
--- a/Engine/Source/Editor/UnrealEd/Public/AnimationEditorUtils.h
+++ b/Engine/Source/Editor/UnrealEd/Public/AnimationEditorUtils.h
@@ -60,6 +60,13 @@ namespace AnimationEditorUtils
 	UNREALED_API void FillCreateAssetMenu(FMenuBuilder& MenuBuilder, TArray<TWeakObjectPtr<USkeleton>> Skeletons, FAnimAssetCreated AssetCreated, bool bInContentBrowser=true);
 	UNREALED_API void CreateUniqueAssetName(const FString& InBasePackageName, const FString& InSuffix, FString& OutPackageName, FString& OutAssetName);
 	UNREALED_API bool ApplyCompressionAlgorithm(TArray<UAnimSequence*>& AnimSequencePtrs, class UAnimCompress* Algorithm);
+	
+	// The following functions are used to fix subgraph arrays for assets
+	UNREALED_API void RegenerateSubGraphArrays(UAnimBlueprint* Blueprint);
+	void RegenerateGraphSubGraphs(UAnimBlueprint* OwningBlueprint, UEdGraph* GraphToFix);
+	void RemoveDuplicateSubGraphs(UEdGraph* GraphToClean);
+	void FindChildGraphsFromNodes(UEdGraph* GraphToSearch, TArray<UEdGraph*>& ChildGraphs);
+	//////////////////////////////////////////////////////////////////////////////////////////
 
 	template <typename TFactory, typename T>
 	void ExecuteNewAnimAsset(TArray<TWeakObjectPtr<USkeleton>> Objects, const FString InSuffix, FAnimAssetCreated AssetCreated, bool bInContentBrowser)
diff --git a/Engine/Source/Editor/UnrealEd/Public/Kismet2/BlueprintEditorUtils.h b/Engine/Source/Editor/UnrealEd/Public/Kismet2/BlueprintEditorUtils.h
--- a/Engine/Source/Editor/UnrealEd/Public/Kismet2/BlueprintEditorUtils.h
+++ b/Engine/Source/Editor/UnrealEd/Public/Kismet2/BlueprintEditorUtils.h
@@ -386,6 +386,19 @@ public:
 		}
 	}
 
+	/** Returns all nodes in all graphs of at least the minimum node type */
+	template< class MinNodeType, class ArrayClassType>
+	static inline void GetAllNodesOfClassEx(const UBlueprint* Blueprint, TArray<ArrayClassType*>& OutNodes)
+	{
+		TArray<UEdGraph*> AllGraphs;
+		Blueprint->GetAllGraphs(AllGraphs);
+		for(UEdGraph* Graph : AllGraphs)
+		{
+			check(Graph != nullptr);
+			Graph->GetNodesOfClassEx<MinNodeType, ArrayClassType>(OutNodes);
+		}
+	}
+
 	/**
 	 * Searches all nodes in a Blueprint and checks for a matching Guid
 	 *
diff --git a/Engine/Source/Runtime/AIModule/Classes/Navigation/CrowdAgentInterface.h b/Engine/Source/Runtime/AIModule/Classes/Navigation/CrowdAgentInterface.h
--- a/Engine/Source/Runtime/AIModule/Classes/Navigation/CrowdAgentInterface.h
+++ b/Engine/Source/Runtime/AIModule/Classes/Navigation/CrowdAgentInterface.h
@@ -25,4 +25,13 @@ class ICrowdAgentInterface
 
 	/** @return max speed of crowd agent */
 	virtual float GetCrowdAgentMaxSpeed() const { return 0.0f; }
+
+	/** Group mask for this agent */
+	virtual int32 GetCrowdAgentAvoidanceGroup() const { return 1; }
+
+	/** Will avoid other agents if they are in one of specified groups */
+	virtual int32 GetCrowdAgentGroupsToAvoid() const { return MAX_int32; }
+
+	/** Will NOT avoid other agents if they are in one of specified groups, higher priority than GroupsToAvoid */
+	virtual int32 GetCrowdAgentGroupsToIgnore() const { return 0; }
 };
diff --git a/Engine/Source/Runtime/AIModule/Classes/Navigation/CrowdFollowingComponent.h b/Engine/Source/Runtime/AIModule/Classes/Navigation/CrowdFollowingComponent.h
--- a/Engine/Source/Runtime/AIModule/Classes/Navigation/CrowdFollowingComponent.h
+++ b/Engine/Source/Runtime/AIModule/Classes/Navigation/CrowdFollowingComponent.h
@@ -43,6 +43,9 @@ class AIMODULE_API UCrowdFollowingComponent : public UPathFollowingComponent, pu
 	virtual FVector GetCrowdAgentVelocity() const override;
 	virtual void GetCrowdAgentCollisions(float& CylinderRadius, float& CylinderHalfHeight) const override;
 	virtual float GetCrowdAgentMaxSpeed() const override;
+	virtual int32 GetCrowdAgentAvoidanceGroup() const override;
+	virtual int32 GetCrowdAgentGroupsToAvoid() const override;
+	virtual int32 GetCrowdAgentGroupsToIgnore() const override;
 	// ICrowdAgentInterface END
 
 	// PathFollowingComponent BEGIN
diff --git a/Engine/Source/Runtime/AIModule/Private/Navigation/CrowdFollowingComponent.cpp b/Engine/Source/Runtime/AIModule/Private/Navigation/CrowdFollowingComponent.cpp
--- a/Engine/Source/Runtime/AIModule/Private/Navigation/CrowdFollowingComponent.cpp
+++ b/Engine/Source/Runtime/AIModule/Private/Navigation/CrowdFollowingComponent.cpp
@@ -64,6 +64,21 @@ float UCrowdFollowingComponent::GetCrowdAgentMaxSpeed() const
 	return MovementComp ? MovementComp->GetMaxSpeed() : 0.0f;
 }
 
+int32 UCrowdFollowingComponent::GetCrowdAgentAvoidanceGroup() const
+{
+	return GetAvoidanceGroup();
+}
+
+int32 UCrowdFollowingComponent::GetCrowdAgentGroupsToAvoid() const
+{
+	return GetGroupsToAvoid();
+}
+
+int32 UCrowdFollowingComponent::GetCrowdAgentGroupsToIgnore() const
+{
+	return GetGroupsToIgnore();
+}
+
 void UCrowdFollowingComponent::SetCrowdAnticipateTurns(bool bEnable, bool bUpdateAgent)
 {
 	if (bEnableAnticipateTurns != bEnable)
diff --git a/Engine/Source/Runtime/AIModule/Private/Navigation/CrowdManager.cpp b/Engine/Source/Runtime/AIModule/Private/Navigation/CrowdManager.cpp
--- a/Engine/Source/Runtime/AIModule/Private/Navigation/CrowdManager.cpp
+++ b/Engine/Source/Runtime/AIModule/Private/Navigation/CrowdManager.cpp
@@ -186,6 +186,7 @@ void UCrowdManager::BeginDestroy()
 	delete TickHelper;
 #endif
 
+	DestroyCrowdManager();
 	Super::BeginDestroy();
 }
 
@@ -692,6 +693,11 @@ void UCrowdManager::GetAgentParams(const ICrowdAgentInterface* Agent, dtCrowdAge
 
 	AgentParams.radius = CylRadius;
 	AgentParams.height = CylHalfHeight * 2.0f;
+	AgentParams.avoidanceQueryMultiplier = 1.0f;
+	AgentParams.avoidanceGroup = Agent->GetCrowdAgentAvoidanceGroup();
+	AgentParams.groupsToAvoid = Agent->GetCrowdAgentGroupsToAvoid();
+	AgentParams.groupsToIgnore = Agent->GetCrowdAgentGroupsToIgnore();
+
 	// skip maxSpeed, it will be constantly updated in every tick
 	// skip maxAcceleration, we don't use Detour's movement code
 
@@ -703,7 +709,7 @@ void UCrowdManager::GetAgentParams(const ICrowdAgentInterface* Agent, dtCrowdAge
 		AgentParams.separationWeight = CrowdComponent->GetCrowdSeparationWeight();
 		AgentParams.obstacleAvoidanceType = CrowdComponent->GetCrowdAvoidanceQuality();
 		AgentParams.avoidanceQueryMultiplier = CrowdComponent->GetCrowdAvoidanceRangeMultiplier();
-	
+
 		if (CrowdComponent->IsCrowdSimulationEnabled())
 		{
 			AgentParams.updateFlags =
@@ -715,16 +721,6 @@ void UCrowdManager::GetAgentParams(const ICrowdAgentInterface* Agent, dtCrowdAge
 				(CrowdComponent->IsCrowdPathOffsetEnabled() ? DT_CROWD_OFFSET_PATH : 0) |
 				(CrowdComponent->IsCrowdSlowdownAtGoalEnabled() ? DT_CROWD_SLOWDOWN_AT_GOAL : 0);
 		}
-
-		AgentParams.avoidanceGroup = CrowdComponent->GetAvoidanceGroup();
-		AgentParams.groupsToAvoid = CrowdComponent->GetGroupsToAvoid();
-		AgentParams.groupsToIgnore = CrowdComponent->GetGroupsToIgnore();
-	}
-	else
-	{
-		AgentParams.avoidanceQueryMultiplier = 1.0f;
-		AgentParams.avoidanceGroup = 1;
-		AgentParams.groupsToAvoid = MAX_uint32;
 	}
 }
 
diff --git a/Engine/Source/Runtime/AnimGraphRuntime/Private/AnimNodes/AnimNode_BlendSpacePlayer.cpp b/Engine/Source/Runtime/AnimGraphRuntime/Private/AnimNodes/AnimNode_BlendSpacePlayer.cpp
--- a/Engine/Source/Runtime/AnimGraphRuntime/Private/AnimNodes/AnimNode_BlendSpacePlayer.cpp
+++ b/Engine/Source/Runtime/AnimGraphRuntime/Private/AnimNodes/AnimNode_BlendSpacePlayer.cpp
@@ -20,6 +20,29 @@ FAnimNode_BlendSpacePlayer::FAnimNode_BlendSpacePlayer()
 {
 }
 
+float FAnimNode_BlendSpacePlayer::GetCurrentAssetTime()
+{
+	if(const FBlendSampleData* HighestWeightedSample = GetHighestWeightedSample())
+	{
+		return HighestWeightedSample->Time;
+	}
+
+	// No sample
+	return 0.0f;
+}
+
+float FAnimNode_BlendSpacePlayer::GetCurrentAssetLength()
+{
+	if(const FBlendSampleData* HighestWeightedSample = GetHighestWeightedSample())
+	{
+		FBlendSample& Sample = BlendSpace->SampleData[HighestWeightedSample->SampleDataIndex];
+		return Sample.Animation->SequenceLength;
+	}
+
+	// No sample
+	return 0.0f;
+}
+
 void FAnimNode_BlendSpacePlayer::Initialize(const FAnimationInitializeContext& Context)
 {
 	FAnimNode_AssetPlayerBase::Initialize(Context);
@@ -112,6 +135,26 @@ UAnimationAsset* FAnimNode_BlendSpacePlayer::GetAnimAsset()
 	return BlendSpace;
 }
 
+const FBlendSampleData* FAnimNode_BlendSpacePlayer::GetHighestWeightedSample() const
+{
+	if(BlendSampleDataCache.Num() == 0)
+	{
+		return nullptr;
+	}
+
+	const FBlendSampleData* HighestSample = &BlendSampleDataCache[0];
+
+	for(int32 Idx = 1; Idx < BlendSampleDataCache.Num(); ++Idx)
+	{
+		if(BlendSampleDataCache[Idx].TotalWeight > HighestSample->TotalWeight)
+		{
+			HighestSample = &BlendSampleDataCache[Idx];
+		}
+	}
+
+	return HighestSample;
+}
+
 void FAnimNode_BlendSpacePlayer::Reinitialize()
 {
 	BlendSampleDataCache.Empty();
diff --git a/Engine/Source/Runtime/AnimGraphRuntime/Private/AnimNodes/AnimNode_SequenceEvaluator.cpp b/Engine/Source/Runtime/AnimGraphRuntime/Private/AnimNodes/AnimNode_SequenceEvaluator.cpp
--- a/Engine/Source/Runtime/AnimGraphRuntime/Private/AnimNodes/AnimNode_SequenceEvaluator.cpp
+++ b/Engine/Source/Runtime/AnimGraphRuntime/Private/AnimNodes/AnimNode_SequenceEvaluator.cpp
@@ -4,6 +4,16 @@
 #include "AnimNodes/AnimNode_SequenceEvaluator.h"
 #include "Animation/AnimInstanceProxy.h"
 
+float FAnimNode_SequenceEvaluator::GetCurrentAssetTime()
+{
+	return ExplicitTime;
+}
+
+float FAnimNode_SequenceEvaluator::GetCurrentAssetLength()
+{
+	return Sequence ? Sequence->SequenceLength : 0.0f;
+}
+
 /////////////////////////////////////////////////////
 // FAnimSequenceEvaluatorNode
 
diff --git a/Engine/Source/Runtime/AnimGraphRuntime/Public/AnimNodes/AnimNode_BlendSpacePlayer.h b/Engine/Source/Runtime/AnimGraphRuntime/Public/AnimNodes/AnimNode_BlendSpacePlayer.h
--- a/Engine/Source/Runtime/AnimGraphRuntime/Public/AnimNodes/AnimNode_BlendSpacePlayer.h
+++ b/Engine/Source/Runtime/AnimGraphRuntime/Public/AnimNodes/AnimNode_BlendSpacePlayer.h
@@ -56,6 +56,11 @@ protected:
 public:	
 	FAnimNode_BlendSpacePlayer();
 
+	// FAnimNode_AssetPlayerBase interface
+	virtual float GetCurrentAssetTime();
+	virtual float GetCurrentAssetLength();
+	// End of FAnimNode_AssetPlayerBase interface
+
 	// FAnimNode_Base interface
 	virtual void Initialize(const FAnimationInitializeContext& Context) override;
 	virtual void CacheBones(const FAnimationCacheBonesContext& Context) override;
@@ -74,4 +79,6 @@ protected:
 
 private:
 	void Reinitialize();
+
+	const FBlendSampleData* GetHighestWeightedSample() const;
 };
diff --git a/Engine/Source/Runtime/AnimGraphRuntime/Public/AnimNodes/AnimNode_SequenceEvaluator.h b/Engine/Source/Runtime/AnimGraphRuntime/Public/AnimNodes/AnimNode_SequenceEvaluator.h
--- a/Engine/Source/Runtime/AnimGraphRuntime/Public/AnimNodes/AnimNode_SequenceEvaluator.h
+++ b/Engine/Source/Runtime/AnimGraphRuntime/Public/AnimNodes/AnimNode_SequenceEvaluator.h
@@ -37,6 +37,11 @@ public:
 	{
 	}
 
+	// FAnimNode_AssetPlayerBase interface
+	virtual float GetCurrentAssetTime();
+	virtual float GetCurrentAssetLength();
+	// End of FAnimNode_AssetPlayerBase interface
+
 	// FAnimNode_Base interface
 	virtual void Initialize(const FAnimationInitializeContext& Context) override;
 	virtual void CacheBones(const FAnimationCacheBonesContext& Context) override;
diff --git a/Engine/Source/Runtime/Core/Public/Containers/Array.h b/Engine/Source/Runtime/Core/Public/Containers/Array.h
--- a/Engine/Source/Runtime/Core/Public/Containers/Array.h
+++ b/Engine/Source/Runtime/Core/Public/Containers/Array.h
@@ -3610,6 +3610,12 @@ public:
 		(Super&)*this = Other;
 	}
 
+	// Functions dependent on Add, Remove.
+	void AssignButKeepOwner( Super&& Other )
+	{
+		(Super&)*this = MoveTemp(Other);
+	}
+
 	int32 Add( const T& Item )
 	{
 		new(*this) T(Item);
diff --git a/Engine/Source/Runtime/Core/Public/UObject/FrameworkObjectVersion.h b/Engine/Source/Runtime/Core/Public/UObject/FrameworkObjectVersion.h
--- a/Engine/Source/Runtime/Core/Public/UObject/FrameworkObjectVersion.h
+++ b/Engine/Source/Runtime/Core/Public/UObject/FrameworkObjectVersion.h
@@ -12,6 +12,10 @@ struct CORE_API FFrameworkObjectVersion
 		// BodySetup's default instance collision profile is used by default when creating a new instance.
 		UseBodySetupCollisionProfile,
 
+		// Regenerate subgraph arrays correctly in animation blueprints to remove duplicates and add
+		// missing graphs that appear read only when edited
+		AnimBlueprintSubgraphFix,
+
 		// -----<new versions can be added above this line>-------------------------------------------------
 		VersionPlusOne,
 		LatestVersion = VersionPlusOne - 1
diff --git a/Engine/Source/Runtime/Core/Public/UObject/WeakObjectPtrTemplates.h b/Engine/Source/Runtime/Core/Public/UObject/WeakObjectPtrTemplates.h
--- a/Engine/Source/Runtime/Core/Public/UObject/WeakObjectPtrTemplates.h
+++ b/Engine/Source/Runtime/Core/Public/UObject/WeakObjectPtrTemplates.h
@@ -117,11 +117,19 @@ public:
 	 * @param bEvenIfPendingKill, if this is true, pendingkill objects are considered valid	
 	 * @return NULL if this object is gone or the weak pointer was NULL, otherwise a valid uobject pointer
 	**/
-	FORCEINLINE T* Get(bool bEvenIfPendingKill = false) const
+	FORCEINLINE T* Get(bool bEvenIfPendingKill) const
 	{
 		return (T*)TWeakObjectPtrBase::Get(bEvenIfPendingKill);
 	}
 
+	/**  
+	 * Dereference the weak pointer. This is an optimized version implying bEvenIfPendingKill=false.
+	 */
+	FORCEINLINE T* Get(/*bool bEvenIfPendingKill = false*/) const
+	{
+		return (T*)TWeakObjectPtrBase::Get();
+	}
+
 	/** Deferences the weak pointer even if its marked RF_Unreachable. This is needed to resolve weak pointers during GC (such as ::AddReferenceObjects) */
 	FORCEINLINE T* GetEvenIfUnreachable() const
 	{
@@ -151,11 +159,20 @@ public:
 	 *							UObject is gone forever (@return false) or if it is still there (@return true, no object flags checked).
 	 * @return true if Get() would return a valid non-null pointer
 	**/
-	FORCEINLINE bool IsValid(bool bEvenIfPendingKill = false, bool bThreadsafeTest = false) const
+	FORCEINLINE bool IsValid(bool bEvenIfPendingKill, bool bThreadsafeTest = false) const
 	{
 		return TWeakObjectPtrBase::IsValid(bEvenIfPendingKill, bThreadsafeTest);
 	}
 
+	/**
+	 * Test if this points to a live UObject. This is an optimized version implying bEvenIfPendingKill=false, bThreadsafeTest=false.
+	 * @return true if Get() would return a valid non-null pointer
+	 */
+	FORCEINLINE bool IsValid(/*bool bEvenIfPendingKill = false, bool bThreadsafeTest = false*/) const
+	{
+		return TWeakObjectPtrBase::IsValid();
+	}
+
 	/**  
 	 * Slightly different than !IsValid(), returns true if this used to point to a UObject, but doesn't any more and has not been assigned or reset in the mean time.
 	 * @param bIncludingIfPendingKill, if this is true, pendingkill objects are considered stale
diff --git a/Engine/Source/Runtime/CoreUObject/Private/UObject/WeakObjectPtr.cpp b/Engine/Source/Runtime/CoreUObject/Private/UObject/WeakObjectPtr.cpp
--- a/Engine/Source/Runtime/CoreUObject/Private/UObject/WeakObjectPtr.cpp
+++ b/Engine/Source/Runtime/CoreUObject/Private/UObject/WeakObjectPtr.cpp
@@ -41,6 +41,11 @@ bool FWeakObjectPtr::IsValid(bool bEvenIfPendingKill, bool bThreadsafeTest) cons
 	return Internal_IsValid(bEvenIfPendingKill, bThreadsafeTest);
 }
 
+bool FWeakObjectPtr::IsValid() const
+{
+	// Using literals here allows the optimizer to remove branches later down the chain.
+	return Internal_IsValid(false, false);
+}
 
 bool FWeakObjectPtr::IsStale(bool bEvenIfPendingKill, bool bThreadsafeTest) const
 {
diff --git a/Engine/Source/Runtime/CoreUObject/Public/UObject/WeakObjectPtr.h b/Engine/Source/Runtime/CoreUObject/Public/UObject/WeakObjectPtr.h
--- a/Engine/Source/Runtime/CoreUObject/Public/UObject/WeakObjectPtr.h
+++ b/Engine/Source/Runtime/CoreUObject/Public/UObject/WeakObjectPtr.h
@@ -78,7 +78,7 @@ public:
 	{
 		return 
 			(ObjectIndex == Other.ObjectIndex && ObjectSerialNumber == Other.ObjectSerialNumber) ||
-			(Get() == Other.Get());
+			(!IsValid() && !Other.IsValid());
 	}
 
 	/**  
@@ -89,7 +89,7 @@ public:
 	{
 		return 
 			(ObjectIndex != Other.ObjectIndex || ObjectSerialNumber != Other.ObjectSerialNumber) &&
-			(Get() != Other.Get());
+			(IsValid() || Other.IsValid());
 	}
 
 
@@ -116,7 +116,13 @@ public:
 	 *							UObject is gone forever (@return false) or if it is still there (@return true, no object flags checked).
 	 * @return true if Get() would return a valid non-null pointer, if bThreadsafeTest == true then @see @param bThreadsafeTest
 	**/
-	COREUOBJECT_API bool IsValid(bool bEvenIfPendingKill = false, bool bThreadsafeTest = false) const;
+	COREUOBJECT_API bool IsValid(bool bEvenIfPendingKill, bool bThreadsafeTest = false) const;
+
+	/**
+	 * Test if this points to a live UObject. This is an optimized version implying bEvenIfPendingKill=false, bThreadsafeTest=false.
+	 * @return true if Get() would return a valid non-null pointer.
+	 */
+	COREUOBJECT_API bool IsValid(/*bool bEvenIfPendingKill = false, bool bThreadsafeTest = false*/) const;
 
 	/**  
 	 * Slightly different than !IsValid(), returns true if this used to point to a UObject, but doesn't any more and has not been assigned or reset in the mean time.
diff --git a/Engine/Source/Runtime/Engine/Classes/Animation/AnimBlueprint.h b/Engine/Source/Runtime/Engine/Classes/Animation/AnimBlueprint.h
--- a/Engine/Source/Runtime/Engine/Classes/Animation/AnimBlueprint.h
+++ b/Engine/Source/Runtime/Engine/Classes/Animation/AnimBlueprint.h
@@ -118,6 +118,8 @@ class ENGINE_API UAnimBlueprint : public UBlueprint
 
 	virtual void PostLoad() override;
 
+	virtual void Serialize(FArchive& Ar) override;
+
 protected:
 	// Broadcast when an override is changed, allowing derived blueprints to be updated
 	FOnOverrideChangedMulticaster OnOverrideChanged;
diff --git a/Engine/Source/Runtime/Engine/Classes/Animation/AnimNode_AssetPlayerBase.h b/Engine/Source/Runtime/Engine/Classes/Animation/AnimNode_AssetPlayerBase.h
--- a/Engine/Source/Runtime/Engine/Classes/Animation/AnimNode_AssetPlayerBase.h
+++ b/Engine/Source/Runtime/Engine/Classes/Animation/AnimNode_AssetPlayerBase.h
@@ -57,6 +57,10 @@ struct ENGINE_API FAnimNode_AssetPlayerBase : public FAnimNode_Base
 	// Create a tick record for this node
 	void CreateTickRecordForNode(const FAnimationUpdateContext& Context, UAnimSequenceBase* Sequence, bool bLooping, float PlayRate);
 
+	// Functions to report data to getters, this is required for all asset players (but can't be pure abstract because of struct instantiation generated code).
+	virtual float GetCurrentAssetLength() { return 0.0f; }
+	virtual float GetCurrentAssetTime() { return 0.0f; }
+
 protected:
 
 	/** Last encountered blendweight for this node */
diff --git a/Engine/Source/Runtime/Engine/Classes/Animation/AnimNode_SequencePlayer.h b/Engine/Source/Runtime/Engine/Classes/Animation/AnimNode_SequencePlayer.h
--- a/Engine/Source/Runtime/Engine/Classes/Animation/AnimNode_SequencePlayer.h
+++ b/Engine/Source/Runtime/Engine/Classes/Animation/AnimNode_SequencePlayer.h
@@ -38,6 +38,11 @@ public:
 	{
 	}
 
+	// FAnimNode_AssetPlayerBase interface
+	virtual float GetCurrentAssetTime();
+	virtual float GetCurrentAssetLength();
+	// End of FAnimNode_AssetPlayerBase interface
+
 	// FAnimNode_Base interface
 	virtual void Initialize(const FAnimationInitializeContext& Context) override;
 	virtual void CacheBones(const FAnimationCacheBonesContext& Context) override;
diff --git a/Engine/Source/Runtime/Engine/Classes/Animation/AnimationAsset.h b/Engine/Source/Runtime/Engine/Classes/Animation/AnimationAsset.h
--- a/Engine/Source/Runtime/Engine/Classes/Animation/AnimationAsset.h
+++ b/Engine/Source/Runtime/Engine/Classes/Animation/AnimationAsset.h
@@ -11,6 +11,7 @@
 #include "AnimInterpFilter.h"
 #include "Animation/Skeleton.h"
 #include "Engine/SkeletalMesh.h"
+#include "Interfaces/Interface_AssetUserData.h"
 #include "AnimationAsset.generated.h"
 
 namespace MarkerIndexSpecialValues
@@ -672,7 +673,7 @@ struct FAnimationGroupReference
 };
 
 UCLASS(abstract, MinimalAPI)
-class UAnimationAsset : public UObject
+class UAnimationAsset : public UObject, public IInterface_AssetUserData
 {
 	GENERATED_UCLASS_BODY()
 
@@ -691,6 +692,11 @@ private:
 	UPROPERTY(Category=MetaData, instanced, EditAnywhere)
 	TArray<class UAnimMetaData*> MetaData;
 
+protected:
+	/** Array of user data stored with the asset */
+	UPROPERTY(EditAnywhere, AdvancedDisplay, Instanced, Category = Animation)
+	TArray<UAssetUserData*> AssetUserData;
+
 public:
 	/** Advances the asset player instance 
 	 * 
@@ -755,6 +761,13 @@ public:
 	/** Return a list of unique marker names for blending compatibility */
 	ENGINE_API virtual TArray<FName>* GetUniqueMarkerNames() { return NULL; }
 
+	//~ Begin IInterface_AssetUserData Interface
+	virtual void AddAssetUserData(UAssetUserData* InUserData) override;
+	virtual void RemoveUserDataOfClass(TSubclassOf<UAssetUserData> InUserDataClass) override;
+	virtual UAssetUserData* GetAssetUserDataOfClass(TSubclassOf<UAssetUserData> InUserDataClass) override;
+	virtual const TArray<UAssetUserData*>* GetAssetUserDataArray() const override;
+	//~ End IInterface_AssetUserData Interface
+
 	/**
 	* return true if this is valid additive animation
 	* false otherwise
diff --git a/Engine/Source/Runtime/Engine/Classes/Components/PrimitiveComponent.h b/Engine/Source/Runtime/Engine/Classes/Components/PrimitiveComponent.h
--- a/Engine/Source/Runtime/Engine/Classes/Components/PrimitiveComponent.h
+++ b/Engine/Source/Runtime/Engine/Classes/Components/PrimitiveComponent.h
@@ -1469,15 +1469,28 @@ protected:
 public:
 	virtual bool IsWorldGeometry() const override;
 
-
+	/** Returns the form of collision for this component */
+	UFUNCTION(BlueprintCallable, Category="Collision")
 	virtual ECollisionEnabled::Type GetCollisionEnabled() const override;
 
+	/** Utility to see if there is any form of collision (query or physics) enabled on this component. */
+	UFUNCTION(BlueprintCallable, meta=(DisplayName="Is Collision Enabled"), Category="Collision")
+	bool K2_IsCollisionEnabled() const;
+
+	/** Utility to see if there is any query collision enabled on this component. */
+	UFUNCTION(BlueprintCallable, meta=(DisplayName="Is Query Collision Enabled"), Category="Collision")
+	bool K2_IsQueryCollisionEnabled() const;
+
+	/** Utility to see if there is any physics collision enabled on this component. */
+	UFUNCTION(BlueprintCallable, meta=(DisplayName="Is Physics Collision Enabled"), Category="Collision")
+	bool K2_IsPhysicsCollisionEnabled() const;
+
 	/** Gets the response type given a specific channel */
-	UFUNCTION(BlueprintCallable, Category="Physics")
+	UFUNCTION(BlueprintCallable, Category="Collision")
 	virtual ECollisionResponse GetCollisionResponseToChannel(ECollisionChannel Channel) const override;
 
 	/** Gets the collision object type */
-	UFUNCTION(BlueprintCallable, Category="Physics")
+	UFUNCTION(BlueprintCallable, Category="Collision")
 	virtual ECollisionChannel GetCollisionObjectType() const override;
 
 	virtual const FCollisionResponseContainer& GetCollisionResponseToChannels() const override;
@@ -1914,3 +1927,18 @@ FORCEINLINE_DEBUGGABLE const TArray<FOverlapInfo>& UPrimitiveComponent::GetOverl
 {
 	return OverlappingComponents;
 }
+
+FORCEINLINE_DEBUGGABLE bool UPrimitiveComponent::K2_IsCollisionEnabled() const
+{
+	return IsCollisionEnabled();
+}
+
+FORCEINLINE_DEBUGGABLE bool UPrimitiveComponent::K2_IsQueryCollisionEnabled() const
+{
+	return IsQueryCollisionEnabled();
+}
+
+FORCEINLINE_DEBUGGABLE bool UPrimitiveComponent::K2_IsPhysicsCollisionEnabled() const
+{
+	return IsPhysicsCollisionEnabled();
+}
diff --git a/Engine/Source/Runtime/Engine/Classes/Components/SkeletalMeshComponent.h b/Engine/Source/Runtime/Engine/Classes/Components/SkeletalMeshComponent.h
--- a/Engine/Source/Runtime/Engine/Classes/Components/SkeletalMeshComponent.h
+++ b/Engine/Source/Runtime/Engine/Classes/Components/SkeletalMeshComponent.h
@@ -594,41 +594,41 @@ public:
 	class UAnimInstance * GetAnimInstance() const;
 
 	/** Below are the interface to control animation when animation mode, not blueprint mode **/
-	UFUNCTION(BlueprintCallable, Category="Components|SkeletalMesh")
+	UFUNCTION(BlueprintCallable, Category = "Components|Animation", meta = (Keywords = "Animation"))
 	void SetAnimationMode(EAnimationMode::Type InAnimationMode);
 	
-	UFUNCTION(BlueprintCallable, Category="Components|SkeletalMesh")
+	UFUNCTION(BlueprintCallable, Category = "Components|Animation", meta = (Keywords = "Animation"))
 	EAnimationMode::Type GetAnimationMode() const;
 
-	UFUNCTION(BlueprintCallable, Category="Components|SkeletalMesh")
+	UFUNCTION(BlueprintCallable, Category = "Components|Animation", meta = (Keywords = "Animation"))
 	void PlayAnimation(class UAnimationAsset* NewAnimToPlay, bool bLooping);
 
-	UFUNCTION(BlueprintCallable, Category="Components|SkeletalMesh")
+	UFUNCTION(BlueprintCallable, Category = "Components|Animation", meta = (Keywords = "Animation"))
 	void SetAnimation(class UAnimationAsset* NewAnimToPlay);
 
 	// @todo block this until we support vertex animation 
 //	UFUNCTION(BlueprintCallable, Category="Components|SkeletalMesh")
 	void SetVertexAnimation(class UVertexAnimation* NewVertexAnimToPlay);
 
-	UFUNCTION(BlueprintCallable, Category="Components|SkeletalMesh")
+	UFUNCTION(BlueprintCallable, Category = "Components|Animation", meta = (Keywords = "Animation"))
 	void Play(bool bLooping);
 
-	UFUNCTION(BlueprintCallable, Category="Components|SkeletalMesh")
+	UFUNCTION(BlueprintCallable, Category = "Components|Animation", meta = (Keywords = "Animation"))
 	void Stop();
 
-	UFUNCTION(BlueprintCallable, Category="Components|SkeletalMesh")
+	UFUNCTION(BlueprintCallable, Category = "Components|Animation", meta = (Keywords = "Animation"))
 	bool IsPlaying() const;
 
-	UFUNCTION(BlueprintCallable, Category="Components|SkeletalMesh")
+	UFUNCTION(BlueprintCallable, Category = "Components|Animation", meta = (Keywords = "Animation"))
 	void SetPosition(float InPos, bool bFireNotifies = true);
 
-	UFUNCTION(BlueprintCallable, Category="Components|SkeletalMesh")
+	UFUNCTION(BlueprintCallable, Category = "Components|Animation", meta = (Keywords = "Animation"))
 	float GetPosition() const;
 
-	UFUNCTION(BlueprintCallable, Category="Components|SkeletalMesh")
+	UFUNCTION(BlueprintCallable, Category = "Components|Animation", meta = (Keywords = "Animation"))
 	void SetPlayRate(float Rate);
 
-	UFUNCTION(BlueprintCallable, Category="Components|SkeletalMesh")
+	UFUNCTION(BlueprintCallable, Category = "Components|Animation", meta = (Keywords = "Animation"))
 	float GetPlayRate() const;
 
 	/**
diff --git a/Engine/Source/Runtime/Engine/Classes/Engine/AssetUserData.h b/Engine/Source/Runtime/Engine/Classes/Engine/AssetUserData.h
--- a/Engine/Source/Runtime/Engine/Classes/Engine/AssetUserData.h
+++ b/Engine/Source/Runtime/Engine/Classes/Engine/AssetUserData.h
@@ -8,7 +8,7 @@
 /**
  * Object that can be subclassed to store custom data on Unreal asset objects.
  */
-UCLASS(DefaultToInstanced)
+UCLASS(DefaultToInstanced, abstract, editinlinenew)
 class ENGINE_API UAssetUserData
 	: public UObject
 {
diff --git a/Engine/Source/Runtime/Engine/Classes/Engine/SkeletalMeshReductionSettings.h b/Engine/Source/Runtime/Engine/Classes/Engine/SkeletalMeshReductionSettings.h
new file mode 100644
--- /dev/null
+++ b/Engine/Source/Runtime/Engine/Classes/Engine/SkeletalMeshReductionSettings.h
@@ -0,0 +1,61 @@
+// Copyright 1998-2016 Epic Games, Inc. All Rights Reserved.
+
+#pragma once
+
+#include "Engine/DeveloperSettings.h"
+#include "Engine/SkeletalMesh.h"
+
+#include "SkeletalMeshReductionSettings.generated.h"
+
+USTRUCT()
+struct FSkeletalMeshLODGroupSettings
+{
+	GENERATED_USTRUCT_BODY()
+
+	FSkeletalMeshLODGroupSettings() {}
+	
+	/** Get Skeletal mesh optimizations setting structure for the given LOD level */
+	ENGINE_API FSkeletalMeshOptimizationSettings GetSettings() const;
+
+	/** Get the correct screen size for the given LOD level */
+	ENGINE_API const float GetScreenSize() const;
+protected:
+	/** FSkeletalMeshLODSettings initializes group entries. */
+	friend class USkeletalMeshReductionSettings;
+
+	/** The screen sizes to use for the respective LOD level */
+	UPROPERTY(EditAnywhere, Category = Reduction)
+	float ScreenSize;
+	
+	/** The optimization settings to use for the respective LOD level */
+	UPROPERTY(EditAnywhere, Category = Reduction)
+	FSkeletalMeshOptimizationSettings OptimizationSettings;
+};
+
+UCLASS(config=Engine, defaultconfig, MinimalAPI) 
+class USkeletalMeshReductionSettings : public UObject
+{
+	GENERATED_UCLASS_BODY()
+protected:
+	UPROPERTY(globalconfig)
+	TArray<FSkeletalMeshLODGroupSettings> Settings;
+public:
+
+	/** Accessor and initializer **/
+	ENGINE_API static USkeletalMeshReductionSettings* Get();
+
+	/** Validates and initializes data retrieved from the ini file */
+	ENGINE_API void Initialize();
+
+	/** Retrieves the Skeletal mesh LOD group settings for the given name */
+	ENGINE_API const FSkeletalMeshLODGroupSettings& GetDefaultSettingsForLODLevel(const int32 LODIndex) const;
+
+	/** Returns whether or not valid settings were retrieved from the ini file */
+	ENGINE_API const bool HasValidSettings() const;
+
+	/** Returns the number of settings parsed from the ini file */
+	ENGINE_API int32 GetNumberOfSettings() const;
+private:
+	/** Flag for whether or not valid settings were found parsed form the ini file */
+	bool bValidSettings;
+};
diff --git a/Engine/Source/Runtime/Engine/Classes/Engine/Texture.h b/Engine/Source/Runtime/Engine/Classes/Engine/Texture.h
--- a/Engine/Source/Runtime/Engine/Classes/Engine/Texture.h
+++ b/Engine/Source/Runtime/Engine/Classes/Engine/Texture.h
@@ -6,6 +6,7 @@
 #include "TextureResource.h"
 #include "MaterialShared.h"
 #include "RenderResource.h"
+#include "Interfaces/Interface_AssetUserData.h"
 #include "Texture.generated.h"
 
 // This needs to be mirrored in EditorFactories.cpp.
@@ -339,7 +340,7 @@ struct FTexturePlatformData
 };
 
 UCLASS(abstract, MinimalAPI, BlueprintType)
-class UTexture : public UObject
+class UTexture : public UObject, public IInterface_AssetUserData
 {
 	GENERATED_UCLASS_BODY()
 
@@ -509,6 +510,11 @@ public:
 	UPROPERTY(transient)
 	uint32 bUseCinematicMipLevels:1;
 
+protected:
+	/** Array of user data stored with the asset */
+	UPROPERTY(EditAnywhere, AdvancedDisplay, Instanced, Category = Texture)
+	TArray<UAssetUserData*> AssetUserData;
+
 private:
 	/** Cached combined group and texture LOD bias to use.	*/
 	UPROPERTY(transient)
@@ -611,6 +617,13 @@ public:
 	 */
 	ENGINE_API void SerializeCookedPlatformData(class FArchive& Ar);
 
+	//~ Begin IInterface_AssetUserData Interface
+	ENGINE_API virtual void AddAssetUserData(UAssetUserData* InUserData) override;
+	ENGINE_API virtual void RemoveUserDataOfClass(TSubclassOf<UAssetUserData> InUserDataClass) override;
+	ENGINE_API virtual UAssetUserData* GetAssetUserDataOfClass(TSubclassOf<UAssetUserData> InUserDataClass) override;
+	ENGINE_API virtual const TArray<UAssetUserData*>* GetAssetUserDataArray() const override;
+	//~ End IInterface_AssetUserData Interface
+
 #if WITH_EDITOR
 	/**
 	 * Caches platform data for the texture.
diff --git a/Engine/Source/Runtime/Engine/Classes/PhysicalMaterials/PhysicalMaterial.h b/Engine/Source/Runtime/Engine/Classes/PhysicalMaterials/PhysicalMaterial.h
--- a/Engine/Source/Runtime/Engine/Classes/PhysicalMaterials/PhysicalMaterial.h
+++ b/Engine/Source/Runtime/Engine/Classes/PhysicalMaterials/PhysicalMaterial.h
@@ -43,7 +43,7 @@ class ENGINE_API UPhysicalMaterial : public UObject
 	//
 	
 	/** Friction value of surface, controls how easily things can slide on this surface */
-	UPROPERTY(EditAnywhere, BlueprintReadOnly, Category=PhysicalMaterial)
+	UPROPERTY(EditAnywhere, BlueprintReadOnly, Category=PhysicalMaterial, meta=(ClampMin=0))
 	float Friction;
 
 	/** Friction combine mode, controls how friction is computed for multiple materials. */
@@ -54,8 +54,8 @@ class ENGINE_API UPhysicalMaterial : public UObject
 	UPROPERTY(EditAnywhere, BlueprintReadWrite, Category = PhysicalMaterial)
 	bool bOverrideFrictionCombineMode;
 
-	/** Restitution or 'bounciness of this surface. */
-	UPROPERTY(EditAnywhere, BlueprintReadOnly, Category = PhysicalMaterial)
+	/** Restitution or 'bounciness' of this surface, between 0 (no bounce) and 1 (outgoing velocity is same as incoming). */
+	UPROPERTY(EditAnywhere, BlueprintReadOnly, Category = PhysicalMaterial,  meta=(ClampMin=0, ClampMax=1))
 	float Restitution;
 
 	/** Restitution combine mode, controls how restitution is computed for multiple materials. */
@@ -71,7 +71,7 @@ class ENGINE_API UPhysicalMaterial : public UObject
 	//
 	
 	/** Used with the shape of the object to calculate its mass properties. The higher the number, the heavier the object. g per cubic cm. */
-	UPROPERTY(EditAnywhere, BlueprintReadOnly, Category = PhysicalMaterial)
+	UPROPERTY(EditAnywhere, BlueprintReadOnly, Category = PhysicalMaterial, meta=(ClampMin=0))
 	float Density;
 
 	/** 
@@ -79,7 +79,7 @@ class ENGINE_API UPhysicalMaterial : public UObject
 	 *	In actuality, larger objects do not tend to be solid, and become more like 'shells' (e.g. a car is not a solid piece of metal).
 	 *	Values are clamped to 1 or less.
 	 */
-	UPROPERTY(EditAnywhere, BlueprintReadOnly, Category = Advanced)
+	UPROPERTY(EditAnywhere, BlueprintReadOnly, Category = Advanced, meta=(ClampMin=0.1, ClampMax=1))
 	float RaiseMassToPower;
 
 	/** How much to scale the damage threshold by on any destructible we are applied to */
diff --git a/Engine/Source/Runtime/Engine/Classes/Sound/SoundMix.h b/Engine/Source/Runtime/Engine/Classes/Sound/SoundMix.h
--- a/Engine/Source/Runtime/Engine/Classes/Sound/SoundMix.h
+++ b/Engine/Source/Runtime/Engine/Classes/Sound/SoundMix.h
@@ -12,59 +12,80 @@ struct FAudioEQEffect
 	/* Start time of effect */
 	double RootTime;
 
-	/* High frequency filter cutoff frequency (Hz) */
-	UPROPERTY(EditAnywhere, Category=HighPass )
-	float HFFrequency;
-
-	/* High frequency gain - 0.0 is silent, 1.0 is full volume. */
-	UPROPERTY(EditAnywhere, Category=HighPass )
-	float HFGain;
-
-	/* Middle (band) frequency filter cutoff frequency (Hz). */
-	UPROPERTY(EditAnywhere, Category=BandPass )
-	float MFCutoffFrequency;
-
-	/* Middle (band) frequency filter bandwidth frequency (Hz) - Range (0.1 to 2.0). */
-	UPROPERTY(EditAnywhere, Category=BandPass )
-	float MFBandwidth;
-
-	/* Middle (band) frequency filter gain - 0.0 is silent, 1.0 is full volume. */
-	UPROPERTY(EditAnywhere, Category=BandPass )
-	float MFGain;
-
-	/* Low frequency filter cutoff frequency (Hz) */
-	UPROPERTY(EditAnywhere, Category=LowPass )
-	float LFFrequency;
-
-	/* Low frequency filter gain - 0.0 is silent, 1.0 is full volume. */
-	UPROPERTY(EditAnywhere, Category=LowPass )
-	float LFGain;
-
-
-
-		// Cannot use strcutdefaultproperties here as this class is a member of a native class
-		FAudioEQEffect()
-		:	RootTime( 0.0 )
-		,	HFFrequency( DEFAULT_HIGH_FREQUENCY )
-		,	HFGain( 1.0f )
-		,	MFCutoffFrequency( DEFAULT_MID_FREQUENCY )
-		,	MFBandwidth( 1.0f )
-		,	MFGain( 1.0f )
-		,	LFFrequency( DEFAULT_LOW_FREQUENCY )
-		,	LFGain( 1.0f )
-		{
-		}
-
-		/** 
-		 * Interpolate EQ settings based on time
-		 */
-		void Interpolate( float InterpValue, const FAudioEQEffect& Start, const FAudioEQEffect& End );
-		
-		/** 
-		 * Validate all settings are in range
-		 */
-		void ClampValues( void );
+	/** Center frequency in Hz for band 0 */
+	UPROPERTY(EditAnywhere, Category = Band0, meta = (ClampMin = "0.0", ClampMax = "20000.0", UIMin = "0.0", UIMax = "20000.0"))
+	float FrequencyCenter0;
 	
+	/** Boost/cut of band 0 */
+	UPROPERTY(EditAnywhere, Category = Band0, meta = (ClampMin = "0.0", ClampMax = "10.0", UIMin = "0.0", UIMax = "10.0"))
+	float Gain0;
+
+	/** Bandwidth of band 0. Region is center frequency +/- Bandwidth /2 */
+	UPROPERTY(EditAnywhere, Category = Band0, meta = (ClampMin = "0.0", ClampMax = "2.0", UIMin = "0.0", UIMax = "2.0"))
+	float Bandwidth0;
+
+	/** Center frequency in Hz for band 1 */
+	UPROPERTY(EditAnywhere, Category = Band0, meta = (ClampMin = "0.0", ClampMax = "20000.0", UIMin = "0.0", UIMax = "20000.0"))
+	float FrequencyCenter1;
+
+	/** Boost/cut of band 1 */
+	UPROPERTY(EditAnywhere, Category = Band0, meta = (ClampMin = "0.0", ClampMax = "10.0", UIMin = "0.0", UIMax = "10.0"))
+	float Gain1;
+
+	/** Bandwidth of band 1. Region is center frequency +/- Bandwidth /2 */
+	UPROPERTY(EditAnywhere, Category = Band0, meta = (ClampMin = "0.0", ClampMax = "2.0", UIMin = "0.0", UIMax = "2.0"))
+	float Bandwidth1;
+
+	/** Center frequency in Hz for band 2 */
+	UPROPERTY(EditAnywhere, Category = Band0, meta = (ClampMin = "0.0", ClampMax = "20000.0", UIMin = "0.0", UIMax = "20000.0"))
+	float FrequencyCenter2;
+
+	/** Boost/cut of band 2 */
+	UPROPERTY(EditAnywhere, Category = Band0, meta = (ClampMin = "0.0", ClampMax = "10.0", UIMin = "0.0", UIMax = "10.0"))
+	float Gain2;
+
+	/** Bandwidth of band 2. Region is center frequency +/- Bandwidth /2 */
+	UPROPERTY(EditAnywhere, Category = Band0, meta = (ClampMin = "0.0", ClampMax = "2.0", UIMin = "0.0", UIMax = "2.0"))
+	float Bandwidth2;
+
+	/** Center frequency in Hz for band 3 */
+	UPROPERTY(EditAnywhere, Category = Band0, meta = (ClampMin = "0.0", ClampMax = "20000.0", UIMin = "0.0", UIMax = "20000.0"))
+	float FrequencyCenter3;
+
+	/** Boost/cut of band 3 */
+	UPROPERTY(EditAnywhere, Category = Band0, meta = (ClampMin = "0.0", ClampMax = "10.0", UIMin = "0.0", UIMax = "10.0"))
+	float Gain3;
+
+	/** Bandwidth of band 3. Region is center frequency +/- Bandwidth /2 */
+	UPROPERTY(EditAnywhere, Category = Band0, meta = (ClampMin = "0.0", ClampMax = "2.0", UIMin = "0.0", UIMax = "2.0"))
+	float Bandwidth3;
+
+	FAudioEQEffect()
+		: RootTime(0.0f)
+		, FrequencyCenter0(600.0f)
+		, Gain0(1.0f)
+		, Bandwidth0(1.0f)
+		, FrequencyCenter1(1000.0f)
+		, Gain1(1.0f)
+		, Bandwidth1(1.0f)
+		, FrequencyCenter2(2000.0f)
+		, Gain2(1.0f)
+		, Bandwidth2(1.0f)
+		, FrequencyCenter3(10000.0f)
+		, Gain3(1.0f)
+		, Bandwidth3(1.0f)
+	{}
+
+	/** 
+	* Interpolate EQ settings based on time
+	*/
+	void Interpolate(float InterpValue, const FAudioEQEffect& Start, const FAudioEQEffect& End );
+		
+	/** 
+	* Clamp all settings in range
+	*/
+	void ClampValues();
+
 };
 
 /**
@@ -95,9 +116,7 @@ struct FSoundClassAdjuster
 	UPROPERTY(EditAnywhere, Category=SoundClassAdjuster )
 	float VoiceCenterChannelVolumeAdjuster;
 
-
-
-		FSoundClassAdjuster()
+	FSoundClassAdjuster()
 		: SoundClassObject(NULL)
 		, VolumeAdjuster(1)
 		, PitchAdjuster(1)
@@ -143,6 +162,12 @@ class USoundMix : public UObject
 	UPROPERTY(EditAnywhere, Category=SoundMix )
 	float FadeOutTime;
 
+#if WITH_EDITORONLY_DATA
+	/** Transient property used to trigger real-time updates of the active EQ filter for editor previewing */
+	UPROPERTY(transient)
+	uint32 bChanged:1;
+#endif
+
 #if WITH_EDITOR
 	bool CausesPassiveDependencyLoop(TArray<USoundClass*>& ProblemClasses) const;
 #endif // WITH_EDITOR
diff --git a/Engine/Source/Runtime/Engine/Private/ActiveSound.cpp b/Engine/Source/Runtime/Engine/Private/ActiveSound.cpp
--- a/Engine/Source/Runtime/Engine/Private/ActiveSound.cpp
+++ b/Engine/Source/Runtime/Engine/Private/ActiveSound.cpp
@@ -19,7 +19,7 @@ FActiveSound::FActiveSound()
 	, ConcurrencySettings(nullptr)
 	, SoundClassOverride(nullptr)
 	, bHasCheckedOcclusion(false)
-	, bIsOccluded(false)
+	, bIsTraceDelegateBound(false)
 	, bAllowSpatialization(true)
 	, bHasAttenuationSettings(false)
 	, bShouldRemainActiveIfDropped(false)
@@ -42,7 +42,9 @@ FActiveSound::FActiveSound()
 	, bWarnedAboutOrphanedLooping(false)
 #endif
 	, bEnableLowPassFilter(false)
-	, bOcclusionAsyncTrace(false)
+	, bOcclusionAsyncTrace(true)
+	, bIsOccluded(false)
+	, bAsyncOcclusionPending(false)
 	, UserIndex(0)
 	, PlaybackTime(0.f)
 	, RequestedStartTime(0.f)
@@ -74,13 +76,12 @@ FActiveSound::FActiveSound()
 	, bIsAudible(true)
 	, ClosestListenerPtr(nullptr)
 {
-	// Bind our async occlusion trace delegate
-	OcclusionTraceDelegate.BindRaw(this, &FActiveSound::OcclusionTraceDone);
 }
 
 FActiveSound::~FActiveSound()
 {
 	ensureMsgf(WaveInstances.Num() == 0, TEXT("Destroyed an active sound that had active wave instances."));
+	check(CanDelete());
 }
 
 FArchive& operator<<( FArchive& Ar, FActiveSound* ActiveSound )
@@ -317,8 +318,6 @@ void FActiveSound::Stop()
 	WaveInstances.Empty();
 
 	AudioDevice->RemoveActiveSound(this);
-
-	delete this;
 }
 
 FWaveInstance* FActiveSound::FindWaveInstance( const UPTRINT WaveInstanceHash )
@@ -370,7 +369,18 @@ void FActiveSound::UpdateOcclusion(const FAttenuationSettings* AttenuationSettin
 
 void FActiveSound::OcclusionTraceDone(const FTraceHandle& TraceHandle, FTraceDatum& TraceDatum)
 {
-	bIsOccluded = (TraceDatum.OutHits.Num() > 0);
+	// Look for any results that resulted in a blocking hit
+	bIsOccluded = false;
+	for (const FHitResult& HitResult : TraceDatum.OutHits)
+	{
+		if (HitResult.bBlockingHit)
+		{
+			bIsOccluded = true;
+			break;
+		}
+	}
+
+	bAsyncOcclusionPending = false;
 }
 
 void FActiveSound::CheckOcclusion(const FVector ListenerLocation, const FVector SoundLocation, const FAttenuationSettings* AttenuationSettingsPtr)
@@ -397,11 +407,23 @@ void FActiveSound::CheckOcclusion(const FVector ListenerLocation, const FVector
 
 		if (bOcclusionAsyncTrace)
 		{
-			WorldPtr->AsyncLineTraceByChannel(SoundLocation, ListenerLocation, ECC_Visibility, Params, FCollisionResponseParams::DefaultResponseParam, &OcclusionTraceDelegate);
-		}
-		else
-		{
-			bIsOccluded = WorldPtr->LineTraceTestByChannel(SoundLocation, ListenerLocation, ECC_Visibility, Params);
+			// Check if we've not already bound our trace delegate
+			if (!bIsTraceDelegateBound)
+			{
+				bIsTraceDelegateBound = true;
+
+				// Bind our async occlusion trace delegate (so next update we'll have it bound)
+				OcclusionTraceDelegate.BindRaw(this, &FActiveSound::OcclusionTraceDone);
+
+				// Only do async occlusion trace if we've already made one. The first trace must be synchronous to avoid issues with sounds starting playing as occluded
+				bIsOccluded = WorldPtr->LineTraceTestByChannel(SoundLocation, ListenerLocation, ECC_Visibility, Params);
+			}
+			// don't need to do anyother async trace if we've already got one pending
+			else if (!bAsyncOcclusionPending)
+			{
+				bAsyncOcclusionPending = true;
+				WorldPtr->AsyncLineTraceByChannel(SoundLocation, ListenerLocation, ECC_Visibility, Params, FCollisionResponseParams::DefaultResponseParam, &OcclusionTraceDelegate);
+			}
 		}
 	}
 
@@ -799,7 +821,8 @@ void FActiveSound::ApplyAttenuation(FSoundParseParameters& ParseParams, const FL
 		}
 	}
 
-	if (Settings->bEnableOcclusion && bIsAudible && FApp::GetVolumeMultiplier() > 0.0f && !AudioDevice->IsAudioDeviceMuted())
+	// Only due occlusion traces if the sound is audible
+	if (Settings->bEnableOcclusion && Volume > 0.0f && !AudioDevice->IsAudioDeviceMuted())
 	{
 		check(ClosestListenerPtr);
 		CheckOcclusion(ClosestListenerPtr->Transform.GetTranslation(), ParseParams.Transform.GetTranslation(), Settings);
diff --git a/Engine/Source/Runtime/Engine/Private/AnimBlueprint.cpp b/Engine/Source/Runtime/Engine/Private/AnimBlueprint.cpp
--- a/Engine/Source/Runtime/Engine/Private/AnimBlueprint.cpp
+++ b/Engine/Source/Runtime/Engine/Private/AnimBlueprint.cpp
@@ -5,6 +5,11 @@
 #include "LatentActions.h"
 #include "Animation/AnimBlueprint.h"
 #include "Animation/AnimBlueprintGeneratedClass.h"
+#include "FrameworkObjectVersion.h"
+
+#if WITH_EDITORONLY_DATA
+#include "AnimationEditorUtils.h"
+#endif
 
 //////////////////////////////////////////////////////////////////////////
 // UAnimBlueprint
@@ -149,6 +154,19 @@ void UAnimBlueprint::PostLoad()
 		});
 	}
 #endif
+
+#if WITH_EDITORONLY_DATA
+	if(GetLinkerCustomVersion(FFrameworkObjectVersion::GUID) < FFrameworkObjectVersion::AnimBlueprintSubgraphFix)
+	{
+		AnimationEditorUtils::RegenerateSubGraphArrays(this);
+	}
+#endif
+}
+
+void UAnimBlueprint::Serialize(FArchive& Ar)
+{
+	Super::Serialize(Ar);
+	Ar.UsingCustomVersion(FFrameworkObjectVersion::GUID);
 }
 
 #endif
\ No newline at end of file
diff --git a/Engine/Source/Runtime/Engine/Private/Animation/AnimInstanceProxy.cpp b/Engine/Source/Runtime/Engine/Private/Animation/AnimInstanceProxy.cpp
--- a/Engine/Source/Runtime/Engine/Private/Animation/AnimInstanceProxy.cpp
+++ b/Engine/Source/Runtime/Engine/Private/Animation/AnimInstanceProxy.cpp
@@ -960,7 +960,7 @@ float FAnimInstanceProxy::GetInstanceAssetPlayerLength(int32 AssetPlayerIndex)
 {
 	if(FAnimNode_AssetPlayerBase* PlayerNode = GetNodeFromIndex<FAnimNode_AssetPlayerBase>(AssetPlayerIndex))
 	{
-		return PlayerNode->GetAnimAsset()->GetMaxCurrentTime();
+		return PlayerNode->GetCurrentAssetLength();
 	}
 
 	return 0.0f;
@@ -980,13 +980,11 @@ float FAnimInstanceProxy::GetInstanceAssetPlayerTimeFraction(int32 AssetPlayerIn
 {
 	if(FAnimNode_AssetPlayerBase* PlayerNode = GetNodeFromIndex<FAnimNode_AssetPlayerBase>(AssetPlayerIndex))
 	{
-		UAnimationAsset* Asset = PlayerNode->GetAnimAsset();
-
-		float Length = Asset ? Asset->GetMaxCurrentTime() : 0.0f;
+		float Length = PlayerNode->GetCurrentAssetLength();
 
 		if(Length > 0.0f)
 		{
-			return PlayerNode->GetAccumulatedTime() / Length;
+			return PlayerNode->GetCurrentAssetTime() / Length;
 		}
 	}
 
@@ -997,13 +995,11 @@ float FAnimInstanceProxy::GetInstanceAssetPlayerTimeFromEndFraction(int32 AssetP
 {
 	if(FAnimNode_AssetPlayerBase* PlayerNode = GetNodeFromIndex<FAnimNode_AssetPlayerBase>(AssetPlayerIndex))
 	{
-		UAnimationAsset* Asset = PlayerNode->GetAnimAsset();
-		
-		float Length = Asset ? Asset->GetMaxCurrentTime() : 0.f;
+		float Length = PlayerNode->GetCurrentAssetLength();
 
 		if(Length > 0.f)
 		{
-			return (Length - PlayerNode->GetAccumulatedTime()) / Length;
+			return (Length - PlayerNode->GetCurrentAssetTime()) / Length;
 		}
 	}
 
@@ -1014,11 +1010,7 @@ float FAnimInstanceProxy::GetInstanceAssetPlayerTimeFromEnd(int32 AssetPlayerInd
 {
 	if(FAnimNode_AssetPlayerBase* PlayerNode = GetNodeFromIndex<FAnimNode_AssetPlayerBase>(AssetPlayerIndex))
 	{
-		UAnimationAsset* Asset = PlayerNode->GetAnimAsset();
-		if(Asset)
-		{
-			return PlayerNode->GetAnimAsset()->GetMaxCurrentTime() - PlayerNode->GetAccumulatedTime();
-		}
+		return PlayerNode->GetCurrentAssetLength() - PlayerNode->GetCurrentAssetTime();
 	}
 
 	return MAX_flt;
@@ -1082,7 +1074,7 @@ float FAnimInstanceProxy::GetRelevantAnimTimeRemaining(int32 MachineIndex, int32
 	{
 		if(AssetPlayer->GetAnimAsset())
 		{
-			return AssetPlayer->GetAnimAsset()->GetMaxCurrentTime() - AssetPlayer->GetAccumulatedTime();
+			return AssetPlayer->GetCurrentAssetLength() - AssetPlayer->GetCurrentAssetTime();
 		}
 	}
 
@@ -1095,10 +1087,10 @@ float FAnimInstanceProxy::GetRelevantAnimTimeRemainingFraction(int32 MachineInde
 	{
 		if(AssetPlayer->GetAnimAsset())
 		{
-			float Length = AssetPlayer->GetAnimAsset()->GetMaxCurrentTime();
+			float Length = AssetPlayer->GetCurrentAssetLength();
 			if(Length > 0.0f)
 			{
-				return (Length - AssetPlayer->GetAccumulatedTime()) / Length;
+				return (Length - AssetPlayer->GetCurrentAssetTime()) / Length;
 			}
 		}
 	}
@@ -1112,7 +1104,7 @@ float FAnimInstanceProxy::GetRelevantAnimLength(int32 MachineIndex, int32 StateI
 	{
 		if(AssetPlayer->GetAnimAsset())
 		{
-			return AssetPlayer->GetAnimAsset()->GetMaxCurrentTime();
+			return AssetPlayer->GetCurrentAssetLength();
 		}
 	}
 
@@ -1123,7 +1115,7 @@ float FAnimInstanceProxy::GetRelevantAnimTime(int32 MachineIndex, int32 StateInd
 {
 	if(FAnimNode_AssetPlayerBase* AssetPlayer = GetRelevantAssetPlayerFromState(MachineIndex, StateIndex))
 	{
-		return AssetPlayer->GetAccumulatedTime();
+		return AssetPlayer->GetCurrentAssetTime();
 	}
 
 	return 0.0f;
@@ -1133,10 +1125,10 @@ float FAnimInstanceProxy::GetRelevantAnimTimeFraction(int32 MachineIndex, int32
 {
 	if(FAnimNode_AssetPlayerBase* AssetPlayer = GetRelevantAssetPlayerFromState(MachineIndex, StateIndex))
 	{
-		float Length = AssetPlayer->GetAnimAsset()->GetMaxCurrentTime();
+		float Length = AssetPlayer->GetCurrentAssetLength();
 		if(Length > 0.0f)
 		{
-			return AssetPlayer->GetAccumulatedTime() / Length;
+			return AssetPlayer->GetCurrentAssetTime() / Length;
 		}
 	}
 
diff --git a/Engine/Source/Runtime/Engine/Private/Animation/AnimNode_SequencePlayer.cpp b/Engine/Source/Runtime/Engine/Private/Animation/AnimNode_SequencePlayer.cpp
--- a/Engine/Source/Runtime/Engine/Private/Animation/AnimNode_SequencePlayer.cpp
+++ b/Engine/Source/Runtime/Engine/Private/Animation/AnimNode_SequencePlayer.cpp
@@ -7,6 +7,16 @@
 /////////////////////////////////////////////////////
 // FAnimSequencePlayerNode
 
+float FAnimNode_SequencePlayer::GetCurrentAssetTime()
+{
+	return InternalTimeAccumulator;
+}
+
+float FAnimNode_SequencePlayer::GetCurrentAssetLength()
+{
+	return Sequence ? Sequence->SequenceLength : 0.0f;
+}
+
 void FAnimNode_SequencePlayer::Initialize(const FAnimationInitializeContext& Context)
 {
 	FAnimNode_AssetPlayerBase::Initialize(Context);
diff --git a/Engine/Source/Runtime/Engine/Private/Animation/AnimSequence.cpp b/Engine/Source/Runtime/Engine/Private/Animation/AnimSequence.cpp
--- a/Engine/Source/Runtime/Engine/Private/Animation/AnimSequence.cpp
+++ b/Engine/Source/Runtime/Engine/Private/Animation/AnimSequence.cpp
@@ -844,7 +844,7 @@ void UAnimSequence::ExtractBoneTransform(const TArray<struct FRawAnimSequenceTra
 	// Bail out if the animation data doesn't exists (e.g. was stripped by the cooker).
 	if(InRawAnimationData.Num() == 0)
 	{
-		UE_LOG(LogAnimation, Log, TEXT("UAnimSequence::GetBoneTransform : No anim data in AnimSequence!"));
+		UE_LOG(LogAnimation, Log, TEXT("UAnimSequence::GetBoneTransform : No anim data in AnimSequence[%s]!"),*GetFullName());
 		OutAtom.SetIdentity();
 		return;
 	}
@@ -857,7 +857,7 @@ void UAnimSequence::ExtractBoneTransform(const struct FRawAnimSequenceTrack& Raw
 	// Bail out (with rather wacky data) if data is empty for some reason.
 	if(RawTrack.PosKeys.Num() == 0 || RawTrack.RotKeys.Num() == 0)
 	{
-		UE_LOG(LogAnimation, Log, TEXT("UAnimSequence::GetBoneTransform : No anim data in AnimSequence!"));
+		UE_LOG(LogAnimation, Log, TEXT("UAnimSequence::GetBoneTransform : No anim data in AnimSequence[%s]!"),*GetFullName());
 		OutAtom.SetIdentity();
 		return;
 	}
@@ -1159,10 +1159,11 @@ void UAnimSequence::GetBonePose(FCompactPose& OutPose, FBlendedCurve& OutCurve,
 {
 	SCOPE_CYCLE_COUNTER(STAT_AnimSeq_GetBonePose);
 
-	USkeleton* MySkeleton = GetSkeleton();
+	const bool bIsBakedAdditive = !bCalcAdditiveDynamically && IsValidAdditive();
+	const USkeleton* MySkeleton = GetSkeleton();
 	if (!MySkeleton)
 	{
-		if (!bCalcAdditiveDynamically && IsValidAdditive())
+		if (bIsBakedAdditive)
 		{
 			OutPose.ResetToIdentity(); 
 		}
@@ -1174,31 +1175,31 @@ void UAnimSequence::GetBonePose(FCompactPose& OutPose, FBlendedCurve& OutCurve,
 	}
 
 	const FBoneContainer& RequiredBones = OutPose.GetBoneContainer();
+	const bool bDisableRetargeting = RequiredBones.GetDisableRetargeting();
 
-	// if retargeting is disabled, we initialize pose with 'Retargeting Source' ref pose.
-	bool const bDisableRetargeting = RequiredBones.GetDisableRetargeting();
-	if (bDisableRetargeting)
+	// initialize with ref-pose
+	if (bIsBakedAdditive)
 	{
-		TArray<FTransform> const& AuthoredOnRefSkeleton = MySkeleton->GetRefLocalPoses(RetargetSource);
-		TArray<FBoneIndexType> const& RequireBonesIndexArray = RequiredBones.GetBoneIndicesArray();
-
-		int32 const NumRequiredBones = RequireBonesIndexArray.Num();
-		for (FCompactPoseBoneIndex PoseBoneIndex : OutPose.ForEachBoneIndex())
-		{
-			int32 const& SkeletonBoneIndex = RequiredBones.GetSkeletonIndex(PoseBoneIndex);
-
-			// Pose bone index should always exist in Skeleton
-			checkSlow(SkeletonBoneIndex != INDEX_NONE);
-			OutPose[PoseBoneIndex] = AuthoredOnRefSkeleton[SkeletonBoneIndex];
-		}
+		//When using baked additive ref pose is identity
+		OutPose.ResetToIdentity();
 	}
 	else
 	{
-		// initialize with ref-pose
-		if (!bCalcAdditiveDynamically && IsValidAdditive())
+		// if retargeting is disabled, we initialize pose with 'Retargeting Source' ref pose.
+		if (bDisableRetargeting)
 		{
-			//When using baked additive ref pose is identity
-			OutPose.ResetToIdentity();
+			TArray<FTransform> const& AuthoredOnRefSkeleton = MySkeleton->GetRefLocalPoses(RetargetSource);
+			TArray<FBoneIndexType> const& RequireBonesIndexArray = RequiredBones.GetBoneIndicesArray();
+
+			int32 const NumRequiredBones = RequireBonesIndexArray.Num();
+			for (FCompactPoseBoneIndex PoseBoneIndex : OutPose.ForEachBoneIndex())
+			{
+				int32 const& SkeletonBoneIndex = RequiredBones.GetSkeletonIndex(PoseBoneIndex);
+
+				// Pose bone index should always exist in Skeleton
+				checkSlow(SkeletonBoneIndex != INDEX_NONE);
+				OutPose[PoseBoneIndex] = AuthoredOnRefSkeleton[SkeletonBoneIndex];
+			}
 		}
 		else
 		{
@@ -1307,7 +1308,13 @@ void UAnimSequence::GetBonePose(FCompactPose& OutPose, FBlendedCurve& OutCurve,
 						break;
 					case EBoneTranslationRetargetingMode::AnimationRelative:
 						TranslationPairs.Add(BoneTrackPair(CompactPoseBoneIndex, TrackIndex));
-						AnimRelativeRetargetingPairs.Add(BoneTrackPair(CompactPoseBoneIndex, SkeletonBoneIndex));
+
+						// With baked additives, we can skip 'AnimationRelative' tracks, as the relative transform gets canceled out.
+						// (A1 + Rel) - (A2 + Rel) = A1 - A2.
+						if (!bIsBakedAdditive)
+						{
+							AnimRelativeRetargetingPairs.Add(BoneTrackPair(CompactPoseBoneIndex, SkeletonBoneIndex));
+						}
 						break;
 					}
 				}
@@ -1499,35 +1506,33 @@ void UAnimSequence::RetargetBoneTransform(FTransform& BoneTransform, const int32
 	const USkeleton* MySkeleton = GetSkeleton();
 	const TArray<FBoneNode>& BoneTree = MySkeleton->GetBoneTree();
 
-	check(!bBakedAdditive || BoneTree[SkeletonBoneIndex].TranslationRetargetingMode == EBoneTranslationRetargetingMode::Skeleton || BoneTree[SkeletonBoneIndex].TranslationRetargetingMode == EBoneTranslationRetargetingMode::Animation);
-
-	if (BoneTree[SkeletonBoneIndex].TranslationRetargetingMode == EBoneTranslationRetargetingMode::AnimationScaled)
-	{
-		// @todo - precache that in FBoneContainer when we have SkeletonIndex->TrackIndex mapping. So we can just apply scale right away.
-		const TArray<FTransform>& SkeletonRefPoseArray = GetSkeleton()->GetRefLocalPoses(RetargetSource);
-		const float SourceTranslationLength = SkeletonRefPoseArray[SkeletonBoneIndex].GetTranslation().Size();
-		if (SourceTranslationLength > KINDA_SMALL_NUMBER)
-		{
-			const float TargetTranslationLength = RequiredBones.GetRefPoseTransform(BoneIndex).GetTranslation().Size();
-			BoneTransform.ScaleTranslation(TargetTranslationLength / SourceTranslationLength);
-		}
-	}
-	else
+	switch (BoneTree[SkeletonBoneIndex].TranslationRetargetingMode)
 	{
-		if (bBakedAdditive)
+		case EBoneTranslationRetargetingMode::AnimationScaled:
 		{
-			if (BoneTree[SkeletonBoneIndex].TranslationRetargetingMode == EBoneTranslationRetargetingMode::Skeleton)
+			// @todo - precache that in FBoneContainer when we have SkeletonIndex->TrackIndex mapping. So we can just apply scale right away.
+			const TArray<FTransform>& SkeletonRefPoseArray = GetSkeleton()->GetRefLocalPoses(RetargetSource);
+			const float SourceTranslationLength = SkeletonRefPoseArray[SkeletonBoneIndex].GetTranslation().Size();
+			if (SourceTranslationLength > KINDA_SMALL_NUMBER)
 			{
-				BoneTransform.SetTranslation(FVector::ZeroVector);
+				const float TargetTranslationLength = RequiredBones.GetRefPoseTransform(BoneIndex).GetTranslation().Size();
+				BoneTransform.ScaleTranslation(TargetTranslationLength / SourceTranslationLength);
 			}
+			break;
 		}
-		else
+		
+		case EBoneTranslationRetargetingMode::Skeleton:
 		{
-			if (BoneTree[SkeletonBoneIndex].TranslationRetargetingMode == EBoneTranslationRetargetingMode::Skeleton)
-			{
-				BoneTransform.SetTranslation(RequiredBones.GetRefPoseTransform(BoneIndex).GetTranslation());
-			}
-			else if (BoneTree[SkeletonBoneIndex].TranslationRetargetingMode == EBoneTranslationRetargetingMode::AnimationRelative)
+			BoneTransform.SetTranslation(bBakedAdditive ? FVector::ZeroVector : RequiredBones.GetRefPoseTransform(BoneIndex).GetTranslation());
+			break;
+		}
+		
+
+		case EBoneTranslationRetargetingMode::AnimationRelative:
+		{
+			// With baked additive animations, Animation Relative delta gets canceled out, so we can skip it.
+			// (A1 + Rel) - (A2 + Rel) = A1 - A2.
+			if (!bBakedAdditive)
 			{
 				const TArray<FTransform>& AuthoredOnRefSkeleton = GetSkeleton()->GetRefLocalPoses(RetargetSource);
 				const TArray<FTransform>& PlayingOnRefSkeleton = RequiredBones.GetRefPoseArray();
@@ -1540,8 +1545,8 @@ void UAnimSequence::RetargetBoneTransform(FTransform& BoneTransform, const int32
 				BoneTransform.SetScale3D(BoneTransform.GetScale3D() * (RefPoseTransform.GetScale3D() * AuthoredOnRefSkeleton[SkeletonBoneIndex].GetSafeScaleReciprocal(AuthoredOnRefSkeleton[SkeletonBoneIndex].GetScale3D())));
 				BoneTransform.NormalizeRotation();
 			}
+			break;
 		}
-
 	}
 }
 
diff --git a/Engine/Source/Runtime/Engine/Private/AnimationAsset.cpp b/Engine/Source/Runtime/Engine/Private/AnimationAsset.cpp
--- a/Engine/Source/Runtime/Engine/Private/AnimationAsset.cpp
+++ b/Engine/Source/Runtime/Engine/Private/AnimationAsset.cpp
@@ -1,6 +1,7 @@
 // Copyright 1998-2016 Epic Games, Inc. All Rights Reserved.
 
 #include "EnginePrivate.h"
+#include "Engine/AssetUserData.h"
 #include "Animation/AnimSequence.h"
 #include "AnimationUtils.h"
 #include "Animation/AnimInstanceProxy.h"
@@ -287,6 +288,50 @@ void UAnimationAsset::ValidateSkeleton()
 	}
 }
 
+void UAnimationAsset::AddAssetUserData(UAssetUserData* InUserData)
+{
+	if (InUserData != NULL)
+	{
+		UAssetUserData* ExistingData = GetAssetUserDataOfClass(InUserData->GetClass());
+		if (ExistingData != NULL)
+		{
+			AssetUserData.Remove(ExistingData);
+		}
+		AssetUserData.Add(InUserData);
+	}
+}
+
+UAssetUserData* UAnimationAsset::GetAssetUserDataOfClass(TSubclassOf<UAssetUserData> InUserDataClass)
+{
+	for (int32 DataIdx = 0; DataIdx < AssetUserData.Num(); DataIdx++)
+	{
+		UAssetUserData* Datum = AssetUserData[DataIdx];
+		if (Datum != NULL && Datum->IsA(InUserDataClass))
+		{
+			return Datum;
+		}
+	}
+	return NULL;
+}
+
+void UAnimationAsset::RemoveUserDataOfClass(TSubclassOf<UAssetUserData> InUserDataClass)
+{
+	for (int32 DataIdx = 0; DataIdx < AssetUserData.Num(); DataIdx++)
+	{
+		UAssetUserData* Datum = AssetUserData[DataIdx];
+		if (Datum != NULL && Datum->IsA(InUserDataClass))
+		{
+			AssetUserData.RemoveAt(DataIdx);
+			return;
+		}
+	}
+}
+
+const TArray<UAssetUserData*>* UAnimationAsset::GetAssetUserDataArray() const
+{
+	return &AssetUserData;
+}
+
 //////////////////////////////////////////////////////////////////////////
 // FBoneContainer
 
diff --git a/Engine/Source/Runtime/Engine/Private/AudioDevice.cpp b/Engine/Source/Runtime/Engine/Private/AudioDevice.cpp
--- a/Engine/Source/Runtime/Engine/Private/AudioDevice.cpp
+++ b/Engine/Source/Runtime/Engine/Private/AudioDevice.cpp
@@ -2413,11 +2413,34 @@ void FAudioDevice::AddNewActiveSound(const FActiveSound& NewActiveSound)
 
 void FAudioDevice::ProcessingPendingActiveSoundStops()
 {
+	// Process the PendingSoundsToDelete. These may have 
+	// had their deletion deferred due to an async operation
+	for (int32 i = PendingSoundsToDelete.Num() - 1; i >= 0; --i)
+	{
+		FActiveSound* ActiveSound = PendingSoundsToDelete[i];
+		if (ActiveSound->CanDelete())
+		{
+			PendingSoundsToDelete.RemoveAtSwap(i, 1, false);
+			delete ActiveSound;
+		}
+	}
+
 	// Stop any pending active sounds that need to be stopped
 	for (FActiveSound* ActiveSound : PendingSoundsToStop)
 	{
 		check(ActiveSound);
 		ActiveSound->Stop();
+
+		// If we can delete the active sound now, then delete it
+		if (ActiveSound->CanDelete())
+		{
+			delete ActiveSound;
+		}
+		else
+		{
+			// There was an async operation pending. We need to defer deleting this sound
+			PendingSoundsToDelete.Add(ActiveSound);
+		}
 	}
 	PendingSoundsToStop.Reset();
 }
diff --git a/Engine/Source/Runtime/Engine/Private/AudioEffect.cpp b/Engine/Source/Runtime/Engine/Private/AudioEffect.cpp
--- a/Engine/Source/Runtime/Engine/Private/AudioEffect.cpp
+++ b/Engine/Source/Runtime/Engine/Private/AudioEffect.cpp
@@ -111,29 +111,43 @@ void FAudioReverbEffect::Interpolate( float InterpValue, const FAudioReverbEffec
  */
 void FAudioEQEffect::ClampValues( void )
 {
-	HFFrequency = FMath::Clamp<float>( HFFrequency, MIN_FILTER_FREQUENCY, MAX_FILTER_FREQUENCY );
-	HFGain = FMath::Clamp<float>( HFGain, MIN_FILTER_GAIN, MAX_FILTER_GAIN );
-	MFCutoffFrequency = FMath::Clamp<float>( MFCutoffFrequency, MIN_FILTER_FREQUENCY, MAX_FILTER_FREQUENCY );
-	MFBandwidth = FMath::Clamp<float>( MFBandwidth, MIN_FILTER_BANDWIDTH, MAX_FILTER_BANDWIDTH );
-	MFGain = FMath::Clamp<float>( MFGain, MIN_FILTER_GAIN, MAX_FILTER_GAIN );
-	LFFrequency = FMath::Clamp<float>( LFFrequency, MIN_FILTER_FREQUENCY, MAX_FILTER_FREQUENCY );
-	LFGain = FMath::Clamp<float>( LFGain, MIN_FILTER_GAIN, MAX_FILTER_GAIN );
+	FrequencyCenter0 = FMath::Clamp(FrequencyCenter0, MIN_FILTER_FREQUENCY, MAX_FILTER_FREQUENCY);
+	FrequencyCenter1 = FMath::Clamp(FrequencyCenter1, MIN_FILTER_FREQUENCY, MAX_FILTER_FREQUENCY);
+	FrequencyCenter2 = FMath::Clamp(FrequencyCenter2, MIN_FILTER_FREQUENCY, MAX_FILTER_FREQUENCY);
+	FrequencyCenter3 = FMath::Clamp(FrequencyCenter3, MIN_FILTER_FREQUENCY, MAX_FILTER_FREQUENCY);
+
+	Gain0 = FMath::Clamp(Gain0, 0.0f, MAX_FILTER_GAIN);
+	Gain1 = FMath::Clamp(Gain1, 0.0f, MAX_FILTER_GAIN);
+	Gain2 = FMath::Clamp(Gain2, 0.0f, MAX_FILTER_GAIN);
+	Gain3 = FMath::Clamp(Gain3, 0.0f, MAX_FILTER_GAIN);
+
+	Bandwidth0 = FMath::Clamp(Bandwidth0, MIN_FILTER_BANDWIDTH, MAX_FILTER_BANDWIDTH);
+	Bandwidth1 = FMath::Clamp(Bandwidth1, MIN_FILTER_BANDWIDTH, MAX_FILTER_BANDWIDTH);
+	Bandwidth2 = FMath::Clamp(Bandwidth2, MIN_FILTER_BANDWIDTH, MAX_FILTER_BANDWIDTH);
+	Bandwidth3 = FMath::Clamp(Bandwidth3, MIN_FILTER_BANDWIDTH, MAX_FILTER_BANDWIDTH);
 }
 
 /** 
  * Interpolate EQ settings based on time
  */
-void FAudioEQEffect::Interpolate( float InterpValue, const FAudioEQEffect& Start, const FAudioEQEffect& End )
+void FAudioEQEffect::Interpolate(float InterpValue, const FAudioEQEffect& Start, const FAudioEQEffect& End )
 {
-	float InvInterpValue = 1.0f - InterpValue;
 	RootTime = FApp::GetCurrentTime();
-	HFFrequency = ( Start.HFFrequency * InvInterpValue ) + ( End.HFFrequency * InterpValue );
-	HFGain = ( Start.HFGain * InvInterpValue ) + ( End.HFGain * InterpValue );
-	MFCutoffFrequency = ( Start.MFCutoffFrequency * InvInterpValue ) + ( End.MFCutoffFrequency * InterpValue );
-	MFBandwidth = ( Start.MFBandwidth * InvInterpValue ) + ( End.MFBandwidth * InterpValue );
-	MFGain = ( Start.MFGain * InvInterpValue ) + ( End.MFGain * InterpValue );
-	LFFrequency = ( Start.LFFrequency * InvInterpValue ) + ( End.LFFrequency * InterpValue );
-	LFGain = ( Start.LFGain * InvInterpValue ) + ( End.LFGain * InterpValue );
+
+	FrequencyCenter0 = FMath::Lerp(Start.FrequencyCenter0, End.FrequencyCenter0, InterpValue);
+	FrequencyCenter1 = FMath::Lerp(Start.FrequencyCenter1, End.FrequencyCenter1, InterpValue);
+	FrequencyCenter2 = FMath::Lerp(Start.FrequencyCenter2, End.FrequencyCenter2, InterpValue);
+	FrequencyCenter3 = FMath::Lerp(Start.FrequencyCenter3, End.FrequencyCenter3, InterpValue);
+
+	Gain0 = FMath::Lerp(Start.Gain0, End.Gain0, InterpValue);
+	Gain1 = FMath::Lerp(Start.Gain1, End.Gain1, InterpValue);
+	Gain2 = FMath::Lerp(Start.Gain2, End.Gain2, InterpValue);
+	Gain3 = FMath::Lerp(Start.Gain3, End.Gain3, InterpValue);
+
+	Bandwidth0 = FMath::Lerp(Start.Bandwidth0, End.Bandwidth0, InterpValue);
+	Bandwidth1 = FMath::Lerp(Start.Bandwidth1, End.Bandwidth1, InterpValue);
+	Bandwidth2 = FMath::Lerp(Start.Bandwidth2, End.Bandwidth2, InterpValue);
+	Bandwidth3 = FMath::Lerp(Start.Bandwidth3, End.Bandwidth3, InterpValue);
 }
 
 /**
@@ -296,9 +310,9 @@ void FAudioEffectsManager::SetReverbSettings( const FReverbSettings& ReverbSetti
  * Sets new EQ mix if necessary. Otherwise interpolates to the current settings and calls SetEffect to handle
  * the platform specific aspect.
  */
-void FAudioEffectsManager::SetMixSettings( USoundMix* NewMix, bool bIgnorePriority )
+void FAudioEffectsManager::SetMixSettings(USoundMix* NewMix, bool bIgnorePriority, bool bForce)
 {
-	if( NewMix && NewMix != CurrentEQMix )
+	if (NewMix && (NewMix != CurrentEQMix || bForce))
 	{
 		// Check whether the priority of this SoundMix is higher than existing one
 		if (CurrentEQMix == NULL || bIgnorePriority || NewMix->EQPriority > CurrentEQMix->EQPriority)
@@ -336,11 +350,12 @@ void FAudioEffectsManager::ClearMixSettings()
 		UE_LOG(LogAudio, Log, TEXT( "FAudioEffectsManager::ClearMixSettings(): %s" ), *CurrentEQMix->GetName() );
 
 		SourceEQEffect = CurrentEQEffect;
-		SourceEQEffect.RootTime = FApp::GetCurrentTime();
+		double CurrentTime = FApp::GetCurrentTime();
+		SourceEQEffect.RootTime = CurrentTime;
 
 		// interpolate back to default
 		DestinationEQEffect = FAudioEQEffect();
-		DestinationEQEffect.RootTime = FApp::GetCurrentTime() + CurrentEQMix->FadeOutTime;
+		DestinationEQEffect.RootTime = CurrentTime + CurrentEQMix->FadeOutTime;
 
 		CurrentEQMix = NULL;
 	}
@@ -349,19 +364,22 @@ void FAudioEffectsManager::ClearMixSettings()
 /** 
  * Feed in new settings to the audio effect system
  */
-void FAudioEffectsManager::Update( void )
+void FAudioEffectsManager::Update()
 {
-	/** Interpolate the settings depending on time */
-	Interpolate( CurrentReverbEffect, SourceReverbEffect, DestinationReverbEffect );
-
-	/** Call the platform-specific code to set the effect */
-	SetReverbEffectParameters( CurrentReverbEffect );
+	// Check for changes to the mix so we can hear EQ changes in real-time
+#if WITH_EDITORONLY_DATA
+	if (CurrentEQMix && CurrentEQMix->bChanged)
+	{
+		CurrentEQMix->bChanged = false;
+		SetMixSettings(CurrentEQMix, true, true);
+	}
+#endif
 
-	/** Interpolate the settings depending on time */
-	Interpolate( CurrentEQEffect, SourceEQEffect, DestinationEQEffect );
+	Interpolate(CurrentReverbEffect, SourceReverbEffect, DestinationReverbEffect);
+	SetReverbEffectParameters(CurrentReverbEffect);
 
-	/** Call the platform-specific code to set the effect */
-	SetEQEffectParameters( CurrentEQEffect );
+	Interpolate(CurrentEQEffect, SourceEQEffect, DestinationEQEffect);
+	SetEQEffectParameters(CurrentEQEffect);
 }
 
 // end 
diff --git a/Engine/Source/Runtime/Engine/Private/Collision/CollisionConversions.cpp b/Engine/Source/Runtime/Engine/Private/Collision/CollisionConversions.cpp
--- a/Engine/Source/Runtime/Engine/Private/Collision/CollisionConversions.cpp
+++ b/Engine/Source/Runtime/Engine/Private/Collision/CollisionConversions.cpp
@@ -431,13 +431,12 @@ EConvertQueryResult ConvertQueryImpactHit(const UWorld* World, const PxLocationH
 	SCOPE_CYCLE_COUNTER(STAT_ConvertQueryImpactHit);
 
 #if WITH_EDITOR
-	bool bIgnoreFaceIndexForPIE = false;
 	if(bReturnFaceIndex && World->IsGameWorld())
 	{
 		if(!ensure(UPhysicsSettings::Get()->bSuppressFaceRemapTable == false))
 		{
-			UE_LOG(LogPhysics, Error, TEXT("A scene query is relying on face indices, but bSuppressFaceRemapTable is false."));
-			bIgnoreFaceIndexForPIE = true;
+			UE_LOG(LogPhysics, Error, TEXT("A scene query is relying on face indices, but bSuppressFaceRemapTable is true."));
+			bReturnFaceIndex = false;
 		}
 	}
 #endif
@@ -523,21 +522,17 @@ EConvertQueryResult ConvertQueryImpactHit(const UWorld* World, const PxLocationH
 			}
 		}
 	}
-	else
-	if(bReturnFaceIndex && PHit.shape->getGeometryType() == PxGeometryType::eTRIANGLEMESH)
+	else if (bReturnFaceIndex && PHit.shape->getGeometryType() == PxGeometryType::eTRIANGLEMESH)
 	{
 		PxTriangleMeshGeometry PTriMeshGeom;
 		if(	PHit.shape->getTriangleMeshGeometry(PTriMeshGeom) && 
 			PTriMeshGeom.triangleMesh != NULL &&
 			PHit.faceIndex < PTriMeshGeom.triangleMesh->getNbTriangles() )
 		{
-			OutResult.FaceIndex	= PTriMeshGeom.triangleMesh->getTrianglesRemap()[PHit.faceIndex];
-#if WITH_EDITOR
-			if(bIgnoreFaceIndexForPIE)
+			if (const PxU32* TriangleRemap = PTriMeshGeom.triangleMesh->getTrianglesRemap())
 			{
-				OutResult.FaceIndex = INDEX_NONE;	//if we need to ignore face index in PIE it's because at runtime we don't have the info. We want to make sure PIE is consistent with cooked version
+				OutResult.FaceIndex	= TriangleRemap[PHit.faceIndex];
 			}
-#endif
 		}
 	}
 
@@ -574,7 +569,7 @@ EConvertQueryResult ConvertRaycastResults(bool& OutHasValidBlockingHit, const UW
 	return ConvertResult;
 }
 
-EConvertQueryResult AddSweepResults(bool& OutHasValidBlockingHit, const UWorld* World, int32 NumHits, const PxSweepHit* Hits, float CheckLength, const PxFilterData& QueryFilter, TArray<FHitResult>& OutHits, const FVector& StartLoc, const FVector& EndLoc, const PxGeometry& Geom, const PxTransform& QueryTM, float MaxDistance, bool bReturnPhysMat)
+EConvertQueryResult AddSweepResults(bool& OutHasValidBlockingHit, const UWorld* World, int32 NumHits, const PxSweepHit* Hits, float CheckLength, const PxFilterData& QueryFilter, TArray<FHitResult>& OutHits, const FVector& StartLoc, const FVector& EndLoc, const PxGeometry& Geom, const PxTransform& QueryTM, float MaxDistance, bool bReturnFaceIndex, bool bReturnPhysMat)
 {
 	OutHits.Reserve(OutHits.Num() + NumHits);
 	EConvertQueryResult ConvertResult = EConvertQueryResult::Valid;
@@ -587,7 +582,7 @@ EConvertQueryResult AddSweepResults(bool& OutHasValidBlockingHit, const UWorld*
 		if(PHit.distance <= MaxDistance)
 		{
 			FHitResult& NewResult = OutHits[OutHits.AddDefaulted()];
-			if (ConvertQueryImpactHit(World, PHit, NewResult, CheckLength, QueryFilter, StartLoc, EndLoc, &Geom, QueryTM, false, bReturnPhysMat) == EConvertQueryResult::Valid)
+			if (ConvertQueryImpactHit(World, PHit, NewResult, CheckLength, QueryFilter, StartLoc, EndLoc, &Geom, QueryTM, bReturnFaceIndex, bReturnPhysMat) == EConvertQueryResult::Valid)
 			{
 				bHadBlockingHit |= NewResult.bBlockingHit;
 			}
diff --git a/Engine/Source/Runtime/Engine/Private/Collision/CollisionConversions.h b/Engine/Source/Runtime/Engine/Private/Collision/CollisionConversions.h
--- a/Engine/Source/Runtime/Engine/Private/Collision/CollisionConversions.h
+++ b/Engine/Source/Runtime/Engine/Private/Collision/CollisionConversions.h
@@ -56,9 +56,11 @@ EConvertQueryResult ConvertRaycastResults(bool& OutHasValidBlockingHit, const UW
  * @param	StartLoc	Start of trace
  * @param	EndLoc		End of trace
  * @param	Geom
+ * @param	bReturnFaceIndex	True if we want to lookup the face index
+ * @param	bReturnPhysMat		True if we want to lookup the physical material
  * @return	Whether all results passed NaN/Inf checks.
  */
-EConvertQueryResult AddSweepResults(bool& OutHasValidBlockingHit, const UWorld* World, int32 NumHits, const PxSweepHit* Hits, float CheckLength, const PxFilterData& QueryFilter, TArray<FHitResult>& OutHits, const FVector& StartLoc, const FVector& EndLoc, const PxGeometry& Geom, const PxTransform& QueryTM, float MaxDistance, bool bReturnPhysMat);
+EConvertQueryResult AddSweepResults(bool& OutHasValidBlockingHit, const UWorld* World, int32 NumHits, const PxSweepHit* Hits, float CheckLength, const PxFilterData& QueryFilter, TArray<FHitResult>& OutHits, const FVector& StartLoc, const FVector& EndLoc, const PxGeometry& Geom, const PxTransform& QueryTM, float MaxDistance, bool bReturnFaceIndex, bool bReturnPhysMat);
 
 /** 
  * Util to convert physX overlap query to our overlap result
diff --git a/Engine/Source/Runtime/Engine/Private/Collision/PhysXCollision.cpp b/Engine/Source/Runtime/Engine/Private/Collision/PhysXCollision.cpp
--- a/Engine/Source/Runtime/Engine/Private/Collision/PhysXCollision.cpp
+++ b/Engine/Source/Runtime/Engine/Private/Collision/PhysXCollision.cpp
@@ -1394,7 +1394,7 @@ bool GeomSweepSingle(const UWorld* World, const struct FCollisionShape& Collisio
 
 		if(bHaveBlockingHit) // If we got a hit, convert it to unreal type
 		{
-			if (ConvertQueryImpactHit(World, PHit, OutHit, DeltaMag, PFilter, Start, End, &PGeom, PStartTM, false, Params.bReturnPhysicalMaterial) == EConvertQueryResult::Invalid)
+			if (ConvertQueryImpactHit(World, PHit, OutHit, DeltaMag, PFilter, Start, End, &PGeom, PStartTM, Params.bReturnFaceIndex, Params.bReturnPhysicalMaterial) == EConvertQueryResult::Invalid)
 			{
 				bHaveBlockingHit = false;
 				UE_LOG(LogCollision, Error, TEXT("GeomSweepSingle resulted in a NaN/INF in PHit!"));
@@ -1535,7 +1535,7 @@ bool GeomSweepMulti_PhysX(const UWorld* World, const PxGeometry& PGeom, const Px
 		// Convert all hits to unreal structs. This will remove any hits further than MinBlockDistance, and sort results.
 		if (NumHits > 0)
 		{
-			if (AddSweepResults(bBlockingHit, World, NumHits, PHits, DeltaMag, PFilter, OutHits, Start, End, PGeom, PStartTM, MinBlockDistance, Params.bReturnPhysicalMaterial) == EConvertQueryResult::Invalid)
+			if (AddSweepResults(bBlockingHit, World, NumHits, PHits, DeltaMag, PFilter, OutHits, Start, End, PGeom, PStartTM, MinBlockDistance, Params.bReturnFaceIndex, Params.bReturnPhysicalMaterial) == EConvertQueryResult::Invalid)
 			{
 				// We don't need to change bBlockingHit, that's done by AddSweepResults if it removed the blocking hit.
 				UE_LOG(LogCollision, Error, TEXT("GeomSweepMulti resulted in a NaN/INF in PHit!"));
diff --git a/Engine/Source/Runtime/Engine/Private/Components/BoxComponent.cpp b/Engine/Source/Runtime/Engine/Private/Components/BoxComponent.cpp
--- a/Engine/Source/Runtime/Engine/Private/Components/BoxComponent.cpp
+++ b/Engine/Source/Runtime/Engine/Private/Components/BoxComponent.cpp
@@ -16,6 +16,7 @@ UBoxComponent::UBoxComponent(const FObjectInitializer& ObjectInitializer)
 void UBoxComponent::SetBoxExtent(FVector NewBoxExtent, bool bUpdateOverlaps)
 {
 	BoxExtent = NewBoxExtent;
+	UpdateBounds();
 	MarkRenderStateDirty();
 	UpdateBodySetup();
 
diff --git a/Engine/Source/Runtime/Engine/Private/Components/CapsuleComponent.cpp b/Engine/Source/Runtime/Engine/Private/Components/CapsuleComponent.cpp
--- a/Engine/Source/Runtime/Engine/Private/Components/CapsuleComponent.cpp
+++ b/Engine/Source/Runtime/Engine/Private/Components/CapsuleComponent.cpp
@@ -146,6 +146,7 @@ void UCapsuleComponent::SetCapsuleSize(float NewRadius, float NewHalfHeight, boo
 {
 	CapsuleHalfHeight = FMath::Max3(0.f, NewHalfHeight, NewRadius);
 	CapsuleRadius = FMath::Max(0.f, NewRadius);
+	UpdateBounds();
 	UpdateBodySetup();
 	MarkRenderStateDirty();
 
diff --git a/Engine/Source/Runtime/Engine/Private/Components/CharacterMovementComponent.cpp b/Engine/Source/Runtime/Engine/Private/Components/CharacterMovementComponent.cpp
--- a/Engine/Source/Runtime/Engine/Private/Components/CharacterMovementComponent.cpp
+++ b/Engine/Source/Runtime/Engine/Private/Components/CharacterMovementComponent.cpp
@@ -1107,7 +1107,6 @@ void UCharacterMovementComponent::AdjustProxyCapsuleSize()
 			Radius * ComponentScale, HalfHeight * ComponentScale, NewRadius * ComponentScale, NewHalfHeight * ComponentScale);
 
 		CharacterOwner->GetCapsuleComponent()->SetCapsuleSize(NewRadius, NewHalfHeight, true);
-		CharacterOwner->GetCapsuleComponent()->UpdateBounds();
 	}	
 }
 
@@ -2091,7 +2090,6 @@ void UCharacterMovementComponent::UnCrouch(bool bClientSimulation)
 	check(CharacterOwner->GetCapsuleComponent());
 	bool bUpdateOverlaps = false;
 	CharacterOwner->GetCapsuleComponent()->SetCapsuleSize(DefaultCharacter->GetCapsuleComponent()->GetUnscaledCapsuleRadius(), DefaultCharacter->GetCapsuleComponent()->GetUnscaledCapsuleHalfHeight(), bUpdateOverlaps);
-	CharacterOwner->GetCapsuleComponent()->UpdateBounds(); // Force an update of the bounds with the new dimensions
 
 	if( !bClientSimulation )
 	{
@@ -2175,7 +2173,6 @@ void UCharacterMovementComponent::UnCrouch(bool bClientSimulation)
 		if (bEncroached)
 		{
 			CharacterOwner->GetCapsuleComponent()->SetCapsuleSize(CharacterOwner->GetCapsuleComponent()->GetUnscaledCapsuleRadius(), OldUnscaledHalfHeight, false);
-			CharacterOwner->GetCapsuleComponent()->UpdateBounds(); // Update bounds again back to old value
 			return;
 		}
 
diff --git a/Engine/Source/Runtime/Engine/Private/Components/SphereComponent.cpp b/Engine/Source/Runtime/Engine/Private/Components/SphereComponent.cpp
--- a/Engine/Source/Runtime/Engine/Private/Components/SphereComponent.cpp
+++ b/Engine/Source/Runtime/Engine/Private/Components/SphereComponent.cpp
@@ -74,6 +74,7 @@ void USphereComponent::UpdateBodySetup()
 void USphereComponent::SetSphereRadius( float InSphereRadius, bool bUpdateOverlaps )
 {
 	SphereRadius = InSphereRadius;
+	UpdateBounds();
 	UpdateBodySetup();
 	MarkRenderStateDirty();
 
diff --git a/Engine/Source/Runtime/Engine/Private/EnginePrivate.h b/Engine/Source/Runtime/Engine/Private/EnginePrivate.h
--- a/Engine/Source/Runtime/Engine/Private/EnginePrivate.h
+++ b/Engine/Source/Runtime/Engine/Private/EnginePrivate.h
@@ -99,4 +99,5 @@ EnginePrivate.h: Unreal engine private header file.
 #include "SlateBasics.h"
 #include "Engine/TextureLODSettings.h"
 #include "DeviceProfiles/DeviceProfile.h"
-#include "DeviceProfiles/DeviceProfileManager.h"
\ No newline at end of file
+#include "DeviceProfiles/DeviceProfileManager.h"
+#include "Engine/SkeletalMeshReductionSettings.h"
\ No newline at end of file
diff --git a/Engine/Source/Runtime/Engine/Private/LevelActor.cpp b/Engine/Source/Runtime/Engine/Private/LevelActor.cpp
--- a/Engine/Source/Runtime/Engine/Private/LevelActor.cpp
+++ b/Engine/Source/Runtime/Engine/Private/LevelActor.cpp
@@ -399,7 +399,7 @@ AActor* UWorld::SpawnActor( UClass* Class, FTransform const* UserTransformPtr, c
 		if (EncroachingBlockingGeometry(Template, FinalRootLocation, FinalRootRotation))
 		{
 			// a native component is colliding, that's enough to reject spawning
-			UE_LOG(LogSpawn, Warning, TEXT("SpawnActor failed because of collision at the spawn location [%s] for [%s]"), *FinalRootLocation.ToString(), *Class->GetName());
+			UE_LOG(LogSpawn, Log, TEXT("SpawnActor failed because of collision at the spawn location [%s] for [%s]"), *FinalRootLocation.ToString(), *Class->GetName());
 			return nullptr;
 		}
 	}
@@ -440,7 +440,7 @@ AActor* UWorld::SpawnActor( UClass* Class, FTransform const* UserTransformPtr, c
 
 	if (Actor->IsPendingKill() && !SpawnParameters.bNoFail)
 	{
-		UE_LOG(LogSpawn, Warning, TEXT("SpawnActor failed because the spawned actor %s IsPendingKill"), *Actor->GetPathName());
+		UE_LOG(LogSpawn, Log, TEXT("SpawnActor failed because the spawned actor %s IsPendingKill"), *Actor->GetPathName());
 		return NULL;
 	}
 
diff --git a/Engine/Source/Runtime/Engine/Private/NetworkDriver.cpp b/Engine/Source/Runtime/Engine/Private/NetworkDriver.cpp
--- a/Engine/Source/Runtime/Engine/Private/NetworkDriver.cpp
+++ b/Engine/Source/Runtime/Engine/Private/NetworkDriver.cpp
@@ -204,7 +204,6 @@ FNetworkObjectInfo* UNetDriver::GetNetworkActor(const AActor* InActor)
 	FNetworkObjectInfo* NetActor = nullptr;
 	if (InActor)
 	{
-		// TODO: this Find() operation is implicitly creating a TSharedPtr and causing a malloc.
 		// note: const_cast<> required because the set keys are not const.
 		TSharedPtr<FNetworkObjectInfo>* NetActorInfo = GetNetworkObjectList().GetObjects().Find(const_cast<AActor*>(InActor));
 		NetActor = (NetActorInfo ? NetActorInfo->Get() : nullptr);
diff --git a/Engine/Source/Runtime/Engine/Private/PhysicsEngine/PhysLevel.cpp b/Engine/Source/Runtime/Engine/Private/PhysicsEngine/PhysLevel.cpp
--- a/Engine/Source/Runtime/Engine/Private/PhysicsEngine/PhysLevel.cpp
+++ b/Engine/Source/Runtime/Engine/Private/PhysicsEngine/PhysLevel.cpp
@@ -494,13 +494,32 @@ void TermGamePhys()
 void DeferredPhysResourceCleanup()
 {
 #if WITH_PHYSX
+
 	// Release all tri meshes and reset array
 	for(int32 MeshIdx=0; MeshIdx<GPhysXPendingKillTriMesh.Num(); MeshIdx++)
 	{
 		PxTriangleMesh* PTriMesh = GPhysXPendingKillTriMesh[MeshIdx];
+
+		// Check this as it shouldn't be null, but then gate on it so we can
+		// avoid a crash if we end up in this state in shipping
 		check(PTriMesh);
-		PTriMesh->release();
-		GPhysXPendingKillTriMesh[MeshIdx] = NULL;
+		if(PTriMesh)
+		{
+			PTriMesh->release();
+
+			if(GPhysXPendingKillTriMesh.IsValidIndex(MeshIdx))
+			{
+				GPhysXPendingKillTriMesh[MeshIdx] = NULL;
+			}
+			else
+			{
+				UE_LOG(LogPhysics, Warning, TEXT("DeferredPhysResourceCleanup found invalid index into GPhysXPendingKillTriMesh, another thread may have modified the array."), MeshIdx);
+			}
+		}
+		else
+		{
+			UE_LOG(LogPhysics, Warning, TEXT("DeferredPhysResourceCleanup found null PxTriangleMesh in pending kill array, another thread may have modified the array."), MeshIdx);
+		}
 	}
 	GPhysXPendingKillTriMesh.Reset();
 
@@ -508,9 +527,27 @@ void DeferredPhysResourceCleanup()
 	for(int32 MeshIdx=0; MeshIdx<GPhysXPendingKillConvex.Num(); MeshIdx++)
 	{
 		PxConvexMesh* PConvexMesh = GPhysXPendingKillConvex[MeshIdx];
+
+		// Check this as it shouldn't be null, but then gate on it so we can
+		// avoid a crash if we end up in this state in shipping
 		check(PConvexMesh);
-		PConvexMesh->release();
-		GPhysXPendingKillConvex[MeshIdx] = NULL;
+		if(PConvexMesh)
+		{
+			PConvexMesh->release();
+
+			if(GPhysXPendingKillConvex.IsValidIndex(MeshIdx))
+			{
+				GPhysXPendingKillConvex[MeshIdx] = NULL;
+			}
+			else
+			{
+				UE_LOG(LogPhysics, Warning, TEXT("DeferredPhysResourceCleanup found invalid index into GPhysXPendingKillConvex (%d), another thread may have modified the array."), MeshIdx);
+			}
+		}
+		else
+		{
+			UE_LOG(LogPhysics, Warning, TEXT("DeferredPhysResourceCleanup found null PxConvexMesh in pending kill array (at %d), another thread may have modified the array."), MeshIdx);
+		}
 	}
 	GPhysXPendingKillConvex.Reset();
 
@@ -518,9 +555,27 @@ void DeferredPhysResourceCleanup()
 	for(int32 HfIdx=0; HfIdx<GPhysXPendingKillHeightfield.Num(); HfIdx++)
 	{
 		PxHeightField* PHeightfield = GPhysXPendingKillHeightfield[HfIdx];
+
+		// Check this as it shouldn't be null, but then gate on it so we can
+		// avoid a crash if we end up in this state in shipping
 		check(PHeightfield);
-		PHeightfield->release();
-		GPhysXPendingKillHeightfield[HfIdx] = NULL;
+		if(PHeightfield)
+		{
+			PHeightfield->release();
+
+			if(GPhysXPendingKillHeightfield.IsValidIndex(HfIdx))
+			{
+				GPhysXPendingKillHeightfield[HfIdx] = NULL;
+			}
+			else
+			{
+				UE_LOG(LogPhysics, Warning, TEXT("DeferredPhysResourceCleanup found invalid index into GPhysXPendingKillHeightfield (%d), another thread may have modified the array."), HfIdx);
+			}
+		}
+		else
+		{
+			UE_LOG(LogPhysics, Warning, TEXT("DeferredPhysResourceCleanup found null PxHeightField in pending kill array (at %d), another thread may have modified the array."), HfIdx);
+		}
 	}
 	GPhysXPendingKillHeightfield.Reset();
 
@@ -528,9 +583,26 @@ void DeferredPhysResourceCleanup()
 	for(int32 MeshIdx=0; MeshIdx<GPhysXPendingKillMaterial.Num(); MeshIdx++)
 	{
 		PxMaterial* PMaterial = GPhysXPendingKillMaterial[MeshIdx];
+
+		// Check this as it shouldn't be null, but then gate on it so we can
+		// avoid a crash if we end up in this state in shipping
 		check(PMaterial);
-		PMaterial->release();
-		GPhysXPendingKillMaterial[MeshIdx] = NULL;
+		if(PMaterial)
+		{
+			PMaterial->release();
+			if(GPhysXPendingKillMaterial.IsValidIndex(MeshIdx))
+			{
+				GPhysXPendingKillMaterial[MeshIdx] = NULL;
+			}
+			else
+			{
+				UE_LOG(LogPhysics, Warning, TEXT("DeferredPhysResourceCleanup found invalid index into GPhysXPendingKillMaterial(%d), another thread may have modified the array."), MeshIdx);
+			}
+		}
+		else
+		{
+			UE_LOG(LogPhysics, Warning, TEXT("DeferredPhysResourceCleanup found null PxMaterial in pending kill array (at %d), another thread may have modified the array."), MeshIdx);
+		}
 	}
 	GPhysXPendingKillMaterial.Reset();
 #endif
diff --git a/Engine/Source/Runtime/Engine/Private/SkeletalMeshReductionSettings.cpp b/Engine/Source/Runtime/Engine/Private/SkeletalMeshReductionSettings.cpp
new file mode 100644
--- /dev/null
+++ b/Engine/Source/Runtime/Engine/Private/SkeletalMeshReductionSettings.cpp
@@ -0,0 +1,70 @@
+// Copyright 1998-2016 Epic Games, Inc. All Rights Reserved.
+
+#include "EnginePrivate.h"
+
+DEFINE_LOG_CATEGORY_STATIC(LogSkeletalMeshReductionSettings, Warning, All)
+
+#include "Developer/MeshUtilities/Public/MeshUtilities.h"
+
+USkeletalMeshReductionSettings::USkeletalMeshReductionSettings(const FObjectInitializer& ObjectInitializer)
+	: Super(ObjectInitializer)
+	, bValidSettings(false)
+{
+}
+
+USkeletalMeshReductionSettings* USkeletalMeshReductionSettings::Get()
+{
+	static bool bInitialized = false;
+	// This is a singleton, use default object
+	USkeletalMeshReductionSettings* DefaultSettings = GetMutableDefault<USkeletalMeshReductionSettings>();
+	
+	if (!bInitialized)
+	{
+		DefaultSettings->Initialize();
+		bInitialized = true;
+	}
+
+	return DefaultSettings;
+}
+
+void USkeletalMeshReductionSettings::Initialize()
+{
+	// Check if valid settings were parsed from the .ini file
+	bValidSettings = ( Settings.Num() > 0 );
+}
+
+const FSkeletalMeshLODGroupSettings& USkeletalMeshReductionSettings::GetDefaultSettingsForLODLevel(const int32 LODIndex) const
+{
+	if (LODIndex < Settings.Num())
+	{
+		return Settings[LODIndex];
+	}
+
+	// This should not happen as of right now, since the function is only called with 'Default' as name
+	checkf(false, TEXT("Invalid Skeletal mesh default settings LOD Level"));
+
+	// Default return value to prevent compile issue
+	static FSkeletalMeshLODGroupSettings DefaultReturnValue;
+	return DefaultReturnValue;
+}
+
+
+int32 USkeletalMeshReductionSettings::GetNumberOfSettings() const
+{
+	return Settings.Num();
+}
+
+const bool USkeletalMeshReductionSettings::HasValidSettings() const
+{
+	return bValidSettings;
+}
+
+FSkeletalMeshOptimizationSettings FSkeletalMeshLODGroupSettings::GetSettings() const
+{
+	return OptimizationSettings;
+}
+
+const float FSkeletalMeshLODGroupSettings::GetScreenSize() const
+{
+	return ScreenSize;
+}
diff --git a/Engine/Source/Runtime/Engine/Private/SoundMix.cpp b/Engine/Source/Runtime/Engine/Private/SoundMix.cpp
--- a/Engine/Source/Runtime/Engine/Private/SoundMix.cpp
+++ b/Engine/Source/Runtime/Engine/Private/SoundMix.cpp
@@ -22,6 +22,10 @@ USoundMix::USoundMix(const FObjectInitializer& ObjectInitializer)
 	Duration = -1.0f;
 	FadeInTime = 0.2f;
 	FadeOutTime = 0.2f;
+
+#if WITH_EDITORONLY_DATA
+	bChanged = false;
+#endif
 }
 
 void USoundMix::BeginDestroy()
@@ -46,17 +50,6 @@ FString USoundMix::GetDesc( void )
 #if WITH_EDITOR
 void USoundMix::PostEditChangeProperty(FPropertyChangedEvent& PropertyChangedEvent)
 {
-	// Sanity check the EQ values
-	EQSettings.LFFrequency = FMath::Clamp<float>( EQSettings.LFFrequency, MIN_FILTER_FREQUENCY, MAX_FILTER_FREQUENCY );
-	EQSettings.LFGain = FMath::Clamp<float>( EQSettings.LFGain, MIN_FILTER_GAIN, MAX_FILTER_GAIN );
-
-	EQSettings.MFCutoffFrequency = FMath::Clamp<float>( EQSettings.MFCutoffFrequency, MIN_FILTER_FREQUENCY, MAX_FILTER_FREQUENCY );
-	EQSettings.MFBandwidth = FMath::Clamp<float>( EQSettings.MFBandwidth, MIN_FILTER_BANDWIDTH, MAX_FILTER_BANDWIDTH );
-	EQSettings.MFGain = FMath::Clamp<float>( EQSettings.MFGain, MIN_FILTER_GAIN, MAX_FILTER_GAIN );
-
-	EQSettings.HFFrequency = FMath::Clamp<float>( EQSettings.HFFrequency, MIN_FILTER_FREQUENCY, MAX_FILTER_FREQUENCY );
-	EQSettings.HFGain = FMath::Clamp<float>( EQSettings.HFGain, MIN_FILTER_GAIN, MAX_FILTER_GAIN );
-
 	if (PropertyChangedEvent.ChangeType != EPropertyChangeType::Interactive && PropertyChangedEvent.MemberProperty)
 	{
 		if (PropertyChangedEvent.MemberProperty->GetFName() == GET_MEMBER_NAME_CHECKED(USoundMix,SoundClassEffects))
@@ -77,6 +70,10 @@ void USoundMix::PostEditChangeProperty(FPropertyChangedEvent& PropertyChangedEve
 			}
 		}
 	}
+
+#if WITH_EDITORONLY_DATA
+	bChanged = true;
+#endif
 }
 
 bool USoundMix::CausesPassiveDependencyLoop(TArray<USoundClass*>& ProblemClasses) const
diff --git a/Engine/Source/Runtime/Engine/Private/SoundWave.cpp b/Engine/Source/Runtime/Engine/Private/SoundWave.cpp
--- a/Engine/Source/Runtime/Engine/Private/SoundWave.cpp
+++ b/Engine/Source/Runtime/Engine/Private/SoundWave.cpp
@@ -71,24 +71,30 @@ SIZE_T USoundWave::GetResourceSize(EResourceSizeMode::Type Mode)
 		return 0;
 	}
 
+	SIZE_T CalculatedResourceSize = 0;
+
 	if (FAudioDevice* LocalAudioDevice = GEngine->GetMainAudioDevice())
 	{
 		if (LocalAudioDevice->HasCompressedAudioInfoClass(this) && DecompressionType == DTYPE_Native)
 		{
 			check(ResourceSize == 0);
-			return RawPCMDataSize;
-		}
-		else if (DecompressionType == DTYPE_RealTime && CachedRealtimeFirstBuffer)
-		{
-			return MONO_PCM_BUFFER_SIZE * NumChannels;
+			CalculatedResourceSize = RawPCMDataSize;
 		}
-		else if ((!FPlatformProperties::SupportsAudioStreaming() || !IsStreaming()))
+		else 
 		{
-			return GetCompressedDataSize(LocalAudioDevice->GetRuntimeFormat(this));
+			if (DecompressionType == DTYPE_RealTime && CachedRealtimeFirstBuffer)
+			{
+				CalculatedResourceSize = MONO_PCM_BUFFER_SIZE * NumChannels;
+			}
+			
+			if ((!FPlatformProperties::SupportsAudioStreaming() || !IsStreaming()))
+			{
+				CalculatedResourceSize += GetCompressedDataSize(LocalAudioDevice->GetRuntimeFormat(this));
+			}
 		}
 	}
 
-	return 0;
+	return CalculatedResourceSize;
 }
 
 int32 USoundWave::GetResourceSizeForFormat(FName Format)
diff --git a/Engine/Source/Runtime/Engine/Private/Texture.cpp b/Engine/Source/Runtime/Engine/Private/Texture.cpp
--- a/Engine/Source/Runtime/Engine/Private/Texture.cpp
+++ b/Engine/Source/Runtime/Engine/Private/Texture.cpp
@@ -5,6 +5,7 @@
 =============================================================================*/
 
 #include "EnginePrivate.h"
+#include "Engine/AssetUserData.h"
 #include "EditorSupportDelegates.h"
 #include "TargetPlatform.h"
 #include "ImageWrapper.h"
@@ -491,6 +492,49 @@ bool UTexture::ForceUpdateTextureStreaming()
 	return true;
 }
 
+void UTexture::AddAssetUserData(UAssetUserData* InUserData)
+{
+	if (InUserData != NULL)
+	{
+		UAssetUserData* ExistingData = GetAssetUserDataOfClass(InUserData->GetClass());
+		if (ExistingData != NULL)
+		{
+			AssetUserData.Remove(ExistingData);
+		}
+		AssetUserData.Add(InUserData);
+	}
+}
+
+UAssetUserData* UTexture::GetAssetUserDataOfClass(TSubclassOf<UAssetUserData> InUserDataClass)
+{
+	for (int32 DataIdx = 0; DataIdx < AssetUserData.Num(); DataIdx++)
+	{
+		UAssetUserData* Datum = AssetUserData[DataIdx];
+		if (Datum != NULL && Datum->IsA(InUserDataClass))
+		{
+			return Datum;
+		}
+	}
+	return NULL;
+}
+
+void UTexture::RemoveUserDataOfClass(TSubclassOf<UAssetUserData> InUserDataClass)
+{
+	for (int32 DataIdx = 0; DataIdx < AssetUserData.Num(); DataIdx++)
+	{
+		UAssetUserData* Datum = AssetUserData[DataIdx];
+		if (Datum != NULL && Datum->IsA(InUserDataClass))
+		{
+			AssetUserData.RemoveAt(DataIdx);
+			return;
+		}
+	}
+}
+
+const TArray<UAssetUserData*>* UTexture::GetAssetUserDataArray() const
+{
+	return &AssetUserData;
+}
 
 
 /*------------------------------------------------------------------------------
diff --git a/Engine/Source/Runtime/Engine/Private/UserInterface/PlayerInput.cpp b/Engine/Source/Runtime/Engine/Private/UserInterface/PlayerInput.cpp
--- a/Engine/Source/Runtime/Engine/Private/UserInterface/PlayerInput.cpp
+++ b/Engine/Source/Runtime/Engine/Private/UserInterface/PlayerInput.cpp
@@ -739,7 +739,7 @@ void UPlayerInput::GetChordsForAction(const FInputActionBinding& ActionBinding,
 			{
 				for (auto KeyStateIt(KeyStateMap.CreateConstIterator()); KeyStateIt; ++KeyStateIt)
 				{
-					if (!KeyStateIt.Key().IsFloatAxis() && !KeyStateIt.Key().IsVectorAxis() && !IsKeyConsumed(KeyStateIt.Key()))
+					if (!KeyStateIt.Key().IsFloatAxis() && !KeyStateIt.Key().IsVectorAxis() && !KeyStateIt.Value().bConsumed)
 					{
 						FInputActionKeyMapping SubKeyMapping(KeyMapping);
 						SubKeyMapping.Key = KeyStateIt.Key();
@@ -763,7 +763,7 @@ void UPlayerInput::GetChordForKey(const FInputKeyBinding& KeyBinding, const bool
 	{
 		for (auto KeyStateIt(KeyStateMap.CreateConstIterator()); KeyStateIt; ++KeyStateIt)
 		{
-			if (!KeyStateIt.Key().IsFloatAxis() && !KeyStateIt.Key().IsVectorAxis() && !IsKeyConsumed(KeyStateIt.Key()))
+			if (!KeyStateIt.Key().IsFloatAxis() && !KeyStateIt.Key().IsVectorAxis() && !KeyStateIt.Value().bConsumed)
 			{
 				FInputKeyBinding SubKeyBinding(KeyBinding);
 				SubKeyBinding.Chord.Key = KeyStateIt.Key();
@@ -1408,19 +1408,50 @@ void UPlayerInput::DisplayDebug(class UCanvas* Canvas, const FDebugDisplayInfo&
 
 bool UPlayerInput::WasJustPressed( FKey InKey ) const
 {
-	FKeyState const* const KeyState = KeyStateMap.Find(InKey);
-	return KeyState ? (KeyState->EventCounts[IE_Pressed].Num() > 0) : false;
+	if (InKey == EKeys::AnyKey)
+	{
+		// Is there any key that has just been pressed
+		for (const TPair<FKey, FKeyState>& KeyStatePair : KeyStateMap)
+		{
+			if (!KeyStatePair.Key.IsFloatAxis() && !KeyStatePair.Key.IsVectorAxis() && KeyStatePair.Value.EventCounts[IE_Pressed].Num() > 0)
+			{
+				return true;
+			}
+		}
+	}
+	else if (FKeyState const* const KeyState = KeyStateMap.Find(InKey))
+	{
+		return (KeyState->EventCounts[IE_Pressed].Num() > 0);
+	}
+
+	return false;
 }
 
 
 bool UPlayerInput::WasJustReleased( FKey InKey ) const
 {
-	FKeyState const* const KeyState = KeyStateMap.Find(InKey);
-	return KeyState ? (KeyState->EventCounts[IE_Released].Num() > 0) : false;
+	if (InKey == EKeys::AnyKey)
+	{
+		// Is there any key that has just been released
+		for (const TPair<FKey, FKeyState>& KeyStatePair : KeyStateMap)
+		{
+			if (!KeyStatePair.Key.IsFloatAxis() && !KeyStatePair.Key.IsVectorAxis() && KeyStatePair.Value.EventCounts[IE_Released].Num() > 0)
+			{
+				return true;
+			}
+		}
+	}
+	else if (FKeyState const* const KeyState = KeyStateMap.Find(InKey))
+	{
+		return (KeyState->EventCounts[IE_Released].Num() > 0);
+	}
+
+	return false;
 }
 
 float UPlayerInput::GetTimeDown( FKey InKey ) const
 {
+	UE_CLOG(InKey == EKeys::AnyKey, LogInput, Warning, TEXT("GetTimeDown cannot return a meaningful result for AnyKey"));
 	UWorld* World = GetWorld();	
 	float DownTime = 0.f;
 	if( World )
@@ -1437,32 +1468,65 @@ float UPlayerInput::GetTimeDown( FKey InKey ) const
 
 bool UPlayerInput::IsKeyConsumed( FKey InKey ) const
 {
-	FKeyState const* const KeyState = KeyStateMap.Find(InKey);
-	return KeyState ? KeyState->bConsumed : false;
+	if (InKey == EKeys::AnyKey)
+	{
+		// Is there any key that is consumed
+		for (const TPair<FKey, FKeyState>& KeyStatePair : KeyStateMap)
+		{
+			if (KeyStatePair.Value.bConsumed)
+			{
+				return true;
+			}
+		}
+	}
+	else if (FKeyState const* const KeyState = KeyStateMap.Find(InKey))
+	{
+		return KeyState->bConsumed;
+	}
+
+	return false;
 }
 
 float UPlayerInput::GetKeyValue( FKey InKey ) const
 {
+	UE_CLOG(InKey == EKeys::AnyKey, LogInput, Warning, TEXT("GetKeyValue cannot return a meaningful result for AnyKey"));
 	FKeyState const* const KeyState = KeyStateMap.Find(InKey);
 	return KeyState ? KeyState->Value.X : 0.f;
 }
 
 float UPlayerInput::GetRawKeyValue( FKey InKey ) const
 {
+	UE_CLOG(InKey == EKeys::AnyKey, LogInput, Warning, TEXT("GetRawKeyValue cannot return a meaningful result for AnyKey"));
 	FKeyState const* const KeyState = KeyStateMap.Find(InKey);
 	return KeyState ? KeyState->RawValue.X : 0.f;
 }
 
 FVector UPlayerInput::GetVectorKeyValue( FKey InKey ) const
 {
+	UE_CLOG(InKey == EKeys::AnyKey, LogInput, Warning, TEXT("GetVectorKeyValue cannot return a meaningful result for AnyKey"));
 	FKeyState const* const KeyState = KeyStateMap.Find(InKey);
 	return KeyState ? KeyState->RawValue : FVector(0,0,0);
 }
 
 bool UPlayerInput::IsPressed( FKey InKey ) const
 {
-	FKeyState const* const KeyState = KeyStateMap.Find(InKey);
-	return KeyState ? KeyState->bDown : false;
+	if (InKey == EKeys::AnyKey)
+	{
+		// Is there any key that is down
+		for (const TPair<FKey, FKeyState>& KeyStatePair : KeyStateMap)
+		{
+			if (!KeyStatePair.Key.IsFloatAxis() && !KeyStatePair.Key.IsVectorAxis() && KeyStatePair.Value.bDown)
+			{
+				return true;
+			}
+		}
+	}
+	else if (FKeyState const* const KeyState = KeyStateMap.Find(InKey))
+	{
+		return KeyState->bDown;
+	}
+
+	return false;
 }
 
 
diff --git a/Engine/Source/Runtime/Engine/Private/VorbisAudioInfo.cpp b/Engine/Source/Runtime/Engine/Private/VorbisAudioInfo.cpp
--- a/Engine/Source/Runtime/Engine/Private/VorbisAudioInfo.cpp
+++ b/Engine/Source/Runtime/Engine/Private/VorbisAudioInfo.cpp
@@ -183,6 +183,8 @@ bool FVorbisAudioInfo::ReadCompressedInfo( const uint8* InSrcBufferData, uint32
 	Callbacks.close_func = OggClose;
 	Callbacks.tell_func = OggTell;
 
+	FScopeLock ScopeLock(&VorbisCriticalSection);
+
 	// Set up the read from memory variables
 	int Result = ov_open_callbacks(this, &VFWrapper->vf, NULL, 0, Callbacks);
 	if (Result < 0)
@@ -226,6 +228,8 @@ void FVorbisAudioInfo::ExpandFile( uint8* DstBuffer, FSoundQualityInfo* QualityI
 	check( DstBuffer );
 	check( QualityInfo );
 
+	FScopeLock ScopeLock(&VorbisCriticalSection);
+
 	// A zero buffer size means decompress the entire ogg vorbis stream to PCM.
 	TotalBytesRead = 0;
 	BytesToRead = QualityInfo->SampleDataSize;
@@ -265,6 +269,8 @@ bool FVorbisAudioInfo::ReadCompressedData( uint8* InDestination, bool bLooping,
 
 	SCOPE_CYCLE_COUNTER( STAT_VorbisDecompressTime );
 
+	FScopeLock ScopeLock(&VorbisCriticalSection);
+
 	bLooped = false;
 
 	// Work out number of samples to read
@@ -314,13 +320,17 @@ bool FVorbisAudioInfo::ReadCompressedData( uint8* InDestination, bool bLooping,
 
 void FVorbisAudioInfo::SeekToTime( const float SeekTime )
 {
-	const float TargetTime = FMath::Min(SeekTime, ( float )ov_time_total( &VFWrapper->vf, -1 ));
+	FScopeLock ScopeLock(&VorbisCriticalSection);
+
+	const float TargetTime = FMath::Min(SeekTime, (float)ov_time_total(&VFWrapper->vf, -1));
 	ov_time_seek( &VFWrapper->vf, TargetTime );
 }
 
 void FVorbisAudioInfo::EnableHalfRate( bool HalfRate )
 {
-	ov_halfrate( &VFWrapper->vf, int32(HalfRate));
+	FScopeLock ScopeLock(&VorbisCriticalSection);
+
+	ov_halfrate(&VFWrapper->vf, int32(HalfRate));
 }
 
 void LoadVorbisLibraries()
diff --git a/Engine/Source/Runtime/Engine/Public/ActiveSound.h b/Engine/Source/Runtime/Engine/Public/ActiveSound.h
--- a/Engine/Source/Runtime/Engine/Public/ActiveSound.h
+++ b/Engine/Source/Runtime/Engine/Public/ActiveSound.h
@@ -120,6 +120,9 @@ public:
 		AudioDevice = InAudioDevice;
 	}
 
+	/** Returns whether or not the active sound can be deleted. */
+	bool CanDelete() const { return !bAsyncOcclusionPending; }
+
 	FAudioDevice* AudioDevice;
 
 	/** The group of active concurrent sounds that this sound is playing in. */
@@ -137,8 +140,8 @@ public:
 	/** Whether or not the sound has checked if it was occluded already. Used to initialize a sound as occluded and bypassing occlusion interpolation. */
 	uint32 bHasCheckedOcclusion:1;
 
-	/** whether we were occluded the last time we checked */
-	uint32 bIsOccluded:1;
+	/** Flag to trigger binding our trace delegate for async trace calls */
+	uint32 bIsTraceDelegateBound:1;
 
 	/** Is this sound allowed to be spatialized? */
 	uint32 bAllowSpatialization:1;
@@ -204,6 +207,12 @@ public:
 	/** Whether or not to use an async trace for occlusion */
 	uint32 bOcclusionAsyncTrace : 1;
 
+	/** whether we were occluded the last time we checked */
+	FThreadSafeBool bIsOccluded;
+
+	/** Whether or not there is an async occlusion trace pending */
+	FThreadSafeBool bAsyncOcclusionPending;
+
 	uint8 UserIndex;
 
 	float PlaybackTime;
diff --git a/Engine/Source/Runtime/Engine/Public/AudioDevice.h b/Engine/Source/Runtime/Engine/Public/AudioDevice.h
--- a/Engine/Source/Runtime/Engine/Public/AudioDevice.h
+++ b/Engine/Source/Runtime/Engine/Public/AudioDevice.h
@@ -969,8 +969,13 @@ private:
 
 	TArray<struct FActiveSound*> ActiveSounds;
 	TArray<FWaveInstance*> ActiveWaveInstances;
+
+	/** Set of sounds which will be stopped next audio frame update */
 	TSet<struct FActiveSound*> PendingSoundsToStop;
 
+	/** A set of sounds which need to be deleted but weren't able to be deleted due to pending async operations */
+	TArray<struct FActiveSound*> PendingSoundsToDelete;
+
 	TMap<UPTRINT, struct FActiveSound*> AudioComponentToActiveSoundMap;
 
 	/** List of passive SoundMixes active last frame */
diff --git a/Engine/Source/Runtime/Engine/Public/AudioEffect.h b/Engine/Source/Runtime/Engine/Public/AudioEffect.h
--- a/Engine/Source/Runtime/Engine/Public/AudioEffect.h
+++ b/Engine/Source/Runtime/Engine/Public/AudioEffect.h
@@ -97,7 +97,7 @@ public:
 	 * @param Mix The SoundMix we want to switch to.
 	 * @param bIgnorePriority Whether EQPriority should be ignored to force mix change.
 	 */
-	void SetMixSettings( USoundMix* Mix, bool bIgnorePriority = false );
+	void SetMixSettings(USoundMix* Mix, bool bIgnorePriority = false, bool bForce = false);
 
 	/**
 	 * Clears the current SoundMix and any EQ settings it has applied
diff --git a/Engine/Source/Runtime/Engine/Public/IAudioExtensionPlugin.h b/Engine/Source/Runtime/Engine/Public/IAudioExtensionPlugin.h
--- a/Engine/Source/Runtime/Engine/Public/IAudioExtensionPlugin.h
+++ b/Engine/Source/Runtime/Engine/Public/IAudioExtensionPlugin.h
@@ -11,14 +11,21 @@ class FAudioSpatializationParams
 public:
 	FAudioSpatializationParams()
 		: EmitterPosition(0.0f)
+		, EmitterWorldPosition(0.0f)
 	{}
 
-	FAudioSpatializationParams(const FVector& InEmitterPosition)
+	FAudioSpatializationParams(const FVector& InEmitterPosition, const FVector& InEmitterWorldPosition)
 		: EmitterPosition(InEmitterPosition)
+		, EmitterWorldPosition(InEmitterWorldPosition)
 	{}
 
 	// Add any more params one might need here...
+
+	/** Normalized emitter position relative to the listener */
 	FVector EmitterPosition;
+
+	/** Raw emitter world position, untransformed relative to listener */
+	FVector EmitterWorldPosition;
 };
 
 /**
diff --git a/Engine/Source/Runtime/Engine/Public/VorbisAudioInfo.h b/Engine/Source/Runtime/Engine/Public/VorbisAudioInfo.h
--- a/Engine/Source/Runtime/Engine/Public/VorbisAudioInfo.h
+++ b/Engine/Source/Runtime/Engine/Public/VorbisAudioInfo.h
@@ -86,5 +86,8 @@ public:
 	const uint8*		SrcBufferData;
 	uint32			SrcBufferDataSize;
 	uint32			BufferOffset;
+
+	/** Critical section used to prevent multiple threads accessing same ogg-vorbis file handles at the same time */
+	FCriticalSection VorbisCriticalSection;
 };
 #endif
diff --git a/Engine/Source/Runtime/Mac/CoreAudio/Private/CoreAudioEffects.cpp b/Engine/Source/Runtime/Mac/CoreAudio/Private/CoreAudioEffects.cpp
--- a/Engine/Source/Runtime/Mac/CoreAudio/Private/CoreAudioEffects.cpp
+++ b/Engine/Source/Runtime/Mac/CoreAudio/Private/CoreAudioEffects.cpp
@@ -164,9 +164,10 @@ void FCoreAudioEffectsManager::SetReverbEffectParameters( const FAudioReverbEffe
  */
 void FCoreAudioEffectsManager::SetEQEffectParameters( const FAudioEQEffect& Params )
 {
-	float LowGain = VolumeToDeciBels(Params.LFGain);
-	float CenterGain = VolumeToDeciBels(Params.MFGain);
-	float HighGain = VolumeToDeciBels(Params.HFGain);
+	float Gain0 = VolumeToDeciBels(Params.Gain0);
+	float Gain1 = VolumeToDeciBels(Params.Gain1);
+	float Gain2 = VolumeToDeciBels(Params.Gain2);
+	float Gain3 = VolumeToDeciBels(Params.Gain3);
 
 	for( uint32 Index = 1; Index < CORE_AUDIO_MAX_CHANNELS + 1; Index++ )
 	{
@@ -175,21 +176,21 @@ void FCoreAudioEffectsManager::SetEQEffectParameters( const FAudioEQEffect& Para
 		{
 			if (Source->EQUnit)
 			{
-				AudioUnitSetParameter(Source->EQUnit, kAUNBandEQParam_Frequency + 0,	kAudioUnitScope_Global, 0, Params.LFFrequency,			0);
-				AudioUnitSetParameter(Source->EQUnit, kAUNBandEQParam_Gain + 0,			kAudioUnitScope_Global, 0, LowGain,						0);
-				AudioUnitSetParameter(Source->EQUnit, kAUNBandEQParam_Bandwidth + 0,	kAudioUnitScope_Global, 0, 1.0f,						0);	// from FXEQ_DEFAULT_BANDWIDTH
+				AudioUnitSetParameter(Source->EQUnit, kAUNBandEQParam_Frequency + 0,	kAudioUnitScope_Global, 0, Params.FrequencyCenter0,		0);
+				AudioUnitSetParameter(Source->EQUnit, kAUNBandEQParam_Gain + 0,			kAudioUnitScope_Global, 0, Gain0,						0);
+				AudioUnitSetParameter(Source->EQUnit, kAUNBandEQParam_Bandwidth + 0,	kAudioUnitScope_Global, 0, Params.Bandwidth0,			0);
 
-				AudioUnitSetParameter(Source->EQUnit, kAUNBandEQParam_Frequency + 1,	kAudioUnitScope_Global, 0, Params.MFCutoffFrequency,	0);
-				AudioUnitSetParameter(Source->EQUnit, kAUNBandEQParam_Gain + 1,			kAudioUnitScope_Global, 0, CenterGain,					0);
-				AudioUnitSetParameter(Source->EQUnit, kAUNBandEQParam_Bandwidth + 1,	kAudioUnitScope_Global, 0, Params.MFBandwidth,			0);
+				AudioUnitSetParameter(Source->EQUnit, kAUNBandEQParam_Frequency + 1,	kAudioUnitScope_Global, 0, Params.FrequencyCenter1,		0);
+				AudioUnitSetParameter(Source->EQUnit, kAUNBandEQParam_Gain + 1,			kAudioUnitScope_Global, 0, Gain1,						0);
+				AudioUnitSetParameter(Source->EQUnit, kAUNBandEQParam_Bandwidth + 1,	kAudioUnitScope_Global, 0, Params.Bandwidth1,			0);
 
-				AudioUnitSetParameter(Source->EQUnit, kAUNBandEQParam_Frequency + 2,	kAudioUnitScope_Global, 0, Params.HFFrequency,			0);
-				AudioUnitSetParameter(Source->EQUnit, kAUNBandEQParam_Gain + 2,			kAudioUnitScope_Global, 0, HighGain,					0);
-				AudioUnitSetParameter(Source->EQUnit, kAUNBandEQParam_Bandwidth + 2,	kAudioUnitScope_Global, 0, 1.0f,						0); // from FXEQ_DEFAULT_BANDWIDTH
+				AudioUnitSetParameter(Source->EQUnit, kAUNBandEQParam_Frequency + 2,	kAudioUnitScope_Global, 0, Params.FrequencyCenter2,		0);
+				AudioUnitSetParameter(Source->EQUnit, kAUNBandEQParam_Gain + 2,			kAudioUnitScope_Global, 0, Gain2,						0);
+				AudioUnitSetParameter(Source->EQUnit, kAUNBandEQParam_Bandwidth + 2,	kAudioUnitScope_Global, 0, Params.Bandwidth2,			0);
 
-				AudioUnitSetParameter(Source->EQUnit, kAUNBandEQParam_Frequency + 3,	kAudioUnitScope_Global, 0, 10000.0f,					0);	// from FXEQ_DEFAULT_CENTER_3
-				AudioUnitSetParameter(Source->EQUnit, kAUNBandEQParam_Gain + 3,			kAudioUnitScope_Global, 0, 1.0,							0); // from FXEQ_DEFAULT_GAIN
-				AudioUnitSetParameter(Source->EQUnit, kAUNBandEQParam_Bandwidth + 3,	kAudioUnitScope_Global, 0, 1.0f,						0); // from FXEQ_DEFAULT_BANDWIDTH
+				AudioUnitSetParameter(Source->EQUnit, kAUNBandEQParam_Frequency + 3,	kAudioUnitScope_Global, 0, Params.FrequencyCenter3,		0);
+				AudioUnitSetParameter(Source->EQUnit, kAUNBandEQParam_Gain + 3,			kAudioUnitScope_Global, 0, Gain3,						0);
+				AudioUnitSetParameter(Source->EQUnit, kAUNBandEQParam_Bandwidth + 3,	kAudioUnitScope_Global, 0, Params.Bandwidth3,			0);
 			}
 
 			if (Source->LowPassUnit && Source->LPFFrequency < MAX_FILTER_FREQUENCY)
diff --git a/Engine/Source/Runtime/RHI/Public/RHICommandList.h b/Engine/Source/Runtime/RHI/Public/RHICommandList.h
--- a/Engine/Source/Runtime/RHI/Public/RHICommandList.h
+++ b/Engine/Source/Runtime/RHI/Public/RHICommandList.h
@@ -134,7 +134,6 @@ public:
 
 	FORCEINLINE_DEBUGGABLE void* Alloc(int32 AllocSize, int32 Alignment)
 	{
-		checkSlow(!Bypass());
 		return MemManager.Alloc(AllocSize, Alignment);
 	}
 
diff --git a/Engine/Source/Runtime/Renderer/Private/ShadowSetup.cpp b/Engine/Source/Runtime/Renderer/Private/ShadowSetup.cpp
--- a/Engine/Source/Runtime/Renderer/Private/ShadowSetup.cpp
+++ b/Engine/Source/Runtime/Renderer/Private/ShadowSetup.cpp
@@ -506,7 +506,6 @@ bool FProjectedShadowInfo::SetupPerObjectProjection(
 
 		if (ResolutionX == 0 || ResolutionY == 0)
 		{
-			checkfSlow(0, TEXT("Tried to create shadow with invalid resolution [%d * %d]."), ResolutionX, ResolutionY);
 			bRet = false;
 		}
 		else
diff --git a/Engine/Source/Runtime/Windows/XAudio2/Private/XAudio2Effects.cpp b/Engine/Source/Runtime/Windows/XAudio2/Private/XAudio2Effects.cpp
--- a/Engine/Source/Runtime/Windows/XAudio2/Private/XAudio2Effects.cpp
+++ b/Engine/Source/Runtime/Windows/XAudio2/Private/XAudio2Effects.cpp
@@ -588,21 +588,21 @@ void FXAudio2EffectsManager::SetEQEffectParameters( const FAudioEQEffect& EQEffe
 	{
 		FXEQ_PARAMETERS NativeParameters;
 
-		NativeParameters.FrequencyCenter0 = EQEffectParameters.LFFrequency;
-		NativeParameters.Gain0 = EQEffectParameters.LFGain;
-		NativeParameters.Bandwidth0 = FXEQ_DEFAULT_BANDWIDTH;
+		NativeParameters.FrequencyCenter0 = EQEffectParameters.FrequencyCenter0;
+		NativeParameters.Gain0 = EQEffectParameters.Gain0;
+		NativeParameters.Bandwidth0 = EQEffectParameters.Bandwidth0;
 
-		NativeParameters.FrequencyCenter1 = EQEffectParameters.MFCutoffFrequency;
-		NativeParameters.Gain1 = EQEffectParameters.MFGain;
-		NativeParameters.Bandwidth1 = EQEffectParameters.MFBandwidth;
+		NativeParameters.FrequencyCenter1 = EQEffectParameters.FrequencyCenter1;
+		NativeParameters.Gain1 = EQEffectParameters.Gain1;
+		NativeParameters.Bandwidth1 = EQEffectParameters.Bandwidth1;
 
-		NativeParameters.FrequencyCenter2 = EQEffectParameters.HFFrequency;
-		NativeParameters.Gain2 = EQEffectParameters.HFGain;
-		NativeParameters.Bandwidth2 = FXEQ_DEFAULT_BANDWIDTH;
+		NativeParameters.FrequencyCenter2 = EQEffectParameters.FrequencyCenter2;
+		NativeParameters.Gain2 = EQEffectParameters.Gain2;
+		NativeParameters.Bandwidth2 = EQEffectParameters.Bandwidth2;
 
-		NativeParameters.FrequencyCenter3 = FXEQ_DEFAULT_FREQUENCY_CENTER_3;
-		NativeParameters.Gain3 = FXEQ_DEFAULT_GAIN;
-		NativeParameters.Bandwidth3 = FXEQ_DEFAULT_BANDWIDTH;
+		NativeParameters.FrequencyCenter3 = EQEffectParameters.FrequencyCenter3;
+		NativeParameters.Gain3 = EQEffectParameters.Gain3;
+		NativeParameters.Bandwidth3 = EQEffectParameters.Bandwidth3;
 
 		XAudio2Device->ValidateAPICall( TEXT( "SetEffectParameters (EQ)" ), 
 			EQPremasterVoice->SetEffectParameters( 0, &NativeParameters, sizeof( NativeParameters ) ) );
diff --git a/Engine/Source/Runtime/Windows/XAudio2/Private/XAudio2Source.cpp b/Engine/Source/Runtime/Windows/XAudio2/Private/XAudio2Source.cpp
--- a/Engine/Source/Runtime/Windows/XAudio2/Private/XAudio2Source.cpp
+++ b/Engine/Source/Runtime/Windows/XAudio2/Private/XAudio2Source.cpp
@@ -626,7 +626,7 @@ void FXAudio2SoundSource::GetMonoChannelVolumes(float ChannelVolumes[CHANNEL_MAT
 		}
 		check(AudioDevice->SpatializeProcessor != nullptr);
 
-		AudioDevice->SpatializeProcessor->SetSpatializationParameters(VoiceId, FAudioSpatializationParams(SpatializationParams.EmitterPosition));
+		AudioDevice->SpatializeProcessor->SetSpatializationParameters(VoiceId, FAudioSpatializationParams(SpatializationParams.EmitterPosition, WaveInstance->Location));
 		GetStereoChannelVolumes(ChannelVolumes, AttenuatedVolume);
 	}
 	else // Spatialize the mono stream using the normal 3d audio algorithm
