Linux: fix native build.

diff --git a/Engine/Source/Runtime/Engine/Private/TextureStreamingBuild.cpp b/Engine/Source/Runtime/Engine/Private/TextureStreamingBuild.cpp
--- a/Engine/Source/Runtime/Engine/Private/TextureStreamingBuild.cpp
+++ b/Engine/Source/Runtime/Engine/Private/TextureStreamingBuild.cpp
@@ -5,7 +5,7 @@ TextureStreamingBuild.cpp : Contains definitions to build texture streaming data
 =============================================================================*/
 
 #include "EnginePrivate.h"
-#include "DebugVieWModeMaterialProxy.h"
+#include "DebugViewModeMaterialProxy.h"
 #include "ShaderCompiler.h"
 
 ENGINE_API void BuildTextureStreamingData(UWorld* InWorld)
