build error
emscripten is getting confused with struct size of a giant blob instead of multiples of i32's

diff --git a/Engine/Source/Runtime/Engine/Private/UnrealClient.cpp b/Engine/Source/Runtime/Engine/Private/UnrealClient.cpp
--- a/Engine/Source/Runtime/Engine/Private/UnrealClient.cpp
+++ b/Engine/Source/Runtime/Engine/Private/UnrealClient.cpp
@@ -795,7 +795,7 @@ void FViewport::HighResScreenshot()
 	BeginInitResource(DummyViewport);
 
 	bool MaskShowFlagBackup = ViewportClient->GetEngineShowFlags()->HighResScreenshotMask;
-	const uint32 MotionBlurShowFlagBackup = ViewportClient->GetEngineShowFlags()->MotionBlur;
+	const bool MotionBlurShowFlagBackup = ViewportClient->GetEngineShowFlags()->MotionBlur;
 
 	ViewportClient->GetEngineShowFlags()->SetHighResScreenshotMask(GetHighResScreenshotConfig().bMaskEnabled);
 	ViewportClient->GetEngineShowFlags()->SetMotionBlur(false);
diff --git a/Engine/Source/Runtime/Engine/Public/ShowFlags.h b/Engine/Source/Runtime/Engine/Public/ShowFlags.h
--- a/Engine/Source/Runtime/Engine/Public/ShowFlags.h
+++ b/Engine/Source/Runtime/Engine/Public/ShowFlags.h
@@ -50,7 +50,7 @@ struct FEngineShowFlags
 	// ---------------------------------------------------------
 	// Define the showflags.
 	// A show flag is either an uint32:1 or static const bool (if optimized out according to UE_BUILD_OPTIMIZED_SHOWFLAGS)
-	#define SHOWFLAG_ALWAYS_ACCESSIBLE(a,...) uint32 a : 1; void Set##a(bool bVal){ a = bVal?1:0;}
+	#define SHOWFLAG_ALWAYS_ACCESSIBLE(a,...) bool a; void Set##a(bool bVal){ a = bVal?1:0;}
 
 	#if UE_BUILD_OPTIMIZED_SHOWFLAGS 
 		#define SHOWFLAG_FIXED_IN_SHIPPING(v,a,...) static const bool a = v; void Set##a(bool bVal){}
