build error - redo

emscripten is getting confused with struct size of a giant blob instead of multiples of i32's
using #if PLATFORM_HTML5 to guard this against other platforms -- i.e. Paragon has been seeing black screens when using bool instead of uint32:1
- this will be sent as a test case to emscripten &/or clang for compiler bit field bug

diff --git a/Engine/Source/Runtime/Engine/Private/UnrealClient.cpp b/Engine/Source/Runtime/Engine/Private/UnrealClient.cpp
--- a/Engine/Source/Runtime/Engine/Private/UnrealClient.cpp
+++ b/Engine/Source/Runtime/Engine/Private/UnrealClient.cpp
@@ -795,7 +795,7 @@ void FViewport::HighResScreenshot()
 	BeginInitResource(DummyViewport);
 
 	bool MaskShowFlagBackup = ViewportClient->GetEngineShowFlags()->HighResScreenshotMask;
-	const uint32 MotionBlurShowFlagBackup = ViewportClient->GetEngineShowFlags()->MotionBlur;
+	const auto MotionBlurShowFlagBackup = ViewportClient->GetEngineShowFlags()->MotionBlur;
 
 	ViewportClient->GetEngineShowFlags()->SetHighResScreenshotMask(GetHighResScreenshotConfig().bMaskEnabled);
 	ViewportClient->GetEngineShowFlags()->SetMotionBlur(false);
diff --git a/Engine/Source/Runtime/Engine/Public/ShowFlags.h b/Engine/Source/Runtime/Engine/Public/ShowFlags.h
--- a/Engine/Source/Runtime/Engine/Public/ShowFlags.h
+++ b/Engine/Source/Runtime/Engine/Public/ShowFlags.h
@@ -50,7 +50,11 @@ struct FEngineShowFlags
 	// ---------------------------------------------------------
 	// Define the showflags.
 	// A show flag is either an uint32:1 or static const bool (if optimized out according to UE_BUILD_OPTIMIZED_SHOWFLAGS)
+#if PLATFORM_HTML5 // broken fit field compiler -- will be sending this file to the emscripten &/or clang keepers
+	#define SHOWFLAG_ALWAYS_ACCESSIBLE(a,...) bool; void Set##a(bool bVal){ a = bVal;}
+#else
 	#define SHOWFLAG_ALWAYS_ACCESSIBLE(a,...) uint32 a : 1; void Set##a(bool bVal){ a = bVal?1:0;}
+#endif
 
 	#if UE_BUILD_OPTIMIZED_SHOWFLAGS 
 		#define SHOWFLAG_FIXED_IN_SHIPPING(v,a,...) static const bool a = v; void Set##a(bool bVal){}
@@ -268,7 +272,11 @@ private:
 		}
 
 		// Most flags are on by default. With the following line we only need disable flags.
+#if PLATFORM_HTML5
+		FMemory::Memset(this, uint8(true), sizeof(*this));
+#else
 		FMemory::Memset(this, 0xff, sizeof(*this));
+#endif
 
 		// The following code sets what should be off by default.
 		SetVisualizeHDR(false);
