Copying //UE4/Orion-Staging to //UE4/Main (Origin //Orion/Dev-General @ 2870388)

diff --git a/Engine/Source/Runtime/Engine/Private/Particles/ParticleComponents.cpp b/Engine/Source/Runtime/Engine/Private/Particles/ParticleComponents.cpp
--- a/Engine/Source/Runtime/Engine/Private/Particles/ParticleComponents.cpp
+++ b/Engine/Source/Runtime/Engine/Private/Particles/ParticleComponents.cpp
@@ -2088,7 +2088,7 @@ void UParticleSystem::PostEditChangeProperty(FPropertyChangedEvent& PropertyChan
 		}
 	}
 
-	bShouldManageSignificance = GetLowestSignificance() != EParticleSignificanceLevel::Critical && !ContainsEmitterType(UParticleModuleTypeDataBeam2::StaticClass());
+	bShouldManageSignificance = GetLowestSignificance() != EParticleSignificanceLevel::Critical/* && !ContainsEmitterType(UParticleModuleTypeDataBeam2::StaticClass())*/;
 
 	//cap the WarmupTickRate to realistic values
 	if (WarmupTickRate <= 0)
@@ -2211,7 +2211,7 @@ void UParticleSystem::PostLoad()
 	}
 	}
 
-	bShouldManageSignificance = GetLowestSignificance() != EParticleSignificanceLevel::Critical && !ContainsEmitterType(UParticleModuleTypeDataBeam2::StaticClass());
+	bShouldManageSignificance = GetLowestSignificance() != EParticleSignificanceLevel::Critical /* && !ContainsEmitterType(UParticleModuleTypeDataBeam2::StaticClass())*/;
 
 	if (LODSettings.Num() == 0)
 	{
@@ -3109,7 +3109,6 @@ UParticleSystemComponent::UParticleSystemComponent(const FObjectInitializer& Obj
 
 	WarmupTime = 0.0f;
 	SecondsBeforeInactive = 1.0f;
-	bAllowSecondsBeforeInactive = true;
 	bIsTransformDirty = false;
 	bSkipUpdateDynamicDataDuringTick = false;
 	bIsViewRelevanceDirty = true;
@@ -3137,7 +3136,8 @@ UParticleSystemComponent::UParticleSystemComponent(const FObjectInitializer& Obj
 
 	RequiredSignificance = EParticleSignificanceLevel::Low;
 	LastSignificantTime = 0.0f;
-	bIsSignificant = 1;
+	bIsManagingSignificance = 0;
+	bWasManagingSignificance = 0;
 }
 
 void UParticleSystemComponent::SetRequiredSignificance(EParticleSignificanceLevel NewRequiredSignificance)
@@ -3156,17 +3156,24 @@ void UParticleSystemComponent::SetRequiredSignificance(EParticleSignificanceLeve
 	if (!IsComponentTickEnabled() && Reaction == EParticleSystemInsignificanceReaction::DisableTick && Template->GetHighestSignificance() >= NewRequiredSignificance)
 	{
 		//Set us to be significant again.
-		OnSignificanceChanged(true, true);
+		OnSignificanceChanged(true, true, true);
 	}
 	}
 
-void UParticleSystemComponent::OnSignificanceChanged(bool bSignificant, bool bApplyToEmitters)
+void UParticleSystemComponent::OnSignificanceChanged(bool bSignificant, bool bApplyToEmitters, bool bAsync)
 {
 	int32 LocalNumSignificantEmitters = 0;
 	if (bSignificant)
 	{
-		SetComponentTickEnabled(true);
-		bIsSignificant = true;
+		if (bAsync)
+		{
+			SetComponentTickEnabledAsync(true);
+		}
+		else
+		{
+			SetComponentTickEnabled(true);
+		}
+
 		if (bApplyToEmitters && EmitterInstances.Num() > 0)
 		{
 			//Mark any emitters as significant if needed.
@@ -3197,8 +3204,15 @@ void UParticleSystemComponent::OnSignificanceChanged(bool bSignificant, bool bAp
 	}
 	else
 	{
-		SetComponentTickEnabled(false);
-		bIsSignificant = false;
+		if (bAsync)
+		{
+			SetComponentTickEnabledAsync(false);
+		}
+		else
+		{
+			SetComponentTickEnabled(false);
+		}
+
 		if (bApplyToEmitters && EmitterInstances.Num() > 0)
 		{
 			//Mark any emitters as significant if needed.
@@ -3253,7 +3267,6 @@ void UParticleSystemComponent::OnSignificanceChanged(bool bSignificant, bool bAp
 			break;
 		}
 	}
-	MarkRenderStateDirty();
 }
 
 bool UParticleSystemComponent::ShouldManageSignificance() const
@@ -3261,6 +3274,35 @@ bool UParticleSystemComponent::ShouldManageSignificance() const
 	return Template ? Template->ShouldManageSignificance() : false;
 }
 
+float UParticleSystemComponent::GetApproxDistanceSquared(FVector Point) const
+{
+	return Bounds.ComputeSquaredDistanceFromBoxToPoint(Point);
+	//TODO: Consider beam line segment?
+}
+
+bool UParticleSystemComponent::CanBeOccluded()const
+{
+	return Template->OcclusionBoundsMethod != EPSOBM_None && 
+		(Template->FixedRelativeBoundingBox.IsValid || (Template->OcclusionBoundsMethod == EPSOBM_CustomBounds)); //We can only be occluded if we have fixed bounds or custom occlusion bounds.
+}
+
+bool UParticleSystemComponent::CanConsiderInvisible()const
+{
+	if (World)
+	{
+		const float MaxSecondsBeforeInactive = FMath::Max(SecondsBeforeInactive, Template->SecondsBeforeInactive);
+
+		// Clamp MaxSecondsBeforeInactive to be at least twice the maximum smoothed frame time (45.45ms) because the rendering thread runs one 
+		// frame behind the game thread and so smaller time differences cannot be reliably detected.
+		const float ClampedMaxSecondsBeforeInactive = MaxSecondsBeforeInactive > 0 ? FMath::Max(MaxSecondsBeforeInactive, 0.1f) : 0.0f;
+		if (ClampedMaxSecondsBeforeInactive > 0.0f && AccumTickTime > ClampedMaxSecondsBeforeInactive && World->IsGameWorld())
+		{
+			return World->GetTimeSeconds() > (LastRenderTime + ClampedMaxSecondsBeforeInactive);
+		}
+	}
+	return false;
+}
+
 #if WITH_EDITOR
 void UParticleSystemComponent::CheckForErrors()
 {
@@ -3708,6 +3750,8 @@ FParticleDynamicData* UParticleSystemComponent::CreateDynamicData()
 {
 	//SCOPE_CYCLE_COUNTER(STAT_ParticleSystemComponent_CreateDynamicData);
 
+	FInGameScopedCycleCounter InGameCycleCounter(GetWorld(), EInGamePerfTrackers::VFXSignificance, EInGamePerfTrackerThreads::GameThread, bIsManagingSignificance);
+
 	// Only proceed if we have any live particles or if we're actively replaying/capturing
 	if (EmitterInstances.Num() > 0)
 	{
@@ -4242,8 +4286,7 @@ DECLARE_CYCLE_STAT(TEXT("PSys Comp Marshall Time"),STAT_UParticleSystemComponent
 
 void UParticleSystemComponent::TickComponent(float DeltaTime, enum ELevelTick TickType, FActorComponentTickFunction *ThisTickFunction)
 {
-	bool bManagingSignificance = ShouldManageSignificance();
-	FInGameScopedCycleCounter InGameCycleCounter(GetWorld(), EInGamePerfTrackers::VFXSignificance, EInGamePerfTrackerThreads::GameThread, bManagingSignificance);
+	FInGameScopedCycleCounter InGameCycleCounter(GetWorld(), EInGamePerfTrackers::VFXSignificance, EInGamePerfTrackerThreads::GameThread, bIsManagingSignificance);
 
 	if (Template == nullptr || Template->Emitters.Num() == 0)
 	{
@@ -4264,6 +4307,12 @@ void UParticleSystemComponent::TickComponent(float DeltaTime, enum ELevelTick Ti
 	ForceAsyncWorkCompletion(ENSURE_AND_STALL);
 	SCOPE_CYCLE_COUNTER(STAT_PSysCompTickTime);
 
+	if (bWasManagingSignificance != bIsManagingSignificance)
+	{	
+		bWasManagingSignificance = bIsManagingSignificance;
+		MarkRenderStateDirty();
+	}
+
 	bool bDisallowAsync = false;
 
 	// Bail out if inactive and not AutoActivate
@@ -4346,34 +4395,22 @@ void UParticleSystemComponent::TickComponent(float DeltaTime, enum ELevelTick Ti
 	{
 		//For now, we're only allowing the SecondsBeforeInactive optimization on looping emitters as it can cause leaks with non-looping effects.
 		//Longer term, there is likely a better solution.
-		if (bAllowSecondsBeforeInactive && Template->IsLooping())
+		if (Template->IsLooping() && CanConsiderInvisible())
 		{
-			const float MaxSecondsBeforeInactive = FMath::Max(SecondsBeforeInactive, Template->SecondsBeforeInactive);
-
-			// Clamp MaxSecondsBeforeInactive to be at least twice the maximum
-			// smoothed frame time (45.45ms) because the rendering thread runs one 
-			// frame behind the game thread and so smaller time differences cannot 
-			// be reliably detected.
-			const float ClampedMaxSecondsBeforeInactive = MaxSecondsBeforeInactive > 0 ? FMath::Max(MaxSecondsBeforeInactive, 0.1f) : 0;
+			SCOPE_CYCLE_COUNTER(STAT_UParticleSystemComponent_LOD_Inactive);
+			bForcedInActive = true;
+			SpawnEvents.Empty();
+			DeathEvents.Empty();
+			CollisionEvents.Empty();
+			KismetEvents.Empty();
 
-			if (ClampedMaxSecondsBeforeInactive > 0
-				&& AccumTickTime > ClampedMaxSecondsBeforeInactive//SecondsBeforeInactive
-				&&	GetWorld()->IsGameWorld())
+			if (bIsManagingSignificance && Template->GetHighestSignificance() < RequiredSignificance)
 			{
-				SCOPE_CYCLE_COUNTER(STAT_UParticleSystemComponent_LOD_Inactive);
-				const float CurrentTimeSeconds = World->GetTimeSeconds();
-				if (CurrentTimeSeconds > (LastRenderTime + ClampedMaxSecondsBeforeInactive))
-				{
-					bForcedInActive = true;
-
-					SpawnEvents.Empty();
-					DeathEvents.Empty();
-					CollisionEvents.Empty();
-					KismetEvents.Empty();
-
-					return;
-				}
+				//We're definitely insignificant so we can stop ticking entirely.
+				OnSignificanceChanged(false, true);
 			}
+
+			return;
 		}
 
 		AccumLODDistanceCheckTime += DeltaTime;
@@ -4533,8 +4570,7 @@ int32 UParticleSystemComponent::GetCurrentDetailMode() const
 
 void UParticleSystemComponent::ComputeTickComponent_Concurrent()
 {
-	bool bLocalShouldManageSignificance = ShouldManageSignificance();
-	FInGameScopedCycleCounter InGameCycleCounter(GetWorld(), EInGamePerfTrackers::VFXSignificance, IsInGameThread() ? EInGamePerfTrackerThreads::GameThread : EInGamePerfTrackerThreads::OtherThread, bLocalShouldManageSignificance);
+	FInGameScopedCycleCounter InGameCycleCounter(GetWorld(), EInGamePerfTrackers::VFXSignificance, IsInGameThread() ? EInGamePerfTrackerThreads::GameThread : EInGamePerfTrackerThreads::OtherThread, bIsManagingSignificance);
 
 	SCOPE_CYCLE_COUNTER(STAT_ParticleComputeTickTime);
 	FScopeCycleCounterUObject AdditionalScope(AdditionalStatObject(), GET_STATID(STAT_ParticleComputeTickTime));
@@ -4559,7 +4595,7 @@ void UParticleSystemComponent::ComputeTickComponent_Concurrent()
 			UParticleLODLevel* SpriteLODLevel = Instance->SpriteTemplate->GetCurrentLODLevel(Instance);
 			if (SpriteLODLevel && SpriteLODLevel->bEnabled)
 			{
-				if (bLocalShouldManageSignificance)
+				if (bIsManagingSignificance)
 				{
 					bool bEmitterIsSignificant = Instance->SpriteTemplate->IsSignificant(RequiredSignificance);
 					if (bEmitterIsSignificant)
@@ -4609,8 +4645,7 @@ void UParticleSystemComponent::ComputeTickComponent_Concurrent()
 
 void UParticleSystemComponent::FinalizeTickComponent()
 {
-	bool bLocalShouldManageSignificance = ShouldManageSignificance();
-	FInGameScopedCycleCounter InGameCycleCounter(GetWorld(), EInGamePerfTrackers::VFXSignificance, EInGamePerfTrackerThreads::GameThread, bLocalShouldManageSignificance);
+	FInGameScopedCycleCounter InGameCycleCounter(GetWorld(), EInGamePerfTrackers::VFXSignificance, IsInGameThread() ? EInGamePerfTrackerThreads::GameThread : EInGamePerfTrackerThreads::OtherThread, bIsManagingSignificance);
 
 	SCOPE_CYCLE_COUNTER(STAT_ParticleFinalizeTickTime);
 
@@ -4675,9 +4710,9 @@ void UParticleSystemComponent::FinalizeTickComponent()
 	float CurrTime = GetWorld()->GetTimeSeconds();
 
 	//Are we still significant?
-	if ((bIsActive && !bWasDeactivated) && bLocalShouldManageSignificance && NumSignificantEmitters == 0 && CurrTime >= LastSignificantTime + Template->InsignificanceDelay)
+	if ((bIsActive && !bWasDeactivated) && bIsManagingSignificance && NumSignificantEmitters == 0 && CurrTime >= LastSignificantTime + Template->InsignificanceDelay)
 	{
-		OnSignificanceChanged(false, false);
+		OnSignificanceChanged(false, true);
 	}
 	else
 	{
@@ -5057,9 +5092,6 @@ void UParticleSystemComponent::SetTemplate(class UParticleSystem* NewTemplate)
 
 void UParticleSystemComponent::ActivateSystem(bool bFlagAsJustAttached)
 {
-	bool bManagingSignificance = ShouldManageSignificance();
-	FInGameScopedCycleCounter InGameCycleCounter(GetWorld(), EInGamePerfTrackers::VFXSignificance, EInGamePerfTrackerThreads::GameThread, bManagingSignificance);
-
 	SCOPE_CYCLE_COUNTER(STAT_ParticleActivateTime);
 	ForceAsyncWorkCompletion(STALL);
 
@@ -5134,16 +5166,20 @@ void UParticleSystemComponent::ActivateSystem(bool bFlagAsJustAttached)
 			}
 		}
 
+		AccumTickTime = 0.0;
+
 		if (!bIsActive)
 		{
 			LastSignificantTime = GetWorld()->GetTimeSeconds();
-			bIsSignificant = true;
 			RequiredSignificance = EParticleSignificanceLevel::Low;
 
 		//Call this now after any attachment has happened.
 			OnSystemPreActivationChange.Broadcast(this, true);
 		}
 
+		//We start this here as before the PreActivation call above, we don't know if this component is managing significance or not.
+		FInGameScopedCycleCounter InGameCycleCounter(GetWorld(), EInGamePerfTrackers::VFXSignificance, EInGamePerfTrackerThreads::GameThread, bIsManagingSignificance);
+
 		if (bFlagAsJustAttached)
 		{
 			bJustRegistered = true;
@@ -5246,7 +5282,6 @@ void UParticleSystemComponent::ActivateSystem(bool bFlagAsJustAttached)
 			WarmupTime = 0.0f;
 			bSkipUpdateDynamicDataDuringTick = bSaveSkipUpdate;
 		}
-		AccumTickTime = 0.0;
 	}
 
 	// Mark render state dirty to ensure the scene proxy is added and registered with the scene.
@@ -5258,7 +5293,7 @@ void UParticleSystemComponent::ActivateSystem(bool bFlagAsJustAttached)
 	}
 
 	//We are definitely insignificant already so set insignificant before we ever begin ticking.
-	if (bManagingSignificance && Template->GetHighestSignificance() < RequiredSignificance && Template->InsignificanceDelay == 0.0f)
+	if (bIsManagingSignificance && Template->GetHighestSignificance() < RequiredSignificance && Template->InsignificanceDelay == 0.0f)
 	{
 		OnSignificanceChanged(false, true);
 	}
@@ -5289,8 +5324,7 @@ void UParticleSystemComponent::Complete()
 
 void UParticleSystemComponent::DeactivateSystem()
 {
-	bool bManagingSignificance = ShouldManageSignificance();
-	FInGameScopedCycleCounter InGameCycleCounter(GetWorld(), EInGamePerfTrackers::VFXSignificance, EInGamePerfTrackerThreads::GameThread, bManagingSignificance);
+	FInGameScopedCycleCounter InGameCycleCounter(GetWorld(), EInGamePerfTrackers::VFXSignificance, EInGamePerfTrackerThreads::GameThread, bIsManagingSignificance);
 
 	if (IsTemplate() == true)
 	{
diff --git a/Engine/Source/Runtime/Engine/Private/Particles/ParticleSystemRender.cpp b/Engine/Source/Runtime/Engine/Private/Particles/ParticleSystemRender.cpp
--- a/Engine/Source/Runtime/Engine/Private/Particles/ParticleSystemRender.cpp
+++ b/Engine/Source/Runtime/Engine/Private/Particles/ParticleSystemRender.cpp
@@ -7012,14 +7012,14 @@ FParticleSystemOcclusionSceneProxy::~FParticleSystemOcclusionSceneProxy()
 
 FPrimitiveSceneProxy* UParticleSystemComponent::CreateSceneProxy()
 {
-	bool bManagingSignificance = ShouldManageSignificance();
-	FInGameScopedCycleCounter InGameCycleCounter(GetWorld(), EInGamePerfTrackers::VFXSignificance, EInGamePerfTrackerThreads::GameThread, bManagingSignificance);
 	FParticleSystemSceneProxy* NewProxy = NULL;
 
 	//@fixme EmitterInstances.Num() check should be here to avoid proxies for dead emitters but there are some edge cases where it happens for emitters that have just activated...
 	//@fixme Get non-instanced path working in ES2!
-	if ((bIsActive == true)/** && (EmitterInstances.Num() > 0)*/ && Template && bIsSignificant)
+	if ((bIsActive == true)/** && (EmitterInstances.Num() > 0)*/ && Template)
 	{
+		FInGameScopedCycleCounter InGameCycleCounter(GetWorld(), EInGamePerfTrackers::VFXSignificance, EInGamePerfTrackerThreads::GameThread, bIsManagingSignificance);
+
 		UE_LOG(LogParticles,Verbose,
 			TEXT("CreateSceneProxy @ %fs %s bIsActive=%d"), GetWorld()->TimeSeconds,
 			Template != NULL ? *Template->GetName() : TEXT("NULL"), bIsActive);
@@ -7032,14 +7032,14 @@ FPrimitiveSceneProxy* UParticleSystemComponent::CreateSceneProxy()
 		// Create the dynamic data for rendering this particle system.
 		FParticleDynamicData* ParticleDynamicData = CreateDynamicData();
 
-		if (Template->OcclusionBoundsMethod == EPSOBM_None)
+		if (CanBeOccluded())
 		{
-			NewProxy = ::new FParticleSystemSceneProxy(this,ParticleDynamicData);
+			Template->CustomOcclusionBounds.IsValid = true;
+			NewProxy = ::new FParticleSystemOcclusionSceneProxy(this,ParticleDynamicData);
 		}
 		else
 		{
-			Template->CustomOcclusionBounds.IsValid = true;
-			NewProxy = ::new FParticleSystemOcclusionSceneProxy(this,ParticleDynamicData);
+			NewProxy = ::new FParticleSystemSceneProxy(this,ParticleDynamicData);
 		}
 		check (NewProxy);
 		if (ParticleDynamicData)
diff --git a/Engine/Source/Runtime/Engine/Private/Particles/ParticleTrail2EmitterInstance.cpp b/Engine/Source/Runtime/Engine/Private/Particles/ParticleTrail2EmitterInstance.cpp
--- a/Engine/Source/Runtime/Engine/Private/Particles/ParticleTrail2EmitterInstance.cpp
+++ b/Engine/Source/Runtime/Engine/Private/Particles/ParticleTrail2EmitterInstance.cpp
@@ -2097,36 +2097,57 @@ void FParticleRibbonEmitterInstance::UpdateSourceData(float DeltaTime, bool bFir
 	float TangentStrength;
 	// For each possible trail in this emitter, update it's source information
 	float ElapsedTime = RunningTime;//SecondsSinceCreation;
+	bool bCanBeValidParticleSource = ((SourceModule != NULL) && (SourceModule->SourceMethod == PET2SRCM_Particle));
 	for (int32 TrailIdx = 0; TrailIdx < MaxTrailCount; TrailIdx++)
 	{
 		bool bNewSource = (SourceIndices[TrailIdx] == -1);
 		if (ResolveSourcePoint(TrailIdx, Position, Rotation, Up, Tangent, TangentStrength) == true)
 		{
-			if ((bFirstTime == true) || 
-				((bNewSource == true) && ((SourceModule != NULL) && (SourceModule->SourceMethod == PET2SRCM_Particle))))
+			if (SourceIndices[TrailIdx] == -1 && bCanBeValidParticleSource)
 			{
+				//No valid particle for source so set all last and prev data to the same defaults. Stops issues with SpawnPerUnit trails having a segment back to the default location.
 				LastSourcePosition[TrailIdx] = Position;
-				LastSourceTangent[TrailIdx] = FVector::ZeroVector;//Component->LocalToWorld.TransformVector(FVector(1,0,0));
+				LastSourceTangent[TrailIdx] = Tangent;
 				LastSourceTangentStrength[TrailIdx] = TangentStrength;
+				LastSourceRotation[TrailIdx] = Rotation;
 				LastSourceUp[TrailIdx] = Up;
-				TrailSpawnTimes[TrailIdx] = RunningTime;
-			}
-			CurrentSourcePosition[TrailIdx] = Position;
-			CurrentSourceRotation[TrailIdx] = Rotation;
-			float ElapsedTimeSinceSpanwed = ElapsedTime - TrailSpawnTimes[TrailIdx];
-			if ( ElapsedTimeSinceSpanwed != 0.f )
-			{
-				CurrentSourceTangent[TrailIdx] = (CurrentSourcePosition[TrailIdx] - LastSourcePosition[TrailIdx]) / ElapsedTimeSinceSpanwed;
+
+				CurrentSourcePosition[TrailIdx] = Position;
+				CurrentSourceTangent[TrailIdx] = Tangent;
+				CurrentSourceTangentStrength[TrailIdx] = TangentStrength;
+				CurrentSourceRotation[TrailIdx] = Rotation;
+				CurrentSourceUp[TrailIdx] = Up;
+
+				TrailSpawnTimes[TrailIdx] = 0.0f;				
 			}
 			else
 			{
-				CurrentSourceTangent[TrailIdx] = FVector(1.f, 0.f, 0.f);
-			}
-			CurrentSourceTangentStrength[TrailIdx] = TangentStrength;
-			CurrentSourceUp[TrailIdx] = Up;
-			if (bFirstTime == true)
-			{
-				LastSourceRotation[TrailIdx] = CurrentSourceRotation[TrailIdx];
+				if ((bFirstTime == true) || 
+					((bNewSource == true) && bCanBeValidParticleSource))
+				{
+					LastSourcePosition[TrailIdx] = Position;
+					LastSourceTangent[TrailIdx] = FVector::ZeroVector;//Component->LocalToWorld.TransformVector(FVector(1,0,0));
+					LastSourceTangentStrength[TrailIdx] = TangentStrength;
+					LastSourceUp[TrailIdx] = Up;
+					TrailSpawnTimes[TrailIdx] = RunningTime;
+				}
+				CurrentSourcePosition[TrailIdx] = Position;
+				CurrentSourceRotation[TrailIdx] = Rotation;
+				float ElapsedTimeSinceSpanwed = ElapsedTime - TrailSpawnTimes[TrailIdx];
+				if ( ElapsedTimeSinceSpanwed != 0.f )
+				{
+					CurrentSourceTangent[TrailIdx] = (CurrentSourcePosition[TrailIdx] - LastSourcePosition[TrailIdx]) / ElapsedTimeSinceSpanwed;
+				}
+				else
+				{
+					CurrentSourceTangent[TrailIdx] = FVector(1.f, 0.f, 0.f);
+				}
+				CurrentSourceTangentStrength[TrailIdx] = TangentStrength;
+				CurrentSourceUp[TrailIdx] = Up;
+				if (bFirstTime == true)
+				{
+					LastSourceRotation[TrailIdx] = CurrentSourceRotation[TrailIdx];
+				}
 			}
 		}
 	}
@@ -2282,6 +2303,7 @@ bool FParticleRibbonEmitterInstance::ResolveSourcePoint(int32 InTrailIdx,
 						//@todo. How to handle this... can potentially cause a jump from the emitter to the
 						// particle...
 						SourceTimes[InTrailIdx] = 0.0f;
+						SourceIndices[InTrailIdx] = -1;//No valid particle source;
 					}
 					OutTangentStrength = OutTangent.SizeSquared();
 					//@todo. Allow particle rotation to define up??
diff --git a/Engine/Source/Runtime/Engine/Private/PlayerController.cpp b/Engine/Source/Runtime/Engine/Private/PlayerController.cpp
--- a/Engine/Source/Runtime/Engine/Private/PlayerController.cpp
+++ b/Engine/Source/Runtime/Engine/Private/PlayerController.cpp
@@ -1457,6 +1457,11 @@ void APlayerController::FOV(float F)
 
 void APlayerController::PreClientTravel( const FString& PendingURL, ETravelType TravelType, bool bIsSeamlessTravel )
 {
+	UGameInstance* GameInstance = GetGameInstance();
+	if (GameInstance)
+	{
+		GameInstance->NotifyPreClientTravel(PendingURL, TravelType, bIsSeamlessTravel);
+	}
 }
 
 void APlayerController::Camera( FName NewMode )
@@ -2922,12 +2927,12 @@ void APlayerController::LevelStreamingStatusChanged(ULevelStreaming* LevelObject
 void APlayerController::ClientPrepareMapChange_Implementation(FName LevelName, bool bFirst, bool bLast)
 {
 	// Only call on the first local player controller to handle it being called on multiple PCs for splitscreen.
-	if (GetWorld()->GetGameInstance() == nullptr)
+	if (GetGameInstance() == nullptr)
 	{
 		return;
 	}
 
-	APlayerController* PlayerController = GetWorld()->GetGameInstance()->GetFirstLocalPlayerController();
+	APlayerController* PlayerController = GetGameInstance()->GetFirstLocalPlayerController();
 	if( PlayerController != this )
 	{
 		return;
diff --git a/Engine/Source/Runtime/Engine/Private/RepLayout.cpp b/Engine/Source/Runtime/Engine/Private/RepLayout.cpp
--- a/Engine/Source/Runtime/Engine/Private/RepLayout.cpp
+++ b/Engine/Source/Runtime/Engine/Private/RepLayout.cpp
@@ -334,13 +334,13 @@ public:
 		}
 	}
 
-	void ProcessCmds( FRepState* RepState, uint8* RESTRICT Data )
+	void ProcessCmds( uint8* RESTRICT Data, uint8* RESTRICT ShadowData )
 	{
-		TStackState StackState( 0, Cmds.Num() - 1, NULL, NULL, (uint8*)RepState->StaticBuffer.GetData(), Data );
+		TStackState StackState( 0, Cmds.Num() - 1, NULL, NULL, ShadowData, Data );
 
 		static_cast< TImpl* >( this )->InitStack( StackState );
 
-		ProcessCmds_r( StackState, (uint8*)RepState->StaticBuffer.GetData(), Data );
+		ProcessCmds_r( StackState, ShadowData, Data );
 	}
 
 	const TArray< FRepParentCmd >&	Parents;
@@ -1184,14 +1184,15 @@ public:
 class FReceivePropertiesImpl : public FRepLayoutCmdIterator< FReceivePropertiesImpl, FReceivedPropertiesStackState >
 {
 public:
-	FReceivePropertiesImpl( FNetBitReader & InBunch, FRepState * InRepState, bool bInDoChecksum, const TArray< FRepParentCmd >& InParents, const TArray< FRepLayoutCmd >& InCmds ) :
+	FReceivePropertiesImpl( FNetBitReader & InBunch, FRepState * InRepState, bool bInDoChecksum, const TArray< FRepParentCmd >& InParents, const TArray< FRepLayoutCmd >& InCmds, const bool bInDoRepNotify ) :
         FRepLayoutCmdIterator( InParents, InCmds ),
 		WaitingHandle( 0 ),
 		CurrentHandle( 0 ), 
 		Bunch( InBunch ),
 		RepState( InRepState ),
 		bDoChecksum( bInDoChecksum ),
-		bHasUnmapped( false )
+		bHasUnmapped( false ),
+		bDoRepNotify( bInDoRepNotify )
 	{}
 
 	void ReadNextHandle()
@@ -1254,7 +1255,7 @@ public:
 
 		const FRepParentCmd& Parent = Parents[Cmd.ParentIndex];
 
-		if ( ( StackState.DataArray->Num() != ArrayNum || Parent.RepNotifyCondition == REPNOTIFY_Always ) && Parent.Property->HasAnyPropertyFlags( CPF_RepNotify ) )
+		if ( bDoRepNotify && ( StackState.DataArray->Num() != ArrayNum || Parent.RepNotifyCondition == REPNOTIFY_Always ) && Parent.Property->HasAnyPropertyFlags( CPF_RepNotify ) )
 		{
 			RepState->RepNotifies.AddUnique( Parent.Property );
 		}
@@ -1265,12 +1266,17 @@ public:
 		FScriptArrayHelper ArrayHelper( (UArrayProperty *)Cmd.Property, Data );
 		ArrayHelper.Resize( ArrayNum );
 
-		FScriptArrayHelper ShadowArrayHelper( (UArrayProperty *)Cmd.Property, ShadowData );
-		ShadowArrayHelper.Resize( ArrayNum );
-
 		// Re-compute the base data values since they could have changed after the resize above
 		StackState.BaseData			= (uint8*)StackState.DataArray->GetData();
-		StackState.ShadowBaseData	= (uint8*)StackState.ShadowArray->GetData();
+		
+		// Only resize the shadow data array if we're actually tracking RepNotifies
+		if (bDoRepNotify)
+		{
+			FScriptArrayHelper ShadowArrayHelper((UArrayProperty*)Cmd.Property, ShadowData);
+			ShadowArrayHelper.Resize(ArrayNum);
+
+			StackState.ShadowBaseData = (uint8*)StackState.ShadowArray->GetData();
+		}
 
 		// Save the old handle so we can restore it when we pop out of the array
 		const uint16 OldHandle = CurrentHandle;
@@ -1304,7 +1310,7 @@ public:
 		// Remember where we started reading from, so that if we have unmapped properties, we can re-deserialize from this data later
 		FBitReaderMark Mark( Bunch );
 
-		if ( Parent.Property->HasAnyPropertyFlags( CPF_RepNotify ) )
+		if ( bDoRepNotify && Parent.Property->HasAnyPropertyFlags( CPF_RepNotify ) )
 		{
 			// Copy current value over so we can check to see if it changed
 			StoreProperty( Cmd, ShadowData + Cmd.Offset, Data + SwappedCmd.Offset );
@@ -1374,9 +1380,10 @@ public:
 	FRepState *				RepState;
 	bool					bDoChecksum;
 	bool					bHasUnmapped;
+	bool					bDoRepNotify;
 };
 
-bool FRepLayout::ReceiveProperties( UClass * InObjectClass, FRepState * RESTRICT RepState, void* RESTRICT Data, FNetBitReader & InBunch, bool & bOutHasUnmapped ) const
+bool FRepLayout::ReceiveProperties( UClass * InObjectClass, FRepState * RESTRICT RepState, void* RESTRICT Data, FNetBitReader & InBunch, bool & bOutHasUnmapped, const bool bEnableRepNotifies ) const
 {
 	check( InObjectClass == Owner );
 
@@ -1388,13 +1395,13 @@ bool FRepLayout::ReceiveProperties( UClass * InObjectClass, FRepState * RESTRICT
 
 	bOutHasUnmapped = false;
 
-	FReceivePropertiesImpl ReceivePropertiesImpl( InBunch, RepState, bDoChecksum, Parents, Cmds );
+	FReceivePropertiesImpl ReceivePropertiesImpl( InBunch, RepState, bDoChecksum, Parents, Cmds, bEnableRepNotifies );
 
 	// Read first handle
 	ReceivePropertiesImpl.ReadNextHandle();
 
 	// Read all properties
-	ReceivePropertiesImpl.ProcessCmds( RepState, (uint8*)Data );
+	ReceivePropertiesImpl.ProcessCmds( (uint8*)Data, RepState->StaticBuffer.GetData() );
 
 	// Make sure we're waiting on the last NULL terminator
 	if ( ReceivePropertiesImpl.WaitingHandle != 0 )
@@ -1789,7 +1796,7 @@ void FRepLayout::MergeDirtyList( FRepState * RepState, const void* RESTRICT Data
 	MergePropertiesImpl.bDirtyValid2 = Dirty2.Num() > 0;
 
 	// Merge lists
-	MergePropertiesImpl.ProcessCmds( RepState, (uint8*)Data );
+	MergePropertiesImpl.ProcessCmds( (uint8*)Data, (uint8*)RepState->StaticBuffer.GetData() );
 
 	MergePropertiesImpl.MergedDirtyList.Add( 0 );
 }
@@ -1877,7 +1884,7 @@ void FRepLayout::SanityCheckChangeList( const uint8* RESTRICT Data, TArray< uint
 class FDiffPropertiesImpl : public FRepLayoutCmdIterator< FDiffPropertiesImpl, FCmdIteratorBaseStackState >
 {
 public:
-	FDiffPropertiesImpl( const bool bInSync, TArray< UProperty * >&	InRepNotifies, const TArray< FRepParentCmd >& InParents, const TArray< FRepLayoutCmd >& InCmds ) : 
+	FDiffPropertiesImpl( const bool bInSync, TArray< UProperty * >*	InRepNotifies, const TArray< FRepParentCmd >& InParents, const TArray< FRepLayoutCmd >& InCmds ) : 
 		FRepLayoutCmdIterator( InParents, InCmds ),
 		bSync( bInSync ),
 		RepNotifies( InRepNotifies ),
@@ -1944,23 +1951,23 @@ public:
 
 			StoreProperty( Cmd, (void*)( Data + Cmd.Offset ), (const void*)( ShadowData + Cmd.Offset ) );
 
-			if ( Parent.Property->HasAnyPropertyFlags( CPF_RepNotify ) )
+			if ( RepNotifies && Parent.Property->HasAnyPropertyFlags( CPF_RepNotify ) )
 			{
-				RepNotifies.AddUnique( Parent.Property );
+				RepNotifies->AddUnique( Parent.Property );
 			}
 		}
 	}
 
 	bool					bSync;
-	TArray< UProperty * >&	RepNotifies;
+	TArray< UProperty * >*	RepNotifies;
 	bool					bDifferent;
 };
 
-bool FRepLayout::DiffProperties( FRepState * RepState, const void* RESTRICT Data, const bool bSync ) const
+bool FRepLayout::DiffProperties( TArray<UProperty*>* RepNotifies, void* RESTRICT Destination, const void* RESTRICT Source, const bool bSync ) const
 {	
-	FDiffPropertiesImpl DiffPropertiesImpl( bSync, RepState->RepNotifies, Parents, Cmds );
+	FDiffPropertiesImpl DiffPropertiesImpl( bSync, RepNotifies, Parents, Cmds );
 
-	DiffPropertiesImpl.ProcessCmds( RepState, (uint8*)Data );
+	DiffPropertiesImpl.ProcessCmds( (uint8*)Destination, (uint8*)Source );
 
 	return DiffPropertiesImpl.bDifferent;
 }
diff --git a/Engine/Source/Runtime/Engine/Private/SkeletalMeshComponentPhysics.cpp b/Engine/Source/Runtime/Engine/Private/SkeletalMeshComponentPhysics.cpp
--- a/Engine/Source/Runtime/Engine/Private/SkeletalMeshComponentPhysics.cpp
+++ b/Engine/Source/Runtime/Engine/Private/SkeletalMeshComponentPhysics.cpp
@@ -2141,21 +2141,30 @@ bool USkeletalMeshComponent::ComponentOverlapMultiImpl(TArray<struct FOverlapRes
 
 
 
-void USkeletalMeshComponent::AddClothingBounds(FBoxSphereBounds& InOutBounds) const
+void USkeletalMeshComponent::AddClothingBounds(FBoxSphereBounds& InOutBounds, const FTransform& LocalToWorld) const
 {
 	for(const FClothingActor& ClothingActor : ClothingActors)
-		{
+	{
 		if(NxClothingActor* Actor = ClothingActor.ApexClothingActor)
+		{
+			physx::PxBounds3 ApexClothingBounds = Actor->getBounds();
+			if (!ApexClothingBounds.isEmpty())
 			{
-				physx::PxBounds3 ApexClothingBounds = Actor->getBounds();
+				const physx::PxMat44& ApexGlobalPose = Actor->getGlobalPose();
 
-				if (!ApexClothingBounds.isEmpty())
-				{
-					FBoxSphereBounds BoxBounds = FBox( P2UVector(ApexClothingBounds.minimum), P2UVector(ApexClothingBounds.maximum) );
-					InOutBounds = InOutBounds + BoxBounds;
-				}
+				FBoxSphereBounds BoxBounds = FBox( P2UVector(ApexClothingBounds.minimum), P2UVector(ApexClothingBounds.maximum) );
+
+				//It's possible that the cloth sim has never run, in which case its global pose is 0,0,0 - In this case the final bounds will be huge because of our LocalToWorld.
+				//It's also possible that after submitting ComponentToWorld on the game thread we modified it. In this case the cloth sim will have a different bounds origin, so again we'll get a huge final bounds
+				//To fix this we convert the bounds to local space and then to world. Note that this ignores scale! If this becomes an issue feel free to use both LocalToWorld and ApexGlobalPose to solve the issue.
+
+				BoxBounds.Origin -= P2UVector(ApexGlobalPose.getPosition());	//to local world
+				BoxBounds.Origin += LocalToWorld.GetLocation();					//to world as seen by our final bounds
+
+				InOutBounds = InOutBounds + BoxBounds;
 			}
 		}
+	}
 }
 
 bool USkeletalMeshComponent::HasValidClothingActors() const
@@ -3307,20 +3316,6 @@ void USkeletalMeshComponent::UpdateClothTransform()
 		ComponentToWorld.SetIdentity();
 	}
 //#endif
-
-	physx::PxMat44 PxGlobalPose = U2PMatrix(ComponentToWorld.ToMatrixWithScale());
-
-	for(FClothingActor& ClothingActor : ClothingActors)
-	{
-		if(NxClothingActor* ApexClothingActor = ClothingActor.ApexClothingActor)
-		{
-			NxParameterized::Interface* ActorDesc = ApexClothingActor->getActorDesc();
-			if(ActorDesc != nullptr)
-			{
-				verify(NxParameterized::setParamMat44(*ActorDesc, "globalPose", PxGlobalPose));
-			}
-		}
-	}
 }
 
 void USkeletalMeshComponent::CheckClothTeleport()
diff --git a/Engine/Source/Runtime/Engine/Private/SkeletalRenderGPUSkin.cpp b/Engine/Source/Runtime/Engine/Private/SkeletalRenderGPUSkin.cpp
--- a/Engine/Source/Runtime/Engine/Private/SkeletalRenderGPUSkin.cpp
+++ b/Engine/Source/Runtime/Engine/Private/SkeletalRenderGPUSkin.cpp
@@ -38,6 +38,9 @@ DEFINE_LOG_CATEGORY_STATIC(LogSkeletalGPUSkinMesh, Warning, All);
 DECLARE_CYCLE_STAT(TEXT("Morph Vertex Buffer Update"),STAT_MorphVertexBuffer_Update,STATGROUP_MorphTarget);
 DECLARE_CYCLE_STAT(TEXT("Morph Vertex Buffer Init"),STAT_MorphVertexBuffer_Init,STATGROUP_MorphTarget);
 DECLARE_CYCLE_STAT(TEXT("Morph Vertex Buffer Apply Delta"),STAT_MorphVertexBuffer_ApplyDelta,STATGROUP_MorphTarget);
+DECLARE_CYCLE_STAT(TEXT("Morph Vertex Buffer Alloc"), STAT_MorphVertexBuffer_Alloc, STATGROUP_MorphTarget);
+DECLARE_CYCLE_STAT(TEXT("Morph Vertex Buffer RHI Lock and copy"), STAT_MorphVertexBuffer_RhiLockAndCopy, STATGROUP_MorphTarget);
+DECLARE_CYCLE_STAT(TEXT("Morph Vertex Buffer RHI Unlock"), STAT_MorphVertexBuffer_RhiUnlock, STATGROUP_MorphTarget);
 
 static TAutoConsoleVariable<int32> CVarMotionBlurDebug(
 	TEXT("r.MotionBlurDebug"),
@@ -397,7 +400,11 @@ void FSkeletalMeshObjectGPUSkin::FSkeletalMeshObjectLOD::UpdateMorphVertexBuffer
 		FStaticLODModel& LodModel = SkelMeshResource->LODModels[LODIndex];
 		uint32 Size = LodModel.NumVertices * sizeof(FMorphGPUSkinVertex);
 
-		FMorphGPUSkinVertex* Buffer = (FMorphGPUSkinVertex*)FMemory::Malloc(Size);
+		FMorphGPUSkinVertex* Buffer = nullptr;
+		{
+			SCOPE_CYCLE_COUNTER(STAT_MorphVertexBuffer_Alloc)
+			Buffer = (FMorphGPUSkinVertex*)FMemory::Malloc(Size);
+		}
 
 		{
 			SCOPE_CYCLE_COUNTER(STAT_MorphVertexBuffer_Init);
@@ -473,14 +480,20 @@ void FSkeletalMeshObjectGPUSkin::FSkeletalMeshObjectLOD::UpdateMorphVertexBuffer
 		} // ApplyDelta
 
 		// Lock the real buffer.
-		FMorphGPUSkinVertex* ActualBuffer = (FMorphGPUSkinVertex*)RHILockVertexBuffer(MorphVertexBuffer.VertexBufferRHI, 0, Size, RLM_WriteOnly);
-		FMemory::Memcpy(ActualBuffer, Buffer, Size);
-		FMemory::Free(Buffer);
-
-		// Unlock the buffer.
-		RHIUnlockVertexBuffer(MorphVertexBuffer.VertexBufferRHI);
-		// set update flag
-		MorphVertexBuffer.bHasBeenUpdated = true;
+		{
+			SCOPE_CYCLE_COUNTER(STAT_MorphVertexBuffer_RhiLockAndCopy)
+			FMorphGPUSkinVertex* ActualBuffer = (FMorphGPUSkinVertex*)RHILockVertexBuffer(MorphVertexBuffer.VertexBufferRHI, 0, Size, RLM_WriteOnly);
+			FMemory::Memcpy(ActualBuffer, Buffer, Size);
+			FMemory::Free(Buffer);
+		}
+
+		{
+			SCOPE_CYCLE_COUNTER(STAT_MorphVertexBuffer_RhiUnlock)
+			// Unlock the buffer.
+			RHIUnlockVertexBuffer(MorphVertexBuffer.VertexBufferRHI);
+			// set update flag
+			MorphVertexBuffer.bHasBeenUpdated = true;
+		}
 	}
 }
 
diff --git a/Engine/Source/Runtime/Engine/Private/Slate/SceneViewport.cpp b/Engine/Source/Runtime/Engine/Private/Slate/SceneViewport.cpp
--- a/Engine/Source/Runtime/Engine/Private/Slate/SceneViewport.cpp
+++ b/Engine/Source/Runtime/Engine/Private/Slate/SceneViewport.cpp
@@ -1094,6 +1094,8 @@ void FSceneViewport::ResizeViewport(uint32 NewSizeX, uint32 NewSizeY, EWindowMod
 
 		UpdateViewportRHI(false, NewSizeX, NewSizeY, NewWindowMode);
 
+		FSystemResolution::RequestResolutionChange(NewSizeX, NewSizeY, NewWindowMode);
+
 		if (ViewportClient)
 		{
 			// Invalidate, then redraw immediately so the user isn't left looking at an empty black viewport
diff --git a/Engine/Source/Runtime/Engine/Private/Texture2D.cpp b/Engine/Source/Runtime/Engine/Private/Texture2D.cpp
--- a/Engine/Source/Runtime/Engine/Private/Texture2D.cpp
+++ b/Engine/Source/Runtime/Engine/Private/Texture2D.cpp
@@ -2075,6 +2075,21 @@ void FTexture2DResource::UploadMipData()
 	}
 }
 
+/**
+* Helper function for cleaning up bulk data files after streaming
+* @todo: make it smarter, close only when we know we won't be streaming anymore or at least for a while
+* @todo: What if each mip is in a different bulk data file? might need to loop over all
+*/
+inline void HintDoneWithStreamedTextureFiles(const UTexture2D* InTexture)
+{
+	if (FPlatformProperties::RequiresCookedData())
+	{
+		const TIndirectArray<FTexture2DMipMap>& OwnerMips = InTexture->GetPlatformMips();
+		const FTexture2DMipMap& MipMap = OwnerMips[0];
+		FIOSystem::Get().HintDoneWithFile(MipMap.BulkData.GetFilename());
+	}
+}
+
 /**
  * Called from the rendering thread to finalize a mip change.
  */
@@ -2115,13 +2130,7 @@ void FTexture2DResource::FinalizeMipCount()
 		// We're done.
 		Owner->PendingMipChangeRequestStatus.Decrement();
 
-		// @todo: make it smarter, close only when we know we won't be streaming anymore or at least for a while
-		if (FPlatformProperties::RequiresCookedData())
-		{
-			const TIndirectArray<FTexture2DMipMap>& OwnerMips = Owner->GetPlatformMips();
-			const FTexture2DMipMap& MipMap = OwnerMips[0];
-			FIOSystem::Get().HintDoneWithFile(MipMap.BulkData.GetFilename());
-		}
+		HintDoneWithStreamedTextureFiles(Owner);
 
 		return;
 	}
@@ -2189,6 +2198,8 @@ void FTexture2DResource::FinalizeMipCount()
 		IntermediateTextureRHI.SafeRelease();
 
 		GStreamMemoryTracker.RenderThread_Finalize( *Owner, bSuccess );
+
+		HintDoneWithStreamedTextureFiles(Owner);
 	}
 	else
 	{
diff --git a/Engine/Source/Runtime/Engine/Private/UnrealEngine.cpp b/Engine/Source/Runtime/Engine/Private/UnrealEngine.cpp
--- a/Engine/Source/Runtime/Engine/Private/UnrealEngine.cpp
+++ b/Engine/Source/Runtime/Engine/Private/UnrealEngine.cpp
@@ -2355,7 +2355,7 @@ MSVC_PRAGMA(warning(disable : 4717))
 // clang can optimize this out (http://stackoverflow.com/questions/18478078/clang-infinite-tail-recursion-optimization), make the function do some useful work
 
 volatile int GInfiniteRecursionCount = 0;
-static int InfiniteRecursionFunction(int B)
+int InfiniteRecursionFunction(int B)
 {
 	GInfiniteRecursionCount += InfiniteRecursionFunction(B + 1);
 	return GInfiniteRecursionCount;
@@ -4505,7 +4505,7 @@ bool UEngine::HandleMemCommand( const TCHAR* Cmd, FOutputDevice& Ar )
 static bool GDebugAllocMemEveryFrame = false;
 
 /** Helper function to cause a stack overflow crash */
-FORCENOINLINE static void StackOverflowFunction(int32* DummyArg)
+FORCENOINLINE void StackOverflowFunction(int32* DummyArg)
 {
 	int32 StackArray[8196];
 	FMemory::Memset(StackArray, 0, sizeof(StackArray));
@@ -4756,12 +4756,35 @@ bool UEngine::HandleDebugCommand( const TCHAR* Cmd, FOutputDevice& Ar )
 	}
 	else if( FParse::Command( &Cmd, TEXT("RECURSE") ) )
 	{
-		Ar.Logf( TEXT("Recursing") );
+		Ar.Logf( TEXT("Recursing to create a very deep callstack.") );
 		GLog->Flush();
 		InfiniteRecursionFunction(1);
 		Ar.Logf(TEXT("You will never see this log line."));
 		return true;
 	}
+	else if (FParse::Command(&Cmd, TEXT("THREADRECURSE")))
+	{
+		Ar.Log(TEXT("Recursing to create a very deep callstack (in a separate thread)."));
+		struct FThread
+		{
+			static void InfiniteRecursion(ENamedThreads::Type, const FGraphEventRef&)
+			{
+				InfiniteRecursionFunction(1);
+			}
+		};
+
+		DECLARE_CYCLE_STAT(TEXT("FThread::InfiniteRecursion"),
+		STAT_FThread__InfiniteRecursion,
+			STATGROUP_TaskGraphTasks);
+
+		FTaskGraphInterface::Get().WaitUntilTaskCompletes(
+			FDelegateGraphTask::CreateAndDispatchWhenReady(
+			FDelegateGraphTask::FDelegate::CreateStatic(FThread::InfiniteRecursion),
+			GET_STATID(STAT_FThread__InfiniteRecursion)),
+			ENamedThreads::GameThread
+			);
+		return true;
+	}
 	else if( FParse::Command( &Cmd, TEXT("EATMEM") ) )
 	{
 		Ar.Log( TEXT("Eating up all available memory") );
@@ -4780,10 +4803,33 @@ bool UEngine::HandleDebugCommand( const TCHAR* Cmd, FOutputDevice& Ar )
 	}
 	else if (FParse::Command(&Cmd, TEXT("STACKOVERFLOW")))
 	{
-		Ar.Log(TEXT("Inifnite recursion to cause stack overflow"));
+		Ar.Log(TEXT("Infinite recursion to cause stack overflow"));
 		StackOverflowFunction(nullptr);
 		return true;
 	}
+	else if (FParse::Command(&Cmd, TEXT("THREADSTACKOVERFLOW")))
+	{
+		Ar.Log(TEXT("Infinite recursion to cause stack overflow will happen in a separate thread."));
+		struct FThread
+		{
+			static void StackOverflow(ENamedThreads::Type, const FGraphEventRef&)
+			{
+				StackOverflowFunction(nullptr);
+			}
+		};
+
+		DECLARE_CYCLE_STAT(TEXT("FThread::StackOverflow"),
+		STAT_FThread__StackOverflow,
+			STATGROUP_TaskGraphTasks);
+
+		FTaskGraphInterface::Get().WaitUntilTaskCompletes(
+			FDelegateGraphTask::CreateAndDispatchWhenReady(
+			FDelegateGraphTask::FDelegate::CreateStatic(FThread::StackOverflow),
+			GET_STATID(STAT_FThread__StackOverflow)),
+			ENamedThreads::GameThread
+			);
+		return true;
+	}
 
 	return false;
 }
diff --git a/Engine/Source/Runtime/Engine/Public/Animation/AnimInstanceProxy.h b/Engine/Source/Runtime/Engine/Public/Animation/AnimInstanceProxy.h
--- a/Engine/Source/Runtime/Engine/Public/Animation/AnimInstanceProxy.h
+++ b/Engine/Source/Runtime/Engine/Public/Animation/AnimInstanceProxy.h
@@ -85,14 +85,6 @@ public:
 	void TickSyncGroupWriteIndex()
 	{ 
 		SyncGroupWriteIndex = GetSyncGroupReadIndex();
-
-		// Reset 'write' SyncGroups for next frame.
-		TArray<FAnimGroupInstance>& SyncGroups = SyncGroupArrays[SyncGroupWriteIndex];
-		for (int32 GroupIndex = 0; GroupIndex < SyncGroups.Num(); ++GroupIndex)
-		{
-			FAnimGroupInstance& SyncGroup = SyncGroups[GroupIndex];
-			SyncGroup = FAnimGroupInstance();
-		}
 	}
 
 	/** Get the sync group we are currently reading from */
diff --git a/Engine/Source/Runtime/Engine/Public/Net/RepLayout.h b/Engine/Source/Runtime/Engine/Public/Net/RepLayout.h
--- a/Engine/Source/Runtime/Engine/Public/Net/RepLayout.h
+++ b/Engine/Source/Runtime/Engine/Public/Net/RepLayout.h
@@ -292,7 +292,7 @@ public:
 
 	ENGINE_API void InitFromObjectClass( UClass * InObjectClass );
 
-	bool ReceiveProperties( UClass * InObjectClass, FRepState * RESTRICT RepState, void* RESTRICT Data, FNetBitReader & InBunch, bool & bOutHasUnmapped ) const;
+	bool ReceiveProperties( UClass * InObjectClass, FRepState * RESTRICT RepState, void* RESTRICT Data, FNetBitReader & InBunch, bool & bOutHasUnmapped, const bool bEnableRepNotifies ) const;
 	void UpdateUnmappedObjects( FRepState *	RepState, UPackageMap * PackageMap, UObject* Object, bool & bOutSomeObjectsWereMapped, bool & bOutHasMoreUnmapped ) const;
 
 	void CallRepNotifies( FRepState * RepState, UObject* Object ) const;
@@ -306,7 +306,7 @@ public:
 
 	void MergeDirtyList( FRepState * RepState, const void* RESTRICT Data, const TArray< uint16 > & Dirty1, const TArray< uint16 > & Dirty2, TArray< uint16 > & MergedDirty ) const;
 
-	bool DiffProperties( FRepState * RepState, const void* RESTRICT Data, const bool bSync ) const;
+	bool DiffProperties( TArray<UProperty*>* RepNotifies, void* RESTRICT Destination, const void* RESTRICT Source, const bool bSync ) const;
 
 	void GetLifetimeCustomDeltaProperties(TArray< int32 > & OutCustom, TArray< ELifetimeCondition >	& OutConditions);
 
diff --git a/Engine/Source/Runtime/FriendsAndChat/Private/FriendsAndChatStyle.cpp b/Engine/Source/Runtime/FriendsAndChat/Private/FriendsAndChatStyle.cpp
--- a/Engine/Source/Runtime/FriendsAndChat/Private/FriendsAndChatStyle.cpp
+++ b/Engine/Source/Runtime/FriendsAndChat/Private/FriendsAndChatStyle.cpp
@@ -22,6 +22,12 @@ FFriendsAndChatStyle& FFriendsAndChatStyle::SetLargeFriendsFontStyle(const FFrie
 	return *this;
 }
 
+FFriendsAndChatStyle& FFriendsAndChatStyle::SetChatFontStyle(const FFriendsFontStyle& FontStyle)
+{
+	ChatFontStyle = FontStyle;
+	return *this;
+}
+
 FFriendsAndChatStyle& FFriendsAndChatStyle::SetFriendsListStyle(const FFriendsListStyle& InFriendsListStyle)
 {
 	FriendsListStyle = InFriendsListStyle;
@@ -110,85 +116,211 @@ TSharedRef< FSlateStyleSet > FFriendsAndChatModuleStyle::Create(FFriendsAndChatS
 {
 	TSharedRef< FSlateStyleSet > Style = MakeShareable(new FSlateStyleSet("FriendsAndChatStyle"));
 
-	const FTextBlockStyle DefaultText = FTextBlockStyle(FriendStyle.FriendsChatStyle.TextStyle)
-		.SetFont(FriendStyle.FriendsNormalFontStyle.FriendsFontSmall);
-
-	// Name Style
-	const FTextBlockStyle GlobalChatFont = FTextBlockStyle(DefaultText)
-		.SetFont(FriendStyle.FriendsChatStyle.TextStyle.Font)
-		.SetColorAndOpacity(FriendStyle.FriendsChatStyle.GlobalHyperlinkChatColor);
-
-	const FTextBlockStyle GameChatFont = FTextBlockStyle(DefaultText)
-		.SetFont(FriendStyle.FriendsChatStyle.TextStyle.Font)
-		.SetColorAndOpacity(FriendStyle.FriendsChatStyle.GameHyperlinkChatColor);
-
-	const FTextBlockStyle TeamChatFont = FTextBlockStyle(DefaultText)
-		.SetFont(FriendStyle.FriendsChatStyle.TextStyle.Font)
-		.SetColorAndOpacity(FriendStyle.FriendsChatStyle.TeamHyperlinkChatColor);
-
-	const FTextBlockStyle PartyChatFont = FTextBlockStyle(DefaultText)
-		.SetFont(FriendStyle.FriendsChatStyle.TextStyle.Font)
-		.SetColorAndOpacity(FriendStyle.FriendsChatStyle.PartyHyperlinkChatColor);
-
-	const FTextBlockStyle WhisperChatFont = FTextBlockStyle(DefaultText)
-		.SetFont(FriendStyle.FriendsChatStyle.TextStyle.Font)
-		.SetColorAndOpacity(FriendStyle.FriendsChatStyle.WhisperHyperlinkChatColor);
-
 	const FButtonStyle UserNameButton = FButtonStyle()
 		.SetNormal(FSlateNoResource())
 		.SetPressed(FSlateNoResource())
 		.SetHovered(FSlateNoResource());
 
-	const FHyperlinkStyle GlobalChatHyperlink = FHyperlinkStyle()
-		.SetUnderlineStyle(UserNameButton)
-		.SetTextStyle(GlobalChatFont)
-		.SetPadding(FMargin(0.0f));
-
-	const FHyperlinkStyle GameChatHyperlink = FHyperlinkStyle()
-		.SetUnderlineStyle(UserNameButton)
-		.SetTextStyle(GameChatFont)
-		.SetPadding(FMargin(0.0f));
-
-	const FHyperlinkStyle TeamChatHyperlink = FHyperlinkStyle()
-		.SetUnderlineStyle(UserNameButton)
-		.SetTextStyle(TeamChatFont)
-		.SetPadding(FMargin(0.0f));
-
-	const FHyperlinkStyle PartyChatHyperlink = FHyperlinkStyle()
-		.SetUnderlineStyle(UserNameButton)
-		.SetTextStyle(PartyChatFont)
-		.SetPadding(FMargin(0.0f));
-
-	const FHyperlinkStyle WhisperChatHyperlink = FHyperlinkStyle()
-		.SetUnderlineStyle(UserNameButton)
-		.SetTextStyle(WhisperChatFont)
-		.SetPadding(FMargin(0.0f));
-
-	const FHyperlinkStyle DefaultChatHyperlink = FHyperlinkStyle()
-		.SetUnderlineStyle(UserNameButton)
-		.SetTextStyle(DefaultText)
-		.SetPadding(FMargin(0.0f));
-
-	Style->Set("UserNameTextStyle.Default", DefaultText);
-
-	Style->Set("UserNameTextStyle.GlobalHyperlink", GlobalChatHyperlink);
-	Style->Set("UserNameTextStyle.GameHyperlink", GameChatHyperlink);
-	Style->Set("UserNameTextStyle.GameHyperlink", TeamChatHyperlink);
-	Style->Set("UserNameTextStyle.PartyHyperlink", PartyChatHyperlink);
-	Style->Set("UserNameTextStyle.Whisperlink", WhisperChatHyperlink);
-	Style->Set("UserNameTextStyle.DefaultHyperlink", DefaultChatHyperlink);
-	Style->Set("UserNameTextStyle.GlobalTextStyle", GlobalChatFont);
-	Style->Set("UserNameTextStyle.GameTextStyle", GameChatFont);
-	Style->Set("UserNameTextStyle.TeamTextStyle", TeamChatFont);
-	Style->Set("UserNameTextStyle.PartyTextStyle", PartyChatFont);
-	Style->Set("UserNameTextStyle.WhisperTextStyle", WhisperChatFont);
-
-	Style->Set("MessageBreak", FTextBlockStyle(DefaultText)
-		.SetFont(FSlateFontInfo(
-		FriendStyle.FriendsNormalFontStyle.FriendsFontSmall.FontObject,
-		6,
-		FriendStyle.FriendsNormalFontStyle.FriendsFontSmall.TypefaceFontName
-		)));
+	// Small
+	{
+		const FTextBlockStyle DefaultTextSmall = FTextBlockStyle(FriendStyle.FriendsChatStyle.TextStyle)
+			.SetFont(FriendStyle.ChatFontStyle.FriendsFontSmallBold);
+
+		const FTextBlockStyle GlobalChatFontSmall = FTextBlockStyle(DefaultTextSmall)
+			.SetColorAndOpacity(FriendStyle.FriendsChatStyle.GlobalHyperlinkChatColor);
+
+		const FTextBlockStyle GameChatFontSmall = FTextBlockStyle(DefaultTextSmall)
+			.SetColorAndOpacity(FriendStyle.FriendsChatStyle.GameHyperlinkChatColor);
+
+		const FTextBlockStyle TeamChatFontSmall = FTextBlockStyle(DefaultTextSmall)
+			.SetColorAndOpacity(FriendStyle.FriendsChatStyle.TeamHyperlinkChatColor);
+
+		const FTextBlockStyle PartyChatFontSmall = FTextBlockStyle(DefaultTextSmall)
+			.SetColorAndOpacity(FriendStyle.FriendsChatStyle.PartyHyperlinkChatColor);
+
+		const FTextBlockStyle WhisperChatFontSmall = FTextBlockStyle(DefaultTextSmall)
+			.SetColorAndOpacity(FriendStyle.FriendsChatStyle.WhisperHyperlinkChatColor);
+
+		const FHyperlinkStyle GlobalChatHyperlinkSmall = FHyperlinkStyle()
+			.SetUnderlineStyle(UserNameButton)
+			.SetTextStyle(GlobalChatFontSmall)
+			.SetPadding(FMargin(0.0f));
+
+		const FHyperlinkStyle GameChatHyperlinkSmall = FHyperlinkStyle()
+			.SetUnderlineStyle(UserNameButton)
+			.SetTextStyle(GameChatFontSmall)
+			.SetPadding(FMargin(0.0f));
+
+		const FHyperlinkStyle TeamChatHyperlinkSmall = FHyperlinkStyle()
+			.SetUnderlineStyle(UserNameButton)
+			.SetTextStyle(TeamChatFontSmall)
+			.SetPadding(FMargin(0.0f));
+
+		const FHyperlinkStyle PartyChatHyperlinkSmall = FHyperlinkStyle()
+			.SetUnderlineStyle(UserNameButton)
+			.SetTextStyle(PartyChatFontSmall)
+			.SetPadding(FMargin(0.0f));
+
+		const FHyperlinkStyle WhisperChatHyperlinkSmall = FHyperlinkStyle()
+			.SetUnderlineStyle(UserNameButton)
+			.SetTextStyle(WhisperChatFontSmall)
+			.SetPadding(FMargin(0.0f));
+
+		const FHyperlinkStyle DefaultChatHyperlinkSmall = FHyperlinkStyle()
+			.SetUnderlineStyle(UserNameButton)
+			.SetTextStyle(DefaultTextSmall)
+			.SetPadding(FMargin(0.0f));
+
+		Style->Set("UserNameTextStyle.DefaultSmall", DefaultTextSmall);
+		Style->Set("UserNameTextStyle.GlobalHyperlinkSmall", GlobalChatHyperlinkSmall);
+		Style->Set("UserNameTextStyle.GameHyperlinkSmall", GameChatHyperlinkSmall);
+		Style->Set("UserNameTextStyle.TeamHyperlinkSmall", TeamChatHyperlinkSmall);
+		Style->Set("UserNameTextStyle.PartyHyperlinkSmall", PartyChatHyperlinkSmall);
+		Style->Set("UserNameTextStyle.WhisperlinkSmall", WhisperChatHyperlinkSmall);
+		Style->Set("UserNameTextStyle.DefaultHyperlinkSmall", DefaultChatHyperlinkSmall);
+		Style->Set("UserNameTextStyle.GlobalTextStyleSmall", GlobalChatFontSmall);
+		Style->Set("UserNameTextStyle.GameTextStyleSmall", GameChatFontSmall);
+		Style->Set("UserNameTextStyle.TeamTextStyleSmall", TeamChatFontSmall);
+		Style->Set("UserNameTextStyle.PartyTextStyleSmall", PartyChatFontSmall);
+		Style->Set("UserNameTextStyle.WhisperTextStyleSmall", WhisperChatFontSmall);
+
+		Style->Set("MessageBreak", FTextBlockStyle(DefaultTextSmall)
+			.SetFont(FSlateFontInfo(
+			FriendStyle.FriendsNormalFontStyle.FriendsFontSmall.FontObject,
+			6,
+			FriendStyle.FriendsNormalFontStyle.FriendsFontSmall.TypefaceFontName
+			)));
+	}
+
+
+	// Normal
+	{
+		const FTextBlockStyle DefaultText = FTextBlockStyle(FriendStyle.FriendsChatStyle.TextStyle)
+			.SetFont(FriendStyle.ChatFontStyle.FriendsFontNormalBold);
+
+		const FTextBlockStyle GlobalChatFont = FTextBlockStyle(DefaultText)
+			.SetColorAndOpacity(FriendStyle.FriendsChatStyle.GlobalHyperlinkChatColor);
+
+		const FTextBlockStyle GameChatFont = FTextBlockStyle(DefaultText)
+			.SetColorAndOpacity(FriendStyle.FriendsChatStyle.GameHyperlinkChatColor);
+
+		const FTextBlockStyle TeamChatFont = FTextBlockStyle(DefaultText)
+			.SetColorAndOpacity(FriendStyle.FriendsChatStyle.TeamHyperlinkChatColor);
+
+		const FTextBlockStyle PartyChatFont = FTextBlockStyle(DefaultText)
+			.SetColorAndOpacity(FriendStyle.FriendsChatStyle.PartyHyperlinkChatColor);
+
+		const FTextBlockStyle WhisperChatFont = FTextBlockStyle(DefaultText)
+			.SetFont(FriendStyle.FriendsNormalFontStyle.FriendsChatFont)
+			.SetColorAndOpacity(FriendStyle.FriendsChatStyle.WhisperHyperlinkChatColor);
+
+		const FHyperlinkStyle GlobalChatHyperlink = FHyperlinkStyle()
+			.SetUnderlineStyle(UserNameButton)
+			.SetTextStyle(GlobalChatFont)
+			.SetPadding(FMargin(0.0f));
+
+		const FHyperlinkStyle GameChatHyperlink = FHyperlinkStyle()
+			.SetUnderlineStyle(UserNameButton)
+			.SetTextStyle(GameChatFont)
+			.SetPadding(FMargin(0.0f));
+
+		const FHyperlinkStyle TeamChatHyperlink = FHyperlinkStyle()
+			.SetUnderlineStyle(UserNameButton)
+			.SetTextStyle(TeamChatFont)
+			.SetPadding(FMargin(0.0f));
+
+		const FHyperlinkStyle PartyChatHyperlink = FHyperlinkStyle()
+			.SetUnderlineStyle(UserNameButton)
+			.SetTextStyle(PartyChatFont)
+			.SetPadding(FMargin(0.0f));
+
+		const FHyperlinkStyle WhisperChatHyperlink = FHyperlinkStyle()
+			.SetUnderlineStyle(UserNameButton)
+			.SetTextStyle(WhisperChatFont)
+			.SetPadding(FMargin(0.0f));
+
+		const FHyperlinkStyle DefaultChatHyperlink = FHyperlinkStyle()
+			.SetUnderlineStyle(UserNameButton)
+			.SetTextStyle(DefaultText)
+			.SetPadding(FMargin(0.0f));
+
+		Style->Set("UserNameTextStyle.Default", DefaultText);
+		Style->Set("UserNameTextStyle.GlobalHyperlink", GlobalChatHyperlink);
+		Style->Set("UserNameTextStyle.GameHyperlink", GameChatHyperlink);
+		Style->Set("UserNameTextStyle.TeamHyperlink", TeamChatHyperlink);
+		Style->Set("UserNameTextStyle.PartyHyperlink", PartyChatHyperlink);
+		Style->Set("UserNameTextStyle.Whisperlink", WhisperChatHyperlink);
+		Style->Set("UserNameTextStyle.DefaultHyperlink", DefaultChatHyperlink);
+		Style->Set("UserNameTextStyle.GlobalTextStyle", GlobalChatFont);
+		Style->Set("UserNameTextStyle.GameTextStyle", GameChatFont);
+		Style->Set("UserNameTextStyle.TeamTextStyle", TeamChatFont);
+		Style->Set("UserNameTextStyle.PartyTextStyle", PartyChatFont);
+		Style->Set("UserNameTextStyle.WhisperTextStyle", WhisperChatFont);
+	}
+
+	// Large
+	{
+		const FTextBlockStyle DefaultTextLarge = FTextBlockStyle(FriendStyle.FriendsChatStyle.TextStyle)
+			.SetFont(FriendStyle.ChatFontStyle.FriendsFontLargeBold);
+
+		const FTextBlockStyle GlobalChatFontLarge = FTextBlockStyle(DefaultTextLarge)
+			.SetColorAndOpacity(FriendStyle.FriendsChatStyle.GlobalHyperlinkChatColor);
+
+		const FTextBlockStyle GameChatFontLarge = FTextBlockStyle(DefaultTextLarge)
+			.SetColorAndOpacity(FriendStyle.FriendsChatStyle.GameHyperlinkChatColor);
+
+		const FTextBlockStyle TeamChatFontLarge = FTextBlockStyle(DefaultTextLarge)
+			.SetColorAndOpacity(FriendStyle.FriendsChatStyle.TeamHyperlinkChatColor);
+
+		const FTextBlockStyle PartyChatFontLarge = FTextBlockStyle(DefaultTextLarge)
+			.SetColorAndOpacity(FriendStyle.FriendsChatStyle.PartyHyperlinkChatColor);
+
+		const FTextBlockStyle WhisperChatFontLarge = FTextBlockStyle(DefaultTextLarge)
+			.SetColorAndOpacity(FriendStyle.FriendsChatStyle.WhisperHyperlinkChatColor);
+
+		const FHyperlinkStyle GlobalChatHyperlinkLarge = FHyperlinkStyle()
+			.SetUnderlineStyle(UserNameButton)
+			.SetTextStyle(GlobalChatFontLarge)
+			.SetPadding(FMargin(0.0f));
+
+		const FHyperlinkStyle GameChatHyperlinkLarge = FHyperlinkStyle()
+			.SetUnderlineStyle(UserNameButton)
+			.SetTextStyle(GameChatFontLarge)
+			.SetPadding(FMargin(0.0f));
+
+		const FHyperlinkStyle TeamChatHyperlinkLarge = FHyperlinkStyle()
+			.SetUnderlineStyle(UserNameButton)
+			.SetTextStyle(TeamChatFontLarge)
+			.SetPadding(FMargin(0.0f));
+
+		const FHyperlinkStyle PartyChatHyperlinkLarge = FHyperlinkStyle()
+			.SetUnderlineStyle(UserNameButton)
+			.SetTextStyle(PartyChatFontLarge)
+			.SetPadding(FMargin(0.0f));
+
+		const FHyperlinkStyle WhisperChatHyperlinkLarge = FHyperlinkStyle()
+			.SetUnderlineStyle(UserNameButton)
+			.SetTextStyle(WhisperChatFontLarge)
+			.SetPadding(FMargin(0.0f));
+
+		const FHyperlinkStyle DefaultChatHyperlinkLarge = FHyperlinkStyle()
+			.SetUnderlineStyle(UserNameButton)
+			.SetTextStyle(DefaultTextLarge)
+			.SetPadding(FMargin(0.0f));
+
+		Style->Set("UserNameTextStyle.DefaultLarge", DefaultTextLarge);
+		Style->Set("UserNameTextStyle.GlobalHyperlinkLarge", GlobalChatHyperlinkLarge);
+		Style->Set("UserNameTextStyle.GameHyperlinkLarge", GameChatHyperlinkLarge);
+		Style->Set("UserNameTextStyle.TeamHyperlinkLarge", TeamChatHyperlinkLarge);
+		Style->Set("UserNameTextStyle.PartyHyperlinkLarge", PartyChatHyperlinkLarge);
+		Style->Set("UserNameTextStyle.WhisperlinkLarge", WhisperChatHyperlinkLarge);
+		Style->Set("UserNameTextStyle.DefaultHyperlinkLarge", DefaultChatHyperlinkLarge);
+		Style->Set("UserNameTextStyle.GlobalTextStyleLarge", GlobalChatFontLarge);
+		Style->Set("UserNameTextStyle.GameTextStyleLarge", GameChatFontLarge);
+		Style->Set("UserNameTextStyle.TeamTextStyleLarge", TeamChatFontLarge);
+		Style->Set("UserNameTextStyle.PartyTextStyleLarge", PartyChatFontLarge);
+		Style->Set("UserNameTextStyle.WhisperTextStyleLarge", WhisperChatFontLarge);
+	}
 
 	return Style;
 }
diff --git a/Engine/Source/Runtime/FriendsAndChat/Private/FriendsFontStyle.cpp b/Engine/Source/Runtime/FriendsAndChat/Private/FriendsFontStyle.cpp
--- a/Engine/Source/Runtime/FriendsAndChat/Private/FriendsFontStyle.cpp
+++ b/Engine/Source/Runtime/FriendsAndChat/Private/FriendsFontStyle.cpp
@@ -46,6 +46,12 @@ FFriendsFontStyle& FFriendsFontStyle::SetFontLargeBold(const FSlateFontInfo& Fon
 	return *this;
 }
 
+FFriendsFontStyle& FFriendsFontStyle::SetChatFont(const FSlateFontInfo& FontStyle)
+{
+	FriendsChatFont = FontStyle;
+	return *this;
+}
+
 FFriendsFontStyle& FFriendsFontStyle::SetDefaultFontColor(const FLinearColor& FontColor)
 {
 	DefaultFontColor = FontColor;
diff --git a/Engine/Source/Runtime/FriendsAndChat/Public/FriendsAndChatStyle.h b/Engine/Source/Runtime/FriendsAndChat/Public/FriendsAndChatStyle.h
--- a/Engine/Source/Runtime/FriendsAndChat/Public/FriendsAndChatStyle.h
+++ b/Engine/Source/Runtime/FriendsAndChat/Public/FriendsAndChatStyle.h
@@ -69,6 +69,10 @@ struct FRIENDSANDCHAT_API FFriendsAndChatStyle
 	FFriendsFontStyle FriendsLargeFontStyle;
 	FFriendsAndChatStyle& SetLargeFriendsFontStyle(const FFriendsFontStyle& FontStyle);
 
+	UPROPERTY(EditAnywhere, Category = Appearance)
+	FFriendsFontStyle ChatFontStyle;
+	FFriendsAndChatStyle& SetChatFontStyle(const FFriendsFontStyle& FontStyle);
+
 	UPROPERTY(EditAnywhere, Category = Appearance)
 	FCheckBoxStyle CheckBoxStyle;
 	FFriendsAndChatStyle& SetCheckBoxStyle(const FCheckBoxStyle& InCheckBoxStyle);
diff --git a/Engine/Source/Runtime/FriendsAndChat/Public/FriendsFontStyle.h b/Engine/Source/Runtime/FriendsAndChat/Public/FriendsFontStyle.h
--- a/Engine/Source/Runtime/FriendsAndChat/Public/FriendsFontStyle.h
+++ b/Engine/Source/Runtime/FriendsAndChat/Public/FriendsFontStyle.h
@@ -69,6 +69,11 @@ struct FRIENDSANDCHAT_API FFriendsFontStyle
 	FSlateFontInfo FriendsFontLargeBold;
 	FFriendsFontStyle& SetFontLargeBold(const FSlateFontInfo& InFontStyle);
 
+	/** Font Style */
+	UPROPERTY(EditAnywhere, Category=Appearance)
+	FSlateFontInfo FriendsChatFont;
+	FFriendsFontStyle& SetChatFont(const FSlateFontInfo& InFontStyle);
+
 	/** Default Font Color */
 	UPROPERTY(EditAnywhere, Category=Appearance)
 	FLinearColor DefaultFontColor;
diff --git a/Engine/Source/Runtime/GameplayAbilities/Private/Abilities/GameplayAbility.cpp b/Engine/Source/Runtime/GameplayAbilities/Private/Abilities/GameplayAbility.cpp
--- a/Engine/Source/Runtime/GameplayAbilities/Private/Abilities/GameplayAbility.cpp
+++ b/Engine/Source/Runtime/GameplayAbilities/Private/Abilities/GameplayAbility.cpp
@@ -777,6 +777,14 @@ float UGameplayAbility::GetCooldownTimeRemaining(const FGameplayAbilityActorInfo
 	return 0.f;
 }
 
+void UGameplayAbility::InvalidateClientPredictionKey() const
+{
+	if (CurrentActorInfo && CurrentActorInfo->AbilitySystemComponent.IsValid())
+	{
+		CurrentActorInfo->AbilitySystemComponent->ScopedPredictionKey = FPredictionKey();
+	}
+}
+
 void UGameplayAbility::GetCooldownTimeRemainingAndDuration(FGameplayAbilitySpecHandle Handle, const FGameplayAbilityActorInfo* ActorInfo, float& TimeRemaining, float& CooldownDuration) const
 {
 	SCOPE_CYCLE_COUNTER(STAT_GameplayAbilityGetCooldownTimeRemainingAndDuration);
diff --git a/Engine/Source/Runtime/GameplayAbilities/Private/Abilities/Tasks/AbilityTask_WaitAttributeChangeRatioThreshold.cpp b/Engine/Source/Runtime/GameplayAbilities/Private/Abilities/Tasks/AbilityTask_WaitAttributeChangeRatioThreshold.cpp
new file mode 100644
--- /dev/null
+++ b/Engine/Source/Runtime/GameplayAbilities/Private/Abilities/Tasks/AbilityTask_WaitAttributeChangeRatioThreshold.cpp
@@ -0,0 +1,131 @@
+// Copyright 1998-2016 Epic Games, Inc. All Rights Reserved.
+
+#include "AbilitySystemPrivatePCH.h"
+#include "Abilities/Tasks/AbilityTask_WaitAttributeChangeRatioThreshold.h"
+
+#include "AbilitySystemComponent.h"
+
+#include "Abilities/GameplayAbility.h"
+#include "GameplayEffectExtension.h"
+
+UAbilityTask_WaitAttributeChangeRatioThreshold::UAbilityTask_WaitAttributeChangeRatioThreshold(const FObjectInitializer& ObjectInitializer)
+	: Super(ObjectInitializer)
+{
+	bTriggerOnce = false;
+	bMatchedComparisonLastAttributeChange = false;
+	LastAttributeNumeratorValue = 1.f;
+	LastAttributeDenominatorValue = 1.f;
+}
+
+UAbilityTask_WaitAttributeChangeRatioThreshold* UAbilityTask_WaitAttributeChangeRatioThreshold::WaitForAttributeChangeRatioThreshold(UObject* WorldContextObject, FGameplayAttribute AttributeNumerator, FGameplayAttribute AttributeDenominator, TEnumAsByte<EWaitAttributeChangeComparison::Type> ComparisonType, float ComparisonValue, bool bTriggerOnce)
+{
+	auto MyTask = NewAbilityTask<UAbilityTask_WaitAttributeChangeRatioThreshold>(WorldContextObject);
+	MyTask->AttributeNumerator = AttributeNumerator;
+	MyTask->AttributeDenominator = AttributeDenominator;
+	MyTask->ComparisonType = ComparisonType;
+	MyTask->ComparisonValue = ComparisonValue;
+	MyTask->bTriggerOnce = bTriggerOnce;
+
+	return MyTask;
+}
+
+void UAbilityTask_WaitAttributeChangeRatioThreshold::Activate()
+{
+	if (AbilitySystemComponent)
+	{
+		LastAttributeNumeratorValue = AbilitySystemComponent->GetNumericAttribute(AttributeNumerator);
+		LastAttributeDenominatorValue = AbilitySystemComponent->GetNumericAttribute(AttributeDenominator);
+		bMatchedComparisonLastAttributeChange = DoesValuePassComparison(LastAttributeNumeratorValue, LastAttributeDenominatorValue);
+
+		// Broadcast OnChange immediately with current value
+		OnChange.Broadcast(bMatchedComparisonLastAttributeChange, LastAttributeDenominatorValue != 0.f ? LastAttributeNumeratorValue/LastAttributeDenominatorValue : 0.f);
+
+		OnNumeratorAttributeChangeDelegateHandle = AbilitySystemComponent->RegisterGameplayAttributeEvent(AttributeNumerator).AddUObject(this, &UAbilityTask_WaitAttributeChangeRatioThreshold::OnNumeratorAttributeChange);
+		OnDenominatorAttributeChangeDelegateHandle = AbilitySystemComponent->RegisterGameplayAttributeEvent(AttributeDenominator).AddUObject(this, &UAbilityTask_WaitAttributeChangeRatioThreshold::OnDenominatorAttributeChange);
+	}
+}
+
+void UAbilityTask_WaitAttributeChangeRatioThreshold::OnAttributeChange()
+{
+	UWorld* World = GetWorld();
+	if (World && !CheckAttributeTimer.IsValid())
+	{
+		// Trigger OnRatioChange check at the end of this frame/next so that any individual attribute change
+		// has time for the other attribute to update (in case they're linked)
+		World->GetTimerManager().SetTimer(CheckAttributeTimer, this, &UAbilityTask_WaitAttributeChangeRatioThreshold::OnRatioChange, 0.001f, false);
+	}
+}
+
+void UAbilityTask_WaitAttributeChangeRatioThreshold::OnRatioChange()
+{
+	CheckAttributeTimer.Invalidate();
+
+	bool bPassedComparison = DoesValuePassComparison(LastAttributeNumeratorValue, LastAttributeDenominatorValue);
+	if (bPassedComparison != bMatchedComparisonLastAttributeChange)
+	{
+		bMatchedComparisonLastAttributeChange = bPassedComparison;
+		OnChange.Broadcast(bMatchedComparisonLastAttributeChange, LastAttributeDenominatorValue != 0.f ? LastAttributeNumeratorValue/LastAttributeDenominatorValue : 0.f);
+		if (bTriggerOnce)
+		{
+			EndTask();
+		}
+	}
+}
+
+void UAbilityTask_WaitAttributeChangeRatioThreshold::OnNumeratorAttributeChange(float NewValue, const FGameplayEffectModCallbackData* Data)
+{
+	LastAttributeNumeratorValue = NewValue;
+	OnAttributeChange();
+}
+
+void UAbilityTask_WaitAttributeChangeRatioThreshold::OnDenominatorAttributeChange(float NewValue, const FGameplayEffectModCallbackData* Data)
+{
+	LastAttributeDenominatorValue = NewValue;
+	OnAttributeChange();
+}
+
+bool UAbilityTask_WaitAttributeChangeRatioThreshold::DoesValuePassComparison(float ValueNumerator, float ValueDenominator) const
+{
+	if (ValueDenominator == 0.f)
+	{
+		return bMatchedComparisonLastAttributeChange;
+	}
+
+	const float CurrentRatio = ValueNumerator / ValueDenominator;
+	bool bPassedComparison = true;
+	switch (ComparisonType)
+	{
+	case EWaitAttributeChangeComparison::ExactlyEqualTo:
+		bPassedComparison = (CurrentRatio == ComparisonValue);
+		break;		
+	case EWaitAttributeChangeComparison::GreaterThan:
+		bPassedComparison = (CurrentRatio > ComparisonValue);
+		break;
+	case EWaitAttributeChangeComparison::GreaterThanOrEqualTo:
+		bPassedComparison = (CurrentRatio >= ComparisonValue);
+		break;
+	case EWaitAttributeChangeComparison::LessThan:
+		bPassedComparison = (CurrentRatio < ComparisonValue);
+		break;
+	case EWaitAttributeChangeComparison::LessThanOrEqualTo:
+		bPassedComparison = (CurrentRatio <= ComparisonValue);
+		break;
+	case EWaitAttributeChangeComparison::NotEqualTo:
+		bPassedComparison = (CurrentRatio != ComparisonValue);
+		break;
+	default:
+		break;
+	}
+	return bPassedComparison;
+}
+
+void UAbilityTask_WaitAttributeChangeRatioThreshold::OnDestroy(bool AbilityEnded)
+{
+	if (AbilitySystemComponent)
+	{
+		AbilitySystemComponent->RegisterGameplayAttributeEvent(AttributeNumerator).Remove(OnNumeratorAttributeChangeDelegateHandle);
+		AbilitySystemComponent->RegisterGameplayAttributeEvent(AttributeDenominator).Remove(OnDenominatorAttributeChangeDelegateHandle);
+	}
+
+	Super::OnDestroy(AbilityEnded);
+}
\ No newline at end of file
diff --git a/Engine/Source/Runtime/GameplayAbilities/Private/Abilities/Tasks/AbilityTask_WaitAttributeChangeThreshold.cpp b/Engine/Source/Runtime/GameplayAbilities/Private/Abilities/Tasks/AbilityTask_WaitAttributeChangeThreshold.cpp
new file mode 100644
--- /dev/null
+++ b/Engine/Source/Runtime/GameplayAbilities/Private/Abilities/Tasks/AbilityTask_WaitAttributeChangeThreshold.cpp
@@ -0,0 +1,94 @@
+// Copyright 1998-2016 Epic Games, Inc. All Rights Reserved.
+
+#include "AbilitySystemPrivatePCH.h"
+#include "Abilities/Tasks/AbilityTask_WaitAttributeChangeThreshold.h"
+
+#include "AbilitySystemComponent.h"
+
+#include "Abilities/GameplayAbility.h"
+#include "GameplayEffectExtension.h"
+
+UAbilityTask_WaitAttributeChangeThreshold::UAbilityTask_WaitAttributeChangeThreshold(const FObjectInitializer& ObjectInitializer)
+	: Super(ObjectInitializer)
+{
+	bTriggerOnce = false;
+	bMatchedComparisonLastAttributeChange = false;
+}
+
+UAbilityTask_WaitAttributeChangeThreshold* UAbilityTask_WaitAttributeChangeThreshold::WaitForAttributeChangeThreshold(UObject* WorldContextObject, FGameplayAttribute Attribute, TEnumAsByte<EWaitAttributeChangeComparison::Type> ComparisonType, float ComparisonValue, bool bTriggerOnce)
+{
+	auto MyTask = NewAbilityTask<UAbilityTask_WaitAttributeChangeThreshold>(WorldContextObject);
+	MyTask->Attribute = Attribute;
+	MyTask->ComparisonType = ComparisonType;
+	MyTask->ComparisonValue = ComparisonValue;
+	MyTask->bTriggerOnce = bTriggerOnce;
+
+	return MyTask;
+}
+
+void UAbilityTask_WaitAttributeChangeThreshold::Activate()
+{
+	if (AbilitySystemComponent)
+	{
+		const float CurrentValue = AbilitySystemComponent->GetNumericAttribute(Attribute);
+		bMatchedComparisonLastAttributeChange = DoesValuePassComparison(CurrentValue);
+
+		// Broadcast OnChange immediately with current value
+		OnChange.Broadcast(bMatchedComparisonLastAttributeChange, CurrentValue);
+
+		OnAttributeChangeDelegateHandle = AbilitySystemComponent->RegisterGameplayAttributeEvent(Attribute).AddUObject(this, &UAbilityTask_WaitAttributeChangeThreshold::OnAttributeChange);
+	}
+}
+
+void UAbilityTask_WaitAttributeChangeThreshold::OnAttributeChange(float NewValue, const FGameplayEffectModCallbackData* Data)
+{
+	bool bPassedComparison = DoesValuePassComparison(NewValue);
+	if (bPassedComparison != bMatchedComparisonLastAttributeChange)
+	{
+		bMatchedComparisonLastAttributeChange = bPassedComparison;
+		OnChange.Broadcast(bPassedComparison, NewValue);
+		if (bTriggerOnce)
+		{
+			EndTask();
+		}
+	}
+}
+
+bool UAbilityTask_WaitAttributeChangeThreshold::DoesValuePassComparison(float Value) const
+{
+	bool bPassedComparison = true;
+	switch (ComparisonType)
+	{
+	case EWaitAttributeChangeComparison::ExactlyEqualTo:
+		bPassedComparison = (Value == ComparisonValue);
+		break;		
+	case EWaitAttributeChangeComparison::GreaterThan:
+		bPassedComparison = (Value > ComparisonValue);
+		break;
+	case EWaitAttributeChangeComparison::GreaterThanOrEqualTo:
+		bPassedComparison = (Value >= ComparisonValue);
+		break;
+	case EWaitAttributeChangeComparison::LessThan:
+		bPassedComparison = (Value < ComparisonValue);
+		break;
+	case EWaitAttributeChangeComparison::LessThanOrEqualTo:
+		bPassedComparison = (Value <= ComparisonValue);
+		break;
+	case EWaitAttributeChangeComparison::NotEqualTo:
+		bPassedComparison = (Value != ComparisonValue);
+		break;
+	default:
+		break;
+	}
+	return bPassedComparison;
+}
+
+void UAbilityTask_WaitAttributeChangeThreshold::OnDestroy(bool AbilityEnded)
+{
+	if (AbilitySystemComponent)
+	{
+		AbilitySystemComponent->RegisterGameplayAttributeEvent(Attribute).Remove(OnAttributeChangeDelegateHandle);
+	}
+
+	Super::OnDestroy(AbilityEnded);
+}
\ No newline at end of file
diff --git a/Engine/Source/Runtime/GameplayAbilities/Private/AbilitySystemBlueprintLibrary.cpp b/Engine/Source/Runtime/GameplayAbilities/Private/AbilitySystemBlueprintLibrary.cpp
--- a/Engine/Source/Runtime/GameplayAbilities/Private/AbilitySystemBlueprintLibrary.cpp
+++ b/Engine/Source/Runtime/GameplayAbilities/Private/AbilitySystemBlueprintLibrary.cpp
@@ -550,7 +550,7 @@ bool UAbilitySystemBlueprintLibrary::GetGameplayCueDirection(AActor* TargetActor
 {
 	if (Parameters.Normal.IsNearlyZero() == false)
 	{
-		Direction = Parameters.Normal;
+		Direction = -Parameters.Normal;
 		return true;
 	}
 
diff --git a/Engine/Source/Runtime/GameplayAbilities/Private/AbilitySystemComponent_Abilities.cpp b/Engine/Source/Runtime/GameplayAbilities/Private/AbilitySystemComponent_Abilities.cpp
--- a/Engine/Source/Runtime/GameplayAbilities/Private/AbilitySystemComponent_Abilities.cpp
+++ b/Engine/Source/Runtime/GameplayAbilities/Private/AbilitySystemComponent_Abilities.cpp
@@ -1000,7 +1000,9 @@ bool UAbilitySystemComponent::InternalTryActivateAbility(FGameplayAbilitySpecHan
 {
 	const FGameplayTag& NetworkFailTag = UAbilitySystemGlobals::Get().ActivateFailNetworkingTag;
 
-	FGameplayTagContainer FailureTags;
+	static FGameplayTagContainer FailureTags;
+	FailureTags.Reset();
+
 	FGameplayAbilitySpec* Spec = FindAbilitySpecFromHandle(Handle);
 	if (!Spec)
 	{
diff --git a/Engine/Source/Runtime/GameplayAbilities/Private/AbilitySystemGlobals.cpp b/Engine/Source/Runtime/GameplayAbilities/Private/AbilitySystemGlobals.cpp
--- a/Engine/Source/Runtime/GameplayAbilities/Private/AbilitySystemGlobals.cpp
+++ b/Engine/Source/Runtime/GameplayAbilities/Private/AbilitySystemGlobals.cpp
@@ -14,7 +14,6 @@ UAbilitySystemGlobals::UAbilitySystemGlobals(const FObjectInitializer& ObjectIni
 : Super(ObjectInitializer)
 {
 	AbilitySystemGlobalsClassName = FStringClassReference(TEXT("/Script/GameplayAbilities.AbilitySystemGlobals"));
-	GlobalGameplayCueManagerName = FStringClassReference(TEXT("/Script/GameplayAbilities.GameplayCueManager"));
 
 	PredictTargetGameplayEffects = true;
 
@@ -293,6 +292,7 @@ UGameplayCueManager* UAbilitySystemGlobals::GetGameplayCueManager()
 {
 	if (GlobalGameplayCueManager == nullptr)
 	{
+		// Load specific gameplaycue manager object if specified
 		if (GlobalGameplayCueManagerName.IsValid())
 		{
 			GlobalGameplayCueManager = LoadObject<UGameplayCueManager>(nullptr, *GlobalGameplayCueManagerName.ToString(), nullptr, LOAD_None, nullptr);
@@ -302,6 +302,16 @@ UGameplayCueManager* UAbilitySystemGlobals::GetGameplayCueManager()
 			}
 		}
 
+		// Load specific gameplaycue manager class if specified
+		if ( GlobalGameplayCueManager == nullptr && GlobalGameplayCueManagerClass.IsValid() )
+		{
+			UClass* GCMClass = LoadClass<UObject>(NULL, *GlobalGameplayCueManagerClass.ToString(), NULL, LOAD_None, NULL);
+			if (GCMClass)
+			{
+				GlobalGameplayCueManager = NewObject<UGameplayCueManager>(this, GCMClass, NAME_None);
+			}
+		}
+
 		if ( GlobalGameplayCueManager == nullptr)
 		{
 			// Fallback to CDO
diff --git a/Engine/Source/Runtime/GameplayAbilities/Private/GameplayCueManager.cpp b/Engine/Source/Runtime/GameplayAbilities/Private/GameplayCueManager.cpp
--- a/Engine/Source/Runtime/GameplayAbilities/Private/GameplayCueManager.cpp
+++ b/Engine/Source/Runtime/GameplayAbilities/Private/GameplayCueManager.cpp
@@ -53,7 +53,7 @@ UGameplayCueManager::UGameplayCueManager(const FObjectInitializer& PCIP)
 
 void UGameplayCueManager::OnCreated()
 {
-	FWorldDelegates::OnPostWorldCreation.AddUObject(this, &UGameplayCueManager::OnWorldCreated);
+	FWorldDelegates::OnPostWorldInitialization.AddUObject(this, &UGameplayCueManager::OnWorldCreated);
 	FWorldDelegates::OnWorldCleanup.AddUObject(this, &UGameplayCueManager::OnWorldCleanup);
 	FWorldDelegates::OnPreWorldFinishDestroy.AddUObject(this, &UGameplayCueManager::OnWorldCleanup, true, true);
 
@@ -1108,7 +1108,7 @@ FPreallocationInfo& UGameplayCueManager::GetPreallocationInfo(UWorld* World)
 
 }
 
-void UGameplayCueManager::OnWorldCreated(UWorld* NewWorld)
+void UGameplayCueManager::OnWorldCreated(UWorld* NewWorld, const UWorld::InitializationValues IV )
 {
 	// Attempting to track down rare GC error where PreallocationInfo_Internal.OwningWorld is not cleaned up.
 	ABILITY_LOG(Display, TEXT("UGameplayCueManager::OnWorldCreated %s. Current PreallocationInfo_Internal: %s"), *GetNameSafe(NewWorld), *GetNameSafe(PreallocationInfo_Internal.OwningWorld));
@@ -1129,7 +1129,7 @@ void UGameplayCueManager::OnWorldCleanup(UWorld* World, bool bSessionEnded, bool
 	if (PreallocationInfo_Internal.OwningWorld == World)
 	{
 		// Reset PreallocationInfo_Internal
-		OnWorldCreated(nullptr);
+		OnWorldCreated(nullptr, UWorld::InitializationValues());
 	}
 
 #if WITH_EDITOR
diff --git a/Engine/Source/Runtime/GameplayAbilities/Private/GameplayCueSet.cpp b/Engine/Source/Runtime/GameplayAbilities/Private/GameplayCueSet.cpp
--- a/Engine/Source/Runtime/GameplayAbilities/Private/GameplayCueSet.cpp
+++ b/Engine/Source/Runtime/GameplayAbilities/Private/GameplayCueSet.cpp
@@ -175,7 +175,7 @@ bool UGameplayCueSet::HandleGameplayCueNotify_Internal(AActor* TargetActor, int3
 				{
 					CueManager->StreamableManager.SimpleAsyncLoad(CueData.GameplayCueNotifyObj);
 
-					ABILITY_LOG(Warning, TEXT("GameplayCueNotify %s was not loaded when GameplayCue was invoked. Starting async loading."), *CueData.GameplayCueNotifyObj.ToString());
+					ABILITY_LOG(Display, TEXT("GameplayCueNotify %s was not loaded when GameplayCue was invoked. Starting async loading."), *CueData.GameplayCueNotifyObj.ToString());
 				}
 				return false;
 			}
diff --git a/Engine/Source/Runtime/GameplayAbilities/Private/GameplayEffect.cpp b/Engine/Source/Runtime/GameplayAbilities/Private/GameplayEffect.cpp
--- a/Engine/Source/Runtime/GameplayAbilities/Private/GameplayEffect.cpp
+++ b/Engine/Source/Runtime/GameplayAbilities/Private/GameplayEffect.cpp
@@ -1387,6 +1387,9 @@ void FActiveGameplayEffect::CheckOngoingTagRequirements(const FGameplayTagContai
 		// All OnDirty callbacks must be inhibited until we update this entire GameplayEffect.
 		FScopedAggregatorOnDirtyBatch	AggregatorOnDirtyBatcher;
 
+		// Important to set this prior to adding or removing, so that any delegates that are triggered can query accurately against this GE
+		bIsInhibited = bShouldBeInhibited;
+
 		if (bShouldBeInhibited)
 		{
 			// Remove our ActiveGameplayEffects modifiers with our Attribute Aggregators
@@ -1396,8 +1399,6 @@ void FActiveGameplayEffect::CheckOngoingTagRequirements(const FGameplayTagContai
 		{
 			OwningContainer.AddActiveGameplayEffectGrantedTagsAndModifiers(*this, bInvokeGameplayCueEvents);
 		}
-
-		bIsInhibited = bShouldBeInhibited;
 	}
 }
 
@@ -3318,27 +3319,41 @@ TArray<FActiveGameplayEffectHandle> FActiveGameplayEffectsContainer::GetActiveEf
 
 float FActiveGameplayEffectsContainer::GetActiveEffectsEndTime(const FGameplayEffectQuery& Query) const
 {
-	float ReturnTime = 0.f;
+	float EndTime = 0.f;
+	float Duration = 0.f;
+	GetActiveEffectsEndTimeAndDuration(Query, EndTime, Duration);
+	return EndTime;
+}
+
+bool FActiveGameplayEffectsContainer::GetActiveEffectsEndTimeAndDuration(const FGameplayEffectQuery& Query, float& EndTime, float& Duration) const
+{
+	bool FoundSomething = false;
+	
 	for (const FActiveGameplayEffect& Effect : this)
 	{
 		if (!Query.Matches(Effect))
 		{
 			continue;
 		}
+		
+		FoundSomething = true;
 
-		float EndTime = Effect.GetEndTime();
-		if (EndTime <= UGameplayEffect::INFINITE_DURATION)
+		float ThisEndTime = Effect.GetEndTime();
+		if (ThisEndTime <= UGameplayEffect::INFINITE_DURATION)
 		{
 			// This is an infinite duration effect, so this end time is indeterminate
-			return -1.f;
+			EndTime = -1.f;
+			Duration = -1.f;
+			return true;
 		}
 
-		if (EndTime > ReturnTime)
+		if (ThisEndTime > EndTime)
 		{
-			ReturnTime = EndTime;
+			EndTime = ThisEndTime;
+			Duration = Effect.GetDuration();
 		}
 	}
-	return ReturnTime;
+	return FoundSomething;
 }
 
 TArray<FActiveGameplayEffectHandle> FActiveGameplayEffectsContainer::GetAllActiveEffectHandles() const
diff --git a/Engine/Source/Runtime/GameplayAbilities/Private/GameplayEffectCustomApplicationRequirement.cpp b/Engine/Source/Runtime/GameplayAbilities/Private/GameplayEffectCustomApplicationRequirement.cpp
new file mode 100644
--- /dev/null
+++ b/Engine/Source/Runtime/GameplayAbilities/Private/GameplayEffectCustomApplicationRequirement.cpp
@@ -0,0 +1,15 @@
+// Copyright 1998-2016 Epic Games, Inc. All Rights Reserved.
+
+#include "AbilitySystemPrivatePCH.h"
+#include "GameplayEffectCustomApplicationRequirement.h"
+
+UGameplayEffectCustomApplicationRequirement::UGameplayEffectCustomApplicationRequirement(const FObjectInitializer& ObjectInitializer)
+	: Super(ObjectInitializer)
+{
+
+}
+
+bool UGameplayEffectCustomApplicationRequirement::CanApplyGameplayEffect_Implementation(const UGameplayEffect* GameplayEffect, const FGameplayEffectSpec& Spec, UAbilitySystemComponent* ASC) const
+{
+	return true;
+}
diff --git a/Engine/Source/Runtime/GameplayAbilities/Public/Abilities/GameplayAbility.h b/Engine/Source/Runtime/GameplayAbilities/Public/Abilities/GameplayAbility.h
--- a/Engine/Source/Runtime/GameplayAbilities/Public/Abilities/GameplayAbility.h
+++ b/Engine/Source/Runtime/GameplayAbilities/Public/Abilities/GameplayAbility.h
@@ -639,6 +639,13 @@ public:
 	UFUNCTION(BlueprintCallable, Category = Ability)
 	float GetCooldownTimeRemaining() const;
 
+	/** 
+	 * Invalidates the current prediction key. This should be used in cases where there is a valid prediction window, but the server is doing logic that only he can do, and afterwards performs an action that the client could predict (had the client been able to run the server-only code prior).
+	 * This returns instantly and has no other side effects other than clearing the current prediction key.
+	 */ 
+	UFUNCTION(BlueprintCallable, Category = Ability)
+	void InvalidateClientPredictionKey() const;
+
 protected:
 
 	// -----------------------------------------------	
diff --git a/Engine/Source/Runtime/GameplayAbilities/Public/Abilities/Tasks/AbilityTask_WaitAttributeChangeRatioThreshold.h b/Engine/Source/Runtime/GameplayAbilities/Public/Abilities/Tasks/AbilityTask_WaitAttributeChangeRatioThreshold.h
new file mode 100644
--- /dev/null
+++ b/Engine/Source/Runtime/GameplayAbilities/Public/Abilities/Tasks/AbilityTask_WaitAttributeChangeRatioThreshold.h
@@ -0,0 +1,56 @@
+// Copyright 1998-2016 Epic Games, Inc. All Rights Reserved.
+#pragma once
+#include "AbilityTask.h"
+#include "Abilities/GameplayAbilityTypes.h"
+#include "AttributeSet.h"
+#include "AbilityTask_WaitAttributeChange.h"
+#include "AbilityTask_WaitAttributeChangeRatioThreshold.generated.h"
+
+struct FGameplayEffectModCallbackData;
+
+DECLARE_DYNAMIC_MULTICAST_DELEGATE_TwoParams(FWaitAttributeChangeRatioThresholdDelegate, bool, bMatchesComparison, float, CurrentRatio);
+
+/**
+ *	Waits for the ratio between two attributes to match a threshold
+ */
+UCLASS(MinimalAPI)
+class UAbilityTask_WaitAttributeChangeRatioThreshold : public UAbilityTask
+{
+	GENERATED_UCLASS_BODY()
+
+	UPROPERTY(BlueprintAssignable)
+	FWaitAttributeChangeRatioThresholdDelegate OnChange;
+
+	virtual void Activate() override;
+
+	void OnNumeratorAttributeChange(float NewValue, const FGameplayEffectModCallbackData* Data);
+	void OnDenominatorAttributeChange(float NewValue, const FGameplayEffectModCallbackData* Data);
+
+	/** Wait on attribute ratio change meeting a comparison threshold. */
+	UFUNCTION(BlueprintCallable, Category="Ability|Tasks", meta = (HidePin = "WorldContextObject", DefaultToSelf = "WorldContextObject", BlueprintInternalUseOnly = "TRUE"))
+	static UAbilityTask_WaitAttributeChangeRatioThreshold* WaitForAttributeChangeRatioThreshold(UObject* WorldContextObject, FGameplayAttribute AttributeNumerator, FGameplayAttribute AttributeDenominator, TEnumAsByte<EWaitAttributeChangeComparison::Type> ComparisonType, float ComparisonValue, bool bTriggerOnce);
+
+	FGameplayAttribute AttributeNumerator;
+	FGameplayAttribute AttributeDenominator;
+	TEnumAsByte<EWaitAttributeChangeComparison::Type> ComparisonType;
+	float ComparisonValue;
+	bool bTriggerOnce;
+	FDelegateHandle OnNumeratorAttributeChangeDelegateHandle;
+	FDelegateHandle OnDenominatorAttributeChangeDelegateHandle;
+
+protected:
+
+	float LastAttributeNumeratorValue;
+	float LastAttributeDenominatorValue;
+	bool bMatchedComparisonLastAttributeChange;
+
+	/** Timer used when either numerator or denominator attribute is changed to delay checking of ratio to avoid false positives (MaxHealth changed before Health updates accordingly) */
+	FTimerHandle CheckAttributeTimer;
+
+	void OnAttributeChange();
+	void OnRatioChange();
+
+	virtual void OnDestroy(bool AbilityEnded) override;
+
+	bool DoesValuePassComparison(float ValueNumerator, float ValueDenominator) const;
+};
diff --git a/Engine/Source/Runtime/GameplayAbilities/Public/Abilities/Tasks/AbilityTask_WaitAttributeChangeThreshold.h b/Engine/Source/Runtime/GameplayAbilities/Public/Abilities/Tasks/AbilityTask_WaitAttributeChangeThreshold.h
new file mode 100644
--- /dev/null
+++ b/Engine/Source/Runtime/GameplayAbilities/Public/Abilities/Tasks/AbilityTask_WaitAttributeChangeThreshold.h
@@ -0,0 +1,45 @@
+// Copyright 1998-2016 Epic Games, Inc. All Rights Reserved.
+#pragma once
+#include "AbilityTask.h"
+#include "Abilities/GameplayAbilityTypes.h"
+#include "AttributeSet.h"
+#include "AbilityTask_WaitAttributeChange.h"
+#include "AbilityTask_WaitAttributeChangeThreshold.generated.h"
+
+struct FGameplayEffectModCallbackData;
+
+DECLARE_DYNAMIC_MULTICAST_DELEGATE_TwoParams(FWaitAttributeChangeThresholdDelegate, bool, bMatchesComparison, float, CurrentValue);
+
+/**
+ *	Waits for an attribute to match a threshold
+ */
+UCLASS(MinimalAPI)
+class UAbilityTask_WaitAttributeChangeThreshold : public UAbilityTask
+{
+	GENERATED_UCLASS_BODY()
+
+	UPROPERTY(BlueprintAssignable)
+	FWaitAttributeChangeThresholdDelegate OnChange;
+
+	virtual void Activate() override;
+
+	void OnAttributeChange(float NewValue, const FGameplayEffectModCallbackData* Data);
+
+	/** Wait on attribute change meeting a comparison threshold. */
+	UFUNCTION(BlueprintCallable, Category="Ability|Tasks", meta = (HidePin = "WorldContextObject", DefaultToSelf = "WorldContextObject", BlueprintInternalUseOnly = "TRUE"))
+	static UAbilityTask_WaitAttributeChangeThreshold* WaitForAttributeChangeThreshold(UObject* WorldContextObject, FGameplayAttribute Attribute, TEnumAsByte<EWaitAttributeChangeComparison::Type> ComparisonType, float ComparisonValue, bool bTriggerOnce);
+
+	FGameplayAttribute Attribute;
+	TEnumAsByte<EWaitAttributeChangeComparison::Type> ComparisonType;
+	float ComparisonValue;
+	bool bTriggerOnce;
+	FDelegateHandle OnAttributeChangeDelegateHandle;
+
+protected:
+
+	bool bMatchedComparisonLastAttributeChange;
+
+	virtual void OnDestroy(bool AbilityEnded) override;
+
+	bool DoesValuePassComparison(float Value) const;
+};
diff --git a/Engine/Source/Runtime/GameplayAbilities/Public/AbilitySystemComponent.h b/Engine/Source/Runtime/GameplayAbilities/Public/AbilitySystemComponent.h
--- a/Engine/Source/Runtime/GameplayAbilities/Public/AbilitySystemComponent.h
+++ b/Engine/Source/Runtime/GameplayAbilities/Public/AbilitySystemComponent.h
@@ -599,6 +599,11 @@ class GAMEPLAYABILITIES_API UAbilitySystemComponent : public UGameplayTasksCompo
 		return ActiveGameplayEffects.GetActiveEffectsEndTime(Query);
 	}
 
+	bool GetActiveEffectsEndTimeAndDuration(const FGameplayEffectQuery& Query, float& EndTime, float& Duration) const
+	{
+		return ActiveGameplayEffects.GetActiveEffectsEndTimeAndDuration(Query, EndTime, Duration);
+	}
+
 	void ModifyActiveEffectStartTime(FActiveGameplayEffectHandle Handle, float StartTimeDiff)
 	{
 		ActiveGameplayEffects.ModifyActiveEffectStartTime(Handle, StartTimeDiff);
diff --git a/Engine/Source/Runtime/GameplayAbilities/Public/AbilitySystemGlobals.h b/Engine/Source/Runtime/GameplayAbilities/Public/AbilitySystemGlobals.h
--- a/Engine/Source/Runtime/GameplayAbilities/Public/AbilitySystemGlobals.h
+++ b/Engine/Source/Runtime/GameplayAbilities/Public/AbilitySystemGlobals.h
@@ -207,7 +207,11 @@ protected:
 	UPROPERTY(config)
 	FString GlobalAttributeSetDefaultsTableName;
 
-	/** The class to instantiate as the global GameplayCue manager. This class is responsible for directing GameplayCue events, loading and cooking assets related to GameplayCues. */
+	/** Class reference to gameplay cue manager. Use this if you want to just instantiate a class for your gameplay cue manager without having to create an asset. */
+	UPROPERTY(config)
+	FStringAssetReference GlobalGameplayCueManagerClass;
+
+	/** Object reference to gameplay cue manager (E.g., reference to a specific blueprint of your GameplayCueManager class. This is not necessary unless you want to have data or blueprints in your gameplay cue manager. */
 	UPROPERTY(config)
 	FStringAssetReference GlobalGameplayCueManagerName;
 
diff --git a/Engine/Source/Runtime/GameplayAbilities/Public/GameplayCueManager.h b/Engine/Source/Runtime/GameplayAbilities/Public/GameplayCueManager.h
--- a/Engine/Source/Runtime/GameplayAbilities/Public/GameplayCueManager.h
+++ b/Engine/Source/Runtime/GameplayAbilities/Public/GameplayCueManager.h
@@ -103,7 +103,7 @@ class GAMEPLAYABILITIES_API UGameplayCueManager : public UDataAsset
 	/** Prespawns a single actor for gameplaycue notify actor classes that need prespawning (should be called by outside gamecode, such as gamestate) */
 	void UpdatePreallocation(UWorld* World);
 
-	void OnWorldCreated(UWorld* NewWorld);
+	void OnWorldCreated(UWorld* NewWorld, const UWorld::InitializationValues);
 
 	void OnWorldCleanup(UWorld* World, bool bSessionEnded, bool bCleanupResources);
 
@@ -203,21 +203,21 @@ protected:
 	void CheckForPreallocation(UClass* GCClass);
 
 	/** Hardref to the gameplaycue notify classes we have async loaded*/
-	UPROPERTY()
+	UPROPERTY(transient)
 	TArray<UClass*> LoadedGameplayCueNotifyClasses;
 
 	/** Classes that we need to preallocate instances for */
-	UPROPERTY()
+	UPROPERTY(transient)
 	TArray<AGameplayCueNotify_Actor*> GameplayCueClassesForPreallocation;
 
 	TArray<FString>	LoadedPaths;
 
 	/** List of gameplay cue executes that haven't been processed yet */
-	UPROPERTY()
+	UPROPERTY(transient)
 	TArray<FGameplayCuePendingExecute> PendingExecuteCues;
 
 	/** Number of active gameplay cue send contexts, when it goes to 0 cues are flushed */
-	UPROPERTY()
+	UPROPERTY(transient)
 	int32 GameplayCueSendContextCount;
 
 	/** Cached world we are currently handling cues for. Used for non instanced GC Notifies that need world. */
diff --git a/Engine/Source/Runtime/GameplayAbilities/Public/GameplayEffect.h b/Engine/Source/Runtime/GameplayAbilities/Public/GameplayEffect.h
--- a/Engine/Source/Runtime/GameplayAbilities/Public/GameplayEffect.h
+++ b/Engine/Source/Runtime/GameplayAbilities/Public/GameplayEffect.h
@@ -20,6 +20,7 @@ class UGameplayEffectTemplate;
 class UAbilitySystemComponent;
 class UGameplayModMagnitudeCalculation;
 class UGameplayEffectExecutionCalculation;
+class UGameplayEffectCustomApplicationRequirement;
 
 /** Enumeration outlining the possible gameplay effect magnitude calculation policies. */
 UENUM()
@@ -609,6 +610,9 @@ public:
 	UPROPERTY(EditDefaultsOnly, Category=Application, meta=(GameplayAttribute="True"))
 	FScalableFloat	ChanceToApplyToTarget;
 
+	UPROPERTY(EditDefaultsOnly, Category=Application, DisplayName="Application Requirement")
+	TArray<TSubclassOf<UGameplayEffectCustomApplicationRequirement> > ApplicationRequirements;
+
 	/** other gameplay effects that will be applied to the target of this effect if this effect applies */
 	UPROPERTY(EditDefaultsOnly, BlueprintReadOnly, Category = GameplayEffect,meta = (DisplayName = "Linked Gameplay Effects"))
 	TArray<TSubclassOf<UGameplayEffect>> TargetEffectClasses;
@@ -1642,6 +1646,7 @@ struct GAMEPLAYABILITIES_API FActiveGameplayEffectsContainer : public FFastArray
 	TArray<FActiveGameplayEffectHandle> GetActiveEffects(const FGameplayEffectQuery& Query) const;
 
 	float GetActiveEffectsEndTime(const FGameplayEffectQuery& Query) const;
+	bool GetActiveEffectsEndTimeAndDuration(const FGameplayEffectQuery& Query, float& EndTime, float& Duration) const;
 
 	/** Returns an array of all of the active gameplay effect handles */
 	TArray<FActiveGameplayEffectHandle> GetAllActiveEffectHandles() const;
diff --git a/Engine/Source/Runtime/GameplayAbilities/Public/GameplayEffectCustomApplicationRequirement.h b/Engine/Source/Runtime/GameplayAbilities/Public/GameplayEffectCustomApplicationRequirement.h
new file mode 100644
--- /dev/null
+++ b/Engine/Source/Runtime/GameplayAbilities/Public/GameplayEffectCustomApplicationRequirement.h
@@ -0,0 +1,24 @@
+// Copyright 1998-2016 Epic Games, Inc. All Rights Reserved.
+
+#pragma once
+
+
+#include "GameplayEffectCustomApplicationRequirement.generated.h"
+
+// Forward declarations
+struct FGameplayEffectSpec;
+class UAbilitySystemComponent;
+class UGameplayEffect;
+
+/** Class used to perform custom gameplay effect modifier calculations, either via blueprint or native code */ 
+UCLASS(BlueprintType, Blueprintable, Abstract)
+class GAMEPLAYABILITIES_API UGameplayEffectCustomApplicationRequirement : public UObject
+{
+
+public:
+	GENERATED_UCLASS_BODY()
+	
+	/** Return whether the gameplay effect should be applied or not */
+	UFUNCTION(BlueprintNativeEvent, Category="Calculation")
+	bool CanApplyGameplayEffect(const UGameplayEffect* GameplayEffect, const FGameplayEffectSpec& Spec, UAbilitySystemComponent* ASC) const;
+};
\ No newline at end of file
diff --git a/Engine/Source/Runtime/MoviePlayer/Private/DefaultGameMoviePlayer.cpp b/Engine/Source/Runtime/MoviePlayer/Private/DefaultGameMoviePlayer.cpp
--- a/Engine/Source/Runtime/MoviePlayer/Private/DefaultGameMoviePlayer.cpp
+++ b/Engine/Source/Runtime/MoviePlayer/Private/DefaultGameMoviePlayer.cpp
@@ -15,6 +15,9 @@
 #include "MoviePlayerSettings.h"
 #include "ShaderCompiler.h"
 
+
+DEFINE_LOG_CATEGORY_STATIC(LogMoviePlayer, Log, All);
+
 class SDefaultMovieBorder : public SBorder
 {
 public:
@@ -114,6 +117,8 @@ void FDefaultGameMoviePlayer::SetSlateRenderer(TSharedPtr<FSlateRenderer> InSlat
 
 void FDefaultGameMoviePlayer::Initialize()
 {
+	UE_LOG(LogMoviePlayer, Log, TEXT("Initializing movie player"));
+
 	ENQUEUE_UNIQUE_RENDER_COMMAND_ONEPARAMETER(RegisterMoviePlayerTickable, FDefaultGameMoviePlayer*, MoviePlayer, this,
 	{
 		MoviePlayer->Register();
@@ -182,6 +187,8 @@ void FDefaultGameMoviePlayer::Initialize()
 
 void FDefaultGameMoviePlayer::Shutdown()
 {
+	UE_LOG(LogMoviePlayer, Log, TEXT("Shutting down movie player"));
+
 	StopMovie();
 	WaitForMovieToFinish();
 
@@ -209,15 +216,19 @@ void FDefaultGameMoviePlayer::Shutdown()
 		delete SyncMechanism;
 		SyncMechanism = NULL;
 	}
-
 }
+
 void FDefaultGameMoviePlayer::PassLoadingScreenWindowBackToGame() const
 {
-	auto GameEngine = Cast<UGameEngine>(GEngine);
+	UGameEngine* GameEngine = Cast<UGameEngine>(GEngine);
 	if (LoadingScreenWindowPtr.IsValid() && GameEngine)
 	{
 		GameEngine->GameViewportWindow = LoadingScreenWindowPtr;
 	}
+	else
+	{
+		UE_LOG(LogMoviePlayer, Warning, TEXT("PassLoadingScreenWindowBackToGame failed.  No Window") );
+	}
 }
 
 void FDefaultGameMoviePlayer::SetupLoadingScreen(const FLoadingScreenAttributes& InLoadingScreenAttributes)
diff --git a/Engine/Source/Runtime/Online/Hotfix/Private/OnlineHotfixManager.cpp b/Engine/Source/Runtime/Online/Hotfix/Private/OnlineHotfixManager.cpp
--- a/Engine/Source/Runtime/Online/Hotfix/Private/OnlineHotfixManager.cpp
+++ b/Engine/Source/Runtime/Online/Hotfix/Private/OnlineHotfixManager.cpp
@@ -121,7 +121,8 @@ void UOnlineHotfixManager::Cleanup()
 void UOnlineHotfixManager::StartHotfixProcess()
 {
 	// Patching the editor this way seems like a bad idea
-	if (!IsRunningGame())
+	bool bShouldHotfix = IsRunningGame() || IsRunningDedicatedServer() || IsRunningClientOnly();
+	if (!bShouldHotfix)
 	{
 		UE_LOG(LogHotfixManager, Warning, TEXT("Hotfixing skipped due to IsRunningGame() == false"));
 		TriggerHotfixComplete(EHotfixResult::SuccessNoChange);
diff --git a/Engine/Source/Runtime/Online/OnlineFramework/Party/Private/Party.cpp b/Engine/Source/Runtime/Online/OnlineFramework/Party/Private/Party.cpp
--- a/Engine/Source/Runtime/Online/OnlineFramework/Party/Private/Party.cpp
+++ b/Engine/Source/Runtime/Online/OnlineFramework/Party/Private/Party.cpp
@@ -21,7 +21,11 @@ UParty::UParty(const FObjectInitializer& ObjectInitializer) :
 
 void UParty::Init()
 {
-	FCoreUObjectDelegates::PostLoadMap.AddUObject(this, &UParty::OnPostLoadMap);
+	UGameInstance* GameInstance = GetGameInstance();
+	check(GameInstance);
+	GameInstance->OnNotifyPreClientTravel().AddUObject(this, &ThisClass::NotifyPreClientTravel);
+
+	FCoreUObjectDelegates::PostLoadMap.AddUObject(this, &ThisClass::OnPostLoadMap);
 }
 
 void UParty::InitPIE()
@@ -58,6 +62,12 @@ void UParty::OnShutdown()
 		UnregisterIdentityDelegates();
 		UnregisterPartyDelegates();
 	}
+
+	UGameInstance* GameInstance = GetGameInstance();
+	if (GameInstance)
+	{
+		GameInstance->OnNotifyPreClientTravel().RemoveAll(this);
+	}
 }
 
 void UParty::OnLogoutComplete(int32 LocalUserNum, bool bWasSuccessful)
@@ -232,7 +242,7 @@ void UParty::UnregisterPartyDelegates()
 	}
 }
 
-void UParty::NotifyPreClientTravel()
+void UParty::NotifyPreClientTravel(const FString& PendingURL, ETravelType TravelType, bool bIsSeamlessTravel)
 {
 	for (auto& JoinedParty : JoinedParties)
 	{
diff --git a/Engine/Source/Runtime/Online/OnlineFramework/Party/Public/Party.h b/Engine/Source/Runtime/Online/OnlineFramework/Party/Public/Party.h
--- a/Engine/Source/Runtime/Online/OnlineFramework/Party/Public/Party.h
+++ b/Engine/Source/Runtime/Online/OnlineFramework/Party/Public/Party.h
@@ -290,7 +290,7 @@ public:
 	UPartyGameState* GetPersistentParty() const;
 
 	/** Notify the party system that travel is occurring */
-	void NotifyPreClientTravel();
+	void NotifyPreClientTravel(const FString& PendingURL, ETravelType TravelType, bool bIsSeamlessTravel);
 
 	/**
 	 * Party invites
diff --git a/Engine/Source/Runtime/Online/XMPP/Private/XmppJingle/XmppMessagesJingle.cpp b/Engine/Source/Runtime/Online/XMPP/Private/XmppJingle/XmppMessagesJingle.cpp
--- a/Engine/Source/Runtime/Online/XMPP/Private/XmppJingle/XmppMessagesJingle.cpp
+++ b/Engine/Source/Runtime/Online/XMPP/Private/XmppJingle/XmppMessagesJingle.cpp
@@ -124,7 +124,7 @@ public:
 			return buzz::XMPP_RETURN_BADSTATE;
 		}
 
-		buzz::Jid ToJidFull(ToJid.node(), GetClient()->jid().domain(), buzz::STR_EMPTY);
+		buzz::Jid ToJidFull(ToJid.node(), GetClient()->jid().domain(), ToJid.resource());
 		buzz::XmlElement* Stanza = MessageToStanza(ToJidFull, Message);
 		QueueStanza(Stanza);
 		delete Stanza;
