Fix for UE-26989 and UE-26954; set deadtrail flag on particles only under proper conditions, to avoid cached indices going out of sync.

diff --git a/Engine/Source/Runtime/Engine/Private/Particles/ParticleTrail2EmitterInstance.cpp b/Engine/Source/Runtime/Engine/Private/Particles/ParticleTrail2EmitterInstance.cpp
--- a/Engine/Source/Runtime/Engine/Private/Particles/ParticleTrail2EmitterInstance.cpp
+++ b/Engine/Source/Runtime/Engine/Private/Particles/ParticleTrail2EmitterInstance.cpp
@@ -3469,7 +3469,14 @@ float FParticleAnimTrailEmitterInstance::Spawn(float DeltaTime)
 			int32 CheckStartIndex = ParticleIndices[FindTrailIdx];
 			DECLARE_PARTICLE_PTR(CheckParticle, ParticleData + ParticleStride * CheckStartIndex);
 			FAnimTrailTypeDataPayload* CheckTrailData = ((FAnimTrailTypeDataPayload*)((uint8*)CheckParticle + TypeDataOffset));
-			CheckTrailData->Flags = TRAIL_EMITTER_SET_DEADTRAIL(CheckTrailData->Flags);
+
+			if (CheckTrailData->TrailIndex == TrailIdx)
+			{
+				if (TRAIL_EMITTER_IS_START(CheckTrailData->Flags))
+				{
+					CheckTrailData->Flags = TRAIL_EMITTER_SET_DEADTRAIL(CheckTrailData->Flags);
+				}
+			}
 			SetDeadIndex(CheckTrailData->TrailIndex, CheckStartIndex);
 		}
 		bTagTrailAsDead = false;
