build error - HTML5 infinite loop

emscripten builds sometimes will be stuck in an infinite loop if builds dump terse error output...
in this case, a missing variable "a" for PLATFORM_HTML5 builds...  cause this to happen...
this will be another test case to send to emscripten maintainers...

diff --git a/Engine/Source/Runtime/Engine/Public/ShowFlags.h b/Engine/Source/Runtime/Engine/Public/ShowFlags.h
--- a/Engine/Source/Runtime/Engine/Public/ShowFlags.h
+++ b/Engine/Source/Runtime/Engine/Public/ShowFlags.h
@@ -50,8 +50,9 @@ struct FEngineShowFlags
 	// ---------------------------------------------------------
 	// Define the showflags.
 	// A show flag is either an uint32:1 or static const bool (if optimized out according to UE_BUILD_OPTIMIZED_SHOWFLAGS)
+
 #if PLATFORM_HTML5 // broken fit field compiler -- will be sending this file to the emscripten &/or clang keepers
-	#define SHOWFLAG_ALWAYS_ACCESSIBLE(a,...) bool; void Set##a(bool bVal){ a = bVal;}
+	#define SHOWFLAG_ALWAYS_ACCESSIBLE(a,...) bool a; void Set##a(bool bVal){ a = bVal;}
 #else
 	#define SHOWFLAG_ALWAYS_ACCESSIBLE(a,...) uint32 a : 1; void Set##a(bool bVal){ a = bVal?1:0;}
 #endif
