Copying //UE4/Dev-Core to //UE4/Main

diff --git a/Engine/Source/Runtime/Engine/Private/Materials/Material.cpp b/Engine/Source/Runtime/Engine/Private/Materials/Material.cpp
--- a/Engine/Source/Runtime/Engine/Private/Materials/Material.cpp
+++ b/Engine/Source/Runtime/Engine/Private/Materials/Material.cpp
@@ -73,6 +73,7 @@ FMaterialResource::FMaterialResource()
 
 int32 FMaterialResource::CompilePropertyAndSetMaterialProperty(EMaterialProperty Property, FMaterialCompiler* Compiler, EShaderFrequency OverrideShaderFrequency, bool bUsePreviousFrameTime) const
 {
+#if WITH_EDITOR
 	// needs to be called in this function!!
 	// sets CurrentShaderFrequency
 	Compiler->SetMaterialProperty(Property, OverrideShaderFrequency, bUsePreviousFrameTime);
@@ -119,6 +120,10 @@ int32 FMaterialResource::CompilePropertyAndSetMaterialProperty(EMaterialProperty
 
 	// output should always be the right type for this property
 	return Compiler->ForceCast(Ret, GetMaterialPropertyType(Property));
+#else // WITH_EDITOR
+	check(0); // This is editor-only function
+	return INDEX_NONE;
+#endif // WITH_EDITOR
 }
 
 void FMaterialResource::GatherCustomOutputExpressions(TArray<UMaterialExpressionCustomOutput*>& OutCustomOutputs) const
@@ -1432,18 +1437,16 @@ bool UMaterial::GetVectorParameterValue(FName ParameterName, FLinearColor& OutVa
 {
 	for (const UMaterialExpression* Expression : Expressions)
 	{
-		if (Expression->IsA<UMaterialExpressionVectorParameter>())
+		if (const UMaterialExpressionVectorParameter* Parameter = Cast<const UMaterialExpressionVectorParameter>(Expression))
 		{
-			const UMaterialExpressionVectorParameter* Parameter = CastChecked<const UMaterialExpressionVectorParameter>(Expression);
 			if (Parameter->IsNamedParameter(ParameterName, OutValue))
 			{
 				// Warning: in the case of duplicate parameters with different default values, this will find the first in the expression array, not necessarily the one that's used for rendering
 				return true;
 			}
 		}
-		else if (Expression->IsA<UMaterialExpressionMaterialFunctionCall>())
+		else if (const UMaterialExpressionMaterialFunctionCall* FunctionCall = Cast<const UMaterialExpressionMaterialFunctionCall>(Expression))
 		{
-			const UMaterialExpressionMaterialFunctionCall* FunctionCall = CastChecked<const UMaterialExpressionMaterialFunctionCall>(Expression);
 			if (FunctionCall->MaterialFunction)
 			{
 				TArray<UMaterialFunction*> Functions;
@@ -1454,10 +1457,9 @@ bool UMaterial::GetVectorParameterValue(FName ParameterName, FLinearColor& OutVa
 				{
 					for (UMaterialExpression* FunctionExpression : Function->FunctionExpressions)
 					{
-						if (FunctionExpression->IsA<UMaterialExpressionVectorParameter>())
+						if (const UMaterialExpressionVectorParameter* FunctionExpressionParameter = Cast<const UMaterialExpressionVectorParameter>(FunctionExpression))
 						{
-							const UMaterialExpressionVectorParameter* Parameter = CastChecked<const UMaterialExpressionVectorParameter>(FunctionExpression);
-							if (Parameter->IsNamedParameter(ParameterName, OutValue))
+							if (FunctionExpressionParameter->IsNamedParameter(ParameterName, OutValue))
 							{
 								return true;
 							}
@@ -1475,18 +1477,16 @@ bool UMaterial::GetScalarParameterValue(FName ParameterName, float& OutValue) co
 {
 	for (const UMaterialExpression* Expression : Expressions)
 	{
-		if (Expression->IsA<UMaterialExpressionScalarParameter>())
+		if (const UMaterialExpressionScalarParameter* Parameter = Cast<const UMaterialExpressionScalarParameter>(Expression))
 		{
-			const UMaterialExpressionScalarParameter* Parameter = CastChecked<const UMaterialExpressionScalarParameter>(Expression);
 			if (Parameter->IsNamedParameter(ParameterName, OutValue))
 			{
 				// Warning: in the case of duplicate parameters with different default values, this will find the first in the expression array, not necessarily the one that's used for rendering
 				return true;
 			}
 		}
-		else if (Expression->IsA<UMaterialExpressionMaterialFunctionCall>())
+		else if (const UMaterialExpressionMaterialFunctionCall* FunctionCall = Cast<const UMaterialExpressionMaterialFunctionCall>(Expression))
 		{
-			const UMaterialExpressionMaterialFunctionCall* FunctionCall = CastChecked<const UMaterialExpressionMaterialFunctionCall>(Expression);
 			if (FunctionCall->MaterialFunction)
 			{
 				TArray<UMaterialFunction*> Functions;
@@ -1497,10 +1497,9 @@ bool UMaterial::GetScalarParameterValue(FName ParameterName, float& OutValue) co
 				{
 					for (UMaterialExpression* FunctionExpression : Function->FunctionExpressions)
 					{
-						if (FunctionExpression->IsA<UMaterialExpressionScalarParameter>())
+						if (const UMaterialExpressionScalarParameter* FunctionExpressionParameter = Cast<const UMaterialExpressionScalarParameter>(FunctionExpression))
 						{
-							const UMaterialExpressionScalarParameter* Parameter = CastChecked<const UMaterialExpressionScalarParameter>(FunctionExpression);
-							if (Parameter->IsNamedParameter(ParameterName, OutValue))
+							if (FunctionExpressionParameter->IsNamedParameter(ParameterName, OutValue))
 							{
 								return true;
 							}
@@ -1567,18 +1566,16 @@ bool UMaterial::GetTextureParameterValue(FName ParameterName, UTexture*& OutValu
 {
 	for (const UMaterialExpression* Expression : Expressions)
 	{
-		if (Expression->IsA<UMaterialExpressionTextureSampleParameter>())
+		if (const UMaterialExpressionTextureSampleParameter* Parameter = Cast<const UMaterialExpressionTextureSampleParameter>(Expression))
 		{
-			const UMaterialExpressionTextureSampleParameter* Parameter = CastChecked<const UMaterialExpressionTextureSampleParameter>(Expression);
 			if (Parameter->IsNamedParameter(ParameterName, OutValue))
 			{
 				// Warning: in the case of duplicate parameters with different default values, this will find the first in the expression array, not necessarily the one that's used for rendering
 				return true;
 			}
 		}
-		else if (Expression->IsA<UMaterialExpressionMaterialFunctionCall>())
+		else if (const UMaterialExpressionMaterialFunctionCall* FunctionCall = Cast<const UMaterialExpressionMaterialFunctionCall>(Expression))
 		{
-			const UMaterialExpressionMaterialFunctionCall* FunctionCall = CastChecked<const UMaterialExpressionMaterialFunctionCall>(Expression);
 			if (FunctionCall->MaterialFunction)
 			{
 				TArray<UMaterialFunction*> Functions;
@@ -1589,10 +1586,9 @@ bool UMaterial::GetTextureParameterValue(FName ParameterName, UTexture*& OutValu
 				{
 					for (UMaterialExpression* FunctionExpression : Function->FunctionExpressions)
 					{
-						if (FunctionExpression->IsA<UMaterialExpressionTextureSampleParameter>())
+						if (const UMaterialExpressionTextureSampleParameter* FunctionExpressionParameter = Cast<const UMaterialExpressionTextureSampleParameter>(FunctionExpression))
 						{
-							const UMaterialExpressionTextureSampleParameter* Parameter = CastChecked<const UMaterialExpressionTextureSampleParameter>(FunctionExpression);
-							if (Parameter->IsNamedParameter(ParameterName, OutValue))
+							if (FunctionExpressionParameter->IsNamedParameter(ParameterName, OutValue))
 							{
 								return true;
 							}
@@ -1610,18 +1606,16 @@ bool UMaterial::GetFontParameterValue(FName ParameterName, UFont*& OutFontValue,
 {
 	for (const UMaterialExpression* Expression : Expressions)
 	{
-		if (Expression->IsA<UMaterialExpressionFontSampleParameter>())
+		if (const UMaterialExpressionFontSampleParameter* Parameter = Cast<const UMaterialExpressionFontSampleParameter>(Expression))
 		{
-			const UMaterialExpressionFontSampleParameter* Parameter = CastChecked<const UMaterialExpressionFontSampleParameter>(Expression);
 			if (Parameter->IsNamedParameter(ParameterName, OutFontValue, OutFontPage))
 			{
 				// Warning: in the case of duplicate parameters with different default values, this will find the first in the expression array, not necessarily the one that's used for rendering
 				return true;
 			}
 		}
-		else if (Expression->IsA<UMaterialExpressionMaterialFunctionCall>())
+		else if (const UMaterialExpressionMaterialFunctionCall* FunctionCall = Cast<const UMaterialExpressionMaterialFunctionCall>(Expression))
 		{
-			const UMaterialExpressionMaterialFunctionCall* FunctionCall = CastChecked<const UMaterialExpressionMaterialFunctionCall>(Expression);
 			if (FunctionCall->MaterialFunction)
 			{
 				TArray<UMaterialFunction*> Functions;
@@ -1632,10 +1626,9 @@ bool UMaterial::GetFontParameterValue(FName ParameterName, UFont*& OutFontValue,
 				{
 					for (UMaterialExpression* FunctionExpression : Function->FunctionExpressions)
 					{
-						if (FunctionExpression->IsA<UMaterialExpressionFontSampleParameter>())
+						if (const UMaterialExpressionFontSampleParameter* FunctionExpressionParameter = Cast<const UMaterialExpressionFontSampleParameter>(FunctionExpression))
 						{
-							const UMaterialExpressionFontSampleParameter* Parameter = CastChecked<const UMaterialExpressionFontSampleParameter>(FunctionExpression);
-							if (Parameter->IsNamedParameter(ParameterName, OutFontValue, OutFontPage))
+							if (FunctionExpressionParameter->IsNamedParameter(ParameterName, OutFontValue, OutFontPage))
 							{
 								return true;
 							}
@@ -1654,18 +1647,16 @@ bool UMaterial::GetStaticSwitchParameterValue(FName ParameterName, bool& OutValu
 {
 	for (const UMaterialExpression* Expression : Expressions)
 	{
-		if (Expression->IsA<UMaterialExpressionStaticBoolParameter>())
+		if (const UMaterialExpressionStaticBoolParameter* Parameter = Cast<const UMaterialExpressionStaticBoolParameter>(Expression))
 		{
-			const UMaterialExpressionStaticBoolParameter* Parameter = CastChecked<const UMaterialExpressionStaticBoolParameter>(Expression);
 			if (Parameter->IsNamedParameter(ParameterName, OutValue, OutExpressionGuid))
 			{
 				// Warning: in the case of duplicate parameters with different default values, this will find the first in the expression array, not necessarily the one that's used for rendering
 				return true;
 			}
 		}
-		else if (Expression->IsA<UMaterialExpressionMaterialFunctionCall>())
+		else if (const UMaterialExpressionMaterialFunctionCall* FunctionCall = Cast<const UMaterialExpressionMaterialFunctionCall>(Expression))
 		{
-			const UMaterialExpressionMaterialFunctionCall* FunctionCall = CastChecked<const UMaterialExpressionMaterialFunctionCall>(Expression);
 			if (FunctionCall->MaterialFunction)
 			{
 				TArray<UMaterialFunction*> Functions;
@@ -1676,10 +1667,9 @@ bool UMaterial::GetStaticSwitchParameterValue(FName ParameterName, bool& OutValu
 				{
 					for (UMaterialExpression* FunctionExpression : Function->FunctionExpressions)
 					{
-						if (FunctionExpression->IsA<UMaterialExpressionStaticBoolParameter>())
+						if (const UMaterialExpressionStaticBoolParameter* FunctionExpressionParameter = Cast<const UMaterialExpressionStaticBoolParameter>(FunctionExpression))
 						{
-							const UMaterialExpressionStaticBoolParameter* Parameter = CastChecked<const UMaterialExpressionStaticBoolParameter>(FunctionExpression);
-							if (Parameter->IsNamedParameter(ParameterName, OutValue, OutExpressionGuid))
+							if (FunctionExpressionParameter->IsNamedParameter(ParameterName, OutValue, OutExpressionGuid))
 							{
 								return true;
 							}
@@ -1698,18 +1688,16 @@ bool UMaterial::GetStaticComponentMaskParameterValue(FName ParameterName, bool&
 {
 	for (const UMaterialExpression* Expression : Expressions)
 	{
-		if (Expression->IsA<UMaterialExpressionStaticComponentMaskParameter>())
+		if (const UMaterialExpressionStaticComponentMaskParameter* Parameter = Cast<const UMaterialExpressionStaticComponentMaskParameter>(Expression))
 		{
-			const UMaterialExpressionStaticComponentMaskParameter* Parameter = CastChecked<const UMaterialExpressionStaticComponentMaskParameter>(Expression);
 			if (Parameter->IsNamedParameter(ParameterName, OutR, OutG, OutB, OutA, OutExpressionGuid))
 			{
 				// Warning: in the case of duplicate parameters with different default values, this will find the first in the expression array, not necessarily the one that's used for rendering
 				return true;
 			}
 		}
-		else if (Expression->IsA<UMaterialExpressionMaterialFunctionCall>())
+		else if (const UMaterialExpressionMaterialFunctionCall* FunctionCall = Cast<const UMaterialExpressionMaterialFunctionCall>(Expression))
 		{
-			const UMaterialExpressionMaterialFunctionCall* FunctionCall = CastChecked<const UMaterialExpressionMaterialFunctionCall>(Expression);
 			if (FunctionCall->MaterialFunction)
 			{
 				TArray<UMaterialFunction*> Functions;
@@ -1720,10 +1708,9 @@ bool UMaterial::GetStaticComponentMaskParameterValue(FName ParameterName, bool&
 				{
 					for (UMaterialExpression* FunctionExpression : Function->FunctionExpressions)
 					{
-						if (FunctionExpression->IsA<UMaterialExpressionStaticComponentMaskParameter>())
+						if (const UMaterialExpressionStaticComponentMaskParameter* FunctionExpressionParameter = Cast<const UMaterialExpressionStaticComponentMaskParameter>(FunctionExpression))
 						{
-							const UMaterialExpressionStaticComponentMaskParameter* Parameter = CastChecked<const UMaterialExpressionStaticComponentMaskParameter>(FunctionExpression);
-							if (Parameter->IsNamedParameter(ParameterName, OutR, OutG, OutB, OutA, OutExpressionGuid))
+							if (FunctionExpressionParameter->IsNamedParameter(ParameterName, OutR, OutG, OutB, OutA, OutExpressionGuid))
 							{
 								return true;
 							}
@@ -2085,6 +2072,11 @@ void UMaterial::Serialize(FArchive& Ar)
 {
 	Super::Serialize(Ar);
 
+	if (FPlatformProperties::RequiresCookedData() && Ar.IsLoading())
+	{
+		Expressions.Remove(nullptr);
+	}
+
 	if (Ar.UE4Ver() >= VER_UE4_PURGED_FMATERIAL_COMPILE_OUTPUTS)
 	{
 #if WITH_EDITOR
@@ -2121,6 +2113,7 @@ void UMaterial::Serialize(FArchive& Ar)
 	}
 #endif // #if WITH_EDITOR
 
+#if WITH_EDITORONLY_DATA
 	DoMaterialAttributeReorder(&DiffuseColor_DEPRECATED,Ar.UE4Ver());
 	DoMaterialAttributeReorder(&SpecularColor_DEPRECATED,Ar.UE4Ver());
 	DoMaterialAttributeReorder(&BaseColor,				Ar.UE4Ver());
@@ -2148,6 +2141,7 @@ void UMaterial::Serialize(FArchive& Ar)
 	DoMaterialAttributeReorder(&CustomizedUVs[6],		Ar.UE4Ver());
 	DoMaterialAttributeReorder(&CustomizedUVs[7],		Ar.UE4Ver());
 	DoMaterialAttributeReorder(&PixelDepthOffset,		Ar.UE4Ver());
+#endif // WITH_EDITORONLY_DATA
 	static_assert(MP_MAX == 28, "New material properties must have DoMaterialAttributesReorder called on them to ensure that any future reordering of property pins is correctly applied.");
 
 	if (Ar.UE4Ver() < VER_UE4_MATERIAL_MASKED_BLENDMODE_TIDY)
@@ -2241,7 +2235,7 @@ void UMaterial::GetQualityLevelNodeUsage(TArray<bool, TInlineAllocator<EMaterial
 		{
 			for (int32 InputIndex = 0; InputIndex < EMaterialQualityLevel::Num; InputIndex++)
 			{
-				if (QualitySwitchNode->Inputs[InputIndex].Expression)
+				if (QualitySwitchNode->Inputs[InputIndex].IsConnected())
 				{
 					OutQualityLevelsUsed[InputIndex] = true;
 				}
@@ -2267,7 +2261,7 @@ void UMaterial::GetQualityLevelNodeUsage(TArray<bool, TInlineAllocator<EMaterial
 					{
 						for (int32 InputIndex = 0; InputIndex < EMaterialQualityLevel::Num; InputIndex++)
 						{
-							if (SwitchNode->Inputs[InputIndex].Expression)
+							if (SwitchNode->Inputs[InputIndex].IsConnected())
 							{
 								OutQualityLevelsUsed[InputIndex] = true;
 							}
@@ -3535,6 +3529,7 @@ void UMaterial::GetAllCustomOutputExpressions(TArray<class UMaterialExpressionCu
 	}
 }
 
+#if WITH_EDITOR
 bool UMaterial::GetAllReferencedExpressions(TArray<UMaterialExpression*>& OutExpressions, class FStaticParameterSet* InStaticParameterSet)
 {
 	OutExpressions.Empty();
@@ -3585,7 +3580,6 @@ bool UMaterial::GetExpressionsInPropertyChain(EMaterialProperty InProperty,
 	return true;
 }
 
-
 bool UMaterial::GetTexturesInPropertyChain(EMaterialProperty InProperty, TArray<UTexture*>& OutTextures,  
 	TArray<FName>* OutTextureParamNames, class FStaticParameterSet* InStaticParameterSet)
 {
@@ -3627,7 +3621,6 @@ bool UMaterial::GetTexturesInPropertyChain(EMaterialProperty InProperty, TArray<
 	return false;
 }
 
-
 bool UMaterial::RecursiveGetExpressionChain(UMaterialExpression* InExpression, TArray<FExpressionInput*>& InOutProcessedInputs, 
 	TArray<UMaterialExpression*>& OutExpressions, class FStaticParameterSet* InStaticParameterSet)
 {
@@ -3697,6 +3690,7 @@ bool UMaterial::RecursiveGetExpressionChain(UMaterialExpression* InExpression, T
 
 	return true;
 }
+#endif // WITH_EDITOR
 
 void UMaterial::AppendReferencedTextures(TArray<UTexture*>& InOutTextures) const
 {
@@ -3719,11 +3713,9 @@ void UMaterial::AppendReferencedTextures(TArray<UTexture*>& InOutTextures) const
 				CurrentFunction->AppendReferencedTextures(InOutTextures);
 			}
 		}
-		else
+		// Not all expressions are cooked
+		else if (Expression)
 		{
-			// should never be 0 but there is a bug, see TTP 340934
-			check(Expression);
-
 			UTexture* ReferencedTexture = Expression->GetReferencedTexture();
 
 			if (ReferencedTexture)
@@ -3734,6 +3726,7 @@ void UMaterial::AppendReferencedTextures(TArray<UTexture*>& InOutTextures) const
 	}
 }
 
+#if WITH_EDITOR
 void UMaterial::RecursiveUpdateRealtimePreview( UMaterialExpression* InExpression, TArray<UMaterialExpression*>& InOutExpressionsToProcess )
 {
 	// remove ourselves from the list to process
@@ -3786,6 +3779,7 @@ void UMaterial::RecursiveUpdateRealtimePreview( UMaterialExpression* InExpressio
 		InExpression->bNeedToUpdatePreview = true;
 	}
 }
+#endif // WITH_EDITOR
 
 void UMaterial::GetReferencedFunctionIds(TArray<FGuid>& Ids) const
 {
@@ -3807,6 +3801,7 @@ void UMaterial::GetReferencedParameterCollectionIds(TArray<FGuid>& Ids) const
 	}
 }
 
+#if WITH_EDITOR
 int32 UMaterial::CompilePropertyEx( FMaterialCompiler* Compiler, EMaterialProperty Property )
 {
 	if( bUseMaterialAttributes && MP_DiffuseColor != Property && MP_SpecularColor != Property )
@@ -3855,6 +3850,7 @@ int32 UMaterial::CompilePropertyEx( FMaterialCompiler* Compiler, EMaterialProper
 	check(0);
 	return INDEX_NONE;
 }
+#endif // WITH_EDITOR
 
 void UMaterial::NotifyCompilationFinished(FMaterialResource* CompiledResource)
 {
diff --git a/Engine/Source/Runtime/Engine/Private/Materials/MaterialShared.cpp b/Engine/Source/Runtime/Engine/Private/Materials/MaterialShared.cpp
--- a/Engine/Source/Runtime/Engine/Private/Materials/MaterialShared.cpp
+++ b/Engine/Source/Runtime/Engine/Private/Materials/MaterialShared.cpp
@@ -25,6 +25,8 @@
 #include "VertexFactory.h"
 #include "RendererInterface.h"
 #include "MaterialShaderQualitySettings.h"
+#include "UObject/CoreObjectVersion.h"
+
 DEFINE_LOG_CATEGORY(LogMaterial);
 
 FName MaterialQualityLevelNames[] = 
@@ -64,10 +66,14 @@ int32 FMaterialCompiler::Errorf(const TCHAR* Format,...)
 	return Error(ErrorText);
 }
 
-//
-//	FExpressionInput::Compile
-//
+IMPLEMENT_STRUCT(ExpressionInput);
+IMPLEMENT_STRUCT(ColorMaterialInput);
+IMPLEMENT_STRUCT(ScalarMaterialInput);
+IMPLEMENT_STRUCT(VectorMaterialInput);
+IMPLEMENT_STRUCT(Vector2MaterialInput);
+IMPLEMENT_STRUCT(MaterialAttributesInput);
 
+#if WITH_EDITOR
 int32 FExpressionInput::Compile(class FMaterialCompiler* Compiler, int32 MultiplexIndex)
 {
 	if(Expression)
@@ -106,7 +112,96 @@ void FExpressionInput::Connect( int32 InOutputIndex, class UMaterialExpression*
 	MaskB = Output->MaskB;
 	MaskA = Output->MaskA;
 }
+#endif // WITH_EDITOR
+
+/** Native serialize for FMaterialExpression struct */
+static bool SerializeExpressionInput(FArchive& Ar, FExpressionInput& Input)
+{
+	Ar.UsingCustomVersion(FCoreObjectVersion::GUID);
+
+	if (Ar.CustomVer(FCoreObjectVersion::GUID) < FCoreObjectVersion::MaterialInputNativeSerialize)
+	{
+		return false;
+	}
+
+#if WITH_EDITORONLY_DATA
+	if (!Ar.IsCooking())
+	{
+		Ar << Input.Expression;
+	}
+#endif
+	Ar << Input.OutputIndex;
+	Ar << Input.InputName;
+	Ar << Input.Mask;
+	Ar << Input.MaskR;
+	Ar << Input.MaskG;
+	Ar << Input.MaskB;
+	Ar << Input.MaskA;
+
+	// Some expressions may have been stripped when cooking and Expression can be null after loading
+	// so make sure we keep the information about the connected node in cooked packages
+	if (FPlatformProperties::RequiresCookedData() || (Ar.IsSaving() && Ar.IsCooking()))
+	{
+#if WITH_EDITORONLY_DATA
+		if (Ar.IsSaving())
+		{
+			Input.ExpressionName = Input.Expression ? Input.Expression->GetFName() : NAME_None;
+		}
+#endif // WITH_EDITORONLY_DATA
+		Ar << Input.ExpressionName;
+	}
+
+	return true;
+}
+
+template <typename InputType>
+static bool SerializeMaterialInput(FArchive& Ar, FMaterialInput<InputType>& Input)
+{
+	if (SerializeExpressionInput(Ar, Input))
+	{
+		bool bUseConstantValue = Input.UseConstant;
+		Ar << bUseConstantValue;
+		Input.UseConstant = bUseConstantValue;
+		Ar << Input.Constant;
+		return true;
+	}
+	else
+	{
+		return false;
+	}
+}
+
+bool FExpressionInput::Serialize(FArchive& Ar)
+{
+	return SerializeExpressionInput(Ar, *this);
+}
+
+bool FColorMaterialInput::Serialize(FArchive& Ar)
+{
+	return SerializeMaterialInput<FColor>(Ar, *this);
+}
+
+bool FScalarMaterialInput::Serialize(FArchive& Ar)
+{
+	return SerializeMaterialInput<float>(Ar, *this);
+}
 
+bool FVectorMaterialInput::Serialize(FArchive& Ar)
+{
+	return SerializeMaterialInput<FVector>(Ar, *this);
+}
+
+bool FVector2MaterialInput::Serialize(FArchive& Ar)
+{
+	return SerializeMaterialInput<FVector2D>(Ar, *this);
+}
+
+bool FMaterialAttributesInput::Serialize(FArchive& Ar)
+{
+	return SerializeExpressionInput(Ar, *this);
+}
+
+#if WITH_EDITOR
 int32 FColorMaterialInput::CompileWithDefault(class FMaterialCompiler* Compiler, EMaterialProperty Property)
 {
 	if (UseConstant)
@@ -206,6 +301,7 @@ int32 FMaterialAttributesInput::CompileWithDefault(class FMaterialCompiler* Comp
 
 	return Ret;
 }
+#endif  // WITH_EDITOR
 
 EMaterialValueType GetMaterialPropertyType(EMaterialProperty Property)
 {
@@ -687,7 +783,7 @@ bool IsTranslucentBlendMode(EBlendMode BlendMode)
 }
 
 int32 FMaterialResource::GetMaterialDomain() const { return Material->MaterialDomain; }
-bool FMaterialResource::IsTangentSpaceNormal() const { return Material->bTangentSpaceNormal || (Material->Normal.Expression == nullptr && !Material->bUseMaterialAttributes); }
+bool FMaterialResource::IsTangentSpaceNormal() const { return Material->bTangentSpaceNormal || (!Material->Normal.IsConnected() && !Material->bUseMaterialAttributes); }
 bool FMaterialResource::ShouldInjectEmissiveIntoLPV() const { return Material->bUseEmissiveForDynamicAreaLighting; }
 bool FMaterialResource::ShouldBlockGI() const { return Material->bBlockGI; }
 bool FMaterialResource::ShouldGenerateSphericalParticleNormals() const { return Material->bGenerateSphericalParticleNormals; }
@@ -700,9 +796,9 @@ bool FMaterialResource::IsLightFunction() const { return Material->MaterialDomai
 bool FMaterialResource::IsUsedWithEditorCompositing() const { return Material->bUsedWithEditorCompositing; }
 bool FMaterialResource::IsUsedWithDeferredDecal() const { return Material->MaterialDomain == MD_DeferredDecal; }
 bool FMaterialResource::IsSpecialEngineMaterial() const { return Material->bUsedAsSpecialEngineMaterial; }
-bool FMaterialResource::HasVertexPositionOffsetConnected() const { return !Material->bUseMaterialAttributes && Material->WorldPositionOffset.Expression != nullptr; }
-bool FMaterialResource::HasPixelDepthOffsetConnected() const { return !Material->bUseMaterialAttributes && Material->PixelDepthOffset.Expression != nullptr; }
-bool FMaterialResource::HasMaterialAttributesConnected() const { return Material->bUseMaterialAttributes && Material->MaterialAttributes.Expression != nullptr; }
+bool FMaterialResource::HasVertexPositionOffsetConnected() const { return !Material->bUseMaterialAttributes && Material->WorldPositionOffset.IsConnected(); }
+bool FMaterialResource::HasPixelDepthOffsetConnected() const { return !Material->bUseMaterialAttributes && Material->PixelDepthOffset.IsConnected(); }
+bool FMaterialResource::HasMaterialAttributesConnected() const { return Material->bUseMaterialAttributes && Material->MaterialAttributes.IsConnected(); }
 FString FMaterialResource::GetBaseMaterialPathName() const { return Material->GetPathName(); }
 
 bool FMaterialResource::IsUsedWithSkeletalMesh() const
@@ -2456,6 +2552,7 @@ bool UMaterialInterface::IsPropertyActive(EMaterialProperty InProperty)const
 	return false;
 }
 
+#if WITH_EDITOR
 int32 UMaterialInterface::CompilePropertyEx( class FMaterialCompiler* Compiler, EMaterialProperty Property )
 {
 	return INDEX_NONE;
@@ -2472,6 +2569,7 @@ int32 UMaterialInterface::CompileProperty(FMaterialCompiler* Compiler, EMaterial
 		return GetDefaultExpressionForMaterialProperty(Compiler, Property);
 	}
 }
+#endif // WITH_EDITOR
 
 void UMaterialInterface::AnalyzeMaterialProperty(EMaterialProperty InProperty, int32& OutNumTextureCoordinates, bool& bOutRequiresVertexData)
 {
@@ -2518,6 +2616,7 @@ void UMaterialInterface::AnalyzeMaterialProperty(EMaterialProperty InProperty, i
 #endif
 }
 
+#if WITH_EDITORONLY_DATA
 //Reorder the output index for any FExpressionInput connected to a UMaterialExpressionBreakMaterialAttributes.
 //If the order of pins in the material results or the make/break attributes nodes changes 
 //then the OutputIndex stored in any FExpressionInput coming from UMaterialExpressionBreakMaterialAttributes will be wrong and needs reordering.
@@ -2552,7 +2651,7 @@ void DoMaterialAttributeReorder(FExpressionInput* Input, int32 UE4Ver)
 		}
 	}
 }
-
+#endif // WITH_EDITORONLY_DATA
 //////////////////////////////////////////////////////////////////////////
 
 FMaterialInstanceBasePropertyOverrides::FMaterialInstanceBasePropertyOverrides()
diff --git a/Engine/Source/Runtime/Engine/Public/MaterialExpressionIO.h b/Engine/Source/Runtime/Engine/Public/MaterialExpressionIO.h
--- a/Engine/Source/Runtime/Engine/Public/MaterialExpressionIO.h
+++ b/Engine/Source/Runtime/Engine/Public/MaterialExpressionIO.h
@@ -9,8 +9,10 @@
 //@warning: FExpressionInput is mirrored in MaterialExpression.h and manually "subclassed" in Material.h (FMaterialInput)
 struct FExpressionInput
 {
+#if WITH_EDITORONLY_DATA
 	/** Material expression that this input is connected to, or NULL if not connected. */
 	class UMaterialExpression*	Expression;
+#endif
 
 	/** Index into Expression's outputs array that this input is connected to. */
 	int32						OutputIndex;
@@ -26,31 +28,58 @@ struct FExpressionInput
 								MaskG,
 								MaskB,
 								MaskA;
-	uint32						GCC64Padding; // @todo 64: if the C++ didn't mismirror this structure, we might not need this
+
+	/** Material expression name that this input is connected to, or None if not connected. Used only in cooked builds */
+	FName				ExpressionName;
 
 	FExpressionInput()
-		: Expression(NULL)
-		, OutputIndex(0)
+		: OutputIndex(0)
 		, Mask(0)
 		, MaskR(0)
 		, MaskG(0)
 		, MaskB(0)
 		, MaskA(0)
-		, GCC64Padding(0)
 	{
+#if WITH_EDITORONLY_DATA
+		Expression = nullptr;
+#endif
 	}
 
+#if WITH_EDITOR
 	ENGINE_API int32 Compile(class FMaterialCompiler* Compiler, int32 MultiplexIndex=INDEX_NONE);
+#endif // WITH_EDITOR
 
 	/**
 	 * Tests if the input has a material expression connected to it
 	 *
 	 * @return	true if an expression is connected, otherwise false
 	 */
-	bool IsConnected() const { return (NULL != Expression); }
+	bool IsConnected() const 
+	{ 
+#if WITH_EDITORONLY_DATA
+		return (nullptr != Expression);
+#else
+		return ExpressionName != NAME_None;
+#endif // WITH_EDITORONLY_DATA
+	}
 
+#if WITH_EDITOR
 	/** Connects output of InExpression to this input */
 	ENGINE_API void Connect( int32 InOutputIndex, class UMaterialExpression* InExpression );
+#endif // WITH_EDITOR
+
+	/** ICPPStructOps interface */
+	ENGINE_API bool Serialize(FArchive& Ar);
+};
+
+template<>
+struct TStructOpsTypeTraits<FExpressionInput>
+	: public TStructOpsTypeTraitsBase
+{
+	enum
+	{
+		WithSerializer = true,
+	};
 };
 
 //
@@ -96,35 +125,112 @@ template<class InputType> struct FMaterialInput : FExpressionInput
 
 struct FColorMaterialInput : FMaterialInput<FColor>
 {
+#if WITH_EDITOR
 	ENGINE_API int32 CompileWithDefault(class FMaterialCompiler* Compiler, EMaterialProperty Property);
+#endif  // WITH_EDITOR
+	/** ICPPStructOps interface */
+	ENGINE_API bool Serialize(FArchive& Ar);
 };
+
+template<>
+struct TStructOpsTypeTraits<FColorMaterialInput>
+	: public TStructOpsTypeTraitsBase
+{
+	enum
+	{
+		WithSerializer = true,
+	};
+};
+
 struct FScalarMaterialInput : FMaterialInput<float>
 {
+#if WITH_EDITOR
 	ENGINE_API int32 CompileWithDefault(class FMaterialCompiler* Compiler, EMaterialProperty Property);
+#endif  // WITH_EDITOR
+	/** ICPPStructOps interface */
+	ENGINE_API bool Serialize(FArchive& Ar);
+};
+
+template<>
+struct TStructOpsTypeTraits<FScalarMaterialInput>
+	: public TStructOpsTypeTraitsBase
+{
+	enum
+	{
+		WithSerializer = true,
+	};
 };
+
 struct FVectorMaterialInput : FMaterialInput<FVector>
 {
+#if WITH_EDITOR
 	ENGINE_API int32 CompileWithDefault(class FMaterialCompiler* Compiler, EMaterialProperty Property);
+#endif  // WITH_EDITOR
+	/** ICPPStructOps interface */
+	ENGINE_API bool Serialize(FArchive& Ar);
+};
+
+template<>
+struct TStructOpsTypeTraits<FVectorMaterialInput>
+	: public TStructOpsTypeTraitsBase
+{
+	enum
+	{
+		WithSerializer = true,
+	};
 };
+
 struct FVector2MaterialInput : FMaterialInput<FVector2D>
 {
+#if WITH_EDITOR
 	ENGINE_API int32 CompileWithDefault(class FMaterialCompiler* Compiler, EMaterialProperty Property);
+#endif  // WITH_EDITOR
+	/** ICPPStructOps interface */
+	ENGINE_API bool Serialize(FArchive& Ar);
+};
+
+template<>
+struct TStructOpsTypeTraits<FVector2MaterialInput>
+	: public TStructOpsTypeTraitsBase
+{
+	enum
+	{
+		WithSerializer = true,
+	};
 };
-struct FMaterialAttributesInput : FMaterialInput<int32>
+
+struct FMaterialAttributesInput : FExpressionInput
 {
-	FMaterialAttributesInput() : FMaterialInput<int32>(), PropertyConnectedBitmask(0)
+	FMaterialAttributesInput() 
+	: PropertyConnectedBitmask(0)
 	{ 
 		// ensure PropertyConnectedBitmask can contain all properties.
 		static_assert((uint32)(MP_MAX)-1 <= (8 * sizeof(PropertyConnectedBitmask)), "PropertyConnectedBitmask cannot contain entire EMaterialProperty enumeration.");
 	}
 
+#if WITH_EDITOR
 	ENGINE_API int32 CompileWithDefault(class FMaterialCompiler* Compiler, EMaterialProperty Property);
+#endif  // WITH_EDITOR
 	ENGINE_API bool IsConnected(EMaterialProperty Property) { return ((PropertyConnectedBitmask >> (uint32)Property) & 0x1) != 0; }
+	ENGINE_API bool IsConnected() const { return FExpressionInput::IsConnected(); }
 	ENGINE_API void SetConnectedProperty(EMaterialProperty Property, bool bIsConnected) 
 	{
 		PropertyConnectedBitmask = bIsConnected ? PropertyConnectedBitmask | (1 << (uint32)Property) : PropertyConnectedBitmask & ~(1 << (uint32)Property);
 	}
 
+	/** ICPPStructOps interface */
+	ENGINE_API bool Serialize(FArchive& Ar);
+
 	// each bit corresponds to EMaterialProperty connection status.
 	uint32 PropertyConnectedBitmask;
+};
+
+template<>
+struct TStructOpsTypeTraits<FMaterialAttributesInput>
+	: public TStructOpsTypeTraitsBase
+{
+	enum
+	{
+		WithSerializer = true,
+	};
 };
\ No newline at end of file
diff --git a/Engine/Source/Runtime/Engine/Public/MaterialShared.h b/Engine/Source/Runtime/Engine/Public/MaterialShared.h
--- a/Engine/Source/Runtime/Engine/Public/MaterialShared.h
+++ b/Engine/Source/Runtime/Engine/Public/MaterialShared.h
@@ -1711,5 +1711,7 @@ ENGINE_API int32 GetDefaultExpressionForMaterialProperty(FMaterialCompiler* Comp
  */
 ENGINE_API FString GetNameOfMaterialProperty(EMaterialProperty Property);
 
+#if WITH_EDITORONLY_DATA
 /** TODO - This can be removed whenever VER_UE4_MATERIAL_ATTRIBUTES_REORDERING is no longer relevant. */
 ENGINE_API void DoMaterialAttributeReorder(FExpressionInput* Input, int32 UE4Ver);
+#endif // WITH_EDITORONLY_DATA
diff --git a/Engine/Source/Runtime/Landscape/Classes/Materials/MaterialExpressionLandscapeGrassOutput.h b/Engine/Source/Runtime/Landscape/Classes/Materials/MaterialExpressionLandscapeGrassOutput.h
--- a/Engine/Source/Runtime/Landscape/Classes/Materials/MaterialExpressionLandscapeGrassOutput.h
+++ b/Engine/Source/Runtime/Landscape/Classes/Materials/MaterialExpressionLandscapeGrassOutput.h
@@ -35,13 +35,17 @@ UCLASS(collapsecategories, hidecategories=Object, MinimalAPI)
 class UMaterialExpressionLandscapeGrassOutput : public UMaterialExpressionCustomOutput
 {
 	GENERATED_UCLASS_BODY()
+
+#if WITH_EDITOR
 	virtual int32 Compile(class FMaterialCompiler* Compiler, int32 OutputIndex, int32 MultiplexIndex) override;
 	virtual void GetCaption(TArray<FString>& OutCaptions) const override;
+#endif
 	virtual const TArray<FExpressionInput*> GetInputs() override;
 	virtual FExpressionInput* GetInput(int32 InputIndex) override;
 	virtual FString GetInputName(int32 InputIndex) const override;
 
 	//~ Begin UObject Interface
+	virtual bool NeedsLoadForClient() const override;
 #if WITH_EDITOR
 	virtual void PostEditChangeProperty(FPropertyChangedEvent& PropertyChangedEvent) override;
 	//~ End UObject Interface
diff --git a/Engine/Source/Runtime/Landscape/Classes/Materials/MaterialExpressionLandscapeLayerBlend.h b/Engine/Source/Runtime/Landscape/Classes/Materials/MaterialExpressionLandscapeLayerBlend.h
--- a/Engine/Source/Runtime/Landscape/Classes/Materials/MaterialExpressionLandscapeLayerBlend.h
+++ b/Engine/Source/Runtime/Landscape/Classes/Materials/MaterialExpressionLandscapeLayerBlend.h
@@ -65,7 +65,8 @@ class UMaterialExpressionLandscapeLayerBlend : public UMaterialExpression
 
 
 	//~ Begin UObject Interface
-	void Serialize(FArchive& Ar) override;
+	virtual void Serialize(FArchive& Ar) override;
+	virtual bool NeedsLoadForClient() const override;
 #if WITH_EDITOR
 	virtual void PostEditChangeProperty(FPropertyChangedEvent& PropertyChangedEvent) override;
 	//~ End UObject Interface
@@ -73,10 +74,10 @@ class UMaterialExpressionLandscapeLayerBlend : public UMaterialExpression
 	//~ Begin UMaterialExpression Interface
 	virtual uint32 GetInputType(int32 InputIndex) override { return MCT_Unknown; }
 	virtual uint32 GetOutputType(int32 InputIndex) override { return MCT_Unknown; }
-#endif
 	virtual bool IsResultMaterialAttributes(int32 OutputIndex) override;
 	virtual int32 Compile(class FMaterialCompiler* Compiler, int32 OutputIndex, int32 MultiplexIndex) override;
 	virtual void GetCaption(TArray<FString>& OutCaptions) const override;
+#endif
 	virtual const TArray<FExpressionInput*> GetInputs() override;
 	virtual FExpressionInput* GetInput(int32 InputIndex) override;
 	virtual FString GetInputName(int32 InputIndex) const override;
diff --git a/Engine/Source/Runtime/Landscape/Classes/Materials/MaterialExpressionLandscapeLayerCoords.h b/Engine/Source/Runtime/Landscape/Classes/Materials/MaterialExpressionLandscapeLayerCoords.h
--- a/Engine/Source/Runtime/Landscape/Classes/Materials/MaterialExpressionLandscapeLayerCoords.h
+++ b/Engine/Source/Runtime/Landscape/Classes/Materials/MaterialExpressionLandscapeLayerCoords.h
@@ -58,9 +58,15 @@ class UMaterialExpressionLandscapeLayerCoords : public UMaterialExpression
 	float MappingPanV;
 
 	//~ Begin UMaterialExpression Interface
+#if WITH_EDITOR
 	virtual int32 Compile(class FMaterialCompiler* Compiler, int32 OutputIndex, int32 MultiplexIndex) override;
 	virtual void GetCaption(TArray<FString>& OutCaptions) const override;
+#endif
 	//~ End UMaterialExpression Interface
+
+	//~ Begin UObject Interface
+	virtual bool NeedsLoadForClient() const override;
+	//~ End UObject Interface
 };
 
 
diff --git a/Engine/Source/Runtime/Landscape/Classes/Materials/MaterialExpressionLandscapeLayerSample.h b/Engine/Source/Runtime/Landscape/Classes/Materials/MaterialExpressionLandscapeLayerSample.h
--- a/Engine/Source/Runtime/Landscape/Classes/Materials/MaterialExpressionLandscapeLayerSample.h
+++ b/Engine/Source/Runtime/Landscape/Classes/Materials/MaterialExpressionLandscapeLayerSample.h
@@ -24,8 +24,10 @@ class UMaterialExpressionLandscapeLayerSample : public UMaterialExpression
 public:
 
 	//~ Begin UMaterialExpression Interface
+#if WITH_EDITOR
 	virtual int32 Compile(class FMaterialCompiler* Compiler, int32 OutputIndex, int32 MultiplexIndex) override;
 	virtual void GetCaption(TArray<FString>& OutCaptions) const override;
+#endif
 	virtual UTexture* GetReferencedTexture() override;
 	//~ End UMaterialExpression Interface
 
@@ -35,6 +37,10 @@ public:
 	 * Called to get list of parameter names for static parameter sets
 	 */
 	void GetAllParameterNames(TArray<FName> &OutParameterNames, TArray<FGuid> &OutParameterIds) const;
+
+	//~ Begin UObject Interface
+	virtual bool NeedsLoadForClient() const override;
+	//~ End UObject Interface
 };
 
 
diff --git a/Engine/Source/Runtime/Landscape/Classes/Materials/MaterialExpressionLandscapeLayerSwitch.h b/Engine/Source/Runtime/Landscape/Classes/Materials/MaterialExpressionLandscapeLayerSwitch.h
--- a/Engine/Source/Runtime/Landscape/Classes/Materials/MaterialExpressionLandscapeLayerSwitch.h
+++ b/Engine/Source/Runtime/Landscape/Classes/Materials/MaterialExpressionLandscapeLayerSwitch.h
@@ -34,9 +34,11 @@ public:
 	//~ End UObject Interface
 
 	//~ Begin UMaterialExpression Interface
+#if WITH_EDITOR
 	virtual bool IsResultMaterialAttributes(int32 OutputIndex) override;
 	virtual int32 Compile(class FMaterialCompiler* Compiler, int32 OutputIndex, int32 MultiplexIndex) override;
 	virtual void GetCaption(TArray<FString>& OutCaptions) const override;
+#endif
 	virtual UTexture* GetReferencedTexture() override;
 #if WITH_EDITOR
 	virtual uint32 GetInputType(int32 InputIndex) override {return MCT_Unknown;}
@@ -48,11 +50,14 @@ public:
 
 	void GetAllParameterNames(TArray<FName> &OutParameterNames, TArray<FGuid> &OutParameterIds) const;
 
+	//~ Begin UObject Interface
 	/**
 	 * Do any object-specific cleanup required immediately after loading an object,
 	 * and immediately after any undo/redo.
 	 */
-	void PostLoad() override;
+	virtual void PostLoad() override;
+	virtual bool NeedsLoadForClient() const override;
+	//~ End UObject Interface
 };
 
 
diff --git a/Engine/Source/Runtime/Landscape/Classes/Materials/MaterialExpressionLandscapeLayerWeight.h b/Engine/Source/Runtime/Landscape/Classes/Materials/MaterialExpressionLandscapeLayerWeight.h
--- a/Engine/Source/Runtime/Landscape/Classes/Materials/MaterialExpressionLandscapeLayerWeight.h
+++ b/Engine/Source/Runtime/Landscape/Classes/Materials/MaterialExpressionLandscapeLayerWeight.h
@@ -33,9 +33,11 @@ class UMaterialExpressionLandscapeLayerWeight : public UMaterialExpression
 public:
 
 	//~ Begin UMaterialExpression Interface
+#if WITH_EDITOR
 	virtual bool IsResultMaterialAttributes(int32 OutputIndex) override;
 	virtual int32 Compile(class FMaterialCompiler* Compiler, int32 OutputIndex, int32 MultiplexIndex) override;
 	virtual void GetCaption(TArray<FString>& OutCaptions) const override;
+#endif
 	virtual UTexture* GetReferencedTexture() override;
 #if WITH_EDITOR
 	virtual uint32 GetInputType(int32 InputIndex) override {return MCT_Float | MCT_MaterialAttributes;}
@@ -49,11 +51,14 @@ public:
 	 */
 	void GetAllParameterNames(TArray<FName> &OutParameterNames, TArray<FGuid> &OutParameterIds) const;
 
+	//~ Begin UObject Interface
 	/**
 	 * Do any object-specific cleanup required immediately after loading an object,
 	 * and immediately after any undo/redo.
 	 */
-	void PostLoad() override;
+	virtual void PostLoad() override;
+	virtual bool NeedsLoadForClient() const override;
+	//~ End UObject Interface
 };
 
 
diff --git a/Engine/Source/Runtime/Landscape/Private/LandscapeGrass.cpp b/Engine/Source/Runtime/Landscape/Private/LandscapeGrass.cpp
--- a/Engine/Source/Runtime/Landscape/Private/LandscapeGrass.cpp
+++ b/Engine/Source/Runtime/Landscape/Private/LandscapeGrass.cpp
@@ -774,6 +774,7 @@ UMaterialExpressionLandscapeGrassOutput::UMaterialExpressionLandscapeGrassOutput
 	new(GrassTypes)FGrassInput(ConstructorStatics.NAME_Grass);
 }
 
+#if WITH_EDITOR
 int32 UMaterialExpressionLandscapeGrassOutput::Compile(class FMaterialCompiler* Compiler, int32 OutputIndex, int32 MultiplexIndex)
 {
 	if (GrassTypes.IsValidIndex(OutputIndex))
@@ -795,6 +796,7 @@ void UMaterialExpressionLandscapeGrassOutput::GetCaption(TArray<FString>& OutCap
 {
 	OutCaptions.Add(TEXT("Grass"));
 }
+#endif // WITH_EDITOR
 
 const TArray<FExpressionInput*> UMaterialExpressionLandscapeGrassOutput::GetInputs()
 {
@@ -816,6 +818,11 @@ FString UMaterialExpressionLandscapeGrassOutput::GetInputName(int32 InputIndex)
 	return GrassTypes[InputIndex].Name.ToString();
 }
 
+bool UMaterialExpressionLandscapeGrassOutput::NeedsLoadForClient() const
+{
+	return true;
+}
+
 #if WITH_EDITOR
 void UMaterialExpressionLandscapeGrassOutput::PostEditChangeProperty(FPropertyChangedEvent& PropertyChangedEvent)
 {
diff --git a/Engine/Source/Runtime/Landscape/Private/LandscapeRender.cpp b/Engine/Source/Runtime/Landscape/Private/LandscapeRender.cpp
--- a/Engine/Source/Runtime/Landscape/Private/LandscapeRender.cpp
+++ b/Engine/Source/Runtime/Landscape/Private/LandscapeRender.cpp
@@ -2661,13 +2661,23 @@ void ULandscapeComponent::GetStreamingTextureInfo(TArray<FStreamingTexturePrimit
 				UMaterialExpressionTextureSample* TextureSample = Cast<UMaterialExpressionTextureSample>(Expression);
 
 				// TODO: This is only works for direct Coordinate Texture Sample cases
-				if (TextureSample && TextureSample->Coordinates.Expression)
+				if (TextureSample && TextureSample->Coordinates.IsConnected())
 				{
-					UMaterialExpressionTextureCoordinate* TextureCoordinate =
-						Cast<UMaterialExpressionTextureCoordinate>(TextureSample->Coordinates.Expression);
-
-					UMaterialExpressionLandscapeLayerCoords* TerrainTextureCoordinate =
-						Cast<UMaterialExpressionLandscapeLayerCoords>(TextureSample->Coordinates.Expression);
+					UMaterialExpressionTextureCoordinate* TextureCoordinate = nullptr;
+					UMaterialExpressionLandscapeLayerCoords* TerrainTextureCoordinate = nullptr;
+		
+					for (UMaterialExpression* FindExp : Material->Expressions)
+					{
+						if (FindExp && FindExp->GetFName() == TextureSample->Coordinates.ExpressionName)
+						{
+							TextureCoordinate = Cast<UMaterialExpressionTextureCoordinate>(FindExp);
+							if (!TextureCoordinate)
+							{
+								TerrainTextureCoordinate = Cast<UMaterialExpressionLandscapeLayerCoords>(FindExp);
+							}
+							break;
+						}
+					}
 
 					if (TextureCoordinate || TerrainTextureCoordinate)
 					{
diff --git a/Engine/Source/Runtime/Landscape/Private/Materials/MaterialExpressionLandscapeLayerBlend.cpp b/Engine/Source/Runtime/Landscape/Private/Materials/MaterialExpressionLandscapeLayerBlend.cpp
--- a/Engine/Source/Runtime/Landscape/Private/Materials/MaterialExpressionLandscapeLayerBlend.cpp
+++ b/Engine/Source/Runtime/Landscape/Private/Materials/MaterialExpressionLandscapeLayerBlend.cpp
@@ -113,7 +113,7 @@ FString UMaterialExpressionLandscapeLayerBlend::GetInputName(int32 InputIndex) c
 	return TEXT("");
 }
 
-
+#if WITH_EDITOR
 bool UMaterialExpressionLandscapeLayerBlend::IsResultMaterialAttributes(int32 OutputIndex)
 {
 	if (ContainsInputLoop())
@@ -131,7 +131,6 @@ bool UMaterialExpressionLandscapeLayerBlend::IsResultMaterialAttributes(int32 Ou
 	return false;
 }
 
-
 int32 UMaterialExpressionLandscapeLayerBlend::Compile(class FMaterialCompiler* Compiler, int32 OutputIndex, int32 MultiplexIndex)
 {
 	// For renormalization
@@ -260,22 +259,19 @@ int32 UMaterialExpressionLandscapeLayerBlend::Compile(class FMaterialCompiler* C
 
 	return OutputCode;
 }
-
+#endif // WITH_EDITOR
 
 UTexture* UMaterialExpressionLandscapeLayerBlend::GetReferencedTexture()
 {
 	return GEngine->WeightMapPlaceholderTexture;
 }
 
-
+#if WITH_EDITOR
 void UMaterialExpressionLandscapeLayerBlend::GetCaption(TArray<FString>& OutCaptions) const
 {
 	OutCaptions.Add(FString(TEXT("Layer Blend")));
 }
 
-
-#if WITH_EDITOR
-
 void UMaterialExpressionLandscapeLayerBlend::PostEditChangeProperty(FPropertyChangedEvent& PropertyChangedEvent)
 {
 	Super::PostEditChangeProperty(PropertyChangedEvent);
@@ -301,7 +297,7 @@ void UMaterialExpressionLandscapeLayerBlend::PostEditChangeProperty(FPropertyCha
 		}
 	}
 }
-#endif
+#endif // WITH_EDITOR
 
 
 void UMaterialExpressionLandscapeLayerBlend::GetAllParameterNames(TArray<FName> &OutParameterNames, TArray<FGuid> &OutParameterIds) const
@@ -320,5 +316,10 @@ void UMaterialExpressionLandscapeLayerBlend::GetAllParameterNames(TArray<FName>
 	}
 }
 
+bool UMaterialExpressionLandscapeLayerBlend::NeedsLoadForClient() const
+{
+	return true;
+}
+
 
 #undef LOCTEXT_NAMESPACE
diff --git a/Engine/Source/Runtime/Landscape/Private/Materials/MaterialExpressionLandscapeLayerCoords.cpp b/Engine/Source/Runtime/Landscape/Private/Materials/MaterialExpressionLandscapeLayerCoords.cpp
--- a/Engine/Source/Runtime/Landscape/Private/Materials/MaterialExpressionLandscapeLayerCoords.cpp
+++ b/Engine/Source/Runtime/Landscape/Private/Materials/MaterialExpressionLandscapeLayerCoords.cpp
@@ -30,7 +30,7 @@ UMaterialExpressionLandscapeLayerCoords::UMaterialExpressionLandscapeLayerCoords
 	bCollapsed = false;
 }
 
-
+#if WITH_EDITOR
 int32 UMaterialExpressionLandscapeLayerCoords::Compile(class FMaterialCompiler* Compiler, int32 OutputIndex, int32 MultiplexIndex)
 {
 	switch (CustomUVType)
@@ -100,6 +100,11 @@ void UMaterialExpressionLandscapeLayerCoords::GetCaption(TArray<FString>& OutCap
 {
 	OutCaptions.Add(FString(TEXT("LandscapeCoords")));
 }
+#endif // WITH_EDITOR
 
+bool UMaterialExpressionLandscapeLayerCoords::NeedsLoadForClient() const
+{
+	return true;
+}
 
 #undef LOCTEXT_NAMESPACE
diff --git a/Engine/Source/Runtime/Landscape/Private/Materials/MaterialExpressionLandscapeLayerSample.cpp b/Engine/Source/Runtime/Landscape/Private/Materials/MaterialExpressionLandscapeLayerSample.cpp
--- a/Engine/Source/Runtime/Landscape/Private/Materials/MaterialExpressionLandscapeLayerSample.cpp
+++ b/Engine/Source/Runtime/Landscape/Private/Materials/MaterialExpressionLandscapeLayerSample.cpp
@@ -37,7 +37,7 @@ FGuid& UMaterialExpressionLandscapeLayerSample::GetParameterExpressionId()
 	return ExpressionGUID;
 }
 
-
+#if WITH_EDITOR
 int32 UMaterialExpressionLandscapeLayerSample::Compile(class FMaterialCompiler* Compiler, int32 OutputIndex, int32 MultiplexIndex)
 {
 	const int32 WeightCode = Compiler->StaticTerrainLayerWeight(ParameterName, Compiler->Constant(PreviewWeight));
@@ -51,19 +51,19 @@ int32 UMaterialExpressionLandscapeLayerSample::Compile(class FMaterialCompiler*
 		return WeightCode;
 	}
 }
-
+#endif // WITH_EDITOR
 
 UTexture* UMaterialExpressionLandscapeLayerSample::GetReferencedTexture()
 {
 	return GEngine->WeightMapPlaceholderTexture;
 }
 
-
+#if WITH_EDITOR
 void UMaterialExpressionLandscapeLayerSample::GetCaption(TArray<FString>& OutCaptions) const
 {
 	OutCaptions.Add(FString::Printf(TEXT("Sample '%s'"), *ParameterName.ToString()));
 }
-
+#endif // WITH_EDITOR
 
 void UMaterialExpressionLandscapeLayerSample::GetAllParameterNames(TArray<FName> &OutParameterNames, TArray<FGuid> &OutParameterIds) const
 {
@@ -74,5 +74,9 @@ void UMaterialExpressionLandscapeLayerSample::GetAllParameterNames(TArray<FName>
 	}
 }
 
+bool UMaterialExpressionLandscapeLayerSample::NeedsLoadForClient() const
+{
+	return ParameterName != NAME_None;
+}
 
 #undef LOCTEXT_NAMESPACE
diff --git a/Engine/Source/Runtime/Landscape/Private/Materials/MaterialExpressionLandscapeLayerSwitch.cpp b/Engine/Source/Runtime/Landscape/Private/Materials/MaterialExpressionLandscapeLayerSwitch.cpp
--- a/Engine/Source/Runtime/Landscape/Private/Materials/MaterialExpressionLandscapeLayerSwitch.cpp
+++ b/Engine/Source/Runtime/Landscape/Private/Materials/MaterialExpressionLandscapeLayerSwitch.cpp
@@ -33,7 +33,7 @@ UMaterialExpressionLandscapeLayerSwitch::UMaterialExpressionLandscapeLayerSwitch
 	bCollapsed = false;
 }
 
-
+#if WITH_EDITOR
 bool UMaterialExpressionLandscapeLayerSwitch::IsResultMaterialAttributes(int32 OutputIndex)
 {
 	if (ContainsInputLoop())
@@ -46,7 +46,6 @@ bool UMaterialExpressionLandscapeLayerSwitch::IsResultMaterialAttributes(int32 O
 	return bLayerUsedIsMaterialAttributes || bLayerNotUsedIsMaterialAttributes;
 }
 
-
 int32 UMaterialExpressionLandscapeLayerSwitch::Compile(class FMaterialCompiler* Compiler, int32 OutputIndex, int32 MultiplexIndex)
 {
 	const int32 WeightCode = Compiler->StaticTerrainLayerWeight(
@@ -73,20 +72,20 @@ int32 UMaterialExpressionLandscapeLayerSwitch::Compile(class FMaterialCompiler*
 
 	return ReturnCode;
 }
-
+#endif // WITH_EDITOR
 
 UTexture* UMaterialExpressionLandscapeLayerSwitch::GetReferencedTexture()
 {
 	return GEngine->WeightMapPlaceholderTexture;
 }
 
-
+#if WITH_EDITOR
 void UMaterialExpressionLandscapeLayerSwitch::GetCaption(TArray<FString>& OutCaptions) const
 {
 	OutCaptions.Add(TEXT("Layer Switch"));
 	OutCaptions.Add(FString::Printf(TEXT("'%s'"), *ParameterName.ToString()));
 }
-
+#endif // WITH_EDITOR
 
 void UMaterialExpressionLandscapeLayerSwitch::Serialize(FArchive& Ar)
 {
@@ -127,5 +126,9 @@ void UMaterialExpressionLandscapeLayerSwitch::GetAllParameterNames(TArray<FName>
 	}
 }
 
+bool UMaterialExpressionLandscapeLayerSwitch::NeedsLoadForClient() const
+{
+	return ParameterName != NAME_None;
+}
 
 #undef LOCTEXT_NAMESPACE
diff --git a/Engine/Source/Runtime/Landscape/Private/Materials/MaterialExpressionLandscapeLayerWeight.cpp b/Engine/Source/Runtime/Landscape/Private/Materials/MaterialExpressionLandscapeLayerWeight.cpp
--- a/Engine/Source/Runtime/Landscape/Private/Materials/MaterialExpressionLandscapeLayerWeight.cpp
+++ b/Engine/Source/Runtime/Landscape/Private/Materials/MaterialExpressionLandscapeLayerWeight.cpp
@@ -50,7 +50,7 @@ void UMaterialExpressionLandscapeLayerWeight::PostLoad()
 	}
 }
 
-
+#if WITH_EDITOR
 bool UMaterialExpressionLandscapeLayerWeight::IsResultMaterialAttributes(int32 OutputIndex)
 {
 	if (ContainsInputLoop())
@@ -63,7 +63,6 @@ bool UMaterialExpressionLandscapeLayerWeight::IsResultMaterialAttributes(int32 O
 	return bLayerIsMaterialAttributes || bBaseIsMaterialAttributes;
 }
 
-
 int32 UMaterialExpressionLandscapeLayerWeight::Compile(class FMaterialCompiler* Compiler, int32 OutputIndex, int32 MultiplexIndex)
 {
 	const int32 BaseCode = Base.Expression ? Base.Compile(Compiler, MultiplexIndex) : Compiler->Constant3(ConstBase.X, ConstBase.Y, ConstBase.Z);
@@ -89,19 +88,19 @@ int32 UMaterialExpressionLandscapeLayerWeight::Compile(class FMaterialCompiler*
 
 	return ReturnCode;
 }
-
+#endif // WITH_EDITOR
 
 UTexture* UMaterialExpressionLandscapeLayerWeight::GetReferencedTexture()
 {
 	return GEngine->WeightMapPlaceholderTexture;
 }
 
-
+#if WITH_EDITOR
 void UMaterialExpressionLandscapeLayerWeight::GetCaption(TArray<FString>& OutCaptions) const
 {
 	OutCaptions.Add(FString::Printf(TEXT("Layer '%s'"), *ParameterName.ToString()));
 }
-
+#endif // WITH_EDITOR
 
 void UMaterialExpressionLandscapeLayerWeight::GetAllParameterNames(TArray<FName> &OutParameterNames, TArray<FGuid> &OutParameterIds) const
 {
@@ -114,5 +113,9 @@ void UMaterialExpressionLandscapeLayerWeight::GetAllParameterNames(TArray<FName>
 	}
 }
 
+bool UMaterialExpressionLandscapeLayerWeight::NeedsLoadForClient() const
+{
+	return ParameterName != NAME_None;
+}
 
 #undef LOCTEXT_NAMESPACE
diff --git a/Engine/Source/Runtime/Landscape/Private/Materials/MaterialExpressionLandscapeVisibilityMask.cpp b/Engine/Source/Runtime/Landscape/Private/Materials/MaterialExpressionLandscapeVisibilityMask.cpp
--- a/Engine/Source/Runtime/Landscape/Private/Materials/MaterialExpressionLandscapeVisibilityMask.cpp
+++ b/Engine/Source/Runtime/Landscape/Private/Materials/MaterialExpressionLandscapeVisibilityMask.cpp
@@ -38,11 +38,13 @@ FGuid& UMaterialExpressionLandscapeVisibilityMask::GetParameterExpressionId()
 	return ExpressionGUID;
 }
 
+#if WITH_EDITOR
 int32 UMaterialExpressionLandscapeVisibilityMask::Compile(class FMaterialCompiler* Compiler, int32 OutputIndex, int32 MultiplexIndex)
 {
 	int32 MaskLayerCode = Compiler->StaticTerrainLayerWeight(ParameterName, Compiler->Constant(0.f));
 	return MaskLayerCode == INDEX_NONE ? Compiler->Constant(1.f) : Compiler->Sub(Compiler->Constant(1.f), MaskLayerCode);
 }
+#endif // WITH_EDITOR
 
 UTexture* UMaterialExpressionLandscapeVisibilityMask::GetReferencedTexture()
 {
@@ -60,9 +62,16 @@ void UMaterialExpressionLandscapeVisibilityMask::GetAllParameterNames(TArray<FNa
 	}
 }
 
+#if WITH_EDITOR
 void UMaterialExpressionLandscapeVisibilityMask::GetCaption(TArray<FString>& OutCaptions) const
 {
 	OutCaptions.Add(FString(TEXT("Landscape Visibility Mask")));
 }
+#endif // WITH_EDITOR
+
+bool UMaterialExpressionLandscapeVisibilityMask::NeedsLoadForClient() const
+{
+	return ParameterName != NAME_None;
+}
 
 #undef LOCTEXT_NAMESPACE
