Copying //UE4/Dev-Framework to Dev-Main (//UE4/Dev-Main) @ 2879625

diff --git a/Engine/Source/Runtime/Engine/Classes/Engine/EngineTypes.h b/Engine/Source/Runtime/Engine/Classes/Engine/EngineTypes.h
--- a/Engine/Source/Runtime/Engine/Classes/Engine/EngineTypes.h
+++ b/Engine/Source/Runtime/Engine/Classes/Engine/EngineTypes.h
@@ -898,7 +898,7 @@ enum ETimelineSigType
 
 
 /** Enum used to describe what type of collision is enabled on a body. */
-UENUM()
+UENUM(BlueprintType)
 namespace ECollisionEnabled 
 { 
 	enum Type 
@@ -1613,7 +1613,7 @@ struct FPrimitiveMaterialRef
 /**
  * Structure containing information about one hit of a trace, such as point of impact and surface normal at that point.
  */
-USTRUCT(BlueprintType, meta=(HasNativeBreak="Engine.GameplayStatics.BreakHitResult"))
+USTRUCT(BlueprintType, meta = (HasNativeBreak = "Engine.GameplayStatics.BreakHitResult", HasNativeMake = "Engine.GameplayStatics.MakeHitResult"))
 struct ENGINE_API FHitResult
 {
 	GENERATED_USTRUCT_BODY()
diff --git a/Engine/Source/Runtime/Engine/Classes/Engine/LODActor.h b/Engine/Source/Runtime/Engine/Classes/Engine/LODActor.h
--- a/Engine/Source/Runtime/Engine/Classes/Engine/LODActor.h
+++ b/Engine/Source/Runtime/Engine/Classes/Engine/LODActor.h
@@ -101,6 +101,9 @@ public:
 	/** Cleans the SubActor array (removes all NULL entries) */
 	void CleanSubActorArray();
 
+	/** Cleans the SubObject array (removes all NULL entries) */
+	void CleanSubObjectsArray();
+
 	/** Recalculates the drawing distance according to a fixed FOV of 90 and the transition screen size*/
 	void RecalculateDrawingDistance(const float TransitionScreenSize);
 #endif // WITH_EDITOR
@@ -109,6 +112,7 @@ public:
 	virtual FString GetDetailedInfoInternal() const override;
 	virtual FBox GetComponentsBoundingBox(bool bNonColliding = false) const override;	
 #if WITH_EDITOR
+	virtual void PreEditChange(UProperty* PropertyThatWillChange) override;
 	virtual void PostEditChangeProperty(FPropertyChangedEvent& PropertyChangedEvent) override;
 	virtual void EditorApplyTranslation(const FVector& DeltaTranslation, bool bAltDown, bool bShiftDown, bool bCtrlDown) override;
 	virtual void EditorApplyRotation(const FRotator& DeltaRotation, bool bAltDown, bool bShiftDown, bool bCtrlDown) override;
diff --git a/Engine/Source/Runtime/Engine/Classes/Engine/Level.h b/Engine/Source/Runtime/Engine/Classes/Engine/Level.h
--- a/Engine/Source/Runtime/Engine/Classes/Engine/Level.h
+++ b/Engine/Source/Runtime/Engine/Classes/Engine/Level.h
@@ -8,6 +8,7 @@ class ALevelBounds;
 class UTexture2D;
 class UNavigationDataChunk;
 class AInstancedFoliageActor;
+class AWorldSettings;
 
 /**
  * Structure containing all information needed for determining the screen space
@@ -509,6 +510,9 @@ public:
 private:
 	FLevelBoundsActorUpdatedEvent LevelBoundsActorUpdatedEvent; 
 
+	UPROPERTY()
+	AWorldSettings* WorldSettings;
+
 protected:
 
 	/** Array of user data stored with the asset */
@@ -685,8 +689,9 @@ public:
 	 *
 	 * @return		The AWorldSettings for this level.
 	 */
-	ENGINE_API 
-	class AWorldSettings* GetWorldSettings() const;
+	ENGINE_API AWorldSettings* GetWorldSettings(bool bChecked = true) const;
+
+	ENGINE_API void SetWorldSettings(AWorldSettings* NewWorldSettings);
 
 	/**
 	 * Returns the level scripting actor associated with this level
diff --git a/Engine/Source/Runtime/Engine/Classes/Engine/StaticMesh.h b/Engine/Source/Runtime/Engine/Classes/Engine/StaticMesh.h
--- a/Engine/Source/Runtime/Engine/Classes/Engine/StaticMesh.h
+++ b/Engine/Source/Runtime/Engine/Classes/Engine/StaticMesh.h
@@ -352,7 +352,7 @@ class UStaticMesh : public UObject, public IInterface_CollisionDataProvider, pub
 	 *	Sometimes it can be desirable to use a lower poly representation for collision to reduce memory usage, improve performance and behaviour.
 	 *	Collision representation does not change based on distance to camera.
 	 */
-	UPROPERTY(EditAnywhere,  Category = StaticMesh)
+	UPROPERTY(EditAnywhere, Category = StaticMesh, meta=(DisplayName="LOD For Collision"))
 	int32 LODForCollision;
 
 	/** True if mesh should use a less-conservative method of mip LOD texture factor computation.
diff --git a/Engine/Source/Runtime/Engine/Classes/Engine/World.h b/Engine/Source/Runtime/Engine/Classes/Engine/World.h
--- a/Engine/Source/Runtime/Engine/Classes/Engine/World.h
+++ b/Engine/Source/Runtime/Engine/Classes/Engine/World.h
@@ -1044,21 +1044,6 @@ public:
 	
 	// LINE TRACE
 
-	/**
-	 *  Trace a ray against the world using a specific channel and return if a blocking hit is found.
-	 *  @param  Start           Start location of the ray
-	 *  @param  End             End location of the ray
-	 *  @param  TraceChannel    The 'channel' that this ray is in, used to determine which components to hit
-	 *  @param  Params          Additional parameters used for the trace
-	 * 	@param 	ResponseParam	ResponseContainer to be used for this trace
-	 *  @return TRUE if a blocking hit is found
-	 */
-	DEPRECATED(4.8, "Use LineTraceTestByChannel instead.")
-	bool LineTraceTest(const FVector& Start,const FVector& End,ECollisionChannel TraceChannel,const FCollisionQueryParams& Params, const FCollisionResponseParams& ResponseParam = FCollisionResponseParams::DefaultResponseParam) const
-	{
-		return LineTraceTestByChannel(Start, End, TraceChannel, Params, ResponseParam);
-	}
-
 	/**
 	 *  Trace a ray against the world using a specific channel and return if a blocking hit is found.
 	 *  @param  Start           Start location of the ray
@@ -1070,19 +1055,6 @@ public:
 	 */
 	bool LineTraceTestByChannel(const FVector& Start,const FVector& End,ECollisionChannel TraceChannel, const FCollisionQueryParams& Params = FCollisionQueryParams::DefaultQueryParam, const FCollisionResponseParams& ResponseParam = FCollisionResponseParams::DefaultResponseParam) const;
 	
-	/**
-	 *  Trace a ray against the world using object types and return if a blocking hit is found.
-	 *  @param  Start           	Start location of the ray
-	 *  @param  End             	End location of the ray
-	 *  @param  Params          	Additional parameters used for the trace
-	 *	@param	ObjectQueryParams	List of object types it's looking for
-	 *  @return TRUE if any hit is found
-	 */
-	DEPRECATED(4.8, "Use LineTraceTestByObjectType instead.")
-	bool LineTraceTest(const FVector& Start,const FVector& End,const FCollisionQueryParams& Params, const FCollisionObjectQueryParams& ObjectQueryParams) const
-	{
-		return LineTraceTestByObjectType(Start, End, ObjectQueryParams, Params);
-	}
 
 	/**
 	 *  Trace a ray against the world using object types and return if a blocking hit is found.
@@ -1104,21 +1076,6 @@ public:
 	 */
 	bool LineTraceTestByProfile(const FVector& Start, const FVector& End, FName ProfileName, const FCollisionQueryParams& Params = FCollisionQueryParams::DefaultQueryParam) const;
 
-	/**
-	 *  Trace a ray against the world using a specific channel and return the first blocking hit
-	 *  @param  OutHit          First blocking hit found
-	 *  @param  Start           Start location of the ray
-	 *  @param  End             End location of the ray
-	 *  @param  TraceChannel    The 'channel' that this ray is in, used to determine which components to hit
-	 *  @param  Params          Additional parameters used for the trace
-	 * 	@param 	ResponseParam	ResponseContainer to be used for this trace	 
-	 *  @return TRUE if a blocking hit is found
-	 */
-	DEPRECATED(4.8, "Use LineTraceSingleByChannel instead.")
-	bool LineTraceSingle(struct FHitResult& OutHit,const FVector& Start,const FVector& End,ECollisionChannel TraceChannel,const FCollisionQueryParams& Params, const FCollisionResponseParams& ResponseParam = FCollisionResponseParams::DefaultResponseParam) const
-	{
-		return LineTraceSingleByChannel(OutHit, Start, End, TraceChannel, Params, ResponseParam);
-	}
 	/**
 	 *  Trace a ray against the world using a specific channel and return the first blocking hit
 	 *  @param  OutHit          First blocking hit found
@@ -1131,21 +1088,6 @@ public:
 	 */
 	bool LineTraceSingleByChannel(struct FHitResult& OutHit,const FVector& Start,const FVector& End,ECollisionChannel TraceChannel,const FCollisionQueryParams& Params = FCollisionQueryParams::DefaultQueryParam, const FCollisionResponseParams& ResponseParam = FCollisionResponseParams::DefaultResponseParam) const;
 
-	/**
-	 *  Trace a ray against the world using object types and return the first blocking hit
-	 *  @param  OutHit          First blocking hit found
-	 *  @param  Start           Start location of the ray
-	 *  @param  End             End location of the ray
-	 *  @param  Params          Additional parameters used for the trace
-	 *	@param	ObjectQueryParams	List of object types it's looking for
-	 *  @return TRUE if any hit is found
-	 */
-	DEPRECATED(4.8, "Use LineTraceSingleByObjectType instead.")
-	bool LineTraceSingle(struct FHitResult& OutHit,const FVector& Start,const FVector& End,const FCollisionQueryParams& Params, const FCollisionObjectQueryParams& ObjectQueryParams) const
-	{
-		return LineTraceSingleByObjectType(OutHit, Start, End, ObjectQueryParams, Params);
-	}
-
 	/**
 	 *  Trace a ray against the world using object types and return the first blocking hit
 	 *  @param  OutHit          First blocking hit found
@@ -1168,24 +1110,6 @@ public:
 	 */
 	bool LineTraceSingleByProfile(struct FHitResult& OutHit, const FVector& Start, const FVector& End, FName ProfileName, const FCollisionQueryParams& Params = FCollisionQueryParams::DefaultQueryParam) const;
 
-	/**
-	 *  Trace a ray against the world using a specific channel and return overlapping hits and then first blocking hit
-	 *  Results are sorted, so a blocking hit (if found) will be the last element of the array
-	 *  Only the single closest blocking result will be generated, no tests will be done after that
-	 *  @param  OutHits         Array of hits found between ray and the world
-	 *  @param  Start           Start location of the ray
-	 *  @param  End             End location of the ray
-	 *  @param  TraceChannel    The 'channel' that this ray is in, used to determine which components to hit
-	 *  @param  Params          Additional parameters used for the trace
-	 * 	@param 	ResponseParam	ResponseContainer to be used for this trace	 
-	 *  @return TRUE if OutHits contains any blocking hit entries
-	 */
-	DEPRECATED(4.8, "Use LineTraceMultiByChannel instead.")
-	bool LineTraceMulti(TArray<struct FHitResult>& OutHits,const FVector& Start,const FVector& End,ECollisionChannel TraceChannel,const FCollisionQueryParams& Params, const FCollisionResponseParams& ResponseParam = FCollisionResponseParams::DefaultResponseParam) const
-	{
-		return LineTraceMultiByChannel(OutHits, Start, End, TraceChannel, Params, ResponseParam);
-	}
-
 	/**
 	 *  Trace a ray against the world using a specific channel and return overlapping hits and then first blocking hit
 	 *  Results are sorted, so a blocking hit (if found) will be the last element of the array
@@ -1201,23 +1125,6 @@ public:
 	bool LineTraceMultiByChannel(TArray<struct FHitResult>& OutHits,const FVector& Start,const FVector& End,ECollisionChannel TraceChannel,const FCollisionQueryParams& Params = FCollisionQueryParams::DefaultQueryParam, const FCollisionResponseParams& ResponseParam = FCollisionResponseParams::DefaultResponseParam) const;
 
 
-	/**
-	 *  Trace a ray against the world using object types and return overlapping hits and then first blocking hit
-	 *  Results are sorted, so a blocking hit (if found) will be the last element of the array
-	 *  Only the single closest blocking result will be generated, no tests will be done after that
-	 *  @param  OutHits         Array of hits found between ray and the world
-	 *  @param  Start           Start location of the ray
-	 *  @param  End             End location of the ray
-	 *  @param  Params          Additional parameters used for the trace
-	 *	@param	ObjectQueryParams	List of object types it's looking for
-	 *  @return TRUE if any hit is found
-	 */
-	DEPRECATED(4.8, "Use LineTraceMultiByObjectType instead.")
-	bool LineTraceMulti(TArray<struct FHitResult>& OutHits,const FVector& Start,const FVector& End,const FCollisionQueryParams& Params, const FCollisionObjectQueryParams& ObjectQueryParams) const
-	{
-		return LineTraceMultiByObjectType(OutHits, Start, End, ObjectQueryParams, Params);
-	}
-
 	/**
 	 *  Trace a ray against the world using object types and return overlapping hits and then first blocking hit
 	 *  Results are sorted, so a blocking hit (if found) will be the last element of the array
@@ -1245,22 +1152,6 @@ public:
 	 */
 	bool LineTraceMultiByProfile(TArray<struct FHitResult>& OutHits, const FVector& Start, const FVector& End, FName ProfileName, const FCollisionQueryParams& Params = FCollisionQueryParams::DefaultQueryParam) const;
 
-	/**
-	 *  Sweep a shape against the world using a specific channel and return if a blocking hit is found.
-	 *  @param  Start           Start location of the shape
-	 *  @param  End             End location of the shape
-	 *  @param  TraceChannel    The 'channel' that this trace uses, used to determine which components to hit
-	 *  @param	CollisionShape	CollisionShape - supports Box, Sphere, Capsule
-	 *  @param  Params          Additional parameters used for the trace
-	 * 	@param 	ResponseParam	ResponseContainer to be used for this trace	 	 
-	 *  @return TRUE if a blocking hit is found
-	 */
-	DEPRECATED(4.8, "Use SweepTestByChannel instead.")
-	bool SweepTest(const FVector& Start, const FVector& End, const FQuat& Rot, ECollisionChannel TraceChannel, const FCollisionShape& CollisionShape, const FCollisionQueryParams& Params, const FCollisionResponseParams& ResponseParam = FCollisionResponseParams::DefaultResponseParam) const
-	{
-		return SweepTestByChannel(Start, End, Rot, TraceChannel, CollisionShape, Params, ResponseParam);
-	}
-
 	/**
 	 *  Sweep a shape against the world using a specific channel and return if a blocking hit is found.
 	 *  @param  Start           Start location of the shape
@@ -1273,21 +1164,6 @@ public:
 	 */
 	bool SweepTestByChannel(const FVector& Start, const FVector& End, const FQuat& Rot, ECollisionChannel TraceChannel, const FCollisionShape& CollisionShape, const FCollisionQueryParams& Params = FCollisionQueryParams::DefaultQueryParam, const FCollisionResponseParams& ResponseParam = FCollisionResponseParams::DefaultResponseParam) const;
 
-	/**
-	 *  Sweep a shape against the world using object types and return if a blocking hit is found.
-	 *  @param  Start           Start location of the shape
-	 *  @param  End             End location of the shape
-	 *  @param	CollisionShape	CollisionShape - supports Box, Sphere, Capsule
-	 *  @param  Params          Additional parameters used for the trace
-	 *	@param	ObjectQueryParams	List of object types it's looking for
-	 *  @return TRUE if any hit is found
-	 */
-	DEPRECATED(4.8, "Use SweepTestByObjectType instead.")
-	bool SweepTest(const FVector& Start, const FVector& End, const FQuat& Rot, const FCollisionShape& CollisionShape, const FCollisionQueryParams& Params, const FCollisionObjectQueryParams& ObjectQueryParams) const
-	{
-		return SweepTestByObjectType(Start, End, Rot, ObjectQueryParams, CollisionShape, Params);
-	}
-
 	/**
 	 *  Sweep a shape against the world using object types and return if a blocking hit is found.
 	 *  @param  Start           Start location of the shape
@@ -1311,23 +1187,6 @@ public:
 	 */
 	bool SweepTestByProfile(const FVector& Start, const FVector& End, const FQuat& Rot, FName ProfileName, const FCollisionShape& CollisionShape, const FCollisionQueryParams& Params) const;
 
-	/**
-	 *  Sweep a shape against the world and return the first blocking hit using a specific channel
-	 *  @param  OutHit          First blocking hit found
-	 *  @param  Start           Start location of the shape
-	 *  @param  End             End location of the shape
-	 *  @param  TraceChannel    The 'channel' that this trace is in, used to determine which components to hit
-	 *  @param	CollisionShape	CollisionShape - supports Box, Sphere, Capsule
-	 *  @param  Params          Additional parameters used for the trace
-	 * 	@param 	ResponseParam	ResponseContainer to be used for this trace	 
-	 *  @return TRUE if OutHits contains any blocking hit entries
-	 */
-	DEPRECATED(4.8, "Use SweepSingleByChannel instead.")
-	bool SweepSingle(struct FHitResult& OutHit, const FVector& Start, const FVector& End, const FQuat& Rot, ECollisionChannel TraceChannel, const FCollisionShape& CollisionShape, const FCollisionQueryParams& Params, const FCollisionResponseParams& ResponseParam = FCollisionResponseParams::DefaultResponseParam) const
-	{
-		return SweepSingleByChannel(OutHit, Start, End, Rot, TraceChannel, CollisionShape, Params, ResponseParam);
-	}
-
 	/**
 	 *  Sweep a shape against the world and return the first blocking hit using a specific channel
 	 *  @param  OutHit          First blocking hit found
@@ -1341,22 +1200,6 @@ public:
 	 */
 	bool SweepSingleByChannel(struct FHitResult& OutHit, const FVector& Start, const FVector& End, const FQuat& Rot, ECollisionChannel TraceChannel, const FCollisionShape& CollisionShape, const FCollisionQueryParams& Params = FCollisionQueryParams::DefaultQueryParam, const FCollisionResponseParams& ResponseParam = FCollisionResponseParams::DefaultResponseParam) const;
 
-	/**
-	 *  Sweep a shape against the world and return the first blocking hit using object types
-	 *  @param  OutHit          First blocking hit found
-	 *  @param  Start           Start location of the shape
-	 *  @param  End             End location of the shape
-	 *  @param	CollisionShape	CollisionShape - supports Box, Sphere, Capsule
-	 *  @param  Params          Additional parameters used for the trace
-	 *	@param	ObjectQueryParams	List of object types it's looking for
-	 *  @return TRUE if any hit is found
-	 */
-	DEPRECATED(4.8, "Use SweepSingleByObjectType instead.")
-	bool SweepSingle(struct FHitResult& OutHit, const FVector& Start, const FVector& End, const FQuat& Rot, const FCollisionShape& CollisionShape, const FCollisionQueryParams& Params, const FCollisionObjectQueryParams& ObjectQueryParams) const
-	{
-		return SweepSingleByObjectType(OutHit, Start, End, Rot, ObjectQueryParams, CollisionShape, Params);
-	}
-
 	/**
 	 *  Sweep a shape against the world and return the first blocking hit using object types
 	 *  @param  OutHit          First blocking hit found
@@ -1381,25 +1224,6 @@ public:
 	 */
 	bool SweepSingleByProfile(struct FHitResult& OutHit, const FVector& Start, const FVector& End, const FQuat& Rot, FName ProfileName, const FCollisionShape& CollisionShape, const FCollisionQueryParams& Params = FCollisionQueryParams::DefaultQueryParam) const;
 
-	/**
-	 *  Sweep a shape against the world and return all initial overlaps using a specific channel (including blocking) if requested, then overlapping hits and then first blocking hit
-	 *  Results are sorted, so a blocking hit (if found) will be the last element of the array
-	 *  Only the single closest blocking result will be generated, no tests will be done after that
-	 *  @param  OutHits         Array of hits found between ray and the world
-	 *  @param  Start           Start location of the shape
-	 *  @param  End             End location of the shape
-	 *  @param  TraceChannel    The 'channel' that this ray is in, used to determine which components to hit
-	 *  @param	CollisionShape	CollisionShape - supports Box, Sphere, Capsule
-	 *  @param  Params          Additional parameters used for the trace
-	 * 	@param 	ResponseParam	ResponseContainer to be used for this trace	 
-	 *  @return TRUE if OutHits contains any blocking hit entries
-	 */
-	DEPRECATED(4.8, "Use SweepMultiByChannel instead.")
-	bool SweepMulti(TArray<struct FHitResult>& OutHits, const FVector& Start, const FVector& End, const FQuat& Rot, ECollisionChannel TraceChannel, const FCollisionShape& CollisionShape, const FCollisionQueryParams& Params, const FCollisionResponseParams& ResponseParam = FCollisionResponseParams::DefaultResponseParam) const
-	{
-		return SweepMultiByChannel(OutHits, Start, End, Rot, TraceChannel, CollisionShape, Params, ResponseParam);
-	}
-
 	/**
 	 *  Sweep a shape against the world and return all initial overlaps using a specific channel (including blocking) if requested, then overlapping hits and then first blocking hit
 	 *  Results are sorted, so a blocking hit (if found) will be the last element of the array
@@ -1415,24 +1239,6 @@ public:
 	 */
 	bool SweepMultiByChannel(TArray<struct FHitResult>& OutHits, const FVector& Start, const FVector& End, const FQuat& Rot, ECollisionChannel TraceChannel, const FCollisionShape& CollisionShape, const FCollisionQueryParams& Params = FCollisionQueryParams::DefaultQueryParam, const FCollisionResponseParams& ResponseParam = FCollisionResponseParams::DefaultResponseParam) const;
 
-	/**
-	 *  Sweep a shape against the world and return all initial overlaps using object types (including blocking) if requested, then overlapping hits and then first blocking hit
-	 *  Results are sorted, so a blocking hit (if found) will be the last element of the array
-	 *  Only the single closest blocking result will be generated, no tests will be done after that
-	 *  @param  OutHits         Array of hits found between ray and the world
-	 *  @param  Start           Start location of the shape
-	 *  @param  End             End location of the shape
-	 *  @param	CollisionShape	CollisionShape - supports Box, Sphere, Capsule
-	 *  @param  Params          Additional parameters used for the trace
-	 *	@param	ObjectQueryParams	List of object types it's looking for
-	 *  @return TRUE if any hit is found
-	 */
-	DEPRECATED(4.8, "Use SweepMultiByObjectType instead.")
-	bool SweepMulti(TArray<struct FHitResult>& OutHits, const FVector& Start, const FVector& End, const FQuat& Rot, const FCollisionShape& CollisionShape, const FCollisionQueryParams& Params, const FCollisionObjectQueryParams& ObjectQueryParams) const
-	{
-		return SweepMultiByObjectType(OutHits, Start, End, Rot, ObjectQueryParams, CollisionShape, Params);
-	}
-
 	/**
 	 *  Sweep a shape against the world and return all initial overlaps using object types (including blocking) if requested, then overlapping hits and then first blocking hit
 	 *  Results are sorted, so a blocking hit (if found) will be the last element of the array
@@ -1461,21 +1267,6 @@ public:
 	 */
 	bool SweepMultiByProfile(TArray<FHitResult>& OutHits, const FVector& Start, const FVector& End, const FQuat& Rot, FName ProfileName, const FCollisionShape& CollisionShape, const FCollisionQueryParams& Params = FCollisionQueryParams::DefaultQueryParam) const;
 
-	/**
-	 *  Test the collision of a shape at the supplied location using a specific channel, and return if any blocking overlap is found
-	 *  @param  Pos             Location of center of box to test against the world
-	 *  @param  TraceChannel    The 'channel' that this query is in, used to determine which components to hit
-	 *  @param	CollisionShape	CollisionShape - supports Box, Sphere, Capsule
-	 *  @param  Params          Additional parameters used for the trace
-	 * 	@param 	ResponseParam	ResponseContainer to be used for this trace	 
-	 *  @return TRUE if any blocking results are found
-	 */
-	DEPRECATED(4.8, "Use OverlapBlockingTestByChannel instead.")
-	bool OverlapTest(const FVector& Pos, const FQuat& Rot, ECollisionChannel TraceChannel, const FCollisionShape& CollisionShape, const FCollisionQueryParams& Params, const FCollisionResponseParams& ResponseParam = FCollisionResponseParams::DefaultResponseParam) const
-	{
-		return OverlapBlockingTestByChannel(Pos, Rot, TraceChannel, CollisionShape, Params, ResponseParam);
-	}
-
 	/**
 	*  Test the collision of a shape at the supplied location using a specific channel, and return if any blocking overlap is found
 	*  @param  Pos             Location of center of box to test against the world
@@ -1498,20 +1289,6 @@ public:
 	*/
 	bool OverlapAnyTestByChannel(const FVector& Pos, const FQuat& Rot, ECollisionChannel TraceChannel, const FCollisionShape& CollisionShape, const FCollisionQueryParams& Params = FCollisionQueryParams::DefaultQueryParam, const FCollisionResponseParams& ResponseParam = FCollisionResponseParams::DefaultResponseParam) const;
 
-	/**
-	 *  Test the collision of a shape at the supplied location using object types, and return if any blocking overlap is found
-	 *  @param  Pos             Location of center of box to test against the world
-	 *  @param	CollisionShape	CollisionShape - supports Box, Sphere, Capsule
-	 *  @param  Params          Additional parameters used for the trace
-	 *	@param	ObjectQueryParams	List of object types it's looking for
-	 *  @return TRUE if any blocking results are found
-	 */
-	DEPRECATED(4.8, "Use OverlapAnyTestByObjectType instead.")
-	bool OverlapTest(const FVector& Pos, const FQuat& Rot, const FCollisionShape& CollisionShape, const FCollisionQueryParams& Params, const FCollisionObjectQueryParams& ObjectQueryParams) const
-	{
-		return OverlapAnyTestByObjectType(Pos, Rot, ObjectQueryParams, CollisionShape, Params);
-	}
-
 	/**
 	*  Test the collision of a shape at the supplied location using object types, and return if any overlap is found
 	*  @param  Pos             Location of center of box to test against the world
@@ -1522,20 +1299,6 @@ public:
 	*/
 	bool OverlapAnyTestByObjectType(const FVector& Pos, const FQuat& Rot, const FCollisionObjectQueryParams& ObjectQueryParams, const FCollisionShape& CollisionShape, const FCollisionQueryParams& Params = FCollisionQueryParams::DefaultQueryParam) const;
 
-	/**
-	 *  Test the collision of a shape at the supplied location using a specific profile, and return if any blocking overlap is found
-	 *  @param  Pos             Location of center of box to test against the world
-	 *  @param  ProfileName     The 'profile' used to determine which components to hit
-	 *  @param	CollisionShape	CollisionShape - supports Box, Sphere, Capsule
-	 *  @param  Params          Additional parameters used for the trace
-	 *  @return TRUE if any blocking results are found
-	 */
-	DEPRECATED(4.8, "Use OverlapBlockingTestByProfile instead.")
-	bool OverlapTestByProfile(const FVector& Pos, const FQuat& Rot, FName ProfileName, const FCollisionShape& CollisionShape, const FCollisionQueryParams& Params) const
-	{
-		return OverlapBlockingTestByProfile(Pos, Rot, ProfileName, CollisionShape, Params);
-	}
-
 	/**
 	*  Test the collision of a shape at the supplied location using a specific profile, and return if any blocking overlap is found
 	*  @param  Pos             Location of center of box to test against the world
@@ -1556,59 +1319,7 @@ public:
 	*/
 	bool OverlapAnyTestByProfile(const FVector& Pos, const FQuat& Rot, FName ProfileName, const FCollisionShape& CollisionShape, const FCollisionQueryParams& Params = FCollisionQueryParams::DefaultQueryParam) const;
 
-	/**
-	 *  Test the collision of a shape at the supplied location using a specific channel, and determine the set of components that it overlaps
-	 *  @param  OutOverlap      Component found to overlap supplied shape
-	 *  @param  Pos             Location of center of shape to test against the world
-	 *  @param  TraceChannel    The 'channel' that this query is in, used to determine which components to hit
-	 *  @param	CollisionShape	CollisionShape - supports Box, Sphere, Capsule
-	 *  @param  Params          Additional parameters used for the trace
-	 * 	@param 	ResponseParam	ResponseContainer to be used for this trace
-	 *  @return TRUE if OutOverlap is a blocking result
-	 */
-	DEPRECATED(4.8, "Use OverlapMultiByChannel instead. Deciding which overlap to use is gameplay specific.")
-	bool OverlapSingle(struct FOverlapResult& OutOverlap, const FVector& Pos, const FQuat& Rot, ECollisionChannel TraceChannel, const FCollisionShape& CollisionShape, const FCollisionQueryParams& Params, const FCollisionResponseParams& ResponseParam = FCollisionResponseParams::DefaultResponseParam) const;
-
-	/**
-	 *  Test the collision of a shape at the supplied location using object types, and determine the set of components that it overlaps
-	 *  @param  OutOverlap      Component found to overlap supplied shape
-	 *  @param  Pos             Location of center of shape to test against the world
-	 *  @param	CollisionShape	CollisionShape - supports Box, Sphere, Capsule
-	 *  @param  Params          Additional parameters used for the trace
-	 *	@param	ObjectQueryParams	List of object types it's looking for
-	 *  @return TRUE if any overlap is found
-	 */
-	DEPRECATED(4.8, "Use OverlapMultiByObjectType instead. Deciding which overlap to use is gameplay specific.")
-	bool OverlapSingle(struct FOverlapResult& OutOverlap, const FVector& Pos, const FQuat& Rot, const FCollisionShape& CollisionShape, const FCollisionQueryParams& Params, const FCollisionObjectQueryParams& ObjectQueryParams) const;
-
-	/**
-	 *  Test the collision of a shape at the supplied location using a specific profile, and determine the set of components that it overlaps
-	 *  @param  OutOverlap      Component found to overlap supplied shape
-	 *  @param  Pos             Location of center of shape to test against the world
-	 *  @param  ProfileName     The 'profile' used to determine which components to hit
-	 *  @param	CollisionShape	CollisionShape - supports Box, Sphere, Capsule
-	 *  @param  Params          Additional parameters used for the trace
-	 *  @return TRUE if OutOverlap is blocking
-	 */
-	DEPRECATED(4.8, "Use OverlapMultiByProfile instead. Deciding which overlap to use is gameplay specific.")
-	bool OverlapSingleByProfile(struct FOverlapResult& OutOverlap, const FVector& Pos, const FQuat& Rot, FName ProfileName, const FCollisionShape& CollisionShape, const FCollisionQueryParams& Params) const;
-
-	/**
-	 *  Test the collision of a shape at the supplied location using a specific channel, and determine the set of components that it overlaps
-	 *  @param  OutOverlaps     Array of components found to overlap supplied box
-	 *  @param  Pos             Location of center of shape to test against the world
-	 *  @param  TraceChannel    The 'channel' that this query is in, used to determine which components to hit
-	 *  @param	CollisionShape	CollisionShape - supports Box, Sphere, Capsule
-	 *  @param  Params          Additional parameters used for the trace
-	 * 	@param 	ResponseParam	ResponseContainer to be used for this trace
-	 *  @return TRUE if OutOverlaps contains any blocking results
-	 */
-	DEPRECATED(4.8, "Use OverlapMultiByChannel instead.")
-	bool OverlapMulti(TArray<struct FOverlapResult>& OutOverlaps, const FVector& Pos, const FQuat& Rot, ECollisionChannel TraceChannel, const FCollisionShape& CollisionShape, const FCollisionQueryParams& Params, const FCollisionResponseParams& ResponseParam = FCollisionResponseParams::DefaultResponseParam) const
-	{
-		return OverlapMultiByChannel(OutOverlaps, Pos, Rot, TraceChannel, CollisionShape, Params, ResponseParam);
-	}
-
+	
 	/**
 	 *  Test the collision of a shape at the supplied location using a specific channel, and determine the set of components that it overlaps
 	 *  @param  OutOverlaps     Array of components found to overlap supplied box
@@ -1621,21 +1332,7 @@ public:
 	 */
 	bool OverlapMultiByChannel(TArray<struct FOverlapResult>& OutOverlaps, const FVector& Pos, const FQuat& Rot, ECollisionChannel TraceChannel, const FCollisionShape& CollisionShape, const FCollisionQueryParams& Params = FCollisionQueryParams::DefaultQueryParam, const FCollisionResponseParams& ResponseParam = FCollisionResponseParams::DefaultResponseParam) const;
 
-	/**
-	 *  Test the collision of a shape at the supplied location using object types, and determine the set of components that it overlaps
-	 *  @param  OutOverlaps     Array of components found to overlap supplied box
-	 *  @param  Pos             Location of center of shape to test against the world
-	 *  @param	CollisionShape	CollisionShape - supports Box, Sphere, Capsule
-	 *  @param  Params          Additional parameters used for the trace
-	 *	@param	ObjectQueryParams	List of object types it's looking for
-	 *  @return TRUE if any overlap is found
-	 */
-	DEPRECATED(4.8, "Use OverlapMultiByObjectType instead.")
-	bool OverlapMulti(TArray<struct FOverlapResult>& OutOverlaps, const FVector& Pos, const FQuat& Rot, const FCollisionShape& CollisionShape, const FCollisionQueryParams& Params, const FCollisionObjectQueryParams& ObjectQueryParams) const
-	{
-		return OverlapMultiByObjectType(OutOverlaps, Pos, Rot, ObjectQueryParams, CollisionShape, Params);
-	}
-	
+
 	/**
 	 *  Test the collision of a shape at the supplied location using object types, and determine the set of components that it overlaps
 	 *  @param  OutOverlaps     Array of components found to overlap supplied box
@@ -1689,21 +1386,6 @@ public:
 	bool ComponentOverlapMulti(TArray<struct FOverlapResult>& OutOverlaps, const class UPrimitiveComponent* PrimComp, const FVector& Pos, const FQuat& Rot,    const FComponentQueryParams& Params = FComponentQueryParams::DefaultComponentQueryParams, const FCollisionObjectQueryParams& ObjectQueryParams=FCollisionObjectQueryParams::DefaultObjectQueryParam) const;
 	bool ComponentOverlapMulti(TArray<struct FOverlapResult>& OutOverlaps, const class UPrimitiveComponent* PrimComp, const FVector& Pos, const FRotator& Rot, const FComponentQueryParams& Params = FComponentQueryParams::DefaultComponentQueryParams, const FCollisionObjectQueryParams& ObjectQueryParams=FCollisionObjectQueryParams::DefaultObjectQueryParam) const;
 
-	/**
-	 *  Test the collision of the supplied component at the supplied location/rotation using a specific channel, and determine the set of components that it overlaps
-	 *  @param  OutOverlaps     Array of overlaps found between component in specified pose and the world
-	 *  @param  PrimComp        Component's geometry to test against the world. Transform of this component is ignored
-	 *  @param  Pos             Location of PrimComp geometry for test against the world
-	 *  @param  Rot             Rotation of PrimComp geometry for test against the world
-	 *  @param  TraceChannel	The 'channel' that this query is in, used to determine which components to hit
-	 *  @return TRUE if OutOverlaps contains any blocking results
-	 */
-	DEPRECATED(4.8, "Use ComponentOverlapMultiByChannel instead.")
-	bool ComponentOverlapMulti(TArray<struct FOverlapResult>& OutOverlaps, const class UPrimitiveComponent* PrimComp, const FVector& Pos, const FRotator& Rot, ECollisionChannel TraceChannel, const FComponentQueryParams& Params, const FCollisionObjectQueryParams& ObjectQueryParams=FCollisionObjectQueryParams::DefaultObjectQueryParam) const
-	{
-		return ComponentOverlapMultiByChannel(OutOverlaps, PrimComp, Pos, Rot.Quaternion(), TraceChannel, Params, ObjectQueryParams);
-	}
-
 	/**
 	 *  Test the collision of the supplied component at the supplied location/rotation using a specific channel, and determine the set of components that it overlaps
 	 *  @note The overload taking rotation as an FQuat is slightly faster than the version using FRotator (which will be converted to an FQuat)..
@@ -1717,37 +1399,13 @@ public:
 	bool ComponentOverlapMultiByChannel(TArray<struct FOverlapResult>& OutOverlaps, const class UPrimitiveComponent* PrimComp, const FVector& Pos, const FQuat& Rot,    ECollisionChannel TraceChannel, const FComponentQueryParams& Params = FComponentQueryParams::DefaultComponentQueryParams, const FCollisionObjectQueryParams& ObjectQueryParams=FCollisionObjectQueryParams::DefaultObjectQueryParam) const;
 	bool ComponentOverlapMultiByChannel(TArray<struct FOverlapResult>& OutOverlaps, const class UPrimitiveComponent* PrimComp, const FVector& Pos, const FRotator& Rot, ECollisionChannel TraceChannel, const FComponentQueryParams& Params = FComponentQueryParams::DefaultComponentQueryParams, const FCollisionObjectQueryParams& ObjectQueryParams=FCollisionObjectQueryParams::DefaultObjectQueryParam) const;
 
-	/**
-	 * Interface for Async. Pretty much same parameter set except you can optional set delegate to be called when execution is completed and you can set UserData if you'd like
-	 * if no delegate, you can query trace data using QueryTraceData or QueryOverlapData
-	 * the data is available only in the next frame after request is made - in other words, if request is made in frame X, you can get the result in frame (X+1)
-	 *
-	 *  @param  Start           Start location of the ray
-	 *  @param  End             End location of the ray
-	 *  @param  TraceChannel    The 'channel' that this ray is in, used to determine which components to hit
-	 *  @param  Params          Additional parameters used for the trace
-	 * 	@param 	ResponseParam	ResponseContainer to be used for this trace
-	 *	@param	InDeleagte		Delegate function to be called - to see example, search FTraceDelegate
-	 *							Example can be void MyActor::TraceDone(const FTraceHandle& TraceHandle, FTraceDatum & TraceData)
-	 *							Before sending to the function, 
-	 *						
-	 *							FTraceDelegate TraceDelegate;
-	 *							TraceDelegate.BindRaw(this, &MyActor::TraceDone);
-	 * 
-	 *	@param	UserData		UserData
-	 *	@param bMultiTrace		true if you'd like to get continuous result from the trace, false if you want single
-	 */ 
-	DEPRECATED(4.8, "Use AsyncLineTraceByChannel instead.")
-	FTraceHandle	AsyncLineTrace(const FVector& Start,const FVector& End, ECollisionChannel TraceChannel, const FCollisionQueryParams& Params, const FCollisionResponseParams& ResponseParam = FCollisionResponseParams::DefaultResponseParam, FTraceDelegate * InDelegate=NULL, uint32 UserData = 0, bool bMultiTrace=false )
-	{
-		return AsyncLineTraceByChannel(Start, End, TraceChannel, Params, ResponseParam, InDelegate, UserData, bMultiTrace);
-	}
 
 	/**
 	 * Interface for Async. Pretty much same parameter set except you can optional set delegate to be called when execution is completed and you can set UserData if you'd like
 	 * if no delegate, you can query trace data using QueryTraceData or QueryOverlapData
 	 * the data is available only in the next frame after request is made - in other words, if request is made in frame X, you can get the result in frame (X+1)
 	 *
+	 *	@param	InTraceType		Indicates if you want multiple results, single result, or just yes/no (no hit information)
 	 *  @param  Start           Start location of the ray
 	 *  @param  End             End location of the ray
 	 *  @param  TraceChannel    The 'channel' that this ray is in, used to determine which components to hit
@@ -1761,33 +1419,13 @@ public:
 	 *							TraceDelegate.BindRaw(this, &MyActor::TraceDone);
 	 * 
 	 *	@param	UserData		UserData
-	 *	@param bMultiTrace		true if you'd like to get continuous result from the trace, false if you want single
 	 */ 
-	FTraceHandle	AsyncLineTraceByChannel(const FVector& Start,const FVector& End, ECollisionChannel TraceChannel, const FCollisionQueryParams& Params = FCollisionQueryParams::DefaultQueryParam, const FCollisionResponseParams& ResponseParam = FCollisionResponseParams::DefaultResponseParam, FTraceDelegate * InDelegate=NULL, uint32 UserData = 0, bool bMultiTrace=false );
+	FTraceHandle	AsyncLineTraceByChannel(EAsyncTraceType InTraceType, const FVector& Start,const FVector& End, ECollisionChannel TraceChannel, const FCollisionQueryParams& Params = FCollisionQueryParams::DefaultQueryParam, const FCollisionResponseParams& ResponseParam = FCollisionResponseParams::DefaultResponseParam, FTraceDelegate * InDelegate=NULL, uint32 UserData = 0 );
 
-	/**
-	 * Interface for Async. Pretty much same parameter set except you can optional set delegate to be called when execution is completed and you can set UserData if you'd like
-	 * if no delegate, you can query trace data using QueryTraceData or QueryOverlapData
-	 * the data is available only in the next frame after request is made - in other words, if request is made in frame X, you can get the result in frame (X+1)
-	 *
-	 *  @param  Start           Start location of the ray
-	 *  @param  End             End location of the ray
-	 *  @param  Params          Additional parameters used for the trace
-	 *	@param	ObjectQueryParams	List of object types it's looking for
-	 *	@param	InDeleagte		Delegate function to be called - to see example, search FTraceDelegate
-	 *							Example can be void MyActor::TraceDone(const FTraceHandle& TraceHandle, FTraceDatum & TraceData)
-	 *							Before sending to the function, 
-	 *						
-	 *							FTraceDelegate TraceDelegate;
-	 *							TraceDelegate.BindRaw(this, &MyActor::TraceDone);
-	 * 
-	 *	@param	UserData		UserData
-	 *	@param bMultiTrace		true if you'd like to get continuous result from the trace, false if you want single
-	 */ 
-	DEPRECATED(4.8, "Use AsyncLineTraceByObjectType instead.")
-	FTraceHandle	AsyncLineTrace(const FVector& Start,const FVector& End, const FCollisionQueryParams& Params, const FCollisionObjectQueryParams& ObjectQueryParams, FTraceDelegate * InDelegate=NULL, uint32 UserData = 0, bool bMultiTrace=false )
+	DEPRECATED(4.12, "bMultiTrace option replaced with required EAsyncTraceType enum.")
+	FTraceHandle	AsyncLineTraceByChannel(const FVector& Start, const FVector& End, ECollisionChannel TraceChannel, const FCollisionQueryParams& Params = FCollisionQueryParams::DefaultQueryParam, const FCollisionResponseParams& ResponseParam = FCollisionResponseParams::DefaultResponseParam, FTraceDelegate * InDelegate = NULL, uint32 UserData = 0, bool bMultiTrace = false)
 	{
-		return AsyncLineTraceByObjectType(Start, End, ObjectQueryParams, Params, InDelegate, UserData, bMultiTrace);
+		return AsyncLineTraceByChannel(bMultiTrace ? EAsyncTraceType::Multi : EAsyncTraceType::Single, Start, End, TraceChannel, Params, ResponseParam, InDelegate, UserData);
 	}
 
 	/**
@@ -1795,6 +1433,7 @@ public:
 	 * if no delegate, you can query trace data using QueryTraceData or QueryOverlapData
 	 * the data is available only in the next frame after request is made - in other words, if request is made in frame X, you can get the result in frame (X+1)
 	 *
+	 *	@param	InTraceType		Indicates if you want multiple results, single hit result, or just yes/no (no hit information)
 	 *  @param  Start           Start location of the ray
 	 *  @param  End             End location of the ray
 	 *	@param	ObjectQueryParams	List of object types it's looking for
@@ -1807,44 +1446,23 @@ public:
 	 *							TraceDelegate.BindRaw(this, &MyActor::TraceDone);
 	 * 
 	 *	@param	UserData		UserData
-	 *	@param bMultiTrace		true if you'd like to get continuous result from the trace, false if you want single
 	 */ 
-	FTraceHandle	AsyncLineTraceByObjectType(const FVector& Start,const FVector& End, const FCollisionObjectQueryParams& ObjectQueryParams, const FCollisionQueryParams& Params = FCollisionQueryParams::DefaultQueryParam, FTraceDelegate * InDelegate=NULL, uint32 UserData = 0, bool bMultiTrace=false );
+	FTraceHandle	AsyncLineTraceByObjectType(EAsyncTraceType InTraceType, const FVector& Start,const FVector& End, const FCollisionObjectQueryParams& ObjectQueryParams, const FCollisionQueryParams& Params = FCollisionQueryParams::DefaultQueryParam, FTraceDelegate * InDelegate=NULL, uint32 UserData = 0 );
 
-	/**
-	 * Interface for Async trace
-	 * Pretty much same parameter set except you can optional set delegate to be called when execution is completed and you can set UserData if you'd like
-	 * if no delegate, you can query trace data using QueryTraceData or QueryOverlapData
-	 * the data is available only in the next frame after request is made - in other words, if request is made in frame X, you can get the result in frame (X+1)
-	 *
-	 *  @param  Start           Start location of the shape
-	 *  @param  End             End location of the shape
-	 *  @param  TraceChannel    The 'channel' that this trace is in, used to determine which components to hit
-	 *  @param	CollisionShape		CollisionShape - supports Box, Sphere, Capsule
-	 *  @param  Params          Additional parameters used for the trace
-	 * 	@param 	ResponseParam	ResponseContainer to be used for this trace	 
-	 *	@param	InDeleagte		Delegate function to be called - to see example, search FTraceDelegate
-	 *							Example can be void MyActor::TraceDone(const FTraceHandle& TraceHandle, FTraceDatum & TraceData)
-	 *							Before sending to the function, 
-	 *						
-	 *							FTraceDelegate TraceDelegate;
-	 *							TraceDelegate.BindRaw(this, &MyActor::TraceDone);
-	 * 
-	 *	@param	UserData		UserData
-	 *	@param bMultiTrace		true if you'd like to get continuous result from the trace, false if you want single
-	 */ 
-	DEPRECATED(4.8, "Use AsyncSweepByChannel instead.")
-	FTraceHandle	AsyncSweep(const FVector& Start, const FVector& End, ECollisionChannel TraceChannel, const FCollisionShape& CollisionShape, const FCollisionQueryParams& Params, const FCollisionResponseParams& ResponseParam = FCollisionResponseParams::DefaultResponseParam, FTraceDelegate * InDelegate = NULL, uint32 UserData = 0, bool bMultiTrace = false)
+	DEPRECATED(4.12, "bMultiTrace option replaced with required EAsyncTraceType enum.")
+	FTraceHandle	AsyncLineTraceByObjectType(const FVector& Start,const FVector& End, const FCollisionObjectQueryParams& ObjectQueryParams, const FCollisionQueryParams& Params = FCollisionQueryParams::DefaultQueryParam, FTraceDelegate * InDelegate=NULL, uint32 UserData = 0, bool bMultiTrace=false )
 	{
-		return AsyncSweepByChannel(Start, End, TraceChannel, CollisionShape, Params, ResponseParam, InDelegate, UserData, bMultiTrace);
+		return AsyncLineTraceByObjectType(bMultiTrace ? EAsyncTraceType::Multi : EAsyncTraceType::Single, Start, End, ObjectQueryParams, Params, InDelegate, UserData);
 	}
 
+
 	/**
 	 * Interface for Async trace
 	 * Pretty much same parameter set except you can optional set delegate to be called when execution is completed and you can set UserData if you'd like
 	 * if no delegate, you can query trace data using QueryTraceData or QueryOverlapData
 	 * the data is available only in the next frame after request is made - in other words, if request is made in frame X, you can get the result in frame (X+1)
 	 *
+	 *	@param	InTraceType		Indicates if you want multiple results, single hit result, or just yes/no (no hit information)
 	 *  @param  Start           Start location of the shape
 	 *  @param  End             End location of the shape
 	 *  @param  TraceChannel    The 'channel' that this trace is in, used to determine which components to hit
@@ -1859,44 +1477,22 @@ public:
 	 *							TraceDelegate.BindRaw(this, &MyActor::TraceDone);
 	 * 
 	 *	@param	UserData		UserData
-	 *	@param bMultiTrace		true if you'd like to get continuous result from the trace, false if you want single
 	 */ 
-	FTraceHandle	AsyncSweepByChannel(const FVector& Start, const FVector& End, ECollisionChannel TraceChannel, const FCollisionShape& CollisionShape, const FCollisionQueryParams& Params = FCollisionQueryParams::DefaultQueryParam, const FCollisionResponseParams& ResponseParam = FCollisionResponseParams::DefaultResponseParam, FTraceDelegate * InDelegate = NULL, uint32 UserData = 0, bool bMultiTrace = false);
+	FTraceHandle	AsyncSweepByChannel(EAsyncTraceType InTraceType, const FVector& Start, const FVector& End, ECollisionChannel TraceChannel, const FCollisionShape& CollisionShape, const FCollisionQueryParams& Params = FCollisionQueryParams::DefaultQueryParam, const FCollisionResponseParams& ResponseParam = FCollisionResponseParams::DefaultResponseParam, FTraceDelegate * InDelegate = NULL, uint32 UserData = 0);
 
-	/**
-	 * Interface for Async trace
-	 * Pretty much same parameter set except you can optional set delegate to be called when execution is completed and you can set UserData if you'd like
-	 * if no delegate, you can query trace data using QueryTraceData or QueryOverlapData
-	 * the data is available only in the next frame after request is made - in other words, if request is made in frame X, you can get the result in frame (X+1)
-	 *
-	 *  @param  Start           Start location of the shape
-	 *  @param  End             End location of the shape
-	 *  @param	CollisionShape		CollisionShape - supports Box, Sphere, Capsule
-	 *  @param  Params          Additional parameters used for the trace
-	 *	@param	ObjectQueryParams	List of object types it's looking for
-	 *	@param	InDeleagte		Delegate function to be called - to see example, search FTraceDelegate
-	 *							Example can be void MyActor::TraceDone(const FTraceHandle& TraceHandle, FTraceDatum & TraceData)
-	 *							Before sending to the function, 
-	 *						
-	 *							FTraceDelegate TraceDelegate;
-	 *							TraceDelegate.BindRaw(this, &MyActor::TraceDone);
-	 * 
-	 *	@param	UserData		UserData
-	 *	@param bMultiTrace		true if you'd like to get continuous result from the trace, false if you want single
-	 */ 
-	DEPRECATED(4.8, "Use AsyncSweepByObjectType instead.")
-	FTraceHandle	AsyncSweep(const FVector& Start, const FVector& End, const FCollisionShape& CollisionShape, const FCollisionQueryParams& Params, const FCollisionObjectQueryParams& ObjectQueryParams, FTraceDelegate * InDelegate = NULL, uint32 UserData = 0, bool bMultiTrace = false)
+	DEPRECATED(4.12, "bMultiTrace option replaced with required ETraceDatumType enum.")
+	FTraceHandle	AsyncSweepByChannel(const FVector& Start, const FVector& End, ECollisionChannel TraceChannel, const FCollisionShape& CollisionShape, const FCollisionQueryParams& Params = FCollisionQueryParams::DefaultQueryParam, const FCollisionResponseParams& ResponseParam = FCollisionResponseParams::DefaultResponseParam, FTraceDelegate * InDelegate = NULL, uint32 UserData = 0, bool bMultiTrace = false)
 	{
-		return AsyncSweepByObjectType(Start, End, ObjectQueryParams, CollisionShape, Params, InDelegate, UserData, bMultiTrace);
+		return AsyncSweepByChannel(bMultiTrace ? EAsyncTraceType::Multi : EAsyncTraceType::Single, Start, End, TraceChannel, CollisionShape, Params, ResponseParam, InDelegate, UserData);
 	}
-	
-	
+
 	/**
 	 * Interface for Async trace
 	 * Pretty much same parameter set except you can optional set delegate to be called when execution is completed and you can set UserData if you'd like
 	 * if no delegate, you can query trace data using QueryTraceData or QueryOverlapData
 	 * the data is available only in the next frame after request is made - in other words, if request is made in frame X, you can get the result in frame (X+1)
 	 *
+	 *	@param	InTraceType		Indicates if you want multiple results, single hit result, or just yes/no (no hit information)
 	 *  @param  Start           Start location of the shape
 	 *  @param  End             End location of the shape
 	 *	@param	ObjectQueryParams	List of object types it's looking for
@@ -1910,39 +1506,18 @@ public:
 	 *							TraceDelegate.BindRaw(this, &MyActor::TraceDone);
 	 * 
 	 *	@param	UserData		UserData
-	 *	@param bMultiTrace		true if you'd like to get continuous result from the trace, false if you want single
 	 */ 
-	FTraceHandle	AsyncSweepByObjectType(const FVector& Start, const FVector& End, const FCollisionObjectQueryParams& ObjectQueryParams, const FCollisionShape& CollisionShape, const FCollisionQueryParams& Params = FCollisionQueryParams::DefaultQueryParam, FTraceDelegate * InDelegate = NULL, uint32 UserData = 0, bool bMultiTrace = false);
-
-	// overlap functions
+	FTraceHandle	AsyncSweepByObjectType(EAsyncTraceType InTraceType, const FVector& Start, const FVector& End, const FCollisionObjectQueryParams& ObjectQueryParams, const FCollisionShape& CollisionShape, const FCollisionQueryParams& Params = FCollisionQueryParams::DefaultQueryParam, FTraceDelegate * InDelegate = NULL, uint32 UserData = 0);
 
-	/**
-	 * Interface for Async trace
-	 * Pretty much same parameter set except you can optional set delegate to be called when execution is completed and you can set UserData if you'd like
-	 * if no delegate, you can query trace data using QueryTraceData or QueryOverlapData
-	 * the data is available only in the next frame after request is made - in other words, if request is made in frame X, you can get the result in frame (X+1)
-	 *
-	 *  @param  Pos             Location of center of shape to test against the world
-	 *	@param	bMultiTrace		true if you'd like to do multi trace, or false otherwise
-	 *  @param  TraceChannel    The 'channel' that this query is in, used to determine which components to hit
-	 *  @param	CollisionShape		CollisionShape - supports Box, Sphere, Capsule
-	 *  @param  Params          Additional parameters used for the trace
-	 * 	@param 	ResponseParam	ResponseContainer to be used for this trace
-	 *	@param	InDeleagte		Delegate function to be called - to see example, search FTraceDelegate
-	 *							Example can be void MyActor::TraceDone(const FTraceHandle& TraceHandle, FTraceDatum & TraceData)
-	 *							Before sending to the function, 
-	 *						
-	 *							FTraceDelegate TraceDelegate;
-	 *							TraceDelegate.BindRaw(this, &MyActor::TraceDone);
-	 * 
-	 *	@param UserData			UserData
-	 */ 
-	DEPRECATED(4.8, "Use AsyncOverlapByChannel instead.")
-	FTraceHandle	AsyncOverlap(const FVector& Pos, const FQuat& Rot, ECollisionChannel TraceChannel, const FCollisionShape& CollisionShape, const FCollisionQueryParams& Params, const FCollisionResponseParams& ResponseParam = FCollisionResponseParams::DefaultResponseParam, FOverlapDelegate * InDelegate = NULL, uint32 UserData = 0)
+	DEPRECATED(4.12, "bMultiTrace option replaced with required ETraceDatumType enum.")
+	FTraceHandle	AsyncSweepByObjectType(const FVector& Start, const FVector& End, const FCollisionObjectQueryParams& ObjectQueryParams, const FCollisionShape& CollisionShape, const FCollisionQueryParams& Params = FCollisionQueryParams::DefaultQueryParam, FTraceDelegate * InDelegate = NULL, uint32 UserData = 0, bool bMultiTrace = false)
 	{
-		return AsyncOverlapByChannel(Pos, Rot, TraceChannel, CollisionShape, Params, ResponseParam, InDelegate, UserData);
+		return AsyncSweepByObjectType(bMultiTrace ? EAsyncTraceType::Multi : EAsyncTraceType::Single, Start, End, ObjectQueryParams, CollisionShape, Params, InDelegate, UserData);
 	}
 
+
+	// overlap functions
+
 	/**
 	 * Interface for Async trace
 	 * Pretty much same parameter set except you can optional set delegate to be called when execution is completed and you can set UserData if you'd like
@@ -1966,32 +1541,6 @@ public:
 	 */ 
 	FTraceHandle	AsyncOverlapByChannel(const FVector& Pos, const FQuat& Rot, ECollisionChannel TraceChannel, const FCollisionShape& CollisionShape, const FCollisionQueryParams& Params = FCollisionQueryParams::DefaultQueryParam, const FCollisionResponseParams& ResponseParam = FCollisionResponseParams::DefaultResponseParam, FOverlapDelegate * InDelegate = NULL, uint32 UserData = 0);
 
-	/**
-	 * Interface for Async trace
-	 * Pretty much same parameter set except you can optional set delegate to be called when execution is completed and you can set UserData if you'd like
-	 * if no delegate, you can query trace data using QueryTraceData or QueryOverlapData
-	 * the data is available only in the next frame after request is made - in other words, if request is made in frame X, you can get the result in frame (X+1)
-	 *
-	 *  @param  Pos             Location of center of shape to test against the world
-	 *  @param  TraceChannel    The 'channel' that this query is in, used to determine which components to hit
-	 *  @param	CollisionShape		CollisionShape - supports Box, Sphere, Capsule
-	 *  @param  Params          Additional parameters used for the trace
-	 * 	@param 	ResponseParam	ResponseContainer to be used for this trace
-	 *	@param	InDeleagte		Delegate function to be called - to see example, search FTraceDelegate
-	 *							Example can be void MyActor::TraceDone(const FTraceHandle& TraceHandle, FTraceDatum & TraceData)
-	 *							Before sending to the function, 
-	 *						
-	 *							FTraceDelegate TraceDelegate;
-	 *							TraceDelegate.BindRaw(this, &MyActor::TraceDone);
-	 * 
-	 *	@param UserData			UserData
-	 */ 
-	DEPRECATED(4.8, "Use AsyncOverlapByObjectType instead.")
-	FTraceHandle	AsyncOverlap(const FVector& Pos, const FQuat& Rot, const FCollisionShape& CollisionShape, const FCollisionQueryParams& Params, const FCollisionObjectQueryParams& ObjectQueryParams, FOverlapDelegate * InDelegate = NULL, uint32 UserData = 0)
-	{
-		return AsyncOverlapByObjectType(Pos, Rot, ObjectQueryParams, CollisionShape, Params, InDelegate, UserData);
-	}
-
 	/**
 	 * Interface for Async trace
 	 * Pretty much same parameter set except you can optional set delegate to be called when execution is completed and you can set UserData if you'd like
diff --git a/Engine/Source/Runtime/Engine/Classes/GameFramework/Actor.h b/Engine/Source/Runtime/Engine/Classes/GameFramework/Actor.h
--- a/Engine/Source/Runtime/Engine/Classes/GameFramework/Actor.h
+++ b/Engine/Source/Runtime/Engine/Classes/GameFramework/Actor.h
@@ -2490,7 +2490,7 @@ public:
 	
 	/** Script exposed version of FindComponentByClass */
 	UFUNCTION()
-	virtual UActorComponent* GetComponentByClass(TSubclassOf<UActorComponent> ComponentClass);
+	UActorComponent* GetComponentByClass(TSubclassOf<UActorComponent> ComponentClass);
 
 	/* Gets all the components that inherit from the given class.
 		Currently returns an array of UActorComponent which must be cast to the correct type. */
diff --git a/Engine/Source/Runtime/Engine/Classes/GameFramework/PlayerInput.h b/Engine/Source/Runtime/Engine/Classes/GameFramework/PlayerInput.h
--- a/Engine/Source/Runtime/Engine/Classes/GameFramework/PlayerInput.h
+++ b/Engine/Source/Runtime/Engine/Classes/GameFramework/PlayerInput.h
@@ -59,7 +59,10 @@ struct FKeyBind
 	UPROPERTY(config)
 	uint32 bIgnoreCmd:1;
 
-	FKeyBind() 
+	UPROPERTY(transient)
+	uint32 bDisabled : 1;
+
+	FKeyBind()
 	{
 	}
 };
diff --git a/Engine/Source/Runtime/Engine/Classes/Kismet/GameplayStatics.h b/Engine/Source/Runtime/Engine/Classes/Kismet/GameplayStatics.h
--- a/Engine/Source/Runtime/Engine/Classes/Kismet/GameplayStatics.h
+++ b/Engine/Source/Runtime/Engine/Classes/Kismet/GameplayStatics.h
@@ -538,6 +538,17 @@ class ENGINE_API UGameplayStatics : public UBlueprintFunctionLibrary
 	UFUNCTION(BlueprintCallable, Category="Audio", meta=(WorldContext = "WorldContextObject"))
 	static void SetBaseSoundMix(UObject* WorldContextObject, class USoundMix* InSoundMix);
 
+	/** Overrides the sound class adjuster in the given sound mix. If the sound class does not exist in the input sound mix, the sound class adjustment will be added to the sound mix.
+	 * @param InSoundMixModifier The sound mix to modify.
+	 * @param InSoundClass The sound class to override (or add) in the sound mix.
+	 * @param Volume The volume scale to set the sound class adjuster to.
+	 * @param Pitch The pitch scale to set the sound class adjuster to.
+	 * @param FadeInTime The interpolation time to use to go from the current sound class adjuster values to the new values.
+	 * @param bApplyToChildren Whether or not to apply this override to the sound class' children or to just the specified sound class.
+	 */
+	UFUNCTION(BlueprintCallable, Category = "Audio", meta = (WorldContext = "WorldContextObject"))
+	static void SetSoundMixClassOverride(UObject* WorldContextObject, class USoundMix* InSoundMixModifier, class USoundClass* InSoundClass, float Volume = 1.0f, float Pitch = 1.0f, float FadeInTime = 1.0f, bool bApplyToChildren = true);
+
 	/** Push a sound mix modifier onto the audio system **/
 	UFUNCTION(BlueprintCallable, Category="Audio", meta=(WorldContext = "WorldContextObject"))
 	static void PushSoundMixModifier(UObject* WorldContextObject, class USoundMix* InSoundMixModifier);
@@ -611,6 +622,26 @@ class ENGINE_API UGameplayStatics : public UBlueprintFunctionLibrary
 	UFUNCTION(BlueprintPure, Category = "Collision", meta=(NativeBreakFunc))
 	static void BreakHitResult(const struct FHitResult& Hit, bool& bBlockingHit, bool& bInitialOverlap, float& Time, FVector& Location, FVector& ImpactPoint, FVector& Normal, FVector& ImpactNormal, class UPhysicalMaterial*& PhysMat, class AActor*& HitActor, class UPrimitiveComponent*& HitComponent, FName& HitBoneName, int32& HitItem, FVector& TraceStart, FVector& TraceEnd);
 
+	/** 
+	 *	Create a HitResult struct
+	 * @param Hit			The source HitResult.
+	 * @param bBlockingHit	True if there was a blocking hit, false otherwise.
+	 * @param bInitialOverlap True if the hit started in an initial overlap. In this case some other values should be interpreted differently. Time will be 0, ImpactPoint will equal Location, and normals will be equal and indicate a depenetration vector.
+	 * @param Time			'Time' of impact along trace direction ranging from [0.0 to 1.0) if there is a hit, indicating time between start and end. Equals 1.0 if there is no hit.
+	 * @param Location		Location of the hit in world space. If this was a swept shape test, this is the location where we can place the shape in the world where it will not penetrate.
+	 * @param Normal		Normal of the hit in world space, for the object that was swept (e.g. for a sphere trace this points towards the sphere's center). Equal to ImpactNormal for line tests.
+	 * @param ImpactPoint	Location of the actual contact point of the trace shape with the surface of the hit object. Equal to Location in the case of an initial overlap.
+	 * @param ImpactNormal	Normal of the hit in world space, for the object that was hit by the sweep.
+	 * @param PhysMat		Physical material that was hit. Must set bReturnPhysicalMaterial to true in the query params for this to be returned.
+	 * @param HitActor		Actor hit by the trace.
+	 * @param HitComponent	PrimitiveComponent hit by the trace.
+	 * @param HitBoneName	Name of the bone hit (valid only if we hit a skeletal mesh).
+	 * @param HitItem		Primitive-specific data recording which item in the primitive was hit
+	 */
+	UFUNCTION(BlueprintPure, Category = "Collision", meta = (NativeMakeFunc, Normal="0,0,1", ImpactNormal="0,0,1"))
+	static FHitResult MakeHitResult(bool bBlockingHit, bool bInitialOverlap, float Time, FVector Location, FVector ImpactPoint, FVector Normal, FVector ImpactNormal, class UPhysicalMaterial* PhysMat, class AActor* HitActor, class UPrimitiveComponent* HitComponent, FName HitBoneName, int32 HitItem, FVector TraceStart, FVector TraceEnd);
+
+
 	/** Returns the EPhysicalSurface type of the given Hit. 
 	 * To edit surface type for your project, use ProjectSettings/Physics/PhysicalSurface section
 	 */
diff --git a/Engine/Source/Runtime/Engine/Classes/Particles/Acceleration/ParticleModuleAccelerationDrag.h b/Engine/Source/Runtime/Engine/Classes/Particles/Acceleration/ParticleModuleAccelerationDrag.h
--- a/Engine/Source/Runtime/Engine/Classes/Particles/Acceleration/ParticleModuleAccelerationDrag.h
+++ b/Engine/Source/Runtime/Engine/Classes/Particles/Acceleration/ParticleModuleAccelerationDrag.h
@@ -14,8 +14,12 @@ class UParticleModuleAccelerationDrag : public UParticleModuleAccelerationBase
 	GENERATED_UCLASS_BODY()
 
 	/** Per-particle drag coefficient. Evaluted using emitter time. */
-	UPROPERTY(EditAnywhere, Category=Drag)
-	class UDistributionFloat* DragCoefficient;
+	UPROPERTY()
+	class UDistributionFloat* DragCoefficient_DEPRECATED;
+
+	/** Per-particle drag coefficient. Evaluted using emitter time. */
+	UPROPERTY(EditAnywhere, Category = Drag)
+	FRawDistributionFloat DragCoefficientRaw;
 
 	/** Initializes the default values for this property */
 	void InitializeDefaults();
@@ -32,6 +36,8 @@ class UParticleModuleAccelerationDrag : public UParticleModuleAccelerationBase
 	virtual void Update( FParticleEmitterInstance* Owner, int32 Offset, float DeltaTime ) override;
 	//End UParticleModule Interface
 
+	virtual void PostLoad() override;
+
 #if WITH_EDITOR
 	virtual bool IsValidForLODLevel(UParticleLODLevel* LODLevel, FString& OutErrorString) override;
 #endif
diff --git a/Engine/Source/Runtime/Engine/Classes/Particles/Acceleration/ParticleModuleAccelerationDragScaleOverLife.h b/Engine/Source/Runtime/Engine/Classes/Particles/Acceleration/ParticleModuleAccelerationDragScaleOverLife.h
--- a/Engine/Source/Runtime/Engine/Classes/Particles/Acceleration/ParticleModuleAccelerationDragScaleOverLife.h
+++ b/Engine/Source/Runtime/Engine/Classes/Particles/Acceleration/ParticleModuleAccelerationDragScaleOverLife.h
@@ -14,8 +14,12 @@ class UParticleModuleAccelerationDragScaleOverLife : public UParticleModuleAccel
 	GENERATED_UCLASS_BODY()
 
 	/** Per-particle drag scale. Evaluted using particle relative time. */
-	UPROPERTY(EditAnywhere, Category=Drag)
-	class UDistributionFloat* DragScale;
+	UPROPERTY()
+	class UDistributionFloat* DragScale_DEPRECATED;
+
+	/** Per-particle drag scale. Evaluted using particle relative time. */
+	UPROPERTY(EditAnywhere, Category = Drag)
+	FRawDistributionFloat DragScaleRaw;
 
 	/** Initializes the default values for this property */
 	void InitializeDefaults();
@@ -25,6 +29,7 @@ class UParticleModuleAccelerationDragScaleOverLife : public UParticleModuleAccel
 	virtual void PostEditChangeProperty(FPropertyChangedEvent& PropertyChangedEvent) override;
 #endif // WITH_EDITOR
 	virtual void PostInitProperties() override;
+	virtual void PostLoad() override;
 	//End UObject Interface
 
 	//Begin UParticleModule Interface
diff --git a/Engine/Source/Runtime/Engine/Classes/Particles/Attractor/ParticleModuleAttractorPointGravity.h b/Engine/Source/Runtime/Engine/Classes/Particles/Attractor/ParticleModuleAttractorPointGravity.h
--- a/Engine/Source/Runtime/Engine/Classes/Particles/Attractor/ParticleModuleAttractorPointGravity.h
+++ b/Engine/Source/Runtime/Engine/Classes/Particles/Attractor/ParticleModuleAttractorPointGravity.h
@@ -22,9 +22,13 @@ class UParticleModuleAttractorPointGravity : public UParticleModuleAttractorBase
 	UPROPERTY(EditAnywhere, Category=PointGravitySource)
 	float Radius;
 
+	/** The strength of the point source. */
+	UPROPERTY()
+	class UDistributionFloat* Strength_DEPRECATED;
+
 	/** The strength of the point source. */
 	UPROPERTY(EditAnywhere, noclear, Category=PointGravitySource)
-	class UDistributionFloat* Strength;
+	FRawDistributionFloat StrengthRaw;
 
 	/** Initializes the default values for this property */
 	void InitializeDefaults();
@@ -34,6 +38,7 @@ class UParticleModuleAttractorPointGravity : public UParticleModuleAttractorBase
 	virtual void PostEditChangeProperty(FPropertyChangedEvent& PropertyChangedEvent) override;
 #endif // WITH_EDITOR
 	virtual void PostInitProperties() override;
+	virtual void PostLoad() override;
 	//End UObject Interface
 
 	//Begin UParticleModule Interface
diff --git a/Engine/Source/Runtime/Engine/Classes/Particles/TypeData/ParticleModuleTypeDataGpu.h b/Engine/Source/Runtime/Engine/Classes/Particles/TypeData/ParticleModuleTypeDataGpu.h
--- a/Engine/Source/Runtime/Engine/Classes/Particles/TypeData/ParticleModuleTypeDataGpu.h
+++ b/Engine/Source/Runtime/Engine/Classes/Particles/TypeData/ParticleModuleTypeDataGpu.h
@@ -399,11 +399,11 @@ class UParticleModuleTypeDataGpu : public UParticleModuleTypeDataBase
 	GENERATED_UCLASS_BODY()
 
 	/** Information for runtime simulation. */
-	UPROPERTY(transient)
+	UPROPERTY()
 	struct FGPUSpriteEmitterInfo EmitterInfo;
 
 	/** Data used to initialize runtime resources. */
-	UPROPERTY(transient)
+	UPROPERTY()
 	struct FGPUSpriteResourceData ResourceData;
 
 	/** TEMP: How much to stretch sprites based on camera motion blur. */
diff --git a/Engine/Source/Runtime/Engine/Classes/Particles/VectorField/ParticleModuleVectorFieldScale.h b/Engine/Source/Runtime/Engine/Classes/Particles/VectorField/ParticleModuleVectorFieldScale.h
--- a/Engine/Source/Runtime/Engine/Classes/Particles/VectorField/ParticleModuleVectorFieldScale.h
+++ b/Engine/Source/Runtime/Engine/Classes/Particles/VectorField/ParticleModuleVectorFieldScale.h
@@ -13,9 +13,13 @@ class UParticleModuleVectorFieldScale : public UParticleModuleVectorFieldBase
 {
 	GENERATED_UCLASS_BODY()
 
+	/** Per-particle vector field scale. Evaluated using emitter time. */
+	UPROPERTY()
+	class UDistributionFloat* VectorFieldScale_DEPRECATED;
+
 	/** Per-particle vector field scale. Evaluated using emitter time. */
 	UPROPERTY(EditAnywhere, Category=VectorField)
-	class UDistributionFloat* VectorFieldScale;
+	FRawDistributionFloat VectorFieldScaleRaw;
 
 	/** Initializes the default values for this property */
 	void InitializeDefaults();
@@ -25,6 +29,7 @@ class UParticleModuleVectorFieldScale : public UParticleModuleVectorFieldBase
 	virtual void PostEditChangeProperty(FPropertyChangedEvent& PropertyChangedEvent) override;
 #endif // WITH_EDITOR
 	virtual void PostInitProperties() override;
+	virtual void PostLoad() override;
 	//~ End UObject Interface
 
 	//~ Begin UParticleModule Interface
diff --git a/Engine/Source/Runtime/Engine/Classes/Particles/VectorField/ParticleModuleVectorFieldScaleOverLife.h b/Engine/Source/Runtime/Engine/Classes/Particles/VectorField/ParticleModuleVectorFieldScaleOverLife.h
--- a/Engine/Source/Runtime/Engine/Classes/Particles/VectorField/ParticleModuleVectorFieldScaleOverLife.h
+++ b/Engine/Source/Runtime/Engine/Classes/Particles/VectorField/ParticleModuleVectorFieldScaleOverLife.h
@@ -13,9 +13,13 @@ class UParticleModuleVectorFieldScaleOverLife : public UParticleModuleVectorFiel
 {
 	GENERATED_UCLASS_BODY()
 
+	/** Per-particle vector field scale. Evaluated using particle relative time. */
+	UPROPERTY()
+	class UDistributionFloat* VectorFieldScaleOverLife_DEPRECATED;
+
 	/** Per-particle vector field scale. Evaluated using particle relative time. */
 	UPROPERTY(EditAnywhere, Category=VectorField)
-	class UDistributionFloat* VectorFieldScaleOverLife;
+	FRawDistributionFloat VectorFieldScaleOverLifeRaw;
 
 	/** Initializes the default values for this property */
 	void InitializeDefaults();
@@ -25,6 +29,7 @@ class UParticleModuleVectorFieldScaleOverLife : public UParticleModuleVectorFiel
 	virtual void PostEditChangeProperty(FPropertyChangedEvent& PropertyChangedEvent) override;
 #endif // WITH_EDITOR
 	virtual void PostInitProperties() override;
+	virtual void PostLoad() override;
 	//~ End UObject Interface
 
 	//~ Begin UParticleModule Interface
diff --git a/Engine/Source/Runtime/Engine/Private/AI/Navigation/PImplRecastNavMesh.cpp b/Engine/Source/Runtime/Engine/Private/AI/Navigation/PImplRecastNavMesh.cpp
--- a/Engine/Source/Runtime/Engine/Private/AI/Navigation/PImplRecastNavMesh.cpp
+++ b/Engine/Source/Runtime/Engine/Private/AI/Navigation/PImplRecastNavMesh.cpp
@@ -2182,12 +2182,12 @@ void FPImplRecastNavMesh::GetDebugGeometry(FRecastDebugGeometry& OutGeometry, in
 			}
 		}
 
-		OutGeometry.MeshVerts.Reserve(NumVertsToReserve);
-		OutGeometry.AreaIndices[0].Reserve(NumIndicesToReserve);
-		OutGeometry.BuiltMeshIndices.Reserve(NumIndicesToReserve);
-		OutGeometry.Clusters.AddZeroed(NumClusters);
+		OutGeometry.MeshVerts.Reserve(OutGeometry.MeshVerts.Num() + NumVertsToReserve);
+		OutGeometry.AreaIndices[0].Reserve(OutGeometry.AreaIndices[0].Num() + NumIndicesToReserve);
+		OutGeometry.BuiltMeshIndices.Reserve(OutGeometry.BuiltMeshIndices.Num() + NumIndicesToReserve);
+		OutGeometry.Clusters.AddZeroed(OutGeometry.Clusters.Num() + NumClusters);
 
-		uint32 AllTilesVertBase = 0;
+		uint32 VertBase = OutGeometry.MeshVerts.Num();
 		for (const FIntPoint& TileLocation : ActiveTiles)
 		{
 			const int32 LayersCount = ConstNavMesh->getTileCountAt(TileLocation.X, TileLocation.Y);
@@ -2197,7 +2197,7 @@ void FPImplRecastNavMesh::GetDebugGeometry(FRecastDebugGeometry& OutGeometry, in
 				dtMeshTile const* const Tile = ConstNavMesh->getTileAt(TileLocation.X, TileLocation.Y, Layer);
 				if (Tile != nullptr && Tile->header != nullptr)
 				{
-					AllTilesVertBase += GetTilesDebugGeometry(Generator, *Tile, AllTilesVertBase, OutGeometry);
+					VertBase += GetTilesDebugGeometry(Generator, *Tile, VertBase, OutGeometry);
 				}
 			}
 		}
@@ -2223,14 +2223,12 @@ void FPImplRecastNavMesh::GetDebugGeometry(FRecastDebugGeometry& OutGeometry, in
 			}
 		}
 
-		OutGeometry.MeshVerts.Reserve(NumVertsToReserve);
-		OutGeometry.AreaIndices[0].Reserve(NumIndicesToReserve);
-		OutGeometry.BuiltMeshIndices.Reserve(NumIndicesToReserve);
-		OutGeometry.Clusters.AddZeroed(NumClusters);
+		OutGeometry.MeshVerts.Reserve(OutGeometry.MeshVerts.Num() + NumVertsToReserve);
+		OutGeometry.AreaIndices[0].Reserve(OutGeometry.AreaIndices[0].Num() + NumIndicesToReserve);
+		OutGeometry.BuiltMeshIndices.Reserve(OutGeometry.BuiltMeshIndices.Num() + NumIndicesToReserve);
+		OutGeometry.Clusters.AddZeroed(OutGeometry.Clusters.Num() + NumClusters);
 
-		// spin through all polys in all tiles and draw them
-		// @see drawMeshTile() in recast code for reference
-		uint32 AllTilesVertBase = 0;
+		uint32 VertBase = OutGeometry.MeshVerts.Num();
 		for (int32 TileIdx = StartingTile; TileIdx < NumTiles; ++TileIdx)
 		{
 			dtMeshTile const* const Tile = ConstNavMesh->getTile(TileIdx);
@@ -2240,7 +2238,7 @@ void FPImplRecastNavMesh::GetDebugGeometry(FRecastDebugGeometry& OutGeometry, in
 				continue;
 			}
 
-			AllTilesVertBase += GetTilesDebugGeometry(Generator, *Tile, AllTilesVertBase, OutGeometry, TileIdx);
+			VertBase += GetTilesDebugGeometry(Generator, *Tile, VertBase, OutGeometry, TileIdx);
 		}
 	}
 }
diff --git a/Engine/Source/Runtime/Engine/Private/AI/Navigation/RecastNavMeshGenerator.cpp b/Engine/Source/Runtime/Engine/Private/AI/Navigation/RecastNavMeshGenerator.cpp
--- a/Engine/Source/Runtime/Engine/Private/AI/Navigation/RecastNavMeshGenerator.cpp
+++ b/Engine/Source/Runtime/Engine/Private/AI/Navigation/RecastNavMeshGenerator.cpp
@@ -24,7 +24,6 @@
 #include "RecastHelpers.h"
 #include "NavigationSystemHelpers.h"
 #include "VisualLogger/VisualLogger.h"
-#include "NavMeshRenderingHelpers.h"
 
 #define SEAMLESS_REBUILDING_ENABLED 1
 
@@ -4246,7 +4245,21 @@ void FRecastNavMeshGenerator::GrabDebugSnapshot(struct FVisualLogEntry* Snapshot
 				if (bExportGeometry && Element.Data->CollisionData.Num())
 				{
 					FRecastGeometryCache CachedGeometry(Element.Data->CollisionData.GetData());
-					AppendGeometry(CoordBuffer, Indices, CachedGeometry.Verts, CachedGeometry.Header.NumVerts, CachedGeometry.Indices, CachedGeometry.Header.NumFaces);
+					
+					const uint32 NumVerts = CachedGeometry.Header.NumVerts;
+					CoordBuffer.Reset(NumVerts);
+					for (uint32 VertIdx = 0; VertIdx < NumVerts * 3; VertIdx += 3)
+					{
+						CoordBuffer.Add(Recast2UnrealPoint(&CachedGeometry.Verts[VertIdx]));
+					}
+
+					const uint32 NumIndices = CachedGeometry.Header.NumFaces * 3;
+					Indices.SetNum(NumIndices, false);
+					for (uint32 IndicesIdx = 0; IndicesIdx < NumIndices; ++IndicesIdx)
+					{
+						Indices[IndicesIdx] = CachedGeometry.Indices[IndicesIdx];
+					}
+
 					Snapshot->AddElement(CoordBuffer, Indices, LogCategory.GetCategoryName(), LogVerbosity, FColorList::LightGrey.WithAlpha(255));
 				}
 				else
diff --git a/Engine/Source/Runtime/Engine/Private/ActiveSound.cpp b/Engine/Source/Runtime/Engine/Private/ActiveSound.cpp
--- a/Engine/Source/Runtime/Engine/Private/ActiveSound.cpp
+++ b/Engine/Source/Runtime/Engine/Private/ActiveSound.cpp
@@ -422,7 +422,7 @@ void FActiveSound::CheckOcclusion(const FVector ListenerLocation, const FVector
 			else if (!bAsyncOcclusionPending)
 			{
 				bAsyncOcclusionPending = true;
-				WorldPtr->AsyncLineTraceByChannel(SoundLocation, ListenerLocation, ECC_Visibility, Params, FCollisionResponseParams::DefaultResponseParam, &OcclusionTraceDelegate);
+				WorldPtr->AsyncLineTraceByChannel(EAsyncTraceType::Test, SoundLocation, ListenerLocation, ECC_Visibility, Params, FCollisionResponseParams::DefaultResponseParam, &OcclusionTraceDelegate);
 			}
 		}
 	}
diff --git a/Engine/Source/Runtime/Engine/Private/Actor.cpp b/Engine/Source/Runtime/Engine/Private/Actor.cpp
--- a/Engine/Source/Runtime/Engine/Private/Actor.cpp
+++ b/Engine/Source/Runtime/Engine/Private/Actor.cpp
@@ -2383,17 +2383,7 @@ UActorComponent* AActor::FindComponentByClass(const TSubclassOf<UActorComponent>
 
 UActorComponent* AActor::GetComponentByClass(TSubclassOf<UActorComponent> ComponentClass)
 {
-	UActorComponent* FoundComponent = NULL;
-	for (UActorComponent* Component : OwnedComponents)
-	{
-		if (Component && Component->IsA(ComponentClass))
-		{
-			FoundComponent = Component;
-			break;
-		}
-	}
-
-	return FoundComponent;
+	return FindComponentByClass(ComponentClass);
 }
 
 TArray<UActorComponent*> AActor::GetComponentsByClass(TSubclassOf<UActorComponent> ComponentClass) const
diff --git a/Engine/Source/Runtime/Engine/Private/Audio.cpp b/Engine/Source/Runtime/Engine/Private/Audio.cpp
--- a/Engine/Source/Runtime/Engine/Private/Audio.cpp
+++ b/Engine/Source/Runtime/Engine/Private/Audio.cpp
@@ -594,7 +594,7 @@ float FWaveInstance::GetVolumeWeightedPriority() const
 	float ActualVolume = GetActualVolume();
 	if (ActualVolume > 0.0f)
 	{
-		return GetActualVolume() * Priority;
+		return ActualVolume * Priority;
 	}
 	else
 	{
