Copying //UE4/Dev-Framework to Dev-Main (//UE4/Dev-Main) @ 2879625

diff --git a/Engine/Source/Runtime/Engine/Private/AudioDevice.cpp b/Engine/Source/Runtime/Engine/Private/AudioDevice.cpp
--- a/Engine/Source/Runtime/Engine/Private/AudioDevice.cpp
+++ b/Engine/Source/Runtime/Engine/Private/AudioDevice.cpp
@@ -132,7 +132,7 @@ bool FAudioDevice::Init(int32 InMaxChannels)
 	const FStringAssetReference DefaultBaseSoundMixName = GetDefault<UAudioSettings>()->DefaultBaseSoundMix;
 	if (DefaultBaseSoundMixName.IsValid())
 	{
-		DefaultBaseSoundMix = LoadObject<USoundMix>(NULL, *DefaultBaseSoundMixName.ToString());
+		DefaultBaseSoundMix = LoadObject<USoundMix>(nullptr, *DefaultBaseSoundMixName.ToString());
 	}
 
 	GetDefault<USoundGroups>()->Initialize();
@@ -179,7 +179,7 @@ bool FAudioDevice::Init(int32 InMaxChannels)
 		PrecacheStartupSounds();
 	}
 
-	UE_LOG(LogInit, Log, TEXT("FAudioDevice initialized." ));
+	UE_LOG(LogInit, Log, TEXT("FAudioDevice initialized."));
 
 	bIsInitialized = true;
 
@@ -210,7 +210,7 @@ void FAudioDevice::SetMaxChannels(int32 InMaxChannels)
 {
 	if (InMaxChannels >= Sources.Num())
 	{
-		UE_LOG(LogAudio, Warning, TEXT( "Can't increase channels past starting number!" ) );
+		UE_LOG(LogAudio, Warning, TEXT("Can't increase channels past starting number!"));
 		return;
 	}
 
@@ -220,18 +220,20 @@ void FAudioDevice::SetMaxChannels(int32 InMaxChannels)
 void FAudioDevice::Teardown()
 {
 	// Flush stops all sources so sources can be safely deleted below.
-	Flush(NULL);
+	Flush(nullptr);
 
 	// Clear out the EQ/Reverb/LPF effects
 	if (Effects)
 	{
 		delete Effects;
-		Effects = NULL;
+		Effects = nullptr;
 	}
 
 	// let platform shutdown
 	TeardownHardware();
 
+	SoundMixClassEffectOverrides.Empty();
+
 	// Note: we don't free audio buffers at this stage since they are managed in the audio device manager
 
 	// Must be after FreeBufferResource as that potentially stops sources
@@ -258,7 +260,7 @@ void FAudioDevice::Teardown()
 
 void FAudioDevice::Suspend(bool bGameTicking)
 {
-	HandlePause( bGameTicking, true );
+	HandlePause(bGameTicking, true);
 }
 
 void FAudioDevice::CountBytes(FArchive& Ar)
@@ -274,11 +276,11 @@ void FAudioDevice::CountBytes(FArchive& Ar)
 	SoundMixModifiers.CountBytes(Ar);
 }
 
-void FAudioDevice::AddReferencedObjects( FReferenceCollector& Collector )
+void FAudioDevice::AddReferencedObjects(FReferenceCollector& Collector)
 {	
 	Collector.AddReferencedObject(DefaultBaseSoundMix);
 
-	for( TMap< USoundMix*, FSoundMixState >::TIterator It( SoundMixModifiers ); It; ++It )
+	for (TMap< USoundMix*, FSoundMixState >::TIterator It(SoundMixModifiers); It; ++It)
 	{
 		Collector.AddReferencedObject(It.Key());
 	}
@@ -288,17 +290,17 @@ void FAudioDevice::AddReferencedObjects( FReferenceCollector& Collector )
 		Effects->AddReferencedObjects(Collector);
 	}
 
-	for( int32 i = 0; i < ActiveSounds.Num(); ++i )
+	for (int32 i = 0; i < ActiveSounds.Num(); ++i)
 	{
 		ActiveSounds[i]->AddReferencedObjects(Collector);
 	}
-	for( int32 i = 0; i < PrevPassiveSoundMixModifiers.Num(); ++i )
+	for (int32 i = 0; i < PrevPassiveSoundMixModifiers.Num(); ++i)
 	{
 		Collector.AddReferencedObject(PrevPassiveSoundMixModifiers[i]);
 	}
 }
 
-void FAudioDevice::ResetInterpolation( void )
+void FAudioDevice::ResetInterpolation()
 {
 	for (FListener& Listener : Listeners)
 	{
@@ -315,7 +317,7 @@ void FAudioDevice::ResetInterpolation( void )
 	}
 
 	// Reset sound class properties to defaults
-	for( TMap<USoundClass*, FSoundClassProperties>::TIterator It( SoundClasses ); It; ++It )
+	for (TMap<USoundClass*, FSoundClassProperties>::TIterator It(SoundClasses); It; ++It)
 	{
 		USoundClass* SoundClass = It.Key();
 		if (SoundClass)
@@ -326,91 +328,91 @@ void FAudioDevice::ResetInterpolation( void )
 
 	SoundMixModifiers.Reset();
 	PrevPassiveSoundMixModifiers.Reset();
-	BaseSoundMix = NULL;
+	BaseSoundMix = nullptr;
 
 	// reset audio effects
 	Effects->ResetInterpolation();
 }
 
-void FAudioDevice::EnableRadioEffect( bool bEnable )
+void FAudioDevice::EnableRadioEffect(bool bEnable)
 {
-	if( bEnable )
+	if (bEnable)
 	{
 		DebugState = DEBUGSTATE_None;
 	}
 	else
 	{
-		UE_LOG(LogAudio, Log, TEXT( "Radio disabled for all sources" ) );
+		UE_LOG(LogAudio, Log, TEXT("Radio disabled for all sources"));
 		DebugState = DEBUGSTATE_DisableRadio;
 	}
 }
 
 #if !UE_BUILD_SHIPPING
-bool FAudioDevice::HandleShowSoundClassHierarchyCommand( const TCHAR* Cmd, FOutputDevice& Ar )
+bool FAudioDevice::HandleShowSoundClassHierarchyCommand(const TCHAR* Cmd, FOutputDevice& Ar)
 {
-	ShowSoundClassHierarchy( Ar );
+	ShowSoundClassHierarchy(Ar);
 	return true;
 }
 
-bool FAudioDevice::HandleListWavesCommand( const TCHAR* Cmd, FOutputDevice& Ar )
+bool FAudioDevice::HandleListWavesCommand(const TCHAR* Cmd, FOutputDevice& Ar)
 {
 	TArray<FWaveInstance*> WaveInstances;
-	int32 FirstActiveIndex = GetSortedActiveWaveInstances( WaveInstances, ESortedActiveWaveGetType::QueryOnly );
+	int32 FirstActiveIndex = GetSortedActiveWaveInstances(WaveInstances, ESortedActiveWaveGetType::QueryOnly);
 
-	for( int32 InstanceIndex = FirstActiveIndex; InstanceIndex < WaveInstances.Num(); InstanceIndex++ )
+	for (int32 InstanceIndex = FirstActiveIndex; InstanceIndex < WaveInstances.Num(); InstanceIndex++)
 	{
 		FWaveInstance* WaveInstance = WaveInstances[ InstanceIndex ];
-		FSoundSource* Source = WaveInstanceSourceMap.FindRef( WaveInstance );
+		FSoundSource* Source = WaveInstanceSourceMap.FindRef(WaveInstance);
 		UAudioComponent* AudioComponent = WaveInstance->ActiveSound->GetAudioComponent();
 		AActor* SoundOwner = AudioComponent ? AudioComponent->GetOwner() : nullptr;
-		Ar.Logf( TEXT( "%4i.    %s %6.2f %6.2f  %s   %s"), InstanceIndex, Source ? TEXT( "Yes" ) : TEXT( " No" ), WaveInstance->ActiveSound->PlaybackTime, WaveInstance->GetActualVolume(), *WaveInstance->WaveData->GetPathName(), SoundOwner ? *SoundOwner->GetName() : TEXT("None") );
+		Ar.Logf(TEXT("%4i.    %s %6.2f %6.2f  %s   %s"), InstanceIndex, Source ? TEXT("Yes") : TEXT(" No"), WaveInstance->ActiveSound->PlaybackTime, WaveInstance->GetActualVolume(), *WaveInstance->WaveData->GetPathName(), SoundOwner ? *SoundOwner->GetName() : TEXT("None"));
 	}
 
-	Ar.Logf( TEXT("Total: %i"), WaveInstances.Num()-FirstActiveIndex );
+	Ar.Logf(TEXT("Total: %i"), WaveInstances.Num()-FirstActiveIndex);
 	return true;
 }
 
-void FAudioDevice::GetSoundClassInfo( TMap<FName, FAudioClassInfo>& AudioClassInfos )
+void FAudioDevice::GetSoundClassInfo(TMap<FName, FAudioClassInfo>& AudioClassInfos)
 {
 	// Iterate over all sound cues to get a unique map of sound node waves to class names
 	TMap<USoundWave*, FName> SoundWaveClasses;
 
-	for( TObjectIterator<USoundCue> CueIt; CueIt; ++CueIt )
+	for (TObjectIterator<USoundCue> CueIt; CueIt; ++CueIt)
 	{
 		TArray<USoundNodeWavePlayer*> WavePlayers;
 
 		USoundCue* SoundCue = *CueIt;
-		SoundCue->RecursiveFindNode<USoundNodeWavePlayer>( SoundCue->FirstNode, WavePlayers );
+		SoundCue->RecursiveFindNode<USoundNodeWavePlayer>(SoundCue->FirstNode, WavePlayers);
 
-		for( int32 WaveIndex = 0; WaveIndex < WavePlayers.Num(); ++WaveIndex )
+		for (int32 WaveIndex = 0; WaveIndex < WavePlayers.Num(); ++WaveIndex)
 		{
 			// Presume one class per sound node wave
 			USoundWave *SoundWave = WavePlayers[ WaveIndex ]->GetSoundWave();
 			if (SoundWave && SoundCue->GetSoundClass())
 			{
-				SoundWaveClasses.Add( SoundWave, SoundCue->GetSoundClass()->GetFName() );
+				SoundWaveClasses.Add(SoundWave, SoundCue->GetSoundClass()->GetFName());
 			}
 		}
 	}
 
 	// Add any sound node waves that are not referenced by sound cues
-	for( TObjectIterator<USoundWave> WaveIt; WaveIt; ++WaveIt )
+	for (TObjectIterator<USoundWave> WaveIt; WaveIt; ++WaveIt)
 	{
 		USoundWave* SoundWave = *WaveIt;
-		if( SoundWaveClasses.Find( SoundWave ) == NULL )
+		if (SoundWaveClasses.Find(SoundWave) == nullptr)
 		{
-			SoundWaveClasses.Add( SoundWave, NAME_UnGrouped );
+			SoundWaveClasses.Add(SoundWave, NAME_UnGrouped);
 		}
 	}
 
 	// Collate the data into something useful
-	for( TMap<USoundWave*, FName>::TIterator MapIter( SoundWaveClasses ); MapIter; ++MapIter )
+	for (TMap<USoundWave*, FName>::TIterator MapIter(SoundWaveClasses); MapIter; ++MapIter)
 	{
 		USoundWave* SoundWave = MapIter.Key();
 		FName ClassName = MapIter.Value();
 
-		FAudioClassInfo* AudioClassInfo = AudioClassInfos.Find( ClassName );
-		if( AudioClassInfo == NULL )
+		FAudioClassInfo* AudioClassInfo = AudioClassInfos.Find(ClassName);
+		if (AudioClassInfo == nullptr)
 		{
 			FAudioClassInfo NewAudioClassInfo;
 
@@ -419,17 +421,17 @@ void FAudioDevice::GetSoundClassInfo( TMap<FName, FAudioClassInfo>& AudioClassIn
 			NewAudioClassInfo.NumRealTime = 0;
 			NewAudioClassInfo.SizeRealTime = 0;
 
-			AudioClassInfos.Add( ClassName, NewAudioClassInfo );
+			AudioClassInfos.Add(ClassName, NewAudioClassInfo);
 
-			AudioClassInfo = AudioClassInfos.Find( ClassName );
-			check( AudioClassInfo );
+			AudioClassInfo = AudioClassInfos.Find(ClassName);
+			check(AudioClassInfo);
 		}
 
 #if !WITH_EDITOR
 		AudioClassInfo->SizeResident += SoundWave->GetCompressedDataSize(GetRuntimeFormat(SoundWave));
 		AudioClassInfo->NumResident++;
 #else
-		switch( SoundWave->DecompressionType )
+		switch(SoundWave->DecompressionType)
 		{
 		case DTYPE_Native:
 		case DTYPE_Preview:
@@ -458,37 +460,37 @@ void FAudioDevice::GetSoundClassInfo( TMap<FName, FAudioClassInfo>& AudioClassIn
 	}
 }
 
-bool FAudioDevice::HandleListSoundClassesCommand( const TCHAR* Cmd, FOutputDevice& Ar )
+bool FAudioDevice::HandleListSoundClassesCommand(const TCHAR* Cmd, FOutputDevice& Ar)
 {
 	TMap<FName, FAudioClassInfo> AudioClassInfos;
 
-	GetSoundClassInfo( AudioClassInfos );
+	GetSoundClassInfo(AudioClassInfos);
 
-	Ar.Logf( TEXT( "Listing all sound classes." ) );
+	Ar.Logf(TEXT("Listing all sound classes."));
 
 	// Display the collated data
 	int32 TotalSounds = 0;
-	for( TMap<FName, FAudioClassInfo>::TIterator ACIIter( AudioClassInfos ); ACIIter; ++ACIIter )
+	for (TMap<FName, FAudioClassInfo>::TIterator ACIIter(AudioClassInfos); ACIIter; ++ACIIter)
 	{
 		FName ClassName = ACIIter.Key();
-		FAudioClassInfo* ACI = AudioClassInfos.Find( ClassName );
+		FAudioClassInfo* ACI = AudioClassInfos.Find(ClassName);
 
-		FString Line = FString::Printf( TEXT( "Class '%s' has %d resident sounds taking %.2f kb" ), *ClassName.ToString(), ACI->NumResident, ACI->SizeResident / 1024.0f );
+		FString Line = FString::Printf(TEXT("Class '%s' has %d resident sounds taking %.2f kb"), *ClassName.ToString(), ACI->NumResident, ACI->SizeResident / 1024.0f);
 		TotalSounds += ACI->NumResident;
-		if( ACI->NumRealTime > 0 )
+		if (ACI->NumRealTime > 0)
 		{
-			Line += FString::Printf( TEXT( ", and %d real time sounds taking %.2f kb " ), ACI->NumRealTime, ACI->SizeRealTime / 1024.0f );
+			Line += FString::Printf(TEXT(", and %d real time sounds taking %.2f kb "), ACI->NumRealTime, ACI->SizeRealTime / 1024.0f);
 			TotalSounds += ACI->NumRealTime;
 		}
 
-		Ar.Logf( *Line );
+		Ar.Logf(*Line);
 	}
 
-	Ar.Logf( TEXT( "%d total sounds in %d classes" ), TotalSounds, AudioClassInfos.Num() );
+	Ar.Logf(TEXT("%d total sounds in %d classes"), TotalSounds, AudioClassInfos.Num());
 	return true;
 }
 
-void FAudioDevice::ShowSoundClassHierarchy( FOutputDevice& Ar, USoundClass* InSoundClass, int32 Indent  ) const
+void FAudioDevice::ShowSoundClassHierarchy(FOutputDevice& Ar, USoundClass* InSoundClass, int32 Indent ) const
 {
 	TArray<USoundClass*> SoundClassesToShow;
 	if (InSoundClass)
@@ -497,10 +499,10 @@ void FAudioDevice::ShowSoundClassHierarchy( FOutputDevice& Ar, USoundClass* InSo
 	}
 	else
 	{
-		for( TMap<USoundClass*, FSoundClassProperties>::TConstIterator It( SoundClasses ); It; ++It )
+		for (TMap<USoundClass*, FSoundClassProperties>::TConstIterator It(SoundClasses); It; ++It)
 		{
 			USoundClass* SoundClass = It.Key();
-			if (SoundClass && SoundClass->ParentClass == NULL)
+			if (SoundClass && SoundClass->ParentClass == nullptr)
 			{
 				SoundClassesToShow.Add(SoundClass);
 			}
@@ -522,7 +524,7 @@ void FAudioDevice::ShowSoundClassHierarchy( FOutputDevice& Ar, USoundClass* InSo
 		{
 			if (SoundClass->ChildClasses[i])
 			{
-				ShowSoundClassHierarchy( Ar, SoundClass->ChildClasses[i], Indent+1);
+				ShowSoundClassHierarchy(Ar, SoundClass->ChildClasses[i], Indent+1);
 			}
 		}
 	}
@@ -535,15 +537,15 @@ float AverageNativeLength = 0.f;
 TMap<int32, int32> NativeChannelCount;
 TMap<int32, int32> NativeSampleRateCount;
 
-bool FAudioDevice::HandleDumpSoundInfoCommand( const TCHAR* Cmd, FOutputDevice& Ar )
+bool FAudioDevice::HandleDumpSoundInfoCommand(const TCHAR* Cmd, FOutputDevice& Ar)
 {
 	Ar.Logf(TEXT("Native Count: %d\nRealtime Count: %d\n"), PrecachedNative, PrecachedRealtime);
 	float AverageSize = 0.0f;
-	if( PrecachedNative != 0 )
+	if (PrecachedNative != 0)
 	{
 		PrecachedNative = TotalNativeSize / PrecachedNative;
 	}
-	Ar.Logf(TEXT("Average Length: %.3g\nTotal Size: %d\nAverage Size: %.3g\n"), AverageNativeLength, TotalNativeSize, PrecachedNative );
+	Ar.Logf(TEXT("Average Length: %.3g\nTotal Size: %d\nAverage Size: %.3g\n"), AverageNativeLength, TotalNativeSize, PrecachedNative);
 	Ar.Logf(TEXT("Channel counts:\n"));
 	for (auto CountIt = NativeChannelCount.CreateConstIterator(); CountIt; ++CountIt)
 	{
@@ -558,102 +560,102 @@ bool FAudioDevice::HandleDumpSoundInfoCommand( const TCHAR* Cmd, FOutputDevice&
 }
 
 
-bool FAudioDevice::HandleListSoundClassVolumesCommand( const TCHAR* Cmd, FOutputDevice& Ar )
+bool FAudioDevice::HandleListSoundClassVolumesCommand(const TCHAR* Cmd, FOutputDevice& Ar)
 {
 	Ar.Logf(TEXT("SoundClass Volumes: (Volume, Pitch)"));
 
-	for( TMap<USoundClass*, FSoundClassProperties>::TIterator It( SoundClasses ); It; ++It )
+	for (TMap<USoundClass*, FSoundClassProperties>::TIterator It(SoundClasses); It; ++It)
 	{
 		USoundClass* SoundClass = It.Key();
 		if (SoundClass)
 		{
 			const FSoundClassProperties& CurClass = It.Value();
 
-			Ar.Logf( TEXT("Cur (%3.2f, %3.2f) for SoundClass %s"), CurClass.Volume, CurClass.Pitch, *SoundClass->GetName() );
+			Ar.Logf(TEXT("Cur (%3.2f, %3.2f) for SoundClass %s"), CurClass.Volume, CurClass.Pitch, *SoundClass->GetName());
 		}
 	}
 
 	return true;
 }
 
-bool FAudioDevice::HandleListAudioComponentsCommand( const TCHAR* Cmd, FOutputDevice& Ar )
+bool FAudioDevice::HandleListAudioComponentsCommand(const TCHAR* Cmd, FOutputDevice& Ar)
 {
 	int32 Count = 0;
-	Ar.Logf( TEXT( "AudioComponent Dump" ) );
-	for( TObjectIterator<UAudioComponent> It; It; ++It )
+	Ar.Logf(TEXT("AudioComponent Dump"));
+	for (TObjectIterator<UAudioComponent> It; It; ++It)
 	{
 		UAudioComponent* AudioComponent = *It;
 		UObject* Outer = It->GetOuter();
 		UObject* Owner = It->GetOwner();
-		Ar.Logf( TEXT("    0x%p: %s, %s, %s, %s"),
+		Ar.Logf(TEXT("    0x%p: %s, %s, %s, %s"),
 			AudioComponent,
-			*( It->GetPathName() ),
-			It->Sound ? *( It->Sound->GetPathName() ) : TEXT( "NO SOUND" ),
-			Outer ? *( Outer->GetPathName() ) : TEXT( "NO OUTER" ),
-			Owner ? *( Owner->GetPathName() ) : TEXT( "NO OWNER" ) );
-		Ar.Logf( TEXT( "        bAutoDestroy....................%s" ), AudioComponent->bAutoDestroy ? TEXT( "true" ) : TEXT( "false" ) );
-		Ar.Logf( TEXT( "        bStopWhenOwnerDestroyed.........%s" ), AudioComponent->bStopWhenOwnerDestroyed ? TEXT( "true" ) : TEXT( "false" ) );
-		Ar.Logf( TEXT( "        bShouldRemainActiveIfDropped....%s" ), AudioComponent->bShouldRemainActiveIfDropped ? TEXT( "true" ) : TEXT( "false" ) );
-		Ar.Logf( TEXT( "        bIgnoreForFlushing..............%s" ), AudioComponent->bIgnoreForFlushing ? TEXT( "true" ) : TEXT( "false" ) );
+			*(It->GetPathName()),
+			It->Sound ? *(It->Sound->GetPathName()) : TEXT("NO SOUND"),
+			Outer ? *(Outer->GetPathName()) : TEXT("NO OUTER"),
+			Owner ? *(Owner->GetPathName()) : TEXT("NO OWNER"));
+		Ar.Logf(TEXT("        bAutoDestroy....................%s"), AudioComponent->bAutoDestroy ? TEXT("true") : TEXT("false"));
+		Ar.Logf(TEXT("        bStopWhenOwnerDestroyed.........%s"), AudioComponent->bStopWhenOwnerDestroyed ? TEXT("true") : TEXT("false"));
+		Ar.Logf(TEXT("        bShouldRemainActiveIfDropped....%s"), AudioComponent->bShouldRemainActiveIfDropped ? TEXT("true") : TEXT("false"));
+		Ar.Logf(TEXT("        bIgnoreForFlushing..............%s"), AudioComponent->bIgnoreForFlushing ? TEXT("true") : TEXT("false"));
 		Count++;
 	}
-	Ar.Logf( TEXT( "AudioComponent Total = %d" ), Count );
+	Ar.Logf(TEXT("AudioComponent Total = %d"), Count);
 
-	Ar.Logf( TEXT( "AudioDevice 0x%p has %d ActiveSounds" ),
+	Ar.Logf(TEXT("AudioDevice 0x%p has %d ActiveSounds"),
 		this, ActiveSounds.Num());
-	for( int32 ASIndex = 0; ASIndex < ActiveSounds.Num(); ASIndex++ )
+	for (int32 ASIndex = 0; ASIndex < ActiveSounds.Num(); ASIndex++)
 	{
 		const FActiveSound* ActiveSound = ActiveSounds[ASIndex];
 		UAudioComponent* AComp = ActiveSound->GetAudioComponent();
-		if( AComp )
+		if (AComp)
 		{
-			Ar.Logf( TEXT( "    0x%p: %4d - %s, %s, %s, %s" ),
+			Ar.Logf(TEXT("    0x%p: %4d - %s, %s, %s, %s"),
 				AComp,
 				ASIndex,
-				*( AComp->GetPathName() ),
-				ActiveSound->Sound ? *( ActiveSound->Sound->GetPathName() ) : TEXT( "NO SOUND" ),
-				AComp->GetOuter() ? *( AComp->GetOuter()->GetPathName() ) : TEXT( "NO OUTER" ),
-				AComp->GetOwner() ? *( AComp->GetOwner()->GetPathName() ) : TEXT( "NO OWNER" ) );
+				*(AComp->GetPathName()),
+				ActiveSound->Sound ? *(ActiveSound->Sound->GetPathName()) : TEXT("NO SOUND"),
+				AComp->GetOuter() ? *(AComp->GetOuter()->GetPathName()) : TEXT("NO OUTER"),
+				AComp->GetOwner() ? *(AComp->GetOwner()->GetPathName()) : TEXT("NO OWNER"));
 		}
 		else
 		{
-			Ar.Logf( TEXT( "    %4d - %s, %s" ), 
+			Ar.Logf(TEXT("    %4d - %s, %s"), 
 				ASIndex, 
-				ActiveSound->Sound ? *( ActiveSound->Sound->GetPathName() ) : TEXT( "NO SOUND" ),
-				TEXT( "NO COMPONENT" ) );
+				ActiveSound->Sound ? *(ActiveSound->Sound->GetPathName()) : TEXT("NO SOUND"),
+				TEXT("NO COMPONENT"));
 		}
 	}
 	return true;
 }
 
-bool FAudioDevice::HandleListSoundDurationsCommand( const TCHAR* Cmd, FOutputDevice& Ar )
+bool FAudioDevice::HandleListSoundDurationsCommand(const TCHAR* Cmd, FOutputDevice& Ar)
 {
-	Ar.Logf(TEXT( "Sound,Duration,Channels" ) );
-	for( TObjectIterator<USoundWave> It; It; ++It )
+	Ar.Logf(TEXT("Sound,Duration,Channels"));
+	for (TObjectIterator<USoundWave> It; It; ++It)
 	{
 		USoundWave* SoundWave = *It;
-		Ar.Logf(TEXT( "%s,%f,%i" ), *SoundWave->GetPathName(), SoundWave->Duration, SoundWave->NumChannels );
+		Ar.Logf(TEXT("%s,%f,%i"), *SoundWave->GetPathName(), SoundWave->Duration, SoundWave->NumChannels);
 	}
 	return true;
 }
 
 
-bool FAudioDevice::HandlePlaySoundCueCommand( const TCHAR* Cmd, FOutputDevice& Ar )
+bool FAudioDevice::HandlePlaySoundCueCommand(const TCHAR* Cmd, FOutputDevice& Ar)
 {
 	// Stop any existing sound playing
-	if( !TestAudioComponent.IsValid() )
+	if (!TestAudioComponent.IsValid())
 	{
 		TestAudioComponent = NewObject<UAudioComponent>();
 	}
 
 	UAudioComponent* AudioComp = TestAudioComponent.Get();
-	if( AudioComp != NULL )
+	if (AudioComp != nullptr)
 	{
 		AudioComp->Stop();
 
 		// Load up an arbitrary cue
-		USoundCue* Cue = LoadObject<USoundCue>( NULL, Cmd, NULL, LOAD_None, NULL );
-		if( Cue != NULL )
+		USoundCue* Cue = LoadObject<USoundCue>(nullptr, Cmd, nullptr, LOAD_None, nullptr);
+		if (Cue != nullptr)
 		{
 			AudioComp->Sound = Cue;
 			AudioComp->bAllowSpatialization = false;
@@ -661,13 +663,13 @@ bool FAudioDevice::HandlePlaySoundCueCommand( const TCHAR* Cmd, FOutputDevice& A
 			AudioComp->Play();
 
 			TArray<USoundNodeWavePlayer*> WavePlayers;
-			Cue->RecursiveFindNode<USoundNodeWavePlayer>( Cue->FirstNode, WavePlayers );
-			for( int32 i = 0; i < WavePlayers.Num(); ++i )
+			Cue->RecursiveFindNode<USoundNodeWavePlayer>(Cue->FirstNode, WavePlayers);
+			for (int32 i = 0; i < WavePlayers.Num(); ++i)
 			{
 				USoundWave* SoundWave = WavePlayers[ i ]->GetSoundWave();
 				if (SoundWave)
 				{
-					SoundWave->LogSubtitle( Ar );
+					SoundWave->LogSubtitle(Ar);
 				}
 			}
 		}
@@ -675,41 +677,41 @@ bool FAudioDevice::HandlePlaySoundCueCommand( const TCHAR* Cmd, FOutputDevice& A
 	return true;
 }
 
-bool FAudioDevice::HandlePlaySoundWaveCommand( const TCHAR* Cmd, FOutputDevice& Ar )
+bool FAudioDevice::HandlePlaySoundWaveCommand(const TCHAR* Cmd, FOutputDevice& Ar)
 {
 	// Stop any existing sound playing
-	if( !TestAudioComponent.IsValid() )
+	if (!TestAudioComponent.IsValid())
 	{
 		TestAudioComponent = NewObject<UAudioComponent>();
 	}
 
 	UAudioComponent* AudioComp = TestAudioComponent.Get();
-	if( AudioComp != NULL )
+	if (AudioComp != nullptr)
 	{
 		AudioComp->Stop();
 
 		// Load up an arbitrary wave
-		USoundWave* Wave = LoadObject<USoundWave>( NULL, Cmd, NULL, LOAD_None, NULL );
-		if( Wave != NULL )
+		USoundWave* Wave = LoadObject<USoundWave>(NULL, Cmd, NULL, LOAD_None, NULL);
+		if (Wave != NULL)
 		{
 			AudioComp->Sound = Wave;
 			AudioComp->bAllowSpatialization = false;
 			AudioComp->bAutoDestroy = true;
 			AudioComp->Play();
 
-			Wave->LogSubtitle( Ar );
+			Wave->LogSubtitle(Ar);
 		}
 	}
 	return true;
 }
 
-bool FAudioDevice::HandleSetBaseSoundMixCommand( const TCHAR* Cmd, FOutputDevice& Ar )
+bool FAudioDevice::HandleSetBaseSoundMixCommand(const TCHAR* Cmd, FOutputDevice& Ar)
 {
-	// Ar.Logf( TEXT( "Setting base sound mix '%s'" ), Cmd );
-	FName NewMix = FName( Cmd );
+	// Ar.Logf(TEXT("Setting base sound mix '%s'"), Cmd);
+	FName NewMix = FName(Cmd);
 	USoundMix* SoundMix = NULL;
 
-	for( TObjectIterator<USoundMix> It; It; ++It )
+	for (TObjectIterator<USoundMix> It; It; ++It)
 	{
 		if (NewMix == It->GetFName())
 		{
@@ -720,7 +722,7 @@ bool FAudioDevice::HandleSetBaseSoundMixCommand( const TCHAR* Cmd, FOutputDevice
 
 	if (SoundMix)
 	{
-		SetBaseSoundMix( SoundMix );
+		SetBaseSoundMix(SoundMix);
 	}
 	else
 	{
@@ -729,63 +731,63 @@ bool FAudioDevice::HandleSetBaseSoundMixCommand( const TCHAR* Cmd, FOutputDevice
 	return true;
 }
 
-bool FAudioDevice::HandleIsolateDryAudioCommand( const TCHAR* Cmd, FOutputDevice& Ar )
+bool FAudioDevice::HandleIsolateDryAudioCommand(const TCHAR* Cmd, FOutputDevice& Ar)
 {
-	Ar.Logf( TEXT( "Dry audio isolated" ) );
+	Ar.Logf(TEXT("Dry audio isolated"));
 	DebugState = DEBUGSTATE_IsolateDryAudio;
 	return true;
 }
 
-bool FAudioDevice::HandleIsolateReverbCommand( const TCHAR* Cmd, FOutputDevice& Ar )
+bool FAudioDevice::HandleIsolateReverbCommand(const TCHAR* Cmd, FOutputDevice& Ar)
 {
-	Ar.Logf( TEXT( "Reverb audio isolated" ) );
+	Ar.Logf(TEXT("Reverb audio isolated"));
 	DebugState = DEBUGSTATE_IsolateReverb;
 	return true;
 }
 
-bool FAudioDevice::HandleTestLPFCommand( const TCHAR* Cmd, FOutputDevice& Ar )
+bool FAudioDevice::HandleTestLPFCommand(const TCHAR* Cmd, FOutputDevice& Ar)
 {
-	Ar.Logf( TEXT( "LPF set to max for all sources" ) );
+	Ar.Logf(TEXT("LPF set to max for all sources"));
 	DebugState =  DEBUGSTATE_TestLPF;
 	return true;
 }
 
-bool FAudioDevice::HandleTestStereoBleedCommand( const TCHAR* Cmd, FOutputDevice& Ar )
+bool FAudioDevice::HandleTestStereoBleedCommand(const TCHAR* Cmd, FOutputDevice& Ar)
 {
-	Ar.Logf( TEXT( "StereoBleed set to max for all sources" ) );
+	Ar.Logf(TEXT("StereoBleed set to max for all sources"));
 	DebugState =  DEBUGSTATE_TestStereoBleed;
 	return true;
 }
 
-bool FAudioDevice::HandleTestLFEBleedCommand( const TCHAR* Cmd, FOutputDevice& Ar )
+bool FAudioDevice::HandleTestLFEBleedCommand(const TCHAR* Cmd, FOutputDevice& Ar)
 {
-	Ar.Logf( TEXT( "LFEBleed set to max for all sources" ) );
+	Ar.Logf(TEXT("LFEBleed set to max for all sources"));
 	DebugState = DEBUGSTATE_TestLFEBleed;
 	return true;
 }
 
-bool FAudioDevice::HandleDisableLPFCommand( const TCHAR* Cmd, FOutputDevice& Ar )
+bool FAudioDevice::HandleDisableLPFCommand(const TCHAR* Cmd, FOutputDevice& Ar)
 {
-	Ar.Logf( TEXT( "LPF disabled for all sources" ) );
+	Ar.Logf(TEXT("LPF disabled for all sources"));
 	DebugState = DEBUGSTATE_DisableLPF;
 	return true;
 }
 
-bool FAudioDevice::HandleDisableRadioCommand( const TCHAR* Cmd, FOutputDevice& Ar )
+bool FAudioDevice::HandleDisableRadioCommand(const TCHAR* Cmd, FOutputDevice& Ar)
 {
 	EnableRadioEffect(false);
 	return true;
 }
 
-bool FAudioDevice::HandleEnableRadioCommand( const TCHAR* Cmd, FOutputDevice& Ar )
+bool FAudioDevice::HandleEnableRadioCommand(const TCHAR* Cmd, FOutputDevice& Ar)
 {
 	EnableRadioEffect(true);
 	return true;
 }
 
-bool FAudioDevice::HandleResetSoundStateCommand( const TCHAR* Cmd, FOutputDevice& Ar )
+bool FAudioDevice::HandleResetSoundStateCommand(const TCHAR* Cmd, FOutputDevice& Ar)
 {
-	Ar.Logf( TEXT( "All volumes reset to their defaults; all test filters removed" ) );
+	Ar.Logf(TEXT("All volumes reset to their defaults; all test filters removed"));
 	DebugState = DEBUGSTATE_None;
 	return true;
 }
@@ -1073,97 +1075,97 @@ bool FAudioDevice::HandleAudioMemoryInfo(const TCHAR* Cmd, FOutputDevice& Ar)
 
 #endif // !UE_BUILD_SHIPPING
 
-EDebugState FAudioDevice::GetMixDebugState( void )
+EDebugState FAudioDevice::GetMixDebugState()
 {
-	return( ( EDebugState )DebugState );
+	return((EDebugState)DebugState);
 }
 
-bool FAudioDevice::Exec( UWorld* InWorld, const TCHAR* Cmd, FOutputDevice& Ar )
+bool FAudioDevice::Exec(UWorld* InWorld, const TCHAR* Cmd, FOutputDevice& Ar)
 {
 #if !UE_BUILD_SHIPPING
-	if( FParse::Command( &Cmd, TEXT( "DumpSoundInfo" ) ) )
+	if (FParse::Command(&Cmd, TEXT("DumpSoundInfo")))
 	{
-		HandleDumpSoundInfoCommand( Cmd, Ar );
+		HandleDumpSoundInfoCommand(Cmd, Ar);
 	}
-	if( FParse::Command( &Cmd, TEXT( "ListSounds" ) ) )
+	if (FParse::Command(&Cmd, TEXT("ListSounds")))
 	{
-		return HandleListSoundsCommand( Cmd, Ar );
+		return HandleListSoundsCommand(Cmd, Ar);
 	}
-	else if( FParse::Command( &Cmd, TEXT( "ListWaves" ) ) )
+	else if (FParse::Command(&Cmd, TEXT("ListWaves")))
 	{
-		return HandleListWavesCommand( Cmd, Ar );
+		return HandleListWavesCommand(Cmd, Ar);
 	}
-	else if( FParse::Command( &Cmd, TEXT( "ListSoundClasses" ) ) )
+	else if (FParse::Command(&Cmd, TEXT("ListSoundClasses")))
 	{
-		return HandleListSoundClassesCommand( Cmd, Ar );
+		return HandleListSoundClassesCommand(Cmd, Ar);
 	}
-	else if( FParse::Command( &Cmd, TEXT( "ShowSoundClassHierarchy" ) ) )
+	else if (FParse::Command(&Cmd, TEXT("ShowSoundClassHierarchy")))
 	{
-		return HandleShowSoundClassHierarchyCommand( Cmd, Ar );
+		return HandleShowSoundClassHierarchyCommand(Cmd, Ar);
 	}
-	else if( FParse::Command( &Cmd, TEXT( "ListSoundClassVolumes" ) ) )
+	else if (FParse::Command(&Cmd, TEXT("ListSoundClassVolumes")))
 	{
-		return HandleListSoundClassVolumesCommand( Cmd, Ar );
+		return HandleListSoundClassVolumesCommand(Cmd, Ar);
 	}
-	else if( FParse::Command( &Cmd, TEXT( "ListAudioComponents" ) ) )
+	else if (FParse::Command(&Cmd, TEXT("ListAudioComponents")))
 	{
-		return HandleListAudioComponentsCommand( Cmd, Ar );
+		return HandleListAudioComponentsCommand(Cmd, Ar);
 	}
-	else if( FParse::Command( &Cmd, TEXT( "ListSoundDurations" ) ) )
+	else if (FParse::Command(&Cmd, TEXT("ListSoundDurations")))
 	{
-		return HandleListSoundDurationsCommand( Cmd, Ar );
+		return HandleListSoundDurationsCommand(Cmd, Ar);
 	}
-	else if( FParse::Command( &Cmd, TEXT( "PlaySoundCue" ) ) )
+	else if (FParse::Command(&Cmd, TEXT("PlaySoundCue")))
 	{
-		return HandlePlaySoundCueCommand( Cmd, Ar );
+		return HandlePlaySoundCueCommand(Cmd, Ar);
 	}
-	else if( FParse::Command( &Cmd, TEXT( "PlaySoundWave" ) ) )
+	else if (FParse::Command(&Cmd, TEXT("PlaySoundWave")))
 	{
-		return HandlePlaySoundWaveCommand( Cmd, Ar );
+		return HandlePlaySoundWaveCommand(Cmd, Ar);
 	}
-	else if( FParse::Command( &Cmd, TEXT( "SetBaseSoundMix" ) ) )
+	else if (FParse::Command(&Cmd, TEXT("SetBaseSoundMix")))
 	{
-		return HandleSetBaseSoundMixCommand( Cmd, Ar );
+		return HandleSetBaseSoundMixCommand(Cmd, Ar);
 	}
-	else if( FParse::Command( &Cmd, TEXT( "IsolateDryAudio" ) ) )
+	else if (FParse::Command(&Cmd, TEXT("IsolateDryAudio")))
 	{
-		return HandleIsolateDryAudioCommand( Cmd, Ar );
+		return HandleIsolateDryAudioCommand(Cmd, Ar);
 	}
-	else if( FParse::Command( &Cmd, TEXT( "IsolateReverb" ) ) )
+	else if (FParse::Command(&Cmd, TEXT("IsolateReverb")))
 	{
-		return HandleIsolateReverbCommand( Cmd, Ar );
+		return HandleIsolateReverbCommand(Cmd, Ar);
 	}
-	else if( FParse::Command( &Cmd, TEXT( "TestLPF" ) ) )
+	else if (FParse::Command(&Cmd, TEXT("TestLPF")))
 	{
-		return HandleTestLPFCommand( Cmd, Ar ); 
+		return HandleTestLPFCommand(Cmd, Ar); 
 	}
-	else if( FParse::Command( &Cmd, TEXT( "TestStereoBleed" ) ) )
+	else if (FParse::Command(&Cmd, TEXT("TestStereoBleed")))
 	{
-		return HandleTestStereoBleedCommand( Cmd, Ar );
+		return HandleTestStereoBleedCommand(Cmd, Ar);
 	}
-	else if( FParse::Command( &Cmd, TEXT( "TestLFEBleed" ) ) )
+	else if (FParse::Command(&Cmd, TEXT("TestLFEBleed")))
 	{
-		return HandleTestLPFCommand( Cmd, Ar );
+		return HandleTestLPFCommand(Cmd, Ar);
 	}
-	else if( FParse::Command( &Cmd, TEXT( "DisableLPF" ) ) )
+	else if (FParse::Command(&Cmd, TEXT("DisableLPF")))
 	{
-		return HandleDisableLPFCommand( Cmd, Ar );
+		return HandleDisableLPFCommand(Cmd, Ar);
 	}
-	else if( FParse::Command( &Cmd, TEXT( "DisableRadio" ) ) )
+	else if (FParse::Command(&Cmd, TEXT("DisableRadio")))
 	{
-		return HandleDisableRadioCommand( Cmd, Ar );
+		return HandleDisableRadioCommand(Cmd, Ar);
 	}
-	else if( FParse::Command( &Cmd, TEXT( "EnableRadio" ) ) )
+	else if (FParse::Command(&Cmd, TEXT("EnableRadio")))
 	{
-		return HandleEnableRadioCommand( Cmd, Ar );
+		return HandleEnableRadioCommand(Cmd, Ar);
 	}
-	else if( FParse::Command( &Cmd, TEXT( "ResetSoundState" ) ) )
+	else if (FParse::Command(&Cmd, TEXT("ResetSoundState")))
 	{
-		return HandleResetSoundStateCommand( Cmd, Ar );
+		return HandleResetSoundStateCommand(Cmd, Ar);
 	}
 	else if (FParse::Command(&Cmd, TEXT("ToggleSpatExt")))
 	{
-		return HandleToggleSpatializationExtensionCommand( Cmd, Ar);
+		return HandleToggleSpatializationExtensionCommand(Cmd, Ar);
 	}
 	else if (FParse::Command(&Cmd, TEXT("ToggleHRTFForAll")))
 	{
@@ -1194,20 +1196,20 @@ bool FAudioDevice::Exec( UWorld* InWorld, const TCHAR* Cmd, FOutputDevice& Ar )
 	return false;
 }
 
-void FAudioDevice::InitSoundClasses( void )
+void FAudioDevice::InitSoundClasses()
 {
 	// Reset the maps of sound class properties
-	for( TObjectIterator<USoundClass> It; It; ++It )
+	for (TObjectIterator<USoundClass> It; It; ++It)
 	{
 		USoundClass* SoundClass = *It;
-		FSoundClassProperties& Properties = SoundClasses.Add( SoundClass, SoundClass->Properties );
+		FSoundClassProperties& Properties = SoundClasses.Add(SoundClass, SoundClass->Properties);
 	}
 
 	// Propagate the properties down the hierarchy 
 	ParseSoundClasses();
 }
 
-void FAudioDevice::InitSoundSources( void )
+void FAudioDevice::InitSoundSources()
 {
 	if (Sources.Num() == 0)
 	{
@@ -1223,7 +1225,7 @@ void FAudioDevice::InitSoundSources( void )
 	}
 }
 
-void FAudioDevice::SetDefaultBaseSoundMix( USoundMix* SoundMix )
+void FAudioDevice::SetDefaultBaseSoundMix(USoundMix* SoundMix)
 {
 	if (SoundMix == NULL)
 	{
@@ -1238,12 +1240,12 @@ void FAudioDevice::SetDefaultBaseSoundMix( USoundMix* SoundMix )
 	SetBaseSoundMix(SoundMix);
 }
 
-void FAudioDevice::RemoveSoundMix( USoundMix* SoundMix )
+void FAudioDevice::RemoveSoundMix(USoundMix* SoundMix)
 {
-	if( SoundMix )
+	if (SoundMix)
 	{
 		// Not sure if we will ever destroy the default base SoundMix
-		if( SoundMix == DefaultBaseSoundMix )
+		if (SoundMix == DefaultBaseSoundMix)
 		{
 			DefaultBaseSoundMix = NULL;
 		}
@@ -1251,7 +1253,7 @@ void FAudioDevice::RemoveSoundMix( USoundMix* SoundMix )
 		ClearSoundMix(SoundMix);
 
 		// Try setting to global default if base SoundMix has been cleared
-		if( BaseSoundMix == NULL )
+		if (BaseSoundMix == NULL)
 		{
 			SetBaseSoundMix(DefaultBaseSoundMix);
 		}
@@ -1259,17 +1261,17 @@ void FAudioDevice::RemoveSoundMix( USoundMix* SoundMix )
 }
 
 
-void FAudioDevice::RecurseIntoSoundClasses( USoundClass* CurrentClass, FSoundClassProperties& ParentProperties )
+void FAudioDevice::RecurseIntoSoundClasses(USoundClass* CurrentClass, FSoundClassProperties& ParentProperties)
 {
 	// Iterate over all child nodes and recurse.
-	for( int32 ChildIndex = 0; ChildIndex < CurrentClass->ChildClasses.Num(); ChildIndex++ )
+	for (int32 ChildIndex = 0; ChildIndex < CurrentClass->ChildClasses.Num(); ChildIndex++)
 	{
 		// Look up class and propagated properties.
 		USoundClass* ChildClass = CurrentClass->ChildClasses[ChildIndex];
 		FSoundClassProperties* Properties = SoundClasses.Find(ChildClass);
 
 		// Should never be NULL for a properly set up tree.
-		if( ChildClass )
+		if (ChildClass)
 		{
 			if (Properties)
 			{
@@ -1282,11 +1284,11 @@ void FAudioDevice::RecurseIntoSoundClasses( USoundClass* CurrentClass, FSoundCla
 				// VoiceCenterChannelVolume, RadioFilterVolume, RadioFilterVolumeThreshold, bApplyEffects, BleedStereo, bReverb, and bCenterChannelOnly do not propagate (sub-classes can be non-zero even if parent class is zero)
 
 				// ... and recurse into child nodes.
-				RecurseIntoSoundClasses( ChildClass, *Properties );
+				RecurseIntoSoundClasses(ChildClass, *Properties);
 			}
 			else
 			{
-				UE_LOG(LogAudio, Warning, TEXT( "Couldn't find child class properties - sound class functionality will not work correctly! CurrentClass: %s ChildClass: %s" ), *CurrentClass->GetFullName(), *ChildClass->GetFullName() );
+				UE_LOG(LogAudio, Warning, TEXT("Couldn't find child class properties - sound class functionality will not work correctly! CurrentClass: %s ChildClass: %s"), *CurrentClass->GetFullName(), *ChildClass->GetFullName());
 			}
 		}
 	}
@@ -1308,7 +1310,7 @@ void FAudioDevice::ParseSoundClasses()
 	TArray<USoundClass*> RootSoundClasses;
 
 	// Reset to known state - preadjusted by set class volume calls
-	for( TMap<USoundClass*, FSoundClassProperties>::TIterator It( SoundClasses ); It; ++It )
+	for (TMap<USoundClass*, FSoundClassProperties>::TIterator It(SoundClasses); It; ++It)
 	{
 		USoundClass* SoundClass = It.Key();
 		if (SoundClass)
@@ -1325,22 +1327,22 @@ void FAudioDevice::ParseSoundClasses()
 	{
 		USoundClass* RootSoundClass = RootSoundClasses[RootIndex];
 
-		FSoundClassProperties* RootSoundClassProperties = SoundClasses.Find( RootSoundClass );
-		if( RootSoundClass && RootSoundClassProperties )
+		FSoundClassProperties* RootSoundClassProperties = SoundClasses.Find(RootSoundClass);
+		if (RootSoundClass && RootSoundClassProperties)
 		{
 			// Follow the tree.
-			RecurseIntoSoundClasses( RootSoundClass, *RootSoundClassProperties );
+			RecurseIntoSoundClasses(RootSoundClass, *RootSoundClassProperties);
 		}
 	}
 }
 
 
-void FAudioDevice::RecursiveApplyAdjuster( const FSoundClassAdjuster& InAdjuster, USoundClass* InSoundClass )
+void FAudioDevice::RecursiveApplyAdjuster(const FSoundClassAdjuster& InAdjuster, USoundClass* InSoundClass)
 {
 	// Find the sound class properties so we can apply the adjuster
 	// and find the sound class so we can recurse through the children
-	FSoundClassProperties* Properties = SoundClasses.Find( InSoundClass );
-	if( InSoundClass && Properties )
+	FSoundClassProperties* Properties = SoundClasses.Find(InSoundClass);
+	if (InSoundClass && Properties)
 	{
 		// Adjust this class
 		Properties->Volume *= InAdjuster.VolumeAdjuster;
@@ -1348,47 +1350,50 @@ void FAudioDevice::RecursiveApplyAdjuster( const FSoundClassAdjuster& InAdjuster
 		Properties->VoiceCenterChannelVolume *= InAdjuster.VoiceCenterChannelVolumeAdjuster;
 
 		// Recurse through this classes children
-		for( int32 ChildIdx = 0; ChildIdx < InSoundClass->ChildClasses.Num(); ++ChildIdx )
+		for (int32 ChildIdx = 0; ChildIdx < InSoundClass->ChildClasses.Num(); ++ChildIdx)
 		{
 			if (InSoundClass->ChildClasses[ ChildIdx ])
 			{
-				RecursiveApplyAdjuster( InAdjuster, InSoundClass->ChildClasses[ ChildIdx ] );
+				RecursiveApplyAdjuster(InAdjuster, InSoundClass->ChildClasses[ ChildIdx ]);
 			}
 		}
 	}
 	else
 	{
-		UE_LOG(LogAudio, Warning, TEXT( "Sound class '%s' does not exist" ), *InSoundClass->GetName() );
+		UE_LOG(LogAudio, Warning, TEXT("Sound class '%s' does not exist"), *InSoundClass->GetName());
 	}
 }
 
-bool FAudioDevice::ApplySoundMix( USoundMix* NewMix, FSoundMixState* SoundMixState )
+bool FAudioDevice::ApplySoundMix(USoundMix* NewMix, FSoundMixState* SoundMixState)
 {
-	if( NewMix && SoundMixState )
+	if (NewMix && SoundMixState)
 	{
-		UE_LOG(LogAudio, Log, TEXT( "FAudioDevice::ApplySoundMix(): %s" ), *NewMix->GetName() );
+		UE_LOG(LogAudio, Log, TEXT("FAudioDevice::ApplySoundMix(): %s"), *NewMix->GetName());
 
 		SoundMixState->StartTime = FApp::GetCurrentTime();
 		SoundMixState->FadeInStartTime = SoundMixState->StartTime + NewMix->InitialDelay;
 		SoundMixState->FadeInEndTime = SoundMixState->FadeInStartTime + NewMix->FadeInTime;
 		SoundMixState->FadeOutStartTime = -1.0;
 		SoundMixState->EndTime = -1.0;
-		if( NewMix->Duration >= 0.0f )
+		if (NewMix->Duration >= 0.0f)
 		{
 			SoundMixState->FadeOutStartTime = SoundMixState->FadeInEndTime + NewMix->Duration;
 			SoundMixState->EndTime = SoundMixState->FadeOutStartTime + NewMix->FadeOutTime;
 		}
 		SoundMixState->InterpValue = 0.0f;
 
-		ApplyClassAdjusters(NewMix,SoundMixState->InterpValue);
+		// On sound mix application, there is no delta time
+		const float InitDeltaTime = 0.0f;
 
-		return( true );
+		ApplyClassAdjusters(NewMix, SoundMixState->InterpValue, InitDeltaTime);
+
+		return(true);
 	}
 
-	return( false );
+	return(false);
 }
 
-void FAudioDevice::UpdateSoundMix(class USoundMix* SoundMix, FSoundMixState* SoundMixState)
+void FAudioDevice::UpdateSoundMix(USoundMix* SoundMix, FSoundMixState* SoundMixState)
 {
 	// If this SoundMix will automatically end, add some more time
 	if (SoundMixState->FadeOutStartTime >= 0.f)
@@ -1400,14 +1405,14 @@ void FAudioDevice::UpdateSoundMix(class USoundMix* SoundMix, FSoundMixState* Sou
 			SoundMixState->FadeInEndTime = SoundMixState->FadeInStartTime + SoundMix->FadeInTime;
 			SoundMixState->FadeOutStartTime = -1.0;
 			SoundMixState->EndTime = -1.0;
-			if( SoundMix->Duration >= 0.0f )
+			if (SoundMix->Duration >= 0.0f)
 			{
 				SoundMixState->FadeOutStartTime = SoundMixState->FadeInEndTime + SoundMix->Duration;
 				SoundMixState->EndTime = SoundMixState->FadeOutStartTime + SoundMix->FadeOutTime;
 			}
 
 			// Might have already started fading EQ effect so try setting again
-			Effects->SetMixSettings( SoundMix );
+			Effects->SetMixSettings(SoundMix);
 		}
 		else if (SoundMixState->CurrentState == ESoundMixState::Active)
 		{
@@ -1427,13 +1432,13 @@ void FAudioDevice::UpdatePassiveSoundMixModifiers(TArray<FWaveInstance*>& WaveIn
 	TArray<USoundMix*> CurrPassiveSoundMixModifiers;
 
 	// Find all passive SoundMixes from currently active wave instances
-	for( int32 WaveIndex = FirstActiveIndex; WaveIndex < WaveInstances.Num(); WaveIndex++ )
+	for (int32 WaveIndex = FirstActiveIndex; WaveIndex < WaveInstances.Num(); WaveIndex++)
 	{
 		FWaveInstance* WaveInstance = WaveInstances[WaveIndex];
 		if (WaveInstance)
 		{
 			USoundClass* SoundClass = WaveInstance->SoundClass;
-			if( SoundClass ) 
+			if (SoundClass) 
 			{
 				const float WaveInstanceActualVolume = WaveInstance->GetActualVolume();
 				// Check each SoundMix individually for volume levels
@@ -1441,7 +1446,7 @@ void FAudioDevice::UpdatePassiveSoundMixModifiers(TArray<FWaveInstance*>& WaveIn
 				{
 					if (WaveInstanceActualVolume >= PassiveSoundMixModifier.MinVolumeThreshold && WaveInstanceActualVolume <= PassiveSoundMixModifier.MaxVolumeThreshold)
 					{
-						CurrPassiveSoundMixModifiers.AddUnique( PassiveSoundMixModifier.SoundMix );
+						CurrPassiveSoundMixModifiers.AddUnique(PassiveSoundMixModifier.SoundMix);
 					}
 				}
 			}
@@ -1517,7 +1522,7 @@ bool FAudioDevice::TryClearingSoundMix(USoundMix* SoundMix, FSoundMixState* Soun
 	return false;
 }
 
-bool FAudioDevice::TryClearingEQSoundMix(class USoundMix* SoundMix)
+bool FAudioDevice::TryClearingEQSoundMix(USoundMix* SoundMix)
 {
 	if (SoundMix && Effects->GetCurrentEQMix() == SoundMix)
 	{
@@ -1563,62 +1568,192 @@ USoundMix* FAudioDevice::FindNextHighestEQPrioritySoundMix(USoundMix* IgnoredSou
 	return NextEQMix;
 }
 
-void FAudioDevice::ClearSoundMix(class USoundMix* SoundMix)
+void FAudioDevice::ClearSoundMix(USoundMix* SoundMix)
 {
-	if( SoundMix == BaseSoundMix )
+	if (SoundMix == BaseSoundMix)
 	{
 		BaseSoundMix = NULL;
 	}
-	SoundMixModifiers.Remove( SoundMix );
-	PrevPassiveSoundMixModifiers.Remove( SoundMix );
+	SoundMixModifiers.Remove(SoundMix);
+	PrevPassiveSoundMixModifiers.Remove(SoundMix);
+
+	// Check if there are any overrides for this sound mix and if so, reset them so that next time this sound mix is applied, it'll get the new override values
+	FSoundMixClassOverrideMap* SoundMixOverrideMap = SoundMixClassEffectOverrides.Find(SoundMix);
+	if (SoundMixOverrideMap)
+	{
+		for (auto& Entry : *SoundMixOverrideMap)
+		{
+			Entry.Value.bOverrideApplied = false;
+		}
+	}
 
 	TryClearingEQSoundMix(SoundMix);
 }
 
-void FAudioDevice::ApplyClassAdjusters(USoundMix* SoundMix, float InterpValue)
+/** Static helper function which handles setting and updating the sound class adjuster override */
+static void UpdateClassAdjustorOverrideEntry(FSoundClassAdjuster& ClassAdjustor, FSoundMixClassOverride& ClassAdjusterOverride, float DeltaTime)
+{
+	// If we've already applied the override in a previous frame
+	if (ClassAdjusterOverride.bOverrideApplied)
+	{
+		// If we've received a new override value since our last update, then just set the dynamic parameters to the new value
+		// The dynamic parameter objects will automatically smoothly travel to the new target value from its current value in the given time
+		if (ClassAdjusterOverride.bOverrideChanged)
+		{
+			ClassAdjusterOverride.PitchOverride.Set(ClassAdjusterOverride.SoundClassAdjustor.PitchAdjuster, ClassAdjusterOverride.FadeInTime);
+			ClassAdjusterOverride.VolumeOverride.Set(ClassAdjusterOverride.SoundClassAdjustor.VolumeAdjuster, ClassAdjusterOverride.FadeInTime);
+		}
+		else
+		{
+			// We haven't changed so just update the override this frame
+			ClassAdjusterOverride.PitchOverride.Update(DeltaTime);
+			ClassAdjusterOverride.VolumeOverride.Update(DeltaTime);
+		}
+	}
+	else
+	{
+		// We haven't yet applied the override to the mix, so set the override dynamic parameters to immediately
+		// have the current class adjuster values (0.0 interp-time), then set the dynamic parameters to the new target values in the given fade time
+
+		ClassAdjusterOverride.VolumeOverride.Set(ClassAdjustor.VolumeAdjuster, 0.0f);
+		ClassAdjusterOverride.VolumeOverride.Set(ClassAdjusterOverride.SoundClassAdjustor.VolumeAdjuster, ClassAdjusterOverride.FadeInTime);
+
+		ClassAdjusterOverride.PitchOverride.Set(ClassAdjustor.PitchAdjuster, 0.0f);
+		ClassAdjusterOverride.PitchOverride.Set(ClassAdjusterOverride.SoundClassAdjustor.PitchAdjuster, ClassAdjusterOverride.FadeInTime);
+	}
+
+	if (!ClassAdjustor.SoundClassObject)
+	{
+		ClassAdjustor.SoundClassObject = ClassAdjusterOverride.SoundClassAdjustor.SoundClassObject;
+	}
+
+	check(ClassAdjustor.SoundClassObject == ClassAdjusterOverride.SoundClassAdjustor.SoundClassObject);
+
+	// Get the current value of the dynamic parameters
+	ClassAdjustor.PitchAdjuster = ClassAdjusterOverride.PitchOverride.GetValue();
+	ClassAdjustor.VolumeAdjuster = ClassAdjusterOverride.VolumeOverride.GetValue();
+
+	// Override the apply to children if applicable
+	ClassAdjustor.bApplyToChildren = ClassAdjusterOverride.SoundClassAdjustor.bApplyToChildren;
+
+	// Reset the flags on the override adjuster
+	ClassAdjusterOverride.bOverrideApplied = true;
+	ClassAdjusterOverride.bOverrideChanged = false;
+}
+
+void FAudioDevice::ApplyClassAdjusters(USoundMix* SoundMix, float InterpValue, float DeltaTime)
 {
-	if( SoundMix )
+	if (!SoundMix)
+	{
+		return;
+	}
+
+	// Check if there is a sound mix override entry
+	FSoundMixClassOverrideMap* SoundMixOverrideMap = SoundMixClassEffectOverrides.Find(SoundMix);
+
+	// Create a ptr to the array of sound class adjusters ers we want to actually use. Default to using the sound class effects adjuster list
+	TArray<FSoundClassAdjuster>* SoundClassAdjusters = &SoundMix->SoundClassEffects;
+
+	bool bUsingOverride = false;
+
+	// If we have an override for this sound mix, replace any overrides and/or add to the array if the sound class adjustment entry doesn't exist
+	if (SoundMixOverrideMap)
 	{
-		// Adjust the sound class properties non recursively
-		TArray<FSoundClassAdjuster>& Adjusters = SoundMix->SoundClassEffects;
+		// If we have an override map, create a copy of the sound class adjusters for the sound mix, then override the sound mix class overrides
+		SoundClassAdjustersCopy = SoundMix->SoundClassEffects;
+
+		// Use the copied list 
+		SoundClassAdjusters = &SoundClassAdjustersCopy;
+
+		bUsingOverride = true;
 
-		for( int32 i = 0; i < Adjusters.Num(); i++ )
+		// Get the interpolated values of the vanilla adjusters up-front
+		for (FSoundClassAdjuster& Entry : *SoundClassAdjusters)
 		{
-			if (Adjusters[ i ].SoundClassObject)
+			if (Entry.SoundClassObject)
 			{
-				if( Adjusters[ i ].bApplyToChildren )
+				Entry.VolumeAdjuster = InterpolateAdjuster(Entry.VolumeAdjuster, InterpValue);
+				Entry.PitchAdjuster = InterpolateAdjuster(Entry.PitchAdjuster, InterpValue);
+				Entry.VoiceCenterChannelVolumeAdjuster = InterpolateAdjuster(Entry.VoiceCenterChannelVolumeAdjuster, InterpValue);
+			}
+		}
+
+		for (auto& SoundMixOverrideEntry : *SoundMixOverrideMap)
+		{
+			// Get the sound class object of the override
+			FSoundMixClassOverride& ClassAdjusterOverride = SoundMixOverrideEntry.Value;
+			USoundClass* SoundClassObject = ClassAdjusterOverride.SoundClassAdjustor.SoundClassObject;
+
+			// Look for it in the adjusters copy 
+			bool bSoundClassAdjustorExisted = false;
+			for (FSoundClassAdjuster& Entry : *SoundClassAdjusters)
+			{
+				// If we found it, then we need to override the volume and pitch values of the adjuster entry
+				if (Entry.SoundClassObject == SoundClassObject)
 				{
-					// Take a copy of this adjuster and pre-calculate interpolation
-					FSoundClassAdjuster InterpolatedAdjuster = Adjusters[ i ];
-					InterpolatedAdjuster.VolumeAdjuster = InterpolateAdjuster(InterpolatedAdjuster.VolumeAdjuster, InterpValue);
-					InterpolatedAdjuster.PitchAdjuster = InterpolateAdjuster(InterpolatedAdjuster.PitchAdjuster, InterpValue);
-					InterpolatedAdjuster.VoiceCenterChannelVolumeAdjuster = InterpolateAdjuster(InterpolatedAdjuster.VoiceCenterChannelVolumeAdjuster, InterpValue);
-
-					// Apply the adjuster the sound class specified by the adjuster and all its children
-					RecursiveApplyAdjuster( InterpolatedAdjuster, Adjusters[ i ].SoundClassObject );
+					// Flag that we don't need to add it to the SoundClassAdjustorsCopy
+					bSoundClassAdjustorExisted = true;
+
+					UpdateClassAdjustorOverrideEntry(Entry, ClassAdjusterOverride, DeltaTime);
+					break;
 				}
-				else
+			}
+
+			// If we didn't find an existing sound class we need to add the override to the adjuster copy
+			if (!bSoundClassAdjustorExisted)
+			{
+				// Create a default sound class adjuster (1.0 values for pitch and volume)
+				FSoundClassAdjuster NewEntry;
+
+				// Apply and/or update the override
+				UpdateClassAdjustorOverrideEntry(NewEntry, ClassAdjusterOverride, DeltaTime);
+
+				// Add the new sound class adjuster entry to the array
+				SoundClassAdjusters->Add(NewEntry);
+			}
+		}
+	}
+
+	// Loop through the sound class adjusters, everything should be up-to-date
+	for (FSoundClassAdjuster& Entry : *SoundClassAdjusters)
+	{
+		if (Entry.SoundClassObject)
+		{
+			if (Entry.bApplyToChildren)
+			{
+				RecursiveApplyAdjuster(Entry, Entry.SoundClassObject);
+			}
+			else
+			{
+				// Apply the adjuster to only the sound class specified by the adjuster
+				FSoundClassProperties* Properties = SoundClasses.Find(Entry.SoundClassObject);
+				if (Properties)
 				{
-					// Apply the adjuster to only the sound class specified by the adjuster
-					FSoundClassProperties* Properties = SoundClasses.Find( Adjusters[ i ].SoundClassObject );
-					if( Properties )
+					// If we are using an override, we've already interpolated all our dynamic parameters
+					if (bUsingOverride)
 					{
-						Properties->Volume *= InterpolateAdjuster(Adjusters[ i ].VolumeAdjuster, InterpValue);
-						Properties->Pitch *= InterpolateAdjuster(Adjusters[ i ].PitchAdjuster, InterpValue);
-						Properties->VoiceCenterChannelVolume *= InterpolateAdjuster(Adjusters[ i ].VoiceCenterChannelVolumeAdjuster, InterpValue);
+						Properties->Volume *= Entry.VolumeAdjuster;
+						Properties->Pitch *= Entry.PitchAdjuster;
+						Properties->VoiceCenterChannelVolume *= Entry.VoiceCenterChannelVolumeAdjuster;
 					}
+					// Otherwise, we need to use the "static" data and compute the adjustment interpolations now
 					else
 					{
-						UE_LOG(LogAudio, Warning, TEXT( "Sound class '%s' does not exist" ), *Adjusters[ i ].SoundClassObject->GetName() );
+						Properties->Volume *= InterpolateAdjuster(Entry.VolumeAdjuster, InterpValue);
+						Properties->Pitch *= InterpolateAdjuster(Entry.PitchAdjuster, InterpValue);
+						Properties->VoiceCenterChannelVolume *= InterpolateAdjuster(Entry.VoiceCenterChannelVolumeAdjuster, InterpValue);
 					}
 				}
+				else
+				{
+					UE_LOG(LogAudio, Warning, TEXT("Sound class '%s' does not exist"), *Entry.SoundClassObject->GetName());
+				}
 			}
 		}
 	}
 }
 
-
-void FAudioDevice::UpdateSoundClassProperties()
+void FAudioDevice::UpdateSoundClassProperties(float DeltaTime)
 {
 	// Remove SoundMix modifications and propagate the properties down the hierarchy
 	ParseSoundClasses();
@@ -1628,30 +1763,30 @@ void FAudioDevice::UpdateSoundClassProperties()
 		FSoundMixState* SoundMixState = &(It.Value());
 
 		// Initial delay before mix is applied
-		if( FApp::GetCurrentTime() >= SoundMixState->StartTime && FApp::GetCurrentTime() < SoundMixState->FadeInStartTime )
+		if (FApp::GetCurrentTime() >= SoundMixState->StartTime && FApp::GetCurrentTime() < SoundMixState->FadeInStartTime)
 		{
 			SoundMixState->InterpValue = 0.0f;
 			SoundMixState->CurrentState = ESoundMixState::Inactive;
 		}
-		else if( FApp::GetCurrentTime() >= SoundMixState->FadeInStartTime && FApp::GetCurrentTime() < SoundMixState->FadeInEndTime )
+		else if (FApp::GetCurrentTime() >= SoundMixState->FadeInStartTime && FApp::GetCurrentTime() < SoundMixState->FadeInEndTime)
 		{
 			// Work out the fade in portion
-			SoundMixState->InterpValue = ( float )( ( FApp::GetCurrentTime() - SoundMixState->FadeInStartTime ) / ( SoundMixState->FadeInEndTime - SoundMixState->FadeInStartTime ) );
+			SoundMixState->InterpValue = (float)((FApp::GetCurrentTime() - SoundMixState->FadeInStartTime) / (SoundMixState->FadeInEndTime - SoundMixState->FadeInStartTime));
 			SoundMixState->CurrentState = ESoundMixState::FadingIn;
 		}
-		else if( FApp::GetCurrentTime() >= SoundMixState->FadeInEndTime
-			&& ( SoundMixState->IsBaseSoundMix || SoundMixState->PassiveRefCount > 0 || SoundMixState->FadeOutStartTime < 0.f || FApp::GetCurrentTime() < SoundMixState->FadeOutStartTime ) )
+		else if (FApp::GetCurrentTime() >= SoundMixState->FadeInEndTime
+			&& (SoundMixState->IsBaseSoundMix || SoundMixState->PassiveRefCount > 0 || SoundMixState->FadeOutStartTime < 0.f || FApp::GetCurrentTime() < SoundMixState->FadeOutStartTime))
 		{
 			// .. ensure the full mix is applied between the end of the fade in time and the start of the fade out time
 			// or if SoundMix is the base or active via a passive push - ignores duration.
 			SoundMixState->InterpValue = 1.0f;
 			SoundMixState->CurrentState = ESoundMixState::Active;
 		}
-		else if( FApp::GetCurrentTime() >= SoundMixState->FadeOutStartTime && FApp::GetCurrentTime() < SoundMixState->EndTime )
+		else if (FApp::GetCurrentTime() >= SoundMixState->FadeOutStartTime && FApp::GetCurrentTime() < SoundMixState->EndTime)
 		{
 			// Work out the fade out portion
-			SoundMixState->InterpValue = 1.0f - ( float )( ( FApp::GetCurrentTime() - SoundMixState->FadeOutStartTime ) / ( SoundMixState->EndTime - SoundMixState->FadeOutStartTime ) );
-			if( SoundMixState->CurrentState != ESoundMixState::FadingOut )
+			SoundMixState->InterpValue = 1.0f - (float)((FApp::GetCurrentTime() - SoundMixState->FadeOutStartTime) / (SoundMixState->EndTime - SoundMixState->FadeOutStartTime));
+			if (SoundMixState->CurrentState != ESoundMixState::FadingOut)
 			{
 				// Start fading EQ at same time
 				SoundMixState->CurrentState = ESoundMixState::FadingOut;
@@ -1660,46 +1795,45 @@ void FAudioDevice::UpdateSoundClassProperties()
 		}
 		else
 		{
-			check( SoundMixState->EndTime >= 0.f && FApp::GetCurrentTime() >= SoundMixState->EndTime );
+			check(SoundMixState->EndTime >= 0.f && FApp::GetCurrentTime() >= SoundMixState->EndTime);
 			// Clear the effect of this SoundMix - may need to revisit for passive
 			SoundMixState->InterpValue = 0.0f;
 			SoundMixState->CurrentState = ESoundMixState::AwaitingRemoval;
 		}
 
-		ApplyClassAdjusters(It.Key(), SoundMixState->InterpValue);
+		ApplyClassAdjusters(It.Key(), SoundMixState->InterpValue, DeltaTime);
 
-		if( SoundMixState->CurrentState == ESoundMixState::AwaitingRemoval )
+		if (SoundMixState->CurrentState == ESoundMixState::AwaitingRemoval)
 		{
 			ClearSoundMix(It.Key());
 		}
 	}
 }
 
-
-float FListener::Interpolate( const double EndTime )
+float FListener::Interpolate(const double EndTime)
 {
-	if( FApp::GetCurrentTime() < InteriorStartTime )
+	if (FApp::GetCurrentTime() < InteriorStartTime)
 	{
-		return( 0.0f );
+		return(0.0f);
 	}
 
-	if( FApp::GetCurrentTime() >= EndTime )
+	if (FApp::GetCurrentTime() >= EndTime)
 	{
-		return( 1.0f );
+		return(1.0f);
 	}
 
-	float InterpValue = ( float )( ( FApp::GetCurrentTime() - InteriorStartTime ) / ( EndTime - InteriorStartTime ) );
-	return( InterpValue );
+	float InterpValue = (float)((FApp::GetCurrentTime() - InteriorStartTime) / (EndTime - InteriorStartTime));
+	return(InterpValue);
 }
 
 
 void FListener::UpdateCurrentInteriorSettings()
 {
 	// Store the interpolation value, not the actual value
-	InteriorVolumeInterp = Interpolate( InteriorEndTime );
-	ExteriorVolumeInterp = Interpolate( ExteriorEndTime );
-	InteriorLPFInterp = Interpolate( InteriorLPFEndTime );
-	ExteriorLPFInterp = Interpolate( ExteriorLPFEndTime );
+	InteriorVolumeInterp = Interpolate(InteriorEndTime);
+	ExteriorVolumeInterp = Interpolate(ExteriorEndTime);
+	InteriorLPFInterp = Interpolate(InteriorLPFEndTime);
+	ExteriorLPFInterp = Interpolate(ExteriorLPFEndTime);
 }
 
 void FAudioDevice::InvalidateCachedInteriorVolumes() const
@@ -1710,9 +1844,9 @@ void FAudioDevice::InvalidateCachedInteriorVolumes() const
 	}
 }
 
-void FListener::ApplyInteriorSettings( class AAudioVolume* InVolume, const FInteriorSettings& Settings )
+void FListener::ApplyInteriorSettings(class AAudioVolume* InVolume, const FInteriorSettings& Settings)
 {
-	if( InVolume != Volume || Settings != InteriorSettings)
+	if (InVolume != Volume || Settings != InteriorSettings)
 	{
 		// Use previous/ current interpolation time if we're transitioning to the default worldsettings zone.
 		InteriorStartTime = FApp::GetCurrentTime();
@@ -1727,7 +1861,7 @@ void FListener::ApplyInteriorSettings( class AAudioVolume* InVolume, const FInte
 }
 
 
-void FAudioDevice::SetListener( const int32 InViewportIndex, const FTransform& InListenerTransform, const float InDeltaSeconds, class AAudioVolume* Volume, const FInteriorSettings& InteriorSettings )
+void FAudioDevice::SetListener(const int32 InViewportIndex, const FTransform& InListenerTransform, const float InDeltaSeconds, class AAudioVolume* Volume, const FInteriorSettings& InteriorSettings)
 {
 	FTransform ListenerTransform = InListenerTransform;
 	
@@ -1737,10 +1871,10 @@ void FAudioDevice::SetListener( const int32 InViewportIndex, const FTransform& I
 		ListenerTransform = FTransform::Identity;
 	}
 
-	if( InViewportIndex >= Listeners.Num() )
+	if (InViewportIndex >= Listeners.Num())
 	{
-		UE_LOG(LogAudio, Log, TEXT( "Resizing Listeners array: %d -> %d" ), Listeners.Num(), InViewportIndex );
-		Listeners.AddZeroed( InViewportIndex - Listeners.Num() + 1 );
+		UE_LOG(LogAudio, Log, TEXT("Resizing Listeners array: %d -> %d"), Listeners.Num(), InViewportIndex);
+		Listeners.AddZeroed(InViewportIndex - Listeners.Num() + 1);
 	}
 
 	Listeners[ InViewportIndex ].Velocity = InDeltaSeconds > 0.f ? 
@@ -1760,14 +1894,14 @@ void FAudioDevice::SetListener( const int32 InViewportIndex, const FTransform& I
 }
 
 
-bool FAudioDevice::SetBaseSoundMix( USoundMix* NewMix )
+bool FAudioDevice::SetBaseSoundMix(USoundMix* NewMix)
 {
-	if( NewMix && NewMix != BaseSoundMix )
+	if (NewMix && NewMix != BaseSoundMix)
 	{
 		USoundMix* OldBaseSoundMix = BaseSoundMix;
 		BaseSoundMix = NewMix;
 
-		if( OldBaseSoundMix )
+		if (OldBaseSoundMix)
 		{
 			FSoundMixState* OldBaseState = SoundMixModifiers.Find(OldBaseSoundMix);
 			check(OldBaseState);
@@ -1777,7 +1911,7 @@ bool FAudioDevice::SetBaseSoundMix( USoundMix* NewMix )
 
 		// Check whether this SoundMix is already active
 		FSoundMixState* ExistingState = SoundMixModifiers.Find(NewMix);
-		if( !ExistingState )
+		if (!ExistingState)
 		{
 			// First time this mix has been set - add it and setup mix modifications
 			ExistingState = &SoundMixModifiers.Add(NewMix, FSoundMixState());
@@ -1786,7 +1920,7 @@ bool FAudioDevice::SetBaseSoundMix( USoundMix* NewMix )
 			ApplySoundMix(NewMix, ExistingState);
 
 			// Use it to set EQ Settings, which will check its priority
-			Effects->SetMixSettings( NewMix );
+			Effects->SetMixSettings(NewMix);
 		}
 
 		ExistingState->IsBaseSoundMix = true;
@@ -1812,7 +1946,7 @@ void FAudioDevice::PushSoundMixModifier(USoundMix* SoundMix, bool bIsPassive)
 			ApplySoundMix(SoundMix, SoundMixState);
 
 			// Use it to set EQ Settings, which will check its priority
-			Effects->SetMixSettings( SoundMix );
+			Effects->SetMixSettings(SoundMix);
 		}
 		else
 		{
@@ -1837,6 +1971,45 @@ void FAudioDevice::PushSoundMixModifier(USoundMix* SoundMix, bool bIsPassive)
 	}
 }
 
+void FAudioDevice::SetSoundMixClassOverride(USoundMix* InSoundMix, USoundClass* InSoundClass, float Volume, float Pitch, float FadeInTime, bool bApplyToChildren)
+{
+	if (!InSoundMix || !InSoundClass)
+	{
+		return;
+	}
+
+	FSoundMixClassOverrideMap& SoundMixClassOverrideMap = SoundMixClassEffectOverrides.FindOrAdd(InSoundMix);
+
+	// Check if we've already added this sound class override
+	FSoundMixClassOverride* ClassOverride = SoundMixClassOverrideMap.Find(InSoundClass);
+	if (ClassOverride)
+	{
+		// Override the values of the sound class override with the new values
+		ClassOverride->SoundClassAdjustor.SoundClassObject = InSoundClass;
+		ClassOverride->SoundClassAdjustor.VolumeAdjuster = Volume;
+		ClassOverride->SoundClassAdjustor.PitchAdjuster = Pitch;
+		ClassOverride->SoundClassAdjustor.bApplyToChildren = bApplyToChildren;
+
+		// Flag that we've changed so that the update will interpolate to new values
+		ClassOverride->bOverrideChanged = true;
+		ClassOverride->FadeInTime = FadeInTime;
+	}
+	else
+	{
+		// Create a new override struct
+		FSoundMixClassOverride NewClassOverride;
+		NewClassOverride.SoundClassAdjustor.SoundClassObject = InSoundClass;
+		NewClassOverride.SoundClassAdjustor.VolumeAdjuster = Volume;
+		NewClassOverride.SoundClassAdjustor.PitchAdjuster = Pitch;
+		NewClassOverride.SoundClassAdjustor.bApplyToChildren = bApplyToChildren;
+		NewClassOverride.FadeInTime = FadeInTime;
+		NewClassOverride.bOverrideApplied = false;
+		NewClassOverride.bOverrideChanged = false;
+
+		SoundMixClassOverrideMap.Add(InSoundClass, NewClassOverride);
+	}
+}
+
 void FAudioDevice::PopSoundMixModifier(USoundMix* SoundMix, bool bIsPassive)
 {
 	if (SoundMix)
@@ -1892,7 +2065,7 @@ void FAudioDevice::ClearSoundMixModifiers()
 	}
 }
 
-void FAudioDevice::ActivateReverbEffect(class UReverbEffect* ReverbEffect, FName TagName, float Priority, float Volume, float FadeTime)
+void FAudioDevice::ActivateReverbEffect(UReverbEffect* ReverbEffect, FName TagName, float Priority, float Volume, float FadeTime)
 {
 	FActivatedReverb* ExistingReverb = ActivatedReverbs.Find(TagName);
 
@@ -1928,19 +2101,19 @@ void FAudioDevice::DeactivateReverbEffect(FName TagName)
 	{
 		if (ExistingReverb == HighestPriorityReverb)
 		{
-			HighestPriorityReverb = NULL;
+			HighestPriorityReverb = nullptr;
 		}
 
 		ActivatedReverbs.Remove(TagName);
 
-		if (HighestPriorityReverb == NULL)
+		if (HighestPriorityReverb == nullptr)
 		{
 			UpdateHighestPriorityReverb();
 		}
 	}
 }
 
-void FAudioDevice::SetReverbSettings( class AAudioVolume* Volume, const FReverbSettings& ReverbSettings )
+void FAudioDevice::SetReverbSettings(class AAudioVolume* Volume, const FReverbSettings& ReverbSettings)
 {
 	const FReverbSettings* ActivatedReverb = &ReverbSettings;
 	if (HighestPriorityReverb && (!Volume || HighestPriorityReverb->Priority > Volume->Priority))
@@ -1948,51 +2121,51 @@ void FAudioDevice::SetReverbSettings( class AAudioVolume* Volume, const FReverbS
 		ActivatedReverb = &HighestPriorityReverb->ReverbSettings;
 	}
 
-	Effects->SetReverbSettings( *ActivatedReverb );
+	Effects->SetReverbSettings(*ActivatedReverb);
 }
 
 
-void* FAudioDevice::InitEffect( FSoundSource* Source )
+void* FAudioDevice::InitEffect(FSoundSource* Source)
 {
-	return( Effects->InitEffect( Source ) );
+	return(Effects->InitEffect(Source));
 }
 
 
-void* FAudioDevice::UpdateEffect( FSoundSource* Source )
+void* FAudioDevice::UpdateEffect(FSoundSource* Source)
 {
-	SCOPE_CYCLE_COUNTER( STAT_AudioUpdateEffects );
+	SCOPE_CYCLE_COUNTER(STAT_AudioUpdateEffects);
 
-	return( Effects->UpdateEffect( Source ) );
+	return(Effects->UpdateEffect(Source));
 }
 
 
-void FAudioDevice::DestroyEffect( FSoundSource* Source )
+void FAudioDevice::DestroyEffect(FSoundSource* Source)
 {
-	return( Effects->DestroyEffect( Source ) );
+	return(Effects->DestroyEffect(Source));
 }
 
 
-void FAudioDevice::HandlePause( bool bGameTicking, bool bGlobalPause )
+void FAudioDevice::HandlePause(bool bGameTicking, bool bGlobalPause)
 {
 	// Pause all sounds if transitioning to pause mode.
-	if( !bGameTicking && (bGameWasTicking || bGlobalPause) )
+	if (!bGameTicking && (bGameWasTicking || bGlobalPause))
 	{
-		for( int32 i = 0; i < Sources.Num(); i++ )
+		for (int32 i = 0; i < Sources.Num(); i++)
 		{
 			FSoundSource* Source = Sources[ i ];
-			if( !Source->IsPaused() && (bGlobalPause || Source->IsGameOnly()) )
+			if (!Source->IsPaused() && (bGlobalPause || Source->IsGameOnly()))
 			{
 				Source->Pause();
 			}
 		}
 	}
 	// Unpause all sounds if transitioning back to game.
-	else if( bGameTicking && (!bGameWasTicking || bGlobalPause) )
+	else if (bGameTicking && (!bGameWasTicking || bGlobalPause))
 	{
-		for( int32 i = 0; i < Sources.Num(); i++ )
+		for (int32 i = 0; i < Sources.Num(); i++)
 		{
 			FSoundSource* Source = Sources[ i ];
-			if( Source->IsPaused() && (bGlobalPause || Source->IsGameOnly()) )
+			if (Source->IsPaused() && (bGlobalPause || Source->IsGameOnly()))
 			{
 				Source->Play();
 			}
@@ -2005,12 +2178,12 @@ void FAudioDevice::HandlePause( bool bGameTicking, bool bGlobalPause )
 
 int32 FAudioDevice::GetSortedActiveWaveInstances(TArray<FWaveInstance*>& WaveInstances, const ESortedActiveWaveGetType::Type GetType)
 {
-	SCOPE_CYCLE_COUNTER( STAT_AudioGatherWaveInstances );
+	SCOPE_CYCLE_COUNTER(STAT_AudioGatherWaveInstances);
 
 	// Tick all the active audio components.  Use a copy as some operations may remove elements from the list, but we want
 	// to evaluate in the order they were added
 	TArray<FActiveSound*> ActiveSoundsCopy = ActiveSounds;
-	for( int32 i = 0; i < ActiveSoundsCopy.Num(); ++i )
+	for (int32 i = 0; i < ActiveSoundsCopy.Num(); ++i)
 	{
 		FActiveSound* ActiveSound = ActiveSoundsCopy[i];
 
@@ -2095,18 +2268,18 @@ int32 FAudioDevice::GetSortedActiveWaveInstances(TArray<FWaveInstance*>& WaveIns
 		// Get the first index that will result in a active source voice
 		FirstActiveIndex = FMath::Max(WaveInstances.Num() - MaxChannels, 0);
 	}
-	return( FirstActiveIndex );
+	return(FirstActiveIndex);
 }
 
 
-void FAudioDevice::StopSources( TArray<FWaveInstance*>& WaveInstances, int32 FirstActiveIndex )
+void FAudioDevice::StopSources(TArray<FWaveInstance*>& WaveInstances, int32 FirstActiveIndex)
 {
 	// Touch sources that are high enough priority to play
 	for (int32 InstanceIndex = FirstActiveIndex; InstanceIndex < WaveInstances.Num(); InstanceIndex++)
 	{
 		FWaveInstance* WaveInstance = WaveInstances[ InstanceIndex ];
-		FSoundSource* Source = WaveInstanceSourceMap.FindRef( WaveInstance );
-		if( Source )
+		FSoundSource* Source = WaveInstanceSourceMap.FindRef(WaveInstance);
+		if (Source)
 		{
 			Source->LastUpdate = CurrentTick;
 
@@ -2128,9 +2301,9 @@ void FAudioDevice::StopSources( TArray<FWaveInstance*>& WaveInstances, int32 Fir
 		if (Source->WaveInstance)
 		{
 #if STATS && WITH_EDITORONLY_DATA
-			if( Source->UsesCPUDecompression() )
+			if (Source->UsesCPUDecompression())
 			{
-				INC_DWORD_STAT( STAT_OggWaveInstances );
+				INC_DWORD_STAT(STAT_OggWaveInstances);
 			}
 #endif
 			// If we need to stop this sound due to max concurrency (i.e. it was quietest in a concurrency group)
@@ -2155,7 +2328,7 @@ void FAudioDevice::StopSources( TArray<FWaveInstance*>& WaveInstances, int32 Fir
 	// Stop wave instances that are no longer playing due to priority reasons. This needs to happen AFTER
 	// stopping sources as calling Stop on a sound source in turn notifies the wave instance of a buffer
 	// being finished which might reset it being finished.
-	for( int32 InstanceIndex = 0; InstanceIndex < FirstActiveIndex; InstanceIndex++ )
+	for (int32 InstanceIndex = 0; InstanceIndex < FirstActiveIndex; InstanceIndex++)
 	{
 		FWaveInstance* WaveInstance = WaveInstances[ InstanceIndex ];
 		WaveInstance->StopWithoutNotification();
@@ -2164,64 +2337,95 @@ void FAudioDevice::StopSources( TArray<FWaveInstance*>& WaveInstances, int32 Fir
 #if STATS
 	uint32 AudibleInactiveSounds = 0;
 	// Count how many sounds are not being played but were audible
-	for( int32 InstanceIndex = 0; InstanceIndex < FirstActiveIndex; InstanceIndex++ )
+	for (int32 InstanceIndex = 0; InstanceIndex < FirstActiveIndex; InstanceIndex++)
 	{
 		FWaveInstance* WaveInstance = WaveInstances[ InstanceIndex ];
-		if( WaveInstance->GetActualVolume() > 0.1f )
+		if (WaveInstance->GetActualVolume() > 0.1f)
 		{
 			AudibleInactiveSounds++;
 		}
 	}
-	SET_DWORD_STAT( STAT_AudibleWavesDroppedDueToPriority, AudibleInactiveSounds );
+	SET_DWORD_STAT(STAT_AudibleWavesDroppedDueToPriority, AudibleInactiveSounds);
 #endif
 }
 
-
-void FAudioDevice::StartSources( TArray<FWaveInstance*>& WaveInstances, int32 FirstActiveIndex, bool bGameTicking )
+void FAudioDevice::StartSources(TArray<FWaveInstance*>& WaveInstances, int32 FirstActiveIndex, bool bGameTicking)
 {
-	SCOPE_CYCLE_COUNTER( STAT_AudioStartSources );
+	SCOPE_CYCLE_COUNTER(STAT_AudioStartSources);
 
 	// Start sources as needed.
-	for( int32 InstanceIndex = FirstActiveIndex; InstanceIndex < WaveInstances.Num(); InstanceIndex++ )
+	for (int32 InstanceIndex = FirstActiveIndex; InstanceIndex < WaveInstances.Num(); InstanceIndex++)
 	{
-		FWaveInstance* WaveInstance = WaveInstances[ InstanceIndex ];
+		FWaveInstance* WaveInstance = WaveInstances[InstanceIndex];
 
 		// Editor uses bIsUISound for sounds played in the browser.
-		if(	bGameTicking || WaveInstance->bIsUISound )
+		if (bGameTicking || WaveInstance->bIsUISound)
 		{
-			FSoundSource* Source = WaveInstanceSourceMap.FindRef( WaveInstance );
-			if( !Source &&
-			  ( !WaveInstance->IsStreaming() ||
-				IStreamingManager::Get().GetAudioStreamingManager().CanCreateSoundSource(WaveInstance) ) )
+			FSoundSource* Source = WaveInstanceSourceMap.FindRef(WaveInstance);
+			if (!Source &&
+				(!WaveInstance->IsStreaming() ||
+				IStreamingManager::Get().GetAudioStreamingManager().CanCreateSoundSource(WaveInstance)))
 			{
-				check( FreeSources.Num() );
+				check(FreeSources.Num());
 				Source = FreeSources.Pop();
-				check( Source);
+				check(Source);
+
+				// Prepare for initialization... 
+				bool bSuccess = false;
+				if (Source->PrepareForInitialization(WaveInstance))
+				{
+					// We successfully prepared for initialization (though we may not be prepared to actually init yet)
+					bSuccess = true;
+
+					// If we are now prepared to init (because the file handle and header synchronously loaded), then init right away
+					if (Source->IsPreparedToInit())
+					{
+						// Init the source, this may result in failure
+						bSuccess = Source->Init(WaveInstance);
+
+						// If we succeeded then play and update the source
+						if (bSuccess)
+						{
+							check(Source->IsInitialized());
+							Source->Play();
+							Source->Update();
+						}
+					}
+				}
 
-				// Try to initialize source.
-				if( Source->Init( WaveInstance ) )
+				// If we succeeded above then we need to map the wave instance to the source
+				if (bSuccess)
 				{
 					IStreamingManager::Get().GetAudioStreamingManager().AddStreamingSoundSource(Source);
 					// Associate wave instance with it which is used earlier in this function.
-					WaveInstanceSourceMap.Add( WaveInstance, Source );
-					// Playback might be deferred till the end of the update function on certain implementations.
-					Source->Play();
-
-					//UE_LOG(LogAudio, Log, TEXT( "Playing: %s" ), *WaveInstance->WaveData->GetName() );
+					WaveInstanceSourceMap.Add(WaveInstance, Source);
 				}
 				else
 				{
+					// If we failed, then we need to stop the wave instance and add the source back to the free list
 					// This can happen if e.g. the USoundWave pointed to by the WaveInstance is not a valid sound file.
 					// If we don't stop the wave file, it will continue to try initializing the file every frame, which is a perf hit
 					WaveInstance->StopWithoutNotification();
-					FreeSources.Add( Source );
+					FreeSources.Add(Source);
 				}
 			}
 			else if (Source)
 			{
-				Source->Update();
+				// If we've already been initialized, then just update the voice
+				if (Source->IsInitialized())
+				{
+					Source->Update();
+				}
+				// Otherwise, we need still need to initialize
+				else if (Source->IsPreparedToInit())
+				{
+					Source->Init(WaveInstance);
+					check(Source->IsInitialized());
+					Source->Play();
+					Source->Update();
+				}
 			}
-			else
+			else if (!Source)
 			{
 				// This can happen if the streaming manager determines that this sound should not be started.
 				// We stop the wave instance to prevent it from attempting to initialize every frame
@@ -2231,10 +2435,9 @@ void FAudioDevice::StartSources( TArray<FWaveInstance*>& WaveInstances, int32 Fi
 	}
 }
 
-
-void FAudioDevice::Update( bool bGameTicking )
+void FAudioDevice::Update(bool bGameTicking)
 {
-	SCOPE_CYCLE_COUNTER( STAT_AudioUpdateTime );
+	SCOPE_CYCLE_COUNTER(STAT_AudioUpdateTime);
 
 	double CurrTime = FPlatformTime::Seconds();
 	double DeltaTime = CurrTime - LastUpdateTime;
@@ -2249,13 +2452,13 @@ void FAudioDevice::Update( bool bGameTicking )
 	CurrentTick++;
 
 	// Handle pause/unpause for the game and editor.
-	HandlePause( bGameTicking );
+	HandlePause(bGameTicking);
 
 	// Update the audio effects - reverb, EQ etc
 	Effects->Update();
 
 	// Gets the current state of the sound classes accounting for sound mix
-	UpdateSoundClassProperties();
+	UpdateSoundClassProperties(DeltaTime);
 
 	// Gets the current state of the interior settings
 	for (FListener& Listener : Listeners)
@@ -2268,10 +2471,10 @@ void FAudioDevice::Update( bool bGameTicking )
 	if (Sources.Num())
 	{
 		// Kill any sources that have finished
-		for( int32 SourceIndex = 0; SourceIndex < Sources.Num(); SourceIndex++ )
+		for (int32 SourceIndex = 0; SourceIndex < Sources.Num(); SourceIndex++)
 		{
 			// Source has finished playing (it's one shot)
-			if( Sources[ SourceIndex ]->IsFinished() )
+			if (Sources[ SourceIndex ]->IsFinished())
 			{
 				Sources[ SourceIndex ]->Stop();
 			}
@@ -2310,14 +2513,14 @@ void FAudioDevice::Update( bool bGameTicking )
 		if (CommonAudioPoolFreeBytes != 0)
 		{
 			UE_LOG(LogAudio, Log, TEXT("Audio pool size mismatch by %d bytes. Please update CommonAudioPoolSize ini setting to %d to avoid waste!"),
-				CommonAudioPoolFreeBytes, CommonAudioPoolSize - CommonAudioPoolFreeBytes );
+				CommonAudioPoolFreeBytes, CommonAudioPoolSize - CommonAudioPoolFreeBytes);
 		}
 	}
 #endif
 }
 
 
-void FAudioDevice::StopAllSounds( bool bShouldStopUISounds )
+void FAudioDevice::StopAllSounds(bool bShouldStopUISounds)
 {
 	for (int32 SoundIndex=ActiveSounds.Num() - 1; SoundIndex >= 0; --SoundIndex)
 	{
@@ -2373,7 +2576,7 @@ void FAudioDevice::AddNewActiveSound(const FActiveSound& NewActiveSound)
 	UAudioComponent* AudioComponent = NewActiveSound.GetAudioComponent();
 
 #if !(UE_BUILD_SHIPPING || UE_BUILD_TEST)
-	UE_LOG(LogAudio, VeryVerbose, TEXT("New ActiveSound %s Comp: %s Loc: %s"), *NewActiveSound.Sound->GetName(), (AudioComponent ? *AudioComponent->GetFullName() : TEXT("No AudioComponent")), *NewActiveSound.Transform.GetTranslation().ToString() );
+	UE_LOG(LogAudio, VeryVerbose, TEXT("New ActiveSound %s Comp: %s Loc: %s"), *NewActiveSound.Sound->GetName(), (AudioComponent ? *AudioComponent->GetFullName() : TEXT("No AudioComponent")), *NewActiveSound.Transform.GetTranslation().ToString());
 #endif // !(UE_BUILD_SHIPPING || UE_BUILD_TEST)
 
 	check(ActiveSound);
@@ -2382,7 +2585,7 @@ void FAudioDevice::AddNewActiveSound(const FActiveSound& NewActiveSound)
 	if (ActiveSound->Sound)
 	{
 		if (!ensureMsgf(ActiveSound->Sound->IsValidLowLevel(), TEXT("AddNewActiveSound with INVALID sound. AudioComponent=%s"),
-			AudioComponent ? *AudioComponent->GetPathName() : TEXT("NO COMPONENT") ))
+			AudioComponent ? *AudioComponent->GetPathName() : TEXT("NO COMPONENT")))
 		{
 			static FName InvalidSoundName(TEXT("INVALID_Sound"));
 			ActiveSound->DebugOriginalSoundName = InvalidSoundName;
@@ -2463,9 +2666,9 @@ void FAudioDevice::AddSoundToStop(FActiveSound* SoundToStop)
 	}
 }
 
-void FAudioDevice::StopActiveSound( UAudioComponent* AudioComponent )
+void FAudioDevice::StopActiveSound(UAudioComponent* AudioComponent)
 {
-	check( AudioComponent );
+	check(AudioComponent);
 
 	FActiveSound* ActiveSound = FindActiveSound(AudioComponent);
 	if (ActiveSound)
@@ -2479,7 +2682,7 @@ void FAudioDevice::StopActiveSound(FActiveSound* ActiveSound)
 	AddSoundToStop(ActiveSound);
 }
 
-FActiveSound* FAudioDevice::FindActiveSound( UAudioComponent* AudioComponent )
+FActiveSound* FAudioDevice::FindActiveSound(UAudioComponent* AudioComponent)
 {
 	// find the active sound corresponding to this audio component
 	if (FActiveSound** ActiveSoundPtr = AudioComponentToActiveSoundMap.Find((UPTRINT)AudioComponent))
@@ -2533,7 +2736,7 @@ bool FAudioDevice::LocationIsAudible(const FVector& Location, const FListener& L
 	return (Listener.Transform.GetTranslation() - Location).SizeSquared() < MaxDistanceSquared;
 }
 
-void FAudioDevice::GetMaxDistanceAndFocusFactor(USoundBase* Sound, const class UWorld* World, const FVector& Location, const FAttenuationSettings* AttenuationSettingsToApply, float *OutMaxDistance, float* OutFocusFactor)
+void FAudioDevice::GetMaxDistanceAndFocusFactor(USoundBase* Sound, const UWorld* World, const FVector& Location, const FAttenuationSettings* AttenuationSettingsToApply, float *OutMaxDistance, float* OutFocusFactor)
 {
 	check(Sound);
 	check(OutMaxDistance);
@@ -2561,7 +2764,7 @@ void FAudioDevice::GetMaxDistanceAndFocusFactor(USoundBase* Sound, const class U
 	}
 }
 
-bool FAudioDevice::SoundIsAudible(USoundBase* Sound, const class UWorld* World, const FVector& Location, const FAttenuationSettings* AttenuationSettingsToApply, float MaxDistance, float FocusFactor)
+bool FAudioDevice::SoundIsAudible(USoundBase* Sound, const UWorld* World, const FVector& Location, const FAttenuationSettings* AttenuationSettingsToApply, float MaxDistance, float FocusFactor)
 {
 	bool bIsInGameWorld = World ? World->IsGameWorld() : true;
 	bool bHasAttenuationSettings = (bIsInGameWorld && AttenuationSettingsToApply);
@@ -2670,9 +2873,9 @@ float FAudioDevice::GetFocusFactor(FAttenuationListenerData& OutListenerData, co
 
 UAudioComponent* FAudioDevice::CreateComponent(USoundBase* Sound, UWorld* World, AActor* Actor, bool bPlay, bool bStopWhenOwnerDestroyed, const FVector* Location, USoundAttenuation* AttenuationSettings, USoundConcurrency* ConcurrencySettings)
 {
-	UAudioComponent* AudioComponent = NULL;
+	UAudioComponent* AudioComponent = nullptr;
 
-	if( Sound && GEngine && GEngine->UseSound() )
+	if (Sound && GEngine && GEngine->UseSound())
 	{
 		// Get the world's audio device if there is a world, otherwise get the main audio device
 		FAudioDevice* AudioDevice = nullptr;
@@ -2738,7 +2941,7 @@ UAudioComponent* FAudioDevice::CreateComponent(USoundBase* Sound, UWorld* World,
 					AudioComponent->SetWorldLocation(*Location);
 				}
 
-				// AudioComponent used in PlayEditorSound sets World to NULL to avoid situations where the world becomes invalid
+				// AudioComponent used in PlayEditorSound sets World to nullptr to avoid situations where the world becomes invalid
 				// and the component is left with invalid pointer.
 				if (World)
 				{
@@ -2758,7 +2961,7 @@ UAudioComponent* FAudioDevice::CreateComponent(USoundBase* Sound, UWorld* World,
 		}
 	}
 
-	return( AudioComponent );
+	return(AudioComponent);
 }
 
 void FAudioDevice::PlaySoundAtLocation(USoundBase* Sound, UWorld* World, float VolumeMultiplier, float PitchMultiplier, float StartTime, const FVector& Location, const FRotator& Rotation, USoundAttenuation* AttenuationSettings, USoundConcurrency* ConcurrencySettings, const TArray<FAudioComponentParam>* Params)
@@ -2821,15 +3024,15 @@ void FAudioDevice::PlaySoundAtLocation(USoundBase* Sound, UWorld* World, float V
 	}
 }
 
-void FAudioDevice::Flush( UWorld* WorldToFlush, bool bClearActivatedReverb )
+void FAudioDevice::Flush(UWorld* WorldToFlush, bool bClearActivatedReverb)
 {
 	// Stop all audio components attached to the scene
 	bool bFoundIgnoredComponent = false;
-	for( int32 Index = ActiveSounds.Num() - 1; Index >= 0; --Index )
+	for (int32 Index = ActiveSounds.Num() - 1; Index >= 0; --Index)
 	{
 		FActiveSound* ActiveSound = ActiveSounds[Index];
 		// if we are in the editor we want to always flush the ActiveSounds
-		if( ActiveSound->bIgnoreForFlushing )
+		if (ActiveSound->bIgnoreForFlushing)
 		{
 			bFoundIgnoredComponent = true;
 		}
@@ -2864,18 +3067,18 @@ void FAudioDevice::Flush( UWorld* WorldToFlush, bool bClearActivatedReverb )
 		HighestPriorityReverb = nullptr;
 	}
 
-	if( WorldToFlush == nullptr )
+	if (WorldToFlush == nullptr)
 	{
 		// Make sure sounds are fully stopped.
-		if( bFoundIgnoredComponent )
+		if (bFoundIgnoredComponent)
 		{
 			// We encountered an ignored component, so address the sounds individually.
 			// There's no need to individually clear WaveInstanceSourceMap elements,
 			// because FSoundSource::Stop(...) takes care of this.
-			for( int32 SourceIndex = 0; SourceIndex < Sources.Num(); SourceIndex++ )
+			for (int32 SourceIndex = 0; SourceIndex < Sources.Num(); SourceIndex++)
 			{
 				const FWaveInstance* WaveInstance = Sources[SourceIndex]->GetWaveInstance();
-				if( WaveInstance == NULL || !WaveInstance->ActiveSound->bIgnoreForFlushing )
+				if (WaveInstance == nullptr || !WaveInstance->ActiveSound->bIgnoreForFlushing)
 				{
 					Sources[ SourceIndex ]->Stop();
 				}
@@ -2884,7 +3087,7 @@ void FAudioDevice::Flush( UWorld* WorldToFlush, bool bClearActivatedReverb )
 		else
 		{
 			// No components were ignored, so stop all sounds.
-			for( int32 SourceIndex = 0; SourceIndex < Sources.Num(); SourceIndex++ )
+			for (int32 SourceIndex = 0; SourceIndex < Sources.Num(); SourceIndex++)
 			{
 				Sources[ SourceIndex ]->Stop();
 			}
@@ -2902,29 +3105,29 @@ void FAudioDevice::Flush( UWorld* WorldToFlush, bool bClearActivatedReverb )
 
 void FAudioDevice::Precache(USoundWave* SoundWave, bool bSynchronous, bool bTrackMemory)
 {
-	if( SoundWave == NULL )
+	if (SoundWave == nullptr)
 	{
 		return;
 	}
 
 	// calculate the decompression type
 	// @todo audio: maybe move this into SoundWave?
-	if( SoundWave->NumChannels == 0 )
+	if (SoundWave->NumChannels == 0)
 	{
 		// No channels - no way of knowing what to play back
 		SoundWave->DecompressionType = DTYPE_Invalid;
 	}
-	else if( SoundWave->RawPCMData )
+	else if (SoundWave->RawPCMData)
 	{
 		// Run time created audio; e.g. editor preview data
 		SoundWave->DecompressionType = DTYPE_Preview;
 	}
-	else if ( SoundWave->bProcedural )
+	else if (SoundWave->bProcedural)
 	{
 		// Streaming data, created programmatically.
 		SoundWave->DecompressionType = DTYPE_Procedural;
 	}
-	else if ( HasCompressedAudioInfoClass(SoundWave) )
+	else if (HasCompressedAudioInfoClass(SoundWave))
 	{
 		const FSoundGroup& SoundGroup = GetDefault<USoundGroups>()->GetSoundGroup(SoundWave->SoundGroup);
 
@@ -3003,8 +3206,8 @@ void FAudioDevice::Precache(USoundWave* SoundWave, bool bSynchronous, bool bTrac
 		SoundWave->TrackedMemoryUsage += ResourceSize;
 
 		// If we aren't decompressing it above, then count the memory
-		INC_DWORD_STAT_BY( STAT_AudioMemorySize, ResourceSize );
-		INC_DWORD_STAT_BY( STAT_AudioMemory, ResourceSize );
+		INC_DWORD_STAT_BY(STAT_AudioMemorySize, ResourceSize);
+		INC_DWORD_STAT_BY(STAT_AudioMemory, ResourceSize);
 	}
 }
 
@@ -3045,7 +3248,7 @@ void FAudioDevice::UnregisterSoundClass(USoundClass* SoundClass)
 	}
 }
 
-FSoundClassProperties* FAudioDevice::GetSoundClassCurrentProperties(class USoundClass* InSoundClass)
+FSoundClassProperties* FAudioDevice::GetSoundClassCurrentProperties(USoundClass* InSoundClass)
 {
 	if (InSoundClass)
 	{
@@ -3153,7 +3356,7 @@ void FAudioDevice::StopSoundsUsingResource(USoundWave* SoundWave, TArray<UAudioC
 
 	if (!GIsEditor && bStoppedSounds)
 	{
-		UE_LOG(LogAudio, Warning, TEXT( "All Sounds using SoundWave '%s' have been stopped" ), *SoundWave->GetName() );
+		UE_LOG(LogAudio, Warning, TEXT("All Sounds using SoundWave '%s' have been stopped"), *SoundWave->GetName());
 	}
 }
 
diff --git a/Engine/Source/Runtime/Engine/Private/Collision/WorldCollision.cpp b/Engine/Source/Runtime/Engine/Private/Collision/WorldCollision.cpp
--- a/Engine/Source/Runtime/Engine/Private/Collision/WorldCollision.cpp
+++ b/Engine/Source/Runtime/Engine/Private/Collision/WorldCollision.cpp
@@ -174,22 +174,6 @@ bool UWorld::OverlapAnyTestByChannel(const FVector& Pos, const FQuat& Rot, EColl
 
 }
 
-bool UWorld::OverlapSingle(struct FOverlapResult& OutOverlap, const FVector& Pos, const FQuat& Rot, ECollisionChannel TraceChannel, const struct FCollisionShape& CollisionShape, const struct FCollisionQueryParams& Params, const struct FCollisionResponseParams& ResponseParam) const
-{
-	bool bBlocking = false;
-#if UE_WITH_PHYSICS
-	TArray<FOverlapResult> Overlaps;
-	return GeomOverlapMulti(this, CollisionShape, Pos, Rot, Overlaps, TraceChannel, Params, ResponseParam, FCollisionObjectQueryParams::DefaultObjectQueryParam);
-	if (Overlaps.Num())
-	{
-		OutOverlap = Overlaps[0];
-		bBlocking = Overlaps[0].bBlockingHit;
-	}
-#endif
-
-	return bBlocking;
-}
-
 bool UWorld::OverlapMultiByChannel(TArray<struct FOverlapResult>& OutOverlaps, const FVector& Pos, const FQuat& Rot, ECollisionChannel TraceChannel, const FCollisionShape& CollisionShape, const FCollisionQueryParams& Params /* = FCollisionQueryParams::DefaultQueryParam */, const FCollisionResponseParams& ResponseParam /* = FCollisionResponseParams::DefaultResponseParam */) const
 {
 #if UE_WITH_PHYSICS
@@ -310,24 +294,6 @@ bool UWorld::OverlapAnyTestByObjectType(const FVector& Pos, const FQuat& Rot, co
 	return bBlocking;
 }
 
-
-bool UWorld::OverlapSingle(struct FOverlapResult& OutOverlap, const FVector& Pos, const FQuat& Rot, const struct FCollisionShape& CollisionShape, const struct FCollisionQueryParams& Params, const struct FCollisionObjectQueryParams& ObjectQueryParams) const
-{
-	bool bBlocking = false;
-#if UE_WITH_PHYSICS
-	TArray<FOverlapResult> Overlaps;
-	GeomOverlapMulti(this, CollisionShape, Pos, Rot, Overlaps, DefaultCollisionChannel, Params, FCollisionResponseParams::DefaultResponseParam, ObjectQueryParams);
-	if (Overlaps.Num() > 0)
-	{
-		OutOverlap = Overlaps[0];
-		bBlocking = true;
-	}
-#endif
-
-	return bBlocking;
-}
-
-
 // profile interfaces
 static void GetCollisionProfileChannelAndResponseParams(FName ProfileName, ECollisionChannel &CollisionChannel, FCollisionResponseParams &ResponseParams)
 {
@@ -415,21 +381,6 @@ bool UWorld::OverlapAnyTestByProfile(const FVector& Pos, const FQuat& Rot, FName
 	return OverlapAnyTestByChannel(Pos, Rot, TraceChannel, CollisionShape, Params, ResponseParam);
 }
 
-bool UWorld::OverlapSingleByProfile(struct FOverlapResult& OutOverlap, const FVector& Pos, const FQuat& Rot, FName ProfileName, const struct FCollisionShape& CollisionShape, const struct FCollisionQueryParams& Params) const
-{
-	TArray<FOverlapResult> Overlaps;
-	OverlapMultiByProfile(Overlaps, Pos, Rot, ProfileName, CollisionShape, Params);
-
-	bool bBlocking = false;
-	if (Overlaps.Num())
-	{
-		bBlocking = Overlaps[0].bBlockingHit;
-		OutOverlap = Overlaps[0];
-	}
-
-	return bBlocking;
-}
-
 bool UWorld::OverlapMultiByProfile(TArray<struct FOverlapResult>& OutOverlaps, const FVector& Pos, const FQuat& Rot, FName ProfileName, const struct FCollisionShape& CollisionShape, const struct FCollisionQueryParams& Params) const
 {
 	ECollisionChannel TraceChannel;
diff --git a/Engine/Source/Runtime/Engine/Private/Collision/WorldCollisionAsync.cpp b/Engine/Source/Runtime/Engine/Private/Collision/WorldCollisionAsync.cpp
--- a/Engine/Source/Runtime/Engine/Private/Collision/WorldCollisionAsync.cpp
+++ b/Engine/Source/Runtime/Engine/Private/Collision/WorldCollisionAsync.cpp
@@ -99,30 +99,72 @@ namespace
 			{
 				if ((TraceData.CollisionParams.CollisionShape.ShapeType == ECollisionShape::Line) || TraceData.CollisionParams.CollisionShape.IsNearlyZero())
 				{
-					if (TraceData.bIsMultiTrace)
+					// MULTI
+					if (TraceData.TraceType == EAsyncTraceType::Multi)
 					{
 						RaycastMulti(TraceData.PhysWorld.Get(), TraceData.OutHits, TraceData.Start, TraceData.End, TraceData.TraceChannel,
 							TraceData.CollisionParams.CollisionQueryParam, TraceData.CollisionParams.ResponseParam, TraceData.CollisionParams.ObjectQueryParam);
 					}
+					// SINGLE
+					else if(TraceData.TraceType == EAsyncTraceType::Single)
+					{
+						FHitResult Result;
+
+						bool bHit = RaycastSingle(TraceData.PhysWorld.Get(), Result, TraceData.Start, TraceData.End, TraceData.TraceChannel,
+							TraceData.CollisionParams.CollisionQueryParam, TraceData.CollisionParams.ResponseParam, TraceData.CollisionParams.ObjectQueryParam);
+
+						if(bHit)
+						{
+							TraceData.OutHits.Add(Result);
+						}
+					}
+					// TEST
 					else
 					{
-						TraceData.OutHits.AddZeroed(1);
-						RaycastSingle(TraceData.PhysWorld.Get(), TraceData.OutHits[0], TraceData.Start, TraceData.End, TraceData.TraceChannel,
+						bool bHit = RaycastTest(TraceData.PhysWorld.Get(), TraceData.Start, TraceData.End, TraceData.TraceChannel,
 							TraceData.CollisionParams.CollisionQueryParam, TraceData.CollisionParams.ResponseParam, TraceData.CollisionParams.ObjectQueryParam);
+
+						if(bHit)
+						{
+							FHitResult Result;
+							Result.bBlockingHit = true;
+							TraceData.OutHits.Add(Result);
+						}
 					}
 				}
 				else
 				{
-					if (TraceData.bIsMultiTrace)
+					// MULTI
+					if (TraceData.TraceType == EAsyncTraceType::Multi)
 					{
 						GeomSweepMulti(TraceData.PhysWorld.Get(), TraceData.CollisionParams.CollisionShape, FQuat::Identity, TraceData.OutHits, TraceData.Start, TraceData.End, TraceData.TraceChannel,
 							TraceData.CollisionParams.CollisionQueryParam, TraceData.CollisionParams.ResponseParam, TraceData.CollisionParams.ObjectQueryParam);
 					}
+					// SINGLE
+					else if (TraceData.TraceType == EAsyncTraceType::Single)
+					{
+						FHitResult Result;
+
+						bool bHit = GeomSweepSingle(TraceData.PhysWorld.Get(), TraceData.CollisionParams.CollisionShape, FQuat::Identity, Result, TraceData.Start, TraceData.End, TraceData.TraceChannel,
+							TraceData.CollisionParams.CollisionQueryParam, TraceData.CollisionParams.ResponseParam, TraceData.CollisionParams.ObjectQueryParam);
+
+						if(bHit)
+						{
+							TraceData.OutHits.Add(Result);
+						}
+					}
+					// TEST
 					else
 					{
-						TraceData.OutHits.AddZeroed(1);
-						GeomSweepSingle(TraceData.PhysWorld.Get(), TraceData.CollisionParams.CollisionShape, FQuat::Identity, TraceData.OutHits[0], TraceData.Start, TraceData.End, TraceData.TraceChannel,
+						bool bHit = GeomSweepTest(TraceData.PhysWorld.Get(), TraceData.CollisionParams.CollisionShape, FQuat::Identity, TraceData.Start, TraceData.End, TraceData.TraceChannel,
 							TraceData.CollisionParams.CollisionQueryParam, TraceData.CollisionParams.ResponseParam, TraceData.CollisionParams.ObjectQueryParam);
+
+						if(bHit)
+						{
+							FHitResult Result;
+							Result.bBlockingHit = true;
+							TraceData.OutHits.Add(Result);
+						}						
 					}
 				}
 			}
@@ -304,24 +346,24 @@ FWorldAsyncTraceState::FWorldAsyncTraceState()
 	DataBuffer[CurrentFrame].bAsyncAllowed = true;
 }
 
-FTraceHandle UWorld::AsyncLineTraceByChannel(const FVector& Start,const FVector& End, ECollisionChannel TraceChannel, const FCollisionQueryParams& Params /* = FCollisionQueryParams::DefaultQueryParam */, const FCollisionResponseParams& ResponseParam /* = FCollisionResponseParams::DefaultResponseParam */, FTraceDelegate * InDelegate/* =NULL */, uint32 UserData /* = 0 */, bool bMultiTrace/* =false */ )
+FTraceHandle UWorld::AsyncLineTraceByChannel(EAsyncTraceType InTraceType, const FVector& Start, const FVector& End, ECollisionChannel TraceChannel, const FCollisionQueryParams& Params /* = FCollisionQueryParams::DefaultQueryParam */, const FCollisionResponseParams& ResponseParam /* = FCollisionResponseParams::DefaultResponseParam */, FTraceDelegate * InDelegate/* =NULL */, uint32 UserData /* = 0 */)
 {
-	return StartNewTrace(AsyncTraceState, FTraceDatum(this, FCollisionShape::LineShape, Params, ResponseParam, FCollisionObjectQueryParams::DefaultObjectQueryParam, TraceChannel, UserData, bMultiTrace, Start, End, InDelegate, AsyncTraceState.CurrentFrame));
+	return StartNewTrace(AsyncTraceState, FTraceDatum(this, FCollisionShape::LineShape, Params, ResponseParam, FCollisionObjectQueryParams::DefaultObjectQueryParam, TraceChannel, UserData, InTraceType, Start, End, InDelegate, AsyncTraceState.CurrentFrame));
 }
 
-FTraceHandle UWorld::AsyncLineTraceByObjectType(const FVector& Start,const FVector& End, const FCollisionObjectQueryParams& ObjectQueryParams, const FCollisionQueryParams& Params /* = FCollisionQueryParams::DefaultQueryParam */, FTraceDelegate * InDelegate/* =NULL */, uint32 UserData /* = 0 */, bool bMultiTrace/* =false */ )
+FTraceHandle UWorld::AsyncLineTraceByObjectType(EAsyncTraceType InTraceType, const FVector& Start,const FVector& End, const FCollisionObjectQueryParams& ObjectQueryParams, const FCollisionQueryParams& Params /* = FCollisionQueryParams::DefaultQueryParam */, FTraceDelegate * InDelegate/* =NULL */, uint32 UserData /* = 0 */)
 {
-	return StartNewTrace(AsyncTraceState, FTraceDatum(this, FCollisionShape::LineShape, Params, FCollisionResponseParams::DefaultResponseParam, ObjectQueryParams, DefaultCollisionChannel, UserData, bMultiTrace, Start, End, InDelegate, AsyncTraceState.CurrentFrame));
+	return StartNewTrace(AsyncTraceState, FTraceDatum(this, FCollisionShape::LineShape, Params, FCollisionResponseParams::DefaultResponseParam, ObjectQueryParams, DefaultCollisionChannel, UserData, InTraceType, Start, End, InDelegate, AsyncTraceState.CurrentFrame));
 }
 
-FTraceHandle UWorld::AsyncSweepByChannel(const FVector& Start, const FVector& End, ECollisionChannel TraceChannel, const FCollisionShape& CollisionShape, const FCollisionQueryParams& Params /* = FCollisionQueryParams::DefaultQueryParam */, const FCollisionResponseParams& ResponseParam /* = FCollisionResponseParams::DefaultResponseParam */, FTraceDelegate * InDelegate /* = NULL */, uint32 UserData /* = 0 */, bool bMultiTrace /* = false */)
+FTraceHandle UWorld::AsyncSweepByChannel(EAsyncTraceType InTraceType, const FVector& Start, const FVector& End, ECollisionChannel TraceChannel, const FCollisionShape& CollisionShape, const FCollisionQueryParams& Params /* = FCollisionQueryParams::DefaultQueryParam */, const FCollisionResponseParams& ResponseParam /* = FCollisionResponseParams::DefaultResponseParam */, FTraceDelegate * InDelegate /* = NULL */, uint32 UserData /* = 0 */)
 {
-	return StartNewTrace(AsyncTraceState, FTraceDatum(this, CollisionShape, Params, ResponseParam, FCollisionObjectQueryParams::DefaultObjectQueryParam, TraceChannel, UserData, bMultiTrace, Start, End, InDelegate, AsyncTraceState.CurrentFrame));
+	return StartNewTrace(AsyncTraceState, FTraceDatum(this, CollisionShape, Params, ResponseParam, FCollisionObjectQueryParams::DefaultObjectQueryParam, TraceChannel, UserData, InTraceType, Start, End, InDelegate, AsyncTraceState.CurrentFrame));
 }
 
-FTraceHandle UWorld::AsyncSweepByObjectType(const FVector& Start, const FVector& End, const FCollisionObjectQueryParams& ObjectQueryParams, const FCollisionShape& CollisionShape, const FCollisionQueryParams& Params /* = FCollisionQueryParams::DefaultQueryParam */, FTraceDelegate * InDelegate /* = NULL */, uint32 UserData /* = 0 */, bool bMultiTrace /* = false */)
+FTraceHandle UWorld::AsyncSweepByObjectType(EAsyncTraceType InTraceType, const FVector& Start, const FVector& End, const FCollisionObjectQueryParams& ObjectQueryParams, const FCollisionShape& CollisionShape, const FCollisionQueryParams& Params /* = FCollisionQueryParams::DefaultQueryParam */, FTraceDelegate * InDelegate /* = NULL */, uint32 UserData /* = 0 */)
 {
-	return StartNewTrace(AsyncTraceState, FTraceDatum(this, CollisionShape, Params, FCollisionResponseParams::DefaultResponseParam, ObjectQueryParams, DefaultCollisionChannel, UserData, bMultiTrace, Start, End, InDelegate, AsyncTraceState.CurrentFrame));
+	return StartNewTrace(AsyncTraceState, FTraceDatum(this, CollisionShape, Params, FCollisionResponseParams::DefaultResponseParam, ObjectQueryParams, DefaultCollisionChannel, UserData, InTraceType, Start, End, InDelegate, AsyncTraceState.CurrentFrame));
 }
 
 // overlap functions
