Copying //UE4/Dev-Framework to Dev-Main (//UE4/Dev-Main) @ 2879625

diff --git a/Engine/Source/Runtime/Engine/Private/Components/NavMeshRenderingComponent.cpp b/Engine/Source/Runtime/Engine/Private/Components/NavMeshRenderingComponent.cpp
--- a/Engine/Source/Runtime/Engine/Private/Components/NavMeshRenderingComponent.cpp
+++ b/Engine/Source/Runtime/Engine/Private/Components/NavMeshRenderingComponent.cpp
@@ -10,333 +10,458 @@
 #include "AI/Navigation/RecastHelpers.h"
 #include "AI/Navigation/RecastNavMeshGenerator.h"
 #include "AI/Navigation/NavigationSystem.h"
+#include "AI/Navigation/NavMeshRenderingComponent.h"
+#include "Debug/DebugDrawService.h"
 
 #if WITH_EDITOR
 #include "Editor.h"
 #include "EditorViewportClient.h"
 #endif
 
-#include "NavMeshRenderingHelpers.h"
-
-#if WITH_EDITOR
-namespace
+static const FColor NavMeshRenderColor_Recast_TriangleEdges(255, 255, 255);
+static const FColor NavMeshRenderColor_Recast_TileEdges(16, 16, 16, 32);
+static const FColor NavMeshRenderColor_Recast_NavMeshEdges(32, 63, 0, 220);
+static const FColor NavMeshRenderColor_RecastMesh(140, 255, 0, 164);
+static const FColor NavMeshRenderColor_TileBounds(255, 255, 64, 255);
+static const FColor NavMeshRenderColor_PathCollidingGeom(255, 255, 255, 40);
+static const FColor NavMeshRenderColor_RecastTileBeingRebuilt(255, 0, 0, 64);
+static const FColor NavMeshRenderColor_OffMeshConnectionInvalid(64, 64, 64);
+
+static const float DefaultEdges_LineThickness = 0.0f;
+static const float PolyEdges_LineThickness = 1.5f;
+static const float NavMeshEdges_LineThickness = 3.5f;
+static const float LinkLines_LineThickness = 2.0f;
+static const float ClusterLinkLines_LineThickness = 2.0f;
+
+namespace FNavMeshRenderingHelpers
 {
-	bool AreAnyViewportsRelevant(const UWorld* World)
+	bool LineInView(const FVector& Start, const FVector& End, const FSceneView* View, bool bUseDistanceCheck)
 	{
-		FWorldContext* WorldContext = GEngine->GetWorldContextFromWorld(World);
-		if (WorldContext && WorldContext->GameViewport)
+		if (FVector::DistSquared(Start, View->ViewMatrices.ViewOrigin) > ARecastNavMesh::GetDrawDistanceSq() ||
+			FVector::DistSquared(End, View->ViewMatrices.ViewOrigin) > ARecastNavMesh::GetDrawDistanceSq())
 		{
-			return true;
+			return false;
 		}
-		
-		for (FEditorViewportClient* CurrentViewport : GEditor->AllViewportClients)
+
+		for (int32 PlaneIdx = 0; PlaneIdx < View->ViewFrustum.Planes.Num(); ++PlaneIdx)
 		{
-			if (CurrentViewport && CurrentViewport->IsVisible())
+			const FPlane& CurPlane = View->ViewFrustum.Planes[PlaneIdx];
+			if (CurPlane.PlaneDot(Start) > 0.f && CurPlane.PlaneDot(End) > 0.f)
 			{
-				return true;
+				return false;
 			}
 		}
 
-		return false;
+		return true;
 	}
-}
-#endif
-
-UNavMeshRenderingComponent::UNavMeshRenderingComponent(const FObjectInitializer& ObjectInitializer)
-	: Super(ObjectInitializer)
-{
-	SetCollisionProfileName(UCollisionProfile::NoCollision_ProfileName);
-	AlwaysLoadOnClient = false;
-	AlwaysLoadOnServer = false;
-	bSelectable = false;
-	bCollectNavigationData = false;
-}
 
-bool UNavMeshRenderingComponent::NeedsLoadForServer() const
-{
-	// This rendering component is not needed in dedicated server builds
-	return false;
-}
+	bool LineInCorrectDistance(const FVector& Start, const FVector& End, const FSceneView* View, float CorrectDistance = -1)
+	{
+		const float MaxDistanceSq = (CorrectDistance > 0) ? FMath::Square(CorrectDistance) : ARecastNavMesh::GetDrawDistanceSq();
+		return FVector::DistSquared(Start, View->ViewMatrices.ViewOrigin) < MaxDistanceSq &&
+			FVector::DistSquared(End, View->ViewMatrices.ViewOrigin) < MaxDistanceSq;
+	}
 
-bool UNavMeshRenderingComponent::IsNavigationShowFlagSet(const UWorld* World)
-{
-	bool bShowNavigation = false;
+	FVector EvalArc(const FVector& Org, const FVector& Dir, const float h, const float u)
+	{
+		FVector Pt = Org + Dir * u;
+		Pt.Z += h * (1 - (u * 2 - 1)*(u * 2 - 1));
 
-	FWorldContext* WorldContext = GEngine->GetWorldContextFromWorld(World);
+		return Pt;
+	}
 
-#if WITH_EDITOR
-	if (GEditor && WorldContext && WorldContext->WorldType != EWorldType::Game)
+	void CacheArc(TArray<FDebugRenderSceneProxy::FDebugLine>& DebugLines, const FVector& Start, const FVector& End, const float Height, const uint32 Segments, const FLinearColor& Color, float LineThickness = 0)
 	{
-		bShowNavigation = WorldContext->GameViewport != nullptr && WorldContext->GameViewport->EngineShowFlags.Navigation;
-		if (bShowNavigation == false)
+		if (Segments == 0)
 		{
-			// we have to check all vieports because we can't to distinguish between SIE and PIE at this point.
-			for (FEditorViewportClient* CurrentViewport : GEditor->AllViewportClients)
-			{
-				if (CurrentViewport && CurrentViewport->IsVisible() && CurrentViewport->EngineShowFlags.Navigation)
-				{
-					bShowNavigation = true;
-					break;
-				}
-			}
+			return;
+		}
+
+		const float ArcPtsScale = 1.0f / (float)Segments;
+		const FVector Dir = End - Start;
+		const float Length = Dir.Size();
+
+		FVector Prev = Start;
+		for (uint32 i = 1; i <= Segments; ++i)
+		{
+			const float u = i * ArcPtsScale;
+			const FVector Pt = EvalArc(Start, Dir, Length*Height, u);
+
+			DebugLines.Add(FDebugRenderSceneProxy::FDebugLine(Prev, Pt, Color.ToFColor(true)));
+			Prev = Pt;
 		}
 	}
-	else
-#endif //WITH_EDITOR
+
+	void CacheArrowHead(TArray<FDebugRenderSceneProxy::FDebugLine>& DebugLines, const FVector& Tip, const FVector& Origin, const float Size, const FLinearColor& Color, float LineThickness = 0)
 	{
-		bShowNavigation = WorldContext && WorldContext->GameViewport && WorldContext->GameViewport->EngineShowFlags.Navigation;
+		FVector Ax, Ay, Az(0, 1, 0);
+		Ay = Origin - Tip;
+		Ay.Normalize();
+		Ax = FVector::CrossProduct(Az, Ay);
+
+		FHitProxyId HitProxyId;
+		DebugLines.Add(FDebugRenderSceneProxy::FDebugLine(Tip, FVector(Tip.X + Ay.X*Size + Ax.X*Size / 3, Tip.Y + Ay.Y*Size + Ax.Y*Size / 3, Tip.Z + Ay.Z*Size + Ax.Z*Size / 3), Color.ToFColor(true)));
+		DebugLines.Add(FDebugRenderSceneProxy::FDebugLine(Tip, FVector(Tip.X + Ay.X*Size - Ax.X*Size / 3, Tip.Y + Ay.Y*Size - Ax.Y*Size / 3, Tip.Z + Ay.Z*Size - Ax.Z*Size / 3), Color.ToFColor(true)));
 	}
 
-	return bShowNavigation;
-}
+	void DrawWireCylinder(TArray<FDebugRenderSceneProxy::FDebugLine>& DebugLines, const FVector& Base, const FVector& X, const FVector& Y, const FVector& Z, FColor Color, float Radius, float HalfHeight, int32 NumSides, uint8 DepthPriority, float LineThickness = 0)
+	{
+		const float	AngleDelta = 2.0f * PI / NumSides;
+		FVector	LastVertex = Base + X * Radius;
 
-void UNavMeshRenderingComponent::TimerFunction()
-{
-	const UWorld* World = GetWorld();
-#if WITH_EDITOR
-	if (GEditor && (AreAnyViewportsRelevant(World) == false))
+		FHitProxyId HitProxyId;
+		for (int32 SideIndex = 0; SideIndex < NumSides; SideIndex++)
+		{
+			const FVector Vertex = Base + (X * FMath::Cos(AngleDelta * (SideIndex + 1)) + Y * FMath::Sin(AngleDelta * (SideIndex + 1))) * Radius;
+
+			DebugLines.Add(FDebugRenderSceneProxy::FDebugLine(LastVertex - Z * HalfHeight, Vertex - Z * HalfHeight, Color));
+			DebugLines.Add(FDebugRenderSceneProxy::FDebugLine(LastVertex + Z * HalfHeight, Vertex + Z * HalfHeight, Color));
+			DebugLines.Add(FDebugRenderSceneProxy::FDebugLine(LastVertex - Z * HalfHeight, LastVertex + Z * HalfHeight, Color));
+
+			LastVertex = Vertex;
+		}
+	}
+
+	inline uint8 GetBit(int32 v, uint8 bit)
 	{
-		// unable to tell if the flag is on or not
-		return;
+		return (v & (1 << bit)) >> bit;
 	}
-#endif // WITH_EDITOR
 
-	const bool bShowNavigation = bForceUpdate || IsNavigationShowFlagSet(World);
+	FColor GetClusterColor(int32 Idx)
+	{
+		uint8 r = 1 + GetBit(Idx, 1) + GetBit(Idx, 3) * 2;
+		uint8 g = 1 + GetBit(Idx, 2) + GetBit(Idx, 4) * 2;
+		uint8 b = 1 + GetBit(Idx, 0) + GetBit(Idx, 5) * 2;
+		return FColor(r * 63, g * 63, b * 63, 164);
+	}
 
-	if (bShowNavigation != !!bCollectNavigationData && bShowNavigation == true)
+	FColor DarkenColor(const FColor& Base)
 	{
-		bForceUpdate = false;
-		bCollectNavigationData = bShowNavigation;
-		MarkRenderStateDirty();
+		const uint32 Col = Base.DWColor();
+		return FColor(((Col >> 1) & 0x007f7f7f) | (Col & 0xff000000));
 	}
-}
 
-void UNavMeshRenderingComponent::OnRegister()
-{
-	Super::OnRegister();
+	void AddVertex(FNavMeshSceneProxyData::FDebugMeshData& MeshData, const FVector& Pos, const FColor Color = FColor::White)
+	{
+		const int32 VertexIndex = MeshData.Vertices.Num();
+		FDynamicMeshVertex* Vertex = new(MeshData.Vertices) FDynamicMeshVertex;
+		Vertex->Position = Pos;
+		Vertex->TextureCoordinate = FVector2D::ZeroVector;
+		Vertex->TangentX = FVector(1.0f, 0.0f, 0.0f);
+		Vertex->TangentZ = FVector(0.0f, 1.0f, 0.0f);
+		// store the sign of the determinant in TangentZ.W (-1=0,+1=255)
+		Vertex->TangentZ.Vector.W = 255;
+		Vertex->Color = Color;
+	}
 
-#if WITH_RECAST && !UE_BUILD_SHIPPING && !UE_BUILD_TEST
-	// it's a kind of HACK but there is no event or other information that show flag was changed by user => we have to check it periodically
-#if WITH_EDITOR
-	if (GEditor)
+	void AddTriangle(FNavMeshSceneProxyData::FDebugMeshData& MeshData, int32 V0, int32 V1, int32 V2)
 	{
-		GEditor->GetTimerManager()->SetTimer(TimerHandle, FTimerDelegate::CreateUObject(this, &UNavMeshRenderingComponent::TimerFunction), 1, true);
+		MeshData.Indices.Add(V0);
+		MeshData.Indices.Add(V1);
+		MeshData.Indices.Add(V2);
 	}
-	else
-#endif //WITH_EDITOR
+
+	void AddRecastGeometry(TArray<FVector>& OutVertexBuffer, TArray<int32>& OutIndexBuffer, const float* Coords, int32 NumVerts, const int32* Faces, int32 NumFaces)
 	{
-		GetWorld()->GetTimerManager().SetTimer(TimerHandle, FTimerDelegate::CreateUObject(this, &UNavMeshRenderingComponent::TimerFunction), 1, true);
+		const int32 VertIndexBase = OutVertexBuffer.Num();
+		for (int32 VertIdx = 0; VertIdx < NumVerts * 3; VertIdx += 3)
+		{
+			OutVertexBuffer.Add(Recast2UnrealPoint(&Coords[VertIdx]));
+		}
+
+		const int32 FirstNewFaceVertexIndex = OutIndexBuffer.Num();
+		const uint32 NumIndices = NumFaces * 3;
+		OutIndexBuffer.AddUninitialized(NumIndices);
+		for (uint32 Index = 0; Index < NumIndices; ++Index)
+		{
+			OutIndexBuffer[FirstNewFaceVertexIndex + Index] = VertIndexBase + Faces[Index];
+		}
+	}
+
+	inline bool HasFlag(int32 Flags, ENavMeshDetailFlags TestFlag)
+	{
+		return (Flags & (1 << static_cast<int32>(TestFlag))) != 0;
 	}
-#endif //WITH_RECAST && !UE_BUILD_SHIPPING && !UE_BUILD_TEST
 }
 
-void UNavMeshRenderingComponent::OnUnregister()
+//////////////////////////////////////////////////////////////////////////
+// FNavMeshSceneProxyData
+
+void FNavMeshSceneProxyData::Reset()
 {
-#if WITH_RECAST && !UE_BUILD_SHIPPING && !UE_BUILD_TEST
-	// it's a kind of HACK but there is no event or other information that show flag was changed by user => we have to check it periodically
-#if WITH_EDITOR
-	if (GEditor)
+	MeshBuilders.Reset();
+	ThickLineItems.Reset();
+	TileEdgeLines.Reset();
+	NavMeshEdgeLines.Reset();
+	NavLinkLines.Reset();
+	ClusterLinkLines.Reset();
+	DebugLabels.Reset();
+	PathCollidingGeomIndices.Reset();
+	PathCollidingGeomVerts.Reset();
+	OctreeBounds.Reset();
+	Bounds.Init();
+
+	bNeedsNewData = true;
+	bDataGathered = false;
+	NavDetailFlags = 0;
+}
+
+void FNavMeshSceneProxyData::Serialize(FArchive& Ar)
+{
+	int32 NumMeshBuilders = MeshBuilders.Num();
+	Ar << NumMeshBuilders;
+	if (Ar.IsLoading())
 	{
-		GEditor->GetTimerManager()->ClearTimer(TimerHandle);
+		MeshBuilders.SetNum(NumMeshBuilders);
 	}
-	else
-#endif //WITH_EDITOR
+
+	for (int32 Idx = 0; Idx < NumMeshBuilders; Idx++)
 	{
-		GetWorld()->GetTimerManager().ClearTimer(TimerHandle);
+		FDebugMeshData& MeshBuilder = MeshBuilders[Idx];
+
+		int32 NumVerts = MeshBuilder.Vertices.Num();
+		Ar << NumVerts;
+		if (Ar.IsLoading())
+		{
+			MeshBuilder.Vertices.SetNum(NumVerts);
+		}
+
+		for (int32 VertIdx = 0; VertIdx < NumVerts; VertIdx++)
+		{
+			FVector SerializedVert = MeshBuilder.Vertices[VertIdx].Position;
+			Ar << SerializedVert;
+
+			if (Ar.IsLoading())
+			{
+				MeshBuilder.Vertices[VertIdx] = FDynamicMeshVertex(SerializedVert);
+			}
+		}
+
+		Ar << MeshBuilder.Indices;
+		Ar << MeshBuilder.ClusterColor;
 	}
-#endif //WITH_RECAST && !UE_BUILD_SHIPPING && !UE_BUILD_TEST
-	Super::OnUnregister();
-}
 
-FPrimitiveSceneProxy* UNavMeshRenderingComponent::CreateSceneProxy()
-{
-#if WITH_RECAST && !UE_BUILD_SHIPPING && !UE_BUILD_TEST
-	FPrimitiveSceneProxy* SceneProxy = NULL;
+	TArray<FDebugRenderSceneProxy::FDebugLine>* LineArraysToSerialize[] = { &ThickLineItems, &TileEdgeLines, &NavMeshEdgeLines, &NavLinkLines, &ClusterLinkLines };
+	for (int32 ArrIdx = 0; ArrIdx < ARRAY_COUNT(LineArraysToSerialize); ArrIdx++)
+	{
+		int32 NumItems = LineArraysToSerialize[ArrIdx]->Num();
+		Ar << NumItems;
+		if (Ar.IsLoading())
+		{
+			LineArraysToSerialize[ArrIdx]->Reset(NumItems);
+			LineArraysToSerialize[ArrIdx]->AddUninitialized(NumItems);
+		}
 
-	const bool bShowNavigation = IsNavigationShowFlagSet(GetWorld());
+		for (int32 Idx = 0; Idx < NumItems; Idx++)
+		{
+			Ar << (*LineArraysToSerialize[ArrIdx])[Idx].Thickness;
+			Ar << (*LineArraysToSerialize[ArrIdx])[Idx].Start;
+			Ar << (*LineArraysToSerialize[ArrIdx])[Idx].End;
+			Ar << (*LineArraysToSerialize[ArrIdx])[Idx].Color;
+		}
+	}
 
-	bCollectNavigationData = bShowNavigation;
+	int32 NumLabels = DebugLabels.Num();
+	Ar << NumLabels;
+	if (Ar.IsLoading())
+	{
+		DebugLabels.SetNum(NumLabels);
+	}
 
-	if (bCollectNavigationData && IsVisible())
+	for (int32 Idx = 0; Idx < NumLabels; Idx++)
 	{
-		FNavMeshSceneProxyData ProxyData;
-		ProxyData.Reset();
-		GatherData(ProxyData);
+		Ar << DebugLabels[Idx].Location;
+		Ar << DebugLabels[Idx].Text;
+	}
 
-		SceneProxy = new FRecastRenderingSceneProxy(this, &ProxyData);
-		ProxyData.Reset();
+	Ar << PathCollidingGeomIndices;
+
+	int32 NumVerts = PathCollidingGeomVerts.Num();
+	Ar << NumVerts;
+	if (Ar.IsLoading())
+	{
+		PathCollidingGeomVerts.SetNum(NumVerts);
 	}
-	return SceneProxy;
-#else
-	return NULL;
-#endif //WITH_RECAST && !UE_BUILD_SHIPPING && !UE_BUILD_TEST
-}
 
-void UNavMeshRenderingComponent::CreateRenderState_Concurrent()
-{
-	Super::CreateRenderState_Concurrent();
+	for (int32 Idx = 0; Idx < NumVerts; Idx++)
+	{
+		FVector SerializedVert = PathCollidingGeomVerts[Idx].Position;
+		Ar << SerializedVert;
 
-#if WITH_RECAST && !UE_BUILD_SHIPPING && !UE_BUILD_TEST
-	if (SceneProxy != NULL)
+		if (Ar.IsLoading())
+		{
+			PathCollidingGeomVerts[Idx] = FDynamicMeshVertex(SerializedVert);
+		}
+	}
+
+	int32 NumBounds = OctreeBounds.Num();
+	Ar << NumBounds;
+	if (Ar.IsLoading())
 	{
-		static_cast<FRecastRenderingSceneProxy*>(SceneProxy)->RegisterDebugDrawDelgate();
+		OctreeBounds.SetNum(NumBounds);
 	}
-#endif
-}
 
-void UNavMeshRenderingComponent::DestroyRenderState_Concurrent()
-{
-#if WITH_RECAST && !UE_BUILD_SHIPPING && !UE_BUILD_TEST
-	if (SceneProxy != NULL)
+	for (int32 Idx = 0; Idx < NumBounds; Idx++)
 	{
-		static_cast<FRecastRenderingSceneProxy*>(SceneProxy)->UnregisterDebugDrawDelgate();
+		Ar << OctreeBounds[Idx].Center;
+		Ar << OctreeBounds[Idx].Extent;
 	}
-#endif
 
-	Super::DestroyRenderState_Concurrent();
+	Ar << Bounds;
+	Ar << NavMeshDrawOffset;
+	Ar << NavDetailFlags;
+
+	int32 BitFlags = ((bDataGathered ? 1 : 0) << 0) | ((bNeedsNewData ? 1 : 0) << 1);
+	Ar << BitFlags;
+	bDataGathered = (BitFlags & (1 << 0)) != 0;
+	bNeedsNewData = (BitFlags & (1 << 1)) != 0;
 }
 
-FBoxSphereBounds UNavMeshRenderingComponent::CalcBounds(const FTransform& LocalToWorld) const
+uint32 FNavMeshSceneProxyData::GetAllocatedSize() const
 {
-	FBox BoundingBox(0);
+	return MeshBuilders.GetAllocatedSize() +
+		ThickLineItems.GetAllocatedSize() +
+		TileEdgeLines.GetAllocatedSize() +
+		NavMeshEdgeLines.GetAllocatedSize() +
+		NavLinkLines.GetAllocatedSize() +
+		ClusterLinkLines.GetAllocatedSize() +
+		DebugLabels.GetAllocatedSize() +
+		PathCollidingGeomIndices.GetAllocatedSize() +
+		PathCollidingGeomVerts.GetAllocatedSize() +
+		OctreeBounds.GetAllocatedSize();
+}
+
 #if WITH_RECAST
-	ARecastNavMesh* NavMesh = Cast<ARecastNavMesh>(GetOwner());
-	if (NavMesh)
-	{
-		BoundingBox = NavMesh->GetNavMeshBounds();
-		if (NavMesh->bDrawOctree)
-		{
-			const UNavigationSystem* NavSys = UNavigationSystem::GetCurrent(GetWorld());
-			const FNavigationOctree* NavOctree = NavSys ? NavSys->GetNavOctree() : NULL;
-			if (NavOctree)
-			{
-				for (FNavigationOctree::TConstIterator<> NodeIt(*NavOctree); NodeIt.HasPendingNodes(); NodeIt.Advance())
-				{
-					const FOctreeNodeContext& CorrentContext = NodeIt.GetCurrentContext();
-					BoundingBox += CorrentContext.Bounds.GetBox();
-				}
-			}
-		}
-	}
-#endif
-	return FBoxSphereBounds(BoundingBox);
+
+int32 FNavMeshSceneProxyData::GetDetailFlags(const ARecastNavMesh* NavMesh) const
+{
+	return (NavMesh == nullptr) ? 0 : 0 |
+		(NavMesh->bDrawTriangleEdges ? (1 << static_cast<int32>(ENavMeshDetailFlags::TriangleEdges)) : 0) |
+		(NavMesh->bDrawPolyEdges ? (1 << static_cast<int32>(ENavMeshDetailFlags::PolyEdges)) : 0) |
+		(NavMesh->bDrawFilledPolys ? (1 << static_cast<int32>(ENavMeshDetailFlags::FilledPolys)) : 0) |
+		(NavMesh->bDrawNavMeshEdges ? (1 << static_cast<int32>(ENavMeshDetailFlags::BoundaryEdges)) : 0) |
+		(NavMesh->bDrawTileBounds ? (1 << static_cast<int32>(ENavMeshDetailFlags::TileBounds)) : 0) |
+		(NavMesh->bDrawPathCollidingGeometry ? (1 << static_cast<int32>(ENavMeshDetailFlags::PathCollidingGeometry)) : 0) |
+		(NavMesh->bDrawTileLabels ? (1 << static_cast<int32>(ENavMeshDetailFlags::TileLabels)) : 0) |
+		(NavMesh->bDrawPolygonLabels ? (1 << static_cast<int32>(ENavMeshDetailFlags::PolygonLabels)) : 0) |
+		(NavMesh->bDrawDefaultPolygonCost ? (1 << static_cast<int32>(ENavMeshDetailFlags::PolygonCost)) : 0) |
+		(NavMesh->bDrawLabelsOnPathNodes ? (1 << static_cast<int32>(ENavMeshDetailFlags::PathLabels)) : 0) |
+		(NavMesh->bDrawNavLinks ? (1 << static_cast<int32>(ENavMeshDetailFlags::NavLinks)) : 0) |
+		(NavMesh->bDrawFailedNavLinks ? (1 << static_cast<int32>(ENavMeshDetailFlags::FailedNavLinks)) : 0) |
+		(NavMesh->bDrawClusters ? (1 << static_cast<int32>(ENavMeshDetailFlags::Clusters)) : 0) |
+		(NavMesh->bDrawOctree ? (1 << static_cast<int32>(ENavMeshDetailFlags::NavOctree)) : 0);
 }
 
-void UNavMeshRenderingComponent::GatherData(FNavMeshSceneProxyData& CurrentData) const
+void FNavMeshSceneProxyData::GatherData(const ARecastNavMesh* NavMesh, int32 InNavDetailFlags, const TArray<int32>& TileSet)
 {
 	QUICK_SCOPE_CYCLE_COUNTER(STAT_NavMesh_GatherDebugDrawingGeometry);
-#if WITH_RECAST
-	const ARecastNavMesh* NavMesh = Cast<ARecastNavMesh>(GetOwner());
+	Reset();
 
-	CurrentData.Reset();
-	
-	if (NavMesh && NavMesh->IsDrawingEnabled())
+	NavDetailFlags = InNavDetailFlags;
+	if (NavMesh && NavDetailFlags)
 	{
-		CurrentData.bEnableDrawing = true;
-		CurrentData.bNeedsNewData = false;
+		bNeedsNewData = false;
+		bDataGathered = true;
 
 		FHitProxyId HitProxyId = FHitProxyId();
-		CurrentData.bDrawPathCollidingGeometry = NavMesh->bDrawPathCollidingGeometry;
+		NavMeshDrawOffset.Z = NavMesh->DrawOffset;
 
-		CurrentData.NavMeshDrawOffset.Z = NavMesh->DrawOffset;
-		CurrentData.NavMeshGeometry.bGatherPolyEdges = NavMesh->bDrawPolyEdges;
-		CurrentData.NavMeshGeometry.bGatherNavMeshEdges = NavMesh->bDrawNavMeshEdges;
+		FRecastDebugGeometry NavMeshGeometry;
+		NavMeshGeometry.bGatherPolyEdges = FNavMeshRenderingHelpers::HasFlag(NavDetailFlags, ENavMeshDetailFlags::PolyEdges);
+		NavMeshGeometry.bGatherNavMeshEdges = FNavMeshRenderingHelpers::HasFlag(NavDetailFlags, ENavMeshDetailFlags::BoundaryEdges);
 
 		const FNavDataConfig& NavConfig = NavMesh->GetConfig();
-		CurrentData.NavMeshColors[RECAST_DEFAULT_AREA] = NavConfig.Color.DWColor() > 0 ? NavConfig.Color : NavMeshRenderColor_RecastMesh;
-		for (uint8 i = 0; i < RECAST_DEFAULT_AREA; ++i)
+		TArray<FColor> NavMeshColors;
+		NavMeshColors.AddDefaulted(RECAST_MAX_AREAS);
+
+		for (uint8 Idx = 0; Idx < RECAST_MAX_AREAS; Idx++)
 		{
-			CurrentData.NavMeshColors[i] = NavMesh->GetAreaIDColor(i);
+			NavMeshColors[Idx] = NavMesh->GetAreaIDColor(Idx);
 		}
+		NavMeshColors[RECAST_DEFAULT_AREA] = NavConfig.Color.DWColor() > 0 ? NavConfig.Color : NavMeshRenderColor_RecastMesh;
 
 		// just a little trick to make sure navmeshes with different sized are not drawn with same offset
-		CurrentData.NavMeshDrawOffset.Z += NavMesh->GetConfig().AgentRadius / 10.f;
+		NavMeshDrawOffset.Z += NavMesh->GetConfig().AgentRadius / 10.f;
 
 		NavMesh->BeginBatchQuery();
-
-		if (NavMesh->bDrawOctree)
+		if (TileSet.Num() > 0)
 		{
-			const UNavigationSystem* NavSys = UNavigationSystem::GetCurrent(GetWorld());
-			const FNavigationOctree* NavOctree = NavSys ? NavSys->GetNavOctree() : NULL;
-			if (NavOctree)
+			for (int32 Idx = 0; Idx < TileSet.Num(); Idx++)
 			{
-				for (FNavigationOctree::TConstIterator<> NodeIt(*NavOctree); NodeIt.HasPendingNodes(); NodeIt.Advance())
-				{
-					const FNavigationOctree::FNode& CurrentNode = NodeIt.GetCurrentNode();
-					const FOctreeNodeContext& CorrentContext = NodeIt.GetCurrentContext();
-					CurrentData.OctreeBounds.Add(CorrentContext.Bounds);
-
-					FOREACH_OCTREE_CHILD_NODE(ChildRef)
-					{
-						if (CurrentNode.HasChild(ChildRef))
-						{
-							NodeIt.PushChild(ChildRef);
-						}
-					}
-				}
+				NavMesh->GetDebugGeometry(NavMeshGeometry, TileSet[Idx]);
 			}
 		}
+		else
+		{
+			NavMesh->GetDebugGeometry(NavMeshGeometry);
+		}
 
-		NavMesh->GetDebugGeometry(CurrentData.NavMeshGeometry);
-		const TArray<FVector>& MeshVerts = CurrentData.NavMeshGeometry.MeshVerts;
+		const TArray<FVector>& MeshVerts = NavMeshGeometry.MeshVerts;
 
 		// @fixme, this is going to double up on lots of interior lines
-		if (NavMesh->bDrawTriangleEdges)
+		const bool bGatherTriEdges = FNavMeshRenderingHelpers::HasFlag(NavDetailFlags, ENavMeshDetailFlags::TriangleEdges);
+		if (bGatherTriEdges)
 		{
 			for (int32 AreaIdx = 0; AreaIdx < RECAST_MAX_AREAS; ++AreaIdx)
 			{
-				const TArray<int32>& MeshIndices = CurrentData.NavMeshGeometry.AreaIndices[AreaIdx];
-				for (int32 Idx=0; Idx<MeshIndices.Num(); Idx += 3)
+				const TArray<int32>& MeshIndices = NavMeshGeometry.AreaIndices[AreaIdx];
+				for (int32 Idx = 0; Idx < MeshIndices.Num(); Idx += 3)
 				{
-					CurrentData.ThickLineItems.Add(FNavMeshSceneProxyData::FDebugThickLine(MeshVerts[MeshIndices[Idx + 0]] + CurrentData.NavMeshDrawOffset, MeshVerts[MeshIndices[Idx + 1]] + CurrentData.NavMeshDrawOffset, NavMeshRenderColor_Recast_TriangleEdges, DefaultEdges_LineThickness));
-					CurrentData.ThickLineItems.Add(FNavMeshSceneProxyData::FDebugThickLine(MeshVerts[MeshIndices[Idx + 1]] + CurrentData.NavMeshDrawOffset, MeshVerts[MeshIndices[Idx + 2]] + CurrentData.NavMeshDrawOffset, NavMeshRenderColor_Recast_TriangleEdges, DefaultEdges_LineThickness));
-					CurrentData.ThickLineItems.Add(FNavMeshSceneProxyData::FDebugThickLine(MeshVerts[MeshIndices[Idx + 2]] + CurrentData.NavMeshDrawOffset, MeshVerts[MeshIndices[Idx + 0]] + CurrentData.NavMeshDrawOffset, NavMeshRenderColor_Recast_TriangleEdges, DefaultEdges_LineThickness));
+					ThickLineItems.Add(FDebugRenderSceneProxy::FDebugLine(MeshVerts[MeshIndices[Idx + 0]] + NavMeshDrawOffset, MeshVerts[MeshIndices[Idx + 1]] + NavMeshDrawOffset, NavMeshRenderColor_Recast_TriangleEdges, DefaultEdges_LineThickness));
+					ThickLineItems.Add(FDebugRenderSceneProxy::FDebugLine(MeshVerts[MeshIndices[Idx + 1]] + NavMeshDrawOffset, MeshVerts[MeshIndices[Idx + 2]] + NavMeshDrawOffset, NavMeshRenderColor_Recast_TriangleEdges, DefaultEdges_LineThickness));
+					ThickLineItems.Add(FDebugRenderSceneProxy::FDebugLine(MeshVerts[MeshIndices[Idx + 2]] + NavMeshDrawOffset, MeshVerts[MeshIndices[Idx + 0]] + NavMeshDrawOffset, NavMeshRenderColor_Recast_TriangleEdges, DefaultEdges_LineThickness));
 				}
 			}
 		}
 
 		// make lines for tile edges
-		if (NavMesh->bDrawPolyEdges)
+		const bool bGatherPolyEdges = FNavMeshRenderingHelpers::HasFlag(NavDetailFlags, ENavMeshDetailFlags::PolyEdges);
+		if (bGatherPolyEdges)
 		{
-			const TArray<FVector>& TileEdgeVerts = CurrentData.NavMeshGeometry.PolyEdges;
-			for (int32 Idx=0; Idx < TileEdgeVerts.Num(); Idx += 2)
+			const TArray<FVector>& TileEdgeVerts = NavMeshGeometry.PolyEdges;
+			for (int32 Idx = 0; Idx < TileEdgeVerts.Num(); Idx += 2)
 			{
-				CurrentData.TileEdgeLines.Add(FDebugRenderSceneProxy::FDebugLine(TileEdgeVerts[Idx] + CurrentData.NavMeshDrawOffset, TileEdgeVerts[Idx + 1] + CurrentData.NavMeshDrawOffset, NavMeshRenderColor_Recast_TileEdges));
+				TileEdgeLines.Add(FDebugRenderSceneProxy::FDebugLine(TileEdgeVerts[Idx] + NavMeshDrawOffset, TileEdgeVerts[Idx + 1] + NavMeshDrawOffset, NavMeshRenderColor_Recast_TileEdges));
 			}
 		}
 
 		// make lines for navmesh edges
-		if (NavMesh->bDrawNavMeshEdges)
+		const bool bGatherBoundaryEdges = FNavMeshRenderingHelpers::HasFlag(NavDetailFlags, ENavMeshDetailFlags::BoundaryEdges);
+		if (bGatherBoundaryEdges)
 		{
-			const FColor EdgesColor = DarkenColor(CurrentData.NavMeshColors[RECAST_DEFAULT_AREA]);
-			const TArray<FVector>& NavMeshEdgeVerts = CurrentData.NavMeshGeometry.NavMeshEdges;
-			for (int32 Idx=0; Idx < NavMeshEdgeVerts.Num(); Idx += 2)
+			const FColor EdgesColor = FNavMeshRenderingHelpers::DarkenColor(NavMeshColors[RECAST_DEFAULT_AREA]);
+			const TArray<FVector>& NavMeshEdgeVerts = NavMeshGeometry.NavMeshEdges;
+			for (int32 Idx = 0; Idx < NavMeshEdgeVerts.Num(); Idx += 2)
 			{
-				CurrentData.NavMeshEdgeLines.Add( FDebugRenderSceneProxy::FDebugLine(NavMeshEdgeVerts[Idx] + CurrentData.NavMeshDrawOffset, NavMeshEdgeVerts[Idx+1] + CurrentData.NavMeshDrawOffset, EdgesColor));
+				NavMeshEdgeLines.Add(FDebugRenderSceneProxy::FDebugLine(NavMeshEdgeVerts[Idx] + NavMeshDrawOffset, NavMeshEdgeVerts[Idx + 1] + NavMeshDrawOffset, EdgesColor));
 			}
 		}
 
 		// offset all navigation-link positions
-		if (!NavMesh->bDrawClusters)
+		const bool bGatherClusters = FNavMeshRenderingHelpers::HasFlag(NavDetailFlags, ENavMeshDetailFlags::Clusters);
+		const bool bGatherNavLinks = FNavMeshRenderingHelpers::HasFlag(NavDetailFlags, ENavMeshDetailFlags::NavLinks);
+		const bool bGatherFailedNavLinks = FNavMeshRenderingHelpers::HasFlag(NavDetailFlags, ENavMeshDetailFlags::FailedNavLinks);
+
+		if (bGatherClusters == false)
 		{
-			for (int32 OffMeshLineIndex = 0; OffMeshLineIndex < CurrentData.NavMeshGeometry.OffMeshLinks.Num(); ++OffMeshLineIndex)
+			for (int32 OffMeshLineIndex = 0; OffMeshLineIndex < NavMeshGeometry.OffMeshLinks.Num(); ++OffMeshLineIndex)
 			{
-				FRecastDebugGeometry::FOffMeshLink& Link = CurrentData.NavMeshGeometry.OffMeshLinks[OffMeshLineIndex];
+				FRecastDebugGeometry::FOffMeshLink& Link = NavMeshGeometry.OffMeshLinks[OffMeshLineIndex];
 				const bool bLinkValid = (Link.ValidEnds & FRecastDebugGeometry::OMLE_Left) && (Link.ValidEnds & FRecastDebugGeometry::OMLE_Right);
 
-				if (NavMesh->bDrawFailedNavLinks || (NavMesh->bDrawNavLinks && bLinkValid))
+				if (bGatherFailedNavLinks || (bGatherNavLinks && bLinkValid))
 				{
-					const FVector V0 = Link.Left + CurrentData.NavMeshDrawOffset;
-					const FVector V1 = Link.Right + CurrentData.NavMeshDrawOffset;
-					const FColor LinkColor = ((Link.Direction && Link.ValidEnds) || (Link.ValidEnds & FRecastDebugGeometry::OMLE_Left)) ? DarkenColor(CurrentData.NavMeshColors[Link.AreaID]) : NavMeshRenderColor_OffMeshConnectionInvalid;
+					const FVector V0 = Link.Left + NavMeshDrawOffset;
+					const FVector V1 = Link.Right + NavMeshDrawOffset;
+					const FColor LinkColor = ((Link.Direction && Link.ValidEnds) || (Link.ValidEnds & FRecastDebugGeometry::OMLE_Left)) ? FNavMeshRenderingHelpers::DarkenColor(NavMeshColors[Link.AreaID]) : NavMeshRenderColor_OffMeshConnectionInvalid;
 
-					CacheArc(CurrentData.NavLinkLines, V0, V1, 0.4f, 4, LinkColor, LinkLines_LineThickness);
+					FNavMeshRenderingHelpers::CacheArc(NavLinkLines, V0, V1, 0.4f, 4, LinkColor, LinkLines_LineThickness);
 
 					const FVector VOffset(0, 0, FVector::Dist(V0, V1) * 1.333f);
-					CacheArrowHead(CurrentData.NavLinkLines, V1, V0+VOffset, 30.f, LinkColor, LinkLines_LineThickness);
+					FNavMeshRenderingHelpers::CacheArrowHead(NavLinkLines, V1, V0 + VOffset, 30.f, LinkColor, LinkLines_LineThickness);
 					if (Link.Direction)
 					{
-						CacheArrowHead(CurrentData.NavLinkLines, V0, V1+VOffset, 30.f, LinkColor, LinkLines_LineThickness);
+						FNavMeshRenderingHelpers::CacheArrowHead(NavLinkLines, V0, V1 + VOffset, 30.f, LinkColor, LinkLines_LineThickness);
 					}
 
 					// if the connection as a whole is valid check if there are any of ends is invalid
@@ -345,25 +470,43 @@ void UNavMeshRenderingComponent::GatherData(FNavMeshSceneProxyData& CurrentData)
 						if (Link.Direction && (Link.ValidEnds & FRecastDebugGeometry::OMLE_Left) == 0)
 						{
 							// left end invalid - mark it
-							DrawWireCylinder(CurrentData.NavLinkLines, V0, FVector(1,0,0), FVector(0,1,0), FVector(0,0,1), NavMeshRenderColor_OffMeshConnectionInvalid, Link.Radius, NavMesh->AgentMaxStepHeight, 16, 0, DefaultEdges_LineThickness);
+							FNavMeshRenderingHelpers::DrawWireCylinder(NavLinkLines, V0, FVector(1, 0, 0), FVector(0, 1, 0), FVector(0, 0, 1), NavMeshRenderColor_OffMeshConnectionInvalid, Link.Radius, NavMesh->AgentMaxStepHeight, 16, 0, DefaultEdges_LineThickness);
 						}
 						if ((Link.ValidEnds & FRecastDebugGeometry::OMLE_Right) == 0)
 						{
-							DrawWireCylinder(CurrentData.NavLinkLines, V1, FVector(1,0,0), FVector(0,1,0), FVector(0,0,1), NavMeshRenderColor_OffMeshConnectionInvalid, Link.Radius, NavMesh->AgentMaxStepHeight, 16, 0, DefaultEdges_LineThickness);
+							FNavMeshRenderingHelpers::DrawWireCylinder(NavLinkLines, V1, FVector(1, 0, 0), FVector(0, 1, 0), FVector(0, 0, 1), NavMeshRenderColor_OffMeshConnectionInvalid, Link.Radius, NavMesh->AgentMaxStepHeight, 16, 0, DefaultEdges_LineThickness);
 						}
 					}
-				}					
+				}
 			}
 		}
 
-		if (NavMesh->bDrawTileLabels || NavMesh->bDrawPolygonLabels || NavMesh->bDrawDefaultPolygonCost || NavMesh->bDrawTileBounds)
+		const bool bGatherTileLabels = FNavMeshRenderingHelpers::HasFlag(NavDetailFlags, ENavMeshDetailFlags::TileLabels);
+		const bool bGatherTileBounds = FNavMeshRenderingHelpers::HasFlag(NavDetailFlags, ENavMeshDetailFlags::TileBounds);
+		const bool bGatherPolygonLabels = FNavMeshRenderingHelpers::HasFlag(NavDetailFlags, ENavMeshDetailFlags::PolygonLabels);
+		const bool bGatherPolygonCost = FNavMeshRenderingHelpers::HasFlag(NavDetailFlags, ENavMeshDetailFlags::PolygonCost);
+
+		if (bGatherTileLabels || bGatherTileBounds || bGatherPolygonLabels || bGatherPolygonCost)
 		{
-			// calculate appropriate points for displaying debug labels
-			const int32 TilesCount = NavMesh->GetNavMeshTilesCount();
-			CurrentData.DebugLabels.Reserve(TilesCount);
+			TArray<int32> UseTileIndices;
+			if (TileSet.Num() > 0)
+			{
+				UseTileIndices = TileSet;
+			}
+			else
+			{
+				const int32 TilesCount = NavMesh->GetNavMeshTilesCount();
+				for (int32 Idx = 0; Idx < TilesCount; Idx++)
+				{
+					UseTileIndices.Add(Idx);
+				}
+			}
 
-			for (int32 TileIndex = 0; TileIndex < TilesCount; ++TileIndex)
+			// calculate appropriate points for displaying debug labels
+			DebugLabels.Reserve(UseTileIndices.Num());
+			for (int32 TileSetIdx = 0; TileSetIdx < UseTileIndices.Num(); TileSetIdx++)
 			{
+				const int32 TileIndex = UseTileIndices[TileSetIdx];
 				int32 X, Y, Layer;
 				if (NavMesh->GetNavMeshTileXY(TileIndex, X, Y, Layer))
 				{
@@ -372,243 +515,251 @@ void UNavMeshRenderingComponent::GatherData(FNavMeshSceneProxyData& CurrentData)
 					TileLabelLocation.Z = TileBoundingBox.Max.Z;
 
 					FNavLocation NavLocation(TileLabelLocation);
-					if (!NavMesh->ProjectPoint(TileLabelLocation, NavLocation, FVector(NavMesh->TileSizeUU/100, NavMesh->TileSizeUU/100, TileBoundingBox.Max.Z-TileBoundingBox.Min.Z)))
+					if (!NavMesh->ProjectPoint(TileLabelLocation, NavLocation, FVector(NavMesh->TileSizeUU / 100, NavMesh->TileSizeUU / 100, TileBoundingBox.Max.Z - TileBoundingBox.Min.Z)))
 					{
-						NavMesh->ProjectPoint(TileLabelLocation, NavLocation, FVector(NavMesh->TileSizeUU/2, NavMesh->TileSizeUU/2, TileBoundingBox.Max.Z-TileBoundingBox.Min.Z));
+						NavMesh->ProjectPoint(TileLabelLocation, NavLocation, FVector(NavMesh->TileSizeUU / 2, NavMesh->TileSizeUU / 2, TileBoundingBox.Max.Z - TileBoundingBox.Min.Z));
 					}
 
-					if (NavMesh->bDrawTileLabels)
+					if (bGatherTileLabels)
 					{
-						CurrentData.DebugLabels.Add(FNavMeshSceneProxyData::FDebugText(
-							/*Location*/NavLocation.Location + CurrentData.NavMeshDrawOffset
-							, /*Text*/FString::Printf(TEXT("(%d,%d:%d)"), X, Y, Layer)
-							));
+						DebugLabels.Add(FDebugText(NavLocation.Location + NavMeshDrawOffset, FString::Printf(TEXT("(%d,%d:%d)"), X, Y, Layer)));
 					}
 
-					if (NavMesh->bDrawPolygonLabels || NavMesh->bDrawDefaultPolygonCost)
+					if (bGatherPolygonLabels || bGatherPolygonCost)
 					{
 						TArray<FNavPoly> Polys;
 						NavMesh->GetPolysInTile(TileIndex, Polys);
 
-						if (NavMesh->bDrawDefaultPolygonCost)
+						if (bGatherPolygonCost)
 						{
 							float DefaultCosts[RECAST_MAX_AREAS];
 							float FixedCosts[RECAST_MAX_AREAS];
 
 							NavMesh->GetDefaultQueryFilter()->GetAllAreaCosts(DefaultCosts, FixedCosts, RECAST_MAX_AREAS);
 
-							for(int k = 0; k < Polys.Num(); ++k)
+							for (int k = 0; k < Polys.Num(); ++k)
 							{
-								uint32 AreaID = NavMesh->GetPolyAreaID(Polys[k].Ref);
-
-								CurrentData.DebugLabels.Add(FNavMeshSceneProxyData::FDebugText(
-									/*Location*/Polys[k].Center + CurrentData.NavMeshDrawOffset
-									, /*Text*/FString::Printf(TEXT("\\%.3f; %.3f\\"), DefaultCosts[AreaID], FixedCosts[AreaID])
-									));
+								const uint32 AreaID = NavMesh->GetPolyAreaID(Polys[k].Ref);
+								DebugLabels.Add(FDebugText(Polys[k].Center + NavMeshDrawOffset, FString::Printf(TEXT("\\%.3f; %.3f\\"), DefaultCosts[AreaID], FixedCosts[AreaID])));
 							}
 						}
 						else
 						{
-							for(int k = 0; k < Polys.Num(); ++k)
+							for (int k = 0; k < Polys.Num(); ++k)
 							{
 								uint32 NavPolyIndex = 0;
 								uint32 NavTileIndex = 0;
 								NavMesh->GetPolyTileIndex(Polys[k].Ref, NavPolyIndex, NavTileIndex);
 
-								CurrentData.DebugLabels.Add(FNavMeshSceneProxyData::FDebugText(
-									/*Location*/Polys[k].Center + CurrentData.NavMeshDrawOffset
-									, /*Text*/FString::Printf(TEXT("[%X:%X]"), NavTileIndex, NavPolyIndex)
-									));
+								DebugLabels.Add(FDebugText(Polys[k].Center + NavMeshDrawOffset, FString::Printf(TEXT("[%X:%X]"), NavTileIndex, NavPolyIndex)));
 							}
 						}
-					}							
+					}
 
-					if (NavMesh->bDrawTileBounds)
+					if (bGatherTileBounds)
 					{
-						FBox TileBox = NavMesh->GetNavMeshTileBounds(TileIndex);
-
-						float DrawZ = (TileBox.Min.Z + TileBox.Max.Z) * 0.5f;		// @hack average
-						FVector LL(TileBox.Min.X, TileBox.Min.Y, DrawZ);
-						FVector UR(TileBox.Max.X, TileBox.Max.Y, DrawZ);
-						FVector UL(LL.X, UR.Y, DrawZ);
-						FVector LR(UR.X, LL.Y, DrawZ);
-						CurrentData.ThickLineItems.Add(FNavMeshSceneProxyData::FDebugThickLine(LL, UL, NavMeshRenderColor_TileBounds, DefaultEdges_LineThickness));
-						CurrentData.ThickLineItems.Add(FNavMeshSceneProxyData::FDebugThickLine(UL, UR, NavMeshRenderColor_TileBounds, DefaultEdges_LineThickness));
-						CurrentData.ThickLineItems.Add(FNavMeshSceneProxyData::FDebugThickLine(UR, LR, NavMeshRenderColor_TileBounds, DefaultEdges_LineThickness));
-						CurrentData.ThickLineItems.Add(FNavMeshSceneProxyData::FDebugThickLine(LR, LL, NavMeshRenderColor_TileBounds, DefaultEdges_LineThickness));
+						const FBox TileBox = NavMesh->GetNavMeshTileBounds(TileIndex);
+						const float DrawZ = (TileBox.Min.Z + TileBox.Max.Z) * 0.5f;
+						const FVector LL(TileBox.Min.X, TileBox.Min.Y, DrawZ);
+						const FVector UR(TileBox.Max.X, TileBox.Max.Y, DrawZ);
+						const FVector UL(LL.X, UR.Y, DrawZ);
+						const FVector LR(UR.X, LL.Y, DrawZ);
+
+						ThickLineItems.Add(FDebugRenderSceneProxy::FDebugLine(LL, UL, NavMeshRenderColor_TileBounds, DefaultEdges_LineThickness));
+						ThickLineItems.Add(FDebugRenderSceneProxy::FDebugLine(UL, UR, NavMeshRenderColor_TileBounds, DefaultEdges_LineThickness));
+						ThickLineItems.Add(FDebugRenderSceneProxy::FDebugLine(UR, LR, NavMeshRenderColor_TileBounds, DefaultEdges_LineThickness));
+						ThickLineItems.Add(FDebugRenderSceneProxy::FDebugLine(LR, LL, NavMeshRenderColor_TileBounds, DefaultEdges_LineThickness));
 					}
 				}
 			}
 		}
 
-		CurrentData.bSkipDistanceCheck = GIsEditor && (GEngine->GetDebugLocalPlayer() == NULL);
-		CurrentData.bDrawClusters = NavMesh->bDrawClusters;
 		NavMesh->FinishBatchQuery();
 
 		// Draw Mesh
-		if (NavMesh->bDrawFilledPolys)
+		const bool bGatherFilledPolys = FNavMeshRenderingHelpers::HasFlag(NavDetailFlags, ENavMeshDetailFlags::FilledPolys);
+		if (bGatherFilledPolys)
 		{
-			if (NavMesh->bDrawClusters)
+			if (bGatherClusters)
 			{
-				for (int32 Idx = 0; Idx < CurrentData.NavMeshGeometry.Clusters.Num(); ++Idx)
+				for (int32 Idx = 0; Idx < NavMeshGeometry.Clusters.Num(); ++Idx)
 				{
-					const TArray<int32>& MeshIndices = CurrentData.NavMeshGeometry.Clusters[Idx].MeshIndices;
-
+					const TArray<int32>& MeshIndices = NavMeshGeometry.Clusters[Idx].MeshIndices;
 					if (MeshIndices.Num() == 0)
 					{
 						continue;
 					}
 
-					FNavMeshSceneProxyData::FDebugMeshData DebugMeshData;
-					DebugMeshData.ClusterColor = GetClusterColor(Idx);
+					FDebugMeshData DebugMeshData;
+					DebugMeshData.ClusterColor = FNavMeshRenderingHelpers::GetClusterColor(Idx);
 					for (int32 VertIdx = 0; VertIdx < MeshVerts.Num(); ++VertIdx)
 					{
-						AddVertexHelper(DebugMeshData, MeshVerts[VertIdx] + CurrentData.NavMeshDrawOffset, DebugMeshData.ClusterColor);
+						FNavMeshRenderingHelpers::AddVertex(DebugMeshData, MeshVerts[VertIdx] + NavMeshDrawOffset, DebugMeshData.ClusterColor);
 					}
 					for (int32 TriIdx = 0; TriIdx < MeshIndices.Num(); TriIdx += 3)
 					{
-						AddTriangleHelper(DebugMeshData, MeshIndices[TriIdx], MeshIndices[TriIdx + 1], MeshIndices[TriIdx + 2]);
+						FNavMeshRenderingHelpers::AddTriangle(DebugMeshData, MeshIndices[TriIdx], MeshIndices[TriIdx + 1], MeshIndices[TriIdx + 2]);
 					}
 
-					CurrentData.MeshBuilders.Add(DebugMeshData);
+					MeshBuilders.Add(DebugMeshData);
 				}
 			}
-			else if (NavMesh->bDrawNavMesh)
+			else
 			{
 				for (int32 AreaType = 0; AreaType < RECAST_MAX_AREAS; ++AreaType)
 				{
-					const TArray<int32>& MeshIndices = CurrentData.NavMeshGeometry.AreaIndices[AreaType];
-
+					const TArray<int32>& MeshIndices = NavMeshGeometry.AreaIndices[AreaType];
 					if (MeshIndices.Num() == 0)
 					{
 						continue;
 					}
 
-					FNavMeshSceneProxyData::FDebugMeshData DebugMeshData;
+					FDebugMeshData DebugMeshData;
 					for (int32 VertIdx = 0; VertIdx < MeshVerts.Num(); ++VertIdx)
 					{
-						AddVertexHelper(DebugMeshData, MeshVerts[VertIdx] + CurrentData.NavMeshDrawOffset, CurrentData.NavMeshColors[AreaType]);
+						FNavMeshRenderingHelpers::AddVertex(DebugMeshData, MeshVerts[VertIdx] + NavMeshDrawOffset, NavMeshColors[AreaType]);
 					}
 					for (int32 TriIdx = 0; TriIdx < MeshIndices.Num(); TriIdx += 3)
 					{
-						AddTriangleHelper(DebugMeshData, MeshIndices[TriIdx], MeshIndices[TriIdx + 1], MeshIndices[TriIdx + 2]);
+						FNavMeshRenderingHelpers::AddTriangle(DebugMeshData, MeshIndices[TriIdx], MeshIndices[TriIdx + 1], MeshIndices[TriIdx + 2]);
 					}
 
-					DebugMeshData.ClusterColor = CurrentData.NavMeshColors[AreaType];
-					CurrentData.MeshBuilders.Add(DebugMeshData);
+					DebugMeshData.ClusterColor = NavMeshColors[AreaType];
+					MeshBuilders.Add(DebugMeshData);
 				}
 			}
 		}
 
-		if (NavMesh->bDrawPathCollidingGeometry)
+		const bool bGatherOctree = FNavMeshRenderingHelpers::HasFlag(NavDetailFlags, ENavMeshDetailFlags::NavOctree);
+		const bool bGatherPathCollidingGeometry = FNavMeshRenderingHelpers::HasFlag(NavDetailFlags, ENavMeshDetailFlags::PathCollidingGeometry);
+		if (bGatherOctree || bGatherPathCollidingGeometry)
 		{
-			// draw all geometry gathered in navoctree
-			const FNavigationOctree* NavOctree = NavMesh->GetWorld()->GetNavigationSystem()->GetNavOctree();
-
+			const UNavigationSystem* NavSys = UNavigationSystem::GetCurrent(NavMesh->GetWorld());
+			const FNavigationOctree* NavOctree = NavSys ? NavSys->GetNavOctree() : nullptr;
 			if (NavOctree)
 			{
-				TArray<FVector> PathCollidingGeomVerts;
-				TArray <int32> PathCollidingGeomIndices;
-				for (FNavigationOctree::TConstIterator<> It(*NavOctree); It.HasPendingNodes(); It.Advance())
+				if (bGatherOctree)
 				{
-					const FNavigationOctree::FNode& Node = It.GetCurrentNode();
-					for (FNavigationOctree::ElementConstIt ElementIt(Node.GetElementIt()); ElementIt; ElementIt++)
+					for (FNavigationOctree::TConstIterator<> NodeIt(*NavOctree); NodeIt.HasPendingNodes(); NodeIt.Advance())
 					{
-						const FNavigationOctreeElement& Element = *ElementIt;
-						if (Element.ShouldUseGeometry(NavMesh->GetConfig()) && Element.Data->CollisionData.Num())
+						const FNavigationOctree::FNode& CurrentNode = NodeIt.GetCurrentNode();
+						const FOctreeNodeContext& CorrentContext = NodeIt.GetCurrentContext();
+						OctreeBounds.Add(CorrentContext.Bounds);
+
+						FOREACH_OCTREE_CHILD_NODE(ChildRef)
 						{
-							const FRecastGeometryCache CachedGeometry(Element.Data->CollisionData.GetData());
-							AppendGeometry(PathCollidingGeomVerts, PathCollidingGeomIndices, CachedGeometry.Verts, CachedGeometry.Header.NumVerts, CachedGeometry.Indices, CachedGeometry.Header.NumFaces);
+							if (CurrentNode.HasChild(ChildRef))
+							{
+								NodeIt.PushChild(ChildRef);
+							}
 						}
 					}
-					FOREACH_OCTREE_CHILD_NODE(ChildRef)
+				}
+
+				if (bGatherPathCollidingGeometry)
+				{
+					TArray<FVector> CollidingVerts;
+
+					for (FNavigationOctree::TConstIterator<> It(*NavOctree); It.HasPendingNodes(); It.Advance())
 					{
-						if (Node.HasChild(ChildRef))
+						const FNavigationOctree::FNode& Node = It.GetCurrentNode();
+						for (FNavigationOctree::ElementConstIt ElementIt(Node.GetElementIt()); ElementIt; ElementIt++)
+						{
+							const FNavigationOctreeElement& Element = *ElementIt;
+							if (Element.ShouldUseGeometry(NavMesh->GetConfig()) && Element.Data->CollisionData.Num())
+							{
+								const FRecastGeometryCache CachedGeometry(Element.Data->CollisionData.GetData());
+								FNavMeshRenderingHelpers::AddRecastGeometry(CollidingVerts, PathCollidingGeomIndices, CachedGeometry.Verts, CachedGeometry.Header.NumVerts, CachedGeometry.Indices, CachedGeometry.Header.NumFaces);
+							}
+						}
+						FOREACH_OCTREE_CHILD_NODE(ChildRef)
 						{
-							It.PushChild(ChildRef);
+							if (Node.HasChild(ChildRef))
+							{
+								It.PushChild(ChildRef);
+							}
 						}
 					}
-				}
-				CurrentData.PathCollidingGeomIndices = PathCollidingGeomIndices;
-				for (const auto& Vertex : PathCollidingGeomVerts)
-				{
-					CurrentData.PathCollidingGeomVerts.Add(FDynamicMeshVertex(Vertex));
+
+					PathCollidingGeomVerts.SetNum(CollidingVerts.Num(), false);
+					for (int32 Idx = 0; Idx < 0; Idx++)
+					{
+						PathCollidingGeomVerts[Idx] = FDynamicMeshVertex(CollidingVerts[Idx]);
+					}
 				}
 			}
 		}
 
-		if (CurrentData.NavMeshGeometry.BuiltMeshIndices.Num() > 0)
+		if (NavMeshGeometry.BuiltMeshIndices.Num() > 0)
 		{
-			FNavMeshSceneProxyData::FDebugMeshData DebugMeshData;
-			for (int32 VertIdx=0; VertIdx < MeshVerts.Num(); ++VertIdx)
+			FDebugMeshData DebugMeshData;
+			for (int32 VertIdx = 0; VertIdx < MeshVerts.Num(); ++VertIdx)
 			{
-				AddVertexHelper(DebugMeshData, MeshVerts[VertIdx] + CurrentData.NavMeshDrawOffset, NavMeshRenderColor_RecastTileBeingRebuilt);
+				FNavMeshRenderingHelpers::AddVertex(DebugMeshData, MeshVerts[VertIdx] + NavMeshDrawOffset, NavMeshRenderColor_RecastTileBeingRebuilt);
 			}
-			DebugMeshData.Indices.Append(CurrentData.NavMeshGeometry.BuiltMeshIndices);
+			DebugMeshData.Indices.Append(NavMeshGeometry.BuiltMeshIndices);
 			DebugMeshData.ClusterColor = NavMeshRenderColor_RecastTileBeingRebuilt;
-			CurrentData.MeshBuilders.Add(DebugMeshData);
-
-			// updates should be requested by FRecastNavMeshGenerator::TickAsyncBuild after tiles were refreshed
+			MeshBuilders.Add(DebugMeshData);
 		}
 
-		if (NavMesh->bDrawClusters)
+		if (bGatherClusters)
 		{
-			for (int i = 0; i < CurrentData.NavMeshGeometry.ClusterLinks.Num(); i++)
+			for (int32 Idx = 0; Idx < NavMeshGeometry.ClusterLinks.Num(); Idx++)
 			{
-				const FRecastDebugGeometry::FClusterLink& CLink = CurrentData.NavMeshGeometry.ClusterLinks[i];
-				const FVector V0 = CLink.FromCluster + CurrentData.NavMeshDrawOffset;
-				const FVector V1 = CLink.ToCluster + CurrentData.NavMeshDrawOffset + FVector(0,0,20.0f);
+				const FRecastDebugGeometry::FClusterLink& CLink = NavMeshGeometry.ClusterLinks[Idx];
+				const FVector V0 = CLink.FromCluster + NavMeshDrawOffset;
+				const FVector V1 = CLink.ToCluster + NavMeshDrawOffset + FVector(0, 0, 20.0f);
 
-				CacheArc(CurrentData.ClusterLinkLines, V0, V1, 0.4f, 4, FColor::Black, ClusterLinkLines_LineThickness);
+				FNavMeshRenderingHelpers::CacheArc(ClusterLinkLines, V0, V1, 0.4f, 4, FColor::Black, ClusterLinkLines_LineThickness);
 				const FVector VOffset(0, 0, FVector::Dist(V0, V1) * 1.333f);
-				CacheArrowHead(CurrentData.ClusterLinkLines, V1, V0+VOffset, 30.f, FColor::Black, ClusterLinkLines_LineThickness);
+				FNavMeshRenderingHelpers::CacheArrowHead(ClusterLinkLines, V1, V0 + VOffset, 30.f, FColor::Black, ClusterLinkLines_LineThickness);
 			}
 		}
 
 		// cache segment links
-		if (NavMesh->bDrawNavLinks)
+		if (bGatherNavLinks)
 		{
-			for (int32 iArea = 0; iArea < RECAST_MAX_AREAS; iArea++)
+			for (int32 AreaIdx = 0; AreaIdx < RECAST_MAX_AREAS; AreaIdx++)
 			{
-				const TArray<int32>& Indices = CurrentData.NavMeshGeometry.OffMeshSegmentAreas[iArea];
+				const TArray<int32>& Indices = NavMeshGeometry.OffMeshSegmentAreas[AreaIdx];
 				FNavMeshSceneProxyData::FDebugMeshData DebugMeshData;
 				int32 VertBase = 0;
 
-				for (int32 i = 0; i < Indices.Num(); i++)
+				for (int32 Idx = 0; Idx < Indices.Num(); Idx++)
 				{
-					FRecastDebugGeometry::FOffMeshSegment& SegInfo = CurrentData.NavMeshGeometry.OffMeshSegments[Indices[i]];
-					const FVector A0 = SegInfo.LeftStart + CurrentData.NavMeshDrawOffset;
-					const FVector A1 = SegInfo.LeftEnd + CurrentData.NavMeshDrawOffset;
-					const FVector B0 = SegInfo.RightStart + CurrentData.NavMeshDrawOffset;
-					const FVector B1 = SegInfo.RightEnd + CurrentData.NavMeshDrawOffset;
+					FRecastDebugGeometry::FOffMeshSegment& SegInfo = NavMeshGeometry.OffMeshSegments[Indices[Idx]];
+					const FVector A0 = SegInfo.LeftStart + NavMeshDrawOffset;
+					const FVector A1 = SegInfo.LeftEnd + NavMeshDrawOffset;
+					const FVector B0 = SegInfo.RightStart + NavMeshDrawOffset;
+					const FVector B1 = SegInfo.RightEnd + NavMeshDrawOffset;
 					const FVector Edge0 = B0 - A0;
 					const FVector Edge1 = B1 - A1;
 					const float Len0 = Edge0.Size();
 					const float Len1 = Edge1.Size();
-					const FColor SegColor = DarkenColor(CurrentData.NavMeshColors[SegInfo.AreaID]);
+					const FColor SegColor = FNavMeshRenderingHelpers::DarkenColor(NavMeshColors[SegInfo.AreaID]);
 					const FColor ColA = (SegInfo.ValidEnds & FRecastDebugGeometry::OMLE_Left) ? FColor::White : FColor::Black;
 					const FColor ColB = (SegInfo.ValidEnds & FRecastDebugGeometry::OMLE_Right) ? FColor::White : FColor::Black;
 
 					const int32 NumArcPoints = 8;
 					const float ArcPtsScale = 1.0f / NumArcPoints;
 
-					FVector Prev0 = EvalArc(A0, Edge0, Len0*0.25f, 0);
-					FVector Prev1 = EvalArc(A1, Edge1, Len1*0.25f, 0);
-					AddVertexHelper(DebugMeshData, Prev0, ColA);
-					AddVertexHelper(DebugMeshData, Prev1, ColA);
-					for (int32 j = 1; j <= NumArcPoints; j++)
+					FVector Prev0 = FNavMeshRenderingHelpers::EvalArc(A0, Edge0, Len0*0.25f, 0);
+					FVector Prev1 = FNavMeshRenderingHelpers::EvalArc(A1, Edge1, Len1*0.25f, 0);
+					FNavMeshRenderingHelpers::AddVertex(DebugMeshData, Prev0, ColA);
+					FNavMeshRenderingHelpers::AddVertex(DebugMeshData, Prev1, ColA);
+					for (int32 ArcIdx = 1; ArcIdx <= NumArcPoints; ArcIdx++)
 					{
-						const float u = j * ArcPtsScale;
-						FVector Pt0 = EvalArc(A0, Edge0, Len0*0.25f, u);
-						FVector Pt1 = EvalArc(A1, Edge1, Len1*0.25f, u);
+						const float u = ArcIdx * ArcPtsScale;
+						FVector Pt0 = FNavMeshRenderingHelpers::EvalArc(A0, Edge0, Len0*0.25f, u);
+						FVector Pt1 = FNavMeshRenderingHelpers::EvalArc(A1, Edge1, Len1*0.25f, u);
 
-						AddVertexHelper(DebugMeshData, Pt0, (j == NumArcPoints) ? ColB : FColor::White);
-						AddVertexHelper(DebugMeshData, Pt1, (j == NumArcPoints) ? ColB : FColor::White);
+						FNavMeshRenderingHelpers::AddVertex(DebugMeshData, Pt0, (ArcIdx == NumArcPoints) ? ColB : FColor::White);
+						FNavMeshRenderingHelpers::AddVertex(DebugMeshData, Pt1, (ArcIdx == NumArcPoints) ? ColB : FColor::White);
 
-						AddTriangleHelper(DebugMeshData, VertBase+0, VertBase+2, VertBase+1);
-						AddTriangleHelper(DebugMeshData, VertBase+2, VertBase+3, VertBase+1);
-						AddTriangleHelper(DebugMeshData, VertBase+0, VertBase+1, VertBase+2);
-						AddTriangleHelper(DebugMeshData, VertBase+2, VertBase+1, VertBase+3);
+						FNavMeshRenderingHelpers::AddTriangle(DebugMeshData, VertBase + 0, VertBase + 2, VertBase + 1);
+						FNavMeshRenderingHelpers::AddTriangle(DebugMeshData, VertBase + 2, VertBase + 3, VertBase + 1);
+						FNavMeshRenderingHelpers::AddTriangle(DebugMeshData, VertBase + 0, VertBase + 1, VertBase + 2);
+						FNavMeshRenderingHelpers::AddTriangle(DebugMeshData, VertBase + 2, VertBase + 1, VertBase + 3);
 
 						VertBase += 2;
 						Prev0 = Pt0;
@@ -621,13 +772,580 @@ void UNavMeshRenderingComponent::GatherData(FNavMeshSceneProxyData& CurrentData)
 
 				if (DebugMeshData.Indices.Num())
 				{
-					CurrentData.MeshBuilders.Add(DebugMeshData);
+					MeshBuilders.Add(DebugMeshData);
 				}
 			}
 		}
-
-		CurrentData.NavMeshGeometry.PolyEdges.Empty();
-		CurrentData.NavMeshGeometry.NavMeshEdges.Empty();
 	}
-#endif //WITH_RECAST
+}
+
+#endif
+
+//////////////////////////////////////////////////////////////////////////
+// FNavMeshSceneProxy
+
+void FNavMeshSceneProxy::FNavMeshIndexBuffer::InitRHI()
+{
+	if (Indices.Num() > 0)
+	{
+		FRHIResourceCreateInfo CreateInfo;
+		void* Buffer = nullptr;
+		IndexBufferRHI = RHICreateAndLockIndexBuffer(sizeof(int32), Indices.Num() * sizeof(int32), BUF_Static, CreateInfo, Buffer);
+
+		// Write the indices to the index buffer.			
+		FMemory::Memcpy(Buffer, Indices.GetData(), Indices.Num() * sizeof(int32));
+		RHIUnlockIndexBuffer(IndexBufferRHI);
+	}
+}
+
+void FNavMeshSceneProxy::FNavMeshVertexBuffer::InitRHI()
+{
+	if (Vertices.Num() > 0)
+	{
+		FRHIResourceCreateInfo CreateInfo;
+		void* VertexBufferData = nullptr;
+		VertexBufferRHI = RHICreateAndLockVertexBuffer(Vertices.Num() * sizeof(FDynamicMeshVertex), BUF_Static, CreateInfo, VertexBufferData);
+
+		// Copy the vertex data into the vertex buffer.			
+		FMemory::Memcpy(VertexBufferData, Vertices.GetData(), Vertices.Num() * sizeof(FDynamicMeshVertex));
+		RHIUnlockVertexBuffer(VertexBufferRHI);
+	}
+}
+
+void FNavMeshSceneProxy::FNavMeshVertexFactory::Init(const FNavMeshVertexBuffer* InVertexBuffer)
+{
+	if (IsInRenderingThread())
+	{
+		// Initialize the vertex factory's stream components.
+		FDataType NewData;
+		NewData.PositionComponent = STRUCTMEMBER_VERTEXSTREAMCOMPONENT(InVertexBuffer, FDynamicMeshVertex, Position, VET_Float3);
+		NewData.TextureCoordinates.Add(
+			FVertexStreamComponent(InVertexBuffer, STRUCT_OFFSET(FDynamicMeshVertex, TextureCoordinate), sizeof(FDynamicMeshVertex), VET_Float2)
+			);
+		NewData.TangentBasisComponents[0] = STRUCTMEMBER_VERTEXSTREAMCOMPONENT(InVertexBuffer, FDynamicMeshVertex, TangentX, VET_PackedNormal);
+		NewData.TangentBasisComponents[1] = STRUCTMEMBER_VERTEXSTREAMCOMPONENT(InVertexBuffer, FDynamicMeshVertex, TangentZ, VET_PackedNormal);
+		SetData(NewData);
+	}
+	else
+	{
+		ENQUEUE_UNIQUE_RENDER_COMMAND_TWOPARAMETER(
+			InitNavMeshVertexFactory,
+			FNavMeshVertexFactory*, InVertexFactory, this,
+			const FNavMeshVertexBuffer*, InVertexBuffer, InVertexBuffer,
+			{
+				// Initialize the vertex factory's stream components.
+				FDataType NewData;
+		NewData.PositionComponent = STRUCTMEMBER_VERTEXSTREAMCOMPONENT(InVertexBuffer, FDynamicMeshVertex, Position, VET_Float3);
+		NewData.TextureCoordinates.Add(
+			FVertexStreamComponent(InVertexBuffer, STRUCT_OFFSET(FDynamicMeshVertex, TextureCoordinate), sizeof(FDynamicMeshVertex), VET_Float2)
+			);
+		NewData.TangentBasisComponents[0] = STRUCTMEMBER_VERTEXSTREAMCOMPONENT(InVertexBuffer, FDynamicMeshVertex, TangentX, VET_PackedNormal);
+		NewData.TangentBasisComponents[1] = STRUCTMEMBER_VERTEXSTREAMCOMPONENT(InVertexBuffer, FDynamicMeshVertex, TangentZ, VET_PackedNormal);
+		InVertexFactory->SetData(NewData);
+			});
+	}
+}
+
+FNavMeshSceneProxy::FNavMeshSceneProxy(const UPrimitiveComponent* InComponent, FNavMeshSceneProxyData* InProxyData, bool ForceToRender)
+	: FDebugRenderSceneProxy(InComponent)
+	, bRequestedData(false)
+	, bForceRendering(ForceToRender)
+{
+	if (InProxyData)
+	{
+		ProxyData = *InProxyData;
+	}
+
+	RenderingComponent = Cast<const UNavMeshRenderingComponent>(InComponent);
+	bSkipDistanceCheck = GIsEditor && (GEngine->GetDebugLocalPlayer() == nullptr);
+	bUseThickLines = GIsEditor;
+
+	const int32 NumberOfMeshes = ProxyData.MeshBuilders.Num();
+	if (!NumberOfMeshes)
+	{
+		return;
+	}
+
+	MeshColors.Reserve(NumberOfMeshes + (ProxyData.PathCollidingGeomVerts.Num() ? 1 : 0));
+	MeshBatchElements.Reserve(NumberOfMeshes);
+	const FMaterialRenderProxy* ParentMaterial = GEngine->DebugMeshMaterial->GetRenderProxy(false);
+	for (int32 Index = 0; Index < NumberOfMeshes; ++Index)
+	{
+		const auto& CurrentMeshBuilder = ProxyData.MeshBuilders[Index];
+
+		FMeshBatchElement Element;
+		Element.FirstIndex = IndexBuffer.Indices.Num();
+		Element.NumPrimitives = FMath::FloorToInt(CurrentMeshBuilder.Indices.Num() / 3);
+		Element.MinVertexIndex = VertexBuffer.Vertices.Num();
+		Element.MaxVertexIndex = Element.MinVertexIndex + CurrentMeshBuilder.Vertices.Num() - 1;
+		Element.IndexBuffer = &IndexBuffer;
+		MeshBatchElements.Add(Element);
+
+		MeshColors.Add(FColoredMaterialRenderProxy(ParentMaterial, CurrentMeshBuilder.ClusterColor));
+
+		VertexBuffer.Vertices.Append(CurrentMeshBuilder.Vertices);
+		IndexBuffer.Indices.Append(CurrentMeshBuilder.Indices);
+	}
+
+	MeshColors.Add(FColoredMaterialRenderProxy(ParentMaterial, NavMeshRenderColor_PathCollidingGeom));
+
+	VertexFactory.Init(&VertexBuffer);
+	BeginInitResource(&IndexBuffer);
+	BeginInitResource(&VertexBuffer);
+	BeginInitResource(&VertexFactory);
+}
+
+FNavMeshSceneProxy::~FNavMeshSceneProxy()
+{
+	VertexBuffer.ReleaseResource();
+	IndexBuffer.ReleaseResource();
+	VertexFactory.ReleaseResource();
+}
+
+void FNavMeshSceneProxy::RegisterDebugDrawDelgate()
+{
+	DebugTextDrawingDelegate = FDebugDrawDelegate::CreateRaw(this, &FNavMeshSceneProxy::DrawDebugLabels);
+	DebugTextDrawingDelegateHandle = UDebugDrawService::Register(TEXT("Navigation"), DebugTextDrawingDelegate);
+}
+
+void FNavMeshSceneProxy::UnregisterDebugDrawDelgate()
+{
+	if (DebugTextDrawingDelegate.IsBound())
+	{
+		UDebugDrawService::Unregister(DebugTextDrawingDelegateHandle);
+	}
+}
+
+void FNavMeshSceneProxy::DrawDebugBox(FPrimitiveDrawInterface* PDI, FVector const& Center, FVector const& Box, FColor const& Color) const
+{
+	// no debug line drawing on dedicated server
+	if (PDI != NULL)
+	{
+		PDI->DrawLine(Center + FVector(Box.X, Box.Y, Box.Z), Center + FVector(Box.X, -Box.Y, Box.Z), Color, SDPG_World);
+		PDI->DrawLine(Center + FVector(Box.X, -Box.Y, Box.Z), Center + FVector(-Box.X, -Box.Y, Box.Z), Color, SDPG_World);
+		PDI->DrawLine(Center + FVector(-Box.X, -Box.Y, Box.Z), Center + FVector(-Box.X, Box.Y, Box.Z), Color, SDPG_World);
+		PDI->DrawLine(Center + FVector(-Box.X, Box.Y, Box.Z), Center + FVector(Box.X, Box.Y, Box.Z), Color, SDPG_World);
+
+		PDI->DrawLine(Center + FVector(Box.X, Box.Y, -Box.Z), Center + FVector(Box.X, -Box.Y, -Box.Z), Color, SDPG_World);
+		PDI->DrawLine(Center + FVector(Box.X, -Box.Y, -Box.Z), Center + FVector(-Box.X, -Box.Y, -Box.Z), Color, SDPG_World);
+		PDI->DrawLine(Center + FVector(-Box.X, -Box.Y, -Box.Z), Center + FVector(-Box.X, Box.Y, -Box.Z), Color, SDPG_World);
+		PDI->DrawLine(Center + FVector(-Box.X, Box.Y, -Box.Z), Center + FVector(Box.X, Box.Y, -Box.Z), Color, SDPG_World);
+
+		PDI->DrawLine(Center + FVector(Box.X, Box.Y, Box.Z), Center + FVector(Box.X, Box.Y, -Box.Z), Color, SDPG_World);
+		PDI->DrawLine(Center + FVector(Box.X, -Box.Y, Box.Z), Center + FVector(Box.X, -Box.Y, -Box.Z), Color, SDPG_World);
+		PDI->DrawLine(Center + FVector(-Box.X, -Box.Y, Box.Z), Center + FVector(-Box.X, -Box.Y, -Box.Z), Color, SDPG_World);
+		PDI->DrawLine(Center + FVector(-Box.X, Box.Y, Box.Z), Center + FVector(-Box.X, Box.Y, -Box.Z), Color, SDPG_World);
+	}
+}
+
+void FNavMeshSceneProxy::GetDynamicMeshElements(const TArray<const FSceneView*>& Views, const FSceneViewFamily& ViewFamily, uint32 VisibilityMap, FMeshElementCollector& Collector) const
+{
+	QUICK_SCOPE_CYCLE_COUNTER(STAT_RecastRenderingSceneProxy_GetDynamicMeshElements);
+
+	for (int32 ViewIndex = 0; ViewIndex < Views.Num(); ViewIndex++)
+	{
+		if (VisibilityMap & (1 << ViewIndex))
+		{
+			const FSceneView* View = Views[ViewIndex];
+			const bool bVisible = !!View->Family->EngineShowFlags.Navigation || bForceRendering;
+			if (!bVisible)
+			{
+				continue;
+			}
+			FPrimitiveDrawInterface* PDI = Collector.GetPDI(ViewIndex);
+
+			for (int32 Index = 0; Index < ProxyData.OctreeBounds.Num(); ++Index)
+			{
+				const FBoxCenterAndExtent& Bounds = ProxyData.OctreeBounds[Index];
+				DrawDebugBox(PDI, Bounds.Center, Bounds.Extent, FColor::White);
+			}
+
+			// Draw Mesh
+			if (MeshBatchElements.Num())
+			{
+				for (int32 Index = 0; Index < MeshBatchElements.Num(); ++Index)
+				{
+					if (MeshBatchElements[Index].NumPrimitives == 0)
+					{
+						continue;
+					}
+
+					FMeshBatch& Mesh = Collector.AllocateMesh();
+					FMeshBatchElement& BatchElement = Mesh.Elements[0];
+					BatchElement = MeshBatchElements[Index];
+					BatchElement.PrimitiveUniformBuffer = CreatePrimitiveUniformBufferImmediate(FMatrix::Identity, GetBounds(), GetLocalBounds(), false, UseEditorDepthTest());
+
+					Mesh.bWireframe = false;
+					Mesh.VertexFactory = &VertexFactory;
+					Mesh.MaterialRenderProxy = &MeshColors[Index];
+					Mesh.ReverseCulling = IsLocalToWorldDeterminantNegative();
+					Mesh.Type = PT_TriangleList;
+					Mesh.DepthPriorityGroup = SDPG_World;
+					Mesh.bCanApplyViewModeOverrides = false;
+					Collector.AddMesh(ViewIndex, Mesh);
+				}
+
+				if (ProxyData.PathCollidingGeomIndices.Num() > 2)
+				{
+					FDynamicMeshBuilder MeshBuilder;
+					MeshBuilder.AddVertices(ProxyData.PathCollidingGeomVerts);
+					MeshBuilder.AddTriangles(ProxyData.PathCollidingGeomIndices);
+
+					MeshBuilder.GetMesh(FMatrix::Identity, &MeshColors[MeshBatchElements.Num()], SDPG_World, false, false, ViewIndex, Collector);
+				}
+			}
+
+			int32 Num = ProxyData.NavMeshEdgeLines.Num();
+			PDI->AddReserveLines(SDPG_World, Num, false, false);
+			PDI->AddReserveLines(SDPG_Foreground, Num, false, true);
+			for (int32 Index = 0; Index < Num; ++Index)
+			{
+				const FDebugLine &Line = ProxyData.NavMeshEdgeLines[Index];
+				if (FNavMeshRenderingHelpers::LineInView(Line.Start, Line.End, View, false))
+				{
+					if (FNavMeshRenderingHelpers::LineInCorrectDistance(Line.Start, Line.End, View))
+					{
+						PDI->DrawLine(Line.Start, Line.End, Line.Color, SDPG_World, NavMeshEdges_LineThickness, 0, true);
+					}
+					else if (bUseThickLines)
+					{
+						PDI->DrawLine(Line.Start, Line.End, Line.Color, SDPG_Foreground, DefaultEdges_LineThickness, 0, true);
+					}
+				}
+			}
+
+			Num = ProxyData.ClusterLinkLines.Num();
+			PDI->AddReserveLines(SDPG_World, Num, false, false);
+			PDI->AddReserveLines(SDPG_Foreground, Num, false, true);
+			for (int32 Index = 0; Index < Num; ++Index)
+			{
+				const FDebugLine &Line = ProxyData.ClusterLinkLines[Index];
+				if (FNavMeshRenderingHelpers::LineInView(Line.Start, Line.End, View, false))
+				{
+					if (FNavMeshRenderingHelpers::LineInCorrectDistance(Line.Start, Line.End, View))
+					{
+						PDI->DrawLine(Line.Start, Line.End, Line.Color, SDPG_World, ClusterLinkLines_LineThickness, 0, true);
+					}
+					else if (bUseThickLines)
+					{
+						PDI->DrawLine(Line.Start, Line.End, Line.Color, SDPG_Foreground, DefaultEdges_LineThickness, 0, true);
+					}
+				}
+			}
+
+			Num = ProxyData.TileEdgeLines.Num();
+			PDI->AddReserveLines(SDPG_World, Num, false, false);
+			PDI->AddReserveLines(SDPG_Foreground, Num, false, true);
+			for (int32 Index = 0; Index < Num; ++Index)
+			{
+				const FDebugLine &Line = ProxyData.TileEdgeLines[Index];
+				if (FNavMeshRenderingHelpers::LineInView(Line.Start, Line.End, View, false))
+				{
+					if (FNavMeshRenderingHelpers::LineInCorrectDistance(Line.Start, Line.End, View))
+					{
+						PDI->DrawLine(Line.Start, Line.End, Line.Color, SDPG_World, PolyEdges_LineThickness, 0, true);
+					}
+					else if (bUseThickLines)
+					{
+						PDI->DrawLine(Line.Start, Line.End, Line.Color, SDPG_Foreground, DefaultEdges_LineThickness, 0, true);
+					}
+				}
+			}
+
+			Num = ProxyData.NavLinkLines.Num();
+			PDI->AddReserveLines(SDPG_World, Num, false, false);
+			PDI->AddReserveLines(SDPG_Foreground, Num, false, true);
+			for (int32 Index = 0; Index < Num; ++Index)
+			{
+				const FDebugLine &Line = ProxyData.NavLinkLines[Index];
+				if (FNavMeshRenderingHelpers::LineInView(Line.Start, Line.End, View, false))
+				{
+					if (FNavMeshRenderingHelpers::LineInCorrectDistance(Line.Start, Line.End, View))
+					{
+						PDI->DrawLine(Line.Start, Line.End, Line.Color, SDPG_World, LinkLines_LineThickness, 0, true);
+					}
+					else if (bUseThickLines)
+					{
+						PDI->DrawLine(Line.Start, Line.End, Line.Color, SDPG_Foreground, DefaultEdges_LineThickness, 0, true);
+					}
+				}
+			}
+
+			Num = ProxyData.ThickLineItems.Num();
+			PDI->AddReserveLines(SDPG_Foreground, Num, false, true);
+			for (int32 Index = 0; Index < Num; ++Index)
+			{
+				const auto &Line = ProxyData.ThickLineItems[Index];
+				if (FNavMeshRenderingHelpers::LineInView(Line.Start, Line.End, View, false))
+				{
+					if (FNavMeshRenderingHelpers::LineInCorrectDistance(Line.Start, Line.End, View))
+					{
+						PDI->DrawLine(Line.Start, Line.End, Line.Color, SDPG_World, Line.Thickness, 0, true);
+					}
+					else if (bUseThickLines)
+					{
+						PDI->DrawLine(Line.Start, Line.End, Line.Color, SDPG_Foreground, DefaultEdges_LineThickness, 0, true);
+					}
+				}
+			}
+		}
+	}
+}
+
+void FNavMeshSceneProxy::DrawDebugLabels(UCanvas* Canvas, APlayerController*)
+{
+	const bool bVisible = (Canvas && Canvas->SceneView && !!Canvas->SceneView->Family->EngineShowFlags.Navigation) || bForceRendering;
+	if (!bVisible || ProxyData.bNeedsNewData || ProxyData.DebugLabels.Num() == 0)
+	{
+		return;
+	}
+
+	const FColor OldDrawColor = Canvas->DrawColor;
+	Canvas->SetDrawColor(FColor::White);
+	const FSceneView* View = Canvas->SceneView;
+	UFont* Font = GEngine->GetSmallFont();
+	const FNavMeshSceneProxyData::FDebugText* DebugText = ProxyData.DebugLabels.GetData();
+	for (int32 Idx = 0; Idx < ProxyData.DebugLabels.Num(); ++Idx, ++DebugText)
+	{
+		if (View->ViewFrustum.IntersectSphere(DebugText->Location, 1.0f))
+		{
+			const FVector ScreenLoc = Canvas->Project(DebugText->Location);
+			Canvas->DrawText(Font, DebugText->Text, ScreenLoc.X, ScreenLoc.Y);
+		}
+	}
+
+	Canvas->SetDrawColor(OldDrawColor);
+}
+
+FPrimitiveViewRelevance FNavMeshSceneProxy::GetViewRelevance(const FSceneView* View) const
+{
+	const bool bVisible = !!View->Family->EngineShowFlags.Navigation || bForceRendering;
+	FPrimitiveViewRelevance Result;
+	Result.bDrawRelevance = bVisible && IsShown(View);
+	Result.bDynamicRelevance = true;
+	// ideally the TranslucencyRelevance should be filled out by the material, here we do it conservative
+	Result.bSeparateTranslucencyRelevance = Result.bNormalTranslucencyRelevance = bVisible && IsShown(View);
+	return Result;
+}
+
+uint32 FNavMeshSceneProxy::GetAllocatedSize() const
+{
+	return FDebugRenderSceneProxy::GetAllocatedSize() +
+		ProxyData.GetAllocatedSize() +
+		IndexBuffer.Indices.GetAllocatedSize() +
+		VertexBuffer.Vertices.GetAllocatedSize() +
+		MeshColors.GetAllocatedSize() +
+		MeshBatchElements.GetAllocatedSize();
+}
+
+//////////////////////////////////////////////////////////////////////////
+// NavMeshRenderingComponent
+
+#if WITH_EDITOR
+namespace
+{
+	bool AreAnyViewportsRelevant(const UWorld* World)
+	{
+		FWorldContext* WorldContext = GEngine->GetWorldContextFromWorld(World);
+		if (WorldContext && WorldContext->GameViewport)
+		{
+			return true;
+		}
+		
+		for (FEditorViewportClient* CurrentViewport : GEditor->AllViewportClients)
+		{
+			if (CurrentViewport && CurrentViewport->IsVisible())
+			{
+				return true;
+			}
+		}
+
+		return false;
+	}
+}
+#endif
+
+UNavMeshRenderingComponent::UNavMeshRenderingComponent(const FObjectInitializer& ObjectInitializer)
+	: Super(ObjectInitializer)
+{
+	SetCollisionProfileName(UCollisionProfile::NoCollision_ProfileName);
+	AlwaysLoadOnClient = false;
+	AlwaysLoadOnServer = false;
+	bSelectable = false;
+	bCollectNavigationData = false;
+}
+
+bool UNavMeshRenderingComponent::NeedsLoadForServer() const
+{
+	// This rendering component is not needed in dedicated server builds
+	return false;
+}
+
+bool UNavMeshRenderingComponent::IsNavigationShowFlagSet(const UWorld* World)
+{
+	bool bShowNavigation = false;
+
+	FWorldContext* WorldContext = GEngine->GetWorldContextFromWorld(World);
+
+#if WITH_EDITOR
+	if (GEditor && WorldContext && WorldContext->WorldType != EWorldType::Game)
+	{
+		bShowNavigation = WorldContext->GameViewport != nullptr && WorldContext->GameViewport->EngineShowFlags.Navigation;
+		if (bShowNavigation == false)
+		{
+			// we have to check all vieports because we can't to distinguish between SIE and PIE at this point.
+			for (FEditorViewportClient* CurrentViewport : GEditor->AllViewportClients)
+			{
+				if (CurrentViewport && CurrentViewport->IsVisible() && CurrentViewport->EngineShowFlags.Navigation)
+				{
+					bShowNavigation = true;
+					break;
+				}
+			}
+		}
+	}
+	else
+#endif //WITH_EDITOR
+	{
+		bShowNavigation = WorldContext && WorldContext->GameViewport && WorldContext->GameViewport->EngineShowFlags.Navigation;
+	}
+
+	return bShowNavigation;
+}
+
+void UNavMeshRenderingComponent::TimerFunction()
+{
+	const UWorld* World = GetWorld();
+#if WITH_EDITOR
+	if (GEditor && (AreAnyViewportsRelevant(World) == false))
+	{
+		// unable to tell if the flag is on or not
+		return;
+	}
+#endif // WITH_EDITOR
+
+	const bool bShowNavigation = bForceUpdate || IsNavigationShowFlagSet(World);
+
+	if (bShowNavigation != !!bCollectNavigationData && bShowNavigation == true)
+	{
+		bForceUpdate = false;
+		bCollectNavigationData = bShowNavigation;
+		MarkRenderStateDirty();
+	}
+}
+
+void UNavMeshRenderingComponent::OnRegister()
+{
+	Super::OnRegister();
+
+#if WITH_RECAST && !UE_BUILD_SHIPPING && !UE_BUILD_TEST
+	// it's a kind of HACK but there is no event or other information that show flag was changed by user => we have to check it periodically
+#if WITH_EDITOR
+	if (GEditor)
+	{
+		GEditor->GetTimerManager()->SetTimer(TimerHandle, FTimerDelegate::CreateUObject(this, &UNavMeshRenderingComponent::TimerFunction), 1, true);
+	}
+	else
+#endif //WITH_EDITOR
+	{
+		GetWorld()->GetTimerManager().SetTimer(TimerHandle, FTimerDelegate::CreateUObject(this, &UNavMeshRenderingComponent::TimerFunction), 1, true);
+	}
+#endif //WITH_RECAST && !UE_BUILD_SHIPPING && !UE_BUILD_TEST
+}
+
+void UNavMeshRenderingComponent::OnUnregister()
+{
+#if WITH_RECAST && !UE_BUILD_SHIPPING && !UE_BUILD_TEST
+	// it's a kind of HACK but there is no event or other information that show flag was changed by user => we have to check it periodically
+#if WITH_EDITOR
+	if (GEditor)
+	{
+		GEditor->GetTimerManager()->ClearTimer(TimerHandle);
+	}
+	else
+#endif //WITH_EDITOR
+	{
+		GetWorld()->GetTimerManager().ClearTimer(TimerHandle);
+	}
+#endif //WITH_RECAST && !UE_BUILD_SHIPPING && !UE_BUILD_TEST
+	Super::OnUnregister();
+}
+
+FPrimitiveSceneProxy* UNavMeshRenderingComponent::CreateSceneProxy()
+{
+#if WITH_RECAST && !UE_BUILD_SHIPPING && !UE_BUILD_TEST
+	FPrimitiveSceneProxy* SceneProxy = nullptr;
+
+	const bool bShowNavigation = IsNavigationShowFlagSet(GetWorld());
+
+	bCollectNavigationData = bShowNavigation;
+
+	if (bCollectNavigationData && IsVisible())
+	{
+		const ARecastNavMesh* NavMesh = Cast<ARecastNavMesh>(GetOwner());
+		if (NavMesh && NavMesh->IsDrawingEnabled())
+		{
+			FNavMeshSceneProxyData ProxyData;
+
+			const int32 DetailFlags = ProxyData.GetDetailFlags(NavMesh);
+			TArray<int32> EmptyTileSet;
+			ProxyData.GatherData(NavMesh, DetailFlags, EmptyTileSet);
+
+			SceneProxy = new FNavMeshSceneProxy(this, &ProxyData);
+		}
+	}
+	return SceneProxy;
+#else
+	return nullptr;
+#endif //WITH_RECAST && !UE_BUILD_SHIPPING && !UE_BUILD_TEST
+}
+
+void UNavMeshRenderingComponent::CreateRenderState_Concurrent()
+{
+	Super::CreateRenderState_Concurrent();
+
+#if WITH_RECAST && !UE_BUILD_SHIPPING && !UE_BUILD_TEST
+	if (SceneProxy)
+	{
+		static_cast<FNavMeshSceneProxy*>(SceneProxy)->RegisterDebugDrawDelgate();
+	}
+#endif
+}
+
+void UNavMeshRenderingComponent::DestroyRenderState_Concurrent()
+{
+#if WITH_RECAST && !UE_BUILD_SHIPPING && !UE_BUILD_TEST
+	if (SceneProxy)
+	{
+		static_cast<FNavMeshSceneProxy*>(SceneProxy)->UnregisterDebugDrawDelgate();
+	}
+#endif
+
+	Super::DestroyRenderState_Concurrent();
+}
+
+FBoxSphereBounds UNavMeshRenderingComponent::CalcBounds(const FTransform& LocalToWorld) const
+{
+	FBox BoundingBox(0);
+#if WITH_RECAST
+	ARecastNavMesh* NavMesh = Cast<ARecastNavMesh>(GetOwner());
+	if (NavMesh)
+	{
+		BoundingBox = NavMesh->GetNavMeshBounds();
+		if (NavMesh->bDrawOctree)
+		{
+			const UNavigationSystem* NavSys = UNavigationSystem::GetCurrent(GetWorld());
+			const FNavigationOctree* NavOctree = NavSys ? NavSys->GetNavOctree() : nullptr;
+			if (NavOctree)
+			{
+				for (FNavigationOctree::TConstIterator<> NodeIt(*NavOctree); NodeIt.HasPendingNodes(); NodeIt.Advance())
+				{
+					const FOctreeNodeContext& CorrentContext = NodeIt.GetCurrentContext();
+					BoundingBox += CorrentContext.Bounds.GetBox();
+				}
+			}
+		}
+	}
+#endif
+	return FBoxSphereBounds(BoundingBox);
 }
diff --git a/Engine/Source/Runtime/Engine/Private/Distributions.cpp b/Engine/Source/Runtime/Engine/Private/Distributions.cpp
--- a/Engine/Source/Runtime/Engine/Private/Distributions.cpp
+++ b/Engine/Source/Runtime/Engine/Private/Distributions.cpp
@@ -30,6 +30,9 @@ ENGINE_API uint32 GDistributionType = 1;
 // The maximum number of values to store in a lookup table.
 #define LOOKUP_TABLE_MAX_SAMPLES (128)
 
+// UDistribution will bake out (if CanBeBaked returns true)
+#define DISTRIBUTIONS_BAKES_OUT 1
+
 // The maximum number of samples must be a power of two.
 static_assert((LOOKUP_TABLE_MAX_SAMPLES & (LOOKUP_TABLE_MAX_SAMPLES - 1)) == 0, "Lookup table max samples is not a power of two.");
 
@@ -909,6 +912,16 @@ void FRawDistributionFloat::Initialize()
 }
 #endif // WITH_EDITOR
 
+bool FRawDistributionFloat::IsCreated()
+{
+	return HasLookupTable(/*bInitializeIfNeeded=*/ false) || (Distribution != nullptr);
+}
+
+bool FRawDistributionVector::IsCreated()
+{
+	return HasLookupTable(/*bInitializeIfNeeded=*/ false) || (Distribution != nullptr);
+}
+
 float FRawDistributionFloat::GetValue(float F, UObject* Data, struct FRandomStream* InRandomStream)
 {
 	if (!HasLookupTable())
@@ -943,7 +956,7 @@ const FRawDistribution *FRawDistributionFloat::GetFastRawDistribution()
 
 void FRawDistributionFloat::GetOutRange(float& MinOut, float& MaxOut)
 {
-	if (!HasLookupTable() || !Distribution)
+	if (!HasLookupTable() && Distribution)
 	{
 		check(Distribution);
 		Distribution->GetOutRange(MinOut, MaxOut);
@@ -967,15 +980,103 @@ void FRawDistributionFloat::InitLookupTable()
 #endif
 }
 
+#if WITH_EDITOR
+template <typename RawDistributionType>
+bool HasBakedDistributionDataHelper(const UDistribution* GivenDistribution)
+{
+	if (UObject* Outer = GivenDistribution->GetOuter())
+	{
+		for (TFieldIterator<UProperty> It(Outer->GetClass()); It; ++It)
+		{
+			UProperty* Property = *It;
+
+			if (UStructProperty* StructProp = Cast<UStructProperty>(Property))
+			{
+				UObject* Distribution = FRawDistribution::TryGetDistributionObjectFromRawDistributionProperty(StructProp, reinterpret_cast<uint8*>(Outer));
+				if (Distribution == GivenDistribution)
+				{
+					if (RawDistributionType* RawDistributionVector = StructProp->ContainerPtrToValuePtr<RawDistributionType>(reinterpret_cast<uint8*>(Outer)))
+					{
+						if (RawDistributionVector->HasLookupTable(false))
+						{
+							return true;	//We have baked data
+						}
+					}
+
+					return false;
+				}
+			}
+			else if (UArrayProperty* ArrayProp = Cast<UArrayProperty>(Property))
+			{
+				if (UStructProperty* InnerStructProp = Cast<UStructProperty>(ArrayProp->Inner))
+				{
+					FScriptArrayHelper ArrayHelper(ArrayProp, Property->ContainerPtrToValuePtr<void>(Outer));
+					for (int32 Idx = 0; Idx < ArrayHelper.Num(); ++Idx)
+					{
+						for (UProperty* ArrayProperty = InnerStructProp->Struct->PropertyLink; ArrayProperty; ArrayProperty = ArrayProperty->PropertyLinkNext)
+						{
+							if (UStructProperty* ArrayStructProp = Cast<UStructProperty>(ArrayProperty))
+							{
+								UObject* Distribution = FRawDistribution::TryGetDistributionObjectFromRawDistributionProperty(ArrayStructProp, ArrayHelper.GetRawPtr(Idx));
+								if (Distribution == GivenDistribution)
+								{
+									if (RawDistributionType* RawDistributionVector = ArrayStructProp->ContainerPtrToValuePtr<RawDistributionType>(ArrayHelper.GetRawPtr(Idx)))
+									{
+										if (RawDistributionVector->HasLookupTable(false))
+										{
+											return true;	//We have baked data
+										}
+									}
+
+									return false;
+								}
+							}
+						}
+					}
+				}
+			}
+		}
+	}
+
+	return false;
+}
+#endif
+
+void UDistributionFloat::Serialize(FArchive& Ar)
+{
+#if WITH_EDITOR
+	if (Ar.IsCooking() || Ar.IsSaving())
+	{
+		bBakedDataSuccesfully = HasBakedDistributionDataHelper<FRawDistributionFloat>(this);
+	}
+#endif
+
+	Super::Serialize(Ar);
+}
+
 
 bool UDistributionFloat::NeedsLoadForClient() const
 {
+#if DISTRIBUTIONS_BAKES_OUT
+	if (CanBeBaked() && HasBakedSuccesfully())
+	{
+		return false;
+	}
+#endif
+
 	return true;
 }
 
 
 bool UDistributionFloat::NeedsLoadForServer() const 
 {
+#if DISTRIBUTIONS_BAKES_OUT
+	if (CanBeBaked() && HasBakedSuccesfully())
+	{
+		return false;
+	}
+#endif
+
 	return true;
 }
 
@@ -1068,6 +1169,8 @@ void FRawDistributionVector::Initialize()
 
 	// fill out our min/max
 	Distribution->GetOutRange(MinValue, MaxValue);
+
+	Distribution->GetRange(MinValueVec, MaxValueVec);
 }
 #endif
 
@@ -1117,6 +1220,20 @@ void FRawDistributionVector::GetOutRange(float& MinOut, float& MaxOut)
 	}
 }
 
+void FRawDistributionVector::GetRange(FVector& MinOut, FVector& MaxOut)
+{
+	if (!HasLookupTable() && Distribution)
+	{
+		check(Distribution);
+		Distribution->GetRange(MinOut, MaxOut);
+	}
+	else
+	{
+		MinOut = MinValueVec;
+		MaxOut = MaxValueVec;
+	}
+}
+
 void FRawDistributionVector::InitLookupTable()
 {
 #if WITH_EDITOR
@@ -1128,14 +1245,39 @@ void FRawDistributionVector::InitLookupTable()
 #endif
 }
 
+void UDistributionVector::Serialize(FArchive& Ar)
+{
+#if WITH_EDITOR
+	if (Ar.IsCooking() || Ar.IsSaving())
+	{
+		bBakedDataSuccesfully = HasBakedDistributionDataHelper<FRawDistributionVector>(this);
+	}
+#endif
+
+	Super::Serialize(Ar);
+}
 
 bool UDistributionVector::NeedsLoadForClient() const
 {
+#if DISTRIBUTIONS_BAKES_OUT
+	if (CanBeBaked() && HasBakedSuccesfully())
+	{
+		return false;
+	}
+#endif
+
 	return true;
 }
 
 bool UDistributionVector::NeedsLoadForServer() const 
 {
+#if DISTRIBUTIONS_BAKES_OUT
+	if (CanBeBaked() && HasBakedSuccesfully())
+	{
+		return false;
+	}
+#endif
+
 	return true;
 }
 
diff --git a/Engine/Source/Runtime/Engine/Private/GameplayStatics.cpp b/Engine/Source/Runtime/Engine/Private/GameplayStatics.cpp
--- a/Engine/Source/Runtime/Engine/Private/GameplayStatics.cpp
+++ b/Engine/Source/Runtime/Engine/Private/GameplayStatics.cpp
@@ -778,6 +778,26 @@ void UGameplayStatics::BreakHitResult(const FHitResult& Hit, bool& bBlockingHit,
 	TraceEnd = Hit.TraceEnd;
 }
 
+FHitResult UGameplayStatics::MakeHitResult(bool bBlockingHit, bool bInitialOverlap, float Time, FVector Location, FVector ImpactPoint, FVector Normal, FVector ImpactNormal, class UPhysicalMaterial* PhysMat, class AActor* HitActor, class UPrimitiveComponent* HitComponent, FName HitBoneName, int32 HitItem, FVector TraceStart, FVector TraceEnd)
+{
+	FHitResult Hit;
+	Hit.bBlockingHit = bBlockingHit;
+	Hit.bStartPenetrating = bInitialOverlap;
+	Hit.Time = Time;
+	Hit.Location = Location;
+	Hit.ImpactPoint = ImpactPoint;
+	Hit.Normal = Normal;
+	Hit.ImpactNormal = ImpactNormal;
+	Hit.PhysMaterial = PhysMat;
+	Hit.Actor = HitActor;
+	Hit.Component = HitComponent;
+	Hit.BoneName = HitBoneName;
+	Hit.Item = HitItem;
+	Hit.TraceStart = TraceStart;
+	Hit.TraceEnd = TraceEnd;
+	return Hit;
+}
+
 EPhysicalSurface UGameplayStatics::GetSurfaceType(const struct FHitResult& Hit)
 {
 	UPhysicalMaterial* const HitPhysMat = Hit.PhysMaterial.Get();
@@ -1098,6 +1118,25 @@ void UGameplayStatics::PushSoundMixModifier(UObject* WorldContextObject, USoundM
 	}
 }
 
+void UGameplayStatics::SetSoundMixClassOverride(UObject* WorldContextObject, class USoundMix* InSoundMixModifier, class USoundClass* InSoundClass, float Volume, float Pitch, float FadeInTime, bool bApplyToChildren)
+{
+	if (!InSoundMixModifier || !GEngine || !GEngine->UseSound())
+	{
+		return;
+	}
+
+	UWorld* ThisWorld = GEngine->GetWorldFromContextObject(WorldContextObject);
+	if (!ThisWorld || !ThisWorld->bAllowAudioPlayback)
+	{
+		return;
+	}
+
+	if (FAudioDevice* AudioDevice = ThisWorld->GetAudioDevice())
+	{
+		AudioDevice->SetSoundMixClassOverride(InSoundMixModifier, InSoundClass, Volume, Pitch, FadeInTime, bApplyToChildren);
+	}
+}
+
 void UGameplayStatics::PopSoundMixModifier(UObject* WorldContextObject, USoundMix* InSoundMixModifier)
 {
 	if (InSoundMixModifier == nullptr || GEngine == nullptr || !GEngine->UseSound())
diff --git a/Engine/Source/Runtime/Engine/Private/LODActor.cpp b/Engine/Source/Runtime/Engine/Private/LODActor.cpp
--- a/Engine/Source/Runtime/Engine/Private/LODActor.cpp
+++ b/Engine/Source/Runtime/Engine/Private/LODActor.cpp
@@ -189,12 +189,34 @@ void ALODActor::PostRegisterAllComponents()
 	// Clean up sub actor if assets were delete manually
 	CleanSubActorArray();
 
+	// Clean up sub objects if assets were delete manually
+	CleanSubObjectsArray();
+
 	UpdateSubActorLODParents();	
 #endif
 }
 
 #if WITH_EDITOR
 
+void ALODActor::PreEditChange(UProperty* PropertyThatWillChange)
+{
+	Super::PreEditChange(PropertyThatWillChange);
+
+	if (PropertyThatWillChange)
+	{
+		const FName PropertyName = PropertyThatWillChange->GetFName();
+
+		// If the Sub Objects array is changed, in case of asset deletion make sure me flag as dirty since the cluster will be invalid
+		if (PropertyName == TEXT("SubObjects"))
+		{
+			SetIsDirty(true);
+		}
+	}
+
+	// Flush all pending rendering commands.
+	FlushRenderingCommands();
+}
+
 void ALODActor::PostEditChangeProperty(FPropertyChangedEvent& PropertyChangedEvent)
 {
 	UProperty* PropertyThatChanged = PropertyChangedEvent.Property;
@@ -516,7 +538,7 @@ void ALODActor::CleanSubActorArray()
 		auto& Actor = SubActors[SubActorIndex];
 		if (Actor == nullptr)
 		{
-			SubActors.RemoveAt(SubActorIndex);
+			SubActors.RemoveAtSwap(SubActorIndex);
 			SubActorIndex--;
 			bIsDirty = true;
 		}
@@ -528,6 +550,26 @@ void ALODActor::CleanSubActorArray()
 	}
 }
 
+void ALODActor::CleanSubObjectsArray()
+{
+	bool bIsDirty = false;
+	for (int32 SubObjectIndex = 0; SubObjectIndex < SubObjects.Num(); ++SubObjectIndex)
+	{
+		UObject* Object = SubObjects[SubObjectIndex];
+		if (Object == nullptr)
+		{
+			SubObjects.RemoveAtSwap(SubObjectIndex);
+			SubObjectIndex--;
+			bIsDirty = true;
+		}
+	}
+
+	if (bIsDirty)
+	{
+		SetIsDirty(true);		
+	}
+}
+
 void ALODActor::RecalculateDrawingDistance(const float InTransitionScreenSize)
 {
 	// At the moment this assumes a fixed field of view of 90 degrees (horizontal and vertical axi)
diff --git a/Engine/Source/Runtime/Engine/Private/Level.cpp b/Engine/Source/Runtime/Engine/Private/Level.cpp
--- a/Engine/Source/Runtime/Engine/Private/Level.cpp
+++ b/Engine/Source/Runtime/Engine/Private/Level.cpp
@@ -495,20 +495,18 @@ void ULevel::SortActorList()
 		return;
 	}
 
-	int32 StartIndex = 0;
 	TArray<AActor*> NewActors;
 	NewActors.Reserve(Actors.Num());
 
-	// The WorldSettings has fixed actor index.
-	check(Actors[StartIndex] == GetWorldSettings());
-	NewActors.Add(Actors[StartIndex++]);
+	check(WorldSettings);
+
+	// The WorldSettings tries to stay at index 0
+	NewActors.Add(WorldSettings);
 
 	// Static not net relevant actors.
-	for (int32 ActorIndex = StartIndex; ActorIndex < Actors.Num(); ActorIndex++)
+	for (AActor* Actor : Actors)
 	{
-		AActor* Actor = Actors[ActorIndex];
-
-		if (Actor != NULL && !Actor->IsPendingKill() && !IsNetActor(Actor))
+		if (Actor != nullptr && Actor != WorldSettings && !Actor->IsPendingKill() && !IsNetActor(Actor))
 		{
 			NewActors.Add(Actor);
 		}
@@ -516,20 +514,19 @@ void ULevel::SortActorList()
 	iFirstNetRelevantActor = NewActors.Num();
 
 	// Static net relevant actors.
-	for (int32 ActorIndex = StartIndex; ActorIndex < Actors.Num(); ActorIndex++)
+	for (AActor* Actor : Actors)
 	{
-		AActor* Actor = Actors[ActorIndex];		
-		if (Actor != NULL && !Actor->IsPendingKill() && IsNetActor(Actor))
+		if (Actor != nullptr && !Actor->IsPendingKill() && IsNetActor(Actor))
 		{
 			NewActors.Add(Actor);
 		}
 	}
 
 	// Replace with sorted list.
-	Actors.AssignButKeepOwner(NewActors);
+	Actors.AssignButKeepOwner(MoveTemp(NewActors));
 
 	// Add all network actors to the owning world
-	if ( OwningWorld != NULL )
+	if ( OwningWorld != nullptr )
 	{
 		// Don't use sorted optimization outside of gameplay so we can safely shuffle around actors e.g. in the Editor
 		// without there being a chance to break code using dynamic/ net relevant actor iterators.
@@ -612,6 +609,11 @@ void ULevel::PostLoad()
 	Actors.Remove(nullptr);
 #endif
 
+	if (WorldSettings == nullptr)
+	{
+		WorldSettings = Cast<AWorldSettings>(Actors[0]);
+	}
+
 	// in the Editor, sort Actor list immediately (at runtime we wait for the level to be added to the world so that it can be delayed in the level streaming case)
 	if (GIsEditor)
 	{
@@ -1765,14 +1767,45 @@ ABrush* ULevel::GetDefaultBrush() const
 }
 
 
-AWorldSettings* ULevel::GetWorldSettings() const
+AWorldSettings* ULevel::GetWorldSettings(bool bChecked) const
 {
-	checkf( Actors.Num() >= 1, *GetPathName() );
-	AWorldSettings* WorldSettings = Cast<AWorldSettings>( Actors[0] );
-	checkf( WorldSettings != NULL, *GetPathName() );
+	if (bChecked)
+	{
+		checkf( WorldSettings != nullptr, *GetPathName() );
+	}
 	return WorldSettings;
 }
 
+void ULevel::SetWorldSettings(AWorldSettings* NewWorldSettings)
+{
+	check(NewWorldSettings); // Doesn't make sense to be clearing a world settings object
+	if (WorldSettings != NewWorldSettings)
+	{
+		// We'll generally endeavor to keep the world settings at its traditional index 0
+		const int32 NewWorldSettingsIndex = Actors.FindLast( NewWorldSettings );
+		if (NewWorldSettingsIndex != 0)
+		{
+			if (Actors[0] == nullptr || Actors[0]->IsA<AWorldSettings>())
+			{
+				Exchange(Actors[0],Actors[NewWorldSettingsIndex]);
+			}
+			else
+			{
+				Actors[NewWorldSettingsIndex] = nullptr;
+				Actors.Insert(NewWorldSettings,0);
+			}
+		}
+
+		if (WorldSettings)
+		{
+			// Makes no sense to have two WorldSettings so destroy existing one
+			WorldSettings->Destroy();
+		}
+
+		WorldSettings = NewWorldSettings;
+	}
+}
+
 ALevelScriptActor* ULevel::GetLevelScriptActor() const
 {
 	return LevelScriptActor;
diff --git a/Engine/Source/Runtime/Engine/Private/Particles/ParticleBeamModules.cpp b/Engine/Source/Runtime/Engine/Private/Particles/ParticleBeamModules.cpp
--- a/Engine/Source/Runtime/Engine/Private/Particles/ParticleBeamModules.cpp
+++ b/Engine/Source/Runtime/Engine/Private/Particles/ParticleBeamModules.cpp
@@ -58,21 +58,21 @@ UParticleModuleTypeDataBeam2::UParticleModuleTypeDataBeam2(const FObjectInitiali
 
 void UParticleModuleTypeDataBeam2::InitializeDefaults()
 {
-	if (!Distance.Distribution)
+	if (!Distance.IsCreated())
 	{
 		UDistributionFloatConstant* DistributionDistance = NewObject<UDistributionFloatConstant>(this, TEXT("DistributionDistance"));
 		DistributionDistance->Constant = 25.0f;
 		Distance.Distribution = DistributionDistance;
 	}
 
-	if (!TaperFactor.Distribution)
+	if (!TaperFactor.IsCreated())
 	{
 		UDistributionFloatConstant* DistributionTaperFactor = NewObject<UDistributionFloatConstant>(this, TEXT("DistributionTaperFactor"));
 		DistributionTaperFactor->Constant = 1.0f;
 		TaperFactor.Distribution = DistributionTaperFactor;
 	}
 
-	if (!TaperScale.Distribution)
+	if (!TaperScale.IsCreated())
 	{
 		UDistributionFloatConstant* DistributionTaperScale = NewObject<UDistributionFloatConstant>(this, TEXT("DistributionTaperScale"));
 		DistributionTaperScale->Constant = 1.0f;
@@ -1015,21 +1015,21 @@ UParticleModuleBeamModifier::UParticleModuleBeamModifier(const FObjectInitialize
 
 void UParticleModuleBeamModifier::InitializeDefaults()
 {
-	if(!Position.Distribution)
+	if(!Position.IsCreated())
 	{
 		UDistributionVectorConstant* DistributionPosition = NewObject<UDistributionVectorConstant>(this, TEXT("DistributionPosition"));
 		DistributionPosition->Constant = FVector(0.0f, 0.0f, 0.0f);
 		Position.Distribution = DistributionPosition;
 	}
 	
-	if(!Tangent.Distribution)
+	if(!Tangent.IsCreated())
 	{
 		UDistributionVectorConstant* DistributionTangent = NewObject<UDistributionVectorConstant>(this, TEXT("DistributionTangent"));
 		DistributionTangent->Constant = FVector(0.0f, 0.0f, 0.0f);
 		Tangent.Distribution = DistributionTangent;
 	}
 
-	if(!Strength.Distribution)
+	if(!Strength.IsCreated())
 	{
 		UDistributionFloatConstant* DistributionStrength = NewObject<UDistributionFloatConstant>(this, TEXT("DistributionStrength"));
 		DistributionStrength->Constant = 0.0f;
@@ -1281,35 +1281,35 @@ UParticleModuleBeamNoise::UParticleModuleBeamNoise(const FObjectInitializer& Obj
 
 void UParticleModuleBeamNoise::InitializeDefaults()
 {
-	if (!NoiseSpeed.Distribution)
+	if (!NoiseSpeed.IsCreated())
 	{
 		UDistributionVectorConstant* DistributionNoiseSpeed = NewObject<UDistributionVectorConstant>(this, TEXT("DistributionNoiseSpeed"));
 		DistributionNoiseSpeed->Constant = FVector(50.0f, 50.0f, 50.0f);
 		NoiseSpeed.Distribution = DistributionNoiseSpeed;
 	}
 
-	if (!NoiseRange.Distribution)
+	if (!NoiseRange.IsCreated())
 	{
 		UDistributionVectorConstant* DistributionNoiseRange = NewObject<UDistributionVectorConstant>(this, TEXT("DistributionNoiseRange"));
 		DistributionNoiseRange->Constant = FVector(50.0f, 50.0f, 50.0f);
 		NoiseRange.Distribution = DistributionNoiseRange;
 	}
 
-	if (!NoiseRangeScale.Distribution)
+	if (!NoiseRangeScale.IsCreated())
 	{
 		UDistributionFloatConstant* DistributionNoiseRangeScale = NewObject<UDistributionFloatConstant>(this, TEXT("DistributionNoiseRangeScale"));
 		DistributionNoiseRangeScale->Constant = 1.0f;
 		NoiseRangeScale.Distribution = DistributionNoiseRangeScale; 
 	}
 
-	if (!NoiseTangentStrength.Distribution)
+	if (!NoiseTangentStrength.IsCreated())
 	{
 		UDistributionFloatConstant* DistributionNoiseTangentStrength = NewObject<UDistributionFloatConstant>(this, TEXT("DistributionNoiseTangentStrength"));
 		DistributionNoiseTangentStrength->Constant = 250.0f;
 		NoiseTangentStrength.Distribution = DistributionNoiseTangentStrength; 
 	}
 
-	if (!NoiseScale.Distribution)
+	if (!NoiseScale.IsCreated())
 	{
 		NoiseScale.Distribution = NewObject<UDistributionFloatConstantCurve>(this, TEXT("DistributionNoiseScale"));
 	}
@@ -1620,21 +1620,21 @@ UParticleModuleBeamSource::UParticleModuleBeamSource(const FObjectInitializer& O
 
 void UParticleModuleBeamSource::InitializeDefaults()
 {
-	if (!Source.Distribution)
+	if (!Source.IsCreated())
 	{
 		UDistributionVectorConstant* DistributionSource = NewObject<UDistributionVectorConstant>(this, TEXT("DistributionSource"));
 		DistributionSource->Constant = FVector(50.0f, 50.0f, 50.0f);
 		Source.Distribution = DistributionSource; 
 	}
 
-	if (!SourceTangent.Distribution)
+	if (!SourceTangent.IsCreated())
 	{
 		UDistributionVectorConstant* DistributionSourceTangent = NewObject<UDistributionVectorConstant>(this, TEXT("DistributionSourceTangent"));
 		DistributionSourceTangent->Constant = FVector(1.0f, 0.0f, 0.0f);
 		SourceTangent.Distribution = DistributionSourceTangent; 
 	}
 
-	if (!SourceStrength.Distribution)
+	if (!SourceStrength.IsCreated())
 	{
 		UDistributionFloatConstant* DistributionSourceStrength = NewObject<UDistributionFloatConstant>(this, TEXT("DistributionSourceStrength"));
 		DistributionSourceStrength->Constant = 25.0f;
@@ -2036,21 +2036,21 @@ UParticleModuleBeamTarget::UParticleModuleBeamTarget(const FObjectInitializer& O
 
 void UParticleModuleBeamTarget::InitializeDefaults()
 {
-	if (!Target.Distribution)
+	if (!Target.IsCreated())
 	{
 		UDistributionVectorConstant* DistributionTarget = NewObject<UDistributionVectorConstant>(this, TEXT("DistributionTarget"));
 		DistributionTarget->Constant = FVector(50.0f, 50.0f, 50.0f);
 		Target.Distribution = DistributionTarget;
 	}
 
-	if (!TargetTangent.Distribution)
+	if (!TargetTangent.IsCreated())
 	{
 		UDistributionVectorConstant* DistributionTargetTangent = NewObject<UDistributionVectorConstant>(this, TEXT("DistributionTargetTangent"));
 		DistributionTargetTangent->Constant = FVector(1.0f, 0.0f, 0.0f);
 		TargetTangent.Distribution = DistributionTargetTangent; 
 	}
 
-	if (!TargetStrength.Distribution)
+	if (!TargetStrength.IsCreated())
 	{
 		UDistributionFloatConstant* DistributionTargetStrength = NewObject<UDistributionFloatConstant>(this, TEXT("DistributionTargetStrength"));
 		DistributionTargetStrength->Constant = 25.0;
diff --git a/Engine/Source/Runtime/Engine/Private/Particles/ParticleGpuSimulation.cpp b/Engine/Source/Runtime/Engine/Private/Particles/ParticleGpuSimulation.cpp
--- a/Engine/Source/Runtime/Engine/Private/Particles/ParticleGpuSimulation.cpp
+++ b/Engine/Source/Runtime/Engine/Private/Particles/ParticleGpuSimulation.cpp
@@ -3247,20 +3247,20 @@ public:
 		// Setup dynamic color parameter. Only set when using particle parameter distributions.
 		FVector4 ColorOverLife(1.0f, 1.0f, 1.0f, 1.0f);
 		FVector4 ColorScaleOverLife(1.0f, 1.0f, 1.0f, 1.0f);
-		if( EmitterInfo.DynamicColorScale.Distribution )
+		if( EmitterInfo.DynamicColorScale.IsCreated() )
 		{
 			ColorScaleOverLife = EmitterInfo.DynamicColorScale.GetValue(0.0f,Component);
 		}
-		if( EmitterInfo.DynamicAlphaScale.Distribution )
+		if( EmitterInfo.DynamicAlphaScale.IsCreated() )
 		{
 			ColorScaleOverLife.W = EmitterInfo.DynamicAlphaScale.GetValue(0.0f,Component);
 		}
 
-		if( EmitterInfo.DynamicColor.Distribution )
+		if( EmitterInfo.DynamicColor.IsCreated() )
 		{
 			ColorOverLife = EmitterInfo.DynamicColor.GetValue(0.0f,Component);
 		}
-		if( EmitterInfo.DynamicAlpha.Distribution )
+		if( EmitterInfo.DynamicAlpha.IsCreated() )
 		{
 			ColorOverLife.W = EmitterInfo.DynamicAlpha.GetValue(0.0f,Component);
 		}
