Copying //UE4/Dev-Framework to Dev-Main (//UE4/Dev-Main) @ 2879625

diff --git a/Engine/Source/Runtime/Engine/Private/Particles/ParticleModules.cpp b/Engine/Source/Runtime/Engine/Private/Particles/ParticleModules.cpp
--- a/Engine/Source/Runtime/Engine/Private/Particles/ParticleModules.cpp
+++ b/Engine/Source/Runtime/Engine/Private/Particles/ParticleModules.cpp
@@ -779,8 +779,15 @@ void UParticleModule::SetTransactionFlag()
 			if( (ObjectPropertyBase->PropertyClass == UDistributionFloat::StaticClass() ||
 				 ObjectPropertyBase->PropertyClass == UDistributionVector::StaticClass()) )
 			{
-				UObject* Distribution = ObjectPropertyBase->GetObjectPropertyValue( ObjectPropertyBase->ContainerPtrToValuePtr<void>(this) );
-				Distribution->SetFlags( RF_Transactional );
+				//We've deprecated all the UFloatDistribution* usage and in PostLoad we set them to nullptr. TODO: can we just remove this code instead of a null check?
+				if(UObject* Distribution = ObjectPropertyBase->GetObjectPropertyValue( ObjectPropertyBase->ContainerPtrToValuePtr<void>(this) ))
+				{
+					Distribution->SetFlags( RF_Transactional );
+				}
+				else
+				{
+					ensure(ObjectPropertyBase->HasAllPropertyFlags(CPF_Deprecated));
+				}
 			}
 		}
 		else if( UArrayProperty* ArrayProp = Cast<UArrayProperty>(Property) )
@@ -897,7 +904,7 @@ UParticleModuleSourceMovement::UParticleModuleSourceMovement(const FObjectInitia
 
 void UParticleModuleSourceMovement::InitializeDefaults()
 {
-	if (!SourceMovementScale.Distribution)
+	if (!SourceMovementScale.IsCreated())
 	{
 		UDistributionVectorConstant* DistributionSourceMovementScale = NewObject<UDistributionVectorConstant>(this, TEXT("DistributionSourceMovementScale"));
 		DistributionSourceMovementScale->Constant = FVector(1.0f, 1.0f, 1.0f);
@@ -1034,7 +1041,7 @@ UParticleModuleRequired::UParticleModuleRequired(const FObjectInitializer& Objec
 
 void UParticleModuleRequired::InitializeDefaults()
 {
-	if (!SpawnRate.Distribution)
+	if (!SpawnRate.IsCreated())
 	{
 		SpawnRate.Distribution = NewObject<UDistributionFloatConstant>(this, TEXT("RequiredDistributionSpawnRate"));
 	}
@@ -1182,7 +1189,7 @@ UParticleModuleMeshRotation::UParticleModuleMeshRotation(const FObjectInitialize
 
 void UParticleModuleMeshRotation::InitializeDefaults()
 {
-	if (!StartRotation.Distribution)
+	if (!StartRotation.IsCreated())
 	{
 		UDistributionVectorUniform* DistributionStartRotation = NewObject<UDistributionVectorUniform>(this, TEXT("DistributionStartRotation"));
 		DistributionStartRotation->Min = FVector(0.0f, 0.0f, 0.0f);
@@ -1297,7 +1304,7 @@ UParticleModuleMeshRotationRate::UParticleModuleMeshRotationRate(const FObjectIn
 
 void UParticleModuleMeshRotationRate::InitializeDefaults()
 {
-	if (!StartRotationRate.Distribution)
+	if (!StartRotationRate.IsCreated())
 	{
 		UDistributionVectorUniform* DistributionStartRotationRate = NewObject<UDistributionVectorUniform>(this, TEXT("DistributionStartRotationRate"));
 		DistributionStartRotationRate->Min = FVector(0.0f, 0.0f, 0.0f);
@@ -1407,7 +1414,7 @@ UParticleModuleMeshRotationRateMultiplyLife::UParticleModuleMeshRotationRateMult
 }
 void UParticleModuleMeshRotationRateMultiplyLife::InitializeDefaults()
 {
-	if (!LifeMultiplier.Distribution)
+	if (!LifeMultiplier.IsCreated())
 	{
 		LifeMultiplier.Distribution = NewObject<UDistributionVectorConstant>(this, TEXT("DistributionLifeMultiplier"));
 	}
@@ -1483,7 +1490,7 @@ UParticleModuleMeshRotationRateOverLife::UParticleModuleMeshRotationRateOverLife
 
 void UParticleModuleMeshRotationRateOverLife::InitializeDefaults()
 {
-	if (!RotRate.Distribution)
+	if (!RotRate.IsCreated())
 	{
 		RotRate.Distribution = NewObject<UDistributionVectorConstantCurve>(this, TEXT("DistributionRotRate"));
 	}
@@ -1589,7 +1596,7 @@ UParticleModuleRotation::UParticleModuleRotation(const FObjectInitializer& Objec
 
 void UParticleModuleRotation::InitializeDefaults()
 {
-	if (!StartRotation.Distribution)
+	if (!StartRotation.IsCreated())
 	{
 		UDistributionFloatUniform* DistributionStartRotation = NewObject<UDistributionFloatUniform>(this, TEXT("DistributionStartRotation"));
 		DistributionStartRotation->Min = 0.0f;
@@ -1675,7 +1682,7 @@ UParticleModuleRotationRate::UParticleModuleRotationRate(const FObjectInitialize
 
 void UParticleModuleRotationRate::InitializeDefaults()
 {
-	if (!StartRotationRate.Distribution)
+	if (!StartRotationRate.IsCreated())
 	{
 		StartRotationRate.Distribution = NewObject<UDistributionFloatConstant>(this, TEXT("DistributionStartRotationRate"));
 	}
@@ -1786,7 +1793,7 @@ UParticleModuleRotationOverLifetime::UParticleModuleRotationOverLifetime(const F
 
 void UParticleModuleRotationOverLifetime::InitializeDefaults()
 {
-	if (!RotationOverLife.Distribution)
+	if (!RotationOverLife.IsCreated())
 	{
 		RotationOverLife.Distribution = NewObject<UDistributionFloatConstantCurve>(this, TEXT("DistributionRotOverLife"));
 	}
@@ -1853,7 +1860,7 @@ UParticleModuleSubUV::UParticleModuleSubUV(const FObjectInitializer& ObjectIniti
 
 void UParticleModuleSubUV::InitializeDefaults()
 {
-	if (!SubImageIndex.Distribution)
+	if (!SubImageIndex.IsCreated())
 	{
 		SubImageIndex.Distribution = NewObject<UDistributionFloatConstant>(this, TEXT("DistributionSubImage"));
 	}
@@ -2092,7 +2099,7 @@ UParticleModuleSubUVMovie::UParticleModuleSubUVMovie(const FObjectInitializer& O
 
 void UParticleModuleSubUVMovie::InitializeDefaults()
 {
-	if (!FrameRate.Distribution)
+	if (!FrameRate.IsCreated())
 	{
 		UDistributionFloatConstant* DistributionFrameRate = NewObject<UDistributionFloatConstant>(this, TEXT("DistributionFrameRate"));
 		DistributionFrameRate->Constant = 30.0f;
@@ -2264,7 +2271,7 @@ UParticleModuleRotationRateMultiplyLife::UParticleModuleRotationRateMultiplyLife
 
 void UParticleModuleRotationRateMultiplyLife::InitializeDefaults()
 {
-	if (!LifeMultiplier.Distribution)
+	if (!LifeMultiplier.IsCreated())
 	{
 		LifeMultiplier.Distribution = NewObject<UDistributionFloatConstant>(this, TEXT("DistributionLifeMultiplier"));
 	}
@@ -2425,11 +2432,11 @@ UParticleModuleAccelerationDrag::UParticleModuleAccelerationDrag(const FObjectIn
 
 void UParticleModuleAccelerationDrag::InitializeDefaults()
 {
-	if (!DragCoefficient)
+	if (!DragCoefficientRaw.IsCreated())
 	{
 		UDistributionFloatConstant* DistributionDragCoefficient = NewObject<UDistributionFloatConstant>(this, TEXT("DistributionDragCoefficient"));
 		DistributionDragCoefficient->Constant = 1.0f;
-		DragCoefficient = DistributionDragCoefficient;
+		DragCoefficientRaw.Distribution = DistributionDragCoefficient;
 	}
 }
 
@@ -2442,9 +2449,22 @@ void UParticleModuleAccelerationDrag::PostInitProperties()
 	}
 }
 
+void UParticleModuleAccelerationDrag::PostLoad()
+{
+	Super::PostLoad();
+#if WITH_EDITOR
+	if(DragCoefficient_DEPRECATED)
+	{
+		DragCoefficientRaw.Distribution = DragCoefficient_DEPRECATED;
+		DragCoefficientRaw.Initialize();
+	}
+	DragCoefficient_DEPRECATED = nullptr;
+#endif
+}
+
 void UParticleModuleAccelerationDrag::CompileModule(FParticleEmitterBuildInfo& EmitterInfo)
 {
-	EmitterInfo.DragCoefficient.Initialize(DragCoefficient);
+	EmitterInfo.DragCoefficient.Initialize(DragCoefficientRaw.Distribution);
 }
 
 #if WITH_EDITOR
@@ -2456,9 +2476,9 @@ void UParticleModuleAccelerationDrag::PostEditChangeProperty(FPropertyChangedEve
 
 bool UParticleModuleAccelerationDrag::IsValidForLODLevel(UParticleLODLevel* LODLevel, FString& OutErrorString)
 {
-	if (DragCoefficient && LODLevel->TypeDataModule && LODLevel->TypeDataModule->IsA(UParticleModuleTypeDataGpu::StaticClass()))
+	if (DragCoefficientRaw.Distribution && LODLevel->TypeDataModule && LODLevel->TypeDataModule->IsA(UParticleModuleTypeDataGpu::StaticClass()))
 	{
-		if(!IsDistributionAllowedOnGPU(DragCoefficient))
+		if(!IsDistributionAllowedOnGPU(DragCoefficientRaw.Distribution))
 		{
 			OutErrorString = GetDistributionNotAllowedOnGPUText(StaticClass()->GetName(), "DragCoefficient" ).ToString();
 			return false;
@@ -2473,7 +2493,7 @@ void UParticleModuleAccelerationDrag::Update(FParticleEmitterInstance* Owner, in
 {
 	BEGIN_UPDATE_LOOP;
 	{
-		FVector Drag  = Particle.Velocity * -DragCoefficient->GetValue(Particle.RelativeTime, Owner->Component);
+		FVector Drag  = Particle.Velocity * -DragCoefficientRaw.GetValue(Particle.RelativeTime, Owner->Component);
 		Particle.Velocity		+= Drag * DeltaTime;
 		Particle.BaseVelocity	+= Drag * DeltaTime;
 	}
@@ -2490,11 +2510,11 @@ UParticleModuleAccelerationDragScaleOverLife::UParticleModuleAccelerationDragSca
 
 void UParticleModuleAccelerationDragScaleOverLife::InitializeDefaults()
 {
-	if (!DragScale)
+	if (!DragScaleRaw.IsCreated())
 	{
 		UDistributionFloatConstant* DistributionDragScale = NewObject<UDistributionFloatConstant>(this, TEXT("DistributionDragScale"));
 		DistributionDragScale->Constant = 1.0f;
-		DragScale = DistributionDragScale;
+		DragScaleRaw.Distribution = DistributionDragScale;
 	}
 }
 
@@ -2507,9 +2527,22 @@ void UParticleModuleAccelerationDragScaleOverLife::PostInitProperties()
 	}
 }
 
+void UParticleModuleAccelerationDragScaleOverLife::PostLoad()
+{
+	Super::PostLoad();
+#if WITH_EDITOR
+	if (DragScale_DEPRECATED)
+	{
+		DragScaleRaw.Distribution = DragScale_DEPRECATED;
+		DragScaleRaw.Initialize();
+	}
+	DragScale_DEPRECATED = nullptr;
+#endif
+}
+
 void UParticleModuleAccelerationDragScaleOverLife::CompileModule(FParticleEmitterBuildInfo& EmitterInfo)
 {
-	EmitterInfo.DragScale.ScaleByDistribution(DragScale);
+	EmitterInfo.DragScale.ScaleByDistribution(DragScaleRaw.Distribution);
 }
 
 #if WITH_EDITOR
@@ -2523,7 +2556,7 @@ bool UParticleModuleAccelerationDragScaleOverLife::IsValidForLODLevel(UParticleL
 {
 	if (LODLevel->TypeDataModule && LODLevel->TypeDataModule->IsA(UParticleModuleTypeDataGpu::StaticClass()))
 	{
-		if(!IsDistributionAllowedOnGPU(DragScale))
+		if(!IsDistributionAllowedOnGPU(DragScaleRaw.Distribution))
 		{
 			OutErrorString = GetDistributionNotAllowedOnGPUText(StaticClass()->GetName(), "DragScale" ).ToString();
 			return false;
@@ -2546,11 +2579,11 @@ UParticleModuleAttractorPointGravity::UParticleModuleAttractorPointGravity(const
 
 void UParticleModuleAttractorPointGravity::InitializeDefaults()
 {
-	if (!Strength)
+	if (!StrengthRaw.IsCreated())
 	{
 		UDistributionFloatConstant* DistributionStrength = NewObject<UDistributionFloatConstant>(this, TEXT("DistributionStrength"));
 		DistributionStrength->Constant = 1.0f;
-		Strength = DistributionStrength;
+		StrengthRaw.Distribution = DistributionStrength;
 	}
 }
 
@@ -2563,11 +2596,24 @@ void UParticleModuleAttractorPointGravity::PostInitProperties()
 	}
 }
 
+void UParticleModuleAttractorPointGravity::PostLoad()
+{
+	Super::PostLoad();
+#if WITH_EDITOR
+	if (Strength_DEPRECATED)
+	{
+		StrengthRaw.Distribution = Strength_DEPRECATED;
+		StrengthRaw.Initialize();
+	}
+	Strength_DEPRECATED = nullptr;
+#endif
+}
+
 void UParticleModuleAttractorPointGravity::CompileModule(FParticleEmitterBuildInfo& EmitterInfo)
 {
 	EmitterInfo.PointAttractorPosition = Position;
 	EmitterInfo.PointAttractorRadius = Radius;
-	EmitterInfo.PointAttractorStrength.Initialize(Strength);
+	EmitterInfo.PointAttractorStrength.Initialize(StrengthRaw.Distribution);
 }
 
 #if WITH_EDITOR
@@ -2603,7 +2649,7 @@ UParticleModuleAcceleration::UParticleModuleAcceleration(const FObjectInitialize
 
 void UParticleModuleAcceleration::InitializeDefaults()
 {
-	if (!Acceleration.Distribution)
+	if (!Acceleration.IsCreated())
 	{
 		Acceleration.Distribution = NewObject<UDistributionVectorUniform>(this, TEXT("DistributionAcceleration"));
 	}
@@ -2734,7 +2780,7 @@ UParticleModuleAccelerationOverLifetime::UParticleModuleAccelerationOverLifetime
 
 void UParticleModuleAccelerationOverLifetime::InitializeDefaults()
 {
-	if (!AccelOverLife.Distribution)
+	if (!AccelOverLife.IsCreated())
 	{
 		AccelOverLife.Distribution = NewObject<UDistributionVectorConstantCurve>(this, TEXT("DistributionAccelOverLife"));
 	}
@@ -2806,22 +2852,22 @@ UParticleModuleLight::UParticleModuleLight(const FObjectInitializer& ObjectIniti
 
 void UParticleModuleLight::InitializeDefaults()
 {
-	if (!ColorScaleOverLife.Distribution)
+	if (!ColorScaleOverLife.IsCreated())
 	{
 		ColorScaleOverLife.Distribution = NewObject<UDistributionVectorConstant>(this, TEXT("DistributionColorScaleOverLife"));
 	}
 
-	if (!BrightnessOverLife.Distribution)
+	if (!BrightnessOverLife.IsCreated())
 	{
 		BrightnessOverLife.Distribution = NewObject<UDistributionFloatConstant>(this, TEXT("DistributionBrightnessOverLife"));
 	}
 
-	if (!RadiusScale.Distribution)
+	if (!RadiusScale.IsCreated())
 	{
 		RadiusScale.Distribution = NewObject<UDistributionFloatConstant>(this, TEXT("DistributionRadiusScale"));
 	}
 
-	if (!LightExponent.Distribution)
+	if (!LightExponent.IsCreated())
 	{
 		LightExponent.Distribution = NewObject<UDistributionFloatConstant>(this, TEXT("DistributionLightExponent"));
 	}
@@ -3222,7 +3268,7 @@ void UParticleModuleTypeDataMesh::Serialize(FArchive& Ar)
 
 void UParticleModuleTypeDataMesh::CreateDistribution()
 {
-	if (!RollPitchYawRange.Distribution)
+	if (!RollPitchYawRange.IsCreated())
 	{
 		RollPitchYawRange.Distribution = NewObject<UDistributionVectorUniform>(this, TEXT("DistributionRollPitchYaw"));
 	}
@@ -3278,12 +3324,12 @@ UParticleModuleKillBox::UParticleModuleKillBox(const FObjectInitializer& ObjectI
 
 void UParticleModuleKillBox::InitializeDefaults()
 {
-	if (!LowerLeftCorner.Distribution)
+	if (!LowerLeftCorner.IsCreated())
 	{
 		LowerLeftCorner.Distribution = NewObject<UDistributionVectorConstant>(this, TEXT("DistributionLowerLeftCorner"));
 	}
 
-	if (!UpperRightCorner.Distribution)
+	if (!UpperRightCorner.IsCreated())
 	{
 		UpperRightCorner.Distribution = NewObject<UDistributionVectorConstant>(this, TEXT("DistributionUpperRightCorner"));
 	}
@@ -3428,7 +3474,7 @@ UParticleModuleKillHeight::UParticleModuleKillHeight(const FObjectInitializer& O
 
 void UParticleModuleKillHeight::InitializeDefaults()
 {
-	if (!Height.Distribution)
+	if (!Height.IsCreated())
 	{
 		Height.Distribution = NewObject<UDistributionFloatConstant>(this, TEXT("DistributionHeight"));
 	}
@@ -3551,7 +3597,7 @@ UParticleModuleLifetime::UParticleModuleLifetime(const FObjectInitializer& Objec
 
 void UParticleModuleLifetime::InitializeDefaults()
 {
-	if(!Lifetime.Distribution)
+	if(!Lifetime.IsCreated())
 	{
 		Lifetime.Distribution = NewObject<UDistributionFloatUniform>(this, TEXT("DistributionLifetime"));
 	}
@@ -3701,12 +3747,12 @@ UParticleModuleAttractorLine::UParticleModuleAttractorLine(const FObjectInitiali
 
 void UParticleModuleAttractorLine::InitializeDefaults()
 {
-	if(!Strength.Distribution)
+	if(!Strength.IsCreated())
 	{
 		Strength.Distribution = NewObject<UDistributionFloatConstant>(this, TEXT("DistributionStrength"));
 	}
 
-	if(!Range.Distribution)
+	if(!Range.IsCreated())
 	{
 		Range.Distribution = NewObject<UDistributionFloatConstant>(this, TEXT("DistributionRange"));
 	}
@@ -3849,12 +3895,12 @@ UParticleModuleAttractorParticle::UParticleModuleAttractorParticle(const FObject
 
 void UParticleModuleAttractorParticle::InitializeDefaults()
 {
-	if(!Range.Distribution)
+	if(!Range.IsCreated())
 	{
 		Range.Distribution = NewObject<UDistributionFloatConstant>(this, TEXT("DistributionRange"));
 	}
 
-	if(!Strength.Distribution)
+	if(!Strength.IsCreated())
 	{
 		Strength.Distribution = NewObject<UDistributionFloatConstant>(this, TEXT("DistributionStrength"));
 	}
@@ -4092,17 +4138,17 @@ UParticleModuleAttractorPoint::UParticleModuleAttractorPoint(const FObjectInitia
 
 void UParticleModuleAttractorPoint::InitializeDefaults()
 {
-	if(!Position.Distribution)
+	if(!Position.IsCreated())
 	{
 		Position.Distribution = NewObject<UDistributionVectorConstant>(this, TEXT("DistributionPosition"));
 	}
 	
-	if(!Range.Distribution)
+	if(!Range.IsCreated())
 	{
 		Range.Distribution = NewObject<UDistributionFloatConstant>(this, TEXT("DistributionRange"));
 	}
 
-	if(!Strength.Distribution)
+	if(!Strength.IsCreated())
 	{	
 		Strength.Distribution = NewObject<UDistributionFloatConstant>(this, TEXT("DistributionStrength"));
 	}
@@ -4325,6 +4371,7 @@ void UParticleModuleTypeDataGpu::BeginDestroy()
 
 void UParticleModuleTypeDataGpu::Build( FParticleEmitterBuildInfo& EmitterBuildInfo )
 {
+#if WITH_EDITOR
 	FVector4Distribution Curve;
 	FComposableFloatDistribution ZeroDistribution;
 	FComposableFloatDistribution OneDistribution;
@@ -4619,6 +4666,8 @@ void UParticleModuleTypeDataGpu::Build( FParticleEmitterBuildInfo& EmitterBuildI
 	// Collision flag.
 	EmitterInfo.bEnableCollision = EmitterBuildInfo.bEnableCollision;
 	EmitterInfo.CollisionMode = (EParticleCollisionMode::Type)EmitterBuildInfo.CollisionMode;
+#endif
+
 
 	// Create or update GPU resources.
 	if ( EmitterInfo.Resources )
diff --git a/Engine/Source/Runtime/Engine/Private/Particles/ParticleModules_Camera.cpp b/Engine/Source/Runtime/Engine/Private/Particles/ParticleModules_Camera.cpp
--- a/Engine/Source/Runtime/Engine/Private/Particles/ParticleModules_Camera.cpp
+++ b/Engine/Source/Runtime/Engine/Private/Particles/ParticleModules_Camera.cpp
@@ -36,7 +36,7 @@ UParticleModuleCameraOffset::UParticleModuleCameraOffset(const FObjectInitialize
 
 void UParticleModuleCameraOffset::InitializeDefaults()
 {
-	if (!CameraOffset.Distribution)
+	if (!CameraOffset.IsCreated())
 	{
 		UDistributionFloatConstant* DistributionCameraOffset = NewObject<UDistributionFloatConstant>(this, TEXT("DistributionCameraOffset"));
 		DistributionCameraOffset->Constant = 1.0f;
diff --git a/Engine/Source/Runtime/Engine/Private/Particles/ParticleModules_Collision.cpp b/Engine/Source/Runtime/Engine/Private/Particles/ParticleModules_Collision.cpp
--- a/Engine/Source/Runtime/Engine/Private/Particles/ParticleModules_Collision.cpp
+++ b/Engine/Source/Runtime/Engine/Private/Particles/ParticleModules_Collision.cpp
@@ -56,31 +56,31 @@ UParticleModuleCollision::UParticleModuleCollision(const FObjectInitializer& Obj
 
 void UParticleModuleCollision::InitializeDefaults()
 {
-	if (!DampingFactor.Distribution)
+	if (!DampingFactor.IsCreated())
 	{
 		DampingFactor.Distribution = NewObject<UDistributionVectorUniform>(this, TEXT("DistributionDampingFactor"));
 	}
 
-	if (!DampingFactorRotation.Distribution)
+	if (!DampingFactorRotation.IsCreated())
 	{
 		UDistributionVectorConstant* DistributionDampingFactorRotation = NewObject<UDistributionVectorConstant>(this, TEXT("DistributionDampingFactorRotation"));
 		DistributionDampingFactorRotation->Constant = FVector(1.0f, 1.0f, 1.0f);
 		DampingFactorRotation.Distribution = DistributionDampingFactorRotation; 
 	}
 
-	if (!MaxCollisions.Distribution)
+	if (!MaxCollisions.IsCreated())
 	{
 		MaxCollisions.Distribution = NewObject<UDistributionFloatUniform>(this, TEXT("DistributionMaxCollisions"));
 	}
 
-	if (!ParticleMass.Distribution)
+	if (!ParticleMass.IsCreated())
 	{
 		UDistributionFloatConstant* DistributionParticleMass = NewObject<UDistributionFloatConstant>(this, TEXT("DistributionParticleMass"));
 		DistributionParticleMass->Constant = 0.1f;
 		ParticleMass.Distribution = DistributionParticleMass;
 	}
 
-	if (!DelayAmount.Distribution)
+	if (!DelayAmount.IsCreated())
 	{
 		UDistributionFloatConstant* DistributionDelayAmount = NewObject<UDistributionFloatConstant>(this, TEXT("DistributionDelayAmount"));
 		DistributionDelayAmount->Constant = 0.0f;
@@ -574,14 +574,14 @@ UParticleModuleCollisionGPU::UParticleModuleCollisionGPU(const FObjectInitialize
 
 void UParticleModuleCollisionGPU::InitializeDefaults()
 {
-	if (!Resilience.Distribution)
+	if (!Resilience.IsCreated())
 	{
 		UDistributionFloatConstant* ResilienceDistribution = NewObject<UDistributionFloatConstant>(this, TEXT("ResilienceDistribution"));
 		ResilienceDistribution->Constant = 0.5f;
 		Resilience.Distribution = ResilienceDistribution;
 	}
 
-	if (!ResilienceScaleOverLife.Distribution)
+	if (!ResilienceScaleOverLife.IsCreated())
 	{
 		UDistributionFloatConstant* ResilienceScaleOverLifeDistribution = NewObject<UDistributionFloatConstant>(this, TEXT("ResilienceScaleOverLifeDistribution"));
 		ResilienceScaleOverLifeDistribution->Constant = 1.0f;
diff --git a/Engine/Source/Runtime/Engine/Private/Particles/ParticleModules_Color.cpp b/Engine/Source/Runtime/Engine/Private/Particles/ParticleModules_Color.cpp
--- a/Engine/Source/Runtime/Engine/Private/Particles/ParticleModules_Color.cpp
+++ b/Engine/Source/Runtime/Engine/Private/Particles/ParticleModules_Color.cpp
@@ -45,12 +45,12 @@ UParticleModuleColor::UParticleModuleColor(const FObjectInitializer& ObjectIniti
 
 void UParticleModuleColor::InitializeDefaults() 
 {
-	if (!StartColor.Distribution)
+	if (!StartColor.IsCreated())
 	{
 		StartColor.Distribution = NewObject<UDistributionVectorConstant>(this, TEXT("DistributionStartColor"));
 	}
 
-	if (!StartAlpha.Distribution)
+	if (!StartAlpha.IsCreated())
 	{
 		UDistributionFloatConstant* DistributionStartAlpha = NewObject<UDistributionFloatConstant>(this, TEXT("DistributionStartAlpha"));
 		DistributionStartAlpha->Constant = 1.0f;
@@ -222,12 +222,12 @@ UParticleModuleColorOverLife::UParticleModuleColorOverLife(const FObjectInitiali
 
 void UParticleModuleColorOverLife::InitializeDefaults()
 {
-	if (!ColorOverLife.Distribution)
+	if (!ColorOverLife.IsCreated())
 	{
 		ColorOverLife.Distribution = NewObject<UDistributionVectorConstantCurve>(this, TEXT("DistributionColorOverLife"));
 	}
 
-	if (!AlphaOverLife.Distribution)
+	if (!AlphaOverLife.IsCreated())
 	{
 		UDistributionFloatConstant* DistributionAlphaOverLife = NewObject<UDistributionFloatConstant>(this, TEXT("DistributionAlphaOverLife"));
 		DistributionAlphaOverLife->Constant = 1.0f;
@@ -472,12 +472,12 @@ UParticleModuleColorScaleOverLife::UParticleModuleColorScaleOverLife(const FObje
 
 void UParticleModuleColorScaleOverLife::InitializeDefaults()
 {
-	if (!ColorScaleOverLife.Distribution)
+	if (!ColorScaleOverLife.IsCreated())
 	{
 		ColorScaleOverLife.Distribution = NewObject<UDistributionVectorConstantCurve>(this, TEXT("DistributionColorScaleOverLife"));
 	}
 
-	if (!AlphaScaleOverLife.Distribution)
+	if (!AlphaScaleOverLife.IsCreated())
 	{
 		UDistributionFloatConstant* DistributionAlphaScaleOverLife = NewObject<UDistributionFloatConstant>(this, TEXT("DistributionAlphaScaleOverLife"));
 		DistributionAlphaScaleOverLife->Constant = 1.0f;
diff --git a/Engine/Source/Runtime/Engine/Private/Particles/ParticleModules_Event.cpp b/Engine/Source/Runtime/Engine/Private/Particles/ParticleModules_Event.cpp
--- a/Engine/Source/Runtime/Engine/Private/Particles/ParticleModules_Event.cpp
+++ b/Engine/Source/Runtime/Engine/Private/Particles/ParticleModules_Event.cpp
@@ -295,14 +295,14 @@ UParticleModuleEventReceiverSpawn::UParticleModuleEventReceiverSpawn(const FObje
 
 void UParticleModuleEventReceiverSpawn::InitializeDefaults()
 {
-	if (!SpawnCount.Distribution)
+	if (!SpawnCount.IsCreated())
 	{
 		UDistributionFloatConstant* RequiredDistributionSpawnCount = NewObject<UDistributionFloatConstant>(this, TEXT("RequiredDistributionSpawnCount"));
 		RequiredDistributionSpawnCount->Constant = 0.0f;
 		SpawnCount.Distribution = RequiredDistributionSpawnCount;
 	}
 
-	if (!InheritVelocityScale.Distribution)
+	if (!InheritVelocityScale.IsCreated())
 	{
 		UDistributionVectorConstant* RequiredDistributionInheritVelocityScale = NewObject<UDistributionVectorConstant>(this, TEXT("RequiredDistributionInheritVelocityScale"));
 		RequiredDistributionInheritVelocityScale->Constant = FVector(1.0f, 1.0f, 1.0f);
diff --git a/Engine/Source/Runtime/Engine/Private/Particles/ParticleModules_Location.cpp b/Engine/Source/Runtime/Engine/Private/Particles/ParticleModules_Location.cpp
--- a/Engine/Source/Runtime/Engine/Private/Particles/ParticleModules_Location.cpp
+++ b/Engine/Source/Runtime/Engine/Private/Particles/ParticleModules_Location.cpp
@@ -55,7 +55,7 @@ UParticleModuleLocation::UParticleModuleLocation(const FObjectInitializer& Objec
 
 void UParticleModuleLocation::InitializeDefaults()
 {
-	if (!StartLocation.Distribution)
+	if (!StartLocation.IsCreated())
 	{
 		StartLocation.Distribution = NewObject<UDistributionVectorUniform>(this, TEXT("DistributionStartLocation"));
 	}
@@ -102,7 +102,7 @@ void UParticleModuleLocation::SpawnEx(FParticleEmitterInstance* Owner, int32 Off
 		else
 		{
 			FVector Min, Max;
-			StartLocation.Distribution->GetRange(Min, Max);
+			StartLocation.GetRange(Min, Max);
 			FVector Lerped = FMath::Lerp(Min, Max, FMath::TruncToFloat((FMath::SRand() * (DistributeOverNPoints - 1.0f)) + 0.5f)/(DistributeOverNPoints - 1.0f));
 			LocationOffset.Set(Lerped.X, Lerped.Y, Lerped.Z);
 		}
@@ -278,19 +278,19 @@ UParticleModuleLocationDirect::UParticleModuleLocationDirect(const FObjectInitia
 
 void UParticleModuleLocationDirect::InitializeDefaults()
 {
-	if (!Location.Distribution)
+	if (!Location.IsCreated())
 	{
 		Location.Distribution = NewObject<UDistributionVectorUniform>(this, TEXT("DistributionLocation"));
 	}
 
-	if (!LocationOffset.Distribution)
+	if (!LocationOffset.IsCreated())
 	{
 		UDistributionVectorConstant* DistributionLocationOffset = NewObject<UDistributionVectorConstant>(this, TEXT("DistributionLocationOffset"));
 		DistributionLocationOffset->Constant = FVector(0.0f, 0.0f, 0.0f);
 		LocationOffset.Distribution = DistributionLocationOffset;
 	}
 
-	if (!Direction.Distribution)
+	if (!Direction.IsCreated())
 	{
 		UDistributionVectorConstant* DistributionScaleFactor = NewObject<UDistributionVectorConstant>(this, TEXT("DistributionScaleFactor"));
 		DistributionScaleFactor->Constant = FVector(1.0f, 1.0f, 1.0f);
@@ -647,14 +647,14 @@ UParticleModuleLocationPrimitiveBase::UParticleModuleLocationPrimitiveBase(const
 
 void UParticleModuleLocationPrimitiveBase::InitializeDefaults()
 {
-	if (!VelocityScale.Distribution)
+	if (!VelocityScale.IsCreated())
 	{
 		UDistributionFloatConstant* DistributionVelocityScale = NewObject<UDistributionFloatConstant>(this, TEXT("DistributionVelocityScale"));
 		DistributionVelocityScale->Constant = 1.0f;
 		VelocityScale.Distribution = DistributionVelocityScale;
 	}
 
-	if (!StartLocation.Distribution)
+	if (!StartLocation.IsCreated())
 	{
 		UDistributionVectorConstant* DistributionStartLocation = NewObject<UDistributionVectorConstant>(this, TEXT("DistributionStartLocation"));
 		DistributionStartLocation->Constant = FVector(0.0f, 0.0f, 0.0f);
@@ -763,28 +763,28 @@ UParticleModuleLocationPrimitiveTriangle::UParticleModuleLocationPrimitiveTriang
 
 void UParticleModuleLocationPrimitiveTriangle::InitializeDefaults()
 {
-	if (!StartOffset.Distribution)
+	if (!StartOffset.IsCreated())
 	{
 		UDistributionVectorConstant* DistributionOffset = NewObject<UDistributionVectorConstant>(this, TEXT("DistributionOffset"));
 		DistributionOffset->Constant = FVector::ZeroVector;
 		StartOffset.Distribution = DistributionOffset;
 	}
 
-	if (!Height.Distribution)
+	if (!Height.IsCreated())
 	{
 		UDistributionFloatConstant* DistributionHeight = NewObject<UDistributionFloatConstant>(this, TEXT("DistributionHeight"));
 		DistributionHeight->Constant = 50.0f;
 		Height.Distribution = DistributionHeight;
 	}
 
-	if (!Angle.Distribution)
+	if (!Angle.IsCreated())
 	{
 		UDistributionFloatConstant* DistributionAngle = NewObject<UDistributionFloatConstant>(this, TEXT("DistributionAngle"));
 		DistributionAngle->Constant = 90.0f;
 		Angle.Distribution = DistributionAngle;
 	}
 
-	if (!Thickness.Distribution)
+	if (!Thickness.IsCreated())
 	{
 		UDistributionFloatConstant* DistributionThickness = NewObject<UDistributionFloatConstant>(this, TEXT("DistributionThickness"));
 		DistributionThickness->Constant = 0.0f;
@@ -930,14 +930,14 @@ UParticleModuleLocationPrimitiveCylinder::UParticleModuleLocationPrimitiveCylind
 
 void UParticleModuleLocationPrimitiveCylinder::InitializeDefaults()
 {
-	if (!StartRadius.Distribution)
+	if (!StartRadius.IsCreated())
 	{
 		UDistributionFloatConstant* DistributionStartRadius = NewObject<UDistributionFloatConstant>(this, TEXT("DistributionStartRadius"));
 		DistributionStartRadius->Constant = 50.0f;
 		StartRadius.Distribution = DistributionStartRadius;
 	}
 
-	if (!StartHeight.Distribution)
+	if (!StartHeight.IsCreated())
 	{
 		UDistributionFloatConstant* DistributionStartHeight = NewObject<UDistributionFloatConstant>(this, TEXT("DistributionStartHeight"));
 		DistributionStartHeight->Constant = 50.0f;
@@ -1239,7 +1239,7 @@ UParticleModuleLocationPrimitiveSphere::UParticleModuleLocationPrimitiveSphere(c
 
 void UParticleModuleLocationPrimitiveSphere::InitializeDefaults()
 {
-	if (!StartRadius.Distribution)
+	if (!StartRadius.IsCreated())
 	{
 		UDistributionFloatConstant* DistributionStartRadius = NewObject<UDistributionFloatConstant>(this, TEXT("DistributionStartRadius"));
 		DistributionStartRadius->Constant = 50.0f;
diff --git a/Engine/Source/Runtime/Engine/Private/Particles/ParticleModules_Orbit.cpp b/Engine/Source/Runtime/Engine/Private/Particles/ParticleModules_Orbit.cpp
--- a/Engine/Source/Runtime/Engine/Private/Particles/ParticleModules_Orbit.cpp
+++ b/Engine/Source/Runtime/Engine/Private/Particles/ParticleModules_Orbit.cpp
@@ -33,7 +33,7 @@ UParticleModuleOrbit::UParticleModuleOrbit(const FObjectInitializer& ObjectIniti
 
 void UParticleModuleOrbit::InitializeDefaults()
 {
-	if (!OffsetAmount.Distribution)
+	if (!OffsetAmount.IsCreated())
 	{
 		UDistributionVectorUniform* DistributionOffsetAmount = NewObject<UDistributionVectorUniform>(this, TEXT("DistributionOffsetAmount"));
 		DistributionOffsetAmount->Min = FVector(0.0f, 0.0f, 0.0f);
@@ -41,7 +41,7 @@ void UParticleModuleOrbit::InitializeDefaults()
 		OffsetAmount.Distribution = DistributionOffsetAmount;
 	}
 
-	if (!RotationAmount.Distribution)
+	if (!RotationAmount.IsCreated())
 	{
 		UDistributionVectorUniform* DistributionRotationAmount = NewObject<UDistributionVectorUniform>(this, TEXT("DistributionRotationAmount"));
 		DistributionRotationAmount->Min = FVector(0.0f, 0.0f, 0.0f);
@@ -49,7 +49,7 @@ void UParticleModuleOrbit::InitializeDefaults()
 		RotationAmount.Distribution = DistributionRotationAmount;
 	}
 
-	if (!RotationRateAmount.Distribution)
+	if (!RotationRateAmount.IsCreated())
 	{
 		UDistributionVectorUniform* DistributionRotationRateAmount = NewObject<UDistributionVectorUniform>(this, TEXT("DistributionRotationRateAmount"));
 		DistributionRotationRateAmount->Min = FVector(0.0f, 0.0f, 0.0f);
diff --git a/Engine/Source/Runtime/Engine/Private/Particles/ParticleModules_Parameter.cpp b/Engine/Source/Runtime/Engine/Private/Particles/ParticleModules_Parameter.cpp
--- a/Engine/Source/Runtime/Engine/Private/Particles/ParticleModules_Parameter.cpp
+++ b/Engine/Source/Runtime/Engine/Private/Particles/ParticleModules_Parameter.cpp
@@ -41,7 +41,7 @@ void UParticleModuleParameterDynamic::InitializeDefaults()
 {
 	for(int32 ParamIdx = 0; ParamIdx < DynamicParams.Num(); ++ParamIdx)
 	{
-		if(!DynamicParams[ParamIdx].ParamValue.Distribution)
+		if(!DynamicParams[ParamIdx].ParamValue.IsCreated())
 		{
 			DynamicParams[ParamIdx].ParamValue.Distribution = NewObject<UDistributionFloatConstant>(this, TEXT("DistributionParam1"));
 		}
diff --git a/Engine/Source/Runtime/Engine/Private/Particles/ParticleModules_Size.cpp b/Engine/Source/Runtime/Engine/Private/Particles/ParticleModules_Size.cpp
--- a/Engine/Source/Runtime/Engine/Private/Particles/ParticleModules_Size.cpp
+++ b/Engine/Source/Runtime/Engine/Private/Particles/ParticleModules_Size.cpp
@@ -40,7 +40,7 @@ UParticleModuleSize::UParticleModuleSize(const FObjectInitializer& ObjectInitial
 
 void UParticleModuleSize::InitializeDefaults()
 {
-	if (!StartSize.Distribution)
+	if (!StartSize.IsCreated())
 	{
 		UDistributionVectorUniform* DistributionStartSize = NewObject<UDistributionVectorUniform>(this, TEXT("DistributionStartSize"));
 		DistributionStartSize->Min = FVector(1.0f, 1.0f, 1.0f);
@@ -144,7 +144,7 @@ UParticleModuleSizeMultiplyLife::UParticleModuleSizeMultiplyLife(const FObjectIn
 
 void UParticleModuleSizeMultiplyLife::InitializeDefaults()
 {
-	if (!LifeMultiplier.Distribution)
+	if (!LifeMultiplier.IsCreated())
 	{
 		LifeMultiplier.Distribution = NewObject<UDistributionVectorConstant>(this, TEXT("DistributionLifeMultiplier"));
 	}
@@ -334,7 +334,7 @@ UParticleModuleSizeScale::UParticleModuleSizeScale(const FObjectInitializer& Obj
 
 void UParticleModuleSizeScale::InitializeDefaults()
 {
-	if (!SizeScale.Distribution)
+	if (!SizeScale.IsCreated())
 	{
 		SizeScale.Distribution = NewObject<UDistributionVectorConstant>(this, TEXT("DistributionSizeScale"));
 	}
diff --git a/Engine/Source/Runtime/Engine/Private/Particles/ParticleModules_Spawn.cpp b/Engine/Source/Runtime/Engine/Private/Particles/ParticleModules_Spawn.cpp
--- a/Engine/Source/Runtime/Engine/Private/Particles/ParticleModules_Spawn.cpp
+++ b/Engine/Source/Runtime/Engine/Private/Particles/ParticleModules_Spawn.cpp
@@ -36,21 +36,21 @@ UParticleModuleSpawn::UParticleModuleSpawn(const FObjectInitializer& ObjectIniti
 
 void UParticleModuleSpawn::InitializeDefaults()
 {
-	if(!Rate.Distribution)
+	if(!Rate.IsCreated())
 	{
 		UDistributionFloatConstant* RequiredDistributionSpawnRate = NewObject<UDistributionFloatConstant>(this, TEXT("RequiredDistributionSpawnRate"));
 		RequiredDistributionSpawnRate->Constant = 20.0f;
 		Rate.Distribution = RequiredDistributionSpawnRate;
 	}
 
-	if(!RateScale.Distribution)
+	if(!RateScale.IsCreated())
 	{
 		UDistributionFloatConstant* RequiredDistributionSpawnRateScale = NewObject<UDistributionFloatConstant>(this, TEXT("RequiredDistributionSpawnRateScale"));
 		RequiredDistributionSpawnRateScale->Constant = 1.0f;
 		RateScale.Distribution = RequiredDistributionSpawnRateScale;
 	}
 
-	if(!BurstScale.Distribution)
+	if(!BurstScale.IsCreated())
 	{
 		UDistributionFloatConstant* BurstScaleDistribution = NewObject<UDistributionFloatConstant>(this, TEXT("BurstScaleDistribution"));
 		BurstScaleDistribution->Constant = 1.0f;
@@ -264,7 +264,7 @@ UParticleModuleSpawnPerUnit::UParticleModuleSpawnPerUnit(const FObjectInitialize
 
 void UParticleModuleSpawnPerUnit::InitializeDefaults()
 {
-	if (!SpawnPerUnit.Distribution)
+	if (!SpawnPerUnit.IsCreated())
 	{
 		UDistributionFloatConstant* RequiredDistributionSpawnPerUnit = NewObject<UDistributionFloatConstant>(this, TEXT("RequiredDistributionSpawnPerUnit"));
 		RequiredDistributionSpawnPerUnit->Constant = 0.0f;
diff --git a/Engine/Source/Runtime/Engine/Private/Particles/ParticleModules_VectorField.cpp b/Engine/Source/Runtime/Engine/Private/Particles/ParticleModules_VectorField.cpp
--- a/Engine/Source/Runtime/Engine/Private/Particles/ParticleModules_VectorField.cpp
+++ b/Engine/Source/Runtime/Engine/Private/Particles/ParticleModules_VectorField.cpp
@@ -45,11 +45,11 @@ UParticleModuleVectorFieldScale::UParticleModuleVectorFieldScale(const FObjectIn
 
 void UParticleModuleVectorFieldScale::InitializeDefaults()
 {
-	if (!VectorFieldScale)
+	if (!VectorFieldScaleRaw.IsCreated())
 	{
 		UDistributionFloatConstant* DistributionVectorFieldScale = NewObject<UDistributionFloatConstant>(this, TEXT("DistributionVectorFieldScale"));
 		DistributionVectorFieldScale->Constant = 1.0f;
-		VectorFieldScale = DistributionVectorFieldScale;
+		VectorFieldScaleRaw.Distribution = DistributionVectorFieldScale;
 	}
 }
 void UParticleModuleVectorFieldScale::PostInitProperties()
@@ -61,6 +61,19 @@ void UParticleModuleVectorFieldScale::PostInitProperties()
 	}
 }
 
+void UParticleModuleVectorFieldScale::PostLoad()
+{
+	Super::PostLoad();
+#if WITH_EDITOR
+	if (VectorFieldScale_DEPRECATED)
+	{
+		VectorFieldScaleRaw.Distribution = VectorFieldScale_DEPRECATED;
+		VectorFieldScaleRaw.Initialize();
+	}
+	VectorFieldScale_DEPRECATED = nullptr;
+#endif
+}
+
 #if WITH_EDITOR
 void UParticleModuleVectorFieldScale::PostEditChangeProperty(FPropertyChangedEvent& PropertyChangedEvent)
 {
@@ -71,7 +84,7 @@ void UParticleModuleVectorFieldScale::PostEditChangeProperty(FPropertyChangedEve
 
 void UParticleModuleVectorFieldScale::CompileModule(FParticleEmitterBuildInfo& EmitterInfo)
 {
-	EmitterInfo.VectorFieldScale.ScaleByDistribution(VectorFieldScale);
+	EmitterInfo.VectorFieldScale.ScaleByDistribution(VectorFieldScaleRaw.Distribution);
 }
 
 #if WITH_EDITOR
@@ -80,7 +93,7 @@ bool UParticleModuleVectorFieldScale::IsValidForLODLevel(UParticleLODLevel* LODL
 {
 	if (LODLevel->TypeDataModule && LODLevel->TypeDataModule->IsA(UParticleModuleTypeDataGpu::StaticClass()))
 	{
-		if(!IsDistributionAllowedOnGPU(VectorFieldScale))
+		if(!IsDistributionAllowedOnGPU(VectorFieldScaleRaw.Distribution))
 		{
 			OutErrorString = GetDistributionNotAllowedOnGPUText(StaticClass()->GetName(), "VectorFieldScale" ).ToString();
 			return false;
@@ -101,11 +114,11 @@ UParticleModuleVectorFieldScaleOverLife::UParticleModuleVectorFieldScaleOverLife
 
 void UParticleModuleVectorFieldScaleOverLife::InitializeDefaults()
 {
-	if (!VectorFieldScaleOverLife)
+	if (!VectorFieldScaleOverLifeRaw.IsCreated())
 	{
 		UDistributionFloatConstant* DistributionVectorFieldScaleOverLife = NewObject<UDistributionFloatConstant>(this, TEXT("DistributionVectorFieldScaleOverLife"));
 		DistributionVectorFieldScaleOverLife->Constant = 1.0f;
-		VectorFieldScaleOverLife = DistributionVectorFieldScaleOverLife;
+		VectorFieldScaleOverLifeRaw.Distribution = DistributionVectorFieldScaleOverLife;
 	}
 }
 
@@ -118,6 +131,19 @@ void UParticleModuleVectorFieldScaleOverLife::PostInitProperties()
 	}
 }
 
+void UParticleModuleVectorFieldScaleOverLife::PostLoad()
+{
+	Super::PostLoad();
+#if WITH_EDITOR
+	if (VectorFieldScaleOverLife_DEPRECATED)
+	{
+		VectorFieldScaleOverLifeRaw.Distribution = VectorFieldScaleOverLife_DEPRECATED;
+		VectorFieldScaleOverLifeRaw.Initialize();
+	}
+	VectorFieldScaleOverLife_DEPRECATED = nullptr;
+#endif
+}
+
 #if WITH_EDITOR
 void UParticleModuleVectorFieldScaleOverLife::PostEditChangeProperty(FPropertyChangedEvent& PropertyChangedEvent)
 {
@@ -128,7 +154,7 @@ void UParticleModuleVectorFieldScaleOverLife::PostEditChangeProperty(FPropertyCh
 
 void UParticleModuleVectorFieldScaleOverLife::CompileModule(FParticleEmitterBuildInfo& EmitterInfo)
 {
-	EmitterInfo.VectorFieldScaleOverLife.ScaleByDistribution(VectorFieldScaleOverLife);
+	EmitterInfo.VectorFieldScaleOverLife.ScaleByDistribution(VectorFieldScaleOverLifeRaw.Distribution);
 }
 
 #if WITH_EDITOR
@@ -137,7 +163,7 @@ bool UParticleModuleVectorFieldScaleOverLife::IsValidForLODLevel(UParticleLODLev
 {
 	if (LODLevel->TypeDataModule && LODLevel->TypeDataModule->IsA(UParticleModuleTypeDataGpu::StaticClass()))
 	{
-		if(!IsDistributionAllowedOnGPU(VectorFieldScaleOverLife))
+		if(!IsDistributionAllowedOnGPU(VectorFieldScaleOverLifeRaw.Distribution))
 		{
 			OutErrorString = GetDistributionNotAllowedOnGPUText(StaticClass()->GetName(), "VectorFieldScaleOverLife" ).ToString();
 			return false;
diff --git a/Engine/Source/Runtime/Engine/Private/Particles/ParticleModules_Velocity.cpp b/Engine/Source/Runtime/Engine/Private/Particles/ParticleModules_Velocity.cpp
--- a/Engine/Source/Runtime/Engine/Private/Particles/ParticleModules_Velocity.cpp
+++ b/Engine/Source/Runtime/Engine/Private/Particles/ParticleModules_Velocity.cpp
@@ -44,12 +44,12 @@ UParticleModuleVelocity::UParticleModuleVelocity(const FObjectInitializer& Objec
 
 void UParticleModuleVelocity::InitializeDefaults()
 {
-	if (!StartVelocity.Distribution)
+	if (!StartVelocity.IsCreated())
 	{
 		StartVelocity.Distribution = NewObject<UDistributionVectorUniform>(this, TEXT("DistributionStartVelocity"));
 	}
 
-	if (!StartVelocityRadial.Distribution)
+	if (!StartVelocityRadial.IsCreated())
 	{
 		StartVelocityRadial.Distribution = NewObject<UDistributionFloatUniform>(this, TEXT("DistributionStartVelocityRadial"));
 	}
@@ -161,7 +161,7 @@ UParticleModuleVelocityInheritParent::UParticleModuleVelocityInheritParent(const
 
 void UParticleModuleVelocityInheritParent::InitializeDefaults()
 {
-	if (!Scale.Distribution)
+	if (!Scale.IsCreated())
 	{
 		UDistributionVectorConstant* DistributionScale = NewObject<UDistributionVectorConstant>(this, TEXT("DistributionScale"));
 		DistributionScale->Constant = FVector(1.0f, 1.0f, 1.0f);
@@ -225,7 +225,7 @@ UParticleModuleVelocityOverLifetime::UParticleModuleVelocityOverLifetime(const F
 
 void UParticleModuleVelocityOverLifetime::InitializeDefaults()
 {
-	if (!VelOverLife.Distribution)
+	if (!VelOverLife.IsCreated())
 	{
 		VelOverLife.Distribution = NewObject<UDistributionVectorConstantCurve>(this, TEXT("DistributionVelOverLife"));
 	}
@@ -386,12 +386,12 @@ UParticleModuleVelocityCone::UParticleModuleVelocityCone(const FObjectInitialize
 
 void UParticleModuleVelocityCone::InitializeDefaults()
 {
-	if (!Angle.Distribution)
+	if (!Angle.IsCreated())
 	{
 		Angle.Distribution = NewObject<UDistributionFloatUniform>(this, TEXT("DistributionAngle"));
 	}
 
-	if (!Velocity.Distribution)
+	if (!Velocity.IsCreated())
 	{
 		Velocity.Distribution = NewObject<UDistributionFloatUniform>(this, TEXT("DistributionVelocity"));
 	}
diff --git a/Engine/Source/Runtime/Engine/Private/Particles/ParticleSystemRender.cpp b/Engine/Source/Runtime/Engine/Private/Particles/ParticleSystemRender.cpp
--- a/Engine/Source/Runtime/Engine/Private/Particles/ParticleSystemRender.cpp
+++ b/Engine/Source/Runtime/Engine/Private/Particles/ParticleSystemRender.cpp
@@ -1374,7 +1374,7 @@ void FDynamicMeshEmitterData::Init( bool bInSelected,
 
 
 		FVector Mins, Maxs;
-		MeshTD->RollPitchYawRange.Distribution->GetRange(Mins, Maxs);
+		MeshTD->RollPitchYawRange.GetRange(Mins, Maxs);
 
 		// Enable/Disable pre-rotation
 		if (Mins.SizeSquared() || Maxs.SizeSquared())
diff --git a/Engine/Source/Runtime/Engine/Private/Particles/ParticleTrailModules.cpp b/Engine/Source/Runtime/Engine/Private/Particles/ParticleTrailModules.cpp
--- a/Engine/Source/Runtime/Engine/Private/Particles/ParticleTrailModules.cpp
+++ b/Engine/Source/Runtime/Engine/Private/Particles/ParticleTrailModules.cpp
@@ -48,7 +48,7 @@ UParticleModuleTrailSource::UParticleModuleTrailSource(const FObjectInitializer&
 
 void UParticleModuleTrailSource::InitializeDefaults()
 {
-	if (!SourceStrength.Distribution)
+	if (!SourceStrength.IsCreated())
 	{
 		UDistributionFloatConstant* DistributionSourceStrength = NewObject<UDistributionFloatConstant>(this, TEXT("DistributionSourceStrength"));
 		DistributionSourceStrength->Constant = 100.0f;
diff --git a/Engine/Source/Runtime/Engine/Private/SoundNodeWavePlayer.cpp b/Engine/Source/Runtime/Engine/Private/SoundNodeWavePlayer.cpp
--- a/Engine/Source/Runtime/Engine/Private/SoundNodeWavePlayer.cpp
+++ b/Engine/Source/Runtime/Engine/Private/SoundNodeWavePlayer.cpp
@@ -58,9 +58,12 @@ void USoundNodeWavePlayer::ParseNodes( FAudioDevice* AudioDevice, const UPTRINT
 {
 	if (bAsyncLoading)
 	{
-		SoundWave = SoundWaveAssetPtr.LoadSynchronous();
-		bAsyncLoading = false;
+		UE_LOG(LogAudio, Warning, TEXT("Asynchronous load of %s not complete in USoundNodeWavePlayer::ParseNodes, will attempt to play later."), *GetFullNameSafe(this));
+		// We're still loading so don't stop this active sound yet
+		ActiveSound.bFinished = false;
+		return;
 	}
+
 	if (SoundWave)
 	{
 		// The SoundWave's bLooping is only for if it is directly referenced, so clear it
diff --git a/Engine/Source/Runtime/Engine/Private/UserInterface/PlayerInput.cpp b/Engine/Source/Runtime/Engine/Private/UserInterface/PlayerInput.cpp
--- a/Engine/Source/Runtime/Engine/Private/UserInterface/PlayerInput.cpp
+++ b/Engine/Source/Runtime/Engine/Private/UserInterface/PlayerInput.cpp
@@ -1276,11 +1276,6 @@ float UPlayerInput::SmoothMouse(float aMouse, float DeltaTime, uint8& SampleCoun
 		}
 	}
 
-	const float DetectedMouseSampleHz = MouseSamples / MouseSamplingTotal;
-
-	// actual num samples should be either this or this+1
-	const int32 ExpectedNumSamplesThisFrame = FMath::TruncToInt(DeltaTime * DetectedMouseSampleHz);
-
 	if (DeltaTime < 0.25f)
 	{
 		// this is seconds/sample
@@ -1821,7 +1816,7 @@ FString UPlayerInput::GetBind(FKey Key)
 		for ( int32 BindIndex = DebugExecBindings.Num() - 1; BindIndex >= 0; BindIndex-- )
 		{
 			const FKeyBind&	Bind = DebugExecBindings[BindIndex];
-			if ( Bind.Key == Key )
+			if ( Bind.Key == Key && !Bind.bDisabled )
 			{
 				// if the modifier key pressed [or this key-bind doesn't require that key], and the key-bind isn't
 				// configured to ignore they modifier key, we've found a match.
diff --git a/Engine/Source/Runtime/Engine/Private/VisualLogger/VisualLogger.cpp b/Engine/Source/Runtime/Engine/Private/VisualLogger/VisualLogger.cpp
--- a/Engine/Source/Runtime/Engine/Private/VisualLogger/VisualLogger.cpp
+++ b/Engine/Source/Runtime/Engine/Private/VisualLogger/VisualLogger.cpp
@@ -286,7 +286,7 @@ void FVisualLogger::EventLog(const UObject* Object, const FName EventTag1, const
 	CurrentEntry->Events[Index].EventTags.Remove(NAME_None);
 }
 
-void FVisualLogger::NavigationDataDump(const UObject* Object, const FLogCategoryBase& Category, ELogVerbosity::Type Verbosity, int32 UniqueLogId, const FBox& Box)
+void FVisualLogger::NavigationDataDump(const UObject* Object, const FLogCategoryBase& Category, const ELogVerbosity::Type Verbosity, const FBox& Box)
 {
 	SCOPE_CYCLE_COUNTER(STAT_VisualLog);
 	UWorld* World = nullptr;
@@ -458,7 +458,7 @@ void FVisualLogger::SetIsRecording(bool InIsRecording)
 { 
 	if (InIsRecording == false && InIsRecording != !!bIsRecording && FParse::Param(FCommandLine::Get(), TEXT("LogNavOctree")))
 	{
-		FVisualLogger::NavigationDataDump(GetWorldForVisualLogger(nullptr), LogNavigation, ELogVerbosity::Log, INDEX_NONE, FBox());
+		FVisualLogger::NavigationDataDump(GetWorldForVisualLogger(nullptr), LogNavigation, ELogVerbosity::Log, FBox());
 	}
 	if (IsRecordingToFile())
 	{
@@ -577,7 +577,7 @@ public:
 #if ENABLE_VISUAL_LOG
 		else if (FParse::Command(&Cmd, TEXT("LogNavOctree")))
 		{
-			FVisualLogger::NavigationDataDump(GetWorldForVisualLogger(nullptr), LogNavigation, ELogVerbosity::Log, INDEX_NONE, FBox());
+			FVisualLogger::NavigationDataDump(GetWorldForVisualLogger(nullptr), LogNavigation, ELogVerbosity::Log, FBox());
 		}
 #endif
 		return false;
diff --git a/Engine/Source/Runtime/Engine/Private/VisualLogger/VisualLoggerKismetLibrary.cpp b/Engine/Source/Runtime/Engine/Private/VisualLogger/VisualLoggerKismetLibrary.cpp
--- a/Engine/Source/Runtime/Engine/Private/VisualLogger/VisualLoggerKismetLibrary.cpp
+++ b/Engine/Source/Runtime/Engine/Private/VisualLogger/VisualLoggerKismetLibrary.cpp
@@ -14,7 +14,7 @@ void UVisualLoggerKismetLibrary::LogText(UObject* WorldContextObject, FString Te
 {
 #if ENABLE_VISUAL_LOG
 	const ELogVerbosity::Type DefaultVerbosity = ELogVerbosity::Log;
-	FVisualLogger::CategorizedLogf(WorldContextObject, FLogCategoryBase(*CategoryName.ToString(), DefaultVerbosity, DefaultVerbosity), DefaultVerbosity, INDEX_NONE, *Text);
+	FVisualLogger::CategorizedLogf(WorldContextObject, FLogCategoryBase(*CategoryName.ToString(), DefaultVerbosity, DefaultVerbosity), DefaultVerbosity, *Text);
 #endif
 }
 
@@ -22,7 +22,7 @@ void UVisualLoggerKismetLibrary::LogLocation(UObject* WorldContextObject, FVecto
 {
 #if ENABLE_VISUAL_LOG
 	const ELogVerbosity::Type DefaultVerbosity = ELogVerbosity::Log;
-	FVisualLogger::GeometryShapeLogf(WorldContextObject, FLogCategoryBase(*CategoryName.ToString(), DefaultVerbosity, DefaultVerbosity), DefaultVerbosity, INDEX_NONE, Location, Radius, Color.ToFColor(true), *Text);
+	FVisualLogger::GeometryShapeLogf(WorldContextObject, FLogCategoryBase(*CategoryName.ToString(), DefaultVerbosity, DefaultVerbosity), DefaultVerbosity, Location, Radius, Color.ToFColor(true), *Text);
 #endif
 }
 
@@ -30,6 +30,6 @@ void UVisualLoggerKismetLibrary::LogBox(UObject* WorldContextObject, FBox Box, F
 {
 #if ENABLE_VISUAL_LOG
 	const ELogVerbosity::Type DefaultVerbosity = ELogVerbosity::Log;
-	FVisualLogger::GeometryShapeLogf(WorldContextObject, FLogCategoryBase(*CategoryName.ToString(), DefaultVerbosity, DefaultVerbosity), DefaultVerbosity, INDEX_NONE, Box, FMatrix::Identity, ObjectColor.ToFColor(true), *Text);
+	FVisualLogger::GeometryShapeLogf(WorldContextObject, FLogCategoryBase(*CategoryName.ToString(), DefaultVerbosity, DefaultVerbosity), DefaultVerbosity, Box, FMatrix::Identity, ObjectColor.ToFColor(true), *Text);
 #endif
 }
diff --git a/Engine/Source/Runtime/Engine/Private/VisualLogger/VisualLoggerTypes.cpp b/Engine/Source/Runtime/Engine/Private/VisualLogger/VisualLoggerTypes.cpp
--- a/Engine/Source/Runtime/Engine/Private/VisualLogger/VisualLoggerTypes.cpp
+++ b/Engine/Source/Runtime/Engine/Private/VisualLogger/VisualLoggerTypes.cpp
@@ -2,9 +2,15 @@
 
 #include "EnginePrivate.h"
 #include "ObjectBase.h"
+#include "VisualLogger/VisualLoggerTypes.h"
 #include "VisualLogger/VisualLogger.h"
 #include "VisualLogger/VisualLoggerDebugSnapshotInterface.h"
 
+namespace
+{
+	const FName NAME_UnnamedCategory = TEXT("UnnamedCategry");
+}
+
 UVisualLoggerDebugSnapshotInterface::UVisualLoggerDebugSnapshotInterface(const FObjectInitializer& ObjectInitializer)
 	: Super(ObjectInitializer)
 {
@@ -15,6 +21,22 @@ UVisualLoggerDebugSnapshotInterface::UVisualLoggerDebugSnapshotInterface(const F
 #define DEPRECATED_VISUAL_LOGGER_MAGIC_NUMBER 0xFAFAAFAF
 #define VISUAL_LOGGER_MAGIC_NUMBER 0xAFAFFAFA
 
+FVisualLogShapeElement::FVisualLogShapeElement(EVisualLoggerShapeElement InType)
+	: Category(NAME_UnnamedCategory)
+	, Verbosity(ELogVerbosity::All)
+	, TransformationMatrix(FMatrix::Identity)
+	, Type(InType)
+	, Color(0xff)
+	, Thicknes(0)
+{
+
+}
+
+FVisualLogStatusCategory::FVisualLogStatusCategory()
+	: Category(NAME_UnnamedCategory.ToString())
+{
+}
+
 bool FVisualLogStatusCategory::GetDesc(int32 Index, FString& Key, FString& Value) const
 {
 	if (Data.IsValidIndex(Index))
diff --git a/Engine/Source/Runtime/Engine/Private/World.cpp b/Engine/Source/Runtime/Engine/Private/World.cpp
--- a/Engine/Source/Runtime/Engine/Private/World.cpp
+++ b/Engine/Source/Runtime/Engine/Private/World.cpp
@@ -853,17 +853,26 @@ UAISystemBase* UWorld::CreateAISystem()
 
 void UWorld::RepairWorldSettings()
 {
+	AWorldSettings* ExistingWorldSettings = PersistentLevel->GetWorldSettings(false);
+
+	if (ExistingWorldSettings == nullptr && PersistentLevel->Actors.Num() > 0)
+	{
+		ExistingWorldSettings = Cast<AWorldSettings>(PersistentLevel->Actors[0]);
+		if (ExistingWorldSettings)
+		{
+			// This means the WorldSettings member just wasn't initialized, get that resolved
+			PersistentLevel->SetWorldSettings(ExistingWorldSettings);
+		}
+	}
+
 	// If for some reason we don't have a valid WorldSettings object go ahead and spawn one to avoid crashing.
 	// This will generally happen if a map is being moved from a different project.
-	const bool bNeedsExchange = PersistentLevel->Actors.Num() > 0;
-	const bool bNeedsDestroy = bNeedsExchange && PersistentLevel->Actors[0] != NULL;
-
-	if (PersistentLevel->Actors.Num() < 1 || PersistentLevel->Actors[0] == NULL || !PersistentLevel->Actors[0]->IsA(GEngine->WorldSettingsClass))
+	if (ExistingWorldSettings == nullptr || !ExistingWorldSettings->IsA(GEngine->WorldSettingsClass))
 	{
 		// Rename invalid WorldSettings to avoid name collisions
-		if (bNeedsDestroy)
+		if (ExistingWorldSettings)
 		{
-			PersistentLevel->Actors[0]->Rename(NULL, PersistentLevel, REN_ForceNoResetLoaders);
+			ExistingWorldSettings->Rename(nullptr, PersistentLevel, REN_ForceNoResetLoaders);
 		}
 		
 		bool bClearOwningWorld = false;
@@ -879,23 +888,16 @@ void UWorld::RepairWorldSettings()
 		SpawnInfo.Name = GEngine->WorldSettingsClass->GetFName();
 		AWorldSettings* const NewWorldSettings = SpawnActor<AWorldSettings>( GEngine->WorldSettingsClass, SpawnInfo );
 
-		const int32 NewWorldSettingsActorIndex = PersistentLevel->Actors.Find( NewWorldSettings );
-
-		if (bNeedsExchange)
-		{
-			// The world info must reside at index 0.
-			Exchange(PersistentLevel->Actors[0],PersistentLevel->Actors[NewWorldSettingsActorIndex]);
-		}
-
-		// If there was an existing actor, copy its properties to the new actor and then destroy it
-		if (bNeedsDestroy)
+		// If there was an existing actor, copy its properties to the new actor (the it will be destroyed by SetWorldSettings)
+		if (ExistingWorldSettings)
 		{
 			NewWorldSettings->UnregisterAllComponents();
-			UEngine::CopyPropertiesForUnrelatedObjects(PersistentLevel->Actors[NewWorldSettingsActorIndex], NewWorldSettings);
+			UEngine::CopyPropertiesForUnrelatedObjects(ExistingWorldSettings, NewWorldSettings);
 			NewWorldSettings->RegisterAllComponents();
-			PersistentLevel->Actors[NewWorldSettingsActorIndex]->Destroy();
 		}
 
+		PersistentLevel->SetWorldSettings(NewWorldSettings);
+
 		// Re-sort actor list as we just shuffled things around.
 		PersistentLevel->SortActorList();
 
@@ -1120,7 +1122,8 @@ void UWorld::InitializeNewWorld(const InitializationValues IVS)
 	SpawnInfo.SpawnCollisionHandlingOverride = ESpawnActorCollisionHandlingMethod::AlwaysSpawn;
 	// Set constant name for WorldSettings to make a network replication work between new worlds on host and client
 	SpawnInfo.Name = GEngine->WorldSettingsClass->GetFName();
-	AActor* WorldSettings = SpawnActor( GEngine->WorldSettingsClass, NULL, NULL, SpawnInfo );
+	AWorldSettings* WorldSettings = SpawnActor<AWorldSettings>( GEngine->WorldSettingsClass, SpawnInfo );
+	PersistentLevel->SetWorldSettings(WorldSettings);
 	check(GetWorldSettings());
 #if WITH_EDITOR
 	WorldSettings->SetIsTemporarilyHiddenInEditor(true);
@@ -3594,30 +3597,19 @@ ALevelScriptActor* UWorld::GetLevelScriptActor( ULevel* OwnerLevel ) const
 AWorldSettings* UWorld::GetWorldSettings( bool bCheckStreamingPesistent, bool bChecked ) const
 {
 	checkSlow(IsInGameThread());
-	AWorldSettings* WorldSettings = NULL;
+	AWorldSettings* WorldSettings = nullptr;
 	if (PersistentLevel)
 	{
-		if (bChecked)
-		{
-			checkSlow(PersistentLevel->Actors.Num());
-			checkSlow(PersistentLevel->Actors[0]);
-			checkSlow(PersistentLevel->Actors[0]->IsA(AWorldSettings::StaticClass()));
-
-			WorldSettings = (AWorldSettings*)PersistentLevel->Actors[0];
-		}
-		else if (PersistentLevel->Actors.Num() > 0)
-		{
-			WorldSettings = Cast<AWorldSettings>(PersistentLevel->Actors[0]);
-		}
+		WorldSettings = PersistentLevel->GetWorldSettings(bChecked);
 
 		if( bCheckStreamingPesistent )
 		{
 			if( StreamingLevels.Num() > 0 &&
 				StreamingLevels[0] &&
-				StreamingLevels[0]->IsA(ULevelStreamingPersistent::StaticClass()) )
+				StreamingLevels[0]->IsA<ULevelStreamingPersistent>()) 
 			{
 				ULevel* Level = StreamingLevels[0]->GetLoadedLevel();
-				if (Level != NULL)
+				if (Level != nullptr)
 				{
 					WorldSettings = Level->GetWorldSettings();
 				}
diff --git a/Engine/Source/Runtime/Engine/Public/Audio.h b/Engine/Source/Runtime/Engine/Public/Audio.h
--- a/Engine/Source/Runtime/Engine/Public/Audio.h
+++ b/Engine/Source/Runtime/Engine/Public/Audio.h
@@ -349,6 +349,11 @@ public:
 	 */
 	FString GetChannelsDesc();
 
+	/**
+	 * Reads the compressed info of the given sound wave. Not implemented on all platforms.
+	 */
+	virtual bool ReadCompressedInfo(USoundWave* SoundWave) { return true; }
+
 	/**
 	 * Gets the chunk index that was last read from (for Streaming Manager requests)
 	 */
@@ -359,6 +364,9 @@ public:
 	 */
 	virtual int32 GetCurrentChunkOffset() const {return -1;}
 
+	/** Returns whether or not a real-time decoding buffer is ready for playback */
+	virtual bool IsRealTimeSourceReady() { return true; }
+
 	/** Unique ID that ties this buffer to a USoundWave */
 	int32	ResourceID;
 	/** Cumulative channels from all streams */
@@ -400,6 +408,7 @@ public:
 		, Buffer(NULL)
 		, Playing(false)
 		, Paused(false)
+		, bInitialized(true) // Note: this is defaulted to true since not all platforms need to deal with async initialization.
 		, bReverbApplied(false)
 		, StereoBleed(0.0f)
 		, LFEBleed(0.5f)
@@ -416,8 +425,10 @@ public:
 	}
 
 	// Initialization & update.
-	virtual bool Init( FWaveInstance* WaveInstance ) = 0;
-	virtual void Update( void ) = 0;
+	virtual bool PrepareForInitialization(FWaveInstance* InWaveInstance) { return true; }
+	virtual bool IsPreparedToInit() { return true; }
+	virtual bool Init(FWaveInstance* InWaveInstance) = 0;
+	virtual void Update(void) = 0;
 
 	// Playback.
 	virtual void Play( void ) = 0;
@@ -427,6 +438,9 @@ public:
 	// Query.
 	virtual	bool IsFinished( void ) = 0;
 
+	/** Returns whether or not the sound source has initialized */
+	bool IsInitialized(void) const { return bInitialized; };
+
 	/**
 	 * Returns a string describing the source (subclass can override, but it should call the base and append)
 	 */
@@ -544,9 +558,10 @@ protected:
 	uint32				Playing:1;
 	/** Cached status information whether we are paused or not. */
 	uint32				Paused:1;
+	/** Whether or not the sound source is ready to be initialized */
+	uint32				bInitialized:1;
 	/** Cached sound mode value used to detect when to switch outputs. */
 	uint32				bReverbApplied:1;
-
 	/** The amount of stereo sounds to bleed to the rear speakers */
 	float				StereoBleed;
 	/** The amount of a sound to bleed to the LFE speaker */
diff --git a/Engine/Source/Runtime/Engine/Public/AudioDecompress.h b/Engine/Source/Runtime/Engine/Public/AudioDecompress.h
--- a/Engine/Source/Runtime/Engine/Public/AudioDecompress.h
+++ b/Engine/Source/Runtime/Engine/Public/AudioDecompress.h
@@ -150,7 +150,13 @@ typedef FAsyncTask<FAsyncAudioDecompressWorker> FAsyncAudioDecompress;
 
 enum class ERealtimeAudioTaskType : uint8
 {
+	/** Parses the wave compressed asset header file info */
+	CompressedInfo,
+
+	/** Decompresses a chunk */
 	Decompress,
+
+	/** Processes a procedural buffer to generate more audio */
 	Procedural
 };
 
@@ -169,6 +175,21 @@ protected:
 	uint32 bLooped:1;
 
 public:
+	FAsyncRealtimeAudioTaskWorker(T* InAudioBuffer, USoundWave* InWaveData)
+		: AudioBuffer(InAudioBuffer)
+		, WaveData(InWaveData)
+		, AudioData(nullptr)
+		, MaxSamples(0)
+		, BytesWritten(0)
+		, TaskType(ERealtimeAudioTaskType::CompressedInfo)
+		, bSkipFirstBuffer(false)
+		, bLoopingMode(false)
+		, bLooped(false)
+	{
+		check(AudioBuffer && AudioBuffer->DecompressionState);
+		check(WaveData);
+	}
+
 	FAsyncRealtimeAudioTaskWorker(T* InAudioBuffer, uint8* InAudioData, bool bInLoopingMode, bool bInSkipFirstBuffer)
 		: AudioBuffer(InAudioBuffer)
 		, AudioData(InAudioData)
@@ -177,6 +198,8 @@ public:
 		, bLoopingMode(bInLoopingMode)
 		, bLooped(false)
 	{
+		check(AudioBuffer && AudioBuffer->DecompressionState);
+		check(AudioData);
 	}
 
 	FAsyncRealtimeAudioTaskWorker(USoundWave* InWaveData, uint8* InAudioData, int32 InMaxSamples)
@@ -186,12 +209,18 @@ public:
 		, BytesWritten(0)
 		, TaskType(ERealtimeAudioTaskType::Procedural)
 	{
+		check(WaveData);
+		check(AudioData);
 	}
 
 	void DoWork()
 	{
 		switch(TaskType)
 		{
+		case ERealtimeAudioTaskType::CompressedInfo:
+			AudioBuffer->ReadCompressedInfo(WaveData);
+			break;
+
 		case ERealtimeAudioTaskType::Decompress:
 			if (bSkipFirstBuffer)
 			{
diff --git a/Engine/Source/Runtime/Engine/Public/AudioDevice.h b/Engine/Source/Runtime/Engine/Public/AudioDevice.h
--- a/Engine/Source/Runtime/Engine/Public/AudioDevice.h
+++ b/Engine/Source/Runtime/Engine/Public/AudioDevice.h
@@ -3,8 +3,18 @@
 #pragma once 
 
 #include "Sound/AudioVolume.h"
+#include "Sound/SoundMix.h"
 #include "AudioDeviceManager.h"
 
+/**
+ * Forward declares
+ */
+
+class FAudioEffectsManager;
+class ICompressedAudioInfo;
+class IAudioSpatializationPlugin;
+class IAudioSpatializationAlgorithm;
+
 /** 
  * Debug state of the audio system
  */
@@ -70,7 +80,7 @@ struct FListener
 	struct FInteriorSettings InteriorSettings;
 
 	/** The volume the listener resides in */
-	class AAudioVolume* Volume;
+	AAudioVolume* Volume;
 
 	/** The times of interior volumes fading in and out */
 	double InteriorStartTime;
@@ -90,7 +100,7 @@ struct FListener
 	/**
 	 * Works out the interp value between source and end
 	 */
-	float Interpolate( const double EndTime );
+	float Interpolate(const double EndTime);
 
 	/**
 	 * Gets the current state of the interior settings for the listener
@@ -100,12 +110,12 @@ struct FListener
 	/** 
 	 * Apply the interior settings to ambient sounds
 	 */
-	void ApplyInteriorSettings( class AAudioVolume* Volume, const FInteriorSettings& Settings );
+	void ApplyInteriorSettings(AAudioVolume* Volume, const FInteriorSettings& Settings);
 
 	FListener()
 		: Transform(FTransform::Identity)
 		, Velocity(ForceInit)
-		, Volume(NULL)
+		, Volume(nullptr)
 		, InteriorStartTime(0.0)
 		, InteriorEndTime(0.0)
 		, ExteriorEndTime(0.0)
@@ -152,6 +162,27 @@ struct FSoundMixState
 	ESoundMixState::Type CurrentState;
 };
 
+struct FSoundMixClassOverride
+{
+	FSoundClassAdjuster SoundClassAdjustor;
+	FDynamicParameter VolumeOverride;
+	FDynamicParameter PitchOverride;
+	float FadeInTime;
+	uint8 bOverrideApplied : 1;
+	uint8 bOverrideChanged : 1;
+
+	FSoundMixClassOverride()
+		: VolumeOverride(1.0f)
+		, PitchOverride(1.0f)
+		, FadeInTime(0.0f)
+		, bOverrideApplied(false)
+		, bOverrideChanged(false)
+	{
+	}
+};
+
+typedef TMap<USoundClass*, FSoundMixClassOverride> FSoundMixClassOverrideMap;
+
 struct FActivatedReverb
 {
 	FReverbSettings ReverbSettings;
@@ -227,51 +258,51 @@ class ENGINE_API FAudioDevice : public FExec
 public:
 
 	//Begin FExec Interface
-	virtual bool Exec( UWorld* InWorld, const TCHAR* Cmd, FOutputDevice& Ar = *GLog ) override;
+	virtual bool Exec(UWorld* InWorld, const TCHAR* Cmd, FOutputDevice& Ar = *GLog) override;
 	//End FExec Interface
 
 #if !UE_BUILD_SHIPPING
 	/**
 	 * Exec command handlers
 	 */
-	bool HandleDumpSoundInfoCommand( const TCHAR* Cmd, FOutputDevice& Ar );
+	bool HandleDumpSoundInfoCommand(const TCHAR* Cmd, FOutputDevice& Ar);
 	/**
 	 * Lists all the loaded sounds and their memory footprint
 	 */
-	bool HandleListSoundsCommand( const TCHAR* Cmd, FOutputDevice& Ar );
+	bool HandleListSoundsCommand(const TCHAR* Cmd, FOutputDevice& Ar);
 	/**
 	 * Lists all the playing waveinstances and their associated source
 	 */
-	bool HandleListWavesCommand( const TCHAR* Cmd, FOutputDevice& Ar );
+	bool HandleListWavesCommand(const TCHAR* Cmd, FOutputDevice& Ar);
 	/**
 	 * Lists a summary of loaded sound collated by class
 	 */
-	bool HandleListSoundClassesCommand( const TCHAR* Cmd, FOutputDevice& Ar );
+	bool HandleListSoundClassesCommand(const TCHAR* Cmd, FOutputDevice& Ar);
 	/**
 	 * shows sound class hierarchy
 	 */
-	bool HandleShowSoundClassHierarchyCommand( const TCHAR* Cmd, FOutputDevice& Ar );
-	bool HandleListSoundClassVolumesCommand( const TCHAR* Cmd, FOutputDevice& Ar );
-	bool HandleListAudioComponentsCommand( const TCHAR* Cmd, FOutputDevice& Ar );
-	bool HandleListSoundDurationsCommand( const TCHAR* Cmd, FOutputDevice& Ar );
-	bool HandleSoundTemplateInfoCommand( const TCHAR* Cmd, FOutputDevice& Ar );
-	bool HandlePlaySoundCueCommand( const TCHAR* Cmd, FOutputDevice& Ar );
-	bool HandlePlaySoundWaveCommand( const TCHAR* Cmd, FOutputDevice& Ar );
-	bool HandleSetBaseSoundMixCommand( const TCHAR* Cmd, FOutputDevice& Ar );
-	bool HandleIsolateDryAudioCommand( const TCHAR* Cmd, FOutputDevice& Ar );
-	bool HandleIsolateReverbCommand( const TCHAR* Cmd, FOutputDevice& Ar );
-	bool HandleTestLPFCommand( const TCHAR* Cmd, FOutputDevice& Ar );
-	bool HandleTestStereoBleedCommand( const TCHAR* Cmd, FOutputDevice& Ar );
-	bool HandleTestLFEBleedCommand( const TCHAR* Cmd, FOutputDevice& Ar );
-	bool HandleDisableLPFCommand( const TCHAR* Cmd, FOutputDevice& Ar );
-	bool HandleDisableRadioCommand( const TCHAR* Cmd, FOutputDevice& Ar );
-	bool HandleEnableRadioCommand( const TCHAR* Cmd, FOutputDevice& Ar );
-	bool HandleResetSoundStateCommand( const TCHAR* Cmd, FOutputDevice& Ar );
-	bool HandleToggleSpatializationExtensionCommand( const TCHAR* Cmd, FOutputDevice& Ar );
+	bool HandleShowSoundClassHierarchyCommand(const TCHAR* Cmd, FOutputDevice& Ar);
+	bool HandleListSoundClassVolumesCommand(const TCHAR* Cmd, FOutputDevice& Ar);
+	bool HandleListAudioComponentsCommand(const TCHAR* Cmd, FOutputDevice& Ar);
+	bool HandleListSoundDurationsCommand(const TCHAR* Cmd, FOutputDevice& Ar);
+	bool HandleSoundTemplateInfoCommand(const TCHAR* Cmd, FOutputDevice& Ar);
+	bool HandlePlaySoundCueCommand(const TCHAR* Cmd, FOutputDevice& Ar);
+	bool HandlePlaySoundWaveCommand(const TCHAR* Cmd, FOutputDevice& Ar);
+	bool HandleSetBaseSoundMixCommand(const TCHAR* Cmd, FOutputDevice& Ar);
+	bool HandleIsolateDryAudioCommand(const TCHAR* Cmd, FOutputDevice& Ar);
+	bool HandleIsolateReverbCommand(const TCHAR* Cmd, FOutputDevice& Ar);
+	bool HandleTestLPFCommand(const TCHAR* Cmd, FOutputDevice& Ar);
+	bool HandleTestStereoBleedCommand(const TCHAR* Cmd, FOutputDevice& Ar);
+	bool HandleTestLFEBleedCommand(const TCHAR* Cmd, FOutputDevice& Ar);
+	bool HandleDisableLPFCommand(const TCHAR* Cmd, FOutputDevice& Ar);
+	bool HandleDisableRadioCommand(const TCHAR* Cmd, FOutputDevice& Ar);
+	bool HandleEnableRadioCommand(const TCHAR* Cmd, FOutputDevice& Ar);
+	bool HandleResetSoundStateCommand(const TCHAR* Cmd, FOutputDevice& Ar);
+	bool HandleToggleSpatializationExtensionCommand(const TCHAR* Cmd, FOutputDevice& Ar);
 	bool HandleEnableHRTFForAllCommand(const TCHAR* Cmd, FOutputDevice& Ar);
-	bool HandleSoloCommand( const TCHAR* Cmd, FOutputDevice& Ar );
-	bool HandleClearSoloCommand( const TCHAR* Cmd, FOutputDevice& Ar );
-	bool HandlePlayAllPIEAudioCommand( const TCHAR* Cmd, FOutputDevice& Ar );
+	bool HandleSoloCommand(const TCHAR* Cmd, FOutputDevice& Ar);
+	bool HandleClearSoloCommand(const TCHAR* Cmd, FOutputDevice& Ar);
+	bool HandlePlayAllPIEAudioCommand(const TCHAR* Cmd, FOutputDevice& Ar);
 	bool HandleAudio3dVisualizeCommand(const TCHAR* Cmd, FOutputDevice& Ar);
 	bool HandleAudioMemoryInfo(const TCHAR* Cmd, FOutputDevice& Ar);
 #endif
@@ -330,9 +361,9 @@ public:
 	int32 GetSortedActiveWaveInstances(TArray<FWaveInstance*>& WaveInstances, const ESortedActiveWaveGetType::Type GetType);
 
 	/**
-	 * Stop all the audio components and sources attached to the world. NULL world means all components.
+	 * Stop all the audio components and sources attached to the world. nullptr world means all components.
 	 */
-	void Flush( class UWorld* WorldToFlush, bool bClearActivatedReverb = true );
+	void Flush(UWorld* WorldToFlush, bool bClearActivatedReverb = true);
 
 	/**
 	 * Stop any playing sounds that are using a particular SoundWave
@@ -357,7 +388,7 @@ public:
 	 * @param	bSynchronous	If true, this function will block until a vorbis decompression is complete
 	 * @param	bTrackMemory	If true, the audio mem stats will be updated
 	 */
-	virtual void Precache(USoundWave* SoundWave, bool bSynchronous=false, bool bTrackMemory=true);
+	virtual void Precache(USoundWave* SoundWave, bool bSynchronous = false, bool bTrackMemory=true);
 
 	/**
 	 * Precaches all existing sounds. Called when audio setup is complete
@@ -381,7 +412,7 @@ public:
 	 *
 	 * @param bShouldStopUISounds If true, this function will stop UI sounds as well
 	 */
-	virtual void StopAllSounds( bool bShouldStopUISounds = false );
+	virtual void StopAllSounds(bool bShouldStopUISounds = false);
 
 	/**
 	 * Sets the details about the listener
@@ -392,13 +423,13 @@ public:
 	 * @param	ReverbSettings		The reverb settings for this user to use.
 	 * @param	InteriorSettings	The interior settings for this user to use.
 	 */
-	void SetListener( const int32 InListenerIndex, const FTransform& ListenerTransform, const float InDeltaSeconds, class AAudioVolume* Volume, const FInteriorSettings& InteriorSettings );
+	void SetListener(const int32 InListenerIndex, const FTransform& ListenerTransform, const float InDeltaSeconds, AAudioVolume* Volume, const FInteriorSettings& InteriorSettings);
 
 	/**
 	 * Starts a transition to new reverb and interior settings
 	 *
 	 */
-	void SetReverbSettings( class AAudioVolume* Volume, const FReverbSettings& ReverbSettings );
+	void SetReverbSettings(AAudioVolume* Volume, const FReverbSettings& ReverbSettings);
 
 	/**
 	 * Creates an audio component to handle playing a sound.
@@ -413,7 +444,7 @@ public:
 	 * @param	USoundConcurrency	The sound's sound concurrency settings to use. Will use the USoundBase's USoundConcurrency if not specified.
 	 * @return	The created audio component if the function successfully created one or a nullptr if not successful. Note: if audio is disabled or if there were no hardware audio devices available, this will return nullptr.
 	 */
-	static class UAudioComponent* CreateComponent(class USoundBase* Sound, class UWorld* World, AActor*  AActor = nullptr, bool Play = true, bool bStopWhenOwnerDestroyed = false, const FVector* Location = nullptr, USoundAttenuation* AttenuationSettings = nullptr, USoundConcurrency* ConcurrencySettings = nullptr);
+	static UAudioComponent* CreateComponent(USoundBase* Sound, UWorld* World, AActor*  AActor = nullptr, bool Play = true, bool bStopWhenOwnerDestroyed = false, const FVector* Location = nullptr, USoundAttenuation* AttenuationSettings = nullptr, USoundConcurrency* ConcurrencySettings = nullptr);
 
 	/** 
 	 * Plays a sound at the given location without creating an audio component.
@@ -428,17 +459,17 @@ public:
 	 * @param	USoundConcurrency	The sound's sound concurrency settings to use (optional). Will use the USoundBase's USoundConcurrency if not specified.
 	 * @param	Params				An optional list of audio component params to immediately apply to a sound.
 	 */
-	void PlaySoundAtLocation(class USoundBase* Sound, class UWorld* World, float VolumeMultiplier, float PitchMultiplier, float StartTime, const FVector& Location, const FRotator& Rotation, USoundAttenuation* AttenuationSettings = nullptr, USoundConcurrency* ConcurrencySettings = nullptr, const TArray<FAudioComponentParam>* Params = nullptr);
+	void PlaySoundAtLocation(USoundBase* Sound, UWorld* World, float VolumeMultiplier, float PitchMultiplier, float StartTime, const FVector& Location, const FRotator& Rotation, USoundAttenuation* AttenuationSettings = nullptr, USoundConcurrency* ConcurrencySettings = nullptr, const TArray<FAudioComponentParam>* Params = nullptr);
 
 	/**
 	 * Adds an active sound to the audio device
 	 */
-	void AddNewActiveSound( const FActiveSound& ActiveSound );
+	void AddNewActiveSound(const FActiveSound& ActiveSound);
 
 	/**
 	 * Removes the active sound for the specified audio component
 	 */
-	void StopActiveSound( class UAudioComponent* AudioComponent );
+	void StopActiveSound(UAudioComponent* AudioComponent);
 
 	/**
 	* Stops the active sound
@@ -448,33 +479,33 @@ public:
 	/**
 	 * Finds the active sound for the specified audio component
 	 */
-	FActiveSound* FindActiveSound( class UAudioComponent* AudioComponent );
+	FActiveSound* FindActiveSound(UAudioComponent* AudioComponent);
 
 	/**
 	 * Removes an active sound from the active sounds array
 	 */
-	void RemoveActiveSound( FActiveSound* ActiveSound );
+	void RemoveActiveSound(FActiveSound* ActiveSound);
 
 	/** 
 	 * Gets the current audio debug state
 	 */
-	EDebugState GetMixDebugState( void );
+	EDebugState GetMixDebugState();
 
 	/**
 	 * Set up the sound class hierarchy
 	 */
-	void InitSoundClasses( void );
+	void InitSoundClasses();
 
 	/**
 	 * Set up the initial sound sources
 	 * Allows us to initialize sound source early on, allowing for render callback hookups for iOS Audio.
 	 */
-	void InitSoundSources( void );
+	void InitSoundSources();
 
 	/** 
 	 * Gets a summary of loaded sound collated by class
 	 */
-	void GetSoundClassInfo( TMap<FName, FAudioClassInfo>& AudioClassInfos );
+	void GetSoundClassInfo(TMap<FName, FAudioClassInfo>& AudioClassInfos);
 
 	/**
 	 * Registers a sound class with the audio device
@@ -482,12 +513,12 @@ public:
 	 * @param	SoundClassName	name of sound class to retrieve
 	 * @return	sound class properties if it exists
 	 */
-	void RegisterSoundClass( class USoundClass* InSoundClass );
+	void RegisterSoundClass(USoundClass* InSoundClass);
 
 	/**
 	* Unregisters a sound class
 	*/
-	void UnregisterSoundClass(class USoundClass* SoundClass);
+	void UnregisterSoundClass(USoundClass* SoundClass);
 
 	/**
 	* Gets the current properties of a sound class, if the sound class hasn't been registered, then it returns nullptr
@@ -495,7 +526,7 @@ public:
 	* @param	SoundClassName	name of sound class to retrieve
 	* @return	sound class properties if it exists
 	*/
-	FSoundClassProperties* GetSoundClassCurrentProperties(class USoundClass* InSoundClass);
+	FSoundClassProperties* GetSoundClassCurrentProperties(USoundClass* InSoundClass);
 
 	/**
 	* Checks to see if a coordinate is within a distance of any listener
@@ -510,26 +541,26 @@ public:
 	/**
 	 * Sets the Sound Mix that should be active by default
 	 */
-	void SetDefaultBaseSoundMix( class USoundMix* SoundMix );
+	void SetDefaultBaseSoundMix(USoundMix* SoundMix);
 
 	/**
 	 * Removes a sound mix - called when SoundMix is unloaded
 	 */
-	void RemoveSoundMix( class USoundMix* SoundMix );
+	void RemoveSoundMix(USoundMix* SoundMix);
 
 	/** 
 	 * Resets all interpolating values to defaults.
 	 */
-	void ResetInterpolation( void );
+	void ResetInterpolation();
 
 	/** Enables or Disables the radio effect. */
-	void EnableRadioEffect( bool bEnable = false );
+	void EnableRadioEffect(bool bEnable = false);
 
 	friend class FAudioEffectsManager;
 	/**
 	 * Sets a new sound mix and applies it to all appropriate sound classes
 	 */
-	bool SetBaseSoundMix( class USoundMix* SoundMix );
+	bool SetBaseSoundMix(USoundMix* SoundMix);
 
 	/**
 	 * Push a SoundMix onto the Audio Device's list.
@@ -537,7 +568,12 @@ public:
 	 * @param SoundMix The SoundMix to push.
 	 * @param bIsPassive Whether this is a passive push from a playing sound.
 	 */
-	void PushSoundMixModifier(class USoundMix* SoundMix, bool bIsPassive = false);
+	void PushSoundMixModifier(USoundMix* SoundMix, bool bIsPassive = false);
+
+	/** 
+	 * Sets a sound class override in the given sound mix.
+	 */
+	void SetSoundMixClassOverride(USoundMix* InSoundMix, USoundClass* InSoundClass, float Volume, float Pitch, float FadeInTime, bool bApplyToChildren);
 
 	/**
 	 * Pop a SoundMix from the Audio Device's list.
@@ -545,14 +581,14 @@ public:
 	 * @param SoundMix The SoundMix to pop.
 	 * @param bIsPassive Whether this is a passive pop from a sound finishing.
 	 */
-	void PopSoundMixModifier(class USoundMix* SoundMix, bool bIsPassive = false);
+	void PopSoundMixModifier(USoundMix* SoundMix, bool bIsPassive = false);
 
 	/**
 	 * Clear the effect of one SoundMix completely.
 	 *
 	 * @param SoundMix The SoundMix to clear.
 	 */
-	void ClearSoundMixModifier(class USoundMix* SoundMix);
+	void ClearSoundMixModifier(USoundMix* SoundMix);
 
 	/**
 	 * Clear the effect of all SoundMix modifiers.
@@ -566,7 +602,7 @@ public:
 	 * @param Volume Volume level of Reverb Effect
 	 * @param FadeTime Time before Reverb Effect is fully active
 	 */
-	void ActivateReverbEffect(class UReverbEffect* ReverbEffect, FName TagName, float Priority, float Volume, float FadeTime);
+	void ActivateReverbEffect(UReverbEffect* ReverbEffect, FName TagName, float Priority, float Volume, float FadeTime);
 	
 	/**
 	 * Deactivates a Reverb Effect not applied by a volume
@@ -575,7 +611,7 @@ public:
 	 */
 	void DeactivateReverbEffect(FName TagName);
 
-	virtual FName GetRuntimeFormat(USoundWave* SoundWave) PURE_VIRTUAL(FAudioDevice::GetRuntimeFormat,return NAME_None;);
+	virtual FName GetRuntimeFormat(USoundWave* SoundWave) = 0;
 
 	/** Whether this SoundWave has an associated info class to decompress it */
 	virtual bool HasCompressedAudioInfoClass(USoundWave* SoundWave) { return false; }
@@ -583,18 +619,18 @@ public:
 	/** Whether this device supports realtime decompression of sound waves (i.e. DTYPE_RealTime) */
 	virtual bool SupportsRealtimeDecompression() const
 	{ 
-		return false;	// assume no support by default
+		return false;
 	}
 
 	/** Creates a Compressed audio info class suitable for decompressing this SoundWave */
-	virtual class ICompressedAudioInfo* CreateCompressedAudioInfo(USoundWave* SoundWave) { return NULL; }
+	virtual ICompressedAudioInfo* CreateCompressedAudioInfo(USoundWave* SoundWave) { return nullptr; }
 
 	/**
 	 * Check for errors and output a human readable string
 	 */
-	virtual bool ValidateAPICall( const TCHAR* Function, uint32 ErrorCode )
+	virtual bool ValidateAPICall(const TCHAR* Function, uint32 ErrorCode)
 	{
-		return( true );
+		return true;
 	}
 
 	const TArray<FActiveSound*>& GetActiveSounds() const { return ActiveSounds; }
@@ -623,10 +659,10 @@ public:
 	void GetAttenuationListenerData(FAttenuationListenerData& OutListenerData, const FTransform& SoundTransform, const FAttenuationSettings& AttenuationSettings, const FListener* InListener = nullptr) const;
 
 	/** Returns the focus factor of a sound based on its position and listener data. */
-	float GetFocusFactor(FAttenuationListenerData& OutListenerData, const class USoundBase* Sound, const FTransform& SoundTransform, const FAttenuationSettings& AttenuationSettings, const FListener* InListener = nullptr) const;
+	float GetFocusFactor(FAttenuationListenerData& OutListenerData, const USoundBase* Sound, const FTransform& SoundTransform, const FAttenuationSettings& AttenuationSettings, const FListener* InListener = nullptr) const;
 
 	/** Gets the max distance and focus factor of a sound. */
-	void GetMaxDistanceAndFocusFactor(USoundBase* Sound, const class UWorld* World, const FVector& Location, const FAttenuationSettings* AttenuationSettingsToApply, float *OutMaxDistance, float* OutFocusFactor);
+	void GetMaxDistanceAndFocusFactor(USoundBase* Sound, const UWorld* World, const FVector& Location, const FAttenuationSettings* AttenuationSettingsToApply, float *OutMaxDistance, float* OutFocusFactor);
 
 	/**
 	* Checks if the given sound would be audible.
@@ -638,7 +674,7 @@ public:
 	* @param FocusFactor			The focus factor of the sound.
 	* @param Returns true if the sound is audible, false otherwise.
 	*/
-	bool SoundIsAudible(USoundBase* Sound, const class UWorld* World, const FVector& Location, const FAttenuationSettings* AttenuationSettingsToApply, float MaxDistance, float FocusFactor);
+	bool SoundIsAudible(USoundBase* Sound, const UWorld* World, const FVector& Location, const FAttenuationSettings* AttenuationSettingsToApply, float MaxDistance, float FocusFactor);
 
 	/** Returns the index of the listener closest to the given sound transform */
 	static int32 FindClosestListenerIndex(const FTransform& SoundTransform, const TArray<FListener>& InListeners);
@@ -650,20 +686,18 @@ protected:
 	/**
 	 * Handle pausing/unpausing of sources when entering or leaving pause mode, or global pause (like device suspend)
 	 */
-	void HandlePause( bool bGameTicking, bool bGlobalPause = false );
+	void HandlePause(bool bGameTicking, bool bGlobalPause = false);
 
 	/**
 	 * Stop sources that need to be stopped, and touch the ones that need to be kept alive
 	 * Stop sounds that are too low in priority to be played
 	 */
-	void StopSources( TArray<FWaveInstance*>& WaveInstances, int32 FirstActiveIndex );
+	void StopSources(TArray<FWaveInstance*>& WaveInstances, int32 FirstActiveIndex);
 
 	/**
 	 * Start and/or update any sources that have a high enough priority to play
 	 */
-	void StartSources( TArray<FWaveInstance*>& WaveInstances, int32 FirstActiveIndex, bool bGameTicking );
-
-	
+	void StartSources(TArray<FWaveInstance*>& WaveInstances, int32 FirstActiveIndex, bool bGameTicking);
 
 	/**
 	 * Sets the 'pause' state of sounds which are always loaded.
@@ -677,8 +711,8 @@ protected:
 	/**
 	 * Lists a summary of loaded sound collated by class
 	 */
-	void ShowSoundClassHierarchy( FOutputDevice& Ar, class USoundClass* SoundClass = NULL, int32 Indent = 0 ) const;
-	
+	void ShowSoundClassHierarchy(FOutputDevice& Ar, USoundClass* SoundClass = nullptr, int32 Indent = 0) const;
+
 	/**
 	 * Parses the sound classes and propagates multiplicative properties down the tree.
 	 */
@@ -686,10 +720,11 @@ protected:
 
 	/**
 	 * Construct the CurrentSoundClassProperties map
+	 * @param DeltaTime The current frame delta. Used to interpolate sound class adjustments.
 	 *
 	 * This contains the original sound class properties propagated properly, and all adjustments due to the sound mixes
 	 */
-	void UpdateSoundClassProperties();
+	void UpdateSoundClassProperties(float DeltaTime);
 
 	/**
 	 * Set the mix for altering sound class properties
@@ -697,7 +732,7 @@ protected:
 	 * @param NewMix The SoundMix to apply
 	 * @param SoundMixState The State associated with this SoundMix
 	 */
-	bool ApplySoundMix( class USoundMix* NewMix, FSoundMixState* SoundMixState );
+	bool ApplySoundMix(USoundMix* NewMix, FSoundMixState* SoundMixState);
 
 	/**
 	 * Updates the state of a sound mix if it is pushed more than once.
@@ -705,7 +740,7 @@ protected:
 	 * @param SoundMix The SoundMix we are updating
 	 * @param SoundMixState The State associated with this SoundMix
 	 */
-	void UpdateSoundMix(class USoundMix* SoundMix, FSoundMixState* SoundMixState);
+	void UpdateSoundMix(USoundMix* SoundMix, FSoundMixState* SoundMixState);
 
 	/**
 	 * Updates list of SoundMixes that are applied passively, pushing and popping those that change
@@ -723,7 +758,7 @@ protected:
 	 *
 	 * @return Whether this SoundMix could be cleared (only true when both ref counts are zero).
 	 */
-	bool TryClearingSoundMix(class USoundMix* SoundMix, FSoundMixState* SoundMixState);
+	bool TryClearingSoundMix(USoundMix* SoundMix, FSoundMixState* SoundMixState);
 
 	/**
 	 * Attempt to remove this SoundMix's EQ effect - it may not currently be active
@@ -732,29 +767,30 @@ protected:
 	 *
 	 * @return Whether the effect of this SoundMix was cleared
 	 */
-	bool TryClearingEQSoundMix(class USoundMix* SoundMix);
+	bool TryClearingEQSoundMix(USoundMix* SoundMix);
 
 	/**
 	 * Find the SoundMix with the next highest EQ priority to the one passed in
 	 *
 	 * @param SoundMix The highest priority SoundMix, which will be ignored
 	 *
-	 * @return The next highest priority SoundMix or NULL if one cannot be found
+	 * @return The next highest priority SoundMix or nullptr if one cannot be found
 	 */
-	class USoundMix* FindNextHighestEQPrioritySoundMix(class USoundMix* IgnoredSoundMix);
+	USoundMix* FindNextHighestEQPrioritySoundMix(USoundMix* IgnoredSoundMix);
 
 	/**
 	 * Clear the effect of a SoundMix completely - only called after checking it's safe to
 	 */
-	void ClearSoundMix(class USoundMix* SoundMix);
+	void ClearSoundMix(USoundMix* SoundMix);
 
 	/**
 	 * Sets the sound class adjusters from a SoundMix.
 	 *
 	 * @param SoundMix		The SoundMix to apply adjusters from
 	 * @param InterpValue	Proportion of adjuster to apply
+	 * @param DeltaTime 	The current frame delta time. Used to interpolate sound class adjusters.
 	 */
-	void ApplyClassAdjusters(class USoundMix* SoundMix, float InterpValue);
+	void ApplyClassAdjusters(USoundMix* SoundMix, float InterpValue, float DeltaTime);
 
 	/**
 	 * Recursively apply an adjuster to the passed in sound class and all children of the sound class
@@ -762,7 +798,7 @@ protected:
 	 * @param InAdjuster		The adjuster to apply
 	 * @param InSoundClassName	The name of the sound class to apply the adjuster to.  Also applies to all children of this class
 	 */
-	void RecursiveApplyAdjuster( const struct FSoundClassAdjuster& InAdjuster, USoundClass* InSoundClass );
+	void RecursiveApplyAdjuster(const FSoundClassAdjuster& InAdjuster, USoundClass* InSoundClass);
 
 	/**
 	 * Takes an adjuster value and modifies it by the proportion that is currently in effect
@@ -775,29 +811,29 @@ protected:
 	/**
 	 * Platform dependent call to init effect data on a sound source
 	 */
-	void* InitEffect( class FSoundSource* Source );
+	void* InitEffect(FSoundSource* Source);
 
 	/**
 	 * Platform dependent call to update the sound output with new parameters
 	 * The audio system's main "Tick" function
 	 */
-	void* UpdateEffect( class FSoundSource* Source );
+	void* UpdateEffect(FSoundSource* Source);
 
 	/**
 	 * Platform dependent call to destroy any effect related data
 	 */
-	void DestroyEffect( class FSoundSource* Source );
+	void DestroyEffect(FSoundSource* Source);
 
 	/**
 	 * Return the pointer to the sound effects handler
 	 */
-	class FAudioEffectsManager* GetEffects( void )
+	FAudioEffectsManager* GetEffects()
 	{
-		return( Effects );
+		return Effects;
 	}
 
 	/** Internal */
-	void SortWaveInstances( int32 MaxChannels );
+	void SortWaveInstances(int32 MaxChannels);
 
 	/**
 	 * Internal helper function used by ParseSoundClasses to traverse the tree.
@@ -805,7 +841,7 @@ protected:
 	 * @param CurrentClass			Subtree to deal with
 	 * @param ParentProperties		Propagated properties of parent node
 	 */
-	void RecurseIntoSoundClasses( class USoundClass* CurrentClass, FSoundClassProperties& ParentProperties );
+	void RecurseIntoSoundClasses(USoundClass* CurrentClass, FSoundClassProperties& ParentProperties);
 
 	/**
 	 * Find the current highest priority reverb after a change to the list of active ones.
@@ -832,10 +868,10 @@ protected:
 	}
 
 	/** Creates a new platform specific sound source */
-	virtual class FAudioEffectsManager* CreateEffectsManager();
+	virtual FAudioEffectsManager* CreateEffectsManager();
 
 	/** Creates a new platform specific sound source */
-	virtual class FSoundSource* CreateSoundSource() PURE_VIRTUAL(FAudioDevice::CreateSoundSource,return NULL;);
+	virtual FSoundSource* CreateSoundSource() = 0;
 
 	/** Low pass filter OneOverQ value */
 	float GetLowPassFilterResonance() const;
@@ -887,12 +923,12 @@ public:
 	/* True once the startup sounds have been precached */
 	uint32 bStartupSoundsPreCached:1;
 
-	TArray<struct FListener> Listeners;
+	TArray<FListener> Listeners;
 
 	uint64 CurrentTick;
 
 	/** An AudioComponent to play test sounds on */
-	TWeakObjectPtr<class UAudioComponent> TestAudioComponent;
+	TWeakObjectPtr<UAudioComponent> TestAudioComponent;
 
 	/** The debug state of the audio device */
 	TEnumAsByte<enum EDebugState> DebugState;
@@ -913,42 +949,48 @@ public:
 	int32 NextResourceID;
 
 	/** Set of sources used to play sounds (platform will subclass these) */
-	TArray<class FSoundSource*>				Sources;
-	TArray<class FSoundSource*>				FreeSources;
-	TMap<struct FWaveInstance*, class FSoundSource*>	WaveInstanceSourceMap;
+	TArray<FSoundSource*> Sources;
+	TArray<FSoundSource*> FreeSources;
+	TMap<FWaveInstance*, FSoundSource*>	WaveInstanceSourceMap;
 
 	/** Current properties of all sound classes */
-	TMap<class USoundClass*, FSoundClassProperties>	SoundClasses;
+	TMap<USoundClass*, FSoundClassProperties>	SoundClasses;
 
 	/** The Base SoundMix that's currently active */
-	class USoundMix* BaseSoundMix;
+	USoundMix* BaseSoundMix;
 
 	/** The Base SoundMix that should be applied by default */
-	class USoundMix* DefaultBaseSoundMix;
+	USoundMix* DefaultBaseSoundMix;
+
+	/** Special sound mix set by user preferences */
+	USoundMix* UserPreferencesSoundMix;
 
 	/** Map of sound mixes currently affecting audio properties */
-	TMap<class USoundMix*, FSoundMixState>		SoundMixModifiers;
+	TMap<USoundMix*, FSoundMixState> SoundMixModifiers;
+
+	/** Map of sound mix sound class overrides. Will override any sound class effects for any sound mixes */
+	TMap<USoundMix*, FSoundMixClassOverrideMap> SoundMixClassEffectOverrides;
 
 	/** Interface to audio effects processing */
-	class FAudioEffectsManager*						Effects;
+	FAudioEffectsManager* Effects;
 
 	/** The volume the listener resides in */
-	const class AAudioVolume*						CurrentAudioVolume;
+	const AAudioVolume* CurrentAudioVolume;
 
 	/** A volume headroom to apply to specific platforms to achieve beter platform consistency. */
 	float PlatformAudioHeadroom;
 
 	/** Reverb Effects activated without volumes */
-	TMap<FName, FActivatedReverb>					ActivatedReverbs;
+	TMap<FName, FActivatedReverb> ActivatedReverbs;
 
 	/** The activated reverb that currently has the highest priority */
-	const FActivatedReverb*								HighestPriorityReverb;
+	const FActivatedReverb* HighestPriorityReverb;
 
 	/** Audio spatialization plugin. */
-	class IAudioSpatializationPlugin* SpatializationPlugin;
+	IAudioSpatializationPlugin* SpatializationPlugin;
 
 	/** Audio spatialization algorithm (derived from a plugin). */
-	class IAudioSpatializationAlgorithm* SpatializeProcessor;
+	IAudioSpatializationAlgorithm* SpatializeProcessor;
 
 	/** Whether or not the spatialization plugin is enabled. */
 	bool bSpatializationExtensionEnabled;
@@ -967,19 +1009,22 @@ public:
 
 private:
 
-	TArray<struct FActiveSound*> ActiveSounds;
+	TArray<FActiveSound*> ActiveSounds;
 	TArray<FWaveInstance*> ActiveWaveInstances;
 
+	/** Cached copy of sound class adjusters array. Cached to avoid allocating every frame. */
+	TArray<FSoundClassAdjuster> SoundClassAdjustersCopy;
+
 	/** Set of sounds which will be stopped next audio frame update */
-	TSet<struct FActiveSound*> PendingSoundsToStop;
+	TSet<FActiveSound*> PendingSoundsToStop;
 
 	/** A set of sounds which need to be deleted but weren't able to be deleted due to pending async operations */
-	TArray<struct FActiveSound*> PendingSoundsToDelete;
+	TArray<FActiveSound*> PendingSoundsToDelete;
 
-	TMap<UPTRINT, struct FActiveSound*> AudioComponentToActiveSoundMap;
+	TMap<UPTRINT, FActiveSound*> AudioComponentToActiveSoundMap;
 
 	/** List of passive SoundMixes active last frame */
-	TArray<class USoundMix*> PrevPassiveSoundMixModifiers;
+	TArray<USoundMix*> PrevPassiveSoundMixModifiers;
 
 	friend class FSoundConcurrencyManager;
 	FSoundConcurrencyManager ConcurrencyManager;
@@ -999,7 +1044,7 @@ class IAudioDeviceModule : public IModuleInterface
 public:
 
 	/** Creates a new instance of the audio device implemented by the module. */
-	virtual class FAudioDevice* CreateAudioDevice() = 0;
+	virtual FAudioDevice* CreateAudioDevice() = 0;
 };
 
 
