Copying //UE4/Dev-Framework to Dev-Main (//UE4/Dev-Main) @ 2879625

diff --git a/Engine/Source/Runtime/Engine/Public/VisualLogger/VisualLogger.h b/Engine/Source/Runtime/Engine/Public/VisualLogger/VisualLogger.h
--- a/Engine/Source/Runtime/Engine/Public/VisualLogger/VisualLogger.h
+++ b/Engine/Source/Runtime/Engine/Public/VisualLogger/VisualLogger.h
@@ -13,46 +13,46 @@
 #define CONNECT_OBJECT_WITH_VLOG(Src, Dest)
 
 // Text, regular log
-#define UE_VLOG(LogOwner, CategoryName, Verbosity, Format, ...) if( FVisualLogger::IsRecording() ) FVisualLogger::CategorizedLogf(LogOwner, CategoryName, ELogVerbosity::Verbosity, INDEX_NONE, Format, ##__VA_ARGS__)
+#define UE_VLOG(LogOwner, CategoryName, Verbosity, Format, ...) if( FVisualLogger::IsRecording() ) FVisualLogger::CategorizedLogf(LogOwner, CategoryName, ELogVerbosity::Verbosity, Format, ##__VA_ARGS__)
 #define UE_CVLOG(Condition, LogOwner, CategoryName, Verbosity, Format, ...)  if(FVisualLogger::IsRecording() && Condition) {UE_VLOG(LogOwner, CategoryName, Verbosity, Format, ##__VA_ARGS__);} 
 // Text, log with output to regular unreal logs too
-#define UE_VLOG_UELOG(LogOwner, CategoryName, Verbosity, Format, ...) { if(FVisualLogger::IsRecording()) FVisualLogger::CategorizedLogf(LogOwner, CategoryName, ELogVerbosity::Verbosity, INDEX_NONE, Format, ##__VA_ARGS__); UE_LOG(CategoryName, Verbosity, Format, ##__VA_ARGS__); }
+#define UE_VLOG_UELOG(LogOwner, CategoryName, Verbosity, Format, ...) { if(FVisualLogger::IsRecording()) FVisualLogger::CategorizedLogf(LogOwner, CategoryName, ELogVerbosity::Verbosity, Format, ##__VA_ARGS__); UE_LOG(CategoryName, Verbosity, Format, ##__VA_ARGS__); }
 #define UE_CVLOG_UELOG(Condition, LogOwner, CategoryName, Verbosity, Format, ...)  if(FVisualLogger::IsRecording() && Condition) {UE_VLOG_UELOG(LogOwner, CategoryName, Verbosity, Format, ##__VA_ARGS__);} 
 // Segment shape
-#define UE_VLOG_SEGMENT(LogOwner, CategoryName, Verbosity, SegmentStart, SegmentEnd, Color, Format, ...) if(FVisualLogger::IsRecording()) FVisualLogger::GeometryShapeLogf(LogOwner, CategoryName, ELogVerbosity::Verbosity, INDEX_NONE, SegmentStart, SegmentEnd, Color, Format, ##__VA_ARGS__)
-#define UE_CVLOG_SEGMENT(Condition, LogOwner, CategoryName, Verbosity, SegmentStart, SegmentEnd, Color, Format, ...) if(FVisualLogger::IsRecording() && Condition) {UE_VLOG_SEGMENT(LogOwner, CategoryName, Verbosity, INDEX_NONE, SegmentStart, SegmentEnd, Color, Format, ##__VA_ARGS__);}
+#define UE_VLOG_SEGMENT(LogOwner, CategoryName, Verbosity, SegmentStart, SegmentEnd, Color, Format, ...) if(FVisualLogger::IsRecording()) FVisualLogger::GeometryShapeLogf(LogOwner, CategoryName, ELogVerbosity::Verbosity, SegmentStart, SegmentEnd, Color, Format, ##__VA_ARGS__)
+#define UE_CVLOG_SEGMENT(Condition, LogOwner, CategoryName, Verbosity, SegmentStart, SegmentEnd, Color, Format, ...) if(FVisualLogger::IsRecording() && Condition) {UE_VLOG_SEGMENT(LogOwner, CategoryName, Verbosity, SegmentStart, SegmentEnd, Color, Format, ##__VA_ARGS__);}
 // Segment shape
-#define UE_VLOG_SEGMENT_THICK(LogOwner, CategoryName, Verbosity, SegmentStart, SegmentEnd, Color, Thickness, Format, ...) if(FVisualLogger::IsRecording()) FVisualLogger::GeometryShapeLogf(LogOwner, CategoryName, ELogVerbosity::Verbosity, INDEX_NONE, SegmentStart, SegmentEnd, Color, Format, ##__VA_ARGS__)
-#define UE_CVLOG_SEGMENT_THICK(Condition, LogOwner, CategoryName, Verbosity, SegmentStart, SegmentEnd, Color, Thickness, Format, ...) if(FVisualLogger::IsRecording() && Condition) {UE_VLOG_SEGMENT_THICK(LogOwner, CategoryName, Verbosity, INDEX_NONE, SegmentStart, SegmentEnd, Color, Format, ##__VA_ARGS__);} 
+#define UE_VLOG_SEGMENT_THICK(LogOwner, CategoryName, Verbosity, SegmentStart, SegmentEnd, Color, Thickness, Format, ...) if(FVisualLogger::IsRecording()) FVisualLogger::GeometryShapeLogf(LogOwner, CategoryName, ELogVerbosity::Verbosity, SegmentStart, SegmentEnd, Color, Format, ##__VA_ARGS__)
+#define UE_CVLOG_SEGMENT_THICK(Condition, LogOwner, CategoryName, Verbosity, SegmentStart, SegmentEnd, Color, Thickness, Format, ...) if(FVisualLogger::IsRecording() && Condition) {UE_VLOG_SEGMENT_THICK(LogOwner, CategoryName, Verbosity, SegmentStart, SegmentEnd, Color, Format, ##__VA_ARGS__);} 
 // Localization as sphere shape
-#define UE_VLOG_LOCATION(LogOwner, CategoryName, Verbosity, Location, Radius, Color, Format, ...) if(FVisualLogger::IsRecording()) FVisualLogger::GeometryShapeLogf(LogOwner, CategoryName, ELogVerbosity::Verbosity, INDEX_NONE, Location, Radius, Color, Format, ##__VA_ARGS__)
+#define UE_VLOG_LOCATION(LogOwner, CategoryName, Verbosity, Location, Radius, Color, Format, ...) if(FVisualLogger::IsRecording()) FVisualLogger::GeometryShapeLogf(LogOwner, CategoryName, ELogVerbosity::Verbosity, Location, Radius, Color, Format, ##__VA_ARGS__)
 #define UE_CVLOG_LOCATION(Condition, LogOwner, CategoryName, Verbosity, Location, Radius, Color, Format, ...)  if(FVisualLogger::IsRecording() && Condition) {UE_VLOG_LOCATION(LogOwner, CategoryName, Verbosity, Location, Radius, Color, Format, ##__VA_ARGS__);} 
 // Box shape
-#define UE_VLOG_BOX(LogOwner, CategoryName, Verbosity, Box, Color, Format, ...) if(FVisualLogger::IsRecording()) FVisualLogger::GeometryShapeLogf(LogOwner, CategoryName, ELogVerbosity::Verbosity, INDEX_NONE, Box, FMatrix::Identity, Color, Format, ##__VA_ARGS__)
+#define UE_VLOG_BOX(LogOwner, CategoryName, Verbosity, Box, Color, Format, ...) if(FVisualLogger::IsRecording()) FVisualLogger::GeometryShapeLogf(LogOwner, CategoryName, ELogVerbosity::Verbosity, Box, FMatrix::Identity, Color, Format, ##__VA_ARGS__)
 #define UE_CVLOG_BOX(Condition, LogOwner, CategoryName, Verbosity, Box, Color, Format, ...)  if(FVisualLogger::IsRecording() && Condition) {UE_VLOG_BOX(Actor, CategoryName, Verbosity, Box, FMatrix::Identity, Color, Format, ##__VA_ARGS__);} 
 // Oriented box shape
-#define UE_VLOG_OBOX(LogOwner, CategoryName, Verbosity, Box, Matrix, Color, Format, ...) if(FVisualLogger::IsRecording()) FVisualLogger::GeometryShapeLogf(LogOwner, CategoryName, ELogVerbosity::Verbosity, INDEX_NONE, Box, Matrix, Color, Format, ##__VA_ARGS__)
+#define UE_VLOG_OBOX(LogOwner, CategoryName, Verbosity, Box, Matrix, Color, Format, ...) if(FVisualLogger::IsRecording()) FVisualLogger::GeometryShapeLogf(LogOwner, CategoryName, ELogVerbosity::Verbosity, Box, Matrix, Color, Format, ##__VA_ARGS__)
 #define UE_CVLOG_OBOX(Condition, LogOwner, CategoryName, Verbosity, Box, Matrix, Color, Format, ...)  if(FVisualLogger::IsRecording() && Condition) {UE_VLOG_BOX(Actor, CategoryName, Verbosity, Box, Matrix, Color, Format, ##__VA_ARGS__);} 
 // Cone shape
-#define UE_VLOG_CONE(LogOwner, CategoryName, Verbosity, Orgin, Direction, Length, Angle, Color, Format, ...) if(FVisualLogger::IsRecording()) FVisualLogger::GeometryShapeLogf(LogOwner, CategoryName, ELogVerbosity::Verbosity, INDEX_NONE, Orgin, Direction, Length, Angle, Color, Format, ##__VA_ARGS__)
+#define UE_VLOG_CONE(LogOwner, CategoryName, Verbosity, Orgin, Direction, Length, Angle, Color, Format, ...) if(FVisualLogger::IsRecording()) FVisualLogger::GeometryShapeLogf(LogOwner, CategoryName, ELogVerbosity::Verbosity, Orgin, Direction, Length, Angle, Color, Format, ##__VA_ARGS__)
 #define UE_CVLOG_CONE(Condition, LogOwner, CategoryName, Verbosity, Orgin, Direction, Length, Angle, Color, Format, ...)  if(FVisualLogger::IsRecording() && Condition) {UE_VLOG_CONE(Actor, CategoryName, Verbosity, Orgin, Direction, Length, Angle, Color, Format, ##__VA_ARGS__);} 
 // Cylinder shape
-#define UE_VLOG_CYLINDER(LogOwner, CategoryName, Verbosity, Start, End, Radius, Color, Format, ...) if(FVisualLogger::IsRecording()) FVisualLogger::GeometryShapeLogf(LogOwner, CategoryName, ELogVerbosity::Verbosity, INDEX_NONE, Start, End, Radius, Color, Format, ##__VA_ARGS__)
+#define UE_VLOG_CYLINDER(LogOwner, CategoryName, Verbosity, Start, End, Radius, Color, Format, ...) if(FVisualLogger::IsRecording()) FVisualLogger::GeometryShapeLogf(LogOwner, CategoryName, ELogVerbosity::Verbosity, Start, End, Radius, Color, Format, ##__VA_ARGS__)
 #define UE_CVLOG_CYLINDER(Condition, LogOwner, CategoryName, Verbosity, Start, End, Radius, Color, Format, ...)  if(FVisualLogger::IsRecording() && Condition) {UE_VLOG_CYLINDER(LogOwner, CategoryName, Verbosity, Start, End, Radius, Color, Format, ##__VA_ARGS__);} 
 // Capsule shape
-#define UE_VLOG_CAPSULE(LogOwner, CategoryName, Verbosity, Center, HalfHeight, Radius, Rotation, Color, Format, ...) if(FVisualLogger::IsRecording()) FVisualLogger::GeometryShapeLogf(LogOwner, CategoryName, ELogVerbosity::Verbosity, INDEX_NONE, Center, HalfHeight, Radius, Rotation, Color, Format, ##__VA_ARGS__)
+#define UE_VLOG_CAPSULE(LogOwner, CategoryName, Verbosity, Center, HalfHeight, Radius, Rotation, Color, Format, ...) if(FVisualLogger::IsRecording()) FVisualLogger::GeometryShapeLogf(LogOwner, CategoryName, ELogVerbosity::Verbosity, Center, HalfHeight, Radius, Rotation, Color, Format, ##__VA_ARGS__)
 #define UE_CVLOG_CAPSULE(Condition, LogOwner, CategoryName, Verbosity, Center, HalfHeight, Radius, Rotation, Color, Format, ...)  if(FVisualLogger::IsRecording() && Condition) {UE_VLOG_CAPSULE(LogOwner, CategoryName, Verbosity, Center, HalfHeight, Radius, Rotation, Color, Format, ##__VA_ARGS__);} 
 // Histogram data for 2d graphs 
-#define UE_VLOG_HISTOGRAM(LogOwner, CategoryName, Verbosity, GraphName, DataName, Data) if(FVisualLogger::IsRecording()) FVisualLogger::HistogramDataLogf(LogOwner, CategoryName, ELogVerbosity::Verbosity, INDEX_NONE, GraphName, DataName, Data, FColor::White, TEXT(""))
+#define UE_VLOG_HISTOGRAM(LogOwner, CategoryName, Verbosity, GraphName, DataName, Data) if(FVisualLogger::IsRecording()) FVisualLogger::HistogramDataLogf(LogOwner, CategoryName, ELogVerbosity::Verbosity, GraphName, DataName, Data, FColor::White, TEXT(""))
 #define UE_CVLOG_HISTOGRAM(Condition, LogOwner, CategoryName, Verbosity, GraphName, DataName, Data) if(FVisualLogger::IsRecording() && Condition) {UE_VLOG_HISTOGRAM(LogOwner, CategoryName, Verbosity, GraphName, DataName, Data);} 
 // NavArea or vertically pulled convex shape
-#define UE_VLOG_PULLEDCONVEX(LogOwner, CategoryName, Verbosity, ConvexPoints, MinZ, MaxZ, Color, Format, ...) if(FVisualLogger::IsRecording()) FVisualLogger::NavAreaShapeLogf(LogOwner, CategoryName, ELogVerbosity::Verbosity, INDEX_NONE, ConvexPoints, MinZ, MaxZ, Color, Format, ##__VA_ARGS__)
+#define UE_VLOG_PULLEDCONVEX(LogOwner, CategoryName, Verbosity, ConvexPoints, MinZ, MaxZ, Color, Format, ...) if(FVisualLogger::IsRecording()) FVisualLogger::NavAreaShapeLogf(LogOwner, CategoryName, ELogVerbosity::Verbosity, ConvexPoints, MinZ, MaxZ, Color, Format, ##__VA_ARGS__)
 #define UE_CVLOG_PULLEDCONVEX(Condition, LogOwner, CategoryName, Verbosity, ConvexPoints, MinZ, MaxZ, Color, Format, ...) if(FVisualLogger::IsRecording() && Condition) {UE_VLOG_PULLEDCONVEX(LogOwner, CategoryName, Verbosity, ConvexPoints, MinZ, MaxZ, Color, Format, ##__VA_ARGS__);}
 // regular 3d mesh shape to log
-#define UE_VLOG_MESH(LogOwner, CategoryName, Verbosity, Vertices, Indices, Color, Format, ...) if (FVisualLogger::IsRecording()) FVisualLogger::GeometryShapeLogf(LogOwner, CategoryName, ELogVerbosity::Verbosity, INDEX_NONE, Vertices, Indices, Format, ##__VA_ARGS__)
+#define UE_VLOG_MESH(LogOwner, CategoryName, Verbosity, Vertices, Indices, Color, Format, ...) if (FVisualLogger::IsRecording()) FVisualLogger::GeometryShapeLogf(LogOwner, CategoryName, ELogVerbosity::Verbosity, Vertices, Indices, Format, ##__VA_ARGS__)
 #define UE_CVLOG_MESH(Condition, LogOwner, CategoryName, Verbosity, Vertices, Indexes, Color, Format, ...) if(FVisualLogger::IsRecording() && Condition) {UE_VLOG_MESH(LogOwner, CategoryName, Verbosity, Vertices, Indices, Color, Format, ##__VA_ARGS__);}
 // 2d convex poly shape
-#define UE_VLOG_CONVEXPOLY(LogOwner, CategoryName, Verbosity, Points, Color, Format, ...) if(FVisualLogger::IsRecording()) FVisualLogger::GeometryConvexLogf(LogOwner, CategoryName, ELogVerbosity::Verbosity, INDEX_NONE, Points, Color, Format, ##__VA_ARGS__)
+#define UE_VLOG_CONVEXPOLY(LogOwner, CategoryName, Verbosity, Points, Color, Format, ...) if(FVisualLogger::IsRecording()) FVisualLogger::GeometryConvexLogf(LogOwner, CategoryName, ELogVerbosity::Verbosity, Points, Color, Format, ##__VA_ARGS__)
 #define UE_CVLOG_CONVEXPOLY(Condition, LogOwner, CategoryName, Verbosity, Points, Color, Format, ...)  if(FVisualLogger::IsRecording() && Condition) {UE_VLOG_CONVEXPOLY(LogOwner, CategoryName, Verbosity, Points, Color, Format, ##__VA_ARGS__);}
 
 
@@ -136,29 +136,29 @@ class ENGINE_API FVisualLogger : public FOutputDevice
 {
 public:
 	// Regular text log
-	VARARG_DECL(static void, static void, return, CategorizedLogf, VARARG_NONE, const TCHAR*, VARARG_EXTRA(const UObject* LogOwner) VARARG_EXTRA(const FLogCategoryBase& Category) VARARG_EXTRA(ELogVerbosity::Type Verbosity) VARARG_EXTRA(int32 UniqueLogId), VARARG_EXTRA(LogOwner) VARARG_EXTRA(Category) VARARG_EXTRA(Verbosity) VARARG_EXTRA(UniqueLogId));
+	VARARG_DECL(static void, static void, return, CategorizedLogf, VARARG_NONE, const TCHAR*, VARARG_EXTRA(const UObject* LogOwner) VARARG_EXTRA(const FLogCategoryBase& Category) VARARG_EXTRA(ELogVerbosity::Type Verbosity), VARARG_EXTRA(LogOwner) VARARG_EXTRA(Category) VARARG_EXTRA(Verbosity));
 	// Segment log
-	VARARG_DECL(static void, static void, return, GeometryShapeLogf, VARARG_NONE, const TCHAR*, VARARG_EXTRA(const UObject* LogOwner) VARARG_EXTRA(const FLogCategoryBase& Category) VARARG_EXTRA(ELogVerbosity::Type Verbosity) VARARG_EXTRA(int32 UniqueLogId) VARARG_EXTRA(const FVector& Start) VARARG_EXTRA(const FVector& End) VARARG_EXTRA(const FColor& Color), VARARG_EXTRA(LogOwner) VARARG_EXTRA(Category) VARARG_EXTRA(Verbosity) VARARG_EXTRA(UniqueLogId) VARARG_EXTRA(Start) VARARG_EXTRA(End) VARARG_EXTRA(Color));
+	VARARG_DECL(static void, static void, return, GeometryShapeLogf, VARARG_NONE, const TCHAR*, VARARG_EXTRA(const UObject* LogOwner) VARARG_EXTRA(const FLogCategoryBase& Category) VARARG_EXTRA(ELogVerbosity::Type Verbosity) VARARG_EXTRA(const FVector& Start) VARARG_EXTRA(const FVector& End) VARARG_EXTRA(const FColor& Color), VARARG_EXTRA(LogOwner) VARARG_EXTRA(Category) VARARG_EXTRA(Verbosity) VARARG_EXTRA(Start) VARARG_EXTRA(End) VARARG_EXTRA(Color));
 	// Location/Shpere log
-	VARARG_DECL(static void, static void, return, GeometryShapeLogf, VARARG_NONE, const TCHAR*, VARARG_EXTRA(const UObject* LogOwner) VARARG_EXTRA(const FLogCategoryBase& Category) VARARG_EXTRA(ELogVerbosity::Type Verbosity) VARARG_EXTRA(int32 UniqueLogId) VARARG_EXTRA(const FVector& Location) VARARG_EXTRA(float Radius) VARARG_EXTRA(const FColor& Color), VARARG_EXTRA(LogOwner) VARARG_EXTRA(Category) VARARG_EXTRA(Verbosity) VARARG_EXTRA(UniqueLogId) VARARG_EXTRA(Location) VARARG_EXTRA(Radius) VARARG_EXTRA(Color));
+	VARARG_DECL(static void, static void, return, GeometryShapeLogf, VARARG_NONE, const TCHAR*, VARARG_EXTRA(const UObject* LogOwner) VARARG_EXTRA(const FLogCategoryBase& Category) VARARG_EXTRA(ELogVerbosity::Type Verbosity) VARARG_EXTRA(const FVector& Location) VARARG_EXTRA(float Radius) VARARG_EXTRA(const FColor& Color), VARARG_EXTRA(LogOwner) VARARG_EXTRA(Category) VARARG_EXTRA(Verbosity) VARARG_EXTRA(Location) VARARG_EXTRA(Radius) VARARG_EXTRA(Color));
 	// Box log
-	VARARG_DECL(static void, static void, return, GeometryShapeLogf, VARARG_NONE, const TCHAR*, VARARG_EXTRA(const UObject* LogOwner) VARARG_EXTRA(const FLogCategoryBase& Category) VARARG_EXTRA(ELogVerbosity::Type Verbosity) VARARG_EXTRA(int32 UniqueLogId) VARARG_EXTRA(const FBox& Box) VARARG_EXTRA(const FMatrix& Matrix) VARARG_EXTRA(const FColor& Color), VARARG_EXTRA(LogOwner) VARARG_EXTRA(Category) VARARG_EXTRA(Verbosity) VARARG_EXTRA(UniqueLogId) VARARG_EXTRA(Box) VARARG_EXTRA(Matrix) VARARG_EXTRA(Color));
+	VARARG_DECL(static void, static void, return, GeometryShapeLogf, VARARG_NONE, const TCHAR*, VARARG_EXTRA(const UObject* LogOwner) VARARG_EXTRA(const FLogCategoryBase& Category) VARARG_EXTRA(ELogVerbosity::Type Verbosity) VARARG_EXTRA(const FBox& Box) VARARG_EXTRA(const FMatrix& Matrix) VARARG_EXTRA(const FColor& Color), VARARG_EXTRA(LogOwner) VARARG_EXTRA(Category) VARARG_EXTRA(Verbosity) VARARG_EXTRA(Box) VARARG_EXTRA(Matrix) VARARG_EXTRA(Color));
 	// Cone log
-	VARARG_DECL(static void, static void, return, GeometryShapeLogf, VARARG_NONE, const TCHAR*, VARARG_EXTRA(const UObject* LogOwner) VARARG_EXTRA(const FLogCategoryBase& Category) VARARG_EXTRA(ELogVerbosity::Type Verbosity) VARARG_EXTRA(int32 UniqueLogId) VARARG_EXTRA(const FVector& Orgin) VARARG_EXTRA(const FVector& Direction) VARARG_EXTRA(const float Length) VARARG_EXTRA(const float Angle) VARARG_EXTRA(const FColor& Color), VARARG_EXTRA(LogOwner) VARARG_EXTRA(Category) VARARG_EXTRA(Verbosity) VARARG_EXTRA(UniqueLogId) VARARG_EXTRA(Orgin) VARARG_EXTRA(Direction) VARARG_EXTRA(Length) VARARG_EXTRA(Angle) VARARG_EXTRA(Color));
+	VARARG_DECL(static void, static void, return, GeometryShapeLogf, VARARG_NONE, const TCHAR*, VARARG_EXTRA(const UObject* LogOwner) VARARG_EXTRA(const FLogCategoryBase& Category) VARARG_EXTRA(ELogVerbosity::Type Verbosity) VARARG_EXTRA(const FVector& Orgin) VARARG_EXTRA(const FVector& Direction) VARARG_EXTRA(const float Length) VARARG_EXTRA(const float Angle) VARARG_EXTRA(const FColor& Color), VARARG_EXTRA(LogOwner) VARARG_EXTRA(Category) VARARG_EXTRA(Verbosity) VARARG_EXTRA(Orgin) VARARG_EXTRA(Direction) VARARG_EXTRA(Length) VARARG_EXTRA(Angle) VARARG_EXTRA(Color));
 	// Cylinder log
-	VARARG_DECL(static void, static void, return, GeometryShapeLogf, VARARG_NONE, const TCHAR*, VARARG_EXTRA(const UObject* LogOwner) VARARG_EXTRA(const FLogCategoryBase& Category) VARARG_EXTRA(ELogVerbosity::Type Verbosity) VARARG_EXTRA(int32 UniqueLogId) VARARG_EXTRA(const FVector& Start) VARARG_EXTRA(const FVector& End) VARARG_EXTRA(const float Radius) VARARG_EXTRA(const FColor& Color), VARARG_EXTRA(LogOwner) VARARG_EXTRA(Category) VARARG_EXTRA(Verbosity) VARARG_EXTRA(UniqueLogId) VARARG_EXTRA(Start) VARARG_EXTRA(End) VARARG_EXTRA(Radius) VARARG_EXTRA(Color));
+	VARARG_DECL(static void, static void, return, GeometryShapeLogf, VARARG_NONE, const TCHAR*, VARARG_EXTRA(const UObject* LogOwner) VARARG_EXTRA(const FLogCategoryBase& Category) VARARG_EXTRA(ELogVerbosity::Type Verbosity) VARARG_EXTRA(const FVector& Start) VARARG_EXTRA(const FVector& End) VARARG_EXTRA(const float Radius) VARARG_EXTRA(const FColor& Color), VARARG_EXTRA(LogOwner) VARARG_EXTRA(Category) VARARG_EXTRA(Verbosity) VARARG_EXTRA(Start) VARARG_EXTRA(End) VARARG_EXTRA(Radius) VARARG_EXTRA(Color));
 	// Capsule log
-	VARARG_DECL(static void, static void, return, GeometryShapeLogf, VARARG_NONE, const TCHAR*, VARARG_EXTRA(const UObject* LogOwner) VARARG_EXTRA(const FLogCategoryBase& Category) VARARG_EXTRA(ELogVerbosity::Type Verbosity) VARARG_EXTRA(int32 UniqueLogId) VARARG_EXTRA(const FVector& Center) VARARG_EXTRA(float HalfHeight) VARARG_EXTRA(float Radius) VARARG_EXTRA(const FQuat & Rotation) VARARG_EXTRA(const FColor& Color), VARARG_EXTRA(LogOwner) VARARG_EXTRA(Category) VARARG_EXTRA(Verbosity) VARARG_EXTRA(UniqueLogId) VARARG_EXTRA(Center) VARARG_EXTRA(HalfHeight) VARARG_EXTRA(Radius) VARARG_EXTRA(Rotation) VARARG_EXTRA(Color));
+	VARARG_DECL(static void, static void, return, GeometryShapeLogf, VARARG_NONE, const TCHAR*, VARARG_EXTRA(const UObject* LogOwner) VARARG_EXTRA(const FLogCategoryBase& Category) VARARG_EXTRA(ELogVerbosity::Type Verbosity) VARARG_EXTRA(const FVector& Center) VARARG_EXTRA(float HalfHeight) VARARG_EXTRA(float Radius) VARARG_EXTRA(const FQuat & Rotation) VARARG_EXTRA(const FColor& Color), VARARG_EXTRA(LogOwner) VARARG_EXTRA(Category) VARARG_EXTRA(Verbosity) VARARG_EXTRA(Center) VARARG_EXTRA(HalfHeight) VARARG_EXTRA(Radius) VARARG_EXTRA(Rotation) VARARG_EXTRA(Color));
 	// NavArea/Extruded convex log
-	VARARG_DECL(static void, static void, return, NavAreaShapeLogf, VARARG_NONE, const TCHAR*, VARARG_EXTRA(const UObject* LogOwner) VARARG_EXTRA(const FLogCategoryBase& Category) VARARG_EXTRA(ELogVerbosity::Type Verbosity) VARARG_EXTRA(int32 UniqueLogId) VARARG_EXTRA(const TArray<FVector> &ConvexPoints) VARARG_EXTRA(float MinZ) VARARG_EXTRA(float MaxZ) VARARG_EXTRA(const FColor& Color), VARARG_EXTRA(LogOwner) VARARG_EXTRA(Category) VARARG_EXTRA(Verbosity) VARARG_EXTRA(UniqueLogId) VARARG_EXTRA(ConvexPoints) VARARG_EXTRA(MinZ) VARARG_EXTRA(MaxZ) VARARG_EXTRA(Color));
+	VARARG_DECL(static void, static void, return, NavAreaShapeLogf, VARARG_NONE, const TCHAR*, VARARG_EXTRA(const UObject* LogOwner) VARARG_EXTRA(const FLogCategoryBase& Category) VARARG_EXTRA(ELogVerbosity::Type Verbosity) VARARG_EXTRA(const TArray<FVector> &ConvexPoints) VARARG_EXTRA(float MinZ) VARARG_EXTRA(float MaxZ) VARARG_EXTRA(const FColor& Color), VARARG_EXTRA(LogOwner) VARARG_EXTRA(Category) VARARG_EXTRA(Verbosity) VARARG_EXTRA(ConvexPoints) VARARG_EXTRA(MinZ) VARARG_EXTRA(MaxZ) VARARG_EXTRA(Color));
 	// 3d Mesh log
-	VARARG_DECL(static void, static void, return, GeometryShapeLogf, VARARG_NONE, const TCHAR*, VARARG_EXTRA(const UObject* LogOwner) VARARG_EXTRA(const FLogCategoryBase& Category) VARARG_EXTRA(ELogVerbosity::Type Verbosity) VARARG_EXTRA(int32 UniqueLogId) VARARG_EXTRA(const TArray<FVector> &Vertices) VARARG_EXTRA(const TArray<int32> &Indices) VARARG_EXTRA(const FColor& Color), VARARG_EXTRA(LogOwner) VARARG_EXTRA(Category) VARARG_EXTRA(Verbosity) VARARG_EXTRA(UniqueLogId) VARARG_EXTRA(Vertices) VARARG_EXTRA(Indices) VARARG_EXTRA(Color));
+	VARARG_DECL(static void, static void, return, GeometryShapeLogf, VARARG_NONE, const TCHAR*, VARARG_EXTRA(const UObject* LogOwner) VARARG_EXTRA(const FLogCategoryBase& Category) VARARG_EXTRA(ELogVerbosity::Type Verbosity) VARARG_EXTRA(const TArray<FVector> &Vertices) VARARG_EXTRA(const TArray<int32> &Indices) VARARG_EXTRA(const FColor& Color), VARARG_EXTRA(LogOwner) VARARG_EXTRA(Category) VARARG_EXTRA(Verbosity) VARARG_EXTRA(Vertices) VARARG_EXTRA(Indices) VARARG_EXTRA(Color));
 	// 2d Convex shape
-	VARARG_DECL(static void, static void, return, GeometryConvexLogf, VARARG_NONE, const TCHAR*, VARARG_EXTRA(const UObject* LogOwner) VARARG_EXTRA(const FLogCategoryBase& Category) VARARG_EXTRA(ELogVerbosity::Type Verbosity) VARARG_EXTRA(int32 UniqueLogId) VARARG_EXTRA(const TArray<FVector> &Points) VARARG_EXTRA(const FColor& Color), VARARG_EXTRA(LogOwner) VARARG_EXTRA(Category) VARARG_EXTRA(Verbosity) VARARG_EXTRA(UniqueLogId) VARARG_EXTRA(Points) VARARG_EXTRA(Color));
+	VARARG_DECL(static void, static void, return, GeometryConvexLogf, VARARG_NONE, const TCHAR*, VARARG_EXTRA(const UObject* LogOwner) VARARG_EXTRA(const FLogCategoryBase& Category) VARARG_EXTRA(ELogVerbosity::Type Verbosity) VARARG_EXTRA(const TArray<FVector> &Points) VARARG_EXTRA(const FColor& Color), VARARG_EXTRA(LogOwner) VARARG_EXTRA(Category) VARARG_EXTRA(Verbosity) VARARG_EXTRA(Points) VARARG_EXTRA(Color));
 	//Histogram data
-	VARARG_DECL(static void, static void, return, HistogramDataLogf, VARARG_NONE, const TCHAR*, VARARG_EXTRA(const UObject* LogOwner) VARARG_EXTRA(const FLogCategoryBase& Category) VARARG_EXTRA(ELogVerbosity::Type Verbosity) VARARG_EXTRA(int32 UniqueLogId) VARARG_EXTRA(FName GraphName) VARARG_EXTRA(FName DataName) VARARG_EXTRA(const FVector2D& Data) VARARG_EXTRA(const FColor& Color), VARARG_EXTRA(LogOwner) VARARG_EXTRA(Category) VARARG_EXTRA(Verbosity) VARARG_EXTRA(UniqueLogId) VARARG_EXTRA(GraphName) VARARG_EXTRA(DataName) VARARG_EXTRA(Data) VARARG_EXTRA(Color));
+	VARARG_DECL(static void, static void, return, HistogramDataLogf, VARARG_NONE, const TCHAR*, VARARG_EXTRA(const UObject* LogOwner) VARARG_EXTRA(const FLogCategoryBase& Category) VARARG_EXTRA(ELogVerbosity::Type Verbosity) VARARG_EXTRA(FName GraphName) VARARG_EXTRA(FName DataName) VARARG_EXTRA(const FVector2D& Data) VARARG_EXTRA(const FColor& Color), VARARG_EXTRA(LogOwner) VARARG_EXTRA(Category) VARARG_EXTRA(Verbosity) VARARG_EXTRA(GraphName) VARARG_EXTRA(DataName) VARARG_EXTRA(Data) VARARG_EXTRA(Color));
 	// Navigation data debug snapshot
-	static void NavigationDataDump(const UObject* LogOwner, const FLogCategoryBase& Category, ELogVerbosity::Type Verbosity, int32 UniqueLogId, const FBox& Box);
+	static void NavigationDataDump(const UObject* LogOwner, const FLogCategoryBase& Category, const ELogVerbosity::Type Verbosity, const FBox& Box);
 
 	/** Log events */
 	static void EventLog(const UObject* LogOwner, const FName EventTag1, const FVisualLogEventBase& Event1, const FName EventTag2 = NAME_None, const FName EventTag3 = NAME_None, const FName EventTag4 = NAME_None, const FName EventTag5 = NAME_None, const FName EventTag6 = NAME_None);
@@ -346,7 +346,7 @@ protected:
 	FMemory::SystemFree(AllocatedBuffer);
 
 inline
-VARARG_BODY(void, FVisualLogger::CategorizedLogf, const TCHAR*, VARARG_EXTRA(const UObject* Object) VARARG_EXTRA(const FLogCategoryBase& Category) VARARG_EXTRA(ELogVerbosity::Type Verbosity) VARARG_EXTRA(int32 UniqueLogId))
+VARARG_BODY(void, FVisualLogger::CategorizedLogf, const TCHAR*, VARARG_EXTRA(const UObject* Object) VARARG_EXTRA(const FLogCategoryBase& Category) VARARG_EXTRA(ELogVerbosity::Type Verbosity))
 {
 	SCOPE_CYCLE_COUNTER(STAT_VisualLog);
 	UWorld *World = nullptr;
@@ -362,7 +362,7 @@ VARARG_BODY(void, FVisualLogger::CategorizedLogf, const TCHAR*, VARARG_EXTRA(con
 }
 
 inline
-VARARG_BODY(void, FVisualLogger::GeometryShapeLogf, const TCHAR*, VARARG_EXTRA(const UObject* Object) VARARG_EXTRA(const FLogCategoryBase& Category) VARARG_EXTRA(ELogVerbosity::Type Verbosity) VARARG_EXTRA(int32 UniqueLogId) VARARG_EXTRA(const FVector& Start) VARARG_EXTRA(const FVector& End) VARARG_EXTRA(const FColor& Color))
+VARARG_BODY(void, FVisualLogger::GeometryShapeLogf, const TCHAR*, VARARG_EXTRA(const UObject* Object) VARARG_EXTRA(const FLogCategoryBase& Category) VARARG_EXTRA(ELogVerbosity::Type Verbosity) VARARG_EXTRA(const FVector& Start) VARARG_EXTRA(const FVector& End) VARARG_EXTRA(const FColor& Color))
 {
 	SCOPE_CYCLE_COUNTER(STAT_VisualLog);
 	UWorld *World = nullptr;
@@ -378,7 +378,7 @@ VARARG_BODY(void, FVisualLogger::GeometryShapeLogf, const TCHAR*, VARARG_EXTRA(c
 }
 
 inline
-VARARG_BODY(void, FVisualLogger::GeometryShapeLogf, const TCHAR*, VARARG_EXTRA(const UObject* Object) VARARG_EXTRA(const FLogCategoryBase& Category) VARARG_EXTRA(ELogVerbosity::Type Verbosity) VARARG_EXTRA(int32 UniqueLogId) VARARG_EXTRA(const FVector& Location) VARARG_EXTRA(float Radius) VARARG_EXTRA(const FColor& Color))
+VARARG_BODY(void, FVisualLogger::GeometryShapeLogf, const TCHAR*, VARARG_EXTRA(const UObject* Object) VARARG_EXTRA(const FLogCategoryBase& Category) VARARG_EXTRA(ELogVerbosity::Type Verbosity) VARARG_EXTRA(const FVector& Location) VARARG_EXTRA(float Radius) VARARG_EXTRA(const FColor& Color))
 {
 	SCOPE_CYCLE_COUNTER(STAT_VisualLog);
 	UWorld *World = nullptr;
@@ -394,7 +394,7 @@ VARARG_BODY(void, FVisualLogger::GeometryShapeLogf, const TCHAR*, VARARG_EXTRA(c
 }
 
 inline
-VARARG_BODY(void, FVisualLogger::GeometryShapeLogf, const TCHAR*, VARARG_EXTRA(const UObject* Object) VARARG_EXTRA(const FLogCategoryBase& Category) VARARG_EXTRA(ELogVerbosity::Type Verbosity) VARARG_EXTRA(int32 UniqueLogId) VARARG_EXTRA(const FBox& Box) VARARG_EXTRA(const FMatrix& Matrix) VARARG_EXTRA(const FColor& Color))
+VARARG_BODY(void, FVisualLogger::GeometryShapeLogf, const TCHAR*, VARARG_EXTRA(const UObject* Object) VARARG_EXTRA(const FLogCategoryBase& Category) VARARG_EXTRA(ELogVerbosity::Type Verbosity) VARARG_EXTRA(const FBox& Box) VARARG_EXTRA(const FMatrix& Matrix) VARARG_EXTRA(const FColor& Color))
 {
 	SCOPE_CYCLE_COUNTER(STAT_VisualLog);
 	UWorld *World = nullptr;
@@ -410,7 +410,7 @@ VARARG_BODY(void, FVisualLogger::GeometryShapeLogf, const TCHAR*, VARARG_EXTRA(c
 }
 
 inline
-VARARG_BODY(void, FVisualLogger::GeometryShapeLogf, const TCHAR*, VARARG_EXTRA(const UObject* Object) VARARG_EXTRA(const FLogCategoryBase& Category) VARARG_EXTRA(ELogVerbosity::Type Verbosity) VARARG_EXTRA(int32 UniqueLogId) VARARG_EXTRA(const FVector& Orgin) VARARG_EXTRA(const FVector& Direction) VARARG_EXTRA(const float Length) VARARG_EXTRA(const float Angle)  VARARG_EXTRA(const FColor& Color))
+VARARG_BODY(void, FVisualLogger::GeometryShapeLogf, const TCHAR*, VARARG_EXTRA(const UObject* Object) VARARG_EXTRA(const FLogCategoryBase& Category) VARARG_EXTRA(ELogVerbosity::Type Verbosity) VARARG_EXTRA(const FVector& Orgin) VARARG_EXTRA(const FVector& Direction) VARARG_EXTRA(const float Length) VARARG_EXTRA(const float Angle)  VARARG_EXTRA(const FColor& Color))
 {
 	SCOPE_CYCLE_COUNTER(STAT_VisualLog);
 	UWorld *World = nullptr;
@@ -426,7 +426,7 @@ VARARG_BODY(void, FVisualLogger::GeometryShapeLogf, const TCHAR*, VARARG_EXTRA(c
 }
 
 inline
-VARARG_BODY(void, FVisualLogger::GeometryShapeLogf, const TCHAR*, VARARG_EXTRA(const UObject* Object) VARARG_EXTRA(const FLogCategoryBase& Category) VARARG_EXTRA(ELogVerbosity::Type Verbosity) VARARG_EXTRA(int32 UniqueLogId) VARARG_EXTRA(const FVector& Start) VARARG_EXTRA(const FVector& End) VARARG_EXTRA(const float Radius) VARARG_EXTRA(const FColor& Color))
+VARARG_BODY(void, FVisualLogger::GeometryShapeLogf, const TCHAR*, VARARG_EXTRA(const UObject* Object) VARARG_EXTRA(const FLogCategoryBase& Category) VARARG_EXTRA(ELogVerbosity::Type Verbosity) VARARG_EXTRA(const FVector& Start) VARARG_EXTRA(const FVector& End) VARARG_EXTRA(const float Radius) VARARG_EXTRA(const FColor& Color))
 {
 	SCOPE_CYCLE_COUNTER(STAT_VisualLog);
 	UWorld *World = nullptr;
@@ -442,7 +442,7 @@ VARARG_BODY(void, FVisualLogger::GeometryShapeLogf, const TCHAR*, VARARG_EXTRA(c
 }
 
 inline
-VARARG_BODY(void, FVisualLogger::GeometryShapeLogf, const TCHAR*, VARARG_EXTRA(const UObject* Object) VARARG_EXTRA(const FLogCategoryBase& Category) VARARG_EXTRA(ELogVerbosity::Type Verbosity) VARARG_EXTRA(int32 UniqueLogId) VARARG_EXTRA(const FVector& Center) VARARG_EXTRA(float HalfHeight) VARARG_EXTRA(float Radius) VARARG_EXTRA(const FQuat& Rotation) VARARG_EXTRA(const FColor& Color))
+VARARG_BODY(void, FVisualLogger::GeometryShapeLogf, const TCHAR*, VARARG_EXTRA(const UObject* Object) VARARG_EXTRA(const FLogCategoryBase& Category) VARARG_EXTRA(ELogVerbosity::Type Verbosity) VARARG_EXTRA(const FVector& Center) VARARG_EXTRA(float HalfHeight) VARARG_EXTRA(float Radius) VARARG_EXTRA(const FQuat& Rotation) VARARG_EXTRA(const FColor& Color))
 {
 	SCOPE_CYCLE_COUNTER(STAT_VisualLog);
 	UWorld *World = nullptr;
@@ -458,7 +458,7 @@ VARARG_BODY(void, FVisualLogger::GeometryShapeLogf, const TCHAR*, VARARG_EXTRA(c
 }
 
 inline
-VARARG_BODY(void, FVisualLogger::NavAreaShapeLogf, const TCHAR*, VARARG_EXTRA(const UObject* Object) VARARG_EXTRA(const FLogCategoryBase& Category) VARARG_EXTRA(ELogVerbosity::Type Verbosity) VARARG_EXTRA(int32 UniqueLogId) VARARG_EXTRA(const TArray<FVector> &ConvexPoints) VARARG_EXTRA(float MinZ) VARARG_EXTRA(float MaxZ) VARARG_EXTRA(const FColor& Color))
+VARARG_BODY(void, FVisualLogger::NavAreaShapeLogf, const TCHAR*, VARARG_EXTRA(const UObject* Object) VARARG_EXTRA(const FLogCategoryBase& Category) VARARG_EXTRA(ELogVerbosity::Type Verbosity) VARARG_EXTRA(const TArray<FVector> &ConvexPoints) VARARG_EXTRA(float MinZ) VARARG_EXTRA(float MaxZ) VARARG_EXTRA(const FColor& Color))
 {
 	SCOPE_CYCLE_COUNTER(STAT_VisualLog);
 	UWorld *World = nullptr;
@@ -474,7 +474,7 @@ VARARG_BODY(void, FVisualLogger::NavAreaShapeLogf, const TCHAR*, VARARG_EXTRA(co
 }
 
 inline
-VARARG_BODY(void, FVisualLogger::GeometryShapeLogf, const TCHAR*, VARARG_EXTRA(const UObject* Object) VARARG_EXTRA(const FLogCategoryBase& Category) VARARG_EXTRA(ELogVerbosity::Type Verbosity) VARARG_EXTRA(int32 UniqueLogId) VARARG_EXTRA(const TArray<FVector> &Vertices) VARARG_EXTRA(const TArray<int32> &Indices) VARARG_EXTRA(const FColor& Color))
+VARARG_BODY(void, FVisualLogger::GeometryShapeLogf, const TCHAR*, VARARG_EXTRA(const UObject* Object) VARARG_EXTRA(const FLogCategoryBase& Category) VARARG_EXTRA(ELogVerbosity::Type Verbosity) VARARG_EXTRA(const TArray<FVector> &Vertices) VARARG_EXTRA(const TArray<int32> &Indices) VARARG_EXTRA(const FColor& Color))
 {
 	SCOPE_CYCLE_COUNTER(STAT_VisualLog);
 	UWorld *World = nullptr;
@@ -490,7 +490,7 @@ VARARG_BODY(void, FVisualLogger::GeometryShapeLogf, const TCHAR*, VARARG_EXTRA(c
 }
 
 inline
-VARARG_BODY(void, FVisualLogger::GeometryConvexLogf, const TCHAR*, VARARG_EXTRA(const UObject* Object) VARARG_EXTRA(const FLogCategoryBase& Category) VARARG_EXTRA(ELogVerbosity::Type Verbosity) VARARG_EXTRA(int32 UniqueLogId) VARARG_EXTRA(const TArray<FVector> &Points) VARARG_EXTRA(const FColor& Color))
+VARARG_BODY(void, FVisualLogger::GeometryConvexLogf, const TCHAR*, VARARG_EXTRA(const UObject* Object) VARARG_EXTRA(const FLogCategoryBase& Category) VARARG_EXTRA(ELogVerbosity::Type Verbosity) VARARG_EXTRA(const TArray<FVector> &Points) VARARG_EXTRA(const FColor& Color))
 {
 	SCOPE_CYCLE_COUNTER(STAT_VisualLog);
 	UWorld *World = nullptr;
@@ -507,7 +507,7 @@ VARARG_BODY(void, FVisualLogger::GeometryConvexLogf, const TCHAR*, VARARG_EXTRA(
 
 
 inline
-VARARG_BODY(void, FVisualLogger::HistogramDataLogf, const TCHAR*, VARARG_EXTRA(const UObject* Object) VARARG_EXTRA(const FLogCategoryBase& Category) VARARG_EXTRA(ELogVerbosity::Type Verbosity) VARARG_EXTRA(int32 UniqueLogId) VARARG_EXTRA(FName GraphName) VARARG_EXTRA(FName DataName) VARARG_EXTRA(const FVector2D& Data) VARARG_EXTRA(const FColor& Color))
+VARARG_BODY(void, FVisualLogger::HistogramDataLogf, const TCHAR*, VARARG_EXTRA(const UObject* Object) VARARG_EXTRA(const FLogCategoryBase& Category) VARARG_EXTRA(ELogVerbosity::Type Verbosity) VARARG_EXTRA(FName GraphName) VARARG_EXTRA(FName DataName) VARARG_EXTRA(const FVector2D& Data) VARARG_EXTRA(const FColor& Color))
 {
 	SCOPE_CYCLE_COUNTER(STAT_VisualLog);
 	UWorld *World = nullptr;
diff --git a/Engine/Source/Runtime/Engine/Public/VisualLogger/VisualLoggerTypes.h b/Engine/Source/Runtime/Engine/Public/VisualLogger/VisualLoggerTypes.h
--- a/Engine/Source/Runtime/Engine/Public/VisualLogger/VisualLoggerTypes.h
+++ b/Engine/Source/Runtime/Engine/Public/VisualLogger/VisualLoggerTypes.h
@@ -110,7 +110,8 @@ struct ENGINE_API FVisualLogStatusCategory
 	int32 UniqueId;
 	TArray<FVisualLogStatusCategory> Children;
 
-	FVisualLogStatusCategory(const FString& InCategory = TEXT(""))
+	FVisualLogStatusCategory();
+	explicit FVisualLogStatusCategory(const FString& InCategory/* = TEXT("")*/)
 		: Category(InCategory)
 	{
 	}
@@ -364,16 +365,7 @@ void FVisualLogStatusCategory::AddChild(const FVisualLogStatusCategory& Child)
 	Children.Add(Child);
 }
 
-inline
-FVisualLogShapeElement::FVisualLogShapeElement(EVisualLoggerShapeElement InType)
-: Verbosity(ELogVerbosity::All)
-, TransformationMatrix(FMatrix::Identity)
-, Type(InType)
-, Color(0xff)
-, Thicknes(0)
-{
 
-}
 
 inline
 FVisualLogShapeElement::FVisualLogShapeElement(const FString& InDescription, const FColor& InColor, uint16 InThickness, const FName& InCategory)
diff --git a/Engine/Source/Runtime/Engine/Public/WorldCollision.h b/Engine/Source/Runtime/Engine/Public/WorldCollision.h
--- a/Engine/Source/Runtime/Engine/Public/WorldCollision.h
+++ b/Engine/Source/Runtime/Engine/Public/WorldCollision.h
@@ -324,6 +324,16 @@ DECLARE_DELEGATE_TwoParams( FTraceDelegate, const FTraceHandle&, FTraceDatum &);
  */
 DECLARE_DELEGATE_TwoParams( FOverlapDelegate, const FTraceHandle&, FOverlapDatum &);
 
+/** Enum to indicate type of test to perfom */
+enum class EAsyncTraceType : uint8
+{
+	/** Return whether the trace succeeds or fails (using bBlockingHit flag on FHitResult), but gives no info about what you hit or where. Good for fast occlusion tests. */
+	Test,
+	/** Returns a single blocking hit */
+	Single,
+	/** Returns a single blocking hit, plus any overlapping hits up to that point */
+	Multi
+};
 
 /**
  * Trace/Sweep Data structure for async trace
@@ -346,14 +356,14 @@ struct FTraceDatum : public FBaseTraceDatum
 	/** Output of the overlap request. Filled up by worker thread */
 	TArray<struct FHitResult> OutHits;
 
-	/** single or multi trace. */
-	bool 	bIsMultiTrace;
+	/** Whether to do test, single or multi test */
+	EAsyncTraceType TraceType;
 
 	FTraceDatum() {}
 
 	/** Set Trace Datum for each shape type **/
 	FTraceDatum(UWorld * World, const FCollisionShape& CollisionShape, const FCollisionQueryParams& Param, const struct FCollisionResponseParams &InResponseParam, const struct FCollisionObjectQueryParams& InObjectQueryParam, 
-		ECollisionChannel Channel, uint32 InUserData, bool bInIsMultiTrace,	const FVector& InStart, const FVector& InEnd, FTraceDelegate * InDelegate, int32 FrameCounter)
+		ECollisionChannel Channel, uint32 InUserData, EAsyncTraceType InTraceType,	const FVector& InStart, const FVector& InEnd, FTraceDelegate * InDelegate, int32 FrameCounter)
 	{
 		Set(World, CollisionShape, Param, InResponseParam, InObjectQueryParam, Channel, InUserData, FrameCounter);
 		Start = InStart;
@@ -367,7 +377,7 @@ struct FTraceDatum : public FBaseTraceDatum
 			Delegate.Unbind();
 		}
 		OutHits.Reset();
-		bIsMultiTrace = bInIsMultiTrace;
+		TraceType = InTraceType;
 	}
 };
 
diff --git a/Engine/Source/Runtime/GameplayAbilities/GameplayAbilities.Build.cs b/Engine/Source/Runtime/GameplayAbilities/GameplayAbilities.Build.cs
--- a/Engine/Source/Runtime/GameplayAbilities/GameplayAbilities.Build.cs
+++ b/Engine/Source/Runtime/GameplayAbilities/GameplayAbilities.Build.cs
@@ -52,6 +52,12 @@ namespace UnrealBuildTool.Rules
                 PrivateDependencyModuleNames.Add("GameplayTagsEditor");
                 PrivateDependencyModuleNames.Add("Slate");
 			}
-		}
+
+            if (UEBuildConfiguration.bBuildDeveloperTools && Target.Configuration != UnrealTargetConfiguration.Shipping && Target.Configuration != UnrealTargetConfiguration.Test)
+            {
+                PrivateDependencyModuleNames.Add("GameplayDebugger");
+                Definitions.Add("WITH_GAMEPLAY_DEBUGGER=1");
+            }
+        }
 	}
 }
\ No newline at end of file
diff --git a/Engine/Source/Runtime/GameplayAbilities/Private/GameplayAbilitiesModule.cpp b/Engine/Source/Runtime/GameplayAbilities/Private/GameplayAbilitiesModule.cpp
--- a/Engine/Source/Runtime/GameplayAbilities/Private/GameplayAbilitiesModule.cpp
+++ b/Engine/Source/Runtime/GameplayAbilities/Private/GameplayAbilitiesModule.cpp
@@ -4,6 +4,11 @@
 #include "AbilitySystemGlobals.h"
 #include "AbilitySystemComponent.h"
 
+#if WITH_GAMEPLAY_DEBUGGER
+#include "GameplayDebugger.h"
+#include "GameplayDebuggerCategory_Abilities.h"
+#endif // WITH_GAMEPLAY_DEBUGGER
+
 class FGameplayAbilitiesModule : public IGameplayAbilitiesModule
 {
 	// Begin IModuleInterface
@@ -39,8 +44,6 @@ class FGameplayAbilitiesModule : public IGameplayAbilitiesModule
 
 	UAbilitySystemGlobals *AbilitySystemGlobals;
 
-	void GetActiveAbilitiesDebugDataForActor(AActor* Actor, FString& AbilityString, bool& bIsUsingAbilities);
-
 private:
 	
 };
@@ -51,6 +54,12 @@ void FGameplayAbilitiesModule::StartupModule()
 {	
 	// This is loaded upon first request
 	AbilitySystemGlobals = NULL;
+
+#if WITH_GAMEPLAY_DEBUGGER
+	IGameplayDebugger& GameplayDebuggerModule = IGameplayDebugger::Get();
+	GameplayDebuggerModule.RegisterCategory("Abilities", IGameplayDebugger::FOnGetCategory::CreateStatic(&FGameplayDebuggerCategory_Abilities::MakeInstance));
+	GameplayDebuggerModule.NotifyCategoriesChanged();
+#endif // WITH_GAMEPLAY_DEBUGGER
 }
 
 void FGameplayAbilitiesModule::ShutdownModule()
@@ -59,39 +68,13 @@ void FGameplayAbilitiesModule::ShutdownModule()
 	// we call this function before unloading the module.
 
 	AbilitySystemGlobals = NULL;
-}
 
-void FGameplayAbilitiesModule::GetActiveAbilitiesDebugDataForActor(AActor* Actor, FString& AbilityString, bool& bIsUsingAbilities)
-{
-	UAbilitySystemComponent* AbilityComp = Actor->FindComponentByClass<UAbilitySystemComponent>();
-	bIsUsingAbilities = (AbilityComp != nullptr);
-	  
-	int32 NumActive = 0;
-	if (AbilityComp)
-	{
-		AbilityString = TEXT("");
-	  
-		for (const FGameplayAbilitySpec& AbilitySpec : AbilityComp->GetActivatableAbilities())
-		{
-	  		if (AbilitySpec.Ability && AbilitySpec.IsActive())
-	  		{
-	  			if (NumActive)
-	  			{
-					AbilityString += TEXT(", ");
-	  			}
-	  
-	  			UClass* AbilityClass = AbilitySpec.Ability->GetClass();
-	  			FString AbClassName = GetNameSafe(AbilityClass);
-	  			AbClassName.RemoveFromEnd(TEXT("_c"));
-	  
-				AbilityString += AbClassName;
-	  			NumActive++;
-	  		}
-		}
-	}
-	  
-	if (NumActive == 0)
+#if WITH_GAMEPLAY_DEBUGGER
+	if (IGameplayDebugger::IsAvailable())
 	{
-		AbilityString = TEXT("None");
+		IGameplayDebugger& GameplayDebuggerModule = IGameplayDebugger::Get();
+		GameplayDebuggerModule.UnregisterCategory("Abilities");
+		GameplayDebuggerModule.NotifyCategoriesChanged();
 	}
+#endif // WITH_GAMEPLAY_DEBUGGER
 }
diff --git a/Engine/Source/Runtime/GameplayAbilities/Private/GameplayDebuggerCategory_Abilities.cpp b/Engine/Source/Runtime/GameplayAbilities/Private/GameplayDebuggerCategory_Abilities.cpp
new file mode 100644
--- /dev/null
+++ b/Engine/Source/Runtime/GameplayAbilities/Private/GameplayDebuggerCategory_Abilities.cpp
@@ -0,0 +1,145 @@
+// Copyright 1998-2016 Epic Games, Inc. All Rights Reserved.
+#include "AbilitySystemPrivatePCH.h"
+#include "AbilitySystemComponent.h"
+#include "GameplayDebuggerCategory_Abilities.h"
+
+#if WITH_GAMEPLAY_DEBUGGER
+
+FGameplayDebuggerCategory_Abilities::FGameplayDebuggerCategory_Abilities()
+{
+	SetDataPackReplication<FRepData>(&DataPack);
+}
+
+TSharedRef<FGameplayDebuggerCategory> FGameplayDebuggerCategory_Abilities::MakeInstance()
+{
+	return MakeShareable(new FGameplayDebuggerCategory_Abilities());
+}
+
+void FGameplayDebuggerCategory_Abilities::FRepData::Serialize(FArchive& Ar)
+{
+	Ar << OwnedTags;
+
+	int32 NumAbilities = Abilities.Num();
+	Ar << NumAbilities;
+	if (Ar.IsLoading())
+	{
+		Abilities.SetNum(NumAbilities);
+	}
+
+	for (int32 Idx = 0; Idx < NumAbilities; Idx++)
+	{
+		Ar << Abilities[Idx].Ability;
+		Ar << Abilities[Idx].Source;
+		Ar << Abilities[Idx].Level;
+		Ar << Abilities[Idx].bIsActive;
+	}
+
+	int32 NumGE = GameplayEffects.Num();
+	Ar << NumGE;
+	if (Ar.IsLoading())
+	{
+		GameplayEffects.SetNum(NumGE);
+	}
+
+	for (int32 Idx = 0; Idx < NumGE; Idx++)
+	{
+		Ar << GameplayEffects[Idx].Effect;
+		Ar << GameplayEffects[Idx].Context;
+		Ar << GameplayEffects[Idx].Duration;
+		Ar << GameplayEffects[Idx].Period;
+		Ar << GameplayEffects[Idx].Stacks;
+		Ar << GameplayEffects[Idx].Level;
+	}
+}
+
+void FGameplayDebuggerCategory_Abilities::CollectData(APlayerController* OwnerPC, AActor* DebugActor)
+{
+	UAbilitySystemComponent* AbilityComp = DebugActor ? DebugActor->FindComponentByClass<UAbilitySystemComponent>() : nullptr;
+	if (AbilityComp)
+	{
+		FGameplayTagContainer OwnerTags;
+		AbilityComp->GetOwnedGameplayTags(OwnerTags);
+		DataPack.OwnedTags = OwnerTags.ToStringSimple();
+
+		TArray<FGameplayEffectSpec> ActiveEffectSpecs;
+		AbilityComp->GetAllActiveGameplayEffectSpecs(ActiveEffectSpecs);
+		for (int32 Idx = 0; Idx < ActiveEffectSpecs.Num(); Idx++)
+		{
+			const FGameplayEffectSpec& EffectSpec = ActiveEffectSpecs[Idx];
+			FRepData::FGameplayEffectDebug ItemData;
+
+			ItemData.Effect = EffectSpec.ToSimpleString();
+			ItemData.Effect.RemoveFromStart(DEFAULT_OBJECT_PREFIX);
+			ItemData.Effect.RemoveFromEnd(TEXT("_C"));
+
+			ItemData.Context = EffectSpec.GetContext().ToString();
+			ItemData.Duration = EffectSpec.GetDuration();
+			ItemData.Period = EffectSpec.GetPeriod();
+			ItemData.Stacks = EffectSpec.StackCount;
+			ItemData.Level = EffectSpec.GetLevel();
+
+			DataPack.GameplayEffects.Add(ItemData);
+		}
+
+		const TArray<FGameplayAbilitySpec>& AbilitySpecs = AbilityComp->GetActivatableAbilities();
+		for (int32 Idx = 0; Idx < AbilitySpecs.Num(); Idx++)
+		{
+			const FGameplayAbilitySpec& AbilitySpec = AbilitySpecs[Idx];
+			FRepData::FGameplayAbilityDebug ItemData;
+
+			ItemData.Ability = GetNameSafe(AbilitySpec.Ability);
+			ItemData.Ability.RemoveFromStart(DEFAULT_OBJECT_PREFIX);
+			ItemData.Ability.RemoveFromEnd(TEXT("_C"));
+
+			ItemData.Source = GetNameSafe(AbilitySpec.SourceObject);
+			ItemData.Source.RemoveFromStart(DEFAULT_OBJECT_PREFIX);
+
+			ItemData.Level = AbilitySpec.Level;
+			ItemData.bIsActive = AbilitySpec.IsActive();
+
+			DataPack.Abilities.Add(ItemData);
+		}
+	}
+}
+
+void FGameplayDebuggerCategory_Abilities::DrawData(APlayerController* OwnerPC, FGameplayDebuggerCanvasContext& CanvasContext)
+{
+	CanvasContext.Printf(TEXT("Owned Tags: {yellow}%s"), *DataPack.OwnedTags);
+
+	CanvasContext.Printf(TEXT("Gameplay Effects: {yellow}%d"), DataPack.GameplayEffects.Num());
+	for (int32 Idx = 0; Idx < DataPack.GameplayEffects.Num(); Idx++)
+	{
+		const FRepData::FGameplayEffectDebug& ItemData = DataPack.GameplayEffects[Idx];
+
+		FString Desc = FString::Printf(TEXT("\t{yellow}%s {grey}source:{white}%s {grey}duration:{white}"), *ItemData.Effect, *ItemData.Context);
+		Desc += (ItemData.Duration > 0.0f) ? FString::Printf(TEXT("%.2f"), ItemData.Duration) : FString(TEXT("INF"));
+
+		if (ItemData.Period > 0.0f)
+		{
+			Desc += FString::Printf(TEXT(" {grey}period:{white}%.2f"), ItemData.Period);
+		}
+
+		if (ItemData.Stacks > 1)
+		{
+			Desc += FString::Printf(TEXT(" {grey}stacks:{white}%d"), ItemData.Stacks);
+		}
+
+		if (ItemData.Level > 1.0f)
+		{
+			Desc += FString::Printf(TEXT(" {grey}level:{white}%.2f"), ItemData.Level);
+		}
+
+		CanvasContext.Print(Desc);
+	}
+
+	CanvasContext.Printf(TEXT("Gameplay Abilities: {yellow}%d"), DataPack.Abilities.Num());
+	for (int32 Idx = 0; Idx < DataPack.Abilities.Num(); Idx++)
+	{
+		const FRepData::FGameplayAbilityDebug& ItemData = DataPack.Abilities[Idx];
+
+		CanvasContext.Printf(TEXT("\t{yellow}%s {grey}source:{white}%s {grey}level:{white}%d {grey}active:{white}%s"),
+			*ItemData.Ability, *ItemData.Source, ItemData.Level, ItemData.bIsActive ? TEXT("YES") : TEXT("no"));
+	}
+}
+
+#endif // WITH_GAMEPLAY_DEBUGGER
diff --git a/Engine/Source/Runtime/GameplayAbilities/Private/GameplayDebuggerCategory_Abilities.h b/Engine/Source/Runtime/GameplayAbilities/Private/GameplayDebuggerCategory_Abilities.h
new file mode 100644
--- /dev/null
+++ b/Engine/Source/Runtime/GameplayAbilities/Private/GameplayDebuggerCategory_Abilities.h
@@ -0,0 +1,48 @@
+// Copyright 1998-2016 Epic Games, Inc. All Rights Reserved.
+
+#if WITH_GAMEPLAY_DEBUGGER
+
+#pragma once
+#include "GameplayDebuggerCategory.h"
+
+class FGameplayDebuggerCategory_Abilities : public FGameplayDebuggerCategory
+{
+public:
+	FGameplayDebuggerCategory_Abilities();
+
+	virtual void CollectData(APlayerController* OwnerPC, AActor* DebugActor) override;
+	virtual void DrawData(APlayerController* OwnerPC, FGameplayDebuggerCanvasContext& CanvasContext) override;
+
+	static TSharedRef<FGameplayDebuggerCategory> MakeInstance();
+
+protected:
+	struct FRepData
+	{
+		FString OwnedTags;
+
+		struct FGameplayAbilityDebug
+		{
+			FString Ability;
+			FString Source;
+			int32 Level;
+			bool bIsActive;
+		};
+		TArray<FGameplayAbilityDebug> Abilities;
+
+		struct FGameplayEffectDebug
+		{
+			FString Effect;
+			FString Context;
+			float Duration;
+			float Period;
+			int32 Stacks;
+			float Level;
+		};
+		TArray<FGameplayEffectDebug> GameplayEffects;
+
+		void Serialize(FArchive& Ar);
+	};
+	FRepData DataPack;
+};
+
+#endif // WITH_GAMEPLAY_DEBUGGER
diff --git a/Engine/Source/Runtime/GameplayAbilities/Public/GameplayAbilitiesModule.h b/Engine/Source/Runtime/GameplayAbilities/Public/GameplayAbilitiesModule.h
--- a/Engine/Source/Runtime/GameplayAbilities/Public/GameplayAbilitiesModule.h
+++ b/Engine/Source/Runtime/GameplayAbilities/Public/GameplayAbilitiesModule.h
@@ -41,7 +41,5 @@ public:
 	virtual UAbilitySystemGlobals* GetAbilitySystemGlobals() = 0;
 
 	virtual bool IsAbilitySystemGlobalsAvailable() = 0;
-
-	virtual void GetActiveAbilitiesDebugDataForActor(AActor* Actor, FString& AbilityString, bool& bIsUsingAbilities)=0;
 };
 
diff --git a/Engine/Source/Runtime/Windows/XAudio2/Private/XAudio2Buffer.cpp b/Engine/Source/Runtime/Windows/XAudio2/Private/XAudio2Buffer.cpp
--- a/Engine/Source/Runtime/Windows/XAudio2/Private/XAudio2Buffer.cpp
+++ b/Engine/Source/Runtime/Windows/XAudio2/Private/XAudio2Buffer.cpp
@@ -80,10 +80,11 @@ struct FXMAInfo
  * @param AudioDevice	audio device this sound buffer is going to be attached to.
  */
 FXAudio2SoundBuffer::FXAudio2SoundBuffer( FAudioDevice* InAudioDevice, ESoundFormat InSoundFormat )
-:	FSoundBuffer(InAudioDevice),
-	SoundFormat( InSoundFormat ),
-	DecompressionState( NULL ),
-	bDynamicResource( false )
+	: FSoundBuffer(InAudioDevice)
+	, SoundFormat(InSoundFormat)
+	, DecompressionState(nullptr)
+	, RealtimeAsyncHeaderParseTask(nullptr)
+	, bDynamicResource(false)
 {
 	PCM.PCMData = NULL;
 	PCM.PCMDataSize = 0;
@@ -216,6 +217,23 @@ int32 FXAudio2SoundBuffer::GetCurrentChunkOffset() const
 	return DecompressionState->GetCurrentChunkOffset();
 }
 
+bool FXAudio2SoundBuffer::IsRealTimeSourceReady()
+{
+	// If we have a realtime async header parse task, then we check if its done
+	if (RealtimeAsyncHeaderParseTask)
+	{
+		bool bIsDone = RealtimeAsyncHeaderParseTask->IsDone();
+		if (bIsDone)
+		{
+			delete RealtimeAsyncHeaderParseTask;
+			RealtimeAsyncHeaderParseTask = nullptr;
+		}
+		return bIsDone;
+	}
+	// Otherwise, we weren't a real time decoding sound buffer (or we've already asked and it was ready)
+	return true;
+}
+
 /** 
  * Setup a WAVEFORMATEX structure
  */
@@ -306,6 +324,16 @@ void FXAudio2SoundBuffer::InitXWMA( USoundWave* Wave, FXMAInfo* XMAInfo )
 #endif	//XAUDIO_SUPPORTS_XMA2WAVEFORMATEX
 }
 
+bool FXAudio2SoundBuffer::ReadCompressedInfo(USoundWave* SoundWave)
+{
+	if (!DecompressionState)
+	{
+		UE_LOG(LogXAudio2, Warning, TEXT("Attempting to read compressed info without a compression state instance for resource '%s'"), *ResourceName);
+		return false;
+	}
+	return DecompressionState->ReadCompressedInfo(SoundWave->ResourceData, SoundWave->ResourceSize, nullptr);
+}
+
 /**
  * Decompresses a chunk of compressed audio to the destination memory
  *
@@ -371,13 +399,16 @@ FXAudio2SoundBuffer* FXAudio2SoundBuffer::CreateQueuedBuffer( FXAudio2Device* XA
 		Wave->InitAudioResource(XAudio2Device->GetRuntimeFormat(Wave));
 	}
 
-	// Prime the first two buffers and prepare the decompression
-	FSoundQualityInfo QualityInfo = { 0 };
-
 	Buffer->DecompressionState = XAudio2Device->CreateCompressedAudioInfo(Wave);
 
-	if (Buffer->DecompressionState && Buffer->DecompressionState->ReadCompressedInfo(Wave->ResourceData, Wave->ResourceSize, &QualityInfo))
+	if (Buffer->DecompressionState && Buffer->DecompressionState->ReadCompressedInfo(Wave->ResourceData, Wave->ResourceSize, nullptr))
 	{
+		// Start the async task that parses the decompressed asset header info
+		// Doing this step synchronously causes huge main-thread hitches for large ogg-vorbis files
+		check(Buffer->RealtimeAsyncHeaderParseTask == nullptr);
+		Buffer->RealtimeAsyncHeaderParseTask = new FAsyncRealtimeAudioTask(Buffer, Wave);
+		Buffer->RealtimeAsyncHeaderParseTask->StartBackgroundTask();
+
 		// Clear out any dangling pointers
 		Buffer->PCM.PCMData = NULL;
 		Buffer->PCM.PCMDataSize = 0;
diff --git a/Engine/Source/Runtime/Windows/XAudio2/Private/XAudio2Source.cpp b/Engine/Source/Runtime/Windows/XAudio2/Private/XAudio2Source.cpp
--- a/Engine/Source/Runtime/Windows/XAudio2/Private/XAudio2Source.cpp
+++ b/Engine/Source/Runtime/Windows/XAudio2/Private/XAudio2Source.cpp
@@ -34,15 +34,19 @@ FXMPHelper* FXMPHelper::GetXMPHelper( void )
  * Simple constructor
  */
 FXAudio2SoundSource::FXAudio2SoundSource(FAudioDevice* InAudioDevice)
-	: FSoundSource( InAudioDevice )
-	, Source( NULL )
+	: FSoundSource(InAudioDevice)
+	, XAudio2Buffer(nullptr)
+	, Source(nullptr)
 	, MaxEffectChainChannels(0)
-	, RealtimeAsyncTask( nullptr )
-	, CurrentBuffer( 0 )
-	, bBuffersToFlush( false )
-	, bLoopCallback( false )
-	, bResourcesNeedFreeing(false)
+	, RealtimeAsyncTask(nullptr)
 	, VoiceId(-1)
+	, CurrentBuffer(0)
+	, bLoopCallback(false)
+	, bIsFinished(false)
+	, bPlayedCachedBuffer(false)
+	, bFirstRTBuffersSubmitted(false)
+	, bBuffersToFlush(false)
+	, bResourcesNeedFreeing(false)
 	, bUsingHRTFSpatialization(false)
 	, bEditorWarnedChangedSpatialization(false)
 {
@@ -85,6 +89,7 @@ void FXAudio2SoundSource::InitializeSourceEffects(uint32 InVoiceId)
  */
 void FXAudio2SoundSource::FreeResources( void )
 {
+	// Clean up the realtime async decoding task
 	if (RealtimeAsyncTask)
 	{
 		RealtimeAsyncTask->EnsureCompletion();
@@ -92,6 +97,14 @@ void FXAudio2SoundSource::FreeResources( void )
 		RealtimeAsyncTask = nullptr;
 	}
 
+	// Clean the realtime async header parsing task
+	if (XAudio2Buffer && XAudio2Buffer->RealtimeAsyncHeaderParseTask)
+	{
+		XAudio2Buffer->RealtimeAsyncHeaderParseTask->EnsureCompletion();
+		delete XAudio2Buffer->RealtimeAsyncHeaderParseTask;
+		XAudio2Buffer->RealtimeAsyncHeaderParseTask = nullptr;
+	}
+
 	// Release voice.
 	if (Source)
 	{
@@ -176,6 +189,7 @@ bool FXAudio2SoundSource::ReadMorePCMData( const int32 BufferIndex, EDataReadMod
 		}
 		else
 		{
+			check(!RealtimeAsyncTask);
 			RealtimeAsyncTask = new FAsyncRealtimeAudioTask(WaveData, (uint8*)XAudio2Buffers[BufferIndex].pAudioData, MaxSamples);
 			RealtimeAsyncTask->StartBackgroundTask();
 		}
@@ -191,7 +205,8 @@ bool FXAudio2SoundSource::ReadMorePCMData( const int32 BufferIndex, EDataReadMod
 		}
 		else
 		{
-			RealtimeAsyncTask = new FAsyncRealtimeAudioTask(XAudio2Buffer, ( uint8* )XAudio2Buffers[BufferIndex].pAudioData, WaveInstance->LoopingMode != LOOP_Never, DataReadMode == EDataReadMode::AsynchronousSkipFirstFrame);
+			check(!RealtimeAsyncTask);
+			RealtimeAsyncTask = new FAsyncRealtimeAudioTask(XAudio2Buffer, (uint8*)XAudio2Buffers[BufferIndex].pAudioData, WaveInstance->LoopingMode != LOOP_Never, DataReadMode == EDataReadMode::AsynchronousSkipFirstFrame);
 			RealtimeAsyncTask->StartBackgroundTask();
 			return false;
 		}
@@ -215,18 +230,20 @@ void FXAudio2SoundSource::SubmitPCMRTBuffers( void )
 	// Set up buffer areas to decompress to
 	XAudio2Buffers[0].pAudioData = (uint8*)FMemory::Malloc(BufferSize);
 	XAudio2Buffers[0].AudioBytes = BufferSize;
+	XAudio2Buffers[0].pContext = this;
 
 	XAudio2Buffers[1].pAudioData = (uint8*)FMemory::Malloc(BufferSize);
 	XAudio2Buffers[1].AudioBytes = BufferSize;
+	XAudio2Buffers[1].pContext = this;
 
 	// Only use the cached data if we're starting from the beginning, otherwise we'll have to take a synchronous hit
-	bool bSkipFirstBuffer = false;
+	bPlayedCachedBuffer = false;
 	bool bIsSeeking = (WaveInstance->StartTime > 0.f);
 	if (WaveInstance->WaveData && WaveInstance->WaveData->CachedRealtimeFirstBuffer && !bIsSeeking)
 	{
+		bPlayedCachedBuffer = true;
 		FMemory::Memcpy((uint8*)XAudio2Buffers[0].pAudioData, WaveInstance->WaveData->CachedRealtimeFirstBuffer, BufferSize);
 		FMemory::Memcpy((uint8*)XAudio2Buffers[1].pAudioData, WaveInstance->WaveData->CachedRealtimeFirstBuffer + BufferSize, BufferSize);
-		bSkipFirstBuffer = true;
 	}
 	else
 	{
@@ -234,35 +251,21 @@ void FXAudio2SoundSource::SubmitPCMRTBuffers( void )
 		ReadMorePCMData(1, EDataReadMode::Synchronous);
 	}
 
-	AudioDevice->ValidateAPICall( TEXT( "SubmitSourceBuffer - PCMRT" ), 
-		Source->SubmitSourceBuffer( &XAudio2Buffers[0] ) );
+	// Immediately submit the first two buffers that were either cached or synchronously read
+	// The first buffer will start the voice processing buffers and trigger an OnBufferEnd callback, which will then 
+	// trigger async tasks to generate more PCMRT buffers.
+	AudioDevice->ValidateAPICall(TEXT("SubmitSourceBuffer - PCMRT"),
+								 Source->SubmitSourceBuffer(&XAudio2Buffers[0]));
 
-	AudioDevice->ValidateAPICall( TEXT( "SubmitSourceBuffer - PCMRT" ), 
-		Source->SubmitSourceBuffer( &XAudio2Buffers[1] ) );
-
-	XAudio2Buffers[2].pAudioData = (uint8*)FMemory::Malloc(BufferSize);
-	XAudio2Buffers[2].AudioBytes = BufferSize;
+	AudioDevice->ValidateAPICall(TEXT("SubmitSourceBuffer - PCMRT"),
+								 Source->SubmitSourceBuffer(&XAudio2Buffers[1]));
 
+	// Prepare the third buffer for the OnBufferEnd callback to write to in the OnBufferEnd callback
 	CurrentBuffer = 2;
 
-	// Start the async population of the next buffer
-	EDataReadMode DataReadMode = EDataReadMode::Asynchronous;
-	if (XAudio2Buffer->SoundFormat == ESoundFormat::SoundFormat_Streaming)
-	{
-		DataReadMode =  EDataReadMode::Synchronous;
-	}
-	else if (bSkipFirstBuffer)
-	{
-		DataReadMode =  EDataReadMode::AsynchronousSkipFirstFrame;
-	}
-
-	ReadMorePCMData(2, DataReadMode);
-
-	if (!bIsSeeking && (DataReadMode == EDataReadMode::Synchronous || (!bSkipFirstBuffer && WaveInstance->WaveData && !WaveInstance->WaveData->bCanProcessAsync)))
-	{
-		AudioDevice->ValidateAPICall(TEXT("SubmitSourceBuffer - PCMRT"),
-									 Source->SubmitSourceBuffer(&XAudio2Buffers[2]));
-	}
+	XAudio2Buffers[2].pAudioData = (uint8*)FMemory::Malloc(BufferSize);
+	XAudio2Buffers[2].AudioBytes = BufferSize;
+	XAudio2Buffers[2].pContext = this;
 
 	bResourcesNeedFreeing = true;
 }
@@ -445,87 +448,116 @@ bool FXAudio2SoundSource::CreateSource( void )
 	return true;
 }
 
-/**
- * Initializes a source with a given wave instance and prepares it for playback.
- *
- * @param	WaveInstance	wave instace being primed for playback
- * @return	true if initialization was successful, false otherwise
- */
-bool FXAudio2SoundSource::Init(FWaveInstance* InWaveInstance)
+bool FXAudio2SoundSource::PrepareForInitialization(FWaveInstance* InWaveInstance)
 {
+	// We are not supporting playing audio on a controller
+	if (InWaveInstance->OutputTarget == EAudioOutputTarget::Controller)
+	{
+		return false;
+	}
+
+	// Flag that we are not initialized yet
+	bInitialized = false;
+
 	// Reset so next instance will warn if algorithm changes inflight
 	bEditorWarnedChangedSpatialization = false;
 
-	if (InWaveInstance->OutputTarget != EAudioOutputTarget::Controller)
+	// Find matching buffer.
+	FAudioDevice* BestAudioDevice = nullptr;
+	if (UAudioComponent* AudioComponent = InWaveInstance->ActiveSound->GetAudioComponent())
 	{
-		// Find matching buffer.
-		FAudioDevice* BestAudioDevice = nullptr;
-		if (UAudioComponent* AudioComponent = InWaveInstance->ActiveSound->GetAudioComponent())
-		{
-			BestAudioDevice = AudioComponent->GetAudioDevice();
-		}
-		else
-		{
-			BestAudioDevice = GEngine->GetMainAudioDevice();
-		}
-		check(BestAudioDevice);
+		BestAudioDevice = AudioComponent->GetAudioDevice();
+	}
+	else
+	{
+		BestAudioDevice = GEngine->GetMainAudioDevice();
+	}
+	check(BestAudioDevice);
 
-		XAudio2Buffer = FXAudio2SoundBuffer::Init(BestAudioDevice, InWaveInstance->WaveData, InWaveInstance->StartTime > 0.f);
+	check(XAudio2Buffer == nullptr);
+	XAudio2Buffer = FXAudio2SoundBuffer::Init(BestAudioDevice, InWaveInstance->WaveData, InWaveInstance->StartTime > 0.f);
+	if (XAudio2Buffer)
+	{
 		Buffer = XAudio2Buffer;
 
+		WaveInstance = InWaveInstance;
+
 		// Reset the LPFFrequency values
 		LPFFrequency = MAX_FILTER_FREQUENCY;
 		LastLPFFrequency = FLT_MAX;
+		bIsFinished = false;
 
-		// Buffer failed to be created, or there was an error with the compressed data
-		if (Buffer && Buffer->NumChannels > 0)
-		{
-			SCOPE_CYCLE_COUNTER( STAT_AudioSourceInitTime );
+		// We succeeded in preparing our xaudio2 buffer for initialization. We are technically not initialized yet.
+		// If the buffer is asynchronously preparing the ogg-vorbis file handle, we may not yet initialize the source
+		return true;
+	}
 
-			WaveInstance = InWaveInstance;
+	// Something went wrong with creating the XAudio2SoundBuffer
+	return false;
+}
 
-			// Set whether to apply reverb
-			SetReverbApplied(Effects->ReverbEffectVoice != nullptr);
+bool FXAudio2SoundSource::IsPreparedToInit()
+{
+	return (XAudio2Buffer && XAudio2Buffer->IsRealTimeSourceReady());
+}
 
-			// Create a new source
-			if (!CreateSource())
-			{
-				return false;
-			}
+bool FXAudio2SoundSource::Init(FWaveInstance* InWaveInstance)
+{
+	check(XAudio2Buffer);
+	check(XAudio2Buffer->IsRealTimeSourceReady());
+	check(Buffer);
 
-			if (WaveInstance->StartTime > 0.f)
-			{
-				XAudio2Buffer->Seek(WaveInstance->StartTime);
-			}
+	// Buffer failed to be created, or there was an error with the compressed data
+	if (Buffer->NumChannels > 0)
+	{
+		SCOPE_CYCLE_COUNTER(STAT_AudioSourceInitTime);
 
-			// Submit audio buffers
-			switch (XAudio2Buffer->SoundFormat)
-			{
+		// Set whether to apply reverb
+		SetReverbApplied(Effects->ReverbEffectVoice != nullptr);
+
+		// Create a new source if we haven't already
+		if (!CreateSource())
+		{
+			return false;
+		}
+
+		check(WaveInstance);
+		if (WaveInstance->StartTime)
+		{
+			XAudio2Buffer->Seek(WaveInstance->StartTime);
+		}
+
+		// Submit audio buffers
+		switch (XAudio2Buffer->SoundFormat)
+		{
 			case SoundFormat_PCM:
 			case SoundFormat_PCMPreview:
-				SubmitPCMBuffers();
-				break;
+			SubmitPCMBuffers();
+			break;
 
 			case SoundFormat_PCMRT:
 			case SoundFormat_Streaming:
-				SubmitPCMRTBuffers();
-				break;
+			SubmitPCMRTBuffers();
+			break;
 
 			case SoundFormat_XMA2:
-				SubmitXMA2Buffers();
-				break;
+			SubmitXMA2Buffers();
+			break;
 
 			case SoundFormat_XWMA:
-				SubmitXWMABuffers();
-				break;
-			}
-
-			// Updates the source which e.g. sets the pitch and volume.
-			Update();
-		
-			// Initialization succeeded.
-			return true;
+			SubmitXWMABuffers();
+			break;
 		}
+
+		// First updates of the source which e.g. sets the pitch and volume.
+		bInitialized = true;
+		bFirstRTBuffersSubmitted = false;
+
+		Update();
+
+		// Now set the source state to initialized so it can be played
+		// Initialization succeeded.
+		return true;
 	}
 
 	// Initialization failed.
@@ -1113,9 +1145,42 @@ int32 FXAudio2SoundSource::GetDestinationVoiceIndexForEffect( SourceDestinations
 	return Index;
 }
 
-/**
- * Callback to let the sound system know the sound has looped
- */
+void FXAudio2SoundSourceCallback::OnBufferEnd(void* BufferContext)
+{
+	if (BufferContext)
+	{
+		FXAudio2SoundSource* SoundSource = (FXAudio2SoundSource*)BufferContext;
+
+		// Only submit more buffers if the source is playing (not stopped or paused)
+		if (SoundSource->Playing && SoundSource->Source)
+		{
+			// Retrieve state source is in.
+			XAUDIO2_VOICE_STATE SourceState;
+
+			SoundSource->Source->GetState(&SourceState);
+
+			const bool bIsRealTimeSource = SoundSource->XAudio2Buffer->SoundFormat == SoundFormat_PCMRT || SoundSource->XAudio2Buffer->SoundFormat == SoundFormat_Streaming;
+
+			// If we have no queued buffers, we're either at the end of a sound, or starved
+			// and we are expecting the sound to be finishing
+			if (SourceState.BuffersQueued == 0 && (SoundSource->bBuffersToFlush || !bIsRealTimeSource))
+			{
+				// Set the flag to notify wave instances that we're finished
+				SoundSource->bIsFinished = true;
+				return;
+			}
+
+			// Service any real time sounds
+			if (bIsRealTimeSource && !SoundSource->bBuffersToFlush && SourceState.BuffersQueued <= 2)
+			{
+				// Continue feeding new sound data (unless we are waiting for the sound to finish)
+				SoundSource->HandleRealTimeSource(SourceState.BuffersQueued < 2);
+				return;
+			}
+		}
+	}
+}
+
 void FXAudio2SoundSourceCallback::OnLoopEnd( void* BufferContext )
 {
 	if( BufferContext )
@@ -1184,15 +1249,11 @@ FString FXAudio2SoundSource::Describe_Internal(bool bUseLongName, bool bIncludeC
 
 }
 
-/**
- * Updates the source specific parameter like e.g. volume and pitch based on the associated
- * wave instance.	
- */
-void FXAudio2SoundSource::Update( void )
+void FXAudio2SoundSource::Update()
 {
 	SCOPE_CYCLE_COUNTER( STAT_AudioUpdateSources );
 
-	if( !WaveInstance || !Source || Paused )
+	if (!WaveInstance || !Source || Paused || !bInitialized)
 	{	
 		return;
 	}
@@ -1261,12 +1322,10 @@ void FXAudio2SoundSource::Update( void )
 
 }
 
-
-/**
- * Plays the current wave instance.	
- */
-void FXAudio2SoundSource::Play( void )
+void FXAudio2SoundSource::Play()
 {
+	check(bInitialized);
+
 	if( WaveInstance )
 	{
 		if( !Playing )
@@ -1290,11 +1349,10 @@ void FXAudio2SoundSource::Play( void )
 	}
 }
 
-/**
- * Stops the current wave instance and detaches it from the source.	
- */
-void FXAudio2SoundSource::Stop( void )
+void FXAudio2SoundSource::Stop()
 {
+	bInitialized = false;
+	
 	if( WaveInstance )
 	{	
 		if( Playing )
@@ -1311,11 +1369,12 @@ void FXAudio2SoundSource::Stop( void )
 				Source->Stop( XAUDIO2_PLAY_TAILS ) );
 		}
 
+		Paused = false;
+		Playing = false;
+
 		// Free resources
 		FreeResources();
 
-		Paused = false;
-		Playing = false;
 		Buffer = XAudio2Buffer = nullptr;
 		bBuffersToFlush = false;
 		bLoopCallback = false;
@@ -1325,17 +1384,15 @@ void FXAudio2SoundSource::Stop( void )
 	FSoundSource::Stop();
 }
 
-/**
- * Pauses playback of current wave instance.
- */
-void FXAudio2SoundSource::Pause( void )
+void FXAudio2SoundSource::Pause()
 {
-	if( WaveInstance )
+	check(bInitialized);
+	if (WaveInstance)
 	{
-		if( Source )
+		if (Source)
 		{
-			AudioDevice->ValidateAPICall( TEXT( "Stop" ), 
-				Source->Stop( 0 ) );
+			AudioDevice->ValidateAPICall(TEXT("Stop"),
+										 Source->Stop(0));
 		}
 
 		Paused = true;
@@ -1349,27 +1406,28 @@ void FXAudio2SoundSource::HandleRealTimeSourceData(bool bLooped)
 	{
 		switch( WaveInstance->LoopingMode )
 		{
-		case LOOP_Never:
-			// Play out any queued buffers - once there are no buffers left, the state check at the beginning of IsFinished will fire
-			bBuffersToFlush = true;
-			XAudio2Buffers[CurrentBuffer].Flags |= XAUDIO2_END_OF_STREAM;
-			break;
+			case LOOP_Never:
+				// Play out any queued buffers - once there are no buffers left, the state check at the beginning of IsFinished will fire
+				bBuffersToFlush = true;
+				XAudio2Buffers[CurrentBuffer].Flags |= XAUDIO2_END_OF_STREAM;
+				break;
 
-		case LOOP_WithNotification:
-			// If we have just looped, and we are programmatically looping, send notification
-			WaveInstance->NotifyFinished();
-			break;
+			case LOOP_WithNotification:
+				// If we have just looped, and we are programmatically looping, send notification
+				// This will trigger a WaveInstance->NotifyFinished() in the FXAudio2SoundSournce::IsFinished() function on main thread.
+				bLoopCallback = true;
+				break;
 
-		case LOOP_Forever:
-			// Let the sound loop indefinitely
-			break;
+			case LOOP_Forever:
+				// Let the sound loop indefinitely
+				break;
 		}
 	}
 
 	if (XAudio2Buffers[CurrentBuffer].AudioBytes > 0)
 	{
-		AudioDevice->ValidateAPICall( TEXT( "SubmitSourceBuffer - IsFinished" ), 
-			Source->SubmitSourceBuffer( &XAudio2Buffers[CurrentBuffer] ) );
+		AudioDevice->ValidateAPICall(TEXT("SubmitSourceBuffer - PCMRT"),
+									 Source->SubmitSourceBuffer(&XAudio2Buffers[CurrentBuffer]));
 	}
 	else
 	{
@@ -1380,9 +1438,6 @@ void FXAudio2SoundSource::HandleRealTimeSourceData(bool bLooped)
 	}
 }
 
-/**
- * Handles feeding new data to a real time decompressed sound
- */
 void FXAudio2SoundSource::HandleRealTimeSource(bool bBlockForData)
 {
 	const bool bGetMoreData = bBlockForData || (RealtimeAsyncTask == nullptr);
@@ -1400,13 +1455,13 @@ void FXAudio2SoundSource::HandleRealTimeSource(bool bBlockForData)
 
 			switch(RealtimeAsyncTask->GetTask().GetTaskType())
 			{
-			case ERealtimeAudioTaskType::Decompress:
-				bLooped = RealtimeAsyncTask->GetTask().GetBufferLooped();
-				break;
+				case ERealtimeAudioTaskType::Decompress:
+					bLooped = RealtimeAsyncTask->GetTask().GetBufferLooped();
+					break;
 
-			case ERealtimeAudioTaskType::Procedural:
-				XAudio2Buffers[CurrentBuffer].AudioBytes = RealtimeAsyncTask->GetTask().GetBytesWritten();
-				break;
+				case ERealtimeAudioTaskType::Procedural:
+					XAudio2Buffers[CurrentBuffer].AudioBytes = RealtimeAsyncTask->GetTask().GetBytesWritten();
+					break;
 			}
 
 			delete RealtimeAsyncTask;
@@ -1424,9 +1479,19 @@ void FXAudio2SoundSource::HandleRealTimeSource(bool bBlockForData)
 			CurrentBuffer = 0;
 		}
 
-		// Get the next bit of streaming data
-		const bool bLooped = ReadMorePCMData(CurrentBuffer, (XAudio2Buffer->SoundFormat == ESoundFormat::SoundFormat_Streaming ? EDataReadMode::Synchronous : EDataReadMode::Asynchronous));
+		EDataReadMode DataReadMode;
+		if (bPlayedCachedBuffer)
+		{
+			bPlayedCachedBuffer = false;
+			DataReadMode = EDataReadMode::AsynchronousSkipFirstFrame;
+		}
+		else
+		{
+			DataReadMode = (XAudio2Buffer->SoundFormat == ESoundFormat::SoundFormat_Streaming ? EDataReadMode::Synchronous : EDataReadMode::Asynchronous);
+		}
+		const bool bLooped = ReadMorePCMData(CurrentBuffer, DataReadMode);
 
+		// If this was a synchronous read, then immediately write it
 		if (RealtimeAsyncTask == nullptr)
 		{
 			HandleRealTimeSourceData(bLooped);
@@ -1434,58 +1499,31 @@ void FXAudio2SoundSource::HandleRealTimeSource(bool bBlockForData)
 	}
 }
 
-/**
- * Queries the status of the currently associated wave instance.
- *
- * @return	true if the wave instance/ source has finished playback and false if it is 
- *			currently playing or paused.
- */
-bool FXAudio2SoundSource::IsFinished( void )
+bool FXAudio2SoundSource::IsFinished()
 {
 	// A paused source is not finished.
-	if( Paused )
+	if (Paused || !bInitialized)
 	{
-		return( false );
+		return(false);
 	}
 
-	if( WaveInstance && Source )
+	if (WaveInstance && Source)
 	{
-		// Retrieve state source is in.
-		XAUDIO2_VOICE_STATE SourceState;
-		Source->GetState( &SourceState );
-
-		const bool bIsRealTimeSource = XAudio2Buffer->SoundFormat == SoundFormat_PCMRT || XAudio2Buffer->SoundFormat == SoundFormat_Streaming;
-
-		// If we have no queued buffers, we're either at the end of a sound, or starved
-		// and we are expecting the sound to be finishing
-		if (SourceState.BuffersQueued == 0 && (bBuffersToFlush || !bIsRealTimeSource))
+		if (bIsFinished)
 		{
-			// ... notify the wave instance that it has finished playing.
 			WaveInstance->NotifyFinished();
 			return true;
 		}
 
-		// Service any real time sounds
-		if (bIsRealTimeSource && !bBuffersToFlush && SourceState.BuffersQueued <= 2)
-		{
-			// Continue feeding new sound data (unless we are waiting for the sound to finish)
-			HandleRealTimeSource(SourceState.BuffersQueued < 2);
-
-			return( false );
-		}
-
-		// Notify the wave instance that the looping callback was hit
-		if( bLoopCallback && WaveInstance->LoopingMode == LOOP_WithNotification )
+		if (bLoopCallback && WaveInstance->LoopingMode == LOOP_WithNotification)
 		{
 			WaveInstance->NotifyFinished();
-
 			bLoopCallback = false;
 		}
 
-		return( false );
+		return false;
 	}
-
-	return( true );
+	return true;
 }
 
 bool FXAudio2SoundSource::IsUsingHrtfSpatializer()
diff --git a/Engine/Source/Runtime/Windows/XAudio2/Private/XAudio2Support.h b/Engine/Source/Runtime/Windows/XAudio2/Private/XAudio2Support.h
--- a/Engine/Source/Runtime/Windows/XAudio2/Private/XAudio2Support.h
+++ b/Engine/Source/Runtime/Windows/XAudio2/Private/XAudio2Support.h
@@ -88,6 +88,8 @@ struct FXWMABufferInfo
 	UINT32						XWMASeekDataSize;
 };
 
+typedef FAsyncTask<class FAsyncRealtimeAudioTaskWorker<FXAudio2SoundBuffer>> FAsyncRealtimeAudioTask;
+
 /**
  * XAudio2 implementation of FSoundBuffer, containing the wave data and format information.
  */
@@ -106,7 +108,14 @@ public:
 	 * 
 	 * Frees wave data and detaches itself from audio device.
 	 */
-	~FXAudio2SoundBuffer( void );
+	~FXAudio2SoundBuffer();
+
+	//~ Begin FSoundBuffer interface
+	int32 GetSize() override;
+	int32 GetCurrentChunkIndex() const override;
+	int32 GetCurrentChunkOffset() const override;
+	bool IsRealTimeSourceReady() override;
+	//~ End FSoundBuffer interface 
 
 	/** 
 	 * Set up this buffer to contain and play XMA2 data
@@ -123,6 +132,11 @@ public:
 	 */
 	void InitWaveFormatEx( uint16 Format, USoundWave* Wave, bool bCheckPCMData );
 
+	/**
+	* Read the compressed info (i.e. parse header and open up a file handle) from the given sound wave.
+	*/
+	bool ReadCompressedInfo(USoundWave* SoundWave) override;
+
 	/**
 	 * Decompresses a chunk of compressed audio to the destination memory
 	 *
@@ -194,16 +208,6 @@ public:
 	 */
 	static FXAudio2SoundBuffer* Init( FAudioDevice* AudioDevice, USoundWave* InWave, bool bForceRealtime );
 
-	/**
-	 * Returns the size of this buffer in bytes.
-	 *
-	 * @return Size in bytes
-	 */
-	int32 GetSize( void );
-
-	virtual int32 GetCurrentChunkIndex() const override;
-	virtual int32 GetCurrentChunkOffset() const override;
-
 	/** Format of the sound referenced by this buffer */
 	int32							SoundFormat;
 
@@ -217,9 +221,16 @@ public:
 	};
 
 	/** Wrapper to handle the decompression of audio codecs */
-	class ICompressedAudioInfo*		DecompressionState;
+	class ICompressedAudioInfo* DecompressionState;
+
+	/** Asynch task for parsing real-time decompressed compressed info headers */
+	FAsyncRealtimeAudioTask* RealtimeAsyncHeaderParseTask;
+
+	/** Bool indicating the that the real-time source is ready for real-time decoding. */
+	FThreadSafeBool bRealTimeSourceReady;
+
 	/** Set to true when the PCM data should be freed when the buffer is destroyed */
-	bool						bDynamicResource;
+	bool bDynamicResource;
 };
 
 /**
@@ -228,45 +239,28 @@ public:
 class FXAudio2SoundSourceCallback : public IXAudio2VoiceCallback
 {
 public:
-	FXAudio2SoundSourceCallback( void )
-	{
-	}
-
-	virtual ~FXAudio2SoundSourceCallback( void )
-	{ 
-	}
-
-	virtual void STDCALL OnStreamEnd( void ) 
-	{ 
-	}
-
-	virtual void STDCALL OnVoiceProcessingPassEnd( void ) 
-	{
-	}
-
-	virtual void STDCALL OnVoiceProcessingPassStart( UINT32 SamplesRequired )
-	{
-	}
+	FXAudio2SoundSourceCallback() {}
+	virtual ~FXAudio2SoundSourceCallback() {}
+	virtual void STDCALL OnStreamEnd() {}
+	virtual void STDCALL OnVoiceProcessingPassEnd() {}
+	virtual void STDCALL OnVoiceProcessingPassStart(UINT32) {}
+	virtual void STDCALL OnBufferStart(void* BufferContext) {}
+	virtual void STDCALL OnVoiceError(void* BufferContext, HRESULT Error) {}
 
-	virtual void STDCALL OnBufferEnd( void* BufferContext )
-	{
-	}
-
-	virtual void STDCALL OnBufferStart( void* BufferContext )
-	{
-	}
-
-	virtual void STDCALL OnLoopEnd( void* BufferContext );
+	/**
+	* Function is called by xaudio2 for a playing voice every time a buffer finishes playing
+	* We use it to launch async decoding tasks and read the results of previous, finished tasks.
+	*/
+	virtual void STDCALL OnBufferEnd(void* BufferContext);
 
-	virtual void STDCALL OnVoiceError( void* BufferContext, HRESULT Error )
-	{
-	}
+	/**
+	* Function called whenever an xaudio2 voice loops on itself. Used to trigger notifications on loop.
+	*/
+	virtual void STDCALL OnLoopEnd(void* BufferContext);
 
 	friend class FXAudio2SoundSource;
 };
 
-typedef FAsyncTask<class FAsyncRealtimeAudioTaskWorker<FXAudio2SoundBuffer>> FAsyncRealtimeAudioTask;
-
 /**
  * XAudio2 implementation of FSoundSource, the interface used to play, stop and update sources
  */
@@ -283,46 +277,60 @@ public:
 	/**
 	 * Destructor, cleaning up voice
 	 */
-	virtual ~FXAudio2SoundSource( void );
+	virtual ~FXAudio2SoundSource();
 
 	/**
 	 * Frees existing resources. Called from destructor and therefore not virtual.
 	 */
-	void FreeResources( void );
+	void FreeResources();
 
 	/**
 	* Initializes any effects used with this source voice
 	*/
 	void InitializeSourceEffects(uint32 InVoiceId) override;
 
+	/**
+	* Attempts to initialize the sound source. If the source needs to decode asynchronously,
+	* this may not result in the source actually initializing, but instead it may launch
+	* an async task to parse the encoded file header and open a file handle. For ogg-vorbis
+	* synchronous file parsing can cause significant hitches.
+	*/
+	virtual bool PrepareForInitialization(FWaveInstance* InWaveInstance) override;
+
+	/**
+	* Queries if the async task has finished parsing the real-time decoded file header and if the
+	* the sound source is ready for initialization.
+	*/
+	virtual bool IsPreparedToInit() override;
+
 	/**
 	 * Initializes a source with a given wave instance and prepares it for playback.
 	 *
 	 * @param	WaveInstance	wave instance being primed for playback
 	 * @return	true if initialization was successful, false otherwise
 	 */
-	virtual bool Init( FWaveInstance* WaveInstance );
+	virtual bool Init(FWaveInstance* WaveInstance) override;
 
 	/**
 	 * Updates the source specific parameter like e.g. volume and pitch based on the associated
 	 * wave instance.	
 	 */
-	virtual void Update( void );
+	virtual void Update();
 
 	/**
 	 * Plays the current wave instance.	
 	 */
-	virtual void Play( void );
+	virtual void Play();
 
 	/**
 	 * Stops the current wave instance and detaches it from the source.	
 	 */
-	virtual void Stop( void );
+	virtual void Stop();
 
 	/**
 	 * Pauses playback of current wave instance.
 	 */
-	virtual void Pause( void );
+	virtual void Pause();
 
 	/**
 	 * Handles feeding new data to a real time decompressed sound
@@ -340,32 +348,32 @@ public:
 	 * @return	true if the wave instance/ source has finished playback and false if it is 
 	 *			currently playing or paused.
 	 */
-	virtual bool IsFinished( void );
+	virtual bool IsFinished();
 
 	/**
 	 * Create a new source voice
 	 */
-	bool CreateSource( void );
+	bool CreateSource();
 
 	/** 
 	 * Submit the relevant audio buffers to the system
 	 */
-	void SubmitPCMBuffers( void );
+	void SubmitPCMBuffers();
 
 	/** 
 	 * Submit the relevant audio buffers to the system
 	 */
-	void SubmitPCMRTBuffers( void );
+	void SubmitPCMRTBuffers();
 
 	/** 
 	 * Submit the relevant audio buffers to the system, accounting for looping modes
 	 */
-	void SubmitXMA2Buffers( void );
+	void SubmitXMA2Buffers();
 
 	/** 
 	 * Submit the relevant audio buffers to the system
 	 */
-	void SubmitXWMABuffers( void );
+	void SubmitXWMABuffers();
 
 	/**
 	 * Calculates the volume for each channel
@@ -451,39 +459,62 @@ protected:
 	void RouteMonoToReverb(float ChannelVolumes[CHANNEL_MATRIX_COUNT]);
 	void RouteStereoToReverb(float ChannelVolumes[CHANNEL_MATRIX_COUNT]);
 
-	/** Owning classes */
-	FXAudio2Device*				AudioDevice;
-	FXAudio2EffectsManager*		Effects;
+	/** Owning audio device object */
+	FXAudio2Device*	AudioDevice;
+
+	/** Pointer to effects manager, which handles updating singleton effects */
+	FXAudio2EffectsManager* Effects;
 
 	/** Cached subclass version of Buffer (which the base class has) */
-	FXAudio2SoundBuffer*		XAudio2Buffer;
+	FXAudio2SoundBuffer* XAudio2Buffer;
+
 	/** XAudio2 source voice associated with this source. */
-	IXAudio2SourceVoice*		Source;
+	IXAudio2SourceVoice* Source;
+
 	/** The max channels in the voice's effect chain. This is used to classify a pool for IXAudio2SourceVoice. */
-	int32						MaxEffectChainChannels;
+	int32 MaxEffectChainChannels;
 
-	/** Asynchronous task for real time audio sources */
+	/** Asynchronous task for real time audio decoding, created from main thread */
 	FAsyncRealtimeAudioTask* RealtimeAsyncTask;
+
 	/** Destination voices */
-	XAUDIO2_SEND_DESCRIPTOR		Destinations[DEST_COUNT];
-	/** Which sound buffer should be written to next - used for double buffering. */
-	int32							CurrentBuffer;
+	XAUDIO2_SEND_DESCRIPTOR Destinations[DEST_COUNT];
+
 	/** A pair of sound buffers to allow notification when a sound loops. */
-	XAUDIO2_BUFFER				XAudio2Buffers[3];
+	XAUDIO2_BUFFER XAudio2Buffers[3];
+
 	/** Additional buffer info for XWMA sounds */
-	XAUDIO2_BUFFER_WMA			XAudio2BufferXWMA[1];
-	/** Set when we wish to let the buffers play themselves out */
-	uint32						bBuffersToFlush:1;
+	XAUDIO2_BUFFER_WMA XAudio2BufferXWMA[1];
+
+	/** Index of this sound source in the audio device sound source array. */
+	uint32 VoiceId;
+
+	/** Which sound buffer should be written to next - used for triple buffering. */
+	int32 CurrentBuffer;
+
 	/** Set to true when the loop end callback is hit */
-	uint32						bLoopCallback:1;
+	FThreadSafeBool bLoopCallback;
+
+	/** Whether or not the sound has finished playing */
+	FThreadSafeBool bIsFinished;
+
+	/** Whether or not the cached first buffer has played. Used to skip first two reads of a RT decoded file */
+	FThreadSafeBool bPlayedCachedBuffer;
+
+	/** Whether or not we need to submit our first buffers to the voice */
+	FThreadSafeBool bFirstRTBuffersSubmitted;
+
+	/** Set when we wish to let the buffers play themselves out */
+	FThreadSafeBool bBuffersToFlush;
+
 	/** Set to true when we've allocated resources that need to be freed */
-	uint32						bResourcesNeedFreeing:1;
-	/** Index of this sound source in the audio device sound source array. */
-	uint32						VoiceId;
+	uint32 bResourcesNeedFreeing : 1;
+
 	/** Whether or not this sound is spatializing using an HRTF spatialization algorithm. */
-	bool						bUsingHRTFSpatialization;
+	uint32 bUsingHRTFSpatialization : 1;
+
 	/** Whether or not we've already logged a warning on this sound about it switching algorithms after init. */
-	bool						bEditorWarnedChangedSpatialization;
+	uint32 bEditorWarnedChangedSpatialization : 1;
 
 	friend class FXAudio2Device;
 	friend class FXAudio2SoundSourceCallback;
