Submit For Ori:
Particles: Disabled distribution baking temporarily to prevent additional assets from being saved in a state that does not work with -game
Note: Particle system assets saved since the merge from staging yesterday are still 'corrupted' WRT -game mode, that will be addressed soonish

diff --git a/Engine/Source/Runtime/Engine/Private/Distributions.cpp b/Engine/Source/Runtime/Engine/Private/Distributions.cpp
--- a/Engine/Source/Runtime/Engine/Private/Distributions.cpp
+++ b/Engine/Source/Runtime/Engine/Private/Distributions.cpp
@@ -31,7 +31,7 @@ ENGINE_API uint32 GDistributionType = 1;
 #define LOOKUP_TABLE_MAX_SAMPLES (128)
 
 // UDistribution will bake out (if CanBeBaked returns true)
-#define DISTRIBUTIONS_BAKES_OUT 1
+#define DISTRIBUTIONS_BAKES_OUT 0
 
 // The maximum number of samples must be a power of two.
 static_assert((LOOKUP_TABLE_MAX_SAMPLES & (LOOKUP_TABLE_MAX_SAMPLES - 1)) == 0, "Lookup table max samples is not a power of two.");
