Copying //UE4/Dev-Editor to Dev-Main (//UE4/Dev-Main)

diff --git a/Engine/Source/Editor/UnrealEd/Private/Fbx/SFbxSceneOptionWindow.cpp b/Engine/Source/Editor/UnrealEd/Private/Fbx/SFbxSceneOptionWindow.cpp
--- a/Engine/Source/Editor/UnrealEd/Private/Fbx/SFbxSceneOptionWindow.cpp
+++ b/Engine/Source/Editor/UnrealEd/Private/Fbx/SFbxSceneOptionWindow.cpp
@@ -32,6 +32,7 @@ SFbxSceneOptionWindow::SFbxSceneOptionWindow()
 	bShouldImport = false;
 	SceneTabTreeview = nullptr;
 	SceneTabDetailsView = nullptr;
+	SceneReimportTabDetailsView = nullptr;
 	OverrideNameOptions.Empty();
 	StaticMeshTabListView = nullptr;
 	StaticMeshTabDetailsView = nullptr;
@@ -65,6 +66,7 @@ SFbxSceneOptionWindow::~SFbxSceneOptionWindow()
 	bShouldImport = false;
 	SceneTabTreeview = nullptr;
 	SceneTabDetailsView = nullptr;
+	SceneReimportTabDetailsView = nullptr;
 	StaticMeshTabListView = nullptr;
 	StaticMeshTabDetailsView = nullptr;
 	SkeletalMeshTabListView = nullptr;
@@ -180,6 +182,9 @@ void SFbxSceneOptionWindow::OnFinishedChangingPropertiesSceneTabDetailView(const
 	//Set the Global Import setting
 	GlobalImportSettings->bBakePivotInVertex = SceneImportOptionsDisplay->bBakePivotInVertex;
 	GlobalImportSettings->bInvertNormalMap = SceneImportOptionsDisplay->bInvertNormalMaps;
+	GlobalImportSettings->ImportTranslation = SceneImportOptionsDisplay->ImportTranslation;
+	GlobalImportSettings->ImportRotation = SceneImportOptionsDisplay->ImportRotation;
+	GlobalImportSettings->ImportUniformScale = SceneImportOptionsDisplay->ImportUniformScale;
 }
 
 TSharedRef<SDockTab> SFbxSceneOptionWindow::SpawnStaticMeshTab(const FSpawnTabArgs& Args)
@@ -739,7 +744,9 @@ TSharedRef<SDockTab> SFbxSceneOptionWindow::SpawnSceneReimportTab(const FSpawnTa
 		.SceneInfo(SceneInfo)
 		.SceneInfoOriginal(SceneInfoOriginal)
 		.NodeStatusMap(NodeStatusMap);
-
+	
+	TSharedPtr<SBox> InspectorBox;
+	
 	TSharedRef<SDockTab> DockTab = SNew(SDockTab)
 		.TabRole(ETabRole::PanelTab)
 		.Label(LOCTEXT("WidgetFbxSceneActorTab", "Scene"))
@@ -748,7 +755,7 @@ TSharedRef<SDockTab> SFbxSceneOptionWindow::SpawnSceneReimportTab(const FSpawnTa
 			SNew(SSplitter)
 			.Orientation(Orient_Horizontal)
 			+ SSplitter::Slot()
-			.Value(0.6f)
+			.Value(0.4f)
 			[
 				SNew(SVerticalBox)
 				+ SVerticalBox::Slot()
@@ -801,7 +808,7 @@ TSharedRef<SDockTab> SFbxSceneOptionWindow::SpawnSceneReimportTab(const FSpawnTa
 				]
 			]
 			+ SSplitter::Slot()
-			.Value(0.4f)
+			.Value(0.6f)
 			[
 				SNew(SVerticalBox)
 				+ SVerticalBox::Slot()
@@ -820,7 +827,7 @@ TSharedRef<SDockTab> SFbxSceneOptionWindow::SpawnSceneReimportTab(const FSpawnTa
 						.IsChecked(this, &SFbxSceneOptionWindow::IsReimportHierarchyChecked)
 					]
 					+ SHorizontalBox::Slot()
-					.FillWidth(1.0f)
+					.AutoWidth()
 					.Padding(0.0f, 3.0f, 6.0f, 3.0f)
 					.VAlign(VAlign_Top)
 					.HAlign(HAlign_Left)
@@ -834,53 +841,8 @@ TSharedRef<SDockTab> SFbxSceneOptionWindow::SpawnSceneReimportTab(const FSpawnTa
 				.HAlign(HAlign_Left)
 				.AutoHeight()
 				[
-					SNew(SHorizontalBox)
-					+ SHorizontalBox::Slot()
-					.AutoWidth()
-					.VAlign(VAlign_Top)
-					.HAlign(HAlign_Left)
-					[
-						SNew(SCheckBox)
-						.HAlign(HAlign_Center)
-						.OnCheckStateChanged(this, &SFbxSceneOptionWindow::OnToggleBakePivotInVertex)
-						.IsChecked(this, &SFbxSceneOptionWindow::IsBakePivotInVertexChecked)
-					]
-					+ SHorizontalBox::Slot()
-					.FillWidth(1.0f)
-					.Padding(0.0f, 3.0f, 6.0f, 3.0f)
-					.VAlign(VAlign_Top)
-					.HAlign(HAlign_Left)
-					[
-						SNew(STextBlock)
-						.Text(LOCTEXT("FbxOptionWindow_Scene_Reimport_BakePivotInVertex", "Bake Pivot In Vertex"))
-						.ToolTipText(LOCTEXT("FbxOptionWindow_Scene_Reimport_BakePivotInVertex_Tooltip", "- Experimental - If this option is true the inverse node pivot will be apply to the mesh vertices. The pivot from the DCC will then be the origin of the mesh."))
-					]
-				]
-				+ SVerticalBox::Slot()
-				.HAlign(HAlign_Left)
-				.AutoHeight()
-				[
-					SNew(SHorizontalBox)
-					+ SHorizontalBox::Slot()
-					.AutoWidth()
-					.VAlign(VAlign_Top)
-					.HAlign(HAlign_Left)
-					[
-						SNew(SCheckBox)
-						.HAlign(HAlign_Center)
-						.OnCheckStateChanged(this, &SFbxSceneOptionWindow::OnToggleInvertNormalMap)
-						.IsChecked(this, &SFbxSceneOptionWindow::IsInvertNormalMapChecked)
-					]
-					+ SHorizontalBox::Slot()
-					.FillWidth(1.0f)
-					.Padding(0.0f, 3.0f, 6.0f, 3.0f)
-					.VAlign(VAlign_Top)
-					.HAlign(HAlign_Left)
-					[
-						SNew(STextBlock)
-						.Text(LOCTEXT("FbxOptionWindow_Scene_Reimport_InvertNormalMap", "Invert Normal Map"))
-						.ToolTipText(LOCTEXT("FbxOptionWindow_Scene_Reimport_InvertNormalMap_Tooltip", "If either importing of textures (or materials) is enabled, this option will cause normal map values to be inverted."))
-					]
+					SAssignNew(InspectorBox, SBox)
+					.WidthOverride(1920.0f)
 				]
 			]
 		];
@@ -888,6 +850,15 @@ TSharedRef<SDockTab> SFbxSceneOptionWindow::SpawnSceneReimportTab(const FSpawnTa
 	//Prevent user to close the tab
 	DockTab->SetCanCloseTab(SDockTab::FCanCloseTab::CreateRaw(this, &SFbxSceneOptionWindow::CanCloseTab));
 
+	FPropertyEditorModule& PropertyEditorModule = FModuleManager::GetModuleChecked<FPropertyEditorModule>("PropertyEditor");
+	FDetailsViewArgs DetailsViewArgs;
+	DetailsViewArgs.bAllowSearch = false;
+	DetailsViewArgs.NameAreaSettings = FDetailsViewArgs::HideNameArea;
+	SceneReimportTabDetailsView = PropertyEditorModule.CreateDetailView(DetailsViewArgs);
+	InspectorBox->SetContent(SceneReimportTabDetailsView->AsShared());
+	SceneReimportTabDetailsView->SetObject(SceneImportOptionsDisplay);
+	SceneReimportTabDetailsView->OnFinishedChangingProperties().AddSP(this, &SFbxSceneOptionWindow::OnFinishedChangingPropertiesSceneTabDetailView);
+
 	return DockTab;
 }
 
@@ -917,19 +888,6 @@ ECheckBoxState SFbxSceneOptionWindow::IsBakePivotInVertexChecked() const
 	return  GlobalImportSettings->bBakePivotInVertex ? ECheckBoxState::Checked : ECheckBoxState::Unchecked;
 }
 
-void SFbxSceneOptionWindow::OnToggleInvertNormalMap(ECheckBoxState CheckType)
-{
-	if (GlobalImportSettings != nullptr)
-	{
-		GlobalImportSettings->bInvertNormalMap = CheckType == ECheckBoxState::Checked;
-	}
-}
-
-ECheckBoxState SFbxSceneOptionWindow::IsInvertNormalMapChecked() const
-{
-	return  GlobalImportSettings->bInvertNormalMap ? ECheckBoxState::Checked : ECheckBoxState::Unchecked;
-}
-
 TSharedRef<SDockTab> SFbxSceneOptionWindow::SpawnStaticMeshReimportTab(const FSpawnTabArgs& Args)
 {
 	//Create the static mesh listview
@@ -1370,6 +1328,7 @@ void SFbxSceneOptionWindow::CloseFbxSceneOption()
 	SkeletalMeshTabListView = nullptr;
 
 	SceneReimportTreeview = nullptr;
+	SceneReimportTabDetailsView = nullptr;
 
 	StaticMeshReimportListView = nullptr;
 	StaticMeshReimportDetailsView = nullptr;
diff --git a/Engine/Source/Editor/UnrealEd/Private/Fbx/SSceneImportNodeTreeView.cpp b/Engine/Source/Editor/UnrealEd/Private/Fbx/SSceneImportNodeTreeView.cpp
--- a/Engine/Source/Editor/UnrealEd/Private/Fbx/SSceneImportNodeTreeView.cpp
+++ b/Engine/Source/Editor/UnrealEd/Private/Fbx/SSceneImportNodeTreeView.cpp
@@ -228,9 +228,10 @@ void SFbxSceneTreeView::OnGetChildrenFbxSceneTreeView(FbxNodeInfoPtr InParent, T
 	}
 	else
 	{
-		for (auto Child : InParent->Childrens)
+		for (TSharedPtr<FFbxNodeInfo> Child : InParent->Childrens)
 		{
-			if (Child.IsValid())
+			//We hide skeletal mesh from the tree, a mesh without a valid attributeinfo is a sub skeletal mesh
+			if (Child.IsValid() && (Child->AttributeType.Compare(TEXT("eMesh")) != 0 || Child->AttributeInfo.IsValid()))
 			{
 				OutChildren.Add(Child);
 			}
diff --git a/Engine/Source/Editor/UnrealEd/Private/Fbx/SSceneMaterialsListView.cpp b/Engine/Source/Editor/UnrealEd/Private/Fbx/SSceneMaterialsListView.cpp
--- a/Engine/Source/Editor/UnrealEd/Private/Fbx/SSceneMaterialsListView.cpp
+++ b/Engine/Source/Editor/UnrealEd/Private/Fbx/SSceneMaterialsListView.cpp
@@ -145,7 +145,8 @@ void SFbxSceneMaterialsListView::GetMaterialsFromHierarchy(TArray<FbxMaterialInf
 	OutMaterials.Reset();
 	for (FbxNodeInfoPtr NodeInfo : SceneInfoSource->HierarchyInfo)
 	{
-		if (!NodeInfo->AttributeInfo.IsValid() || NodeInfo->NodeName.Compare("RootNode") == 0)
+		//Get all the mesh material
+		if (NodeInfo->AttributeType.Compare(TEXT("eMesh")) != 0 || NodeInfo->NodeName.Compare("RootNode") == 0)
 		{
 			continue;
 		}
diff --git a/Engine/Source/Editor/UnrealEd/Private/Fbx/SSceneReimportNodeTreeView.cpp b/Engine/Source/Editor/UnrealEd/Private/Fbx/SSceneReimportNodeTreeView.cpp
--- a/Engine/Source/Editor/UnrealEd/Private/Fbx/SSceneReimportNodeTreeView.cpp
+++ b/Engine/Source/Editor/UnrealEd/Private/Fbx/SSceneReimportNodeTreeView.cpp
@@ -381,6 +381,7 @@ void SFbxReimportSceneTreeView::OnGetChildrenFbxSceneTreeView(FbxNodeInfoPtr InP
 	check(NodeTreeData.Contains(InParent));
 	TSharedPtr<FTreeNodeValue> NodeValue = *(NodeTreeData.Find(InParent));
 	TArray<FString> ChildNames;
+	TArray<FString> ChildProcess;
 	//Current node contain the add and same node
 	if (NodeValue->CurrentNode.IsValid())
 	{
@@ -393,6 +394,7 @@ void SFbxReimportSceneTreeView::OnGetChildrenFbxSceneTreeView(FbxNodeInfoPtr InP
 				if (Child.IsValid() && Child->AttributeType.Compare(TEXT("eMesh")) == 0)
 				{
 					ChildNames.Add(Child->NodeName);
+					ChildProcess.Add(Child->NodeName);
 					OutChildren.Add(Child);
 					return;
 				}
@@ -403,11 +405,16 @@ void SFbxReimportSceneTreeView::OnGetChildrenFbxSceneTreeView(FbxNodeInfoPtr InP
 		{
 			for (FbxNodeInfoPtr Child : NodeValue->CurrentNode->Childrens)
 			{
-				if (Child.IsValid())
+				if (Child.IsValid() && (Child->AttributeType.Compare(TEXT("eMesh")) != 0 || Child->AttributeInfo.IsValid()))
 				{
 					ChildNames.Add(Child->NodeName);
+					ChildProcess.Add(Child->NodeName);
 					OutChildren.Add(Child);
 				}
+				else if(Child.IsValid() && Child->AttributeType.Compare(TEXT("eMesh")) == 0)
+				{
+					ChildProcess.Add(Child->NodeName);
+				}
 			}
 		}
 	}
@@ -423,7 +430,7 @@ void SFbxReimportSceneTreeView::OnGetChildrenFbxSceneTreeView(FbxNodeInfoPtr InP
 				Child = Child->Childrens[0];
 				if (Child.IsValid() && Child->AttributeType.Compare(TEXT("eMesh")) == 0)
 				{
-					if (!ChildNames.Contains(Child->NodeName))
+					if (!ChildProcess.Contains(Child->NodeName))
 					{
 						OutChildren.Add(Child);
 					}
@@ -437,7 +444,7 @@ void SFbxReimportSceneTreeView::OnGetChildrenFbxSceneTreeView(FbxNodeInfoPtr InP
 			for (FbxNodeInfoPtr Child : NodeValue->OriginalNode->Childrens)
 			{
 				//We have a delete node
-				if (Child.IsValid() && !ChildNames.Contains(Child->NodeName))
+				if (Child.IsValid() && !ChildProcess.Contains(Child->NodeName))
 				{
 					OutChildren.Add(Child);
 				}
diff --git a/Engine/Source/Editor/UnrealEd/Private/Fbx/SSceneReimportSkeletalMeshListView.cpp b/Engine/Source/Editor/UnrealEd/Private/Fbx/SSceneReimportSkeletalMeshListView.cpp
--- a/Engine/Source/Editor/UnrealEd/Private/Fbx/SSceneReimportSkeletalMeshListView.cpp
+++ b/Engine/Source/Editor/UnrealEd/Private/Fbx/SSceneReimportSkeletalMeshListView.cpp
@@ -266,7 +266,9 @@ private:
 
 	FText GetAssetPivotNodeName() const
 	{
-		return GlobalImportSettings->bBakePivotInVertex ? FText::FromString(FbxMeshInfo->PivotNodeName) : FText::FromString(TEXT("-"));
+		//TODO support bBakePivotInVertex
+		//return GlobalImportSettings->bBakePivotInVertex ? FText::FromString(FbxMeshInfo->PivotNodeName) : FText::FromString(TEXT("-"));
+		return FText::FromString(TEXT("-"));
 	}
 
 	/** The node info to build the tree view row from. */
@@ -549,7 +551,8 @@ TSharedPtr<SWidget> SFbxSceneSkeletalMeshReimportListView::OnOpenContextMenu()
 	}
 	MenuBuilder.EndSection();
 	
-	AddBakePivotMenu(MenuBuilder);
+	//TODO support bBakePivotInVertex 
+	//AddBakePivotMenu(MenuBuilder);
 
 	bool bShowOptionMenu = false;
 	for (FbxMeshInfoPtr MeshInfo : SelectedItems)
diff --git a/Engine/Source/Editor/UnrealEd/Private/FbxImporter.h b/Engine/Source/Editor/UnrealEd/Private/FbxImporter.h
--- a/Engine/Source/Editor/UnrealEd/Private/FbxImporter.h
+++ b/Engine/Source/Editor/UnrealEd/Private/FbxImporter.h
@@ -180,6 +180,12 @@ struct FBXImportOptions
 		bImportMorph = true;
 		AnimationLengthImportType = FBXALIT_ExportedTime;
 	}
+
+	static void ResetOptions(FBXImportOptions *OptionsToReset)
+	{
+		check(OptionsToReset != nullptr);
+		FMemory::Memzero(OptionsToReset, sizeof(OptionsToReset));
+	}
 };
 
 struct FbxMeshInfo
@@ -783,12 +789,23 @@ protected:
 
 	/**
 	* Compute the global matrix for Fbx Node
+	* If we import scene it will return identity plus the pivot if we turn the bake pivot option
 	*
 	* @param Node	Fbx Node
 	* @return KFbxXMatrix*	The global transform matrix
 	*/
 	FbxAMatrix ComputeTotalMatrix(FbxNode* Node);
-
+	
+	/**
+	* Compute the matrix for skeletal Fbx Node
+	* If we import don't import a scene it will call ComputeTotalMatrix with Node as the parameter. If we import a scene
+	* it will return the relative transform between the RootSkeletalNode and Node.
+	*
+	* @param Node	Fbx Node
+	* @param Node	Fbx RootSkeletalNode
+	* @return KFbxXMatrix*	The global transform matrix
+	*/
+	FbxAMatrix ComputeSkeletalMeshTotalMatrix(FbxNode* Node, FbxNode *RootSkeletalNode);
 	/**
 	 * Check if there are negative scale in the transform matrix and its number is odd.
 	 * @return bool True if there are negative scale and its number is 1 or 3. 
@@ -832,12 +849,24 @@ protected:
 	* @param FbxShape	Fbx Morph object, if not NULL, we are importing a morph object.
 	* @param SortedLinks    Fbx Links(bones) of this skeletal mesh
 	* @param FbxMatList  All material names of the skeletal mesh
+	* @param RootNode       The skeletal mesh root fbx node.
 	*
 	* @returns bool*	true if import successfully.
 	*/
     bool FillSkelMeshImporterFromFbx(FSkeletalMeshImportData& ImportData, FbxMesh*& Mesh, FbxSkin* Skin, 
-										FbxShape* Shape, TArray<FbxNode*> &SortedLinks, const TArray<FbxSurfaceMaterial*>& FbxMaterials );
-
+										FbxShape* Shape, TArray<FbxNode*> &SortedLinks, const TArray<FbxSurfaceMaterial*>& FbxMaterials, FbxNode *RootNode);
+	/**
+	* Fill FSkeletalMeshIMportData from Fbx Nodes and FbxShape Array if exists.  
+	*
+	* @param NodeArray	Fbx node array to look at
+	* @param TemplateImportData template import data 
+	* @param FbxShapeArray	Fbx Morph object, if not NULL, we are importing a morph object.
+	* @param OutData    FSkeletalMeshImportData output data
+	*
+	* @returns bool*	true if import successfully.
+	*/
+	bool FillSkeletalMeshImportData(TArray<FbxNode*>& NodeArray, UFbxSkeletalMeshImportData* TemplateImportData, TArray<FbxShape*> *FbxShapeArray, FSkeletalMeshImportData* OutData);
+	
 	/**
 	 * Import bones from skeletons that NodeArray bind to.
 	 *
@@ -848,7 +877,7 @@ protected:
 	 * @param bDisableMissingBindPoseWarning
 	 * @param bUseTime0AsRefPose	in/out - Use Time 0 as Ref Pose 
 	 */
-	bool ImportBone(TArray<FbxNode*>& NodeArray, FSkeletalMeshImportData &ImportData, UFbxSkeletalMeshImportData* TemplateData, TArray<FbxNode*> &OutSortedLinks, bool& bOutDiffPose, bool bDisableMissingBindPoseWarning, bool & bUseTime0AsRefPose);
+	bool ImportBone(TArray<FbxNode*>& NodeArray, FSkeletalMeshImportData &ImportData, UFbxSkeletalMeshImportData* TemplateData, TArray<FbxNode*> &OutSortedLinks, bool& bOutDiffPose, bool bDisableMissingBindPoseWarning, bool & bUseTime0AsRefPose, FbxNode *SkeletalMeshNode);
 	
 	/**
 	 * Skins the control points of the given mesh or shape using either the default pose for skinning or the first frame of the
diff --git a/Engine/Source/Editor/UnrealEd/Private/Layers/Layers.cpp b/Engine/Source/Editor/UnrealEd/Private/Layers/Layers.cpp
--- a/Engine/Source/Editor/UnrealEd/Private/Layers/Layers.cpp
+++ b/Engine/Source/Editor/UnrealEd/Private/Layers/Layers.cpp
@@ -706,6 +706,36 @@ void FLayers::RemoveViewFromActorViewVisibility( FLevelEditorViewportClient* Vie
 	}
 }
 
+static void UpdateBrushLayerVisibility(ABrush* Brush, bool bIsHidden)
+{
+	ULevel* Level = Brush->GetLevel();
+	if (!Level)
+	{
+		return;
+	}
+
+	UModel* Model = Level->Model;
+	if (!Model)
+	{
+		return;
+	}
+
+	bool bAnySurfaceWasFound = false;
+	for (FBspSurf& Surf : Model->Surfs)
+	{
+		if (Surf.Actor == Brush)
+		{
+			Surf.bHiddenEdLayer = bIsHidden;
+			bAnySurfaceWasFound = true;
+		}
+	}
+
+	if (bAnySurfaceWasFound)
+	{
+		Level->UpdateModelComponents();
+		Model->InvalidSurfaces = true;
+	}
+}
 
 bool FLayers::UpdateActorVisibility( const TWeakObjectPtr< AActor >& Actor, bool& bOutSelectionChanged, bool& bOutActorModified, bool bNotifySelectionChange, bool bRedrawViewports )
 {
@@ -752,6 +782,12 @@ bool FLayers::UpdateActorVisibility( const TWeakObjectPtr< AActor >& Actor, bool
 				Actor->bHiddenEdLayer = false;
 				Actor->MarkComponentsRenderStateDirty();
 				bOutActorModified = true;
+
+				if (ABrush* Brush = Cast<ABrush>(Actor.Get()))
+				{
+					const bool bIsHidden = false;
+					UpdateBrushLayerVisibility(Brush, bIsHidden);
+				}
 			}
 
 			// Stop, because we found at least one visible layer the actor belongs to
@@ -769,6 +805,12 @@ bool FLayers::UpdateActorVisibility( const TWeakObjectPtr< AActor >& Actor, bool
 			Actor->bHiddenEdLayer = true;
 			Actor->MarkComponentsRenderStateDirty();
 			bOutActorModified = true;
+
+			if (ABrush* Brush = Cast<ABrush>(Actor.Get()))
+			{
+				const bool bIsHidden = true;
+				UpdateBrushLayerVisibility(Brush, bIsHidden);
+			}
 		}
 
 		//if the actor was selected, mark it as unselected
diff --git a/Engine/Source/Editor/UnrealEd/Private/PlayLevel.cpp b/Engine/Source/Editor/UnrealEd/Private/PlayLevel.cpp
--- a/Engine/Source/Editor/UnrealEd/Private/PlayLevel.cpp
+++ b/Engine/Source/Editor/UnrealEd/Private/PlayLevel.cpp
@@ -454,6 +454,9 @@ void UEditorEngine::EndPlayMap()
 	// Tear down the output log to message log thunker
 	OutputLogErrorsToMessageLogProxyPtr.Reset();
 
+	// Remove undo barrier
+	GUnrealEd->Trans->RemoveUndoBarrier();
+
 	// display any info if required.
 	FMessageLog(NAME_CategoryPIE).Notify(LOCTEXT("PIEErrorsPresent", "Errors/warnings reported while playing in editor."));
 }
@@ -3099,7 +3102,7 @@ UGameInstance* UEditorEngine::CreatePIEGameInstance(int32 PIEInstance, bool bInS
 	GEngine->BroadcastLevelActorListChanged();
 
 	// Set an undo barrier so that transactions prior to PIE can't be undone
-	GUnrealEd->ResetTransaction(NSLOCTEXT("UnrealEd", "PIEStarted", "PIE Started"));
+	GUnrealEd->Trans->SetUndoBarrier();
 
 	return GameInstance;
 }
diff --git a/Engine/Source/Editor/UnrealEd/Private/SCurveEditor.cpp b/Engine/Source/Editor/UnrealEd/Private/SCurveEditor.cpp
--- a/Engine/Source/Editor/UnrealEd/Private/SCurveEditor.cpp
+++ b/Engine/Source/Editor/UnrealEd/Private/SCurveEditor.cpp
@@ -355,7 +355,7 @@ void SCurveEditor::Construct(const FArguments& InArgs)
 							.MaxValue(TOptional<float>())
 							.MaxSliderValue(TOptional<float>())
 							.MinSliderValue(TOptional<float>())
-							.Delta(0.001f)
+							.Delta(this, &SCurveEditor::GetInputNumericEntryBoxDelta)
 							.MinDesiredValueWidth(60.0f)
 							.Label()
 							[
@@ -387,7 +387,7 @@ void SCurveEditor::Construct(const FArguments& InArgs)
 							.MaxValue(TOptional<float>())
 							.MaxSliderValue(TOptional<float>())
 							.MinSliderValue(TOptional<float>())
-							.Delta(.001f)
+							.Delta(this, &SCurveEditor::GetOutputNumericEntryBoxDelta)
 							.MinDesiredValueWidth(60.0f)
 							.Label()
 							[
@@ -1162,6 +1162,16 @@ void SCurveEditor::PaintMarquee(const FGeometry& AllottedGeometry, const FSlateR
 		);
 }
 
+float SCurveEditor::GetInputNumericEntryBoxDelta() const
+{
+	return bSnappingEnabled.Get() ? InputSnap.Get() : 0;
+}
+
+float SCurveEditor::GetOutputNumericEntryBoxDelta() const
+{
+	return bSnappingEnabled.Get() ? OutputSnap.Get() : 0;
+}
+
 void SCurveEditor::SetCurveOwner(FCurveOwnerInterface* InCurveOwner, bool bCanEdit)
 {
 	if(InCurveOwner != CurveOwner)
diff --git a/Engine/Source/Editor/UnrealEd/Private/SFbxSceneOptionWindow.h b/Engine/Source/Editor/UnrealEd/Private/SFbxSceneOptionWindow.h
--- a/Engine/Source/Editor/UnrealEd/Private/SFbxSceneOptionWindow.h
+++ b/Engine/Source/Editor/UnrealEd/Private/SFbxSceneOptionWindow.h
@@ -90,8 +90,6 @@ public:
 	ECheckBoxState IsReimportHierarchyChecked() const;
 	void OnToggleBakePivotInVertex(ECheckBoxState CheckType);
 	ECheckBoxState IsBakePivotInVertexChecked() const;
-	void OnToggleInvertNormalMap(ECheckBoxState CheckType);
-	ECheckBoxState IsInvertNormalMapChecked() const;
 
 	//Material UI
 	FText GetMaterialBasePath() const;
@@ -173,6 +171,7 @@ private:
 
 	//Scene Reimport tab variables
 	TSharedPtr<SFbxReimportSceneTreeView> SceneReimportTreeview;
+	TSharedPtr<IDetailsView> SceneReimportTabDetailsView;
 
 	//StaticMesh Reimport tab variables
 	TSharedPtr<SFbxSceneStaticMeshReimportListView> StaticMeshReimportListView;
diff --git a/Engine/Source/Editor/UnrealEd/Private/Settings/SettingsClasses.cpp b/Engine/Source/Editor/UnrealEd/Private/Settings/SettingsClasses.cpp
--- a/Engine/Source/Editor/UnrealEd/Private/Settings/SettingsClasses.cpp
+++ b/Engine/Source/Editor/UnrealEd/Private/Settings/SettingsClasses.cpp
@@ -116,7 +116,7 @@ UEditorLoadingSavingSettings::UEditorLoadingSavingSettings( const FObjectInitial
 	, AutoReimportThreshold(3.f)
 	, bAutoCreateAssets(true)
 	, bAutoDeleteAssets(true)
-	, bDetectChangesOnRestart(true)
+	, bDetectChangesOnStartup(false)
 	, bDeleteSourceFilesWithAssets(false)
 {
 	TextDiffToolPath.FilePath = TEXT("P4Merge.exe");
@@ -177,7 +177,7 @@ void UEditorLoadingSavingSettings::CheckSourceControlCompatability()
 		return;
 	}
 
-	if (ISourceControlModule::Get().IsEnabled() && bDetectChangesOnRestart)
+	if (ISourceControlModule::Get().IsEnabled() && bDetectChangesOnStartup)
 	{
 		// Persistent shared payload captured by the lambdas below
 		struct FPersistentPayload { TSharedPtr<SNotificationItem> Notification; };
@@ -187,7 +187,7 @@ void UEditorLoadingSavingSettings::CheckSourceControlCompatability()
 
  		auto OnTurnOffClicked = [=]{
 			auto* Settings = GetMutableDefault<UEditorLoadingSavingSettings>();
-			Settings->bDetectChangesOnRestart = false;
+			Settings->bDetectChangesOnStartup = false;
 			Settings->SaveConfig();
 
 			Payload->Notification->SetEnabled(false);
diff --git a/Engine/Source/Editor/UnrealEd/Private/SkeletalMeshEdit.cpp b/Engine/Source/Editor/UnrealEd/Private/SkeletalMeshEdit.cpp
--- a/Engine/Source/Editor/UnrealEd/Private/SkeletalMeshEdit.cpp
+++ b/Engine/Source/Editor/UnrealEd/Private/SkeletalMeshEdit.cpp
@@ -1008,6 +1008,7 @@ bool UnFbx::FFbxImporter::ImportAnimation(USkeleton* Skeleton, UAnimSequence * D
 		DestSeq->RawCurveData.FloatCurves.Shrink();
 	}
 
+	FbxNode *SkeletalMeshRootNode = NodeArray.Num() > 0 ? NodeArray[0] : nullptr;
 	//
 	// import blend shape curves
 	//
@@ -1180,7 +1181,7 @@ bool UnFbx::FFbxImporter::ImportAnimation(USkeleton* Skeleton, UAnimSequence * D
 
 				FbxNode* Link = SortedLinks[SourceTrackIdx];
 				FbxNode * LinkParent = Link->GetParent();
-			
+				FString LinkParentName = LinkParent->GetName();
 				for(FbxTime CurTime = AnimTimeSpan.GetStart(); CurTime <= AnimTimeSpan.GetStop(); CurTime += TimeIncrement)
 				{
 					// save global trasnform
@@ -1215,6 +1216,13 @@ bool UnFbx::FFbxImporter::ImportAnimation(USkeleton* Skeleton, UAnimSequence * D
 						// I can't rely on LocalMatrix. I need to recalculate quaternion/scale based on global transform if Parent exists
 						FbxAMatrix ParentGlobalMatrix = Link->GetParent()->EvaluateGlobalTransform(CurTime);
 						FTransform ParentGlobalTransform =  Converter.ConvertTransform(ParentGlobalMatrix);
+						//In case we do a scene import we need to add the skeletal mesh root node matrix to the parent link.
+						if (ImportOptions->bImportScene && !ImportOptions->bTransformVertexToAbsolute && BoneTreeIndex == 0 && SkeletalMeshRootNode != nullptr)
+						{
+							FbxAMatrix GlobalSkeletalNodeFbx = SkeletalMeshRootNode->EvaluateGlobalTransform(CurTime);
+							FTransform GlobalSkeletalNode = Converter.ConvertTransform(GlobalSkeletalNodeFbx);
+							ParentGlobalTransform = ParentGlobalTransform * GlobalSkeletalNode;
+						}
 
 						LocalTransform = GlobalTransform.GetRelativeTransform(ParentGlobalTransform);
 						NewDebugData.SourceParentGlobalTransform.Add(ParentGlobalTransform);
diff --git a/Engine/Source/Editor/UnrealEd/Private/Tests/AutomationEditorPromotionCommon.cpp b/Engine/Source/Editor/UnrealEd/Private/Tests/AutomationEditorPromotionCommon.cpp
--- a/Engine/Source/Editor/UnrealEd/Private/Tests/AutomationEditorPromotionCommon.cpp
+++ b/Engine/Source/Editor/UnrealEd/Private/Tests/AutomationEditorPromotionCommon.cpp
@@ -265,8 +265,10 @@ void FEditorPromotionTestUtilities::TakeScreenshot(const FString& ScreenshotName
 
 			TArray<FColor> OutImageData;
 			FIntVector OutImageSize;
-			FSlateApplication::Get().TakeScreenshot(WindowRef, OutImageData, OutImageSize);
-			FAutomationTestFramework::GetInstance().OnScreenshotCaptured().ExecuteIfBound(OutImageSize.X, OutImageSize.Y, OutImageData, ScreenshotFileName);
+			if (FSlateApplication::Get().TakeScreenshot(WindowRef, OutImageData, OutImageSize))
+			{
+				FAutomationTestFramework::GetInstance().OnScreenshotCaptured().ExecuteIfBound(OutImageSize.X, OutImageSize.Y, OutImageData, ScreenshotFileName);
+			}
 		}
 		else
 		{
diff --git a/Engine/Source/Editor/UnrealEd/Private/Tests/EditorAssetAutomationTests.cpp b/Engine/Source/Editor/UnrealEd/Private/Tests/EditorAssetAutomationTests.cpp
--- a/Engine/Source/Editor/UnrealEd/Private/Tests/EditorAssetAutomationTests.cpp
+++ b/Engine/Source/Editor/UnrealEd/Private/Tests/EditorAssetAutomationTests.cpp
@@ -916,8 +916,10 @@ namespace ImportExportAssetHelper
 
 				TArray<FColor> OutImageData;
 				FIntVector OutImageSize;
-				FSlateApplication::Get().TakeScreenshot(WindowRef, OutImageData, OutImageSize);
-				FAutomationTestFramework::GetInstance().OnScreenshotCaptured().ExecuteIfBound(OutImageSize.X, OutImageSize.Y, OutImageData, ScreenshotName);
+				if (FSlateApplication::Get().TakeScreenshot(WindowRef, OutImageData, OutImageSize))
+				{
+					FAutomationTestFramework::GetInstance().OnScreenshotCaptured().ExecuteIfBound(OutImageSize.X, OutImageSize.Y, OutImageData, ScreenshotName);
+				}
 
 				//Close the editor
 				FAssetEditorManager::Get().CloseAllAssetEditors();
diff --git a/Engine/Source/Editor/UnrealEd/Private/ThumbnailRendering/FontThumbnailRenderer.cpp b/Engine/Source/Editor/UnrealEd/Private/ThumbnailRendering/FontThumbnailRenderer.cpp
--- a/Engine/Source/Editor/UnrealEd/Private/ThumbnailRendering/FontThumbnailRenderer.cpp
+++ b/Engine/Source/Editor/UnrealEd/Private/ThumbnailRendering/FontThumbnailRenderer.cpp
@@ -67,7 +67,9 @@ void UFontThumbnailRenderer::Draw(UObject* Object, int32 X, int32 Y, uint32 Widt
 				static const int32 FontSize = 28;
 				static const float FontScale = 1.0f;
 
-				const FText FontNameText = FText::FromName(Object->GetFName());
+				const FString FontName = Object->GetName();
+				const TextBiDi::ETextDirection FontNameBaseTextDirection = TextBiDi::ComputeBaseDirection(FontName);
+
 				TSharedRef<FSlateFontCache> FontCache = FSlateApplication::Get().GetRenderer()->GetFontCache();
 
 				// Draw the object name using each of the default fonts
@@ -75,14 +77,18 @@ void UFontThumbnailRenderer::Draw(UObject* Object, int32 X, int32 Y, uint32 Widt
 				for(const FTypefaceEntry& TypefaceEntry : Font->CompositeFont.DefaultTypeface.Fonts)
 				{
 					FSlateFontInfo FontInfo(Font, FontSize, TypefaceEntry.Name);
+					FShapedGlyphSequenceRef FontNameShapedText = FontCache->ShapeBidirectionalText(FontName, FontInfo, FontScale, FontNameBaseTextDirection, ETextShapingMethod::Auto);
 
-					FCharacterList& CharacterList = FontCache->GetCharacterList(FontInfo, FontScale);
-					const int32 MaxCharHeight = CharacterList.GetMaxHeight();
+					FVector2D TextDrawPos = CurPos;
+					if (FontNameBaseTextDirection == TextBiDi::ETextDirection::RightToLeft)
+					{
+						TextDrawPos.X = Width - CurPos.X - FontNameShapedText->GetMeasuredWidth();
+					}
 
-					FCanvasTextItem TextItem(CurPos, FontNameText, FontInfo, FLinearColor::White);
+					FCanvasShapedTextItem TextItem(TextDrawPos, FontNameShapedText, FLinearColor::White);
 					Canvas->DrawItem(TextItem);
 
-					CurPos.Y += MaxCharHeight;
+					CurPos.Y += FontNameShapedText->GetMaxTextHeight();
 				}
 			}
 			break;
diff --git a/Engine/Source/Editor/UnrealEd/Private/UnrealEdSrv.cpp b/Engine/Source/Editor/UnrealEd/Private/UnrealEdSrv.cpp
--- a/Engine/Source/Editor/UnrealEd/Private/UnrealEdSrv.cpp
+++ b/Engine/Source/Editor/UnrealEd/Private/UnrealEdSrv.cpp
@@ -1285,12 +1285,13 @@ bool UUnrealEdEngine::Exec( UWorld* InWorld, const TCHAR* Stream, FOutputDevice&
 			{
 				TArray<FColor> OutImageData;
 				FIntVector OutImageSize;
-				FSlateApplication::Get().TakeScreenshot(InWidget,OutImageData,OutImageSize);
-
-				FString FileName;
-				const FString BaseFileName = FPaths::ScreenShotDir() / TEXT("EditorScreenshot");
-				FFileHelper::GenerateNextBitmapFilename(BaseFileName, TEXT("bmp"), FileName);
-				FFileHelper::CreateBitmap(*FileName, OutImageSize.X, OutImageSize.Y, OutImageData.GetData());
+				if (FSlateApplication::Get().TakeScreenshot(InWidget, OutImageData, OutImageSize))
+				{
+					FString FileName;
+					const FString BaseFileName = FPaths::ScreenShotDir() / TEXT("EditorScreenshot");
+					FFileHelper::GenerateNextBitmapFilename(BaseFileName, TEXT("bmp"), FileName);
+					FFileHelper::CreateBitmap(*FileName, OutImageSize.X, OutImageSize.Y, OutImageData.GetData());
+				}
 			}
 		};
 
diff --git a/Engine/Source/Editor/UnrealEd/Public/EditorReimportHandler.h b/Engine/Source/Editor/UnrealEd/Public/EditorReimportHandler.h
--- a/Engine/Source/Editor/UnrealEd/Public/EditorReimportHandler.h
+++ b/Engine/Source/Editor/UnrealEd/Public/EditorReimportHandler.h
@@ -46,10 +46,11 @@ public:
 	 * @param	Obj	Object to try reimporting
 	 * @param	bAskForNewFileIfMissing If the file is missing, open a dialog to ask for a new one
 	 * @param	bShowNotification True to show a notification when complete, false otherwise
+	 * @param	PreferedReimportFile if not empty, will be use in case the original file is missing and bAskForNewFileIfMissing is set to false
 	 *
 	 * @return	true if the object was handled by one of the reimport handlers; false otherwise
 	 */
-	UNREALED_API virtual bool Reimport( UObject* Obj, bool bAskForNewFileIfMissing = false, bool bShowNotification = true );
+	UNREALED_API virtual bool Reimport( UObject* Obj, bool bAskForNewFileIfMissing = false, bool bShowNotification = true, FString PreferedReimportFile = TEXT("") );
 
 	/**
 	 * Update the reimport paths for the specified object
diff --git a/Engine/Source/Editor/UnrealEd/Public/SCurveEditor.h b/Engine/Source/Editor/UnrealEd/Public/SCurveEditor.h
--- a/Engine/Source/Editor/UnrealEd/Public/SCurveEditor.h
+++ b/Engine/Source/Editor/UnrealEd/Public/SCurveEditor.h
@@ -440,6 +440,12 @@ private:
 	/** Paints the marquee for selection */
 	void PaintMarquee(const FGeometry& AllottedGeometry, const FSlateRect& MyClippingRect, FSlateWindowElementList& OutDrawElements, int32 LayerId) const;
 
+	/** Gets the delta value for the input value numeric entry box. */
+	float GetInputNumericEntryBoxDelta() const;
+
+	/** Gets the delta value for the output value numeric entry box. */
+	float GetOutputNumericEntryBoxDelta() const;
+
 protected:
 	//~ Begin SWidget Interface
 	UNREALED_API virtual FReply OnMouseButtonDown( const FGeometry& MyGeometry, const FPointerEvent& MouseEvent ) override;
diff --git a/Engine/Source/Runtime/AIModule/Classes/BehaviorTree/Tasks/BTTask_RunEQSQuery.h b/Engine/Source/Runtime/AIModule/Classes/BehaviorTree/Tasks/BTTask_RunEQSQuery.h
--- a/Engine/Source/Runtime/AIModule/Classes/BehaviorTree/Tasks/BTTask_RunEQSQuery.h
+++ b/Engine/Source/Runtime/AIModule/Classes/BehaviorTree/Tasks/BTTask_RunEQSQuery.h
@@ -40,7 +40,7 @@ class AIMODULE_API UBTTask_RunEQSQuery : public UBTTask_BlackboardBase
 	UPROPERTY(VisibleAnywhere, Category = Blackboard, meta = (EditCondition = "bUseBBKey", DisplayName = "EQSQueryBlackboardKey_DEPRECATED"))
 	struct FBlackboardKeySelector EQSQueryBlackboardKey;
 
-	UPROPERTY()
+	UPROPERTY(EditAnywhere, Category=Node, meta=(InlineEditConditionToggle))
 	bool bUseBBKey;
 
 	UPROPERTY(Category = EQS, EditAnywhere)
diff --git a/Engine/Source/Runtime/AIModule/Classes/EnvironmentQuery/EnvQueryTest.h b/Engine/Source/Runtime/AIModule/Classes/EnvironmentQuery/EnvQueryTest.h
--- a/Engine/Source/Runtime/AIModule/Classes/EnvironmentQuery/EnvQueryTest.h
+++ b/Engine/Source/Runtime/AIModule/Classes/EnvironmentQuery/EnvQueryTest.h
@@ -119,7 +119,7 @@ class AIMODULE_API UEnvQueryTest : public UEnvQueryNode
 	FAIDataProviderFloatValue SweetSpotValue;
 
 	/** When set to true enables usage of SweetSpotValue. It's false by default */
-	UPROPERTY()
+	UPROPERTY(EditDefaultsOnly, Category = Score, meta=(InlineEditConditionToggle))
 	bool bDefineSweetSpot;
 	
 	/** Validation: item type that can be used with this test */
diff --git a/Engine/Source/Runtime/AIModule/Classes/EnvironmentQuery/EnvQueryTypes.h b/Engine/Source/Runtime/AIModule/Classes/EnvironmentQuery/EnvQueryTypes.h
--- a/Engine/Source/Runtime/AIModule/Classes/EnvironmentQuery/EnvQueryTypes.h
+++ b/Engine/Source/Runtime/AIModule/Classes/EnvironmentQuery/EnvQueryTypes.h
@@ -1258,7 +1258,7 @@ struct FEQSParametrizedQueryExecutionRequest
 	UPROPERTY(Category = Node, EditAnywhere)
 	TEnumAsByte<EEnvQueryRunMode::Type> RunMode;
 
-	UPROPERTY()
+	UPROPERTY(Category = Node, EditAnywhere, meta=(InlineEditConditionToggle))
 	uint32 bUseBBKeyForQueryTemplate : 1;
 
 	uint32 bInitialized : 1;
diff --git a/Engine/Source/Runtime/AIModule/Classes/EnvironmentQuery/Generators/EnvQueryGenerator_Donut.h b/Engine/Source/Runtime/AIModule/Classes/EnvironmentQuery/Generators/EnvQueryGenerator_Donut.h
--- a/Engine/Source/Runtime/AIModule/Classes/EnvironmentQuery/Generators/EnvQueryGenerator_Donut.h
+++ b/Engine/Source/Runtime/AIModule/Classes/EnvironmentQuery/Generators/EnvQueryGenerator_Donut.h
@@ -44,7 +44,7 @@ class AIMODULE_API UEnvQueryGenerator_Donut : public UEnvQueryGenerator_Projecte
 	UPROPERTY(EditAnywhere, Category = Generator)
 	TSubclassOf<class UEnvQueryContext> Center;
 
-	UPROPERTY()
+	UPROPERTY(EditAnywhere, Category = Generator, meta=(InlineEditConditionToggle))
 	uint32 bDefineArc : 1;
 
 	virtual void GenerateItems(FEnvQueryInstance& QueryInstance) const override;
diff --git a/Engine/Source/Runtime/AIModule/Classes/EnvironmentQuery/Generators/EnvQueryGenerator_OnCircle.h b/Engine/Source/Runtime/AIModule/Classes/EnvironmentQuery/Generators/EnvQueryGenerator_OnCircle.h
--- a/Engine/Source/Runtime/AIModule/Classes/EnvironmentQuery/Generators/EnvQueryGenerator_OnCircle.h
+++ b/Engine/Source/Runtime/AIModule/Classes/EnvironmentQuery/Generators/EnvQueryGenerator_OnCircle.h
@@ -63,7 +63,7 @@ class AIMODULE_API UEnvQueryGenerator_OnCircle : public UEnvQueryGenerator_Proje
 	UPROPERTY(EditAnywhere, Category=Generator)
 	FEnvTraceData TraceData;
 
-	UPROPERTY()
+	UPROPERTY(EditAnywhere, Category=Generator, meta=(InlineEditConditionToggle))
 	uint32 bDefineArc:1;
 
 	virtual void PostLoad() override;
diff --git a/Engine/Source/Runtime/AIModule/Classes/Perception/AISenseConfig_Hearing.h b/Engine/Source/Runtime/AIModule/Classes/Perception/AISenseConfig_Hearing.h
--- a/Engine/Source/Runtime/AIModule/Classes/Perception/AISenseConfig_Hearing.h
+++ b/Engine/Source/Runtime/AIModule/Classes/Perception/AISenseConfig_Hearing.h
@@ -23,7 +23,7 @@ public:
 	float LoSHearingRange;
 
 	/** Warning: has significant runtime cost */
-	UPROPERTY()
+	UPROPERTY(EditDefaultsOnly, Category = "Sense", meta = (InlineEditConditionToggle))
 	uint32 bUseLoSHearing : 1;
 
 	UPROPERTY(EditDefaultsOnly, BlueprintReadOnly, Category = "Sense", config)
diff --git a/Engine/Source/Runtime/AppFramework/Private/Widgets/Testing/STestSuite.cpp b/Engine/Source/Runtime/AppFramework/Private/Widgets/Testing/STestSuite.cpp
--- a/Engine/Source/Runtime/AppFramework/Private/Widgets/Testing/STestSuite.cpp
+++ b/Engine/Source/Runtime/AppFramework/Private/Widgets/Testing/STestSuite.cpp
@@ -2057,6 +2057,8 @@ public:
 	BEGIN_SLATE_FUNCTION_BUILD_OPTIMIZATION
 	void Construct(const FArguments& InArgs)
 	{
+		bIsPassword = true;
+
 		InlineEditableText = LOCTEXT( "TestingInlineEditableTextBlock", "Testing inline editable text block!" );
 
 		Animation = FCurveSequence(0, 5);
@@ -2156,6 +2158,33 @@ public:
 					.HintText(LOCTEXT("DisabledContextMenuHint", "No context menu..."))
 					.OnContextMenuOpening(this, &STextEditTest::OnDisabledContextMenuOpening)
 				]
+				+ SVerticalBox::Slot().AutoHeight() .HAlign(HAlign_Center) .VAlign(VAlign_Center) .Padding( 5 )
+				[
+					SNew(SHorizontalBox)
+
+					+SHorizontalBox::Slot()
+					.AutoWidth()
+					[
+						SNew(SEditableTextBox)
+						.IsPassword(this, &STextEditTest::IsPassword)
+						.RevertTextOnEscape(true)
+						.MinDesiredWidth(400)
+						.HintText(LOCTEXT("EditablePasswordHintText", "This text box can be a password"))
+					]
+
+					+SHorizontalBox::Slot()
+					.AutoWidth()
+					.Padding(FMargin(4, 0))
+					[
+						SNew(SCheckBox)
+						.IsChecked(this, &STextEditTest::GetPasswordCheckState)
+						.OnCheckStateChanged(this, &STextEditTest::OnPasswordCheckStateChanged)
+						[
+							SNew(STextBlock)
+							.Text(LOCTEXT("PasswordCheckBoxText", "Password?"))
+						]
+					]
+				]
 			]
 		];
 	}
@@ -2201,6 +2230,21 @@ public:
 		return TSharedPtr<SWidget>();
 	}
 
+	bool IsPassword() const
+	{
+		return bIsPassword;
+	}
+
+	ECheckBoxState GetPasswordCheckState() const
+	{
+		return bIsPassword ? ECheckBoxState::Checked : ECheckBoxState::Unchecked;
+	}
+
+	void OnPasswordCheckStateChanged(ECheckBoxState NewState)
+	{
+		bIsPassword = NewState == ECheckBoxState::Checked;
+	}
+
 	void ClearSearchBox()
 	{
 		SearchBox->SetText( FText::GetEmpty() );
@@ -2279,6 +2323,7 @@ protected:
 	TSharedPtr<SInlineEditableTextBlock> InlineEditableTextBlock;
 	FText InlineEditableText;
 
+	bool bIsPassword;
  };
 
  #if WITH_FANCY_TEXT
diff --git a/Engine/Source/Runtime/Core/Private/Misc/FbxErrors.cpp b/Engine/Source/Runtime/Core/Private/Misc/FbxErrors.cpp
--- a/Engine/Source/Runtime/Core/Private/Misc/FbxErrors.cpp
+++ b/Engine/Source/Runtime/Core/Private/Misc/FbxErrors.cpp
@@ -116,6 +116,8 @@ FName FFbxErrors::SkeletalMesh_LOD_MissingSocketBone(TEXT("SkeletalMesh_LOD_Miss
 
 FName FFbxErrors::SkeletalMesh_LOD_MissingMorphTarget(TEXT("SkeletalMesh_LOD_MissingMorphTarget"));
 
+FName FFbxErrors::SkeletalMesh_FillImportDataFailed(TEXT("SkeletalMesh_FillImportDataFailed"));
+
 /** Animation **/
 FName FFbxErrors::Animation_CouldNotFindRootTrack(TEXT("Animation_CouldNotFindRootTrack"));
 
diff --git a/Engine/Source/Runtime/Core/Private/Windows/WindowsApplication.cpp b/Engine/Source/Runtime/Core/Private/Windows/WindowsApplication.cpp
--- a/Engine/Source/Runtime/Core/Private/Windows/WindowsApplication.cpp
+++ b/Engine/Source/Runtime/Core/Private/Windows/WindowsApplication.cpp
@@ -50,6 +50,7 @@ FWindowsApplication::FWindowsApplication( const HINSTANCE HInstance, const HICON
 	, InstanceHandle( HInstance )
 	, bUsingHighPrecisionMouseInput( false )
 	, bIsMouseAttached( false )
+	, bForceActivateByMouse( false )
 	, XInput( XInputInterface::Create( MessageHandler ) )
 	, bHasLoadedInputPlugins( false )
 	, bAllowedToDeferMessageProcessing(true)
@@ -919,6 +920,13 @@ int32 FWindowsApplication::ProcessMessage( HWND hwnd, uint32 msg, WPARAM wParam,
 			break;
 #endif
 
+			// Window focus and activation
+		case WM_MOUSEACTIVATE:
+			{
+				DeferMessage(CurrentNativeEventWindowPtr, hwnd, msg, wParam, lParam);
+			}
+			break;
+
 			// Window focus and activation
 		case WM_ACTIVATE:
 			{
@@ -1486,6 +1494,17 @@ int32 FWindowsApplication::ProcessDeferredMessage( const FDeferredWindowsMessage
 			}
 			break;
 
+			// Window focus and activation
+		case WM_MOUSEACTIVATE:
+			{
+				// If the mouse activate isn't in the client area we'll force the WM_ACTIVATE to be EWindowActivation::ActivateByMouse
+				// This ensures that clicking menu buttons on the header doesn't generate a WM_ACTIVATE with EWindowActivation::Activate
+				// which may cause mouse capture to be taken because is not differentiable from Alt-Tabbing back to the application.
+				bForceActivateByMouse = !(LOWORD(lParam) & HTCLIENT);
+				return 0;
+			}
+			break;
+
 			// Window focus and activation
 		case WM_ACTIVATE:
 			{
@@ -1493,7 +1512,7 @@ int32 FWindowsApplication::ProcessDeferredMessage( const FDeferredWindowsMessage
 
 				if (LOWORD(wParam) & WA_ACTIVE)
 				{
-					ActivationType = EWindowActivation::Activate;
+					ActivationType = bForceActivateByMouse ? EWindowActivation::ActivateByMouse : EWindowActivation::Activate;
 				}
 				else if (LOWORD(wParam) & WA_CLICKACTIVE)
 				{
@@ -1503,6 +1522,7 @@ int32 FWindowsApplication::ProcessDeferredMessage( const FDeferredWindowsMessage
 				{
 					ActivationType = EWindowActivation::Deactivate;
 				}
+				bForceActivateByMouse = false;
 
 				if ( CurrentNativeEventWindowPtr.IsValid() )
 				{
diff --git a/Engine/Source/Runtime/Core/Public/GenericPlatform/ITextInputMethodSystem.h b/Engine/Source/Runtime/Core/Public/GenericPlatform/ITextInputMethodSystem.h
--- a/Engine/Source/Runtime/Core/Public/GenericPlatform/ITextInputMethodSystem.h
+++ b/Engine/Source/Runtime/Core/Public/GenericPlatform/ITextInputMethodSystem.h
@@ -18,6 +18,8 @@ public:
 	};
 
 public:
+	virtual ~ITextInputMethodContext() {}
+
 	/**
 	 * Returns whether or not this text is read-only.
 	 *
@@ -141,6 +143,8 @@ public:
 	};
 
 public:
+	virtual ~ITextInputMethodChangeNotifier() {}
+
 	virtual void NotifyLayoutChanged(const ELayoutChangeType ChangeType) = 0;
 	virtual void NotifySelectionChanged() = 0;
 	virtual void NotifyTextChanged(const uint32 BeginIndex, const uint32 OldLength, const uint32 NewLength) = 0;
@@ -153,6 +157,8 @@ public:
 class ITextInputMethodSystem
 {
 public:
+	virtual ~ITextInputMethodSystem() {}
+
 	/**
 	 * Called when a window is first created to allow the text input system to apply any default settings
 	 */
diff --git a/Engine/Source/Runtime/Core/Public/Misc/FbxErrors.h b/Engine/Source/Runtime/Core/Public/Misc/FbxErrors.h
--- a/Engine/Source/Runtime/Core/Public/Misc/FbxErrors.h
+++ b/Engine/Source/Runtime/Core/Public/Misc/FbxErrors.h
@@ -124,6 +124,8 @@ struct CORE_API FFbxErrors
 
 	static FName SkeletalMesh_LOD_MissingMorphTarget;
 
+	static FName SkeletalMesh_FillImportDataFailed;
+
 	/** Animation **/
 	static FName Animation_CouldNotFindRootTrack;
 
diff --git a/Engine/Source/Runtime/Core/Public/Windows/WindowsApplication.h b/Engine/Source/Runtime/Core/Public/Windows/WindowsApplication.h
--- a/Engine/Source/Runtime/Core/Public/Windows/WindowsApplication.h
+++ b/Engine/Source/Runtime/Core/Public/Windows/WindowsApplication.h
@@ -432,6 +432,8 @@ private:
 
 	bool bIsMouseAttached;
 
+	bool bForceActivateByMouse;
+
 	TArray<FDeferredWindowsMessage> DeferredMessages;
 
 	TArray<FDeferredWindowsDragDropOperation> DeferredDragDropOperations;
diff --git a/Engine/Source/Runtime/CoreUObject/Private/Serialization/PropertyLocalizationDataGathering.cpp b/Engine/Source/Runtime/CoreUObject/Private/Serialization/PropertyLocalizationDataGathering.cpp
--- a/Engine/Source/Runtime/CoreUObject/Private/Serialization/PropertyLocalizationDataGathering.cpp
+++ b/Engine/Source/Runtime/CoreUObject/Private/Serialization/PropertyLocalizationDataGathering.cpp
@@ -406,22 +406,50 @@ void FPropertyLocalizationDataGatherer::GatherScriptBytecode(const FString& Path
 
 		void SerializeText(int32& iCode, FArchive& Ar)
 		{
-			bIsParsingText = true;
+			// What kind of text are we dealing with?
+			const EBlueprintTextLiteralType TextLiteralType = (EBlueprintTextLiteralType)Script[iCode++];
 
-			SerializeExpr(iCode, Ar);
-			const FString SourceString = MoveTemp(LastParsedString);
+			switch (TextLiteralType)
+			{
+			case EBlueprintTextLiteralType::Empty:
+				// Don't need to gather empty text
+				break;
+
+			case EBlueprintTextLiteralType::LocalizedText:
+				{
+					bIsParsingText = true;
 
-			SerializeExpr(iCode, Ar);
-			const FString TextKey = MoveTemp(LastParsedString);
+					SerializeExpr(iCode, Ar);
+					const FString SourceString = MoveTemp(LastParsedString);
 
-			SerializeExpr(iCode, Ar);
-			const FString TextNamespace = MoveTemp(LastParsedString);
+					SerializeExpr(iCode, Ar);
+					const FString TextKey = MoveTemp(LastParsedString);
 
-			const FText TextInstance = FInternationalization::ForUseOnlyByLocMacroAndGraphNodeTextLiterals_CreateText(*SourceString, *TextNamespace, *TextKey);
+					SerializeExpr(iCode, Ar);
+					const FString TextNamespace = MoveTemp(LastParsedString);
 
-			PropertyLocalizationDataGatherer.GatherTextInstance(TextInstance, FString::Printf(TEXT("%s [Script Bytecode]"), SourceDescription), bTreatAsEditorOnlyData);
+					bIsParsingText = false;
 
-			bIsParsingText = false;
+					const FText TextInstance = FInternationalization::ForUseOnlyByLocMacroAndGraphNodeTextLiterals_CreateText(*SourceString, *TextNamespace, *TextKey);
+
+					PropertyLocalizationDataGatherer.GatherTextInstance(TextInstance, FString::Printf(TEXT("%s [Script Bytecode]"), SourceDescription), bTreatAsEditorOnlyData);
+				}
+				break;
+
+			case EBlueprintTextLiteralType::InvariantText:
+				// Don't need to gather invariant text, but we do need to walk over the string in the buffer
+				SerializeExpr(iCode, Ar);
+				break;
+
+			case EBlueprintTextLiteralType::LiteralString:
+				// Don't need to gather literal strings, but we do need to walk over the string in the buffer
+				SerializeExpr(iCode, Ar);
+				break;
+
+			default:
+				checkf(false, TEXT("Unknown EBlueprintTextLiteralType! Please update FGatherTextFromScriptBytecode::SerializeText to handle this type of text."));
+				break;
+			}
 		}
 
 		const TCHAR* SourceDescription;
diff --git a/Engine/Source/Runtime/CoreUObject/Private/UObject/Obj.cpp b/Engine/Source/Runtime/CoreUObject/Private/UObject/Obj.cpp
--- a/Engine/Source/Runtime/CoreUObject/Private/UObject/Obj.cpp
+++ b/Engine/Source/Runtime/CoreUObject/Private/UObject/Obj.cpp
@@ -322,6 +322,21 @@ void UObject::PostEditChangeChainProperty(FPropertyChangedChainEvent& PropertyCh
 {
 	FPropertyChangedEvent PropertyEvent(PropertyChangedEvent.PropertyChain.GetActiveNode()->GetValue(), PropertyChangedEvent.ChangeType);
 
+	// Set up array index per object map so that GetArrayIndex returns a valid result
+	TArray<TMap<FString, int32>> ArrayIndexForProperty;
+	if (PropertyChangedEvent.Property)
+	{
+		const FString PropertyName = PropertyChangedEvent.Property->GetName();
+		const int32 ArrayIndex = PropertyChangedEvent.GetArrayIndex(PropertyName);
+		if (ArrayIndex != INDEX_NONE)
+		{
+			PropertyEvent.ObjectIteratorIndex = 0;
+			ArrayIndexForProperty.AddDefaulted();
+			ArrayIndexForProperty.Last().Add(PropertyName, ArrayIndex);
+			PropertyEvent.SetArrayIndexPerObject(ArrayIndexForProperty);
+		}
+	}
+
 	if( PropertyChangedEvent.PropertyChain.GetActiveMemberNode() )
 	{
 		PropertyEvent.SetActiveMemberProperty( PropertyChangedEvent.PropertyChain.GetActiveMemberNode()->GetValue() );
diff --git a/Engine/Source/Runtime/CoreUObject/Private/UObject/ScriptCore.cpp b/Engine/Source/Runtime/CoreUObject/Private/UObject/ScriptCore.cpp
--- a/Engine/Source/Runtime/CoreUObject/Private/UObject/ScriptCore.cpp
+++ b/Engine/Source/Runtime/CoreUObject/Private/UObject/ScriptCore.cpp
@@ -2307,13 +2307,54 @@ IMPLEMENT_VM_FUNCTION( EX_UnicodeStringConst, execUnicodeStringConst );
 
 void UObject::execTextConst( FFrame& Stack, RESULT_DECL )
 {
-	FString SourceString;
-	FString KeyString;
-	FString Namespace;
-	Stack.Step( Stack.Object, &SourceString);
-	Stack.Step( Stack.Object, &KeyString);
-	Stack.Step( Stack.Object, &Namespace);
-	*(FText*)RESULT_PARAM = FInternationalization::ForUseOnlyByLocMacroAndGraphNodeTextLiterals_CreateText(*SourceString, *Namespace, *KeyString);
+	// What kind of text are we dealing with?
+	const EBlueprintTextLiteralType TextLiteralType = (EBlueprintTextLiteralType)*Stack.Code++;
+
+	switch (TextLiteralType)
+	{
+	case EBlueprintTextLiteralType::Empty:
+		{
+			*(FText*)RESULT_PARAM = FText::GetEmpty();
+		}
+		break;
+
+	case EBlueprintTextLiteralType::LocalizedText:
+		{
+			FString SourceString;
+			Stack.Step(Stack.Object, &SourceString);
+
+			FString KeyString;
+			Stack.Step(Stack.Object, &KeyString);
+			
+			FString Namespace;
+			Stack.Step(Stack.Object, &Namespace);
+
+			*(FText*)RESULT_PARAM = FInternationalization::ForUseOnlyByLocMacroAndGraphNodeTextLiterals_CreateText(*SourceString, *Namespace, *KeyString);
+		}
+		break;
+
+	case EBlueprintTextLiteralType::InvariantText:
+		{
+			FString SourceString;
+			Stack.Step(Stack.Object, &SourceString);
+
+			*(FText*)RESULT_PARAM = FText::AsCultureInvariant(MoveTemp(SourceString));
+		}
+		break;
+
+	case EBlueprintTextLiteralType::LiteralString:
+		{
+			FString SourceString;
+			Stack.Step(Stack.Object, &SourceString);
+
+			*(FText*)RESULT_PARAM = FText::FromString(MoveTemp(SourceString));
+		}
+		break;
+
+	default:
+		checkf(false, TEXT("Unknown EBlueprintTextLiteralType! Please update UObject::execTextConst to handle this type of text."));
+		break;
+	}
 }
 IMPLEMENT_VM_FUNCTION( EX_TextConst, execTextConst );
 
diff --git a/Engine/Source/Runtime/CoreUObject/Private/UObject/UTextProperty.cpp b/Engine/Source/Runtime/CoreUObject/Private/UObject/UTextProperty.cpp
--- a/Engine/Source/Runtime/CoreUObject/Private/UObject/UTextProperty.cpp
+++ b/Engine/Source/Runtime/CoreUObject/Private/UObject/UTextProperty.cpp
@@ -53,6 +53,19 @@ void UTextProperty::ExportTextItem( FString& ValueStr, const void* PropertyValue
 	{
 		ValueStr += GenerateCppCodeForTextValue(TextValue, FString());
 	}
+	else if (PortFlags & PPF_PropertyWindow)
+	{
+		if (PortFlags & PPF_Delimited)
+		{
+			ValueStr += TEXT("\"");
+			ValueStr += TextValue.ToString();
+			ValueStr += TEXT("\"");
+		}
+		else
+		{
+			ValueStr += TextValue.ToString();
+		}
+	}
 	else
 	{
 		FTextStringHelper::WriteToString(ValueStr, TextValue, !!(PortFlags & PPF_Delimited));
diff --git a/Engine/Source/Runtime/CoreUObject/Public/UObject/ObjectBase.h b/Engine/Source/Runtime/CoreUObject/Public/UObject/ObjectBase.h
--- a/Engine/Source/Runtime/CoreUObject/Public/UObject/ObjectBase.h
+++ b/Engine/Source/Runtime/CoreUObject/Public/UObject/ObjectBase.h
@@ -1072,6 +1072,9 @@ namespace UM
 		/// [PropertyMetadata] Used for FColor and FLinearColor properties. Indicates that the Alpha property should be hidden when displaying the property widget in the details.
 		HideAlphaChannel,
 
+		/// [PropertyMetadata] Signifies that the bool property is only displayed inline as an edit condition toggle in other properties, and should not be shown on its own row.
+		InlineEditConditionToggle,
+
 		/// [PropertyMetadata] Used for FStringClassReference properties.  Indicates whether only blueprint classes should be shown in the class picker.
 		IsBlueprintBaseOnly,
 
diff --git a/Engine/Source/Runtime/CoreUObject/Public/UObject/Script.h b/Engine/Source/Runtime/CoreUObject/Public/UObject/Script.h
--- a/Engine/Source/Runtime/CoreUObject/Public/UObject/Script.h
+++ b/Engine/Source/Runtime/CoreUObject/Public/UObject/Script.h
@@ -268,6 +268,19 @@ enum ECastToken
 	CST_Max					= 0xFF,
 };
 
+// Kinds of text literals
+enum class EBlueprintTextLiteralType : uint8
+{
+	/* Text is an empty string. The bytecode contains no strings, and you should use FText::GetEmpty() to initialize the FText instance. */
+	Empty,
+	/** Text is localized. The bytecode will contain three strings - source, key, and namespace - and should be loaded via FInternationalization */
+	LocalizedText,
+	/** Text is culture invariant. The bytecode will contain one string, and you should use FText::AsCultureInvariant to initialize the FText instance. */
+	InvariantText,
+	/** Text is a literal FString. The bytecode will contain one string, and you should use FText::FromString to initialize the FText instance. */
+	LiteralString,
+};
+
 // Kinds of Blueprint exceptions
 namespace EBlueprintExceptionType
 {
diff --git a/Engine/Source/Runtime/CoreUObject/Public/UObject/ScriptSerialization.h b/Engine/Source/Runtime/CoreUObject/Public/UObject/ScriptSerialization.h
--- a/Engine/Source/Runtime/CoreUObject/Public/UObject/ScriptSerialization.h
+++ b/Engine/Source/Runtime/CoreUObject/Public/UObject/ScriptSerialization.h
@@ -74,9 +74,27 @@
 #ifndef XFERTEXT
 	#define XFERTEXT() \
 	{ \
-		SerializeExpr( iCode, Ar );	\
-		SerializeExpr( iCode, Ar ); \
-		SerializeExpr( iCode, Ar ); \
+		XFER(uint8); \
+		const EBlueprintTextLiteralType TextLiteralType = (EBlueprintTextLiteralType)Script[iCode - 1]; \
+		switch (TextLiteralType) \
+		{ \
+		case EBlueprintTextLiteralType::Empty: \
+			break; \
+		case EBlueprintTextLiteralType::LocalizedText: \
+			SerializeExpr( iCode, Ar );	\
+			SerializeExpr( iCode, Ar ); \
+			SerializeExpr( iCode, Ar ); \
+			break; \
+		case EBlueprintTextLiteralType::InvariantText: \
+			SerializeExpr( iCode, Ar );	\
+			break; \
+		case EBlueprintTextLiteralType::LiteralString: \
+			SerializeExpr( iCode, Ar );	\
+			break; \
+		default: \
+			checkf(false, TEXT("Unknown EBlueprintTextLiteralType! Please update XFERTEXT to handle this type of text.")); \
+			break; \
+		} \
 	}
 #endif	//XFERTEXT
 
diff --git a/Engine/Source/Runtime/Engine/Classes/AI/Navigation/NavFilters/NavigationQueryFilter.h b/Engine/Source/Runtime/Engine/Classes/AI/Navigation/NavFilters/NavigationQueryFilter.h
--- a/Engine/Source/Runtime/Engine/Classes/AI/Navigation/NavFilters/NavigationQueryFilter.h
+++ b/Engine/Source/Runtime/Engine/Classes/AI/Navigation/NavFilters/NavigationQueryFilter.h
@@ -27,10 +27,10 @@ struct ENGINE_API FNavigationFilterArea
 	UPROPERTY(EditAnywhere, Category=Area)
 	uint32 bIsExcluded : 1;
 
-	UPROPERTY()
+	UPROPERTY(EditAnywhere, Category=Area, meta=(InlineEditConditionToggle))
 	uint32 bOverrideTravelCost : 1;
 
-	UPROPERTY()
+	UPROPERTY(EditAnywhere, Category=Area, meta=(InlineEditConditionToggle))
 	uint32 bOverrideEnteringCost : 1;
 
 	FNavigationFilterArea() : TravelCostOverride(1.0f) {}
diff --git a/Engine/Source/Runtime/Engine/Classes/AI/Navigation/NavLinkDefinition.h b/Engine/Source/Runtime/Engine/Classes/AI/Navigation/NavLinkDefinition.h
--- a/Engine/Source/Runtime/Engine/Classes/AI/Navigation/NavLinkDefinition.h
+++ b/Engine/Source/Runtime/Engine/Classes/AI/Navigation/NavLinkDefinition.h
@@ -41,7 +41,7 @@ struct ENGINE_API FNavigationLinkBase
 	UPROPERTY(EditAnywhere, Category=Default, meta=(ClampMin = "1.0", EditCondition="bUseSnapHeight"))
 	float SnapHeight;
 
-	UPROPERTY()
+	UPROPERTY(EditAnywhere, Category=Default, meta=(InlineEditConditionToggle))
 	uint32 bUseSnapHeight : 1;
 
 	/** If set, link will try to snap to cheapest area in given radius */
diff --git a/Engine/Source/Runtime/Engine/Classes/AI/Navigation/NavigationTestingActor.h b/Engine/Source/Runtime/Engine/Classes/AI/Navigation/NavigationTestingActor.h
--- a/Engine/Source/Runtime/Engine/Classes/AI/Navigation/NavigationTestingActor.h
+++ b/Engine/Source/Runtime/Engine/Classes/AI/Navigation/NavigationTestingActor.h
@@ -50,7 +50,7 @@ private_subobject:
 	UPROPERTY(EditAnywhere, Category = Navigation, meta=(EditCondition="bActAsNavigationInvoker"))
 	UNavigationInvokerComponent* InvokerComponent;
 
-	UPROPERTY()
+	UPROPERTY(EditAnywhere, Category = Navigation, meta=(InlineEditConditionToggle))
 	uint32 bActAsNavigationInvoker : 1;
 
 public:
diff --git a/Engine/Source/Runtime/Engine/Classes/Animation/VertexAnim/MorphTarget.h b/Engine/Source/Runtime/Engine/Classes/Animation/VertexAnim/MorphTarget.h
--- a/Engine/Source/Runtime/Engine/Classes/Animation/VertexAnim/MorphTarget.h
+++ b/Engine/Source/Runtime/Engine/Classes/Animation/VertexAnim/MorphTarget.h
@@ -27,19 +27,20 @@ public:
 	TArray<uint32> Indices;
 	/** indices to original imported wedge points */
 	TArray<uint32> WedgePointIndices;
-	/** mesh that provided the source data */
-	UObject* SourceMesh;
 
 	/** Constructor (default) */
-	ENGINE_API FMorphMeshRawSource() : 
-	SourceMesh(NULL)
+	ENGINE_API FMorphMeshRawSource() 
 	{
 	}
 
 	ENGINE_API FMorphMeshRawSource( USkeletalMesh* SrcMesh, int32 LODIndex=0 );
 	ENGINE_API FMorphMeshRawSource( UStaticMesh* SrcMesh, int32 LODIndex=0 );
+	ENGINE_API FMorphMeshRawSource( FStaticLODModel& LODModel );
 
 	ENGINE_API bool IsValidTarget( const FMorphMeshRawSource& Target ) const;
+
+private:
+	void Initialize(FStaticLODModel& LODModel);
 };
 
 
diff --git a/Engine/Source/Runtime/Engine/Classes/Components/SplineComponent.h b/Engine/Source/Runtime/Engine/Classes/Components/SplineComponent.h
--- a/Engine/Source/Runtime/Engine/Classes/Components/SplineComponent.h
+++ b/Engine/Source/Runtime/Engine/Classes/Components/SplineComponent.h
@@ -76,6 +76,12 @@ class ENGINE_API USplineComponent : public UPrimitiveComponent
 	UPROPERTY(EditAnywhere, Category = Spline, meta=(DisplayName="Override Construction Script"))
 	bool bSplineHasBeenEdited;
 
+	UPROPERTY(EditAnywhere, BlueprintReadWrite, Category = Spline)
+	bool bInputSplinePointsToConstructionScript;
+
+	UPROPERTY(EditAnywhere, BlueprintReadWrite, Category = Spline)
+	bool bAlwaysRenderInEditor;
+
 private:
 	/**
 	 * Whether the spline is to be considered as a closed loop.
@@ -107,10 +113,29 @@ public:
 	float ScaleVisualizationWidth;
 #endif
 
+	//~ Begin UObject Interface
+	virtual void PostLoad() override;
+	virtual void PostEditImport() override;
+#if WITH_EDITOR
+	virtual void Serialize(FArchive& Ar) override;
+	virtual void PostEditChangeChainProperty(FPropertyChangedChainEvent& PropertyChangedEvent) override;
+#endif
+	//~ End UObject Interface
+
 	//~ Begin UActorComponent Interface.
 	virtual FActorComponentInstanceData* GetComponentInstanceData() const override;
 	//~ End UActorComponent Interface.
 
+#if WITH_EDITOR
+	//~ Begin UPrimitiveComponent Interface.
+	virtual FPrimitiveSceneProxy* CreateSceneProxy() override;
+	//~ End UPrimitiveComponent Interface.
+
+	//~ Begin USceneComponent Interface
+	virtual FBoxSphereBounds CalcBounds(const FTransform& LocalToWorld) const override;
+	//~ End USceneComponent Interface
+#endif
+
 	void ApplyComponentInstanceData(class FSplineInstanceData* ComponentInstanceData, const bool bPostUCS);
 
 	/** Update the spline tangents and SplineReparamTable */
@@ -214,6 +239,10 @@ public:
 	UFUNCTION(BlueprintCallable, Category = Spline)
 	void SetTangentAtSplinePoint(int32 PointIndex, const FVector& InTangent, ESplineCoordinateSpace::Type CoordinateSpace);
 
+	/** Specify the up vector at a given spline point */
+	UFUNCTION(BlueprintCallable, Category = Spline)
+	void SetUpVectorAtSplinePoint(int32 PointIndex, const FVector& InUpVector, ESplineCoordinateSpace::Type CoordinateSpace);
+
 	/** Get the type of a spline point */
 	UFUNCTION(BlueprintCallable, Category = Spline)
 	ESplinePointType::Type GetSplinePointType(int32 PointIndex) const;
@@ -243,7 +272,6 @@ public:
 	FVector GetTangentAtSplinePoint(int32 PointIndex, ESplineCoordinateSpace::Type CoordinateSpace) const;
 
 	/** Get the rotation at spline point as a quaternion */
-	UFUNCTION(BlueprintCallable, Category = Spline)
 	FQuat GetQuaternionAtSplinePoint(int32 PointIndex, ESplineCoordinateSpace::Type CoordinateSpace) const;
 
 	/** Get the rotation at spline point as a rotator */
@@ -323,7 +351,6 @@ public:
 	FVector GetWorldTangentAtDistanceAlongSpline(float Distance) const;
 
 	/** Given a distance along the length of this spline, return a quaternion corresponding to the spline's rotation there. */
-	UFUNCTION(BlueprintCallable, Category = Spline)
 	FQuat GetQuaternionAtDistanceAlongSpline(float Distance, ESplineCoordinateSpace::Type CoordinateSpace) const;
 
 	/** Given a distance along the length of this spline, return a rotation corresponding to the spline's rotation there. */
@@ -375,7 +402,6 @@ public:
 	FVector GetTangentAtTime(float Time, ESplineCoordinateSpace::Type CoordinateSpace, bool bUseConstantVelocity = false) const;
 
 	/** Given a time from 0 to the spline duration, return a quaternion corresponding to the spline's rotation there. */
-	UFUNCTION(BlueprintCallable, Category=Spline)
 	FQuat GetQuaternionAtTime(float Time, ESplineCoordinateSpace::Type CoordinateSpace, bool bUseConstantVelocity = false) const;
 
 	/** Given a time from 0 to the spline duration, return a rotation corresponding to the spline's position and direction there. */
@@ -423,7 +449,6 @@ public:
 	FVector FindTangentClosestToWorldLocation(const FVector& WorldLocation, ESplineCoordinateSpace::Type CoordinateSpace) const;
 
 	/** Given a location, in world space, return a quaternion corresponding to the spline's rotation closest to the location. */
-	UFUNCTION(BlueprintCallable, Category = Spline)
 	FQuat FindQuaternionClosestToWorldLocation(const FVector& WorldLocation, ESplineCoordinateSpace::Type CoordinateSpace) const;
 
 	/** Given a location, in world space, return rotation corresponding to the spline's rotation closest to the location. */
@@ -451,15 +476,6 @@ public:
 	FTransform FindTransformClosestToWorldLocation(const FVector& WorldLocation, ESplineCoordinateSpace::Type CoordinateSpace, bool bUseScale = false) const;
 
 
-	//~ Begin UObject Interface
-	virtual void PostLoad() override;
-	virtual void PostEditImport() override;
-#if WITH_EDITOR
-	virtual void Serialize(FArchive& Ar) override;
-	virtual void PostEditChangeChainProperty(FPropertyChangedChainEvent& PropertyChangedEvent) override;
-#endif
-	//~ End UObject Interface
-
 private:
 	/** Returns the length of the specified spline segment up to the parametric value given */
 	float GetSegmentLength(const int32 Index, const float Param = 1.0f) const;
diff --git a/Engine/Source/Runtime/Engine/Classes/Components/StaticMeshComponent.h b/Engine/Source/Runtime/Engine/Classes/Components/StaticMeshComponent.h
--- a/Engine/Source/Runtime/Engine/Classes/Components/StaticMeshComponent.h
+++ b/Engine/Source/Runtime/Engine/Classes/Components/StaticMeshComponent.h
@@ -160,7 +160,7 @@ class ENGINE_API UStaticMeshComponent : public UMeshComponent
 	void OnRep_StaticMesh(class UStaticMesh *OldStaticMesh);
 
 	/** If true, WireframeColorOverride will be used. If false, color is determined based on mobility and physics simulation settings */
-	UPROPERTY(BlueprintReadOnly, Category=Rendering)
+	UPROPERTY(EditAnywhere, BlueprintReadOnly, Category=Rendering, meta=(InlineEditConditionToggle))
 	bool bOverrideWireframeColor;
 
 	/** Wireframe color to use if bOverrideWireframeColor is true */
@@ -206,7 +206,7 @@ class ENGINE_API UStaticMeshComponent : public UMeshComponent
 	uint32 bIgnoreInstanceForTextureStreaming:1;
 
 	/** Whether to override the lightmap resolution defined in the static mesh. */
-	UPROPERTY(BlueprintReadOnly, Category=Lighting)
+	UPROPERTY(EditAnywhere, BlueprintReadOnly, Category=Lighting, meta=(InlineEditConditionToggle))
 	uint32 bOverrideLightMapRes:1;
 
 	/** Light map resolution to use on this component, used if bOverrideLightMapRes is true */
diff --git a/Engine/Source/Runtime/Engine/Classes/Engine/EngineTypes.h b/Engine/Source/Runtime/Engine/Classes/Engine/EngineTypes.h
--- a/Engine/Source/Runtime/Engine/Classes/Engine/EngineTypes.h
+++ b/Engine/Source/Runtime/Engine/Classes/Engine/EngineTypes.h
@@ -2450,7 +2450,7 @@ struct FMaterialSimplificationSettings
 	FIntPoint BaseColorMapSize;
 
 	// Whether to generate normal map
-	UPROPERTY()
+	UPROPERTY(Category=Material, EditAnywhere, meta=(InlineEditConditionToggle))
 	bool bNormalMap;
 	
 	// Size of generated specular map
@@ -2462,7 +2462,7 @@ struct FMaterialSimplificationSettings
 	float MetallicConstant;
 
 	// Whether to generate metallic map
-	UPROPERTY()
+	UPROPERTY(Category=Material, EditAnywhere, meta=(InlineEditConditionToggle))
 	bool bMetallicMap;
 
 	// Size of generated metallic map
@@ -2474,7 +2474,7 @@ struct FMaterialSimplificationSettings
 	float RoughnessConstant;
 
 	// Whether to generate roughness map
-	UPROPERTY()
+	UPROPERTY(Category=Material, EditAnywhere, meta=(InlineEditConditionToggle))
 	bool bRoughnessMap;
 
 	// Size of generated roughness map
@@ -2486,7 +2486,7 @@ struct FMaterialSimplificationSettings
 	float SpecularConstant;
 
 	// Whether to generate specular map
-	UPROPERTY()
+	UPROPERTY(Category=Material, EditAnywhere, meta=(InlineEditConditionToggle))
 	bool bSpecularMap;
 
 	// Size of generated specular map
diff --git a/Engine/Source/Runtime/Engine/Classes/Engine/GameViewportClient.h b/Engine/Source/Runtime/Engine/Classes/Engine/GameViewportClient.h
--- a/Engine/Source/Runtime/Engine/Classes/Engine/GameViewportClient.h
+++ b/Engine/Source/Runtime/Engine/Classes/Engine/GameViewportClient.h
@@ -156,6 +156,8 @@ public:
 	virtual void LostFocus(FViewport* Viewport) override;
 	virtual void ReceivedFocus(FViewport* Viewport) override;
 	virtual bool IsFocused(FViewport* Viewport) override;
+	virtual void Activated(FViewport* InViewport, const FWindowActivateEvent& InActivateEvent) override;
+	virtual void Deactivated(FViewport* InViewport, const FWindowActivateEvent& InActivateEvent) override;
 	virtual void CloseRequested(FViewport* Viewport) override;
 	virtual bool RequiresHitProxyStorage() override { return 0; }
 	virtual bool IsOrtho() const override;
diff --git a/Engine/Source/Runtime/Engine/Classes/Engine/Level.h b/Engine/Source/Runtime/Engine/Classes/Engine/Level.h
--- a/Engine/Source/Runtime/Engine/Classes/Engine/Level.h
+++ b/Engine/Source/Runtime/Engine/Classes/Engine/Level.h
@@ -258,7 +258,7 @@ struct ENGINE_API FLevelSimplificationDetails
 	UPROPERTY(Category = Landscape, EditAnywhere)
 	FMaterialProxySettings StaticMeshMaterialSettings;
 
-	UPROPERTY()
+	UPROPERTY(Category = Landscape, EditAnywhere, meta=(InlineEditConditionToggle))
 	bool bOverrideLandscapeExportLOD;
 
 	/** Landscape LOD to use for static mesh generation, when not specified 'Max LODLevel' from landscape actor will be used */
diff --git a/Engine/Source/Runtime/Engine/Classes/Materials/MaterialExpressionTime.h b/Engine/Source/Runtime/Engine/Classes/Materials/MaterialExpressionTime.h
--- a/Engine/Source/Runtime/Engine/Classes/Materials/MaterialExpressionTime.h
+++ b/Engine/Source/Runtime/Engine/Classes/Materials/MaterialExpressionTime.h
@@ -16,7 +16,7 @@ class UMaterialExpressionTime : public UMaterialExpression
 	uint32 bIgnorePause:1;
 
 	/** Enables or disables the Period value. */
-	UPROPERTY()
+	UPROPERTY(EditAnywhere, Category = MaterialExpressionTime, meta=(InlineEditConditionToggle))
 	uint32 bOverride_Period:1;
 
 	/** Time will loop around once it gets to Period. */
diff --git a/Engine/Source/Runtime/Engine/Classes/Particles/VectorField/ParticleModuleVectorFieldGlobal.h b/Engine/Source/Runtime/Engine/Classes/Particles/VectorField/ParticleModuleVectorFieldGlobal.h
--- a/Engine/Source/Runtime/Engine/Classes/Particles/VectorField/ParticleModuleVectorFieldGlobal.h
+++ b/Engine/Source/Runtime/Engine/Classes/Particles/VectorField/ParticleModuleVectorFieldGlobal.h
@@ -14,7 +14,7 @@ class UParticleModuleVectorFieldGlobal : public UParticleModuleVectorFieldBase
 	GENERATED_UCLASS_BODY()
 
 	/** Property override value for global vector field tightness.  */
-	UPROPERTY()
+	UPROPERTY(EditAnywhere, Category = VectorField, meta=(InlineEditConditionToggle))
 	uint32 bOverrideGlobalVectorFieldTightness:1;
 
 	/** Global vector field scale. */
diff --git a/Engine/Source/Runtime/Engine/Classes/PhysicsEngine/BodyInstance.h b/Engine/Source/Runtime/Engine/Classes/PhysicsEngine/BodyInstance.h
--- a/Engine/Source/Runtime/Engine/Classes/PhysicsEngine/BodyInstance.h
+++ b/Engine/Source/Runtime/Engine/Classes/PhysicsEngine/BodyInstance.h
@@ -197,7 +197,7 @@ public:
 	uint32 bSimulatePhysics : 1;
 
 	/** If true, mass will not be automatically computed and you must set it directly */
-	UPROPERTY(meta = (DisplayName = "Override"))
+	UPROPERTY(EditAnywhere, Category = Physics, meta = (InlineEditConditionToggle))
 	uint32 bOverrideMass : 1;
 
 	/** If object should have the force of gravity applied */
@@ -256,7 +256,7 @@ public:
 	uint32 bLockZRotation : 1;
 
 	/** Override the default max angular velocity */
-	UPROPERTY(meta = (editcondition = "bSimulatePhysics"))
+	UPROPERTY(EditAnywhere, Category = Physics, meta = (editcondition = "bSimulatePhysics", InlineEditConditionToggle))
 	uint32 bOverrideMaxAngularVelocity : 1;
 
 	/** When initializing dynamic instances their component or velocity can override the bStartAwake flag */
@@ -271,11 +271,11 @@ protected:
 	uint32 bUseAsyncScene:1;
 
 	/** Whether this body instance has its own custom MaxDepenetrationVelocity*/
-	UPROPERTY()
+	UPROPERTY(EditAnywhere, Category = Physics, meta=(InlineEditConditionToggle))
 	uint32 bOverrideMaxDepenetrationVelocity : 1;
 
 	/** Whether this instance of the object has its own custom walkable slope override setting. */
-	UPROPERTY()
+	UPROPERTY(EditAnywhere, Category = Physics, meta = (InlineEditConditionToggle))
 	uint32 bOverrideWalkableSlopeOnInstance : 1;
 
 	uint32 bHasSharedShapes : 1;
diff --git a/Engine/Source/Runtime/Engine/Classes/Sound/SoundAttenuation.h b/Engine/Source/Runtime/Engine/Classes/Sound/SoundAttenuation.h
--- a/Engine/Source/Runtime/Engine/Classes/Sound/SoundAttenuation.h
+++ b/Engine/Source/Runtime/Engine/Classes/Sound/SoundAttenuation.h
@@ -181,15 +181,15 @@ struct ENGINE_API FAttenuationSettings
 	float NonFocusVolumeAttenuation;
 
 	/** The low pass filter frequency (in hertz) to apply if the sound playing in this audio component is occluded. This will override the frequency set in LowPassFilterFrequency. A frequency of 0.0 is the device sample rate and will bypass the filter. */
-	UPROPERTY(EditAnywhere, BlueprintReadWrite, Category = Occlusion, meta = (ClampMin = "0.0", UIMin = "0.0", EditCondition = "bEnableOcclusionChecks"))
+	UPROPERTY(EditAnywhere, BlueprintReadWrite, Category = Occlusion, meta = (ClampMin = "0.0", UIMin = "0.0", EditCondition = "bEnableOcclusion"))
 	float OcclusionLowPassFilterFrequency;
 
 	/** The amount of volume attenuation to apply to sounds which are occluded. */
-	UPROPERTY(EditAnywhere, BlueprintReadWrite, Category = Occlusion, meta = (ClampMin = "0.0", UIMin = "0.0", EditCondition = "bEnableOcclusionChecks"))
+	UPROPERTY(EditAnywhere, BlueprintReadWrite, Category = Occlusion, meta = (ClampMin = "0.0", UIMin = "0.0", EditCondition = "bEnableOcclusion"))
 	float OcclusionVolumeAttenuation;
 
 	/** The amount of time in seconds to interpolate to the target OcclusionLowPassFilterFrequency when a sound is occluded. */
-	UPROPERTY(EditAnywhere, BlueprintReadWrite, Category = Occlusion, meta = (ClampMin = "0", UIMin = "0.0", EditCondition = "bEnableOcclusionChecks"))
+	UPROPERTY(EditAnywhere, BlueprintReadWrite, Category = Occlusion, meta = (ClampMin = "0", UIMin = "0.0", EditCondition = "bEnableOcclusion"))
 	float OcclusionInterpolationTime;
 
 	FAttenuationSettings()
diff --git a/Engine/Source/Runtime/Engine/Private/Camera/CameraComponent.cpp b/Engine/Source/Runtime/Engine/Private/Camera/CameraComponent.cpp
--- a/Engine/Source/Runtime/Engine/Private/Camera/CameraComponent.cpp
+++ b/Engine/Source/Runtime/Engine/Private/Camera/CameraComponent.cpp
@@ -159,6 +159,10 @@ void UCameraComponent::RefreshVisualRepresentation()
 		DrawFrustum->FrustumAspectRatio = AspectRatio;
 		DrawFrustum->MarkRenderStateDirty();
 	}
+	if (ProxyMeshComponent != nullptr)
+	{
+		ProxyMeshComponent->ResetRelativeTransform();
+	}
 }
 
 void UCameraComponent::OverrideFrustumColor(FColor OverrideColor)
@@ -249,6 +253,13 @@ void UCameraComponent::GetCameraView(float DeltaTime, FMinimalViewInfo& DesiredV
 	{
 		DesiredView.PostProcessSettings = PostProcessSettings;
 	}
+
+#if WITH_EDITOR
+	if (ProxyMeshComponent != nullptr)
+	{
+		ProxyMeshComponent->ResetRelativeTransform();
+	}
+#endif //WITH_EDITOR
 }
 
 #if WITH_EDITOR
diff --git a/Engine/Source/Runtime/Engine/Private/Components/SplineComponent.cpp b/Engine/Source/Runtime/Engine/Private/Components/SplineComponent.cpp
--- a/Engine/Source/Runtime/Engine/Private/Components/SplineComponent.cpp
+++ b/Engine/Source/Runtime/Engine/Private/Components/SplineComponent.cpp
@@ -16,6 +16,8 @@ USplineComponent::USplineComponent(const FObjectInitializer& ObjectInitializer)
 	, Duration(1.0f)
 	, bStationaryEndpoints(false)
 	, bSplineHasBeenEdited(false)
+	, bInputSplinePointsToConstructionScript(false)
+	, bAlwaysRenderInEditor(true)
 	, bClosedLoop(false)
 	, DefaultUpVector(FVector::UpVector)
 #if WITH_EDITORONLY_DATA
@@ -561,6 +563,23 @@ void USplineComponent::SetTangentAtSplinePoint(int32 PointIndex, const FVector&
 }
 
 
+void USplineComponent::SetUpVectorAtSplinePoint(int32 PointIndex, const FVector& InUpVector, ESplineCoordinateSpace::Type CoordinateSpace)
+{
+	const int32 NumPoints = SplineInfo.Points.Num();
+
+	if ((PointIndex >= 0) && (PointIndex < NumPoints))
+	{
+		const FVector TransformedUpVector = (CoordinateSpace == ESplineCoordinateSpace::World) ?
+			ComponentToWorld.InverseTransformVector(InUpVector.GetSafeNormal()) : InUpVector.GetSafeNormal();
+
+		FQuat Quat = FQuat::FindBetween(DefaultUpVector, TransformedUpVector);
+		SplineRotInfo.Points[PointIndex].OutVal = Quat;
+
+		UpdateSpline();
+	}
+}
+
+
 ESplinePointType::Type USplineComponent::GetSplinePointType(int32 PointIndex) const
 {
 	if ((PointIndex >= 0) && (PointIndex < SplineInfo.Points.Num()))
@@ -598,7 +617,7 @@ FVector USplineComponent::GetLocationAtSplinePoint(int32 PointIndex, ESplineCoor
 
 FVector USplineComponent::GetDirectionAtSplinePoint(int32 PointIndex, ESplineCoordinateSpace::Type CoordinateSpace) const
 {
-	return GetLocationAtSplineInputKey(static_cast<float>(PointIndex), CoordinateSpace);
+	return GetDirectionAtSplineInputKey(static_cast<float>(PointIndex), CoordinateSpace);
 }
 
 
@@ -1084,6 +1103,171 @@ FTransform USplineComponent::FindTransformClosestToWorldLocation(const FVector&
 	return GetTransformAtSplineInputKey(Param, CoordinateSpace, bUseScale);
 }
 
+#if WITH_EDITOR
+FPrimitiveSceneProxy* USplineComponent::CreateSceneProxy()
+{
+	if (!bAlwaysRenderInEditor)
+	{
+		return Super::CreateSceneProxy();
+	}
+
+	class FSplineSceneProxy : public FPrimitiveSceneProxy
+	{
+	public:
+
+		FSplineSceneProxy(const USplineComponent* InComponent)
+			: FPrimitiveSceneProxy(InComponent)
+			, bAlwaysRenderInEditor(InComponent->bAlwaysRenderInEditor)
+			, SplineInfo(InComponent->SplineInfo)
+			, LineColor(InComponent->EditorUnselectedSplineSegmentColor)
+		{}
+
+		virtual void GetDynamicMeshElements(const TArray<const FSceneView*>& Views, const FSceneViewFamily& ViewFamily, uint32 VisibilityMap, FMeshElementCollector& Collector) const override
+		{
+			QUICK_SCOPE_CYCLE_COUNTER(STAT_SplineSceneProxy_GetDynamicMeshElements);
+
+			if (IsSelected())
+			{
+				return;
+			}
+
+			for (int32 ViewIndex = 0; ViewIndex < Views.Num(); ViewIndex++)
+			{
+				if (VisibilityMap & (1 << ViewIndex))
+				{
+					const FSceneView* View = Views[ViewIndex];
+					FPrimitiveDrawInterface* PDI = Collector.GetPDI(ViewIndex);
+
+					const FMatrix& LocalToWorld = GetLocalToWorld();
+
+					// Taking into account the min and maximum drawing distance
+					const float DistanceSqr = (View->ViewMatrices.ViewOrigin - LocalToWorld.GetOrigin()).SizeSquared();
+					if (DistanceSqr < FMath::Square(GetMinDrawDistance()) || DistanceSqr > FMath::Square(GetMaxDrawDistance()))
+					{
+						continue;
+					}
+
+					const int32 GrabHandleSize = 6;
+					FVector OldKeyPos(0);
+
+					const int32 NumPoints = SplineInfo.Points.Num();
+					const int32 NumSegments = SplineInfo.bIsLooped ? NumPoints : NumPoints - 1;
+					for (int32 KeyIdx = 0; KeyIdx < NumSegments + 1; KeyIdx++)
+					{
+						const FVector NewKeyPos = LocalToWorld.TransformPosition(SplineInfo.Eval(static_cast<float>(KeyIdx), FVector(0)));
+
+						// Draw the keypoint
+						if (KeyIdx < NumPoints)
+						{
+							PDI->DrawPoint(NewKeyPos, LineColor, GrabHandleSize, SDPG_World);
+						}
+
+						// If not the first keypoint, draw a line to the previous keypoint.
+						if (KeyIdx > 0)
+						{
+							// For constant interpolation - don't draw ticks - just draw dotted line.
+							if (SplineInfo.Points[KeyIdx - 1].InterpMode == CIM_Constant)
+							{
+								// Calculate dash length according to size on screen
+								const float StartW = View->WorldToScreen(OldKeyPos).W;
+								const float EndW = View->WorldToScreen(NewKeyPos).W;
+
+								const float WLimit = 10.0f;
+								if (StartW > WLimit || EndW > WLimit)
+								{
+									const float Scale = 0.03f;
+									DrawDashedLine(PDI, OldKeyPos, NewKeyPos, LineColor, FMath::Max(StartW, EndW) * Scale, SDPG_World);
+								}
+							}
+							else
+							{
+								// Find position on first keyframe.
+								FVector OldPos = OldKeyPos;
+
+								// Then draw a line for each substep.
+								const int32 NumSteps = 20;
+
+								for (int32 StepIdx = 1; StepIdx <= NumSteps; StepIdx++)
+								{
+									const float Key = (KeyIdx - 1) + (StepIdx / static_cast<float>(NumSteps));
+									const FVector NewPos = LocalToWorld.TransformPosition(SplineInfo.Eval(Key, FVector(0)));
+									PDI->DrawLine(OldPos, NewPos, LineColor, SDPG_World);
+									OldPos = NewPos;
+								}
+							}
+						}
+
+						OldKeyPos = NewKeyPos;
+					}
+				}
+			}
+		}
+
+		virtual FPrimitiveViewRelevance GetViewRelevance(const FSceneView* View) const override
+		{
+			FPrimitiveViewRelevance Result;
+			Result.bDrawRelevance = bAlwaysRenderInEditor && !IsSelected() && IsShown(View) && View->Family->EngineShowFlags.Splines;
+			Result.bDynamicRelevance = true;
+			Result.bShadowRelevance = IsShadowCast(View);
+			Result.bEditorPrimitiveRelevance = UseEditorCompositing(View);
+			return Result;
+		}
+
+		virtual uint32 GetMemoryFootprint(void) const override { return sizeof *this + GetAllocatedSize(); }
+		uint32 GetAllocatedSize(void) const { return FPrimitiveSceneProxy::GetAllocatedSize(); }
+
+	private:
+		bool bAlwaysRenderInEditor;
+		FInterpCurveVector SplineInfo;
+		FLinearColor LineColor;
+	};
+
+	return new FSplineSceneProxy(this);
+}
+
+
+FBoxSphereBounds USplineComponent::CalcBounds(const FTransform& LocalToWorld) const
+{
+	if (!bAlwaysRenderInEditor)
+	{
+		// Do as little as possible if not rendering anything
+		return Super::CalcBounds(LocalToWorld);
+	}
+
+#if SPLINE_FAST_BOUNDS_CALCULATION
+	FBox BoundingBox(0);
+	for (const auto& InterpPoint : SplineInfo.Points)
+	{
+		BoundingBox += InterpPoint.OutVal;
+	}
+
+	return FBoxSphereBounds(BoundingBox.TransformBy(LocalToWorld));
+#else
+	const int32 NumPoints = SplineInfo.Points.Num();
+	const int32 NumSegments = bClosedLoop ? NumPoints : NumPoints - 1;
+
+	FVector Min(WORLD_MAX);
+	FVector Max(-WORLD_MAX);
+	for (int32 Index = 0; Index < NumSegments; Index++)
+	{
+		const bool bLoopSegment = (Index == NumPoints - 1);
+		const int32 NextIndex = bLoopSegment ? 0 : Index + 1;
+		const FInterpCurvePoint<FVector>& ThisInterpPoint = SplineInfo.Points[Index];
+		FInterpCurvePoint<FVector> NextInterpPoint = SplineInfo.Points[NextIndex];
+		if (bLoopSegment)
+		{
+			NextInterpPoint.InVal = ThisInterpPoint.InVal + SplineInfo.LoopKeyOffset;
+		}
+
+		CurveVectorFindIntervalBounds(ThisInterpPoint, NextInterpPoint, Min, Max);
+	}
+
+	return FBoxSphereBounds(FBox(Min, Max).TransformBy(LocalToWorld));
+#endif
+}
+
+#endif
+
 
 /** Used to store spline data during RerunConstructionScripts */
 class FSplineInstanceData : public FSceneComponentInstanceData
@@ -1126,6 +1310,11 @@ void USplineComponent::ApplyComponentInstanceData(FSplineInstanceData* SplineIns
 {
 	check(SplineInstanceData);
 
+	if (bPostUCS && bInputSplinePointsToConstructionScript)
+	{
+		return;
+	}
+
 	if (SplineInstanceData->bSplineHasBeenEdited)
 	{
 		SplineInfo = SplineInstanceData->SplineInfo;
diff --git a/Engine/Source/Runtime/Engine/Private/GameEngine.cpp b/Engine/Source/Runtime/Engine/Private/GameEngine.cpp
--- a/Engine/Source/Runtime/Engine/Private/GameEngine.cpp
+++ b/Engine/Source/Runtime/Engine/Private/GameEngine.cpp
@@ -144,7 +144,6 @@ void UGameEngine::CreateGameViewport( UGameViewportClient* GameViewportClient )
 
 	auto Window = GameViewportWindow.Pin();
 
-	Window->SetWidgetToFocusOnActivate( GameViewportWidgetRef );
 	Window->SetOnWindowClosed( FOnWindowClosed::CreateUObject( this, &UGameEngine::OnGameWindowClosed ) );
 
 	// SAVEWINPOS tells us to load/save window positions to user settings (this is disabled by default)
diff --git a/Engine/Source/Runtime/Engine/Private/GameViewportClient.cpp b/Engine/Source/Runtime/Engine/Private/GameViewportClient.cpp
--- a/Engine/Source/Runtime/Engine/Private/GameViewportClient.cpp
+++ b/Engine/Source/Runtime/Engine/Private/GameViewportClient.cpp
@@ -591,43 +591,13 @@ bool UGameViewportClient::RequiresUncapturedAxisInput() const
 
 EMouseCursor::Type UGameViewportClient::GetCursor(FViewport* InViewport, int32 X, int32 Y)
 {
-	//bool bIsPlayingMovie = false;//GetMoviePlayer()->IsMovieCurrentlyPlaying();
-
-#if !PLATFORM_WINDOWS
-	bool bIsWithinTitleBar = false;
-#else
-	POINT CursorPos = { X, Y };
-	RECT WindowRect;
-
-	bool bIsWithinWindow = true;
-
-	// For Slate based windows the viewport doesnt have access to the OS window handle and shouln't need it
-	bool bIsWithinTitleBar = false;
-	if( InViewport->GetWindow() )
-	{
-		ClientToScreen( (HWND)InViewport->GetWindow(), &CursorPos );
-		GetWindowRect( (HWND)InViewport->GetWindow(), &WindowRect );
-		bIsWithinWindow = ( CursorPos.x >= WindowRect.left && CursorPos.x <= WindowRect.right &&
-			CursorPos.y >= WindowRect.top && CursorPos.y <= WindowRect.bottom );
-
-		// The user is mousing over the title bar if Y is less than zero and within the window rect
-		bIsWithinTitleBar = Y < 0 && bIsWithinWindow;
-	}
-
-#endif
-
-	if ((!InViewport->HasMouseCapture() && !InViewport->HasFocus()) || (ViewportConsole && ViewportConsole->ConsoleActive()))
+	if (!InViewport->HasMouseCapture() || (ViewportConsole && ViewportConsole->ConsoleActive()))
 	{
 		return EMouseCursor::Default;
 	}
-	else if ( /*(!bIsPlayingMovie) && */(InViewport->IsFullscreen() || !bIsWithinTitleBar) ) //bIsPlayingMovie has always false value
+	else if (GameInstance && GameInstance->GetFirstLocalPlayerController())
 	{
-		if (GameInstance && GameInstance->GetFirstLocalPlayerController())
-		{
-			return GameInstance->GetFirstLocalPlayerController()->GetMouseCursor();
-		}
-
-		return EMouseCursor::None;
+		return GameInstance->GetFirstLocalPlayerController()->GetMouseCursor();
 	}
 
 	return FViewportClient::GetCursor(InViewport, X, Y);
@@ -1273,11 +1243,6 @@ void UGameViewportClient::ProcessScreenShots(FViewport* InViewport)
 			bScreenshotSuccessful = FSlateApplication::Get().TakeScreenshot( WindowRef, Bitmap, Size);
 			GScreenshotResolutionX = Size.X;
 			GScreenshotResolutionY = Size.Y;
-
-			if (Size.X == 0 || Size.Y == 0)
-			{
-				bScreenshotSuccessful = false;
-			}
 		}
 		else
 		{
@@ -1411,6 +1376,16 @@ bool UGameViewportClient::IsFocused(FViewport* InViewport)
 	return InViewport->HasFocus() || InViewport->HasMouseCapture();
 }
 
+void UGameViewportClient::Activated(FViewport* InViewport, const FWindowActivateEvent& InActivateEvent)
+{
+	ReceivedFocus(InViewport);
+}
+
+void UGameViewportClient::Deactivated(FViewport* InViewport, const FWindowActivateEvent& InActivateEvent)
+{
+	LostFocus(InViewport);
+}
+
 void UGameViewportClient::CloseRequested(FViewport* InViewport)
 {
 	check(InViewport == Viewport);
diff --git a/Engine/Source/Runtime/Engine/Private/Model.cpp b/Engine/Source/Runtime/Engine/Private/Model.cpp
--- a/Engine/Source/Runtime/Engine/Private/Model.cpp
+++ b/Engine/Source/Runtime/Engine/Private/Model.cpp
@@ -24,7 +24,7 @@ DEFINE_LOG_CATEGORY_STATIC(LogModel, Log, All);
  */
 bool FBspSurf::IsHiddenEd() const
 {
-	return bHiddenEdTemporary || bHiddenEdLevel;
+	return bHiddenEdTemporary || bHiddenEdLevel || bHiddenEdLayer;
 }
 
 /**
@@ -60,6 +60,7 @@ FArchive& operator<<( FArchive& Ar, FBspSurf& Surf )
 	{
 		Ar << Surf.bHiddenEdTemporary;
 		Ar << Surf.bHiddenEdLevel;
+		Ar << Surf.bHiddenEdLayer;
 	}
 
 	return Ar;
@@ -352,6 +353,11 @@ void UModel::PostLoad()
 			FBspSurf& CurSurf = *SurfIter;
 			CurSurf.bHiddenEdTemporary = ( ( CurSurf.PolyFlags & PF_HiddenEd ) != 0 );
 			CurSurf.bHiddenEdLevel = 0;
+#if WITH_EDITOR
+			CurSurf.bHiddenEdLayer = CurSurf.Actor ? CurSurf.Actor->bHiddenEdLayer : 0;
+#else
+			CurSurf.bHiddenEdLayer = 0;
+#endif
 		}
 
 #if WITH_EDITOR
diff --git a/Engine/Source/Runtime/Engine/Private/MorphTools.cpp b/Engine/Source/Runtime/Engine/Private/MorphTools.cpp
--- a/Engine/Source/Runtime/Engine/Private/MorphTools.cpp
+++ b/Engine/Source/Runtime/Engine/Private/MorphTools.cpp
@@ -237,8 +237,7 @@ void UMorphTarget::RemapVertexIndices( USkeletalMesh* InBaseMesh, const TArray<
 * @param	SrcMesh - source skeletal mesh to convert
 * @param	LODIndex - level of detail to use for the geometry
 */
-FMorphMeshRawSource::FMorphMeshRawSource( USkeletalMesh* SrcMesh, int32 LODIndex ) :
-	SourceMesh(SrcMesh)
+FMorphMeshRawSource::FMorphMeshRawSource( USkeletalMesh* SrcMesh, int32 LODIndex ) 
 {
 	check(SrcMesh);
 	FSkeletalMeshResource* ImportedResource = SrcMesh->GetImportedResource();
@@ -247,6 +246,7 @@ FMorphMeshRawSource::FMorphMeshRawSource( USkeletalMesh* SrcMesh, int32 LODIndex
 	// get the mesh data for the given lod
 	FStaticLODModel& LODModel = ImportedResource->LODModels[LODIndex];
 
+
 	// vertices are packed in this order iot stay consistent
 	// with the indexing used by the FStaticLODModel vertex buffer
 	//
@@ -313,6 +313,79 @@ FMorphMeshRawSource::FMorphMeshRawSource( USkeletalMesh* SrcMesh, int32 LODIndex
 	}
 }
 
+FMorphMeshRawSource::FMorphMeshRawSource(FStaticLODModel& LODModel)
+{
+	Initialize(LODModel);
+}
+
+void FMorphMeshRawSource::Initialize(FStaticLODModel& LODModel)
+{
+	// vertices are packed in this order iot stay consistent
+	// with the indexing used by the FStaticLODModel vertex buffer
+	//
+	//	Chunk0
+	//		Rigid0
+	//		Rigid1
+	//		Soft0
+	//		Soft1
+	//	Chunk1
+	//		Rigid0
+	//		Rigid1
+	//		Soft0
+	//		Soft1
+
+	// iterate over the chunks for the skeletal mesh
+	for (int32 ChunkIdx = 0; ChunkIdx < LODModel.Chunks.Num(); ChunkIdx++)
+	{
+		// each chunk has both rigid and smooth vertices
+		const FSkelMeshChunk& Chunk = LODModel.Chunks[ChunkIdx];
+		// rigid vertices should always be added first so that we are
+		// consistent with the way vertices are packed in the FStaticLODModel vertex buffer
+		for (int32 VertexIdx = 0; VertexIdx < Chunk.RigidVertices.Num(); VertexIdx++)
+		{
+			const FRigidSkinVertex& SourceVertex = Chunk.RigidVertices[VertexIdx];
+			FMorphMeshVertexRaw RawVertex =
+			{
+				SourceVertex.Position,
+				SourceVertex.TangentX,
+				SourceVertex.TangentY,
+				SourceVertex.TangentZ
+			};
+			Vertices.Add(RawVertex);
+		}
+		// smooth vertices are added next. The resulting Vertices[] array should
+		// match the FStaticLODModel vertex buffer when indexing vertices
+		for (int32 VertexIdx = 0; VertexIdx < Chunk.SoftVertices.Num(); VertexIdx++)
+		{
+			const FSoftSkinVertex& SourceVertex = Chunk.SoftVertices[VertexIdx];
+			FMorphMeshVertexRaw RawVertex =
+			{
+				SourceVertex.Position,
+				SourceVertex.TangentX,
+				SourceVertex.TangentY,
+				SourceVertex.TangentZ
+			};
+			Vertices.Add(RawVertex);
+		}
+	}
+
+	// Copy the indices manually, since the LODModel's index buffer may have a different alignment.
+	Indices.Empty(LODModel.MultiSizeIndexContainer.GetIndexBuffer()->Num());
+	for (int32 Index = 0; Index < LODModel.MultiSizeIndexContainer.GetIndexBuffer()->Num(); Index++)
+	{
+		Indices.Add(LODModel.MultiSizeIndexContainer.GetIndexBuffer()->Get(Index));
+	}
+
+	// copy the wedge point indices
+	if (LODModel.RawPointIndices.GetBulkDataSize())
+	{
+		WedgePointIndices.Empty(LODModel.RawPointIndices.GetElementCount());
+		WedgePointIndices.AddUninitialized(LODModel.RawPointIndices.GetElementCount());
+		FMemory::Memcpy(WedgePointIndices.GetData(), LODModel.RawPointIndices.Lock(LOCK_READ_ONLY), LODModel.RawPointIndices.GetBulkDataSize());
+		LODModel.RawPointIndices.Unlock();
+	}
+}
+
 /**
 * Constructor. 
 * Converts a static mesh to raw vertex data
@@ -321,8 +394,7 @@ FMorphMeshRawSource::FMorphMeshRawSource( USkeletalMesh* SrcMesh, int32 LODIndex
 * @param	SrcMesh - source static mesh to convert
 * @param	LODIndex - level of detail to use for the geometry
 */
-FMorphMeshRawSource::FMorphMeshRawSource( UStaticMesh* SrcMesh, int32 LODIndex ) :
-	SourceMesh(SrcMesh)
+FMorphMeshRawSource::FMorphMeshRawSource( UStaticMesh* SrcMesh, int32 LODIndex ) 
 {
 	// @todo - not implemented
 	// not sure if we will support static mesh morphing yet
diff --git a/Engine/Source/Runtime/Engine/Private/ObjectLibrary.cpp b/Engine/Source/Runtime/Engine/Private/ObjectLibrary.cpp
--- a/Engine/Source/Runtime/Engine/Private/ObjectLibrary.cpp
+++ b/Engine/Source/Runtime/Engine/Private/ObjectLibrary.cpp
@@ -40,6 +40,11 @@ void UObjectLibrary::PostEditChangeProperty(FPropertyChangedEvent& PropertyChang
 				if(Objects[i] != NULL)
 				{
 					UClass* BlueprintClass = Cast<UClass>(Objects[i]);
+					if (!BlueprintClass)
+					{
+						UBlueprintCore* Blueprint = Cast<UBlueprintCore>(Objects[i]);
+						BlueprintClass = Blueprint ? Blueprint->GeneratedClass : nullptr;
+					}
 
 					if (!BlueprintClass)
 					{
