Disabling some focus adjusting code in the scene viewport to help with vr editor event.

diff --git a/Engine/Source/Runtime/Engine/Private/Slate/SceneViewport.cpp b/Engine/Source/Runtime/Engine/Private/Slate/SceneViewport.cpp
--- a/Engine/Source/Runtime/Engine/Private/Slate/SceneViewport.cpp
+++ b/Engine/Source/Runtime/Engine/Private/Slate/SceneViewport.cpp
@@ -401,20 +401,20 @@ FReply FSceneViewport::OnMouseButtonDown( const FGeometry& InGeometry, const FPo
 				CurrentReplyState = FReply::Unhandled();
 			}
 		}
-		else
-		{
-			TSharedRef<SViewport> ViewportWidgetRef = ViewportWidget.Pin().ToSharedRef();
-			if ( ViewportWidgetRef->HasUserFocusedDescendants(InMouseEvent.GetUserIndex()) )
-			{
-				// If we're still focused on a descendant, force it to stay focused.  Need to do this to fool SlateApplication
-				// otherwise it will reassign focus thinking that nobody requested it.
-				TSharedPtr<SWidget> FocusedWidgetPtr = FSlateApplication::Get().GetUserFocusedWidget(InMouseEvent.GetUserIndex());
-				if ( FocusedWidgetPtr.IsValid() )
-				{
-					CurrentReplyState.SetUserFocus(FocusedWidgetPtr.ToSharedRef(), EFocusCause::SetDirectly, false);
-				}
-			}
-		}
+		//else
+		//{
+		//	TSharedRef<SViewport> ViewportWidgetRef = ViewportWidget.Pin().ToSharedRef();
+		//	if ( ViewportWidgetRef->HasUserFocusedDescendants(InMouseEvent.GetUserIndex()) )
+		//	{
+		//		// If we're still focused on a descendant, force it to stay focused.  Need to do this to fool SlateApplication
+		//		// otherwise it will reassign focus thinking that nobody requested it.
+		//		TSharedPtr<SWidget> FocusedWidgetPtr = FSlateApplication::Get().GetUserFocusedWidget(InMouseEvent.GetUserIndex());
+		//		if ( FocusedWidgetPtr.IsValid() )
+		//		{
+		//			CurrentReplyState.SetUserFocus(FocusedWidgetPtr.ToSharedRef(), EFocusCause::SetDirectly, false);
+		//		}
+		//	}
+		//}
 
 		// a new menu was opened if there was previously not a menu visible but now there is
 		const bool bNewMenuWasOpened = !bAnyMenuWasVisible && FSlateApplication::Get().AnyMenusVisible();
