VR Editor: Never update viewport hit proxies while in VR
- This solves graphical artifacts and poor performance when dragging spinners in Details, as well as other issues

diff --git a/Engine/Source/Runtime/Engine/Private/UnrealClient.cpp b/Engine/Source/Runtime/Engine/Private/UnrealClient.cpp
--- a/Engine/Source/Runtime/Engine/Private/UnrealClient.cpp
+++ b/Engine/Source/Runtime/Engine/Private/UnrealClient.cpp
@@ -1001,6 +1001,7 @@ void FViewport::Draw( bool bShouldPresent /*= true */)
 	static bool bReentrant = false;
 	if(!bReentrant)
 	{
+		//GWarn->Logf( TEXT( "%i: ViewportDRAW!" ), GFrameCounter );
 		// See what screenshot related features are required
 		static const auto CVarDumpFrames = IConsoleManager::Get().FindTConsoleVariableDataInt(TEXT("r.BufferVisualizationDumpFrames"));
 		GIsHighResScreenshot = GIsHighResScreenshot || bTakeHighResScreenShot;
@@ -1160,10 +1161,18 @@ const TArray<FColor>& FViewport::GetRawHitProxyData(FIntRect InRect)
 {
 	FScopedConditionalWorldSwitcher WorldSwitcher(ViewportClient);
 
-	bool bFetchHitProxyBytes = !bHitProxiesCached || (SizeY*SizeX) != CachedHitProxyData.Num();
+	const bool bIsRenderingStereo = GEngine->IsStereoscopic3D( this ) && this->IsStereoRenderingAllowed();
 
+	bool bFetchHitProxyBytes = !bIsRenderingStereo && ( !bHitProxiesCached || (SizeY*SizeX) != CachedHitProxyData.Num() );
+
+	if( bIsRenderingStereo )
+	{
+		// Stereo viewports don't support hit proxies, and we don't want to update them because it will adversely
+		// affect performance.
+		CachedHitProxyData.SetNumZeroed( SizeY * SizeX );
+	}
 	// If the hit proxy map isn't up to date, render the viewport client's hit proxies to it.
-	if (!bHitProxiesCached)
+	else if (!bHitProxiesCached)
 	{
 		EnqueueBeginRenderFrame();
 
