Fixed viewport focus stealing issues caused by Oculus 1.2 preview code

diff --git a/Engine/Source/Runtime/Engine/Private/Slate/SceneViewport.cpp b/Engine/Source/Runtime/Engine/Private/Slate/SceneViewport.cpp
--- a/Engine/Source/Runtime/Engine/Private/Slate/SceneViewport.cpp
+++ b/Engine/Source/Runtime/Engine/Private/Slate/SceneViewport.cpp
@@ -901,34 +901,6 @@ FReply FSceneViewport::OnFocusReceived(const FFocusEvent& InFocusEvent)
 	{
 		FScopedConditionalWorldSwitcher WorldSwitcher(ViewportClient);
 		ViewportClient->ReceivedFocus(this);
-
-		if ((FApp::IsGame() && !GIsEditor) || bIsPlayInEditorViewport)
-		{
-			if (IsForegroundWindow() || FApp::UseVRFocus())
-			{
-				bool bIsCursorForcedVisible = false;
-				if (ViewportClient->GetWorld() && ViewportClient->GetWorld()->GetGameInstance() && ViewportClient->GetWorld()->GetGameInstance()->GetFirstLocalPlayerController())
-				{
-					bIsCursorForcedVisible = ViewportClient->GetWorld()->GetGameInstance()->GetFirstLocalPlayerController()->GetMouseCursor() != EMouseCursor::None;
-				}
-
-				const bool bPlayInEditorCapture = !bIsPlayInEditorViewport || InFocusEvent.GetCause() != EFocusCause::SetDirectly || bPlayInEditorGetsMouseControl;
-
-				// capturing the mouse interferes with slate UI (like the virtual joysticks)
-				if (FPlatformProperties::SupportsWindowedMode() && bPlayInEditorCapture && !bIsCursorForcedVisible && !FSlateApplication::Get().IsFakingTouchEvents())
-				{
-					// Only require the user to click in the window the first time - after that return focus to the game so long as it was the last focused widget.
-					// Means that tabbing in/out will return the mouse control to where it was & the in-game console won't leave the mouse under editor control.
-					bPlayInEditorGetsMouseControl = true;
-					CurrentReplyState.UseHighPrecisionMouseMovement(ViewportWidget.Pin().ToSharedRef());
-					CurrentReplyState.LockMouseToWidget(ViewportWidget.Pin().ToSharedRef());
-				}
-			}
-			else
-			{
-				CurrentReplyState.ClearUserFocus(true);
-			}
-		}
 	}
 
 	// Update key state mappings so that the the viewport modifier states are valid upon focus.
