Copying //UE4/Dev-Rendering to Dev-Main (//UE4/Dev-Main)

==========================
MAJOR FEATURES + CHANGES
==========================

Change 2879377 on 2016/02/24 by [REDACTED_PERSON]

	UE4 - Added render thread start and stop delegates. GitHub 2006.

Change 2879378 on 2016/02/24 by [REDACTED_PERSON]

	UE4 - Avoided using TG_PrePhysics as the first tickgroup so that licensees can add tickgropups.
	https://udn.unrealengine.com/questions/279126/code-assumes-that-tg-prephysics-is-the-first-tick.html

Change 2879382 on 2016/02/24 by [REDACTED_PERSON]

	UE4 - Tweaked automation test framework by request from UDN post.

Change 2879727 on 2016/02/24 by [REDACTED_PERSON]

	adding debug info for Optimus driver detection issue
	#Test:PC

Change 2879728 on 2016/02/24 by [REDACTED_PERSON]

	fixed and improved VisualizeMotionBlur
	#test:PC

Change 2879729 on 2016/02/24 by [REDACTED_PERSON]

	added AngleBetweenVectors() and variants to the FastMath library

Change 2880133 on 2016/02/24 by [REDACTED_PERSON]

	new r.DepthOfFieldQualitySetting
	for GDC squencer demo
	 - OR-15875

Change 2880314 on 2016/02/24 by [REDACTED_PERSON]

	Fixed uses of FDepthDrawingPolicyFactory being affected by bUseAsOccluder
	* This fixes preshadows on HISMC and foliage

Change 2880338 on 2016/02/24 by [REDACTED_PERSON]

	added SkinCache.Debug cvar
	#test:PC

Change 2880344 on 2016/02/24 by [REDACTED_PERSON]

	Added the ability to apply DFAO to static indirect lighting, controlled by r.AOApplyToStaticIndirect
	* Lightmaps, stationary skylight and reflection captures are all affected
	* Specular occlusion on reflection captures requires a fair amount of tweaking of r.SkySpecularOcclusionStrength, MinOcclusion and MaxOcclusionDistance for good quality
	* For now, a movable skylight with low intensity (.0001) must be placed to control MaxOcclusionDistance and MinOcclusion

Change 2880346 on 2016/02/24 by [REDACTED_PERSON]

	Added several cvars to expose mesh distance field limits, which allows higher quality
	* r.DistanceFields.MaxPerMeshResolution
	* r.DistanceFields.DefaultVoxelDensity
	* r.DistanceFields.AtlasSizeXY
	* r.DistanceFields.AtlasSizeZ

Change 2881304 on 2016/02/25 by [REDACTED_PERSON]

	UE4 - Increased the priority of cloth tasks because these are on the critical path.

Change 2881306 on 2016/02/25 by [REDACTED_PERSON]

	UE4 - Added cvar to control background tick list cleanup.

Change 2881790 on 2016/02/25 by [REDACTED_PERSON]

	Screen size fading is only applied to spot and point lights

Change 2882077 on 2016/02/25 by [REDACTED_PERSON]

	DFAO indirect occlusion on static lighting is now correctly applied to IndirectIrradiance

Change 2882391 on 2016/02/25 by [REDACTED_PERSON]

	fixed bad caching of SRV for vertexbuffers in SkinCache (caused rendering artifacts and wasteful memory allocations). Finding a SRV is now O(1), was O(n)

Change 2883008 on 2016/02/26 by [REDACTED_PERSON]

	UE4 - Fixed recursive shader intialization crash on consoles.

Change 2883253 on 2016/02/26 by [REDACTED_PERSON]

	Improved SkinTangent compression

Change 2883295 on 2016/02/26 by [REDACTED_PERSON]

	Added RecomputeSkinTangent feature for GPU SkinCache, not enabled by default (r.SkinCache.RecomputeTangents)

Change 2883363 on 2016/02/26 by [REDACTED_PERSON]

	UE4 - Fixed an issue with recurisve shader init on consoles...again.

Change 2883912 on 2016/02/26 by [REDACTED_PERSON]

	UE4 - Fixed shadows updating static meshes while the prepass is in progress.

Change 2884829 on 2016/02/27 by [REDACTED_PERSON]

	OR-16237 indirect lighting on skin is too dark

Change 2885096 on 2016/02/28 by [REDACTED_PERSON]

	OR-13678

diff --git a/Engine/Shaders/BasePassPixelShader.usf b/Engine/Shaders/BasePassPixelShader.usf
--- a/Engine/Shaders/BasePassPixelShader.usf
+++ b/Engine/Shaders/BasePassPixelShader.usf
@@ -748,9 +748,23 @@ void Main(
 	}
 #endif
 
+	float3 LocalBaseColor = BaseColor;
+	float LocalSpecular = Specular;
+
+#if MATERIAL_SHADINGMODEL_SUBSURFACE_PROFILE
+	// SubsurfaceProfile applies the BaseColor in a later pass. Any lighting output in the base pass needs
+	// to separate specualar and diffuse lighting in a checkerboard pattern
+	{
+		bool bChecker = CheckerFromPixelPos(MaterialParameters.SvPosition.xy);
+
+		LocalBaseColor = bChecker;
+		LocalSpecular = !bChecker;
+	}
+#endif
+
 	// So that the following code can still use DiffuseColor and SpecularColor.
-	half3 DiffuseColor = BaseColor - BaseColor * Metallic;
-	half3 SpecularColor = lerp( 0.08 * Specular.xxx, BaseColor, Metallic.xxx );
+	half3 DiffuseColor = LocalBaseColor - LocalBaseColor * Metallic;
+	half3 SpecularColor = lerp( 0.08 * LocalSpecular.xxx, LocalBaseColor, Metallic.xxx );
 
 #if USE_DEVELOPMENT_SHADERS
 	{
diff --git a/Engine/Shaders/DeferredDecal.usf b/Engine/Shaders/DeferredDecal.usf
--- a/Engine/Shaders/DeferredDecal.usf
+++ b/Engine/Shaders/DeferredDecal.usf
@@ -298,8 +298,20 @@ void MainPS(
 			// not normalized
 			float3 OSNormal = GradientNotNormalized(SvPosition, OSPosition, HitT);
 			
-			// not normalized
-			float3 WSNormal = mul(float4(OSNormal, 0), transpose(WorldToDecal)).xyz;
+			// To transform the normals use tranpose(Inverse(DecalToWorld)) = transpose(WorldToDecal)
+			// But we want to only rotate the normals (we don't want to non-uniformaly scale them).
+			// We assume the matrix is only a scale and rotation, and we remove non-uniform scale:
+			float3 lengthSqr = { length2(DecalToWorld._m00_m01_m02),
+				length2(DecalToWorld._m10_m11_m12),
+				length2(DecalToWorld._m20_m21_m22) };
+
+			float3 scale = rsqrt(lengthSqr);
+
+			// Pre-multiply by the inverse of the non-uniform scale in DecalToWorld
+		    float4 ScaledNormal = float4(OSNormal.xyz * scale.xyz, 0.f); 
+		    float3 WSNormal = mul(ScaledNormal, DecalToWorld).xyz;
+
+			// need to normalized
 
 			MaterialParameters.WorldNormal = normalize(WSNormal);
 		}
@@ -310,9 +322,23 @@ void MainPS(
 
 			// tangent/decal space to world space
 #if MATERIAL_TANGENTSPACENORMAL
-			// normals are transformed by the tranposed of the inverse
-			// DecalToWorld becomes transpose(WorldToDecal)
-			MaterialParameters.WorldNormal = normalize(mul(float4(-MaterialNormal.z, MaterialNormal.y, MaterialNormal.x, 0), transpose(WorldToDecal)).xyz);
+
+	
+			// To transform the normals use tranpose(Inverse(DecalToWorld)) = transpose(WorldToDecal)
+			// But we want to only rotate the normals (we don't want to non-uniformaly scale them).
+			// We assume the matrix is only a scale and rotation, and we remove non-uniform scale:
+			float3 lengthSqr = { length2(DecalToWorld._m00_m01_m02),
+					             length2(DecalToWorld._m10_m11_m12),
+					             length2(DecalToWorld._m20_m21_m22) };
+
+			float3 scale = rsqrt(lengthSqr);
+				
+			// Pre-multiply by the inverse of the non-uniform scale in DecalToWorld
+			float4 ScaledNormal = { -MaterialNormal.z * scale.x, MaterialNormal.y * scale.y, MaterialNormal.x * scale.z, 0.f };
+				
+			// Compute the normal 
+			MaterialParameters.WorldNormal = normalize(mul(ScaledNormal, DecalToWorld).xyz);
+
 #else
 			MaterialParameters.WorldNormal = normalize(MaterialNormal);
 #endif
diff --git a/Engine/Shaders/DistanceFieldLightingPost.usf b/Engine/Shaders/DistanceFieldLightingPost.usf
--- a/Engine/Shaders/DistanceFieldLightingPost.usf
+++ b/Engine/Shaders/DistanceFieldLightingPost.usf
@@ -240,6 +240,29 @@ float UseHistoryFilter;
 Texture2D VelocityTexture;
 SamplerState VelocityTextureSampler;
 
+float ComputeHistoryWeightBasedOnPosition(float2 ScreenPosition, float SceneDepth, float2 OldScreenPosition, float HistoryDepth)
+{
+	float3 WorldPosition;
+	float3 PrevWorldPosition;
+
+	{
+		float4 HomogeneousWorldPosition = mul(float4(ScreenPosition * SceneDepth, SceneDepth, 1), View.ScreenToWorld);
+		WorldPosition = HomogeneousWorldPosition.xyz / HomogeneousWorldPosition.w;
+
+		float4 PrevPositionTranslatedWorld = mul(float4(OldScreenPosition * HistoryDepth, HistoryDepth, 1), View.PrevScreenToTranslatedWorld);
+		PrevPositionTranslatedWorld.xyz /= PrevPositionTranslatedWorld.w;	
+		PrevWorldPosition = PrevPositionTranslatedWorld.xyz - View.PrevPreViewTranslation;
+	}
+
+	float DistanceToHistoryValue = length(PrevWorldPosition - WorldPosition);
+	float RelativeHistoryDistanceThreshold = HistoryDistanceThreshold / 1000.0f;
+
+	return DistanceToHistoryValue / SceneDepth > RelativeHistoryDistanceThreshold ? 0.0f : 1.0f;
+}
+
+// Doesn't work yet
+#define HIGH_QUALITY_HISTORY_REJECTION 0
+
 /** Reproject the occlusion history. */
 void UpdateHistoryPS(
 	in float4 UVAndScreenPos : TEXCOORD0
@@ -253,6 +276,7 @@ void UpdateHistoryPS(
 	float2 DistanceFieldUVs = UVAndScreenPos.xy - Frame.ViewRectMin.xy * Frame.BufferSizeAndInvSize.zw;
 	float4 NewValue = Texture2DSampleLevel(BentNormalAOTexture, BentNormalAOSampler, DistanceFieldUVs, 0);
 	float SceneDepth = abs(NewValue.w);
+	float OutputDepthSign = 1.0f;
 
 	float2 OldUV;
 
@@ -282,15 +306,34 @@ void UpdateHistoryPS(
 	}
 	
 	float2 OldScreenPosition = (OldUV.xy - View.ScreenPositionScaleBias.wz) / View.ScreenPositionScaleBias.xy;
-
 	float2 OldDistanceFieldUVs = OldUV.xy - Frame.ViewRectMin.xy * Frame.BufferSizeAndInvSize.zw;
 
-#define MANUAL_HISTORY_RESAMPLE 0
-#if MANUAL_HISTORY_RESAMPLE
-	float4 HistoryValue;
+	float2 MinUV = Frame.ViewRectMin.xy * Frame.BufferSizeAndInvSize.zw;
+	// Pull in the max UV to exclude the region which will read outside the viewport due to bilinear filtering
+	float2 MaxUV = MinUV + (Frame.ViewSizeAndInvSize.xy - 1 * DOWNSAMPLE_FACTOR) * Frame.BufferSizeAndInvSize.zw;
+
+	float EffectiveHistoryWeight = HistoryWeight;
+
+	FLATTEN
+	if (any(OldUV > MaxUV) || any(OldUV < MinUV))
 	{
+		EffectiveHistoryWeight = 0;
+	}
+
+#define MANUAL_HISTORY_RESAMPLE HIGH_QUALITY_HISTORY_REJECTION
+#if MANUAL_HISTORY_RESAMPLE
+
+		float3 Unused;
+		float Unused2;
+		//@todo - pack into BentNormalAOTexture
+		bool bHasDistanceFieldRepresentation;
+		bool bUnused;
+		GetDownsampledGBuffer(DistanceFieldUVs, Unused, Unused2, bHasDistanceFieldRepresentation, bUnused);
+
+		OutputDepthSign = bHasDistanceFieldRepresentation ? 1.0f : -1.0f;
+
 		float2 HistoryTexelSize = Frame.BufferSizeAndInvSize.zw * DOWNSAMPLE_FACTOR;
-		float2 OldDistanceFieldUVsCorner00 = HistoryTexelSize * (floor(OldDistanceFieldUVs / HistoryTexelSize - .5f) + .5f);
+		float2 OldDistanceFieldUVsCorner00 = HistoryTexelSize * (floor(OldDistanceFieldUVs / HistoryTexelSize) + .5f);
 		float2 OldDistanceFieldUVsCorner01 = OldDistanceFieldUVsCorner00 + float2(0, HistoryTexelSize.y);
 		float2 OldDistanceFieldUVsCorner11 = OldDistanceFieldUVsCorner00 + HistoryTexelSize;
 		float2 OldDistanceFieldUVsCorner10 = OldDistanceFieldUVsCorner00 + float2(HistoryTexelSize.x, 0);
@@ -301,47 +344,51 @@ void UpdateHistoryPS(
 		float4 Corner10Value = Texture2DSampleLevel(BentNormalHistoryTexture, BentNormalHistorySampler, OldDistanceFieldUVsCorner10, 0);
 
 		float2 LerpWeights = OldDistanceFieldUVs - OldDistanceFieldUVsCorner00;
-		float4 VerticalLerp0 = lerp(Corner00Value, Corner01Value, LerpWeights.y);
-		float4 VerticalLerp1 = lerp(Corner10Value, Corner11Value, LerpWeights.y);
-		HistoryValue = lerp(VerticalLerp0, VerticalLerp1, LerpWeights.x);
-	}
-#else
-	float4 HistoryValue = Texture2DSampleLevel(BentNormalHistoryTexture, BentNormalHistorySampler, OldDistanceFieldUVs, 0);
-#endif
+
+	#define COMPUTE_SEPARATE_POSITION_WEIGHTS 1
+	#if COMPUTE_SEPARATE_POSITION_WEIGHTS
+		float PositionWeight00 = ComputeHistoryWeightBasedOnPosition(UVAndScreenPos.zw, SceneDepth, OldScreenPosition, abs(Corner00Value.w));
+		float PositionWeight01 = ComputeHistoryWeightBasedOnPosition(UVAndScreenPos.zw, SceneDepth, OldScreenPosition, abs(Corner01Value.w));
+		float PositionWeight11 = ComputeHistoryWeightBasedOnPosition(UVAndScreenPos.zw, SceneDepth, OldScreenPosition, abs(Corner11Value.w));
+		float PositionWeight10 = ComputeHistoryWeightBasedOnPosition(UVAndScreenPos.zw, SceneDepth, OldScreenPosition, abs(Corner10Value.w));
 	
-	float2 MinUV = Frame.ViewRectMin.xy * Frame.BufferSizeAndInvSize.zw;
-	// Pull in the max UV to exclude the region which will read outside the viewport due to bilinear filtering
-	float2 MaxUV = MinUV + (Frame.ViewSizeAndInvSize.xy - 1 * DOWNSAMPLE_FACTOR) * Frame.BufferSizeAndInvSize.zw;
+		float ClassificationWeight00 = (Corner00Value.w > 0 == bHasDistanceFieldRepresentation) ? PositionWeight00 : 0.0f;
+		float ClassificationWeight01 = (Corner01Value.w > 0 == bHasDistanceFieldRepresentation) ? PositionWeight01 : 0.0f;
+		float ClassificationWeight11 = (Corner11Value.w > 0 == bHasDistanceFieldRepresentation) ? PositionWeight11 : 0.0f;
+		float ClassificationWeight10 = (Corner10Value.w > 0 == bHasDistanceFieldRepresentation) ? PositionWeight10 : 0.0f;
 
-	float3 WorldPosition;
-	float3 PrevWorldPosition;
+		float FinalWeight00 = (1 - LerpWeights.y) * (1 - LerpWeights.x) * ClassificationWeight00;
+		float FinalWeight01 = (LerpWeights.y) * (1 - LerpWeights.x) * ClassificationWeight01;
+		float FinalWeight11 = (LerpWeights.y) * (LerpWeights.x) * ClassificationWeight11;
+		float FinalWeight10 = (1 - LerpWeights.y) * (LerpWeights.x) * ClassificationWeight10;
+		float TotalWeight = FinalWeight00 + FinalWeight01 + FinalWeight11 + FinalWeight10;
 
-	{
-		float4 HomogeneousWorldPosition = mul(float4(UVAndScreenPos.zw * SceneDepth, SceneDepth, 1), View.ScreenToWorld);
-		WorldPosition = HomogeneousWorldPosition.xyz / HomogeneousWorldPosition.w;
+		float3 HistoryValue = (Corner00Value.xyz * FinalWeight00
+			+ Corner01Value.xyz * FinalWeight01
+			+ Corner11Value.xyz * FinalWeight11
+			+ Corner10Value.xyz * FinalWeight10) / max(TotalWeight, .00001f);
 
-		float HistoryDepth = abs(HistoryValue.w);
-		float4 PrevPositionTranslatedWorld = mul(float4(OldScreenPosition * HistoryDepth, HistoryDepth, 1), View.PrevScreenToTranslatedWorld);
-		PrevPositionTranslatedWorld.xyz /= PrevPositionTranslatedWorld.w;	
-		PrevWorldPosition = PrevPositionTranslatedWorld.xyz - View.PrevPreViewTranslation;
-	}
+		EffectiveHistoryWeight *= TotalWeight > 0 ? 1.0f : 0.0f;
 
-	float EffectiveHistoryWeight = HistoryWeight;
-	float RelativeHistoryDistanceThreshold = HistoryDistanceThreshold / 1000.0f;
-	float DistanceToHistoryValue = length(PrevWorldPosition - WorldPosition);
+	#else
+		// Filter encoded depth path, for testing
+		float4 VerticalLerp0 = lerp(Corner00Value, Corner01Value, LerpWeights.y);
+		float4 VerticalLerp1 = lerp(Corner10Value, Corner11Value, LerpWeights.y);
+		float4 HistoryValue = lerp(VerticalLerp0, VerticalLerp1, LerpWeights.x);
+		float PositionWeight = ComputeHistoryWeightBasedOnPosition(UVAndScreenPos.zw, SceneDepth, OldScreenPosition, abs(HistoryValue.w));
+		EffectiveHistoryWeight *= PositionWeight;
+	#endif
 
-	FLATTEN
-	if (any(OldUV > MaxUV) 
-		|| any(OldUV < MinUV)
-		// Discard history if we are shading a new position (newly revealed)
-		//@todo - this test needs to be done on the surrounding 4 texels of the history with manual filtering to avoid foreground leaking into the background
-		//|| DistanceToHistoryValue > HistoryDistanceThreshold)
-		|| DistanceToHistoryValue / SceneDepth > RelativeHistoryDistanceThreshold)
-	{
-		EffectiveHistoryWeight = 0;
-	}
+	OutBentNormal.rgb = lerp(NewValue.rgb, HistoryValue.rgb, EffectiveHistoryWeight);
 
+#else
+	// Fast path that uses hardware interpolation
+	// This fails to reject based on depth consistently because the depth has been filtered
+	float4 HistoryValue = Texture2DSampleLevel(BentNormalHistoryTexture, BentNormalHistorySampler, OldDistanceFieldUVs, 0);
+	float PositionWeight = ComputeHistoryWeightBasedOnPosition(UVAndScreenPos.zw, SceneDepth, OldScreenPosition, abs(HistoryValue.w));
+	EffectiveHistoryWeight *= PositionWeight;
 	OutBentNormal.rgb = lerp(NewValue.rgb, HistoryValue.rgb, EffectiveHistoryWeight);
+#endif
 
 #if SUPPORT_IRRADIANCE
 	float3 NewIrradiance = Texture2DSampleLevel(IrradianceTexture, IrradianceSampler, DistanceFieldUVs, 0).xyz;
@@ -359,6 +406,10 @@ void UpdateHistoryPS(
 		// Sign bit of alpha stores whether the history was rejected or not, to be read by the history filter pass
 		OutBentNormal.a *= EffectiveHistoryWeight > 0 ? 1 : -1;
 	}
+	else
+	{
+		OutBentNormal.a *= OutputDepthSign;
+	}
 }
 
 #define HALF_HISTORY_FILL_KERNEL_SIZE 2
@@ -433,17 +484,41 @@ void FilterHistoryPS(
 	// Remove sign bit so future reprojection interpolation isn't affected
 	OutBentNormal.w = abs(OutBentNormal.w);
 
+#if HIGH_QUALITY_HISTORY_REJECTION
+	float3 Unused;
+	float Unused2;
+	//@todo - pack into BentNormalAOTexture
+	bool bHasDistanceFieldRepresentation;
+	bool bUnused;
+	GetDownsampledGBuffer(UVAndScreenPos.xy, Unused, Unused2, bHasDistanceFieldRepresentation, bUnused);
+	// Store whether the pixel has a distance field representation for next frame
+	// This can only be supported with manual history filtering
+	OutBentNormal.w *= bHasDistanceFieldRepresentation ? 1.0f : -1.0f;
+
+#endif
+
 #if SUPPORT_IRRADIANCE
 	OutIrradiance = float4(IrradianceValue, 0);
 #endif
 }
 
+#if MODULATE_SCENE_COLOR
+#define IRRADIANCE_TARGET SV_Target2
+#else
+#define IRRADIANCE_TARGET SV_Target1
+#endif
+
+float MinIndirectDiffuseOcclusion;
+
 /** Upsamples the AO results to full resolution using a bilateral filter. */
 void AOUpsamplePS(
 	in float4 UVAndScreenPos : TEXCOORD0
 	, out float4 OutBentNormal : SV_Target0
+#if MODULATE_SCENE_COLOR
+	, out float4 OutSceneColorModulate : SV_Target1
+#endif
 #if OUTPUT_BENT_NORMAL && SUPPORT_IRRADIANCE
-	, out float4 OutIrradiance : SV_Target1
+	, out float4 OutIrradiance : IRRADIANCE_TARGET
 #endif
 	)
 {
@@ -509,6 +584,11 @@ void AOUpsamplePS(
 	#endif
 #endif
 
+#if MODULATE_SCENE_COLOR
+	float Visibility = lerp(length(BentNormal), 1, MinIndirectDiffuseOcclusion);
+	OutSceneColorModulate = Visibility;
+#endif
+
 #if OUTPUT_AO
 	OutBentNormal = float4(length(BentNormal).xxx, 1);
 	//OutBentNormal = float4(BentNormal, 1);
diff --git a/Engine/Shaders/DistanceFieldScreenGridLighting.usf b/Engine/Shaders/DistanceFieldScreenGridLighting.usf
--- a/Engine/Shaders/DistanceFieldScreenGridLighting.usf
+++ b/Engine/Shaders/DistanceFieldScreenGridLighting.usf
@@ -127,7 +127,7 @@ void HemisphereConeTraceAgainstTileCulledObjects(uint2 OutputCoordinate, uint Sh
 					float WorldStepOffset = GetStepOffset(StartStepIndex);
 
 					#if USE_GLOBAL_DISTANCE_FIELD
-						WorldStepOffset += 10;
+						WorldStepOffset += 2;
 					#endif
 
 					LOOP
diff --git a/Engine/Shaders/FastMath.usf b/Engine/Shaders/FastMath.usf
--- a/Engine/Shaders/FastMath.usf
+++ b/Engine/Shaders/FastMath.usf
@@ -206,4 +206,29 @@ float acosFast4(float inX)
 float asinFast4( float x )
 {
 	return (0.5 * PI) - acosFast4(x);
-}
\ No newline at end of file
+}
+
+// @param A doesn't have to be normalized, output could be NaN if this is near 0,0,0
+// @param B doesn't have to be normalized, output could be NaN if this is near 0,0,0
+// @return can be passed to a acosFast() or acos() to compute an angle
+float CosBetweenVectors(float3 A, float3 B)
+{
+	// unoptimized: dot(normalize(A), normalize(B))
+	return dot(A, B) * rsqrt(length2(A) * length2(B));
+}
+
+// @param A doesn't have to be normalized, output could be NaN if this is near 0,0,0
+// @param B doesn't have to be normalized, output could be NaN if this is near 0,0,0
+float AngleBetweenVectors(float3 A, float3 B)
+{
+	return acos(CosBetweenVectors(A, B));
+}
+// @param A doesn't have to be normalized, output could be NaN if this is near 0,0,0
+// @param B doesn't have to be normalized, output could be NaN if this is near 0,0,0
+float AngleBetweenVectorsFast(float3 A, float3 B)
+{
+	return acosFast(CosBetweenVectors(A, B));
+}
+
+
+
diff --git a/Engine/Shaders/GpuSkinCacheComputeShader.usf b/Engine/Shaders/GpuSkinCacheComputeShader.usf
--- a/Engine/Shaders/GpuSkinCacheComputeShader.usf
+++ b/Engine/Shaders/GpuSkinCacheComputeShader.usf
@@ -108,6 +108,8 @@ uint SkinCacheVertexCount;
 
 float3 MeshOrigin;
 float3 MeshExtension;
+// see r.SkinCache.Debug, default is 1.0f
+float SkinCacheDebug;
 
 //uint4 DebugParameter;
 
@@ -214,8 +216,8 @@ void SkinCacheUpdateBatchCS(uint3 GroupID : SV_GroupID,
 		LocalTangentX = normalize( LocalTangentX - (dot(LocalTangentX, LocalTangentZ) * LocalTangentZ) );
 	}
 #endif
-
-	float3 SPos = SkinPosition(Intermediates);
+	
+	float3 SPos = SkinPosition(Intermediates) * SkinCacheDebug;
 		
 	float3 TangentX = mul((float3x3)Intermediates.BlendMatrix, LocalTangentX);
 	float3 TangentZ = mul((float3x3)Intermediates.BlendMatrix, LocalTangentZ);
diff --git a/Engine/Shaders/GpuSkinCommon.usf b/Engine/Shaders/GpuSkinCommon.usf
--- a/Engine/Shaders/GpuSkinCommon.usf
+++ b/Engine/Shaders/GpuSkinCommon.usf
@@ -76,10 +76,10 @@ float4 UnpackS8x4N(uint Packed)
 uint PackS8x4N(float4 Unpacked)
 {
 	uint Packed = asuint(
-		(uint((Unpacked.x + 1.0f) * 127.5f) << 0) |
-		(uint((Unpacked.y + 1.0f) * 127.5f) << 8) |
-		(uint((Unpacked.z + 1.0f) * 127.5f) << 16) |
-		(uint((Unpacked.w + 1.0f) * 127.5f) << 24));
+		(uint((Unpacked.x + 1.0f) * 127.499f) << 0) |
+		(uint((Unpacked.y + 1.0f) * 127.499f) << 8) |
+		(uint((Unpacked.z + 1.0f) * 127.499f) << 16) |
+		(uint((Unpacked.w + 1.0f) * 127.499f) << 24));
 	return Packed;
 }
 
@@ -88,8 +88,8 @@ uint PackS8x4N(float4 Unpacked)
 uint PackS8x3N(float3 Unpacked)
 {
 	uint Packed = asuint(
-		(uint((Unpacked.x + 1.0f) * 127.5f) << 0) |
-		(uint((Unpacked.y + 1.0f) * 127.5f) << 8) |
-		(uint((Unpacked.z + 1.0f) * 127.5f) << 16));
+		(uint((Unpacked.x + 1.0f) * 127.499f) << 0) |
+		(uint((Unpacked.y + 1.0f) * 127.499f) << 8) |
+		(uint((Unpacked.z + 1.0f) * 127.499f) << 16));
 	return Packed;
 }
\ No newline at end of file
diff --git a/Engine/Shaders/PostProcessCircleDOF.usf b/Engine/Shaders/PostProcessCircleDOF.usf
--- a/Engine/Shaders/PostProcessCircleDOF.usf
+++ b/Engine/Shaders/PostProcessCircleDOF.usf
@@ -523,9 +523,11 @@ void Circle4Samples(float4 UVAndScreenPos : TEXCOORD0, out float4 OutColor0 : SV
 // pixel shader entry point
 void CirclePS(float4 UVAndScreenPos : TEXCOORD0, out float4 OutColor0 : SV_Target0)
 {
-#if QUALITY
-	// 2 or higher, gets slower but less noisy
+// Count 2 or higher, gets slower but less noisy
+#if ( QUALITY == 2 )
 	const uint Count = 32;
+#elif ( QUALITY == 1 )
+	const uint Count = 12;
 #else
 	const uint Count = 1;
 #endif
@@ -639,9 +641,11 @@ float4 Recombine2Samples(float4 UVAndScreenPos, float PixCoc, float InRand)
 
 float3 RecombineNSamples(float4 UVAndScreenPos, float PixCoc)
 {
-#if QUALITY
-	// 2 or higher, gets slower but less noisy
+// Count 2 or higher, gets slower but less noisy
+#if ( QUALITY == 2 )
 	const uint Count = 32;
+#elif ( QUALITY == 1 ) 
+	const uint Count = 12;
 #else
 	const uint Count = 1;
 #endif
diff --git a/Engine/Shaders/PostProcessCommon.usf b/Engine/Shaders/PostProcessCommon.usf
--- a/Engine/Shaders/PostProcessCommon.usf
+++ b/Engine/Shaders/PostProcessCommon.usf
@@ -71,6 +71,14 @@ float2 ComputePixelPosCenter(float2 ScreenPos, bool bPixelCenter)
 	return ScreenPos * ScreenPosToPixel.xy + ScreenPosToPixel.zw + PixelOffset;
 }
 
+// not well optimized but can be useful
+float2 ScreenPosFromPixelPos(float2 PixelPos, bool bPixelCenter)
+{
+	float2 PixelOffset = bPixelCenter ? 0.5f : 0.0f;
+
+	return (PixelPos - PixelOffset - ScreenPosToPixel.zw) / ScreenPosToPixel.xy;
+}
+
 // shape of the lens for vignette like effects and masking
 // @param ScreenPos -1 .. 1
 // @return 0..1 0:borders, 1:center
diff --git a/Engine/Shaders/PostProcessMotionBlur.usf b/Engine/Shaders/PostProcessMotionBlur.usf
--- a/Engine/Shaders/PostProcessMotionBlur.usf
+++ b/Engine/Shaders/PostProcessMotionBlur.usf
@@ -152,6 +152,118 @@ void OverrideWithTestChart(float2 ScreenPos, inout float2 ObjectVelocity, inout
 #endif
 }
 
+// used to visualize the motion blur
+// @return 0/1
+float Compute3DCheckerBoard(float3 Pos)
+{
+	float3 TiledWorldPos = frac(Pos) > 0.5f;
+	return (float)((uint)dot(float3(1,1,1), TiledWorldPos) % 2);
+}
+
+void VisualizeMotionBlurPS(
+	noperspective float4 UVAndScreenPos : TEXCOORD0,
+	out float4 OutColor : SV_Target0)
+{
+	const float TileSize = 48;
+
+	float2 ScreenSpacePos = UVAndScreenPos.zw;
+	float2 PixelPos = ComputePixelPosCenter(ScreenSpacePos, true);
+	float2 UV = PixelPos * PostprocessInput0Size.zw;
+
+	float2 PixelPosAtTileCenter = PixelPos - (frac(PixelPos / TileSize) - 0.5f) * TileSize;
+	float2 UVAtTileCenter = PixelPosAtTileCenter * PostprocessInput0Size.zw;
+	float2 ScreenPosAtTileCenter = ScreenPosFromPixelPos(PixelPosAtTileCenter, true);
+
+	// World aligned checkerboards
+	{
+		float DeviceZ = LookupDeviceZ(UV);
+		float SceneDepth = ConvertFromDeviceZ(DeviceZ);
+
+		float3 ScreenVector = mul(float4(ScreenSpacePos, 1, 0), View.ScreenToWorld).xyz;
+
+		// world space position of the current pixel
+		float3 OffsetWorldPos = ScreenVector * SceneDepth;
+
+		float3 AbsWorldPos = View.WorldCameraOrigin + OffsetWorldPos;
+		float3 WorldCheckerboard = Compute3DCheckerBoard(AbsWorldPos * 0.02f) * 0.1f + Compute3DCheckerBoard(AbsWorldPos * 0.002f) * 0.3f + Compute3DCheckerBoard(AbsWorldPos * 0.0002f) * 0.6f;
+		OutColor = float4(lerp(WorldCheckerboard, float3(0,0,0), 0.7f), 1);
+	}
+
+	bool bSelectorOpaqueAtTileCenter = PostprocessInput2.SampleLevel(PostprocessInput2Sampler, UVAtTileCenter, 0).x > 0;
+	bool bSelectorOpaque = PostprocessInput2.SampleLevel(PostprocessInput2Sampler, UV, 0).x > 0;
+
+	// relative, in screen space -1...1 -1..1, can be even outside of that range, points into the movement direction
+	float2 VelocityAtTileCenter;
+	{
+		VelocityAtTileCenter = DecodeVelocityFromTexture(PostprocessInput2.SampleLevel(PostprocessInput2Sampler, UVAtTileCenter, 0).xy);
+
+		// reconstruct from camera motion if we don't have velocity data
+		if(!bSelectorOpaqueAtTileCenter)
+		{
+			// UV is per tile
+			float4 ThisClip = float4(ScreenPosAtTileCenter, LookupDeviceZ(UVAtTileCenter), 1);
+			float4 PrevClip = mul( ThisClip, View.ClipToPrevClip );
+			float2 PrevScreen = PrevClip.xy / PrevClip.w;
+
+			// points into the movement direction
+			VelocityAtTileCenter = ScreenPosAtTileCenter - PrevScreen;
+		}
+	}
+
+	// tint yellow if velocity data is stored in texture
+	{
+		float3 TintColor = bSelectorOpaque ? float3(0.5f, 0.5f, 0.2f) : float3(0.5f, 0.5f, 0.5f);
+	
+		OutColor.rgb = lerp(OutColor.rgb, TintColor, 0.45f);
+	}
+
+	// tile center
+	{
+		float2 Delta = PixelPos - PixelPosAtTileCenter;
+
+		float HMask = lerp(saturate(abs(Delta.y)), 1, saturate(abs(Delta.x) / TileSize * 4));
+		float VMask = lerp(saturate(abs(Delta.x)), 1, saturate(abs(Delta.y) / TileSize * 4));
+
+		float Dist = length(Delta);
+		OutColor.rgb *= lerp(1, 0.3f, saturate(3 - Dist));
+	}
+
+	float3 LineColor = bSelectorOpaqueAtTileCenter ? float3(1,1,0) : float3(0.7,0.7,0.7);
+
+	// points into the movement direction
+	float2 PixelDirection = VelocityAtTileCenter * ScreenPosToPixel.xy;
+
+	// arrow
+	{
+		float2 PerpPixelDirection = float2(PixelDirection.y, -PixelDirection.x);
+		float2 DirectionInTile = PixelPos - PixelPosAtTileCenter;
+
+		float DistOnLine = dot(normalize(-PixelDirection), DirectionInTile) + length(PixelDirection);
+
+		bool bArrowHead = DistOnLine < 8;
+
+//		float LocalThickness = bArrowHead ? (DistOnLine * 0.5f) : 0.9f;
+//		float LocalThickness = bArrowHead ? (DistOnLine * 0.3f) : 1 + (frac(DistOnLine/8)*8)*0.2f;
+		float LocalThickness = 1 + (frac(DistOnLine/8)*8)*0.25f;
+
+		float PerpDirectionMask = saturate(LocalThickness - abs(dot(normalize(PerpPixelDirection), DirectionInTile)));
+		float DirectionMask = saturate(length(PixelDirection) - length(DirectionInTile));
+
+		float3 LineMask = PerpDirectionMask * DirectionMask;
+		OutColor.rgb = lerp(OutColor.rgb, LineColor, LineMask);
+	}
+
+	// previous pos is a dot
+	{
+		float3 DotColor = float3(0,1,0);
+		// PixelPos of the previous position
+		float2 PreviousPixelPos = PixelPosAtTileCenter - PixelDirection;
+		float Dist = length(PreviousPixelPos - PixelPos);
+		OutColor.rgb = lerp(OutColor.rgb, LineColor, saturate(4 - Dist));
+		OutColor.rgb = lerp(OutColor.rgb, 0, saturate(2.5f - Dist));
+	}
+}
+
 // ------------------------------------------
 
 // motion blur setup vertex shader
@@ -293,15 +405,6 @@ void SetupPS(
 //	OutVelocity.xy = lerp(float2(0,0), (InUV.y > 0.8f) ? float2(1,0) : float2(0.02f,0), OutVelocity.b);
 }
 
-
-// used to visualize the motion blur
-// @return 0/1
-float Compute3DCheckerBoard(float3 Pos)
-{
-	float3 TiledWorldPos = frac(Pos) > 0.5f;
-	return (float)((uint)dot(float3(1,1,1), TiledWorldPos) % 2);
-}
-
 uint GetStepCountFromQuality()
 {
 #if MOTION_BLUR_QUALITY == 1
diff --git a/Engine/Shaders/PostProcessTemporalCommon.usf b/Engine/Shaders/PostProcessTemporalCommon.usf
--- a/Engine/Shaders/PostProcessTemporalCommon.usf
+++ b/Engine/Shaders/PostProcessTemporalCommon.usf
@@ -7,6 +7,9 @@ This is the common bits of the temporal AA shader
 which can be configured via defines for various temporal AA passes.
 =============================================================================*/
 
+// 0 / 1(default) This can be set to 0 to debug the velocity vectors
+#define NEIGHBORHOOD_CLAMPING 1
+
 // 1 = Use tighter AABB clamp for history.
 // 0 = Use simple min/max clamp.
 #ifndef AA_AABB
@@ -433,6 +436,8 @@ which can be configured via defines for various temporal AA passes.
 		#endif
 	#endif
 	float LumaContrast = LumaMax - LumaMin;
+
+#if NEIGHBORHOOD_CLAMPING
 	#if AA_YCOCG
 		OutColor.rgb = clamp(OutColor.rgb, NeighborMin.rgb, NeighborMax.rgb);
 		#if (AA_ALPHA == 0)
@@ -455,6 +460,7 @@ which can be configured via defines for various temporal AA passes.
 	#if AA_DEBUG
 		OutColor.rg = float2(0.0, 0.0);
 	#endif
+#endif
 
 	// ADD BACK IN ALIASING TO SHARPEN
     // -------------------------------
diff --git a/Engine/Shaders/RecomputeTangentsCommon.usf b/Engine/Shaders/RecomputeTangentsCommon.usf
new file mode 100644
--- /dev/null
+++ b/Engine/Shaders/RecomputeTangentsCommon.usf
@@ -0,0 +1,19 @@
+// Copyright 1998-2016 Epic Games, Inc. All Rights Reserved.
+
+/*=============================================================================
+	RecomputeTangentsCommon.usf: Recompute Skin Tangents
+=============================================================================*/
+
+#include "Common.usf"
+#include "GpuSkinCommon.usf"
+
+// intermediate tangent accumulation buffer, int se we can use atomics and to get deterministic results no matter what order
+// uses as output on MainCS(), used and input and output (to clear) for ResolveCS()
+RWBuffer<int> IntermediateAccumBufferUAV;
+
+// How many units the compute shader should process
+// MainCS(): triangles count, ResolveCS(): vertex count
+uint UnitCount;
+
+// could be larger for better precision but we need to be conservative as we sum up many triangles weighed by the angle (potentially unbound)
+#define TANGENT_RANGE 0x7fff
diff --git a/Engine/Shaders/RecomputeTangentsPerTrianglePass.usf b/Engine/Shaders/RecomputeTangentsPerTrianglePass.usf
new file mode 100644
--- /dev/null
+++ b/Engine/Shaders/RecomputeTangentsPerTrianglePass.usf
@@ -0,0 +1,210 @@
+// Copyright 1998-2016 Epic Games, Inc. All Rights Reserved.
+
+/*=============================================================================
+	RecomputeTangentsPerTrianglePass.usf: Recompute Skin Tangents
+=============================================================================*/
+
+#include "RecomputeTangentsCommon.usf"
+#include "FastMath.usf"
+
+// start offset within GPUSkinCacheBuffer[]
+uint GPUSkinCacheFloatOffset;
+// input vertex data (SkinCache output)
+// Buffer<float> GPUSkinCacheBuffer;
+// used as input in MainCS()
+Buffer<float> GPUSkinCacheBuffer;
+
+uint SkinCacheVertexStride;
+uint SkinCacheInputStreamFloatOffset;
+Buffer<float> SkinStreamInputBuffer;
+
+// start index in IndexBuffer[]
+uint IndexBufferOffset;
+// todo: test 16bit and 32 bit 
+/// only needed for MainCS()
+Buffer<uint> IndexBuffer;
+
+struct FVertexData
+{
+	// vertex index in the vertexbuffer
+	uint VertexIndex;
+	// position including skinning and morph target animation
+	float3 Position;
+	// 
+	float2 UV;
+	// 0..1, used for weighting the triangle contributions
+	float Angle;
+};
+
+// can be optimized
+// @return not neccessarily normalized
+float3 FindPerpendicularTo(float3 In)
+{
+	float3 Other = (abs(normalize(In).z) < 0.5f)
+		? float3(0,0,1)
+		: float3(1,0,0);
+
+	return cross(Other, In);
+}
+
+FVertexData GetVertex(uint VertexIndex)
+{
+	FVertexData Ret;
+
+	// is filled out later
+	Ret.Angle = 0;
+
+	{
+		// start offset within GPUSkinCacheBuffer[]
+		uint OutputOffset = VertexIndex * GPUSKIN_RWBUFFER_NUM_FLOATS + GPUSkinCacheFloatOffset;
+	
+		// no start offset as we reuse the same buffer always from 0 on
+		Ret.VertexIndex = VertexIndex;
+
+		Ret.Position.x = GPUSkinCacheBuffer[OutputOffset + GPUSKIN_RWBUFFER_OFFSET_POSITION + 0];
+		Ret.Position.y = GPUSkinCacheBuffer[OutputOffset + GPUSKIN_RWBUFFER_OFFSET_POSITION + 1];
+		Ret.Position.z = GPUSkinCacheBuffer[OutputOffset + GPUSKIN_RWBUFFER_OFFSET_POSITION + 2];
+	}
+
+	{
+		uint FloatsPerVertex = SkinCacheVertexStride / 4;
+		uint InputOffset = VertexIndex * FloatsPerVertex + SkinCacheInputStreamFloatOffset;
+
+		if(FULL_PRECISION_UV)
+		{
+			// float precision
+			Ret.UV.x = SkinStreamInputBuffer[InputOffset + GPUSKIN_VB_OFFSET_UVS + 0];
+			Ret.UV.y = SkinStreamInputBuffer[InputOffset + GPUSKIN_VB_OFFSET_UVS + 1];
+		}
+		else
+		{
+			// half precision
+			uint IntUV = asuint(SkinStreamInputBuffer[InputOffset + GPUSKIN_VB_OFFSET_UVS + 0]);
+			Ret.UV.x = f16tof32(IntUV & 0xffff);
+			Ret.UV.y = f16tof32(IntUV >> 16);
+		}
+	}
+
+	return Ret;
+}
+
+[numthreads(THREADGROUP_SIZEX, 1, 1)]
+void MainCS(
+	uint DispatchThreadId : SV_DispatchThreadID)	// DispatchThreadId = GroupId * int2(dimx,dimy) + GroupThreadId
+{
+	uint TriangleIndex = DispatchThreadId;
+
+	// We need this to restrict processing of triangles to the valid range (We run a Dispatch with a coarse granularity and would process too many elements without that)
+	if(TriangleIndex < UnitCount)
+	{
+		// 3 triangle corners
+		FVertexData Corner[3];
+		{
+			// 3 vertices per triangle
+			uint TriangleIndexBufferStart = TriangleIndex * 3 + IndexBufferOffset;
+
+			UNROLL for(uint CornerId = 0; CornerId < 3; ++CornerId)
+			{
+				uint VertexIndex = IndexBuffer[TriangleIndexBufferStart + CornerId]; //+ SkinCacheInputStreamFloatOffset;
+
+				Corner[CornerId] = GetVertex(VertexIndex);
+			}
+		}
+
+		// Compute the angle in each corner for weighting
+		{
+			UNROLL for(uint CornerId = 0; CornerId < 3; ++CornerId)
+			{
+				float3 Pos = Corner[CornerId].Position;
+
+				float3 EdgeA = Corner[(CornerId + 1) % 3].Position - Pos;
+				float3 EdgeB = Corner[(CornerId + 2) % 3].Position - Pos;
+			
+				Corner[CornerId].Angle = AngleBetweenVectorsFast(EdgeA, EdgeB);
+			}
+		}
+
+		float3 EdgeA = Corner[1].Position - Corner[0].Position;
+		float3 EdgeB = Corner[2].Position - Corner[0].Position;
+
+		float3 TriangleNormal = cross(EdgeB, EdgeA);
+
+		// in -1..1 range
+		float3 TangentZ = normalize(TriangleNormal);
+		
+		// to avoid having degenerated triangles causing shading artifacts on neighbors (should be compiled out if the boolean is not use)
+		bool bDegenerated = false;
+
+		bDegenerated = length2(TriangleNormal) < 0.000001f;
+
+		float3 TangentX;
+		float Orientation;
+		{
+			float2 UVEdgeA = Corner[1].UV - Corner[0].UV;
+			float2 UVEdgeB = Corner[2].UV - Corner[0].UV;
+
+			// as explained here: http://www.3dkingdoms.com/weekly/weekly.php?a=37
+			float cp = UVEdgeA.y * UVEdgeB.x - UVEdgeA.x * UVEdgeB.y;
+
+			float3 UnnormalizedTangentX = (EdgeA * -UVEdgeB.y + EdgeB * UVEdgeA.y) / cp;
+			// normally this should be normalize(x / cp) but we only need the sign of cp so we can multiply as well, which is faster and avoids a division by 0
+//			float3 UnnormalizedTangentX = (EdgeA * -UVEdgeB.y + EdgeB * UVEdgeA.y) * cp;
+
+			// bad UV assignment
+			if(abs(cp) < 0.000001f)
+			{
+				UnnormalizedTangentX = FindPerpendicularTo(TangentZ);
+			}
+			
+			FLATTEN if(length2(UnnormalizedTangentX) < 0.000001f)
+			{
+			// EdgeA is in the direction EdgeB
+				bDegenerated = true;
+			}
+
+			TangentX = normalize(UnnormalizedTangentX);
+
+			Orientation = (cp > 0) ? 1 : -1;
+		}
+
+/* disabled for now (for performance) but if needed we can enable it later
+		if(bDegenerated)
+		{
+			TangentX = 0;
+			TangentZ = 0;
+			Orientation = 0;
+		}
+*/
+		{
+			UNROLL for(uint CornerId = 0; CornerId < 3; ++CornerId)
+			{
+				// weighting by the angle results in more correct results with triangle meshes (independent of diagonal edges)
+				float Weight = Corner[CornerId].Angle;
+
+				// in -TANGENT_RANGE*PI..TANGENT_RANGE*PI  range
+				int3 IntTangentZ = (int3)(TangentZ * Weight * TANGENT_RANGE);
+				int3 IntTangentX = (int3)(TangentX * Weight * TANGENT_RANGE);
+				int IntOrientation = (int)(Orientation * Weight * TANGENT_RANGE);
+
+				uint VertexIndex = Corner[CornerId].VertexIndex;
+
+				// accumulate the contributions of each triangle on the triangle vertices
+				
+				// no start offset as we reuse the same buffer always from 0 on
+				uint Index = VertexIndex * INTERMEDIATE_ACCUM_BUFFER_NUM_INTS;
+
+				// TangentZ
+				InterlockedAdd(IntermediateAccumBufferUAV[Index + 0], IntTangentZ.x);
+				InterlockedAdd(IntermediateAccumBufferUAV[Index + 1], IntTangentZ.y);
+				InterlockedAdd(IntermediateAccumBufferUAV[Index + 2], IntTangentZ.z);
+				// TangentX
+				InterlockedAdd(IntermediateAccumBufferUAV[Index + 3], IntTangentX.x);
+				InterlockedAdd(IntermediateAccumBufferUAV[Index + 4], IntTangentX.y);
+				InterlockedAdd(IntermediateAccumBufferUAV[Index + 5], IntTangentX.z);
+				// Orientation
+				InterlockedAdd(IntermediateAccumBufferUAV[Index + 6], IntOrientation);
+			}
+		}
+	}
+}
+
diff --git a/Engine/Shaders/RecomputeTangentsPerVertexPass.usf b/Engine/Shaders/RecomputeTangentsPerVertexPass.usf
new file mode 100644
--- /dev/null
+++ b/Engine/Shaders/RecomputeTangentsPerVertexPass.usf
@@ -0,0 +1,52 @@
+// Copyright 1998-2016 Epic Games, Inc. All Rights Reserved.
+
+/*=============================================================================
+	RecomputeTangentsPerVertexPass.usf: Recompute Skin Tangents
+=============================================================================*/
+
+#include "RecomputeTangentsCommon.usf"
+
+// start offset within GPUSkinCacheBufferUAV[]
+uint GPUSkinCacheFloatOffset;
+// used as output in ResolveCS()
+RWBuffer<float> GPUSkinCacheBufferUAV;
+
+uint SkinCacheInputStreamFloatOffset;
+
+[numthreads(THREADGROUP_SIZEX, 1, 1)]
+void MainCS(
+	uint DispatchThreadId : SV_DispatchThreadID)	// DispatchThreadId = GroupId * int2(dimx,dimy) + GroupThreadId
+{
+	if(DispatchThreadId < UnitCount)
+	{
+		uint VertexIndex = DispatchThreadId + SkinCacheInputStreamFloatOffset;
+		// -1..1 range, normalized
+		float3 TangentZ, TangentX;
+		// the sign defines the handyness of the tangent matrix
+		float Orientation;
+		{
+			// no start offset as we reuse the same buffer always from 0 on
+			uint Index = VertexIndex * INTERMEDIATE_ACCUM_BUFFER_NUM_INTS;
+
+			// we don't have to scale down as we normalize anyway
+			TangentZ = normalize(float3(IntermediateAccumBufferUAV[Index + 0], IntermediateAccumBufferUAV[Index + 1], IntermediateAccumBufferUAV[Index + 2]));
+			TangentX = normalize(float3(IntermediateAccumBufferUAV[Index + 3], IntermediateAccumBufferUAV[Index + 4], IntermediateAccumBufferUAV[Index + 5]));
+			Orientation = IntermediateAccumBufferUAV[Index + 6];
+
+			// clear input UAV to 0 again to avoid another Dispatch() call
+			UNROLL for(int i = 0; i < 7; ++i)
+			{
+				IntermediateAccumBufferUAV[Index + i] = 0;
+			}
+		}
+
+		uint Offset = VertexIndex * GPUSKIN_RWBUFFER_NUM_FLOATS + GPUSkinCacheFloatOffset;
+
+		uint TangentXCompressed = PackS8x3N(TangentX);
+		uint TangentZCompressed = PackS8x3N(TangentZ);
+		uint HandynessValue = (Orientation < 0) ? 0xff000000 : 0;
+
+		GPUSkinCacheBufferUAV[Offset + GPUSKIN_RWBUFFER_OFFSET_TANGENT_X] = asfloat(TangentXCompressed);
+		GPUSkinCacheBufferUAV[Offset + GPUSKIN_RWBUFFER_OFFSET_TANGENT_Z] = asfloat(TangentZCompressed | HandynessValue);
+	}
+}
\ No newline at end of file
diff --git a/Engine/Shaders/ReflectionEnvironmentComputeShaders.usf b/Engine/Shaders/ReflectionEnvironmentComputeShaders.usf
--- a/Engine/Shaders/ReflectionEnvironmentComputeShaders.usf
+++ b/Engine/Shaders/ReflectionEnvironmentComputeShaders.usf
@@ -411,6 +411,18 @@ float CountOverlap( float3 WorldPosition )
 
 void GatherRadiance(inout float4 Color, float3 WorldPosition, float3 RayDirection, float Roughness, float2 ScreenPosition, float IndirectIrradiance, uint ShadingModelID)
 {
+	// Indirect occlusion from DFAO, which should be applied to reflection captures and skylight specular, but not SSR
+	float IndirectSpecularOcclusion = 1.0f;
+
+#if SUPPORT_DFAO_INDIRECT_OCCLUSION
+	float2 ScreenUV = ScreenPosition * View.ScreenPositionScaleBias.xy + View.ScreenPositionScaleBias.wz;
+	float IndirectDiffuseOcclusion;
+	GetDistanceFieldAOSpecularOcclusion(ScreenUV, RayDirection, Roughness, ShadingModelID == SHADINGMODELID_TWOSIDED_FOLIAGE, IndirectSpecularOcclusion, IndirectDiffuseOcclusion);
+	// Apply DFAO to IndirectIrradiance before mixing with indirect specular
+	IndirectIrradiance *= IndirectDiffuseOcclusion;
+	
+#endif
+
 	float Mip = ComputeReflectionCaptureMipFromRoughness( Roughness );
 
 	uint NumCapturesAffectingTile = TileNumReflectionCaptures;
@@ -542,7 +554,7 @@ void GatherRadiance(inout float4 Color, float3 WorldPosition, float3 RayDirectio
 				Sample *= DistanceAlpha;
 
 				// Under operator (back to front)
-				Color.rgb += Sample.rgb * Color.a;
+				Color.rgb += Sample.rgb * Color.a * IndirectSpecularOcclusion;
 				Color.a *= 1 - Sample.a;
 			}
 		}
@@ -562,11 +574,7 @@ void GatherRadiance(inout float4 Color, float3 WorldPosition, float3 RayDirectio
 			SkyLighting *= IndirectIrradiance;
 		}
 
-		float2 ScreenUV = ScreenPosition * View.ScreenPositionScaleBias.xy + View.ScreenPositionScaleBias.wz;
-		float Visibility = GetDistanceFieldAOSpecularOcclusion(ScreenUV, RayDirection, Roughness, ShadingModelID == SHADINGMODELID_TWOSIDED_FOLIAGE);
-
-		// TODO use Mip
-		Color.rgb += (Color.a * Visibility) * SkyLighting;
+		Color.rgb += Color.a * SkyLighting * IndirectSpecularOcclusion;
 	}
 #endif
 }
diff --git a/Engine/Shaders/ReflectionEnvironmentShaders.usf b/Engine/Shaders/ReflectionEnvironmentShaders.usf
--- a/Engine/Shaders/ReflectionEnvironmentShaders.usf
+++ b/Engine/Shaders/ReflectionEnvironmentShaders.usf
@@ -474,16 +474,28 @@ void ReflectionApplyPS(
 		float3 V = -normalize( mul( float4(UVAndScreenPos.zw, 1, 0), View.ScreenToWorld ).xyz );
 		float NoV = saturate( dot( GBuffer.WorldNormal, V ) );
 		float3 SpecularColor = EnvBRDF( GBuffer.SpecularColor, GBuffer.Roughness, NoV );
-		
+		float3 ReflectionVector = reflect(-V, GBuffer.WorldNormal);
+		float IndirectIrradiance = GBuffer.IndirectIrradiance;
+
 		float AO = GBuffer.GBufferAO * ScreenSpaceData.AmbientOcclusion;
 		float SpecularOcclusion = saturate( Square( NoV + AO ) - 1 + AO );
 		SpecularColor *= SpecularOcclusion;
 
+		// Indirect occlusion from DFAO, which should be applied to reflection captures and skylight specular, but not SSR
+		float IndirectSpecularOcclusion = 1.0f;
+#if SUPPORT_DFAO_INDIRECT_OCCLUSION
+		float IndirectDiffuseOcclusion = 1.0f;
+		GetDistanceFieldAOSpecularOcclusion(UV, ReflectionVector, GBuffer.Roughness, GBuffer.ShadingModelID == SHADINGMODELID_TWOSIDED_FOLIAGE, IndirectSpecularOcclusion, IndirectDiffuseOcclusion); 
+		// Apply DFAO to IndirectIrradiance before mixing with indirect specular
+		IndirectIrradiance *= IndirectDiffuseOcclusion;
+#endif
+
 		float4 SpecularLighting = float4(0, 0, 0, 1);
 		
 #if APPLY_REFLECTION_ENV
 		// Alpha of the reflection environment texture contains coverage, where 0 the sky reflection should be applied
 		SpecularLighting = Texture2DSample( ReflectionEnvTexture, ReflectionEnvSampler, UV );
+		SpecularLighting.rgb *= IndirectSpecularOcclusion;
 
 #if ALLOW_STATIC_LIGHTING
 		FLATTEN
@@ -496,7 +508,7 @@ void ReflectionApplyPS(
 			// Note: make sure this matches the lightmap mixing done for translucency (BasePassPixelShader.usf)
 
 			// NOTE: this breaks LPV specular if r.allowstaticlighting is enabled
-			SpecularLighting.rgb *= GBuffer.IndirectIrradiance;
+			SpecularLighting.rgb *= IndirectIrradiance;
 		}
 #endif
 #endif
@@ -513,7 +525,6 @@ void ReflectionApplyPS(
 		BRANCH
 		if (SpecularLighting.a > .001f)
 		{
-			float3 ReflectionVector = reflect(-V, GBuffer.WorldNormal);
 			// Normalize for static skylight types which mix with lightmaps
 			bool bNormalize = SkyLightParameters.z < 1 && Frame.UseLightmaps > 0 && !ALLOW_STATIC_LIGHTING;
 			float3 SkyLighting = GetSkyLightReflectionSupportingBlend(ReflectionVector, GBuffer.Roughness, bNormalize);
@@ -525,12 +536,10 @@ void ReflectionApplyPS(
 			FLATTEN
 			if (bNormalize)
 			{
-				SkyLighting *= GBuffer.IndirectIrradiance;
+				SkyLighting *= IndirectIrradiance;
 			}
 
-			float Visibility = GetDistanceFieldAOSpecularOcclusion(UV, ReflectionVector, GBuffer.Roughness, GBuffer.ShadingModelID == SHADINGMODELID_TWOSIDED_FOLIAGE); 
-
-			SpecularLighting.rgb += (Visibility * SpecularLighting.a) * SkyLighting;
+			SpecularLighting.rgb += SpecularLighting.a * SkyLighting * IndirectSpecularOcclusion;
 		}
 #endif
 
diff --git a/Engine/Shaders/SkyLightingShared.usf b/Engine/Shaders/SkyLightingShared.usf
--- a/Engine/Shaders/SkyLightingShared.usf
+++ b/Engine/Shaders/SkyLightingShared.usf
@@ -23,9 +23,10 @@ float ApplyBentNormalAO;
 float InvSkySpecularOcclusionStrength;
 float MinSkySpecularOcclusion;
 
-float GetDistanceFieldAOSpecularOcclusion(float2 ScreenUV, float3 ReflectionVector, float Roughness, bool bTwoSidedFoliage)
+void GetDistanceFieldAOSpecularOcclusion(float2 ScreenUV, float3 ReflectionVector, float Roughness, bool bTwoSidedFoliage, out float IndirectSpecularOcclusion, out float IndirectDiffuseOcclusion)
 {
-	float Visibility = 1;
+	IndirectSpecularOcclusion = 1;
+	IndirectDiffuseOcclusion = 1;
 
 	BRANCH
 	if (ApplyBentNormalAO > 0)
@@ -36,22 +37,22 @@ float GetDistanceFieldAOSpecularOcclusion(float2 ScreenUV, float3 ReflectionVect
 		BRANCH
 		if (bTwoSidedFoliage)
 		{
-			Visibility = BentNormalLength;
+			IndirectSpecularOcclusion = BentNormalLength;
 		}
 		else
 		{
 			float ReflectionConeAngle = Roughness * PI;
 			float UnoccludedAngle = BentNormalLength * PI * InvSkySpecularOcclusionStrength;
 			float AngleBetween = acos(dot(BentNormalAO, ReflectionVector) / max(BentNormalLength, .001f));
-			Visibility = ApproximateConeConeIntersection(ReflectionConeAngle, UnoccludedAngle, AngleBetween);
+			IndirectSpecularOcclusion = ApproximateConeConeIntersection(ReflectionConeAngle, UnoccludedAngle, AngleBetween);
 
 			// Can't rely on the direction of the bent normal when close to fully occluded, lerp to shadowed
-			Visibility = lerp(0, Visibility, saturate((UnoccludedAngle - .1f) / .2f));
-			Visibility = lerp(Visibility, 1, MinSkySpecularOcclusion);
+			IndirectSpecularOcclusion = lerp(0, IndirectSpecularOcclusion, saturate((UnoccludedAngle - .1f) / .2f));
+			IndirectSpecularOcclusion = lerp(IndirectSpecularOcclusion, 1, MinSkySpecularOcclusion);
 		}
-	}
 
-	return Visibility;
+		IndirectDiffuseOcclusion = lerp(BentNormalLength, 1, MinSkySpecularOcclusion);
+	}
 }
 
 float GetDynamicSkyIndirectIrradiance(float2 ScreenUV, float3 WorldNormal)
diff --git a/Engine/Source/Developer/AutomationController/Private/AutomationCommandline.cpp b/Engine/Source/Developer/AutomationController/Private/AutomationCommandline.cpp
--- a/Engine/Source/Developer/AutomationController/Private/AutomationCommandline.cpp
+++ b/Engine/Source/Developer/AutomationController/Private/AutomationCommandline.cpp
@@ -286,6 +286,8 @@ public:
 				if (AutomationCommand == EAutomationCommand::Quit)
 				{
 					GIsRequestingExit = true;
+					// We have finished the testing, and results are available
+					AutomationTestState = EAutomationTestState::Complete;
 				}
 				break;
 			}
diff --git a/Engine/Source/Developer/MeshUtilities/Private/MeshUtilities.cpp b/Engine/Source/Developer/MeshUtilities/Private/MeshUtilities.cpp
--- a/Engine/Source/Developer/MeshUtilities/Private/MeshUtilities.cpp
+++ b/Engine/Source/Developer/MeshUtilities/Private/MeshUtilities.cpp
@@ -684,12 +684,17 @@ class FQueuedThreadPool& ThreadPool,
 			SampleDirections.Add(Sample);
 		}
 
+		static const auto CVar = IConsoleManager::Get().FindTConsoleVariableDataInt(TEXT("r.DistanceFields.MaxPerMeshResolution"));
+		const int32 PerMeshMax = CVar->GetValueOnAnyThread();
+
 		// Meshes with explicit artist-specified scale can go higher
-		const int32 MaxNumVoxelsOneDim = DistanceFieldResolutionScale <= 1 ? 64 : 128;
+		const int32 MaxNumVoxelsOneDim = DistanceFieldResolutionScale <= 1 ? PerMeshMax / 2 : PerMeshMax;
 		const int32 MinNumVoxelsOneDim = 8;
 
-		//@todo - project setting
-		const float NumVoxelsPerLocalSpaceUnit = .1f * DistanceFieldResolutionScale;
+		static const auto CVarDensity = IConsoleManager::Get().FindTConsoleVariableDataFloat(TEXT("r.DistanceFields.DefaultVoxelDensity"));
+		const float VoxelDensity = CVarDensity->GetValueOnAnyThread();
+
+		const float NumVoxelsPerLocalSpaceUnit = VoxelDensity * DistanceFieldResolutionScale;
 		FBox MeshBounds(Bounds.GetBox());
 
 		{
diff --git a/Engine/Source/Runtime/Core/Private/HAL/ConsoleManager.cpp b/Engine/Source/Runtime/Core/Private/HAL/ConsoleManager.cpp
--- a/Engine/Source/Runtime/Core/Private/HAL/ConsoleManager.cpp
+++ b/Engine/Source/Runtime/Core/Private/HAL/ConsoleManager.cpp
@@ -1768,7 +1768,8 @@ static TAutoConsoleVariable<int32> CVarDepthOfFieldQuality(
 	TEXT(" 0: Off\n")
 	TEXT(" 1: Low\n")
 	TEXT(" 2: high quality (default, adaptive, can be 4x slower)\n")
-	TEXT(" 3: Special mode only affecting CircleDOF for very high quality but slow rendering"),
+	TEXT(" 3: very high quality, intended for non realtime cutscenes, CircleDOF only (slow)\n")
+	TEXT(" 4: extremely high quality, intended for non realtime cutscenes, CircleDOF only (very slow)"),
 	ECVF_Scalability | ECVF_RenderThreadSafe);
 
 static TAutoConsoleVariable<float> CVarScreenPercentage(
@@ -1849,6 +1850,20 @@ static TAutoConsoleVariable<int32> CVarDistField(
 	TEXT("Enabling will increase mesh build times and memory usage.  Changing this value will cause a rebuild of all static meshes."),
 	ECVF_ReadOnly);
 
+static TAutoConsoleVariable<int32> CVarDistFieldRes(
+	TEXT("r.DistanceFields.MaxPerMeshResolution"),
+	128,	
+	TEXT("Highest resolution (in one dimension) allowed for a single static mesh asset, used to cap the memory usage of meshes with a large scale.\n")
+	TEXT("Changing this will cause all distance fields to be rebuilt.  Large values such as 512 can consume memory very quickly! (128Mb for one asset at 512)"),
+	ECVF_ReadOnly);
+
+static TAutoConsoleVariable<float> CVarDistFieldResScale(
+	TEXT("r.DistanceFields.DefaultVoxelDensity"),
+	.1f,	
+	TEXT("Determines how the default scale of a mesh converts into distance field voxel dimensions.\n")
+	TEXT("Changing this will cause all distance fields to be rebuilt.  Large values can consume memory very quickly!"),
+	ECVF_ReadOnly);
+
 static TAutoConsoleVariable<int32> CVarLandscapeGI(
 	TEXT("r.GenerateLandscapeGIData"),
 	1,
diff --git a/Engine/Source/Runtime/Core/Private/Misc/CallbackDevice.cpp b/Engine/Source/Runtime/Core/Private/Misc/CallbackDevice.cpp
--- a/Engine/Source/Runtime/Core/Private/Misc/CallbackDevice.cpp
+++ b/Engine/Source/Runtime/Core/Private/Misc/CallbackDevice.cpp
@@ -76,4 +76,6 @@ FCoreDelegates::FOnUserActivityStringChanged FCoreDelegates::UserActivityStringC
 
 FCoreDelegates::FOnGameSessionIDChange FCoreDelegates::GameSessionIDChanged;
 
-FCoreDelegates::FOnAsyncLoadingFlush FCoreDelegates::OnAsyncLoadingFlush;
\ No newline at end of file
+FCoreDelegates::FOnAsyncLoadingFlush FCoreDelegates::OnAsyncLoadingFlush;
+FCoreDelegates::FRenderingThreadChanged FCoreDelegates::PostRenderingThreadCreated;
+FCoreDelegates::FRenderingThreadChanged FCoreDelegates::PreRenderingThreadDestroyed;
diff --git a/Engine/Source/Runtime/Core/Public/Containers/LockFreeVoidPointerListBase.h b/Engine/Source/Runtime/Core/Public/Containers/LockFreeVoidPointerListBase.h
--- a/Engine/Source/Runtime/Core/Public/Containers/LockFreeVoidPointerListBase.h
+++ b/Engine/Source/Runtime/Core/Public/Containers/LockFreeVoidPointerListBase.h
@@ -1926,6 +1926,8 @@ public:
 					FLockFreeLink* Repush = IncomingQueue.Pop(); 
 					if (!Repush)
 					{
+						FPlatformMisc::MemoryBarrier();
+						DequeueLock = 0;
 						break;
 					}
 					bResult = false;
diff --git a/Engine/Source/Runtime/Core/Public/Misc/CallbackDevice.h b/Engine/Source/Runtime/Core/Public/Misc/CallbackDevice.h
--- a/Engine/Source/Runtime/Core/Public/Misc/CallbackDevice.h
+++ b/Engine/Source/Runtime/Core/Public/Misc/CallbackDevice.h
@@ -230,6 +230,13 @@ public:
 	DECLARE_MULTICAST_DELEGATE_OneParam(FOnGameSessionIDChange, const FString&);
 	static FOnGameSessionIDChange GameSessionIDChanged;
 	
+	/** Callback for notifications regarding changes of the rendering thread. */
+	DECLARE_MULTICAST_DELEGATE(FRenderingThreadChanged)
+
+	/** Sent just after the rendering thread has been created. */
+	static FRenderingThreadChanged PostRenderingThreadCreated;
+	/* Sent just before the rendering thread is destroyed. */
+	static FRenderingThreadChanged PreRenderingThreadDestroyed;
 
 private:
 
diff --git a/Engine/Source/Runtime/Engine/Private/DistanceFieldAtlas.cpp b/Engine/Source/Runtime/Engine/Private/DistanceFieldAtlas.cpp
--- a/Engine/Source/Runtime/Engine/Private/DistanceFieldAtlas.cpp
+++ b/Engine/Source/Runtime/Engine/Private/DistanceFieldAtlas.cpp
@@ -13,18 +13,33 @@
 #include "MeshUtilities.h"
 #endif
 
-TGlobalResource<FDistanceFieldVolumeTextureAtlas> GDistanceFieldVolumeTextureAtlas = TGlobalResource<FDistanceFieldVolumeTextureAtlas>(PF_R16F);
+static TAutoConsoleVariable<int32> CVarDistFieldAtlasResXY(
+	TEXT("r.DistanceFields.AtlasSizeXY"),
+	512,
+	TEXT("Size of the global mesh distance field atlas volume texture in X and Y."),
+	ECVF_ReadOnly);
+
+static TAutoConsoleVariable<int32> CVarDistFieldAtlasResZ(
+	TEXT("r.DistanceFields.AtlasSizeZ"),
+	1024,
+	TEXT("Size of the global mesh distance field atlas volume texture in Z."),
+	ECVF_ReadOnly);
 
-// 512Mb
-const int32 MaxAtlasDimensionX = 512;
-const int32 MaxAtlasDimensionY = 512;
-const int32 MaxAtlasDimensionZ = 1024;
+TGlobalResource<FDistanceFieldVolumeTextureAtlas> GDistanceFieldVolumeTextureAtlas = TGlobalResource<FDistanceFieldVolumeTextureAtlas>(PF_R16F);
 
 FDistanceFieldVolumeTextureAtlas::FDistanceFieldVolumeTextureAtlas(EPixelFormat InFormat) :
-	BlockAllocator(0, 0, 0, MaxAtlasDimensionX, MaxAtlasDimensionY, MaxAtlasDimensionZ, false, false)
+	BlockAllocator(0, 0, 0, 0, 0, 0, false, false)
 {
 	Generation = 0;
 	Format = InFormat;
+
+	static const auto CVarXY = IConsoleManager::Get().FindTConsoleVariableDataInt(TEXT("r.DistanceFields.AtlasSizeXY"));
+	const int32 AtlasXY = CVarXY->GetValueOnAnyThread();
+
+	static const auto CVarZ = IConsoleManager::Get().FindTConsoleVariableDataInt(TEXT("r.DistanceFields.AtlasSizeZ"));
+	const int32 AtlasZ = CVarZ->GetValueOnAnyThread();
+
+	BlockAllocator = FTextureLayout3d(0, 0, 0, AtlasXY, AtlasXY, AtlasZ, false, false);
 }
 
 FString FDistanceFieldVolumeTextureAtlas::GetSizeString() const
@@ -94,8 +109,14 @@ void FDistanceFieldVolumeTextureAtlas::UpdateAllocations()
 		{
 			if (CurrentAllocations.Num() > 0)
 			{
+				static const auto CVarXY = IConsoleManager::Get().FindTConsoleVariableDataInt(TEXT("r.DistanceFields.AtlasSizeXY"));
+				const int32 AtlasXY = CVarXY->GetValueOnAnyThread();
+
+				static const auto CVarZ = IConsoleManager::Get().FindTConsoleVariableDataInt(TEXT("r.DistanceFields.AtlasSizeZ"));
+				const int32 AtlasZ = CVarZ->GetValueOnAnyThread();
+
 				// Remove all allocations from the layout so we have a clean slate
-				BlockAllocator = FTextureLayout3d(0, 0, 0, MaxAtlasDimensionX, MaxAtlasDimensionY, MaxAtlasDimensionZ, false, false);
+				BlockAllocator = FTextureLayout3d(0, 0, 0, AtlasXY, AtlasXY, AtlasZ, false, false);
 				
 				Generation++;
 
diff --git a/Engine/Source/Runtime/Engine/Private/GPUSkinCache.cpp b/Engine/Source/Runtime/Engine/Private/GPUSkinCache.cpp
--- a/Engine/Source/Runtime/Engine/Private/GPUSkinCache.cpp
+++ b/Engine/Source/Runtime/Engine/Private/GPUSkinCache.cpp
@@ -25,6 +25,8 @@ DEFINE_STAT(STAT_GPUSkinCache_SkippedForMemory);
 DEFINE_STAT(STAT_GPUSkinCache_SkippedForAlreadyCached);
 DEFINE_STAT(STAT_GPUSkinCache_SkippedForOutOfECS);
 
+DEFINE_LOG_CATEGORY_STATIC(LogSkinCache, Log, All);
+
 int32 GEnableGPUSkinCacheShaders = 0;
 static FAutoConsoleVariableRef CVarEnableGPUSkinCacheShaders(
 	TEXT("r.SkinCacheShaders"),
@@ -49,6 +51,16 @@ static TAutoConsoleVariable<int32> CVarEnableGPUSkinCache(
 	ECVF_RenderThreadSafe
 	);
 
+static TAutoConsoleVariable<int32> CVarGPUSkinCacheRecomputeTangents(
+	TEXT("r.SkinCache.RecomputeTangents"),
+	0,
+	TEXT("If r.SkinCaching is enabled this option recomputes the vertex tangents on the GPU\n")
+	TEXT("Can be changed at runtime, Requires r.SkinCacheShaders=1 and r.SkinCache=1\n")
+	TEXT(" 0: off (default)\n")
+	TEXT(" 1: on"),
+	ECVF_RenderThreadSafe
+	);
+
 static int32 GMaxGPUSkinCacheElementsPerFrame = 1000;
 static FAutoConsoleVariableRef CVarMaxGPUSkinCacheElementsPerFrame(
 	TEXT("r.MaxGPUSkinCacheElementsPerFrame"),
@@ -67,6 +79,28 @@ static FAutoConsoleVariableRef CVarMaxGPUSkinCacheBufferSize(
 	ECVF_ReadOnly | ECVF_RenderThreadSafe
 	);
 
+static TAutoConsoleVariable<float> CVarGPUSkinCacheDebug(
+	TEXT("r.SkinCache.Debug"),
+	1.0f,
+	TEXT("A constant passed to the SkinCache shader, useful for debugging"),
+	ECVF_RenderThreadSafe
+	);
+
+// We don't have it always enabled as it's not clear if this has a performance cost
+// Call on render thread only!
+// Should only be called if SM5 (compute shaders, atomics) are supported.
+ENGINE_API bool DoSkeletalMeshIndexBuffersNeedSRV()
+{
+	// currently only implemented and tested on Window SM5 (needs Compute, Atomics, SRV for index buffers, UAV for VertexBuffers)
+	return GMaxRHIShaderPlatform == SP_PCD3D_SM5 && GEnableGPUSkinCacheShaders != 0;
+}
+
+ENGINE_API bool DoRecomputeSkinTangentsOnGPU()
+{
+	// currently only implemented and tested on Window SM5 (needs Compute, Atomics, SRV for index buffers, UAV for VertexBuffers)
+	return GMaxRHIShaderPlatform == SP_PCD3D_SM5 && GEnableGPUSkinCacheShaders && GEnableGPUSkinCache && CVarGPUSkinCacheRecomputeTangents.GetValueOnRenderThread() != 0;
+}
+
 TGlobalResource<FGPUSkinCache> GGPUSkinCache;
 
 IMPLEMENT_UNIFORM_BUFFER_STRUCT(GPUSkinCacheBonesUniformShaderParameters,TEXT("GPUSkinCacheBones"));
@@ -96,12 +130,11 @@ public:
 
 		MorphBuffer.Bind(Initializer.ParameterMap, TEXT("MorphBuffer"));
 		MorphBufferOffset.Bind(Initializer.ParameterMap, TEXT("MorphBufferOffset"));
+		SkinCacheDebug.Bind(Initializer.ParameterMap, TEXT("SkinCacheDebug"));
 	}
 
-	void SetParameters(uint32 InputStreamStride,
-		uint32 InputStreamFloatOffset, uint32 InputStreamVertexCount, uint32 OutputBufferFloatOffset,
-		const FVertexBufferAndSRV& BoneBuffer, FUniformBufferRHIRef UniformBuffer,
-		FShaderResourceViewRHIRef VertexBufferSRV, FRWBuffer& SkinBuffer, const FVector& MeshOrigin, const FVector& MeshExtension,
+	void SetParameters(const FVertexBufferAndSRV& BoneBuffer, FUniformBufferRHIRef UniformBuffer,
+		FRWBuffer& SkinBuffer, const FVector& MeshOrigin, const FVector& MeshExtension,
 		const FGPUSkinCache::FDispatchData& DispatchData)
 	{
 		FComputeShaderRHIParamRef ShaderRHI = GetComputeShader();
@@ -111,9 +144,9 @@ public:
 		SetShaderValue(RHICmdList, ShaderRHI, SkinMeshOriginParameter, MeshOrigin);
 		SetShaderValue(RHICmdList, ShaderRHI, SkinMeshExtensionParameter, MeshExtension);
 
-		SetShaderValue(RHICmdList, ShaderRHI, SkinInputStreamStride, InputStreamStride);
-		SetShaderValue(RHICmdList, ShaderRHI, SkinInputStreamVertexCount, InputStreamVertexCount);
-		SetShaderValue(RHICmdList, ShaderRHI, SkinInputStreamFloatOffset, InputStreamFloatOffset);
+		SetShaderValue(RHICmdList, ShaderRHI, SkinInputStreamStride, DispatchData.InputStreamStride);
+		SetShaderValue(RHICmdList, ShaderRHI, SkinInputStreamVertexCount, DispatchData.NumVertices);
+		SetShaderValue(RHICmdList, ShaderRHI, SkinInputStreamFloatOffset, DispatchData.InputStreamFloatOffset);
 		//SetShaderValue(RHICmdList, ShaderRHI, DebugParameter, FUintVector4((uint32)(uint64)VertexBufferSRV.GetReference(), 0, 0, 1));
 
 		if (UniformBuffer)
@@ -125,10 +158,11 @@ public:
 			SetSRVParameter(RHICmdList, ShaderRHI, BoneMatrices, BoneBuffer.VertexBufferSRV);
 		}
 
-		SetSRVParameter(RHICmdList, ShaderRHI, SkinInputStream, VertexBufferSRV);
+		SetSRVParameter(RHICmdList, ShaderRHI, SkinInputStream, DispatchData.InputVertexBufferSRV);
+
 		// output UAV
-		RHICmdList.SetUAVParameter(ShaderRHI, SkinCacheBuffer.GetBaseIndex(), SkinBuffer.UAV);
-		SetShaderValue(RHICmdList, ShaderRHI, SkinOutputBufferFloatOffset, OutputBufferFloatOffset);
+		SetUAVParameter(RHICmdList, ShaderRHI, SkinCacheBuffer, SkinBuffer.UAV);
+		SetShaderValue(RHICmdList, ShaderRHI, SkinOutputBufferFloatOffset, DispatchData.OutputBufferFloatOffset);
 	
 		const bool bMorph = DispatchData.SkinType == 1;
 		if(bMorph)
@@ -136,18 +170,15 @@ public:
 			SetSRVParameter(RHICmdList, ShaderRHI, MorphBuffer, DispatchData.MorphBuffer);
 			SetShaderValue(RHICmdList, ShaderRHI, MorphBufferOffset, DispatchData.MorphBufferOffset);
 		}
+
+		SetShaderValue(RHICmdList, ShaderRHI, SkinCacheDebug, CVarGPUSkinCacheDebug.GetValueOnRenderThread());
 	}
 
 	void UnsetParameters(FRHICommandList& RHICmdList)
 	{
 		FComputeShaderRHIParamRef ShaderRHI = GetComputeShader();
 
-		// are the next 3 lines needed?
-		FShaderResourceViewRHIParamRef NullSRV = FShaderResourceViewRHIParamRef();
-		RHICmdList.SetShaderResourceViewParameter(ShaderRHI, BoneMatrices.GetBaseIndex(), NullSRV);
-		RHICmdList.SetShaderResourceViewParameter(ShaderRHI, SkinInputStream.GetBaseIndex(), NullSRV);
-
-		RHICmdList.SetUAVParameter(ShaderRHI, SkinCacheBuffer.GetBaseIndex(), FUnorderedAccessViewRHIParamRef());
+		SetUAVParameter(RHICmdList, ShaderRHI, SkinCacheBuffer, 0);
 	}
 
 	virtual bool Serialize(FArchive& Ar) override
@@ -155,7 +186,8 @@ public:
 		bool bShaderHasOutdatedParameters = FGlobalShader::Serialize(Ar);
 		Ar  << SkinMeshOriginParameter << SkinMeshExtensionParameter << SkinInputStreamStride
 			<< SkinInputStreamVertexCount << SkinInputStreamFloatOffset << SkinOutputBufferFloatOffset
-			<< SkinInputStream << SkinCacheBuffer << BoneMatrices << MorphBuffer << MorphBufferOffset;
+			<< SkinInputStream << SkinCacheBuffer << BoneMatrices << MorphBuffer << MorphBufferOffset
+			<< SkinCacheDebug;
 		//Ar << DebugParameter;
 		return bShaderHasOutdatedParameters;
 	}
@@ -167,6 +199,7 @@ private:
 
 	FShaderParameter SkinInputStreamStride;
 	FShaderParameter SkinInputStreamVertexCount;
+	FShaderParameter SkinCacheDebug;
 	FShaderParameter SkinInputStreamFloatOffset;
 	FShaderParameter SkinOutputBufferFloatOffset;
 
@@ -259,6 +292,8 @@ void FGPUSkinCache::Initialize(FRHICommandListImmediate& RHICmdList)
 {
 	check(!bInitialized);
 
+	UE_LOG(LogSkinCache, Log, TEXT("GPUSkinCache:"));
+
 	static const auto CVar = IConsoleManager::Get().FindTConsoleVariableDataInt(TEXT("r.BasePassOutputsVelocity"));
 	checkf(!CVar || CVar->GetValueOnRenderThread() == 0, TEXT("GPU Skin caching is not allowed with outputting velocity on base pass (r.BasePassOutputsVelocity=1)"));
 
@@ -273,10 +308,27 @@ void FGPUSkinCache::Initialize(FRHICommandListImmediate& RHICmdList)
 		RHICmdList.BindDebugLabelName(SkinCacheBuffer[Index].UAV, *Name);
 	}
 
+	UE_LOG(LogSkinCache, Log, TEXT("  Frames:            %d"), GPUSKINCACHE_FRAMES);
+	UE_LOG(LogSkinCache, Log, TEXT("  VertexBuffers:     %.1f MB"), FMath::DivideAndRoundUp(GGPUSkinCacheBufferSize * GPUSKINCACHE_FRAMES, 1024) / 1024.0f);
+
+	// we always allocate so we can change r.SkinCache.RecomputeTangents at runtime
+	{
+		uint32 MaxVertexCount = ComputeRecomputeTangentMaxVertexCount();
+		uint32 NumIntsPerBuffer = MaxVertexCount * FGPUSkinCache::IntermediateAccumBufferNumInts;
+		SkinTangentIntermediate.Initialize(sizeof(int32), NumIntsPerBuffer, PF_R32_SINT, BUF_UnorderedAccess);
+		RHICmdList.BindDebugLabelName(SkinTangentIntermediate.UAV, TEXT("SkinTangentIntermediate"));
+
+		uint32 MemInB = MaxVertexCount * FGPUSkinCache::IntermediateAccumBufferNumInts * sizeof(int32);
+		UE_LOG(LogSkinCache, Log, TEXT("  RecomputeTangents: %.1f MB"), FMath::DivideAndRoundUp(MemInB, (uint32)1024) / 1024.0f);
+
+		// We clear once in the beginning and then each time in the PerVertexPass (this saves the clear or another Dispatch)
+		uint32 Value[4] = { 0, 0, 0, 0 };
+		RHICmdList.ClearUAV(SkinTangentIntermediate.UAV, Value);
+	}
+
 	TransitionAllToWriteable(RHICmdList);
 
 	CachedElements.Reserve(MaxCachedElements);
-	CachedSRVs.Reserve(MaxCachedVertexBufferSRVs);
 
 	bInitialized = true;
 }
@@ -290,72 +342,10 @@ void FGPUSkinCache::Cleanup()
 			SkinCacheBuffer[Index].Release();
 		}
 
-		bInitialized = false;
-	}
-}
-
-FGPUSkinCache::FSRVCacheEntry* FGPUSkinCache::ReuseOrFindAndCacheSRV(int32& InOutIndex, const FVertexBuffer& InKey)
-{
-	FSRVCacheEntry* Ret = 0;
-	if (InOutIndex >= 0)
-	{
-		Ret = &(CachedSRVs.GetData()[InOutIndex]);
-	}
-	else
-	{
-		// time complexity: O(n) 
-		Ret = CachedSRVs.FindByKey(InKey);
-
-		if (Ret)
-		{
-			InOutIndex = Ret - CachedSRVs.GetData();
-		}
-	}
-
-	if (!Ret)
-	{
-		// Add new cache entry for vertex buffer and SRV
-		FSRVCacheEntry Info;
-		Info.Key = &InKey;
-		Info.SRVValue = RHICreateShaderResourceView(InKey.VertexBufferRHI, 4, PF_R32_FLOAT);
-
-		InOutIndex = CachedSRVs.Add(Info);
-		Ret = &(CachedSRVs.GetData()[InOutIndex]);
-	}
-
-	return Ret;
-}
-
-FGPUSkinCache::FSRVCacheEntry* FGPUSkinCache::ReuseOrFindAndCacheSRV(int32& InOutIndex, const FIndexBuffer& InKey)
-{
-	FSRVCacheEntry* Ret = 0;
-	if (InOutIndex >= 0)
-	{
-		Ret = &(CachedSRVs.GetData()[InOutIndex]);
-	}
-	else
-	{
-		// time complexity: O(n) 
-		Ret = CachedSRVs.FindByKey(InKey);
-
-		if (Ret)
-		{
-			InOutIndex = Ret - CachedSRVs.GetData();
-		}
-	}
-
-	if (!Ret)
-	{
-		// Add new cache entry for index buffer and SRV
-		FSRVCacheEntry Info;
-		Info.Key = &InKey;
-		Info.SRVValue = RHICreateShaderResourceView(InKey.IndexBufferRHI);
+		SkinTangentIntermediate.Release();
 
-		InOutIndex = CachedSRVs.Add(Info);
-		Ret = &(CachedSRVs.GetData()[InOutIndex]);
+		bInitialized = false;
 	}
-
-	return Ret;
 }
 
 // Kick off compute shader skinning for incoming chunk, return key for fast lookup in draw passes
@@ -403,8 +393,7 @@ int32 FGPUSkinCache::StartCacheMesh(FRHICommandListImmediate& RHICmdList, int32
 
 	bool AlignStreamOffsetsFrames = false;
 
-	auto* GPUVertexFactory = (FGPUBaseSkinVertexFactory*)VertexFactory;
-	const bool bExtraBoneInfluences = GPUVertexFactory->UsesExtraBoneInfluences();
+	const bool bExtraBoneInfluences = VertexFactory->UsesExtraBoneInfluences();
 	bool bHasPreviousOffset = false;
 	if (ECSInfo != nullptr)
 	{
@@ -491,16 +480,15 @@ int32 FGPUSkinCache::StartCacheMesh(FRHICommandListImmediate& RHICmdList, int32
 		ECSInfo->PreviousFrameStreamOffset = ECSInfo->StreamOffset;
 	}
 
-	const FVertexBuffer* SkinVertexBuffer = GPUVertexFactory->GetSkinVertexBuffer();
+	const FSkeletalMeshVertexBuffer* SkeletalMeshVertexBuffer = VertexFactory->GetSkinVertexBuffer();
 
-	FSRVCacheEntry* VBInfo = ReuseOrFindAndCacheSRV(ECSInfo->VertexBufferSRVIndex, *SkinVertexBuffer);
-	
-	FDispatchData DispatchData(RHICmdList, Skin->GetFeatureLevel(), Skin);
+	FDispatchData DispatchData(RHICmdList, BatchElement, Skin->GetFeatureLevel(), Skin);
+	DispatchData.NumVertices = NumVertices;
 
 	if(MorphVertexBuffer)
 	{
-		FSRVCacheEntry* MorphVBInfo = ReuseOrFindAndCacheSRV(ECSInfo->MorphVertexBufferSRVIndex, *MorphVertexBuffer);
-		DispatchData.MorphBuffer = MorphVBInfo->SRVValue;
+		DispatchData.MorphBuffer = MorphVertexBuffer->GetSRV();
+		check(DispatchData.MorphBuffer);
 
 		// in bytes
 		const uint32 MorphStride = sizeof(FMorphGPUSkinVertex);
@@ -514,18 +502,43 @@ int32 FGPUSkinCache::StartCacheMesh(FRHICommandListImmediate& RHICmdList, int32
 	INC_DWORD_STAT(STAT_GPUSkinCache_TotalNumChunks);
 	INC_DWORD_STAT_BY(STAT_GPUSkinCache_TotalNumVertices, NumVertices);
 
-	auto& ShaderData = GPUVertexFactory->GetShaderData();
+	auto& ShaderData = VertexFactory->GetShaderData();
 	
 	// SkinType 0:normal, 1:with morph target, 2:with APEX cloth (not yet implemented)
 	DispatchData.SkinType = MorphVertexBuffer ? 1 : 0;
+	DispatchData.InputStreamFloatOffset = InputStreamFloatOffset;
+	DispatchData.OutputBufferFloatOffset = ECSInfo->StreamOffset;
 
-	DispatchSkinCacheProcess(InputStreamFloatOffset, ECSInfo->StreamOffset,
+	int32 UAVIndex = FrameCounter % GPUSKINCACHE_FRAMES;
+	DispatchData.SkinCacheBuffer = &SkinCacheBuffer[UAVIndex];
+	DispatchData.InputStreamStride = StreamStrides[0];
+	DispatchData.InputVertexBufferSRV = SkeletalMeshVertexBuffer->GetSRV();
+	check(DispatchData.InputVertexBufferSRV);
+
+	DispatchSkinCacheProcess(
 		ShaderData.GetBoneBufferForWriting(false, GFrameNumberRenderThread), ShaderData.GetUniformBuffer(),
-		VBInfo, StreamStrides[0], NumVertices, ShaderData.MeshOrigin, ShaderData.MeshExtension, bExtraBoneInfluences,
+		ShaderData.MeshOrigin, ShaderData.MeshExtension, bExtraBoneInfluences,
 		DispatchData);
 
 	TargetVertexFactory->UpdateVertexDeclaration(VertexFactory, &SkinCacheBuffer[FrameCounter % GPUSKINCACHE_FRAMES]);
 
+	const bool bRecomputeTangents = MorphVertexBuffer && CVarGPUSkinCacheRecomputeTangents.GetValueOnRenderThread() != 0;
+
+	if(bRecomputeTangents)
+	{
+		FStaticLODModel* StaticLODModel = MorphVertexBuffer->GetStaticLODModel();
+		FRawStaticIndexBuffer16or32Interface* IndexBuffer = StaticLODModel->MultiSizeIndexContainer.GetIndexBuffer();
+
+		DispatchData.IndexBuffer = IndexBuffer->GetSRV();
+		check(DispatchData.IndexBuffer);
+
+		// todo: this can be optimized, worse: with multiple chunks per section this is likely not to work
+		const FSkelMeshSection* Section = StaticLODModel->FindSectionForChunk(StaticLODModel->FindChunkIndex(BatchElement));
+		DispatchData.NumTriangles = Section->NumTriangles;
+		DispatchData.IndexBufferOffsetValue = Section->BaseIndex;
+		DispatchUpdateSkinTangents(DispatchData);
+	}
+
 	CacheCurrentFloatOffset += NumRWFloats;
 	CachedChunksThisFrameCount++;
 
@@ -656,15 +669,13 @@ FGPUSkinCache::FElementCacheStatusInfo* FGPUSkinCache::FindEvictableCacheStatusI
 	}
 }
 
-void FGPUSkinCache::DispatchSkinCacheProcess(uint32 InputStreamFloatOffset, uint32 OutputBufferFloatOffset,
-	const FVertexBufferAndSRV& BoneBuffer, FUniformBufferRHIRef UniformBuffer, const FGPUSkinCache::FSRVCacheEntry* VBInfo, uint32 VertexStride,
-	uint32 VertexCount, const FVector& MeshOrigin, const FVector& MeshExtension, bool bUseExtraBoneInfluences,
+void FGPUSkinCache::DispatchSkinCacheProcess(
+	const FVertexBufferAndSRV& BoneBuffer, FUniformBufferRHIRef UniformBuffer,
+	const FVector& MeshOrigin, const FVector& MeshExtension, bool bUseExtraBoneInfluences,
 	const FDispatchData& DispatchData)
 {
 	check(DispatchData.FeatureLevel >= ERHIFeatureLevel::SM5);
 
-	uint32 VertexCountAlign64 = FMath::DivideAndRoundUp(VertexCount, (uint32)64);
-
 	SCOPED_DRAW_EVENT(DispatchData.RHICmdList, SkinCacheDispatch);
 
 	int32 UAVIndex = FrameCounter % GPUSKINCACHE_FRAMES;
@@ -685,17 +696,286 @@ void FGPUSkinCache::DispatchSkinCacheProcess(uint32 InputStreamFloatOffset, uint
 	}
 	
 	check(Shader);
-
 	DispatchData.RHICmdList.SetComputeShader(Shader->GetComputeShader());
-	Shader->SetParameters(VertexStride, InputStreamFloatOffset, VertexCount,
-		OutputBufferFloatOffset, BoneBuffer, UniformBuffer, VBInfo->SRVValue,
+	Shader->SetParameters(BoneBuffer, UniformBuffer,
 		SkinCacheBuffer[UAVIndex], MeshOrigin, MeshExtension,
 		DispatchData);
 
+	uint32 VertexCountAlign64 = FMath::DivideAndRoundUp(DispatchData.NumVertices, (uint32)64);
 	DispatchData.RHICmdList.DispatchComputeShader(VertexCountAlign64, 1, 1);
 	Shader->UnsetParameters(DispatchData.RHICmdList);
 }
 
+/** Encapsulates the RecomputeSkinTangents compute shader. */
+template <uint32 FullPrecisionUV>
+class FRecomputeTangentsPerTrianglePassCS : public FGlobalShader
+{
+	DECLARE_SHADER_TYPE(FRecomputeTangentsPerTrianglePassCS, Global);
+
+	static bool ShouldCache(EShaderPlatform Platform)
+	{
+		// currently only implemented and tested on Window SM5 (needs Compute, Atomics, SRV for index buffers, UAV for VertexBuffers)
+		return Platform == SP_PCD3D_SM5;
+	}
+
+	static void ModifyCompilationEnvironment(EShaderPlatform Platform, FShaderCompilerEnvironment& OutEnvironment)
+	{
+		FGlobalShader::ModifyCompilationEnvironment(Platform, OutEnvironment);
+		OutEnvironment.SetDefine(TEXT("THREADGROUP_SIZEX"), ThreadGroupSizeX);
+		OutEnvironment.SetDefine(TEXT("INTERMEDIATE_ACCUM_BUFFER_NUM_INTS"), FGPUSkinCache::IntermediateAccumBufferNumInts);
+		OutEnvironment.SetDefine(TEXT("FULL_PRECISION_UV"), FullPrecisionUV);
+	}
+
+	static const uint32 ThreadGroupSizeX = 64;
+
+	FRecomputeTangentsPerTrianglePassCS() {}
+
+public:
+	FShaderResourceParameter IntermediateAccumBufferUAV;
+	FShaderParameter UnitCount;
+	FShaderResourceParameter GPUSkinCacheBuffer;
+	FShaderParameter GPUSkinCacheFloatOffset;
+	FShaderResourceParameter IndexBuffer;
+	FShaderParameter IndexBufferOffset;
+	FShaderParameter SkinInputStreamFloatOffset;
+	FShaderParameter SkinInputStreamStride;
+	FShaderResourceParameter SkinInputStream;
+
+	FRecomputeTangentsPerTrianglePassCS(const ShaderMetaType::CompiledShaderInitializerType& Initializer)
+		: FGlobalShader(Initializer)
+	{
+		IntermediateAccumBufferUAV.Bind(Initializer.ParameterMap, TEXT("IntermediateAccumBufferUAV"));
+		UnitCount.Bind(Initializer.ParameterMap, TEXT("UnitCount"));
+		GPUSkinCacheBuffer.Bind(Initializer.ParameterMap,TEXT("GPUSkinCacheBuffer"));
+		GPUSkinCacheFloatOffset.Bind(Initializer.ParameterMap,TEXT("GPUSkinCacheFloatOffset"));
+		IndexBuffer.Bind(Initializer.ParameterMap,TEXT("IndexBuffer"));
+		IndexBufferOffset.Bind(Initializer.ParameterMap,TEXT("IndexBufferOffset"));
+
+		SkinInputStreamFloatOffset.Bind(Initializer.ParameterMap, TEXT("SkinCacheInputStreamFloatOffset"));
+		SkinInputStreamStride.Bind(Initializer.ParameterMap, TEXT("SkinCacheVertexStride"));
+		SkinInputStream.Bind(Initializer.ParameterMap, TEXT("SkinStreamInputBuffer"));
+	}
+
+	void SetParameters(const FGPUSkinCache::FDispatchData& DispatchData, FRWBuffer& IntermediateAccumBuffer)
+	{
+		const FComputeShaderRHIParamRef ShaderRHI = GetComputeShader();
+		FRHICommandListImmediate& RHICmdList = DispatchData.RHICmdList;
+
+		check(DispatchData.SkinCacheBuffer);
+
+//later		FGlobalShader::SetParameters(RHICmdList, ShaderRHI, View);
+
+		SetShaderValue(RHICmdList, ShaderRHI, UnitCount, DispatchData.NumTriangles);
+
+		SetSRVParameter(RHICmdList, ShaderRHI, GPUSkinCacheBuffer, DispatchData.SkinCacheBuffer->SRV);
+		SetShaderValue(RHICmdList, ShaderRHI, GPUSkinCacheFloatOffset, DispatchData.OutputBufferFloatOffset);
+
+		SetSRVParameter(RHICmdList, ShaderRHI, IndexBuffer, DispatchData.IndexBuffer);
+		SetShaderValue(RHICmdList, ShaderRHI, IndexBufferOffset, DispatchData.IndexBufferOffsetValue);
+		
+		SetShaderValue(RHICmdList, ShaderRHI, SkinInputStreamFloatOffset, DispatchData.InputStreamFloatOffset);
+		SetShaderValue(RHICmdList, ShaderRHI, SkinInputStreamStride, DispatchData.InputStreamStride);
+ 		SetSRVParameter(RHICmdList, ShaderRHI, SkinInputStream, DispatchData.InputVertexBufferSRV);
+
+		// UAV
+		SetUAVParameter(RHICmdList, ShaderRHI, IntermediateAccumBufferUAV, IntermediateAccumBuffer.UAV);
+	}
+
+	void UnsetParameters(FRHICommandList& RHICmdList)
+	{
+		FComputeShaderRHIParamRef ShaderRHI = GetComputeShader();
+
+		SetUAVParameter(RHICmdList, ShaderRHI, IntermediateAccumBufferUAV, 0);
+	}
+
+	virtual bool Serialize(FArchive& Ar) override
+	{
+		bool bShaderHasOutdatedParameters = FGlobalShader::Serialize(Ar);
+		Ar << IntermediateAccumBufferUAV << UnitCount << GPUSkinCacheBuffer << GPUSkinCacheFloatOffset << IndexBuffer << IndexBufferOffset
+			<< SkinInputStreamFloatOffset << SkinInputStreamStride << SkinInputStream;
+		return bShaderHasOutdatedParameters;
+	}
+
+	static const TCHAR* GetSourceFilename()
+	{
+		return TEXT("RecomputeTangentsPerTrianglePass");
+	}
+
+	static const TCHAR* GetFunctionName()
+	{
+		return TEXT("MainCS");
+	}
+};
+
+// #define avoids a lot of code duplication
+#define VARIATION1(A) typedef FRecomputeTangentsPerTrianglePassCS<A> FRecomputeTangentsPerTrianglePassCS##A; \
+	IMPLEMENT_SHADER_TYPE2(FRecomputeTangentsPerTrianglePassCS##A, SF_Compute);
+	VARIATION1(0) VARIATION1(1)
+#undef VARIATION1
+
+/** Encapsulates the RecomputeSkinTangentsResolve compute shader. */
+class FRecomputeTangentsPerVertexPassCS : public FGlobalShader
+{
+	DECLARE_SHADER_TYPE(FRecomputeTangentsPerVertexPassCS, Global);
+
+	static bool ShouldCache(EShaderPlatform Platform)
+	{
+		// currently only implemented and tested on Window SM5 (needs Compute, Atomics, SRV for index buffers, UAV for VertexBuffers)
+		return Platform == SP_PCD3D_SM5;
+	}
+
+	static void ModifyCompilationEnvironment(EShaderPlatform Platform, FShaderCompilerEnvironment& OutEnvironment)
+	{
+		FGlobalShader::ModifyCompilationEnvironment(Platform, OutEnvironment);
+		OutEnvironment.SetDefine(TEXT("THREADGROUP_SIZEX"), ThreadGroupSizeX);
+		OutEnvironment.SetDefine(TEXT("GPUSKIN_RWBUFFER_OFFSET_TANGENT_X"), FGPUSkinCache::RWTangentXOffsetInFloats);
+		OutEnvironment.SetDefine(TEXT("GPUSKIN_RWBUFFER_OFFSET_TANGENT_Z"), FGPUSkinCache::RWTangentZOffsetInFloats);
+		OutEnvironment.SetDefine(TEXT("GPUSKIN_RWBUFFER_NUM_FLOATS"), FGPUSkinCache::RWStrideInFloats);
+		OutEnvironment.SetDefine(TEXT("INTERMEDIATE_ACCUM_BUFFER_NUM_INTS"), FGPUSkinCache::IntermediateAccumBufferNumInts);
+	}
+
+	static const uint32 ThreadGroupSizeX = 64;
+
+	FRecomputeTangentsPerVertexPassCS() {}
+
+public:
+	FShaderResourceParameter IntermediateAccumBufferUAV;
+	FShaderResourceParameter GPUSkinCacheBufferUAV;
+	FShaderParameter GPUSkinCacheFloatOffset;
+	FShaderParameter UnitCount;
+	FShaderParameter SkinInputStreamFloatOffset;
+
+	FRecomputeTangentsPerVertexPassCS(const ShaderMetaType::CompiledShaderInitializerType& Initializer)
+		: FGlobalShader(Initializer)
+	{
+		IntermediateAccumBufferUAV.Bind(Initializer.ParameterMap, TEXT("IntermediateAccumBufferUAV"));
+		GPUSkinCacheBufferUAV.Bind(Initializer.ParameterMap,TEXT("GPUSkinCacheBufferUAV"));
+		GPUSkinCacheFloatOffset.Bind(Initializer.ParameterMap,TEXT("GPUSkinCacheFloatOffset"));
+		UnitCount.Bind(Initializer.ParameterMap, TEXT("UnitCount"));
+		SkinInputStreamFloatOffset.Bind(Initializer.ParameterMap, TEXT("SkinCacheInputStreamFloatOffset"));
+	}
+
+	void SetParameters(const FGPUSkinCache::FDispatchData& DispatchData, FRWBuffer& IntermediateAccumBuffer)
+	{
+		const FComputeShaderRHIParamRef ShaderRHI = GetComputeShader();
+		FRHICommandListImmediate& RHICmdList = DispatchData.RHICmdList;
+		
+		check(DispatchData.SkinCacheBuffer);
+
+//later		FGlobalShader::SetParameters(RHICmdList, ShaderRHI, View);
+
+		SetShaderValue(RHICmdList, ShaderRHI, GPUSkinCacheFloatOffset, DispatchData.OutputBufferFloatOffset);
+		SetShaderValue(RHICmdList, ShaderRHI, UnitCount, DispatchData.NumVertices);
+		SetShaderValue(RHICmdList, ShaderRHI, SkinInputStreamFloatOffset, DispatchData.InputStreamFloatOffset);
+
+		// UAVs
+		SetUAVParameter(RHICmdList, ShaderRHI, IntermediateAccumBufferUAV, IntermediateAccumBuffer.UAV);
+		SetUAVParameter(RHICmdList, ShaderRHI, GPUSkinCacheBufferUAV, DispatchData.SkinCacheBuffer->UAV);
+	}
+
+	void UnsetParameters(FRHICommandList& RHICmdList)
+	{
+		FComputeShaderRHIParamRef ShaderRHI = GetComputeShader();
+
+		SetUAVParameter(RHICmdList, ShaderRHI, GPUSkinCacheBufferUAV, 0);
+		SetUAVParameter(RHICmdList, ShaderRHI, IntermediateAccumBufferUAV, 0);
+	}
+
+	virtual bool Serialize(FArchive& Ar) override
+	{
+		bool bShaderHasOutdatedParameters = FGlobalShader::Serialize(Ar);
+		Ar << IntermediateAccumBufferUAV << GPUSkinCacheBufferUAV << GPUSkinCacheFloatOffset << UnitCount << SkinInputStreamFloatOffset;
+		return bShaderHasOutdatedParameters;
+	}
+};
+
+IMPLEMENT_SHADER_TYPE(,FRecomputeTangentsPerVertexPassCS,TEXT("RecomputeTangentsPerVertexPass"),TEXT("MainCS"),SF_Compute);
+
+void FGPUSkinCache::DispatchUpdateSkinTangents(const FDispatchData& DispatchData)
+{
+	FRHICommandListImmediate& RHICmdList = DispatchData.RHICmdList;
+
+	// currently only implemented and tested on Window SM5 (needs Compute, Atomics, SRV for index buffers, UAV for VertexBuffers)
+	check(DispatchData.FeatureLevel == ERHIFeatureLevel::SM5);
+
+	//	if( IsValidRef(MorphVertexBuffer.VertexBufferRHI) )
+	{
+		// no need to clear the intermediate buffer becasue we create it cleared and clear it after each usage in the per vertex pass
+
+		FSkeletalMeshResource& SkeletalMeshResource = DispatchData.GPUSkin->GetSkeletalMeshResource();
+		int32 LODIndex = DispatchData.GPUSkin->GetLOD();
+		FStaticLODModel& LodModel = SkeletalMeshResource.LODModels[LODIndex];
+
+		SetRenderTarget(RHICmdList, FTextureRHIRef(), FTextureRHIRef());
+
+		FRawStaticIndexBuffer16or32Interface* IndexBuffer = LodModel.MultiSizeIndexContainer.GetIndexBuffer();
+		FIndexBufferRHIParamRef IndexBufferRHI = IndexBuffer->IndexBufferRHI;
+
+		const uint32 RequiredVertexCount = LodModel.NumVertices;
+
+		uint32 MaxVertexCount = ComputeRecomputeTangentMaxVertexCount();
+
+		// We could do automatically reallocate but that is not yet supported
+		ensure(RequiredVertexCount < MaxVertexCount);
+
+		// todo
+//		RHICmdList.TransitionResource(EResourceTransitionAccess::ERWNoBarrier, EResourceTransitionPipeline::EGfxToCompute, IntermediateAccumBuffer.UAV);
+
+		// This code can be optimized by batched up and doing it with less Dispatch calls (costs more memory)
+		{
+			TShaderMapRef<FRecomputeTangentsPerTrianglePassCS<0> > ComputeShader0(GetGlobalShaderMap(DispatchData.FeatureLevel));
+			TShaderMapRef<FRecomputeTangentsPerTrianglePassCS<1> > ComputeShader1(GetGlobalShaderMap(DispatchData.FeatureLevel));
+			
+			bool bFullPrecisionUV = LodModel.VertexBufferGPUSkin.GetUseFullPrecisionUVs();
+			uint32 NumTriangles = DispatchData.NumTriangles;
+			uint32 ThreadGroupCountValue = FMath::DivideAndRoundUp(NumTriangles, FRecomputeTangentsPerTrianglePassCS<0>::ThreadGroupSizeX);
+
+			SCOPED_DRAW_EVENTF(RHICmdList, SkinTangents_PerTrianglePass, TEXT("SkinTangents_PerTrianglePass IndexStart=%d Tri=%d UVPrecision=%d"),
+				DispatchData.IndexBufferOffsetValue, DispatchData.NumTriangles, bFullPrecisionUV);
+
+			if(bFullPrecisionUV)
+			{
+				const FComputeShaderRHIParamRef ShaderRHI = ComputeShader1->GetComputeShader();
+				RHICmdList.SetComputeShader(ShaderRHI);
+
+				ComputeShader1->SetParameters(DispatchData, SkinTangentIntermediate);
+				DispatchComputeShader(RHICmdList, *ComputeShader1, ThreadGroupCountValue, 1, 1);
+				ComputeShader1->UnsetParameters(RHICmdList);
+			}
+			else
+			{
+				FGlobalShader* ComputeShader = bFullPrecisionUV ? (FGlobalShader*)*ComputeShader1 : (FGlobalShader*)*ComputeShader0;
+				const FComputeShaderRHIParamRef ShaderRHI = ComputeShader0->GetComputeShader();
+				RHICmdList.SetComputeShader(ShaderRHI);
+
+				ComputeShader0->SetParameters(DispatchData, SkinTangentIntermediate);
+				DispatchComputeShader(RHICmdList, *ComputeShader0, ThreadGroupCountValue, 1, 1);
+				ComputeShader0->UnsetParameters(RHICmdList);
+			}
+		}
+
+		{
+			SCOPED_DRAW_EVENTF(RHICmdList, SkinTangents_PerVertexPass, TEXT("SkinTangents_PerVertexPass InputVertStart=%d, SkinCacheStart=%d, Vert=%d"),
+				DispatchData.InputStreamFloatOffset, DispatchData.OutputBufferFloatOffset, DispatchData.NumVertices);
+
+			TShaderMapRef<FRecomputeTangentsPerVertexPassCS> ComputeShader(GetGlobalShaderMap(DispatchData.FeatureLevel));
+			const FComputeShaderRHIParamRef ShaderRHI = ComputeShader->GetComputeShader();
+			RHICmdList.SetComputeShader(ComputeShader->GetComputeShader());
+
+			uint32 VertexCount = DispatchData.NumVertices;
+			uint32 ThreadGroupCountValue = FMath::DivideAndRoundUp(VertexCount, FRecomputeTangentsPerVertexPassCS::ThreadGroupSizeX);
+
+			ComputeShader->SetParameters(DispatchData, SkinTangentIntermediate);
+	
+			DispatchComputeShader(RHICmdList, *ComputeShader, ThreadGroupCountValue, 1, 1);
+
+			ComputeShader->UnsetParameters(RHICmdList);
+		}
+// todo				RHICmdList.TransitionResource(EResourceTransitionAccess::EReadable, EResourceTransitionPipeline::EComputeToGfx, TangentsBlendBuffer.VertexBufferSRV);
+//			ensureMsgf(DestRenderTarget.TargetableTexture == DestRenderTarget.ShaderResourceTexture, TEXT("%s should be resolved to a separate SRV"), *DestRenderTarget.TargetableTexture->GetName().ToString());	
+	}
+}
+
 void FGPUSkinCache::CVarSinkFunction()
 {
 	// 0/1
diff --git a/Engine/Source/Runtime/Engine/Private/Particles/ParticleBeamTrailVertexFactory.cpp b/Engine/Source/Runtime/Engine/Private/Particles/ParticleBeamTrailVertexFactory.cpp
--- a/Engine/Source/Runtime/Engine/Private/Particles/ParticleBeamTrailVertexFactory.cpp
+++ b/Engine/Source/Runtime/Engine/Private/Particles/ParticleBeamTrailVertexFactory.cpp
@@ -112,13 +112,7 @@ void FParticleBeamTrailVertexFactory::InitRHI()
 	SetDeclaration(GParticleBeamTrailVertexDeclaration.VertexDeclarationRHI);
 
 	FVertexStream* VertexStream = new(Streams) FVertexStream;
-	VertexStream->VertexBuffer = 0;
-	VertexStream->Stride = 0;
-	VertexStream->Offset = 0;
 	FVertexStream* DynamicParameterStream = new(Streams) FVertexStream;
-	DynamicParameterStream->VertexBuffer = 0;
-	DynamicParameterStream->Stride = 0;
-	DynamicParameterStream->Offset = 0;
 }
 
 FVertexFactoryShaderParameters* FParticleBeamTrailVertexFactory::ConstructShaderParameters(EShaderFrequency ShaderFrequency)
diff --git a/Engine/Source/Runtime/Engine/Private/Particles/ParticleVertexFactory.cpp b/Engine/Source/Runtime/Engine/Private/Particles/ParticleVertexFactory.cpp
--- a/Engine/Source/Runtime/Engine/Private/Particles/ParticleVertexFactory.cpp
+++ b/Engine/Source/Runtime/Engine/Private/Particles/ParticleVertexFactory.cpp
@@ -245,13 +245,7 @@ void FParticleSpriteVertexFactory::InitStreams()
 		TexCoordStream->Offset = 0;
 	}
 	FVertexStream* InstanceStream = new(Streams) FVertexStream;
-	InstanceStream->VertexBuffer = 0;
-	InstanceStream->Stride = 0;
-	InstanceStream->Offset = 0;
 	FVertexStream* DynamicParameterStream = new(Streams) FVertexStream;
-	DynamicParameterStream->VertexBuffer = 0;
-	DynamicParameterStream->Stride = 0;
-	DynamicParameterStream->Offset = 0;
 }
 
 void FParticleSpriteVertexFactory::SetInstanceBuffer(const FVertexBuffer* InInstanceBuffer, uint32 StreamOffset, uint32 Stride, bool bInstanced)
diff --git a/Engine/Source/Runtime/Engine/Private/SceneManagement.cpp b/Engine/Source/Runtime/Engine/Private/SceneManagement.cpp
--- a/Engine/Source/Runtime/Engine/Private/SceneManagement.cpp
+++ b/Engine/Source/Runtime/Engine/Private/SceneManagement.cpp
@@ -6,6 +6,7 @@
 #include "LightMap.h"
 #include "ShadowMap.h"
 
+bool GDrawListsLocked = false;
 
 static TAutoConsoleVariable<float> CVarLODTemporalLag(
 	TEXT("lod.TemporalLag"),
diff --git a/Engine/Source/Runtime/Engine/Private/SkeletalMesh.cpp b/Engine/Source/Runtime/Engine/Private/SkeletalMesh.cpp
--- a/Engine/Source/Runtime/Engine/Private/SkeletalMesh.cpp
+++ b/Engine/Source/Runtime/Engine/Private/SkeletalMesh.cpp
@@ -226,9 +226,6 @@ void FSkeletalMeshVertexBuffer::CleanUp()
 	VertexData = NULL;
 }
 
-/**
- * Initialize the RHI resource for this vertex buffer
- */
 void FSkeletalMeshVertexBuffer::InitRHI()
 {
 	check(VertexData);
@@ -237,10 +234,20 @@ void FSkeletalMeshVertexBuffer::InitRHI()
 	{
 		// Create the vertex buffer.
 		FRHIResourceCreateInfo CreateInfo(ResourceArray);
-		VertexBufferRHI = RHICreateVertexBuffer( ResourceArray->GetResourceDataSize(), BUF_Static|BUF_ShaderResource, CreateInfo);
+		
+		// BUF_ShaderResource is needed for support of the SkinCache (we could make is dependent on GEnableGPUSkinCacheShaders or are there other users?)
+		VertexBufferRHI = RHICreateVertexBuffer( ResourceArray->GetResourceDataSize(), BUF_Static | BUF_ShaderResource, CreateInfo);
+		SRVValue = RHICreateShaderResourceView(VertexBufferRHI, 4, PF_R32_FLOAT);
 	}
 }
 
+void FSkeletalMeshVertexBuffer::ReleaseRHI()
+{
+	FVertexBuffer::ReleaseRHI();
+
+	SRVValue.SafeRelease();
+}
+
 /**
 * Serializer for this class
 * @param Ar - archive to serialize to
diff --git a/Engine/Source/Runtime/Engine/Private/SkeletalMeshComponentPhysics.cpp b/Engine/Source/Runtime/Engine/Private/SkeletalMeshComponentPhysics.cpp
--- a/Engine/Source/Runtime/Engine/Private/SkeletalMeshComponentPhysics.cpp
+++ b/Engine/Source/Runtime/Engine/Private/SkeletalMeshComponentPhysics.cpp
@@ -3622,8 +3622,9 @@ void USkeletalMeshComponent::UpdateClothMorphTarget()
 FAutoConsoleTaskPriority CPrio_FParallelClothTask(
 	TEXT("TaskGraph.TaskPriorities.ParallelClothTask"),
 	TEXT("Task and thread priority for parallel cloth."),
-	ENamedThreads::NormalThreadPriority,
-	ENamedThreads::NormalTaskPriority
+	ENamedThreads::HighThreadPriority, // if we have high priority task threads, then use them...
+	ENamedThreads::NormalTaskPriority, // .. at normal task priority
+	ENamedThreads::HighTaskPriority // if we don't have hi pri threads, then use normal priority threads at high task priority instead
 	);
 
 class FParallelClothTask
diff --git a/Engine/Source/Runtime/Engine/Private/SkeletalRenderGPUSkin.cpp b/Engine/Source/Runtime/Engine/Private/SkeletalRenderGPUSkin.cpp
--- a/Engine/Source/Runtime/Engine/Private/SkeletalRenderGPUSkin.cpp
+++ b/Engine/Source/Runtime/Engine/Private/SkeletalRenderGPUSkin.cpp
@@ -30,6 +30,10 @@
 #include "Components/SkeletalMeshComponent.h"
 #include "LocalVertexFactory.h"
 
+#include "GlobalShader.h"
+#include "ShaderParameterUtils.h"
+
+
 DEFINE_LOG_CATEGORY_STATIC(LogSkeletalGPUSkinMesh, Warning, All);
 
 // 0/1
@@ -54,9 +58,6 @@ static TAutoConsoleVariable<int32> CVarMotionBlurDebug(
 FMorphVertexBuffer
 -----------------------------------------------------------------------------*/
 
-/** 
-* Initialize the dynamic RHI for this rendering resource 
-*/
 void FMorphVertexBuffer::InitDynamicRHI()
 {
 	// LOD of the skel mesh is used to find number of vertices in buffer
@@ -67,8 +68,22 @@ void FMorphVertexBuffer::InitDynamicRHI()
 	FRHIResourceCreateInfo CreateInfo;
 	void* BufferData = nullptr;
 
-	// BUF_ShaderResource is needed for Morph support of the SkinCache
-	VertexBufferRHI = RHICreateAndLockVertexBuffer(Size, BUF_Dynamic | BUF_ShaderResource, CreateInfo, BufferData);
+	const bool bSRV = GEnableGPUSkinCacheShaders != 0;
+
+	EBufferUsageFlags Flags = BUF_Dynamic;
+
+	if(bSRV)
+	{
+		// BUF_ShaderResource is needed for Morph support of the SkinCache
+		Flags = (EBufferUsageFlags)(Flags | BUF_ShaderResource);
+	}
+
+	VertexBufferRHI = RHICreateAndLockVertexBuffer(Size, Flags, CreateInfo, BufferData);
+
+	if(bSRV)
+	{
+		SRVValue = RHICreateShaderResourceView(VertexBufferRHI, 4, PF_R32_FLOAT);
+	}
 
 	// Lock the buffer.
 	FMorphGPUSkinVertex* Buffer = (FMorphGPUSkinVertex*)BufferData;
@@ -80,12 +95,10 @@ void FMorphVertexBuffer::InitDynamicRHI()
 	bHasBeenUpdated = false;
 }
 
-/** 
-* Release the dynamic RHI for this rendering resource 
-*/
 void FMorphVertexBuffer::ReleaseDynamicRHI()
 {
 	VertexBufferRHI.SafeRelease();
+	SRVValue.SafeRelease();
 }
 
 /*-----------------------------------------------------------------------------
@@ -94,10 +107,10 @@ FSkeletalMeshObjectGPUSkin
 
 FSkeletalMeshObjectGPUSkin::FSkeletalMeshObjectGPUSkin(USkinnedMeshComponent* InMeshComponent, FSkeletalMeshResource* InSkeletalMeshResource, ERHIFeatureLevel::Type InFeatureLevel)
 	: FSkeletalMeshObject(InMeshComponent, InSkeletalMeshResource, InFeatureLevel)
-,	DynamicData(NULL)
-,	bNeedsUpdateDeferred(false)
-,	bMorphNeedsUpdateDeferred(false)
-,	bMorphResourcesInitialized(false)
+	,	DynamicData(NULL)
+	,	bNeedsUpdateDeferred(false)
+	,	bMorphNeedsUpdateDeferred(false)
+	,	bMorphResourcesInitialized(false)
 {
 	// create LODs to match the base mesh
 	LODs.Empty(SkeletalMeshResource->LODModels.Num());
@@ -289,7 +302,7 @@ void FSkeletalMeshObjectGPUSkin::ProcessUpdatedDynamicData(FRHICommandListImmedi
 
 	// if hasn't been updated, force update again
 	bMorphNeedsUpdate = LOD.MorphVertexBuffer.bHasBeenUpdated ? bMorphNeedsUpdate : true;
-	
+
 	bool bMorph = DynamicData->NumWeightedActiveVertexAnims > 0;
 
 	const FStaticLODModel& LODModel = SkeletalMeshResource->LODModels[DynamicData->LODIndex];
@@ -309,7 +322,7 @@ void FSkeletalMeshObjectGPUSkin::ProcessUpdatedDynamicData(FRHICommandListImmedi
 		if(bMorphNeedsUpdate)
 		{
 			// update the morph data for the lod (before SkinCache)
-			LOD.UpdateMorphVertexBuffer( DynamicData->ActiveVertexAnims );
+			LOD.UpdateMorphVertexBuffer(RHICmdList, DynamicData->ActiveVertexAnims);
 		}
 	}
 	else
@@ -360,8 +373,11 @@ void FSkeletalMeshObjectGPUSkin::ProcessUpdatedDynamicData(FRHICommandListImmedi
 			{
 				int32 Key = GGPUSkinCache.StartCacheMesh(RHICmdList, GPUSkinCacheKeys[ChunkIdx], VertexFactoryData.VertexFactories[ChunkIdx].Get(),
 					VertexFactoryData.PassthroughVertexFactories[ChunkIdx].Get(), Chunk, this, bMorph ? &LOD.MorphVertexBuffer : 0);
+
+				// if not failed(Key == -1)
 				if(Key >= 0)
 				{
+					ensure(Key <= 0xffff);
 					GPUSkinCacheKeys[ChunkIdx] = (int16)Key;
 				}
 			}
@@ -390,12 +406,26 @@ void FSkeletalMeshObjectGPUSkin::ProcessUpdatedDynamicData(FRHICommandListImmedi
 TArray<FVector> FSkeletalMeshObjectGPUSkin::FSkeletalMeshObjectLOD::MorphDeltaTangentZAccumulationArray;
 TArray<float> FSkeletalMeshObjectGPUSkin::FSkeletalMeshObjectLOD::MorphAccumulatedWeightArray;
 
-void FSkeletalMeshObjectGPUSkin::FSkeletalMeshObjectLOD::UpdateMorphVertexBuffer(const TArray<FActiveVertexAnim>& ActiveVertexAnims)
+void FSkeletalMeshObjectGPUSkin::FSkeletalMeshObjectLOD::UpdateMorphVertexBuffer(FRHICommandListImmediate& RHICmdList, const TArray<FActiveVertexAnim>& ActiveVertexAnims)
 {
 	SCOPE_CYCLE_COUNTER(STAT_MorphVertexBuffer_Update);
 
 	if( IsValidRef(MorphVertexBuffer.VertexBufferRHI) )
 	{
+		extern ENGINE_API bool DoRecomputeSkinTangentsOnGPU();
+		bool bBlendTangentsOnCPU = !DoRecomputeSkinTangentsOnGPU();
+
+/*		// TODO: Need to finish this to avoid artifacts if the SkinCache is not handing all objects
+		if(!bBlendTangentsOnCPU)
+		{
+			// It's possible that we reject the object from the SkinCache (e.g. all memory is used up), the we want to render the normal way.
+			// Unfortunately we don't know that at this point and code needs to be changed to get the info here.
+			if(!GGPUSkinCache.IsElementProcessed())
+			{
+				bBlendTangentsOnCPU = true;
+			}
+		}
+*/
 		// LOD of the skel mesh is used to find number of vertices in buffer
 		FStaticLODModel& LodModel = SkelMeshResource->LODModels[LODIndex];
 		uint32 Size = LodModel.NumVertices * sizeof(FMorphGPUSkinVertex);
@@ -409,14 +439,17 @@ void FSkeletalMeshObjectGPUSkin::FSkeletalMeshObjectLOD::UpdateMorphVertexBuffer
 		{
 			SCOPE_CYCLE_COUNTER(STAT_MorphVertexBuffer_Init);
 
-			// zero everything
-			int32 vertsToAdd = static_cast<int32>(LodModel.NumVertices) - MorphAccumulatedWeightArray.Num();
-			if(vertsToAdd > 0) 
+			if(bBlendTangentsOnCPU)
 			{
-				MorphAccumulatedWeightArray.AddUninitialized(vertsToAdd);
-			}
+				// zero everything
+				int32 vertsToAdd = static_cast<int32>(LodModel.NumVertices) - MorphAccumulatedWeightArray.Num();
+				if(vertsToAdd > 0) 
+				{
+					MorphAccumulatedWeightArray.AddUninitialized(vertsToAdd);
+				}
 
-			FMemory::Memzero(MorphAccumulatedWeightArray.GetData(), sizeof(float)*LodModel.NumVertices);
+				FMemory::Memzero(MorphAccumulatedWeightArray.GetData(), sizeof(float)*LodModel.NumVertices);
+			}
 
 			// PackedNormals will be wrong init with 0, but they'll be overwritten later
 			FMemory::Memzero(&Buffer[0], sizeof(FMorphGPUSkinVertex)*LodModel.NumVertices);
@@ -453,28 +486,36 @@ void FSkeletalMeshObjectGPUSkin::FSkeletalMeshObjectLOD::UpdateMorphVertexBuffer
 						FMorphGPUSkinVertex& DestVertex = Buffer[MorphVertex.SourceIdx];
 
 						DestVertex.DeltaPosition += MorphVertex.PositionDelta * VertAnim.Weight;
-						DestVertex.DeltaTangentZ += MorphVertex.TangentZDelta * VertAnim.Weight;
-						// accumulate the weight so we can normalized it later
-						MorphAccumulatedWeightArray[MorphVertex.SourceIdx] += VertAnimAbsWeight;
+
+						// todo: could be moved out of the inner loop to be more efficient
+						if(bBlendTangentsOnCPU)
+						{
+							DestVertex.DeltaTangentZ += MorphVertex.TangentZDelta * VertAnim.Weight;
+							// accumulate the weight so we can normalized it later
+							MorphAccumulatedWeightArray[MorphVertex.SourceIdx] += VertAnimAbsWeight;
+						}
 					}
 				} // for all vertices
 
 				VertAnim.VertAnim->TermEval(AnimState);
 			} // for all anim
 			
-			// copy back all the tangent values (can't use Memcpy, since we have to pack the normals)
-			for(uint32 iVertex = 0; iVertex < LodModel.NumVertices; ++iVertex) 
+			if(bBlendTangentsOnCPU)
 			{
-				FMorphGPUSkinVertex& DestVertex = Buffer[iVertex];
-				float AccumulatedWeight = MorphAccumulatedWeightArray[iVertex];
-
-				// if accumualated weight is >1.f
-				// previous code was applying the weight again in GPU if less than 1, but it doesn't make sense to do so
-				// so instead, we just divide by AccumulatedWeight if it's more than 1.
-				// now DeltaTangentZ isn't FPackedNormal, so you can apply any value to it. 
-				if (AccumulatedWeight > 1.f)
+				// copy back all the tangent values (can't use Memcpy, since we have to pack the normals)
+				for(uint32 iVertex = 0; iVertex < LodModel.NumVertices; ++iVertex) 
 				{
-					DestVertex.DeltaTangentZ /= AccumulatedWeight;
+					FMorphGPUSkinVertex& DestVertex = Buffer[iVertex];
+					float AccumulatedWeight = MorphAccumulatedWeightArray[iVertex];
+
+					// if accumulated weight is >1.f
+					// previous code was applying the weight again in GPU if less than 1, but it doesn't make sense to do so
+					// so instead, we just divide by AccumulatedWeight if it's more than 1.
+					// now DeltaTangentZ isn't FPackedNormal, so you can apply any value to it. 
+					if (AccumulatedWeight > 1.f)
+					{
+						DestVertex.DeltaTangentZ /= AccumulatedWeight;
+					}
 				}
 			}
 		} // ApplyDelta
@@ -513,20 +554,20 @@ const FVertexFactory* FSkeletalMeshObjectGPUSkin::GetVertexFactory(int32 LODInde
 		return LOD.GPUSkinVertexFactories.ClothVertexFactories[ChunkIdx]->GetVertexFactory();
 	}
 
-	// If the GPU skinning cache was used, return the passthrough vertex factory
-	if (GGPUSkinCache.IsElementProcessed(GPUSkinCacheKeys[ChunkIdx]))
-	{
-		return LOD.GPUSkinVertexFactories.PassthroughVertexFactories[ChunkIdx].Get();
-	}
+		// If the GPU skinning cache was used, return the passthrough vertex factory
+		if (GGPUSkinCache.IsElementProcessed(GPUSkinCacheKeys[ChunkIdx]))
+		{
+			return LOD.GPUSkinVertexFactories.PassthroughVertexFactories[ChunkIdx].Get();
+		}
 
 	// use the morph enabled vertex factory if any active morphs are set
 	if( DynamicData->NumWeightedActiveVertexAnims > 0 )
-	{
+		{
 		return LOD.GPUSkinVertexFactories.MorphVertexFactories[ChunkIdx].Get();
 	}
 
 	// use the default gpu skin vertex factory
-	return LOD.GPUSkinVertexFactories.VertexFactories[ChunkIdx].Get();
+			return LOD.GPUSkinVertexFactories.VertexFactories[ChunkIdx].Get();
 }
 
 /** 
@@ -563,14 +604,13 @@ void InitGPUSkinVertexFactoryComponents(typename VertexFactoryType::FDataType* V
 			VertexBuffers.VertexBufferGPUSkin,STRUCT_OFFSET(BaseVertexType,InfluenceWeights) + 4,VertexBuffers.VertexBufferGPUSkin->GetStride(),VET_UByte4N);
 	}
 
-	// uvs
+	typedef TGPUSkinVertexFloat16Uvs<MAX_TEXCOORDS, VertexFactoryType::HasExtraBoneInfluences> VertexType;
+	VertexFactoryData->PositionComponent = FVertexStreamComponent(
+		VertexBuffers.VertexBufferGPUSkin,STRUCT_OFFSET(VertexType,Position),VertexBuffers.VertexBufferGPUSkin->GetStride(),VET_Float3);
+
+	// Add a texture coordinate for each texture coordinate set we have
 	if( !VertexBuffers.VertexBufferGPUSkin->GetUseFullPrecisionUVs() )
 	{
-		typedef TGPUSkinVertexFloat16Uvs<MAX_TEXCOORDS, VertexFactoryType::HasExtraBoneInfluences> VertexType;
-		VertexFactoryData->PositionComponent = FVertexStreamComponent(
-			VertexBuffers.VertexBufferGPUSkin,STRUCT_OFFSET(VertexType,Position),VertexBuffers.VertexBufferGPUSkin->GetStride(),VET_Float3);
-
-		// Add a texture coordinate for each texture coordinate set we have
 		for( uint32 UVIndex = 0; UVIndex < VertexBuffers.VertexBufferGPUSkin->GetNumTexCoords(); ++UVIndex )
 		{
 			VertexFactoryData->TextureCoordinates.Add(FVertexStreamComponent(
@@ -579,11 +619,6 @@ void InitGPUSkinVertexFactoryComponents(typename VertexFactoryType::FDataType* V
 	}
 	else
 	{	
-		typedef TGPUSkinVertexFloat32Uvs<MAX_TEXCOORDS, VertexFactoryType::HasExtraBoneInfluences> VertexType;
-		VertexFactoryData->PositionComponent = FVertexStreamComponent(
-			VertexBuffers.VertexBufferGPUSkin,STRUCT_OFFSET(VertexType,Position),VertexBuffers.VertexBufferGPUSkin->GetStride(),VET_Float3);
-
-		// Add a texture coordinate for each texture coordinate set we have
 		for( uint32 UVIndex = 0; UVIndex < VertexBuffers.VertexBufferGPUSkin->GetNumTexCoords(); ++UVIndex )
 		{
 			VertexFactoryData->TextureCoordinates.Add(FVertexStreamComponent(
@@ -895,7 +930,7 @@ void FSkeletalMeshObjectGPUSkin::FVertexFactoryData::ReleaseMorphVertexFactories
 {
 	// Default morph factories
 	for( int32 FactoryIdx=0; FactoryIdx < MorphVertexFactories.Num(); FactoryIdx++ )
-	{		
+	{
 		BeginReleaseResource(MorphVertexFactories[FactoryIdx].Get());
 	}
 }
@@ -946,7 +981,6 @@ void FSkeletalMeshObjectGPUSkin::FVertexFactoryData::ReleaseAPEXClothVertexFacto
  * Init rendering resources for this LOD 
  * @param bUseLocalVertexFactory - use non-gpu skinned factory when rendering in ref pose
  * @param MeshLODInfo - information about the state of the bone influence swapping
- * @param Chunks - relevant chunk information (either original or from swapped influence)
  */
 void FSkeletalMeshObjectGPUSkin::FSkeletalMeshObjectLOD::InitResources(const FSkelMeshObjectLODInfo& MeshLODInfo, ERHIFeatureLevel::Type FeatureLevel)
 {
diff --git a/Engine/Source/Runtime/Engine/Private/SkeletalRenderGPUSkin.h b/Engine/Source/Runtime/Engine/Private/SkeletalRenderGPUSkin.h
--- a/Engine/Source/Runtime/Engine/Private/SkeletalRenderGPUSkin.h
+++ b/Engine/Source/Runtime/Engine/Private/SkeletalRenderGPUSkin.h
@@ -173,10 +173,22 @@ public:
 	/** Has been updated or not by UpdateMorphVertexBuffer**/
 	bool bHasBeenUpdated;
 
+	// @param guaranteed only to be valid if the vertex buffer is valid
+	FShaderResourceViewRHIParamRef GetSRV() const
+	{
+		return SRVValue;
+	}
+
+	FStaticLODModel* GetStaticLODModel() const { return &SkelMeshResource->LODModels[LODIdx]; }
+
+protected:
+	// guaranteed only to be valid if the vertex buffer is valid
+	FShaderResourceViewRHIRef SRVValue;
+
 private:
 	/** index to the SkelMeshResource.LODModels */
 	int32	LODIdx;
-	/** parent mesh containing the source data */
+	// parent mesh containing the source data, never 0
 	FSkeletalMeshResource* SkelMeshResource;
 };
 
@@ -393,6 +405,7 @@ private:
 		}
 
 		FSkeletalMeshResource* SkelMeshResource;
+		// index into FSkeletalMeshResource::LODModels[]
 		int32 LODIndex;
 
 		/** Vertex buffer that stores the morph target vertex deltas. Updated on the CPU */
@@ -406,7 +419,7 @@ private:
 		 * delta positions and delta normals from the set of active vertex anims
 		 * @param ActiveVertexAnims - vertex anims to accumulate. assumed to be weighted and have valid targets
 		 */
-		void UpdateMorphVertexBuffer( const TArray<FActiveVertexAnim>& ActiveVertexAnims );
+		void UpdateMorphVertexBuffer(FRHICommandListImmediate& RHICmdList, const TArray<FActiveVertexAnim>& ActiveVertexAnims);
 
 		/**
 		 * Determine the current vertex buffers valid for this LOD
diff --git a/Engine/Source/Runtime/Engine/Private/StaticMesh.cpp b/Engine/Source/Runtime/Engine/Private/StaticMesh.cpp
--- a/Engine/Source/Runtime/Engine/Private/StaticMesh.cpp
+++ b/Engine/Source/Runtime/Engine/Private/StaticMesh.cpp
@@ -1207,9 +1207,17 @@ static FString BuildStaticMeshDerivedDataKey(UStaticMesh* Mesh, const FStaticMes
 
 static FString BuildDistanceFieldDerivedDataKey(const FString& InMeshKey)
 {
+	static const auto CVar = IConsoleManager::Get().FindTConsoleVariableDataInt(TEXT("r.DistanceFields.MaxPerMeshResolution"));
+	const int32 PerMeshMax = CVar->GetValueOnAnyThread();
+	const FString PerMeshMaxString = PerMeshMax == 128 ? TEXT("") : FString(TEXT("_%u"), PerMeshMax);
+
+	static const auto CVarDensity = IConsoleManager::Get().FindTConsoleVariableDataFloat(TEXT("r.DistanceFields.DefaultVoxelDensity"));
+	const float VoxelDensity = CVarDensity->GetValueOnAnyThread();
+	const FString VoxelDensityString = VoxelDensity == .1f ? TEXT("") : FString(TEXT("_%.3f"), VoxelDensity);
+
 	return FDerivedDataCacheInterface::BuildCacheKey(
 		TEXT("DIST"),
-		*FString::Printf(TEXT("%s_%s"), *InMeshKey, DISTANCEFIELD_DERIVEDDATA_VER),
+		*FString::Printf(TEXT("%s_%s%s%s"), *InMeshKey, DISTANCEFIELD_DERIVEDDATA_VER, *PerMeshMaxString, *VoxelDensityString),
 		TEXT(""));
 }
 
diff --git a/Engine/Source/Runtime/Engine/Private/TickTaskManager.cpp b/Engine/Source/Runtime/Engine/Private/TickTaskManager.cpp
--- a/Engine/Source/Runtime/Engine/Private/TickTaskManager.cpp
+++ b/Engine/Source/Runtime/Engine/Private/TickTaskManager.cpp
@@ -56,6 +56,11 @@ static TAutoConsoleVariable<int32> CVarAllowAsyncTickDispatch(
 	0,
 	TEXT("If true, ticks are dispatched in a task thread."));
 
+static TAutoConsoleVariable<int32> CVarAllowAsyncTickCleanup(
+	TEXT("tick.AllowAsyncTickCleanup"),
+	1,
+	TEXT("If true, ticks are cleaned up in a task thread."));
+
 FAutoConsoleTaskPriority CPrio_DispatchTaskPriority(
 	TEXT("TaskGraph.TaskPriorities.TickDispatchTaskPriority"),
 	TEXT("Task and thread priority for tick tasks dispatch."),
@@ -433,7 +438,7 @@ public:
 	/** Add a completion handle to a tick group, parallel version **/
 	FORCEINLINE void AddTickTaskCompletionParallel(ETickingGroup StartTickGroup, ETickingGroup EndTickGroup, TGraphTask<FTickFunctionTask>* Task, bool bHiPri)
 	{
-		checkSlow(StartTickGroup >=0 && StartTickGroup < TG_MAX && EndTickGroup >=0 && EndTickGroup < TG_MAX && StartTickGroup <= EndTickGroup);
+		check(StartTickGroup >= 0 && StartTickGroup < TG_NewlySpawned && EndTickGroup >= 0 && EndTickGroup < TG_NewlySpawned && StartTickGroup <= EndTickGroup);
 		if (bHiPri)
 		{
 			HiPriTickTasks[StartTickGroup][EndTickGroup].AddThreadsafe(Task);
@@ -522,7 +527,7 @@ public:
 				if (TickCompletionEvents[Block].Num())
 				{
 					FTaskGraphInterface::Get().WaitUntilTasksComplete(TickCompletionEvents[Block], ENamedThreads::GameThread);
-					if (SingleThreadedMode() || WorldTickGroup == TG_NewlySpawned)
+					if (SingleThreadedMode() || Block == TG_NewlySpawned || CVarAllowAsyncTickCleanup.GetValueOnGameThread() == 0)
 					{
 						ResetTickGroup(Block);
 					}
@@ -698,7 +703,7 @@ public:
 	int32 StartFrame(const FTickContext& InContext)
 	{
 		check(!NewlySpawnedTickFunctions.Num()); // There shouldn't be any in here at this point in the frame
-		Context.TickGroup = TG_PrePhysics; // reset this to the start tick group
+		Context.TickGroup = ETickingGroup(0); // reset this to the start tick group
 		Context.DeltaSeconds = InContext.DeltaSeconds;
 		Context.TickType = InContext.TickType;
 		Context.Thread = ENamedThreads::GameThread;
@@ -735,7 +740,7 @@ public:
 	void StartFrameParallel(const FTickContext& InContext, TArray<FTickFunction*>& AllTickFunctions)
 	{
 		check(!NewlySpawnedTickFunctions.Num()); // There shouldn't be any in here at this point in the frame
-		Context.TickGroup = TG_PrePhysics; // reset this to the start tick group
+		Context.TickGroup = ETickingGroup(0); // reset this to the start tick group
 		Context.DeltaSeconds = InContext.DeltaSeconds;
 		Context.TickType = InContext.TickType;
 		Context.Thread = ENamedThreads::GameThread;
@@ -1325,7 +1330,7 @@ public:
 		}
 #endif
 		World = InWorld;
-		Context.TickGroup = TG_PrePhysics; // reset this to the start tick group
+		Context.TickGroup = ETickingGroup(0); // reset this to the start tick group
 		Context.DeltaSeconds = InDeltaSeconds;
 		Context.TickType = InTickType;
 		Context.Thread = ENamedThreads::GameThread;
@@ -1397,7 +1402,7 @@ public:
 	virtual void RunPauseFrame(UWorld* InWorld, float InDeltaSeconds, ELevelTick InTickType) override
 	{
 		bTickNewlySpawned = true; // we don't support new spawns, but lets at least catch them.
-		Context.TickGroup = TG_PrePhysics; // reset this to the start tick group
+		Context.TickGroup = ETickingGroup(0); // reset this to the start tick group
 		Context.DeltaSeconds = InDeltaSeconds;
 		Context.TickType = InTickType;
 		Context.Thread = ENamedThreads::GameThread;
diff --git a/Engine/Source/Runtime/Engine/Public/GPUSkinCache.h b/Engine/Source/Runtime/Engine/Public/GPUSkinCache.h
--- a/Engine/Source/Runtime/Engine/Public/GPUSkinCache.h
+++ b/Engine/Source/Runtime/Engine/Public/GPUSkinCache.h
@@ -5,6 +5,34 @@
 	GPUSkinCache.h: Performs skinning on a compute shader into a buffer to avoid vertex buffer skinning.
 =============================================================================*/
 
+// Requirements
+// * Compute shader support (with Atomics)
+// * Project settings needs to be enabled (r.SkinCacheShaders)
+// * feature need to be enabled (r.SkinCaching)
+
+// Features
+// * Skeletal mesh, 4 / 8 weights per vertex, 16/32 index buffer
+// * Supports Morph target animation (morph target blending is not done by this code)
+// * Saves vertex shader computations when we render an object multiple times (EarlyZ, velocity, shadow, BasePass, CustomDepth, Shadow masking)
+// * Fixes velocity rendering (needed for MotionBlur and TemporalAA) for WorldPosOffset animation and morph target animation
+// * RecomputeTangents results in improved tangent space for WorldPosOffset animation and morph target animation
+// * fixed amount of memory (r.MaxGPUSkinCacheElementsPerFrame, r.SkinCache.BufferSize)
+// * Velocity Rendering for MotionBlur and TemporalAA (test Velocity in BasePass)
+// * r.SkinCaching and r.SkinCaching.RecomputeTangents can be toggled at runtime
+
+// TODO:
+// * Test: Tessellation
+// * Quality/Optimization: increase TANGENT_RANGE for better quality or accumulate two components in one 32bit value
+// * Feature: RecomputeTangents should be a per section setting, the cvar (r.SkinCaching.RecomputeTangents) can stay but should be on by default
+// * Bug: We iterate through all chunks and try to find the Section. This is costly and might fail (Seems the section has a StartChunkIndex but it's called ChunkIndex)
+// * Bug: UpdateMorphVertexBuffer needs to handle SkinCacheObjects that have been rejected by the SkinCache (e.g. because it was running out of memory)
+// * Refactor: Unify the 3 compute shaders to use the same C++ setup code for the variables
+// * Optimization: Dispatch calls can be merged for better performance, stalls between Dispatch calls can be avoided (DX11 back door, DX12, console API)
+// * Feature: Needs show flag to see what objects are using the SkinCache
+// * Investigate: There might be issues with the FrameNumber (SceneCaptures, Split screen, HitProxy rendering, pause, scrubbing in Sequencer) causing TemporalAA or MotionBlur artifacts
+// * Feature: Cloth is not supported yet (Morph targets is a similar code)
+// * Feature: Support Static Meshes ?
+
 #pragma once
 
 #include "Array.h"
@@ -12,7 +40,7 @@
 #include "RenderResource.h"
 #include "UniformBuffer.h"
 
-typedef FRHIShaderResourceView*              FShaderResourceViewRHIParamRef;
+typedef FRHIShaderResourceView* FShaderResourceViewRHIParamRef;
 
 extern ENGINE_API int32 GEnableGPUSkinCacheShaders;
 extern ENGINE_API int32 GEnableGPUSkinCache;
@@ -36,6 +64,10 @@ public:
 		RWTangentZOffsetInFloats = 4,	// Packed U8x4N
 
 		RWStrideInFloats = 5,
+		// 3 ints for normal, 3 ints for tangent, 1 for orientation = 7, rounded up to 8 as it should result in faster math and caching
+		IntermediateAccumBufferNumInts = 8,
+		// max vertex count we support once the r.SkinCache.RecomputeTangents feature is used (GPU memory in bytes = n * sizeof(int) * IntermediateAccumBufferNumInts)
+		IntermediateAccumBufferSizeInKB = 1024,
 	};
 
 	FGPUSkinCache();
@@ -86,11 +118,7 @@ public:
 		const FSkeletalMeshObjectGPUSkin* Skin;
 
 		int32	Key;
-		// index into CachedSRVs[] or -1 if it has no assignment yet
-		int32	VertexBufferSRVIndex;
-		// index into CachedSRVs[] or -1 if it has no assignment yet
-		int32	MorphVertexBufferSRVIndex;
-		
+	
 		uint32	FrameUpdated;
 
 		uint32	InputVBStride;
@@ -101,12 +129,6 @@ public:
 
 		bool	bExtraBoneInfluences;
 
-		FElementCacheStatusInfo()
-			: VertexBufferSRVIndex(-1)
-			, MorphVertexBufferSRVIndex(-1)
-		{
-		}
-
 		struct FSearchInfo
 		{
 			const FSkeletalMeshObjectGPUSkin* Skin;
@@ -133,7 +155,7 @@ public:
 	struct FDispatchData 
 	{
 		FRHICommandListImmediate& RHICmdList;
-
+		const FSkelMeshChunk& Chunk;
 		ERHIFeatureLevel::Type FeatureLevel;
 		
 		// must not be 0
@@ -143,22 +165,40 @@ public:
 		uint32 SkinType;
 
 		// SkinCache output input for RecomputeSkinTagents
-		FShaderResourceViewRHIParamRef VertexBuffer;
-		uint32 VertexBufferOffset;
+		// 0 if no morph
+		FRWBuffer* SkinCacheBuffer;
+		uint32 OutputBufferFloatOffset;
+		uint32 NumVertices;
+
+		uint32 InputStreamFloatOffset;
+		uint32 InputStreamStride;
+		FShaderResourceViewRHIRef InputVertexBufferSRV;
 
 		// morph input
 		FShaderResourceViewRHIParamRef MorphBuffer;
 		uint32 MorphBufferOffset;
 
-		FDispatchData(FRHICommandListImmediate& InRHICmdList, ERHIFeatureLevel::Type InFeatureLevel, FSkeletalMeshObjectGPUSkin* InGPUSkin)
+		// triangle index buffer (input for the RecomputeSkinTangents, might need special index buffer unique to position and normal, not considering UV/vertex color)
+		FShaderResourceViewRHIParamRef IndexBuffer;
+		uint32 IndexBufferOffsetValue;
+		uint32 NumTriangles;
+		
+		FDispatchData(FRHICommandListImmediate& InRHICmdList, const FSkelMeshChunk& InChunk, ERHIFeatureLevel::Type InFeatureLevel, FSkeletalMeshObjectGPUSkin* InGPUSkin)
 			: RHICmdList(InRHICmdList)
+			, Chunk(InChunk)
 			, FeatureLevel(InFeatureLevel)
 			, GPUSkin(InGPUSkin)
 			, SkinType(0)
-			, VertexBuffer(0)
-			, VertexBufferOffset(0)
+			, SkinCacheBuffer(0)
+			, OutputBufferFloatOffset(0)
+			, NumVertices(0)
+			, InputStreamFloatOffset(0)
+			, InputStreamStride(0)
 			, MorphBuffer(0)
 			, MorphBufferOffset(0)
+			, IndexBuffer(0)
+			, IndexBufferOffsetValue(0)
+			, NumTriangles(0)
 		{
 			check(GPUSkin);
 		}
@@ -171,20 +211,26 @@ public:
 private:
 	FElementCacheStatusInfo* FindEvictableCacheStatusInfo();
 
-	struct FSRVCacheEntry;
-
-	void DispatchSkinCacheProcess(uint32 InputStreamFloatOffset, uint32 OutputBufferFloatOffset,
-		const struct FVertexBufferAndSRV& BoneBuffer, FUniformBufferRHIRef UniformBuffer, const FSRVCacheEntry* VBInfo,
-		uint32 VertexStride, uint32 VertexCount, const FVector& MeshOrigin, const FVector& MeshExtension, bool bUseExtraBoneInfluences,
+	void DispatchSkinCacheProcess(
+		const struct FVertexBufferAndSRV& BoneBuffer, FUniformBufferRHIRef UniformBuffer,
+		const FVector& MeshOrigin, const FVector& MeshExtension, bool bUseExtraBoneInfluences,
 		const FDispatchData& DispatchData);
 
-	bool InternalSetVertexStreamFromCache(FRHICommandList& RHICmdList, int32 Key, FShader* Shader, const class FGPUSkinPassthroughVertexFactory* VertexFactory, uint32 BaseVertexIndex, FShaderParameter PreviousStreamFloatOffset, FShaderResourceParameter PreviousStreamBuffer);
+	void DispatchUpdateSkinTangents(const FDispatchData& DispatchData);
 
-	bool	bInitialized : 1;
+	bool InternalSetVertexStreamFromCache(FRHICommandList& RHICmdList, int32 Key, FShader* Shader, const class FGPUSkinPassthroughVertexFactory* VertexFactory, uint32 BaseVertexIndex, FShaderParameter PreviousStreamFloatOffset, FShaderResourceParameter PreviousStreamBuffer);
 
 	void TransitionAllToWriteable(FRHICommandList& RHICmdList);
 	bool InternalIsElementProcessed(int32 Key) const;
 
+	uint32 ComputeRecomputeTangentMaxVertexCount() const
+	{
+		uint32 IntermediateAccumBufferSizeInB = FGPUSkinCache::IntermediateAccumBufferSizeInKB * 1024;
+		return IntermediateAccumBufferSizeInB / (sizeof(int32) * FGPUSkinCache::IntermediateAccumBufferNumInts);
+	}
+
+	bool bInitialized : 1;
+
 	uint32	FrameCounter;
 
 	int		CacheMaxVectorCount;
@@ -195,37 +241,12 @@ private:
 
 	FRWBuffer	SkinCacheBuffer[GPUSKINCACHE_FRAMES];
 
+	// to accumulate tangents on the GPU with atomics, is shared and reused, if not large enough the method fails
+	FRWBuffer SkinTangentIntermediate;
+
 	static void CVarSinkFunction();
 
 	static FAutoConsoleVariableSink CVarSink;
-
-	
-	// -------------------------------------------------
-	
-	// to cache the SRV needed to access the input vertex buffers
-	struct FSRVCacheEntry
-	{
-		// key: input vertex/index buffer from which we need an SRV
-		const FRenderResource* Key;
-		// value: SRV created from the VertexBuffer
-		FShaderResourceViewRHIRef SRVValue;
-		
-		// to support FindByKey()
-		bool operator == (const FRenderResource& Other) const 
-		{
-			return (Key == &Other);
-		}
-	};
-
-	// to cache the SRV needed to access the input vertex buffers
-	TArray<FSRVCacheEntry> CachedSRVs;
-	
-	// @param InOutIndex index -1 if not yet set
-	// @return 0 if we failed to find a buffer
-	FSRVCacheEntry* ReuseOrFindAndCacheSRV(int32& InOutIndex, const FVertexBuffer& InKey);
-	// @param InOutIndex index -1 if not yet set
-	// @return 0 if we failed to find a buffer
-	FSRVCacheEntry* ReuseOrFindAndCacheSRV(int32& InOutIndex, const FIndexBuffer& InKey);
 };
 
 BEGIN_UNIFORM_BUFFER_STRUCT(GPUSkinCacheBonesUniformShaderParameters,)
diff --git a/Engine/Source/Runtime/Engine/Public/GPUSkinVertexFactory.h b/Engine/Source/Runtime/Engine/Public/GPUSkinVertexFactory.h
--- a/Engine/Source/Runtime/Engine/Public/GPUSkinVertexFactory.h
+++ b/Engine/Source/Runtime/Engine/Public/GPUSkinVertexFactory.h
@@ -412,9 +412,9 @@ public:
 
 	static bool SupportsTessellationShaders() { return true; }
 
-	const FVertexBuffer* GetSkinVertexBuffer() const
+	const FSkeletalMeshVertexBuffer* GetSkinVertexBuffer() const
 	{
-		return Streams[0].VertexBuffer;
+		return (FSkeletalMeshVertexBuffer*)(Streams[0].VertexBuffer);
 	}
 
 protected:
diff --git a/Engine/Source/Runtime/Engine/Public/RawIndexBuffer.h b/Engine/Source/Runtime/Engine/Public/RawIndexBuffer.h
--- a/Engine/Source/Runtime/Engine/Public/RawIndexBuffer.h
+++ b/Engine/Source/Runtime/Engine/Public/RawIndexBuffer.h
@@ -204,6 +204,7 @@ public:
 	 * the index array a private member
 	 */
 	virtual bool GetNeedsCPUAccess() const = 0;
+	// number of indices (e.g. 4 triangles would result in 12 elements)
 	virtual int32 Num() const = 0;
 	virtual int32 AddItem(uint32 Val) = 0;
 	virtual uint32 Get(uint32 Idx) const = 0;
@@ -212,6 +213,16 @@ public:
 	virtual void Remove(int32 Idx, int32 Num = 1) = 0;
 	virtual void Empty(int32 Slack = 0) = 0;
 	virtual int32 GetResourceDataSize() const = 0;
+
+	// @param guaranteed only to be valid if the vertex buffer is valid and the buffer was created with the SRV flags
+	FShaderResourceViewRHIParamRef GetSRV() const
+	{
+		return SRVValue;
+	}
+
+protected:
+	// guaranteed only to be valid if the vertex buffer is valid and the buffer was created with the SRV flags
+	FShaderResourceViewRHIRef SRVValue;
 };
 
 template <typename INDEX_TYPE>
@@ -240,9 +251,33 @@ public:
 		{
 			// Create the index buffer.
 			FRHIResourceCreateInfo CreateInfo(&Indices);
-			IndexBufferRHI = RHICreateIndexBuffer(sizeof(INDEX_TYPE),Size,BUF_Static,CreateInfo);
+			
+			extern ENGINE_API bool DoSkeletalMeshIndexBuffersNeedSRV();
+			bool bSRV = DoSkeletalMeshIndexBuffersNeedSRV();
+
+			EBufferUsageFlags Flags = BUF_Static;
+
+			if(bSRV)
+			{
+				// BUF_ShaderResource is needed for SkinCache RecomputeSkinTangents
+				Flags = (EBufferUsageFlags)(Flags | BUF_ShaderResource);
+			}
+
+			IndexBufferRHI = RHICreateIndexBuffer(sizeof(INDEX_TYPE), Size, Flags, CreateInfo);
+
+			if(bSRV)
+			{
+				SRVValue = RHICreateShaderResourceView(IndexBufferRHI);
+			}
 		}    
 	}
+	
+	virtual void ReleaseRHI() override
+	{
+		FRawStaticIndexBuffer16or32Interface::ReleaseRHI();
+
+		SRVValue.SafeRelease();
+	}
 
 	/**
 	* Serializer for this class
diff --git a/Engine/Source/Runtime/Engine/Public/SkeletalMeshTypes.h b/Engine/Source/Runtime/Engine/Public/SkeletalMeshTypes.h
--- a/Engine/Source/Runtime/Engine/Public/SkeletalMeshTypes.h
+++ b/Engine/Source/Runtime/Engine/Public/SkeletalMeshTypes.h
@@ -814,6 +814,8 @@ public:
 	*/
 	virtual void InitRHI() override;
 
+	virtual void ReleaseRHI() override;
+
 	/**
 	* @return text description for the resource type
 	*/
@@ -1069,6 +1071,16 @@ public:
 			ConvertToFullPrecisionUVsTyped<NumTexCoordsT, false>();
 		}
 	}
+	
+	// @param guaranteed only to be valid if the vertex buffer is valid
+	FShaderResourceViewRHIParamRef GetSRV() const
+	{
+		return SRVValue;
+	}
+
+protected:
+	// guaranteed only to be valid if the vertex buffer is valid
+	FShaderResourceViewRHIRef SRVValue;
 
 private:
 	/** InfluenceBones/InfluenceWeights byte order has been swapped */
@@ -1668,6 +1680,39 @@ public:
 		return false;
 	}
 
+	// @return -1 if not found
+	uint32 FindChunkIndex(const FSkelMeshChunk& Chunk) const
+	{
+		uint32 Ret = -1;
+
+		for(int32 ChunkIdx = 0; ChunkIdx < Chunks.Num(); ++ChunkIdx)
+		{
+			if(&Chunks[ChunkIdx] == &Chunk)
+			{
+				Ret = ChunkIdx;
+				break;
+			}
+		}
+
+		return Ret;
+	}
+
+	// @return 0 if not found
+	const FSkelMeshSection* FindSectionForChunk(int16 ChunkIndex) const
+	{
+		for(auto& ref : Sections)
+		{
+			// should we check ranges instead?
+			if(ref.ChunkIndex == ChunkIndex)
+			{
+				return &ref;
+			}
+		}
+		// should never happen
+		ensure(0);
+		return 0;
+	}
+
 	/**
 	 * Get Resource Size
 	 */
diff --git a/Engine/Source/Runtime/RHI/Private/RHICommandList.cpp b/Engine/Source/Runtime/RHI/Private/RHICommandList.cpp
--- a/Engine/Source/Runtime/RHI/Private/RHICommandList.cpp
+++ b/Engine/Source/Runtime/RHI/Private/RHICommandList.cpp
@@ -1515,29 +1515,14 @@ bool FRHICommandListImmediate::StallRHIThread()
 		}
 		if (!RenderThreadSublistDispatchTask.GetReference())
 		{
-			while (RHIThreadTask.GetReference())
+			if (RHIThreadTask.GetReference() && RHIThreadTask->IsComplete())
 			{
-				SCOPE_CYCLE_COUNTER(STAT_ExplicitWaitRHIThread);
-				if (FTaskGraphInterface::Get().IsThreadProcessingTasks(ENamedThreads::RenderThread_Local))
-				{
-					// we have to spin here because all task threads might be stalled, meaning the fire event anythread task might not be hit.
-					// todo, add a third queue
-					SCOPE_CYCLE_COUNTER(STAT_SpinWaitRHIThread);
-					while (!RHIThreadTask->IsComplete())
-					{
-						FPlatformProcess::SleepNoStats(0);
-					}
-				}
-				else
-				{
-					FTaskGraphInterface::Get().WaitUntilTaskCompletes(RHIThreadTask, ENamedThreads::RenderThread_Local);
-				}
-				if (RHIThreadTask.GetReference() && RHIThreadTask->IsComplete())
-				{
-					RHIThreadTask = nullptr;
-				}
+				RHIThreadTask = nullptr;
+			}
+			if (!RHIThreadTask.GetReference())
+			{
+				return false;
 			}
-			return false;
 		}
 		if (!GRHIThreadStallEvent)
 		{
diff --git a/Engine/Source/Runtime/RenderCore/Private/RenderResource.cpp b/Engine/Source/Runtime/RenderCore/Private/RenderResource.cpp
--- a/Engine/Source/Runtime/RenderCore/Private/RenderResource.cpp
+++ b/Engine/Source/Runtime/RenderCore/Private/RenderResource.cpp
@@ -84,17 +84,17 @@ void FRenderResource::InitResourceFromPossiblyParallelRendering()
 			{
 			}
 
-			FORCEINLINE TStatId GetStatId() const
+			static FORCEINLINE TStatId GetStatId()
 			{
 				RETURN_QUICK_DECLARE_CYCLE_STAT(FInitResourceRenderThreadTask, STATGROUP_TaskGraphTasks);
 			}
 
-			ENamedThreads::Type GetDesiredThread()
+			static FORCEINLINE ENamedThreads::Type GetDesiredThread()
 			{
 				return ENamedThreads::RenderThread_Local;
 			}
 
-			static ESubsequentsMode::Type GetSubsequentsMode() { return ESubsequentsMode::FireAndForget; }
+			static FORCEINLINE ESubsequentsMode::Type GetSubsequentsMode() { return ESubsequentsMode::FireAndForget; }
 
 			void DoTask(ENamedThreads::Type CurrentThread, const FGraphEventRef& MyCompletionGraphEvent)
 			{
diff --git a/Engine/Source/Runtime/RenderCore/Private/RenderingThread.cpp b/Engine/Source/Runtime/RenderCore/Private/RenderingThread.cpp
--- a/Engine/Source/Runtime/RenderCore/Private/RenderingThread.cpp
+++ b/Engine/Source/Runtime/RenderCore/Private/RenderingThread.cpp
@@ -305,10 +305,12 @@ void RenderingThreadMain( FEvent* TaskGraphBoundSyncEvent )
 	}
 #endif
 
+	FCoreDelegates::PostRenderingThreadCreated.Broadcast();
 	check(GIsThreadedRendering);
 	FTaskGraphInterface::Get().ProcessThreadUntilRequestReturn(ENamedThreads::RenderThread);
 	FPlatformMisc::MemoryBarrier();
 	check(!GIsThreadedRendering);
+	FCoreDelegates::PreRenderingThreadDestroyed.Broadcast();
 	
 #if STATS
 	if (FThreadStats::WillEverCollectData())
diff --git a/Engine/Source/Runtime/Renderer/Private/CustomDepthRendering.cpp b/Engine/Source/Runtime/Renderer/Private/CustomDepthRendering.cpp
--- a/Engine/Source/Runtime/Renderer/Private/CustomDepthRendering.cpp
+++ b/Engine/Source/Runtime/Renderer/Private/CustomDepthRendering.cpp
@@ -34,7 +34,7 @@ bool FCustomDepthPrimSet::DrawPrims(FRHICommandListImmediate& RHICmdList, const
 			{
 				const FPrimitiveViewRelevance& ViewRelevance = View.PrimitiveViewRelevanceMap[PrimitiveSceneInfo->GetIndex()];
 
-				FDepthDrawingPolicyFactory::ContextType Context(DDM_AllOpaque);
+				FDepthDrawingPolicyFactory::ContextType Context(DDM_AllOpaque, false);
 
 				if (bWriteCustomStencilValues)
 				{
diff --git a/Engine/Source/Runtime/Renderer/Private/DeferredShadingRenderer.cpp b/Engine/Source/Runtime/Renderer/Private/DeferredShadingRenderer.cpp
--- a/Engine/Source/Runtime/Renderer/Private/DeferredShadingRenderer.cpp
+++ b/Engine/Source/Runtime/Renderer/Private/DeferredShadingRenderer.cpp
@@ -899,6 +899,8 @@ void FDeferredShadingSceneRenderer::Render(FRHICommandListImmediate& RHICmdList)
 
 	// Find the visible primitives.
 	bool bDoInitViewAftersPrepass = InitViews(RHICmdList, ILCTaskData, SortEvents);
+
+	TGuardValue<bool> LockDrawLists(GDrawListsLocked, true);
 #if !UE_BUILD_SHIPPING
 	if (CVarStallInitViews.GetValueOnRenderThread() > 0.0f)
 	{
@@ -1303,6 +1305,10 @@ void FDeferredShadingSceneRenderer::Render(FRHICommandListImmediate& RHICmdList)
 		// These modulate the scenecolor output from the basepass, which is assumed to be indirect lighting
 		RenderIndirectCapsuleShadows(RHICmdList);
 
+		TRefCountPtr<IPooledRenderTarget> DynamicBentNormalAO;
+		// These modulate the scenecolor output from the basepass, which is assumed to be indirect lighting
+		RenderDFAOAsIndirectShadowing(RHICmdList, VelocityRT, DynamicBentNormalAO);
+
 		// Clear the translucent lighting volumes before we accumulate
 		ClearTranslucentVolumeLighting(RHICmdList);
 
@@ -1335,7 +1341,6 @@ void FDeferredShadingSceneRenderer::Render(FRHICommandListImmediate& RHICmdList)
 			}
 		}
 
-		TRefCountPtr<IPooledRenderTarget> DynamicBentNormalAO;
 		RenderDynamicSkyLighting(RHICmdList, VelocityRT, DynamicBentNormalAO);
 		ServiceLocalQueue();
 
@@ -1470,7 +1475,7 @@ void FDeferredShadingSceneRenderer::Render(FRHICommandListImmediate& RHICmdList)
 		const float OcclusionMaxDistance = Scene->SkyLight && !Scene->SkyLight->bWantsStaticShadowing ? Scene->SkyLight->OcclusionMaxDistance : GDefaultDFAOMaxOcclusionDistance;
 		TRefCountPtr<IPooledRenderTarget> DummyOutput;
 		RHICmdList.SetCurrentStat(GET_STATID(STAT_CLM_RenderDistanceFieldLighting));
-		RenderDistanceFieldLighting(RHICmdList, FDistanceFieldAOParameters(OcclusionMaxDistance), VelocityRT, DummyOutput, DummyOutput, ViewFamily.EngineShowFlags.VisualizeDistanceFieldAO, ViewFamily.EngineShowFlags.VisualizeDistanceFieldGI);
+		RenderDistanceFieldLighting(RHICmdList, FDistanceFieldAOParameters(OcclusionMaxDistance), VelocityRT, DummyOutput, DummyOutput, false, ViewFamily.EngineShowFlags.VisualizeDistanceFieldAO, ViewFamily.EngineShowFlags.VisualizeDistanceFieldGI); 
 		ServiceLocalQueue();
 	}
 
@@ -1526,7 +1531,7 @@ void FDeferredShadingSceneRenderer::Render(FRHICommandListImmediate& RHICmdList)
 
 bool FDeferredShadingSceneRenderer::RenderPrePassViewDynamic(FRHICommandList& RHICmdList, const FViewInfo& View)
 {
-	FDepthDrawingPolicyFactory::ContextType Context(EarlyZPassMode);
+	FDepthDrawingPolicyFactory::ContextType Context(EarlyZPassMode, true);
 
 	for (int32 MeshBatchIndex = 0; MeshBatchIndex < View.DynamicMeshElements.Num(); MeshBatchIndex++)
 	{
diff --git a/Engine/Source/Runtime/Renderer/Private/DeferredShadingRenderer.h b/Engine/Source/Runtime/Renderer/Private/DeferredShadingRenderer.h
--- a/Engine/Source/Runtime/Renderer/Private/DeferredShadingRenderer.h
+++ b/Engine/Source/Runtime/Renderer/Private/DeferredShadingRenderer.h
@@ -201,6 +201,12 @@ private:
 	/** Render dynamic sky lighting from Movable sky lights. */
 	void RenderDynamicSkyLighting(FRHICommandListImmediate& RHICmdList, const TRefCountPtr<IPooledRenderTarget>& VelocityTexture, TRefCountPtr<IPooledRenderTarget>& DynamicBentNormalAO);
 
+	/** Computes DFAO, modulates it to scene color (which is assumed to contain diffuse indirect lighting), and stores the output bent normal for use occluding specular. */
+	void RenderDFAOAsIndirectShadowing(
+		FRHICommandListImmediate& RHICmdList,
+		const TRefCountPtr<IPooledRenderTarget>& VelocityTexture,
+		TRefCountPtr<IPooledRenderTarget>& DynamicBentNormalAO);
+
 	/** Render Ambient Occlusion using mesh distance fields and the surface cache, which supports dynamic rigid meshes. */
 	bool RenderDistanceFieldLighting(
 		FRHICommandListImmediate& RHICmdList, 
@@ -208,6 +214,7 @@ private:
 		const TRefCountPtr<IPooledRenderTarget>& VelocityTexture,
 		TRefCountPtr<IPooledRenderTarget>& OutDynamicBentNormalAO, 
 		TRefCountPtr<IPooledRenderTarget>& OutDynamicIrradiance,
+		bool bModulateToSceneColor,
 		bool bVisualizeAmbientOcclusion,
 		bool bVisualizeGlobalIllumination);
 
diff --git a/Engine/Source/Runtime/Renderer/Private/DepthRendering.cpp b/Engine/Source/Runtime/Renderer/Private/DepthRendering.cpp
--- a/Engine/Source/Runtime/Renderer/Private/DepthRendering.cpp
+++ b/Engine/Source/Runtime/Renderer/Private/DepthRendering.cpp
@@ -508,7 +508,7 @@ bool FDepthDrawingPolicyFactory::DrawMesh(
 
 	//Do a per-FMeshBatch check on top of the proxy check in RenderPrePass to handle the case where a proxy that is relevant 
 	//to the depth only pass has to submit multiple FMeshElements but only some of them should be used as occluders.
-	if (Mesh.bUseAsOccluder || DrawingContext.DepthDrawingMode == DDM_AllOpaque)
+	if (Mesh.bUseAsOccluder || !DrawingContext.bRespectUseAsOccluderFlag || DrawingContext.DepthDrawingMode == DDM_AllOpaque)
 	{
 		const FMaterialRenderProxy* MaterialRenderProxy = Mesh.MaterialRenderProxy;
 		const FMaterial* Material = MaterialRenderProxy->GetMaterial(View.GetFeatureLevel());
diff --git a/Engine/Source/Runtime/Renderer/Private/DepthRendering.h b/Engine/Source/Runtime/Renderer/Private/DepthRendering.h
--- a/Engine/Source/Runtime/Renderer/Private/DepthRendering.h
+++ b/Engine/Source/Runtime/Renderer/Private/DepthRendering.h
@@ -150,8 +150,12 @@ public:
 	{
 		EDepthDrawingMode DepthDrawingMode;
 
-		ContextType(EDepthDrawingMode InDepthDrawingMode) :
-			DepthDrawingMode(InDepthDrawingMode)
+		/** Uses of FDepthDrawingPolicyFactory that are not the depth pre-pass will not want the bUseAsOccluder flag to interfere. */
+		bool bRespectUseAsOccluderFlag;
+
+		ContextType(EDepthDrawingMode InDepthDrawingMode, bool bInRespectUseAsOccluderFlag) :
+			DepthDrawingMode(InDepthDrawingMode),
+			bRespectUseAsOccluderFlag(bInRespectUseAsOccluderFlag)
 		{}
 	};
 
diff --git a/Engine/Source/Runtime/Renderer/Private/DistanceFieldLightingPost.cpp b/Engine/Source/Runtime/Renderer/Private/DistanceFieldLightingPost.cpp
--- a/Engine/Source/Runtime/Renderer/Private/DistanceFieldLightingPost.cpp
+++ b/Engine/Source/Runtime/Renderer/Private/DistanceFieldLightingPost.cpp
@@ -328,11 +328,14 @@ public:
 		UseHistoryFilter.Bind(Initializer.ParameterMap, TEXT("UseHistoryFilter"));
 		VelocityTexture.Bind(Initializer.ParameterMap, TEXT("VelocityTexture"));
 		VelocityTextureSampler.Bind(Initializer.ParameterMap, TEXT("VelocityTextureSampler"));
+		DistanceFieldNormalTexture.Bind(Initializer.ParameterMap, TEXT("DistanceFieldNormalTexture"));
+		DistanceFieldNormalSampler.Bind(Initializer.ParameterMap, TEXT("DistanceFieldNormalSampler"));
 	}
 
 	void SetParameters(
 		FRHICommandList& RHICmdList, 
 		const FSceneView& View, 
+		FSceneRenderTargetItem& DistanceFieldNormal, 
 		FSceneRenderTargetItem& BentNormalHistoryTextureValue, 
 		TRefCountPtr<IPooledRenderTarget>* IrradianceHistoryRT, 
 		FSceneRenderTargetItem& DistanceFieldAOBentNormal, 
@@ -358,10 +361,19 @@ public:
 			ShaderRHI,
 			BentNormalHistoryTexture,
 			BentNormalHistorySampler,
-			TStaticSamplerState<SF_Bilinear,AM_Clamp,AM_Clamp,AM_Clamp>::GetRHI(),
+			TStaticSamplerState<SF_Point,AM_Clamp,AM_Clamp,AM_Clamp>::GetRHI(),
 			BentNormalHistoryTextureValue.ShaderResourceTexture
 			);
 
+		SetTextureParameter(
+			RHICmdList,
+			ShaderRHI,
+			DistanceFieldNormalTexture,
+			DistanceFieldNormalSampler,
+			TStaticSamplerState<SF_Bilinear,AM_Clamp,AM_Clamp,AM_Clamp>::GetRHI(),
+			DistanceFieldNormal.ShaderResourceTexture
+			);
+
 		if (IrradianceTexture.IsBound())
 		{
 			SetTextureParameter(
@@ -419,6 +431,8 @@ public:
 		Ar << UseHistoryFilter;
 		Ar << VelocityTexture;
 		Ar << VelocityTextureSampler;
+		Ar << DistanceFieldNormalTexture;
+		Ar << DistanceFieldNormalSampler;
 		return bShaderHasOutdatedParameters;
 	}
 
@@ -438,6 +452,8 @@ private:
 	FShaderParameter UseHistoryFilter;
 	FShaderResourceParameter VelocityTexture;
 	FShaderResourceParameter VelocityTextureSampler;
+	FShaderResourceParameter DistanceFieldNormalTexture;
+	FShaderResourceParameter DistanceFieldNormalSampler;
 };
 
 IMPLEMENT_SHADER_TYPE(template<>,TUpdateHistoryPS<true>,TEXT("DistanceFieldLightingPost"),TEXT("UpdateHistoryPS"),SF_Pixel);
@@ -474,11 +490,14 @@ public:
 		IrradianceSampler.Bind(Initializer.ParameterMap, TEXT("IrradianceSampler"));
 		HistoryWeight.Bind(Initializer.ParameterMap, TEXT("HistoryWeight"));
 		BentNormalAOTexelSize.Bind(Initializer.ParameterMap, TEXT("BentNormalAOTexelSize"));
+		DistanceFieldNormalTexture.Bind(Initializer.ParameterMap, TEXT("DistanceFieldNormalTexture"));
+		DistanceFieldNormalSampler.Bind(Initializer.ParameterMap, TEXT("DistanceFieldNormalSampler"));
 	}
 
 	void SetParameters(
 		FRHICommandList& RHICmdList, 
 		const FSceneView& View, 
+		FSceneRenderTargetItem& DistanceFieldNormal, 
 		FSceneRenderTargetItem& BentNormalHistoryTextureValue, 
 		IPooledRenderTarget* IrradianceHistoryRT)
 	{
@@ -495,6 +514,15 @@ public:
 			BentNormalHistoryTextureValue.ShaderResourceTexture
 			);
 
+		SetTextureParameter(
+			RHICmdList,
+			ShaderRHI,
+			DistanceFieldNormalTexture,
+			DistanceFieldNormalSampler,
+			TStaticSamplerState<SF_Bilinear,AM_Clamp,AM_Clamp,AM_Clamp>::GetRHI(),
+			DistanceFieldNormal.ShaderResourceTexture
+			);
+
 		if (IrradianceTexture.IsBound())
 		{
 			SetTextureParameter(
@@ -525,6 +553,8 @@ public:
 		Ar << IrradianceSampler;
 		Ar << HistoryWeight;
 		Ar << BentNormalAOTexelSize;
+		Ar << DistanceFieldNormalTexture;
+		Ar << DistanceFieldNormalSampler;
 		return bShaderHasOutdatedParameters;
 	}
 
@@ -536,6 +566,8 @@ private:
 	FShaderResourceParameter IrradianceSampler;
 	FShaderParameter HistoryWeight;
 	FShaderParameter BentNormalAOTexelSize;
+	FShaderResourceParameter DistanceFieldNormalTexture;
+	FShaderResourceParameter DistanceFieldNormalSampler;
 };
 
 IMPLEMENT_SHADER_TYPE(template<>,TFilterHistoryPS<true>,TEXT("DistanceFieldLightingPost"),TEXT("FilterHistoryPS"),SF_Pixel);
@@ -560,6 +592,7 @@ void UpdateHistory(
 	const TCHAR* BentNormalHistoryRTName,
 	const TCHAR* IrradianceHistoryRTName,
 	IPooledRenderTarget* VelocityTexture,
+	FSceneRenderTargetItem& DistanceFieldNormal,
 	/** Contains last frame's history, if non-NULL.  This will be updated with the new frame's history. */
 	TRefCountPtr<IPooledRenderTarget>* BentNormalHistoryState,
 	TRefCountPtr<IPooledRenderTarget>* IrradianceHistoryState,
@@ -616,7 +649,7 @@ void UpdateHistory(
 
 					static FGlobalBoundShaderState BoundShaderState;
 					SetGlobalBoundShaderState(RHICmdList, View.GetFeatureLevel(), BoundShaderState, GFilterVertexDeclaration.VertexDeclarationRHI, *VertexShader, *PixelShader);
-					PixelShader->SetParameters(RHICmdList, View, (*BentNormalHistoryState)->GetRenderTargetItem(), IrradianceHistoryState, BentNormalSource->GetRenderTargetItem(), IrradianceSource, VelocityTexture);
+					PixelShader->SetParameters(RHICmdList, View, DistanceFieldNormal, (*BentNormalHistoryState)->GetRenderTargetItem(), IrradianceHistoryState, BentNormalSource->GetRenderTargetItem(), IrradianceSource, VelocityTexture);
 				}
 				else
 				{
@@ -624,7 +657,7 @@ void UpdateHistory(
 
 					static FGlobalBoundShaderState BoundShaderState;
 					SetGlobalBoundShaderState(RHICmdList, View.GetFeatureLevel(), BoundShaderState, GFilterVertexDeclaration.VertexDeclarationRHI, *VertexShader, *PixelShader);
-					PixelShader->SetParameters(RHICmdList, View, (*BentNormalHistoryState)->GetRenderTargetItem(), IrradianceHistoryState, BentNormalSource->GetRenderTargetItem(), IrradianceSource, VelocityTexture);
+					PixelShader->SetParameters(RHICmdList, View, DistanceFieldNormal, (*BentNormalHistoryState)->GetRenderTargetItem(), IrradianceHistoryState, BentNormalSource->GetRenderTargetItem(), IrradianceSource, VelocityTexture);
 				}
 
 				VertexShader->SetParameters(RHICmdList, View);
@@ -689,7 +722,7 @@ void UpdateHistory(
 
 						static FGlobalBoundShaderState BoundShaderState;
 						SetGlobalBoundShaderState(RHICmdList, View.GetFeatureLevel(), BoundShaderState, GFilterVertexDeclaration.VertexDeclarationRHI, *VertexShader, *PixelShader);
-						PixelShader->SetParameters(RHICmdList, View, NewBentNormalHistory->GetRenderTargetItem(), NewIrradianceHistory);
+						PixelShader->SetParameters(RHICmdList, View, DistanceFieldNormal, NewBentNormalHistory->GetRenderTargetItem(), NewIrradianceHistory);
 					}
 					else
 					{
@@ -697,7 +730,7 @@ void UpdateHistory(
 
 						static FGlobalBoundShaderState BoundShaderState;
 						SetGlobalBoundShaderState(RHICmdList, View.GetFeatureLevel(), BoundShaderState, GFilterVertexDeclaration.VertexDeclarationRHI, *VertexShader, *PixelShader);
-						PixelShader->SetParameters(RHICmdList, View, NewBentNormalHistory->GetRenderTargetItem(), NewIrradianceHistory);
+						PixelShader->SetParameters(RHICmdList, View, DistanceFieldNormal, NewBentNormalHistory->GetRenderTargetItem(), NewIrradianceHistory);
 					}
 
 					VertexShader->SetParameters(RHICmdList, View);
@@ -946,6 +979,7 @@ void PostProcessBentNormalAOSurfaceCache(
 			TEXT("DistanceFieldAOHistory"),
 			TEXT("DistanceFieldIrradianceHistory"),
 			VelocityTexture,
+			DistanceFieldNormal,
 			BentNormalHistoryState,
 			IrradianceHistoryState,
 			DistanceFieldAOBentNormal, 
@@ -963,7 +997,7 @@ enum EAOUpsampleType
 	AOUpsample_OutputIrradiance
 };
 
-template<EAOUpsampleType UpsampleType>
+template<EAOUpsampleType UpsampleType, bool bModulateToSceneColor>
 class TDistanceFieldAOUpsamplePS : public FGlobalShader
 {
 	DECLARE_SHADER_TYPE(TDistanceFieldAOUpsamplePS, Global);
@@ -977,6 +1011,7 @@ public:
 	static void ModifyCompilationEnvironment(EShaderPlatform Platform, FShaderCompilerEnvironment& OutEnvironment)
 	{
 		OutEnvironment.SetDefine(TEXT("DOWNSAMPLE_FACTOR"), GAODownsampleFactor);
+		OutEnvironment.SetDefine(TEXT("MODULATE_SCENE_COLOR"), bModulateToSceneColor ? TEXT("1") : TEXT("0"));
 		OutEnvironment.SetDefine(TEXT("OUTPUT_BENT_NORMAL"), (UpsampleType == AOUpsample_OutputBentNormal || UpsampleType == AOUpsample_OutputBentNormalAndIrradiance) ? TEXT("1") : TEXT("0"));
 		OutEnvironment.SetDefine(TEXT("SUPPORT_IRRADIANCE"), (UpsampleType == AOUpsample_OutputIrradiance || UpsampleType == AOUpsample_OutputBentNormalAndIrradiance) ? TEXT("1") : TEXT("0"));
 		OutEnvironment.SetDefine(TEXT("OUTPUT_AO"), (UpsampleType == AOUpsample_OutputAO) ? TEXT("1") : TEXT("0"));
@@ -994,6 +1029,7 @@ public:
 		BentNormalAOSampler.Bind(Initializer.ParameterMap,TEXT("BentNormalAOSampler"));
 		IrradianceTexture.Bind(Initializer.ParameterMap,TEXT("IrradianceTexture"));
 		IrradianceSampler.Bind(Initializer.ParameterMap,TEXT("IrradianceSampler"));
+		MinIndirectDiffuseOcclusion.Bind(Initializer.ParameterMap,TEXT("MinIndirectDiffuseOcclusion"));
 	}
 
 	void SetParameters(FRHICommandList& RHICmdList, const FSceneView& View, TRefCountPtr<IPooledRenderTarget>& DistanceFieldAOBentNormal, IPooledRenderTarget* DistanceFieldIrradiance)
@@ -1009,6 +1045,10 @@ public:
 		{
 			SetTextureParameter(RHICmdList, ShaderRHI, IrradianceTexture, IrradianceSampler, TStaticSamplerState<SF_Bilinear>::GetRHI(), DistanceFieldIrradiance->GetRenderTargetItem().ShaderResourceTexture);
 		}
+
+		FScene* Scene = (FScene*)View.Family->Scene;
+		const float MinOcclusion = Scene->SkyLight ? Scene->SkyLight->MinOcclusion : 0;
+		SetShaderValue(RHICmdList, ShaderRHI, MinIndirectDiffuseOcclusion, MinOcclusion);
 	}
 	// FShader interface.
 	virtual bool Serialize(FArchive& Ar) override
@@ -1019,6 +1059,7 @@ public:
 		Ar << BentNormalAOSampler;
 		Ar << IrradianceTexture;
 		Ar << IrradianceSampler;
+		Ar << MinIndirectDiffuseOcclusion;
 		return bShaderHasOutdatedParameters;
 	}
 
@@ -1029,22 +1070,29 @@ private:
 	FShaderResourceParameter BentNormalAOSampler;
 	FShaderResourceParameter IrradianceTexture;
 	FShaderResourceParameter IrradianceSampler;
+	FShaderParameter MinIndirectDiffuseOcclusion;
 };
 
-#define IMPLEMENT_UPSAMPLE_PS_TYPE(UpsampleType) \
-	typedef TDistanceFieldAOUpsamplePS<UpsampleType> TDistanceFieldAOUpsamplePS##UpsampleType; \
-	IMPLEMENT_SHADER_TYPE(template<>,TDistanceFieldAOUpsamplePS##UpsampleType,TEXT("DistanceFieldLightingPost"),TEXT("AOUpsamplePS"),SF_Pixel);
+#define IMPLEMENT_UPSAMPLE_PS_TYPE(UpsampleType, bModulateToSceneColor) \
+	typedef TDistanceFieldAOUpsamplePS<UpsampleType, bModulateToSceneColor> TDistanceFieldAOUpsamplePS##UpsampleType##bModulateToSceneColor; \
+	IMPLEMENT_SHADER_TYPE(template<>,TDistanceFieldAOUpsamplePS##UpsampleType##bModulateToSceneColor,TEXT("DistanceFieldLightingPost"),TEXT("AOUpsamplePS"),SF_Pixel);
+
+IMPLEMENT_UPSAMPLE_PS_TYPE(AOUpsample_OutputBentNormal, true)
+IMPLEMENT_UPSAMPLE_PS_TYPE(AOUpsample_OutputAO, true)
+IMPLEMENT_UPSAMPLE_PS_TYPE(AOUpsample_OutputBentNormalAndIrradiance, true)
+IMPLEMENT_UPSAMPLE_PS_TYPE(AOUpsample_OutputIrradiance, true)
 
-IMPLEMENT_UPSAMPLE_PS_TYPE(AOUpsample_OutputBentNormal)
-IMPLEMENT_UPSAMPLE_PS_TYPE(AOUpsample_OutputAO)
-IMPLEMENT_UPSAMPLE_PS_TYPE(AOUpsample_OutputBentNormalAndIrradiance)
-IMPLEMENT_UPSAMPLE_PS_TYPE(AOUpsample_OutputIrradiance)
+IMPLEMENT_UPSAMPLE_PS_TYPE(AOUpsample_OutputBentNormal, false)
+IMPLEMENT_UPSAMPLE_PS_TYPE(AOUpsample_OutputAO, false)
+IMPLEMENT_UPSAMPLE_PS_TYPE(AOUpsample_OutputBentNormalAndIrradiance, false)
+IMPLEMENT_UPSAMPLE_PS_TYPE(AOUpsample_OutputIrradiance, false)
 
 void UpsampleBentNormalAO(
 	FRHICommandList& RHICmdList, 
 	const TArray<FViewInfo>& Views, 
 	TRefCountPtr<IPooledRenderTarget>& DistanceFieldAOBentNormal, 
 	TRefCountPtr<IPooledRenderTarget>& DistanceFieldIrradiance,
+	bool bModulateSceneColor,
 	bool bVisualizeAmbientOcclusion,
 	bool bVisualizeGlobalIllumination)
 {
@@ -1058,40 +1106,91 @@ void UpsampleBentNormalAO(
 		RHICmdList.SetViewport(View.ViewRect.Min.X, View.ViewRect.Min.Y, 0.0f, View.ViewRect.Max.X, View.ViewRect.Max.Y, 1.0f);
 		RHICmdList.SetRasterizerState(TStaticRasterizerState<FM_Solid, CM_None>::GetRHI());
 		RHICmdList.SetDepthStencilState(TStaticDepthStencilState<false, CF_Always>::GetRHI());
-		RHICmdList.SetBlendState(TStaticBlendState<>::GetRHI());
 
-		TShaderMapRef<FPostProcessVS> VertexShader(View.ShaderMap);
-
-		if (bVisualizeAmbientOcclusion)
+		if (bModulateSceneColor)
 		{
-			TShaderMapRef<TDistanceFieldAOUpsamplePS<AOUpsample_OutputAO> > PixelShader(View.ShaderMap);
-			static FGlobalBoundShaderState BoundShaderState;
-			SetGlobalBoundShaderState(RHICmdList, View.GetFeatureLevel(), BoundShaderState, GFilterVertexDeclaration.VertexDeclarationRHI, *VertexShader, *PixelShader);
-			PixelShader->SetParameters(RHICmdList, View, DistanceFieldAOBentNormal, DistanceFieldIrradiance);
+			RHICmdList.SetBlendState(TStaticBlendState<
+				// Opaque blending to DistanceFieldAOBentNormal
+				CW_RGBA, BO_Add, BF_One, BF_Zero, BO_Add, BF_One, BF_Zero,
+				// Modulate scene color target
+				CW_RGB, BO_Add, BF_DestColor, BF_Zero, BO_Add, BF_Zero, BF_One,
+				// Opaque blending to DistanceFieldIrradiance
+				CW_RGBA, BO_Add, BF_One, BF_Zero, BO_Add, BF_One, BF_Zero>::GetRHI());
 		}
-		else if (bVisualizeGlobalIllumination && bUseDistanceFieldGI)
+		else
 		{
-			TShaderMapRef<TDistanceFieldAOUpsamplePS<AOUpsample_OutputIrradiance> > PixelShader(View.ShaderMap);
-			static FGlobalBoundShaderState BoundShaderState;
-			SetGlobalBoundShaderState(RHICmdList, View.GetFeatureLevel(), BoundShaderState, GFilterVertexDeclaration.VertexDeclarationRHI, *VertexShader, *PixelShader);
-			PixelShader->SetParameters(RHICmdList, View, DistanceFieldAOBentNormal, DistanceFieldIrradiance);
+			RHICmdList.SetBlendState(TStaticBlendState<>::GetRHI());
 		}
-		else
+
+		TShaderMapRef<FPostProcessVS> VertexShader(View.ShaderMap);
+
+		if (bModulateSceneColor)
 		{
-			if (bUseDistanceFieldGI)
+			if (bVisualizeAmbientOcclusion)
 			{
-				TShaderMapRef<TDistanceFieldAOUpsamplePS<AOUpsample_OutputBentNormalAndIrradiance> > PixelShader(View.ShaderMap);
+				TShaderMapRef<TDistanceFieldAOUpsamplePS<AOUpsample_OutputAO, true> > PixelShader(View.ShaderMap);
+				static FGlobalBoundShaderState BoundShaderState;
+				SetGlobalBoundShaderState(RHICmdList, View.GetFeatureLevel(), BoundShaderState, GFilterVertexDeclaration.VertexDeclarationRHI, *VertexShader, *PixelShader);
+				PixelShader->SetParameters(RHICmdList, View, DistanceFieldAOBentNormal, DistanceFieldIrradiance);
+			}
+			else if (bVisualizeGlobalIllumination && bUseDistanceFieldGI)
+			{
+				TShaderMapRef<TDistanceFieldAOUpsamplePS<AOUpsample_OutputIrradiance, true> > PixelShader(View.ShaderMap);
 				static FGlobalBoundShaderState BoundShaderState;
 				SetGlobalBoundShaderState(RHICmdList, View.GetFeatureLevel(), BoundShaderState, GFilterVertexDeclaration.VertexDeclarationRHI, *VertexShader, *PixelShader);
 				PixelShader->SetParameters(RHICmdList, View, DistanceFieldAOBentNormal, DistanceFieldIrradiance);
 			}
 			else
 			{
-				TShaderMapRef<TDistanceFieldAOUpsamplePS<AOUpsample_OutputBentNormal> > PixelShader(View.ShaderMap);
+				if (bUseDistanceFieldGI)
+				{
+					TShaderMapRef<TDistanceFieldAOUpsamplePS<AOUpsample_OutputBentNormalAndIrradiance, true> > PixelShader(View.ShaderMap);
+					static FGlobalBoundShaderState BoundShaderState;
+					SetGlobalBoundShaderState(RHICmdList, View.GetFeatureLevel(), BoundShaderState, GFilterVertexDeclaration.VertexDeclarationRHI, *VertexShader, *PixelShader);
+					PixelShader->SetParameters(RHICmdList, View, DistanceFieldAOBentNormal, DistanceFieldIrradiance);
+				}
+				else
+				{
+					TShaderMapRef<TDistanceFieldAOUpsamplePS<AOUpsample_OutputBentNormal, true> > PixelShader(View.ShaderMap);
+					static FGlobalBoundShaderState BoundShaderState;
+					SetGlobalBoundShaderState(RHICmdList, View.GetFeatureLevel(), BoundShaderState, GFilterVertexDeclaration.VertexDeclarationRHI, *VertexShader, *PixelShader);
+					PixelShader->SetParameters(RHICmdList, View, DistanceFieldAOBentNormal, DistanceFieldIrradiance);
+				}
+			}
+		}
+		else
+		{
+			if (bVisualizeAmbientOcclusion)
+			{
+				TShaderMapRef<TDistanceFieldAOUpsamplePS<AOUpsample_OutputAO, false> > PixelShader(View.ShaderMap);
 				static FGlobalBoundShaderState BoundShaderState;
 				SetGlobalBoundShaderState(RHICmdList, View.GetFeatureLevel(), BoundShaderState, GFilterVertexDeclaration.VertexDeclarationRHI, *VertexShader, *PixelShader);
 				PixelShader->SetParameters(RHICmdList, View, DistanceFieldAOBentNormal, DistanceFieldIrradiance);
 			}
+			else if (bVisualizeGlobalIllumination && bUseDistanceFieldGI)
+			{
+				TShaderMapRef<TDistanceFieldAOUpsamplePS<AOUpsample_OutputIrradiance, false> > PixelShader(View.ShaderMap);
+				static FGlobalBoundShaderState BoundShaderState;
+				SetGlobalBoundShaderState(RHICmdList, View.GetFeatureLevel(), BoundShaderState, GFilterVertexDeclaration.VertexDeclarationRHI, *VertexShader, *PixelShader);
+				PixelShader->SetParameters(RHICmdList, View, DistanceFieldAOBentNormal, DistanceFieldIrradiance);
+			}
+			else
+			{
+				if (bUseDistanceFieldGI)
+				{
+					TShaderMapRef<TDistanceFieldAOUpsamplePS<AOUpsample_OutputBentNormalAndIrradiance, false> > PixelShader(View.ShaderMap);
+					static FGlobalBoundShaderState BoundShaderState;
+					SetGlobalBoundShaderState(RHICmdList, View.GetFeatureLevel(), BoundShaderState, GFilterVertexDeclaration.VertexDeclarationRHI, *VertexShader, *PixelShader);
+					PixelShader->SetParameters(RHICmdList, View, DistanceFieldAOBentNormal, DistanceFieldIrradiance);
+				}
+				else
+				{
+					TShaderMapRef<TDistanceFieldAOUpsamplePS<AOUpsample_OutputBentNormal, false> > PixelShader(View.ShaderMap);
+					static FGlobalBoundShaderState BoundShaderState;
+					SetGlobalBoundShaderState(RHICmdList, View.GetFeatureLevel(), BoundShaderState, GFilterVertexDeclaration.VertexDeclarationRHI, *VertexShader, *PixelShader);
+					PixelShader->SetParameters(RHICmdList, View, DistanceFieldAOBentNormal, DistanceFieldIrradiance);
+				}
+			}
 		}
 
 		DrawRectangle( 
diff --git a/Engine/Source/Runtime/Renderer/Private/DistanceFieldLightingPost.h b/Engine/Source/Runtime/Renderer/Private/DistanceFieldLightingPost.h
--- a/Engine/Source/Runtime/Renderer/Private/DistanceFieldLightingPost.h
+++ b/Engine/Source/Runtime/Renderer/Private/DistanceFieldLightingPost.h
@@ -14,6 +14,7 @@ extern void UpdateHistory(
 	const TCHAR* BentNormalHistoryRTName,
 	const TCHAR* IrradianceHistoryRTName,
 	IPooledRenderTarget* VelocityTexture,
+	FSceneRenderTargetItem& DistanceFieldNormal,
 	/** Contains last frame's history, if non-NULL.  This will be updated with the new frame's history. */
 	TRefCountPtr<IPooledRenderTarget>* BentNormalHistoryState,
 	TRefCountPtr<IPooledRenderTarget>* IrradianceHistoryState,
@@ -40,5 +41,6 @@ extern void UpsampleBentNormalAO(
 	const TArray<FViewInfo>& Views, 
 	TRefCountPtr<IPooledRenderTarget>& DistanceFieldAOBentNormal, 
 	TRefCountPtr<IPooledRenderTarget>& DistanceFieldIrradiance,
+	bool bModulateSceneColor,
 	bool bVisualizeAmbientOcclusion,
 	bool bVisualizeGlobalIllumination);
\ No newline at end of file
diff --git a/Engine/Source/Runtime/Renderer/Private/DistanceFieldScreenGridLighting.cpp b/Engine/Source/Runtime/Renderer/Private/DistanceFieldScreenGridLighting.cpp
--- a/Engine/Source/Runtime/Renderer/Private/DistanceFieldScreenGridLighting.cpp
+++ b/Engine/Source/Runtime/Renderer/Private/DistanceFieldScreenGridLighting.cpp
@@ -743,6 +743,7 @@ void PostProcessBentNormalAOScreenGrid(
 			TEXT("DistanceFieldAOHistory"),
 			TEXT("DistanceFieldIrradianceHistory"),
 			VelocityTexture,
+			DistanceFieldNormal,
 			BentNormalHistoryState,
 			IrradianceHistoryState,
 			DistanceFieldAOBentNormal, 
diff --git a/Engine/Source/Runtime/Renderer/Private/DistanceFieldSurfaceCacheLighting.cpp b/Engine/Source/Runtime/Renderer/Private/DistanceFieldSurfaceCacheLighting.cpp
--- a/Engine/Source/Runtime/Renderer/Private/DistanceFieldSurfaceCacheLighting.cpp
+++ b/Engine/Source/Runtime/Renderer/Private/DistanceFieldSurfaceCacheLighting.cpp
@@ -31,6 +31,14 @@ FAutoConsoleVariableRef CVarDistanceFieldAO(
 	ECVF_Scalability | ECVF_RenderThreadSafe
 	);
 
+int32 GDistanceFieldAOApplyToStaticIndirect = 0;
+FAutoConsoleVariableRef CVarDistanceFieldAOApplyToStaticIndirect(
+	TEXT("r.AOApplyToStaticIndirect"),
+	GDistanceFieldAOApplyToStaticIndirect,
+	TEXT("Whether to apply DFAO as indirect shadowing even to static indirect sources (lightmaps + stationary skylight + reflection captures)"),
+	ECVF_Scalability | ECVF_RenderThreadSafe
+	);
+
 bool IsDistanceFieldGIAllowed(const FViewInfo& View)
 {
 	return DoesPlatformSupportDistanceFieldGI(View.GetShaderPlatform())
@@ -2742,7 +2750,8 @@ bool FDeferredShadingSceneRenderer::ShouldPrepareForDistanceFieldAO() const
 		&& ((ShouldRenderDynamicSkyLight(Scene, ViewFamily) && Scene->SkyLight->bCastShadows && ViewFamily.EngineShowFlags.DistanceFieldAO)
 			|| ViewFamily.EngineShowFlags.VisualizeMeshDistanceFields
 			|| ViewFamily.EngineShowFlags.VisualizeDistanceFieldAO
-			|| ViewFamily.EngineShowFlags.VisualizeDistanceFieldGI);
+			|| ViewFamily.EngineShowFlags.VisualizeDistanceFieldGI
+			|| (GDistanceFieldAOApplyToStaticIndirect && ViewFamily.EngineShowFlags.DistanceFieldAO));
 }
 
 bool FDeferredShadingSceneRenderer::ShouldPrepareDistanceFields() const
@@ -2893,12 +2902,12 @@ void RenderDistanceFieldAOSurfaceCache(
 			// Create new records which haven't been shaded yet for shading points which don't have a valid interpolation from existing records
 			{
 		        FIntPoint DownsampledViewSize = FIntPoint::DivideAndRoundUp(View.ViewRect.Size(), DestLevelDownsampleFactor);
-				        uint32 GroupSizeX = FMath::DivideAndRoundUp(DownsampledViewSize.X, GDistanceFieldAOTileSizeX);
-				        uint32 GroupSizeY = FMath::DivideAndRoundUp(DownsampledViewSize.Y, GDistanceFieldAOTileSizeY);
+				uint32 GroupSizeX = FMath::DivideAndRoundUp(DownsampledViewSize.X, GDistanceFieldAOTileSizeX);
+				uint32 GroupSizeY = FMath::DivideAndRoundUp(DownsampledViewSize.Y, GDistanceFieldAOTileSizeY);
         
 		        TShaderMapRef<FPopulateCacheCS> ComputeShader(View.ShaderMap);
         
-				        RHICmdList.SetComputeShader(ComputeShader->GetComputeShader());
+				RHICmdList.SetComputeShader(ComputeShader->GetComputeShader());
 		        ComputeShader->SetParameters(RHICmdList, View, DistanceFieldAOBentNormalSplat->GetRenderTargetItem(), DistanceFieldNormal->GetRenderTargetItem(), DestLevelDownsampleFactor, DepthLevel, TileListGroupSize, Parameters);
 				DispatchComputeShader(RHICmdList, *ComputeShader, GroupSizeX, GroupSizeY, 1);
 
@@ -2998,18 +3007,31 @@ void RenderDistanceFieldAOSurfaceCache(
 
 	GRenderTargetPool.VisualizeTexture.SetCheckPoint(RHICmdList, BentNormalAccumulation);
 
-			// Post process the AO to cover over artifacts
-
+	// Post process the AO to cover over artifacts
 	PostProcessBentNormalAOSurfaceCache(
-				RHICmdList, 
-				Parameters, 
-				View, 
-	VelocityTexture, 
-			BentNormalAccumulation->GetRenderTargetItem(), 
-			IrradianceAccumulation, 
-			DistanceFieldNormal->GetRenderTargetItem(), 
-	OutDynamicBentNormalAO, 
-	OutDynamicIrradiance);
+		RHICmdList, 
+		Parameters, 
+		View, 
+		VelocityTexture, 
+		BentNormalAccumulation->GetRenderTargetItem(), 
+		IrradianceAccumulation, 
+		DistanceFieldNormal->GetRenderTargetItem(), 
+		OutDynamicBentNormalAO, 
+		OutDynamicIrradiance);
+}
+
+void FDeferredShadingSceneRenderer::RenderDFAOAsIndirectShadowing(
+	FRHICommandListImmediate& RHICmdList,
+	const TRefCountPtr<IPooledRenderTarget>& VelocityTexture, 
+	TRefCountPtr<IPooledRenderTarget>& DynamicBentNormalAO)
+{
+	if (GDistanceFieldAOApplyToStaticIndirect && ShouldRenderDistanceFieldAO())
+	{
+		// Use the skylight's max distance if there is one, to be consistent with DFAO shadowing on the skylight
+		const float OcclusionMaxDistance = Scene->SkyLight && !Scene->SkyLight->bWantsStaticShadowing ? Scene->SkyLight->OcclusionMaxDistance : GDefaultDFAOMaxOcclusionDistance;
+		TRefCountPtr<IPooledRenderTarget> DummyOutput;
+		RenderDistanceFieldLighting(RHICmdList, FDistanceFieldAOParameters(OcclusionMaxDistance), VelocityTexture, DynamicBentNormalAO, DummyOutput, true, false, false);
+	}
 }
 
 bool FDeferredShadingSceneRenderer::RenderDistanceFieldLighting(
@@ -3018,6 +3040,7 @@ bool FDeferredShadingSceneRenderer::RenderDistanceFieldLighting(
 	const TRefCountPtr<IPooledRenderTarget>& VelocityTexture,
 	TRefCountPtr<IPooledRenderTarget>& OutDynamicBentNormalAO, 
 	TRefCountPtr<IPooledRenderTarget>& OutDynamicIrradiance,
+	bool bModulateToSceneColor,
 	bool bVisualizeAmbientOcclusion,
 	bool bVisualizeGlobalIllumination)
 {
@@ -3148,16 +3171,32 @@ bool FDeferredShadingSceneRenderer::RenderDistanceFieldLighting(
 					GRenderTargetPool.FindFreeElement(RHICmdList, Desc, OutDynamicIrradiance, TEXT("DynamicIrradiance"));
 				}
 
-				FTextureRHIParamRef RenderTargets[2] =
+				FTextureRHIParamRef RenderTargets[3] =
 				{
 					OutDynamicBentNormalAO->GetRenderTargetItem().TargetableTexture,
-					bUseDistanceFieldGI ? OutDynamicIrradiance->GetRenderTargetItem().TargetableTexture : NULL
+					NULL,
+					NULL
 				};
-				SetRenderTargets(RHICmdList, ARRAY_COUNT(RenderTargets) - (bUseDistanceFieldGI ? 0 : 1), RenderTargets, FTextureRHIParamRef(), 0, NULL);
+
+				int32 NumRenderTargets = 1;
+
+				if (bModulateToSceneColor)
+				{
+					RenderTargets[NumRenderTargets] = SceneContext.GetSceneColorSurface();
+					NumRenderTargets++;
+				}
+
+				if (bUseDistanceFieldGI)
+				{
+					RenderTargets[NumRenderTargets] = OutDynamicIrradiance->GetRenderTargetItem().TargetableTexture;
+					NumRenderTargets++;
+				}
+
+				SetRenderTargets(RHICmdList, NumRenderTargets, RenderTargets, FTextureRHIParamRef(), 0, NULL);
 			}
 
 			// Upsample to full resolution, write to output
-			UpsampleBentNormalAO(RHICmdList, Views, BentNormalOutput, IrradianceOutput, bVisualizeAmbientOcclusion, bVisualizeGlobalIllumination);
+			UpsampleBentNormalAO(RHICmdList, Views, BentNormalOutput, IrradianceOutput, bModulateToSceneColor, bVisualizeAmbientOcclusion, bVisualizeGlobalIllumination);
 
 			if (!bVisualizeAmbientOcclusion && !bVisualizeGlobalIllumination)
 			{
@@ -3273,8 +3312,7 @@ IMPLEMENT_SKYLIGHT_PS_TYPE(false, false)
 
 bool FDeferredShadingSceneRenderer::ShouldRenderDistanceFieldAO() const
 {
-	return Scene->SkyLight->bCastShadows
-		&& ViewFamily.EngineShowFlags.DistanceFieldAO 
+	return ViewFamily.EngineShowFlags.DistanceFieldAO 
 		&& !ViewFamily.EngineShowFlags.VisualizeDistanceFieldAO
 		&& !ViewFamily.EngineShowFlags.VisualizeDistanceFieldGI
 		&& !ViewFamily.EngineShowFlags.VisualizeMeshDistanceFields;
@@ -3291,10 +3329,14 @@ void FDeferredShadingSceneRenderer::RenderDynamicSkyLighting(FRHICommandListImme
 		FDistanceFieldAOParameters Parameters(Scene->SkyLight->OcclusionMaxDistance, Scene->SkyLight->Contrast);
 		TRefCountPtr<IPooledRenderTarget> DynamicIrradiance;
 
-		if (ShouldRenderDistanceFieldAO() && ViewFamily.EngineShowFlags.AmbientOcclusion)
+		if (Scene->SkyLight->bCastShadows 
+			&& !GDistanceFieldAOApplyToStaticIndirect
+			&& ShouldRenderDistanceFieldAO() 
+			&& ViewFamily.EngineShowFlags.AmbientOcclusion)
 		{
-			bApplyShadowing = RenderDistanceFieldLighting(RHICmdList, Parameters, VelocityTexture, DynamicBentNormalAO, DynamicIrradiance, false, false);
+			bApplyShadowing = RenderDistanceFieldLighting(RHICmdList, Parameters, VelocityTexture, DynamicBentNormalAO, DynamicIrradiance, false, false, false);
 		}
+
 		FSceneRenderTargets& SceneContext = FSceneRenderTargets::Get(RHICmdList);
 
 		SceneContext.BeginRenderingSceneColor(RHICmdList, ESimpleRenderTargetMode::EExistingColorAndDepth, FExclusiveDepthStencil::DepthRead_StencilRead);
diff --git a/Engine/Source/Runtime/Renderer/Private/LightRendering.h b/Engine/Source/Runtime/Renderer/Private/LightRendering.h
--- a/Engine/Source/Runtime/Renderer/Private/LightRendering.h
+++ b/Engine/Source/Runtime/Renderer/Private/LightRendering.h
@@ -92,6 +92,9 @@ void SetDeferredLightParameters(
 		DeferredLightUniformsValue.LightColor *= LightSceneInfo->Proxy->GetIndirectLightingScale();
 	}
 
+	const ELightComponentType LightType = (ELightComponentType)LightSceneInfo->Proxy->GetLightType();
+
+	if (LightType == LightType_Point || LightType == LightType_Spot)
 	{
 		// Distance fade
 		FSphere Bounds = LightSceneInfo->Proxy->GetBoundingSphere();
diff --git a/Engine/Source/Runtime/Renderer/Private/PostProcess/PostProcessCircleDOF.cpp b/Engine/Source/Runtime/Renderer/Private/PostProcess/PostProcessCircleDOF.cpp
--- a/Engine/Source/Runtime/Renderer/Private/PostProcess/PostProcessCircleDOF.cpp
+++ b/Engine/Source/Runtime/Renderer/Private/PostProcess/PostProcessCircleDOF.cpp
@@ -584,7 +584,9 @@ public:
 	IMPLEMENT_SHADER_TYPE2(FPostProcessCircleDOFPS##A##B, SF_Pixel);
 
 	VARIATION1(0,0)			VARIATION1(1,0)
-	VARIATION1(0,1)			VARIATION1(1,1)
+	VARIATION1(0,1)	        VARIATION1(1,1)
+	VARIATION1(0,2)         VARIATION1(1,2)
+
 
 #undef VARIATION1
 
@@ -663,20 +665,30 @@ void FRCPassPostProcessCircleDOF::Process(FRenderingCompositePassContext& Contex
 	check(CVar);
 	int32 DOFQualityCVarValue = CVar->GetValueOnRenderThread();
 	
-	// 0:normal / 1:slow but very high quality
-	uint32 Quality = DOFQualityCVarValue >= 3;
+	
 
 	FShader* VertexShader = 0;
 
 	if (bNearBlurEnabled)
 	{
-		if(Quality) VertexShader = SetShaderTempl<1, 1>(Context);
-		  else		VertexShader = SetShaderTempl<1, 0>(Context);
+		switch (DOFQualityCVarValue)
+		{
+			case 0:  VertexShader = SetShaderTempl<1, 0>(Context); break;
+			case 3:  VertexShader = SetShaderTempl<1, 1>(Context); break;
+			case 4:  VertexShader = SetShaderTempl<1, 2>(Context); break;
+			default: VertexShader = SetShaderTempl<1, 0>(Context);
+		}
 	}
 	else
 	{
-		if(Quality)	VertexShader = SetShaderTempl<0, 1>(Context);
-		  else		VertexShader = SetShaderTempl<0, 0>(Context);
+		switch (DOFQualityCVarValue)
+		{
+			case 0:  VertexShader = SetShaderTempl<0, 0>(Context);
+			case 3:  VertexShader = SetShaderTempl<0, 1>(Context); break;
+			case 4:  VertexShader = SetShaderTempl<0, 2>(Context); break;	
+			default: VertexShader = SetShaderTempl<0, 0>(Context);
+		}
+
 	}
 
 	DrawPostProcessPass(
diff --git a/Engine/Source/Runtime/Renderer/Private/PostProcess/PostProcessMotionBlur.cpp b/Engine/Source/Runtime/Renderer/Private/PostProcess/PostProcessMotionBlur.cpp
--- a/Engine/Source/Runtime/Renderer/Private/PostProcess/PostProcessMotionBlur.cpp
+++ b/Engine/Source/Runtime/Renderer/Private/PostProcess/PostProcessMotionBlur.cpp
@@ -352,8 +352,6 @@ public:
 			SetShaderValue(Context.RHICmdList, ShaderRHI, PrevViewProjMatrix, Context.View.ViewMatrices.GetViewRotationProjMatrix());
 		}
 
-		TRefCountPtr<IPooledRenderTarget> InputPooledElement = Context.Pass->GetInput(ePId_Input0)->GetOutput()->RequestInput();
-
 		// to mask out samples from outside of the view
 		{
 			FIntPoint BufferSize = FSceneRenderTargets::Get(Context.RHICmdList).GetBufferSizeXY();
@@ -423,6 +421,94 @@ VARIATION1(0)			VARIATION1(1)			VARIATION1(2)			VARIATION1(3)			VARIATION1(4)
 #undef VARIATION1
 
 
+class FPostProcessVisualizeMotionBlurPS : public FGlobalShader
+{
+	DECLARE_SHADER_TYPE(FPostProcessVisualizeMotionBlurPS, Global);
+
+	static bool ShouldCache(EShaderPlatform Platform)
+	{
+		return IsFeatureLevelSupported(Platform, ERHIFeatureLevel::SM4);
+	}
+
+	/** Default constructor. */
+	FPostProcessVisualizeMotionBlurPS() {}
+
+public:
+	FPostProcessPassParameters PostprocessParameter;
+	FDeferredPixelShaderParameters DeferredParameters;
+	FShaderParameter PrevViewProjMatrix;
+
+	/** Initialization constructor. */
+	FPostProcessVisualizeMotionBlurPS(const ShaderMetaType::CompiledShaderInitializerType& Initializer)
+		: FGlobalShader(Initializer)
+	{
+		PostprocessParameter.Bind(Initializer.ParameterMap);
+		DeferredParameters.Bind(Initializer.ParameterMap);
+		PrevViewProjMatrix.Bind(Initializer.ParameterMap, TEXT("PrevViewProjMatrix"));
+	}
+
+	// FShader interface.
+	virtual bool Serialize(FArchive& Ar) override
+	{
+		bool bShaderHasOutdatedParameters = FGlobalShader::Serialize(Ar);
+		Ar << PostprocessParameter << DeferredParameters << PrevViewProjMatrix;
+		return bShaderHasOutdatedParameters;
+	}
+
+	void SetParameters(const FRenderingCompositePassContext& Context)
+	{
+		const FPixelShaderRHIParamRef ShaderRHI = GetPixelShader();
+
+		FGlobalShader::SetParameters(Context.RHICmdList, ShaderRHI, Context.View);
+
+		DeferredParameters.Set(Context.RHICmdList, ShaderRHI, Context.View);
+
+		SetUniformBufferParameter(Context.RHICmdList, ShaderRHI, GetUniformBufferParameter<FCameraMotionParameters>(), CreateCameraMotionParametersUniformBuffer(Context.View));
+
+		{
+			bool bFiltered = false;
+
+#if !(UE_BUILD_SHIPPING || UE_BUILD_TEST)
+			bFiltered = CVarMotionBlurFiltering.GetValueOnRenderThread() != 0;
+#endif // !(UE_BUILD_SHIPPING || UE_BUILD_TEST)
+
+			if(bFiltered)
+			{
+				PostprocessParameter.SetPS(ShaderRHI, Context, TStaticSamplerState<SF_Bilinear,AM_Border,AM_Border,AM_Clamp>::GetRHI());
+			}
+			else
+			{
+				PostprocessParameter.SetPS(ShaderRHI, Context, TStaticSamplerState<SF_Point,AM_Border,AM_Border,AM_Clamp>::GetRHI());
+			}
+		}
+		
+		if( Context.View.Family->EngineShowFlags.CameraInterpolation )
+		{
+			// Instead of finding the world space position of the current pixel, calculate the world space position offset by the camera position, 
+			// then translate by the difference between last frame's camera position and this frame's camera position,
+			// then apply the rest of the transforms.  This effectively avoids precision issues near the extents of large levels whose world space position is very large.
+			FVector ViewOriginDelta = Context.View.ViewMatrices.ViewOrigin - Context.View.PrevViewMatrices.ViewOrigin;
+			SetShaderValue(Context.RHICmdList, ShaderRHI, PrevViewProjMatrix, FTranslationMatrix(ViewOriginDelta) * Context.View.PrevViewRotationProjMatrix);
+		}
+		else
+		{
+			SetShaderValue(Context.RHICmdList, ShaderRHI, PrevViewProjMatrix, Context.View.ViewMatrices.GetViewRotationProjMatrix());
+		}
+	}
+
+	static const TCHAR* GetSourceFilename()
+	{
+		return TEXT("PostProcessMotionBlur");
+	}
+
+	static const TCHAR* GetFunctionName()
+	{
+		return TEXT("VisualizeMotionBlurPS");
+	}
+};
+
+IMPLEMENT_SHADER_TYPE3(FPostProcessVisualizeMotionBlurPS, SF_Pixel);
+
 
 // @param Quality 0: visualize, 1:low, 2:medium, 3:high, 4:very high
 template <uint32 Quality>
@@ -572,20 +658,13 @@ public:
 		return bShaderHasOutdatedParameters;
 	}
 
-	void SetParameters(const FRenderingCompositePassContext& Context, bool bBilinear)
+	void SetParameters(const FRenderingCompositePassContext& Context)
 	{
 		const FPixelShaderRHIParamRef ShaderRHI = GetPixelShader();
 
 		FGlobalShader::SetParameters(Context.RHICmdList, ShaderRHI, Context.View);
 
-		if(bBilinear)
-		{
-			PostprocessParameter.SetPS(ShaderRHI, Context, TStaticSamplerState<SF_Bilinear, AM_Border, AM_Border, AM_Clamp>::GetRHI());
-		}
-		else
-		{
-			PostprocessParameter.SetPS(ShaderRHI, Context, TStaticSamplerState<SF_Point, AM_Border, AM_Border, AM_Clamp>::GetRHI());
-		}
+		PostprocessParameter.SetPS(ShaderRHI, Context, TStaticSamplerState<SF_Bilinear, AM_Border, AM_Border, AM_Clamp>::GetRHI());
 	}
 };
 
@@ -640,11 +719,8 @@ void FRCPassPostProcessMotionBlurRecombine::Process(FRenderingCompositePassConte
 	
 	SetGlobalBoundShaderState(Context.RHICmdList, Context.GetFeatureLevel(), BoundShaderState, GFilterVertexDeclaration.VertexDeclarationRHI, *VertexShader, *PixelShader);
 
-	// with point filtering we can better debug in Visualize MotionBlur
-	const bool bBilinear = !View.Family->EngineShowFlags.VisualizeMotionBlur;
-
 	VertexShader->SetParameters(Context);
-	PixelShader->SetParameters(Context, bBilinear);
+	PixelShader->SetParameters(Context);
 
 	DrawPostProcessPass(
 		Context.RHICmdList,
@@ -1462,12 +1538,17 @@ void FRCPassPostProcessVisualizeMotionBlur::Process(FRenderingCompositePassConte
 	Context.RHICmdList.SetRasterizerState(TStaticRasterizerState<>::GetRHI());
 	Context.RHICmdList.SetDepthStencilState(TStaticDepthStencilState<false, CF_Always>::GetRHI());
 
-	// Quality 0: visualize
-	SetMotionBlurShaderTempl<0>(Context);
-
-	// Draw a quad mapping scene color to the view's render target
 	TShaderMapRef<FPostProcessVS> VertexShader(Context.GetShaderMap());
+	TShaderMapRef<FPostProcessVisualizeMotionBlurPS> PixelShader(Context.GetShaderMap());
+
+	static FGlobalBoundShaderState BoundShaderState;
+	
+	SetGlobalBoundShaderState(Context.RHICmdList, Context.GetFeatureLevel(), BoundShaderState, GFilterVertexDeclaration.VertexDeclarationRHI, *VertexShader, *PixelShader);
+
+	VertexShader->SetParameters(Context);
+	PixelShader->SetParameters(Context);
 
+	// Draw a quad mapping scene color to the view's render target
 	DrawRectangle(
 		Context.RHICmdList,
 		0, 0,
@@ -1503,7 +1584,7 @@ void FRCPassPostProcessVisualizeMotionBlur::Process(FRenderingCompositePassConte
 	FCanvas Canvas(&TempRenderTarget, NULL, ViewFamily.CurrentRealTime, ViewFamily.CurrentWorldTime, ViewFamily.DeltaWorldTime, Context.GetFeatureLevel());
 
 	float X = 20;
-	float Y = 8;
+	float Y = 38;
 	const float YStep = 14;
 	const float ColumnWidth = 200;
 
diff --git a/Engine/Source/Runtime/Renderer/Private/PostProcess/PostProcessMotionBlur.h b/Engine/Source/Runtime/Renderer/Private/PostProcess/PostProcessMotionBlur.h
--- a/Engine/Source/Runtime/Renderer/Private/PostProcess/PostProcessMotionBlur.h
+++ b/Engine/Source/Runtime/Renderer/Private/PostProcess/PostProcessMotionBlur.h
@@ -59,9 +59,9 @@ public:
 };
 
 
-// ePId_Input0: Half res SceneColor with depth in alpha
-// ePId_Input1: optional quarter res blurred screen space velocity
-// ePId_Input2: optional half res screen space masked velocity
+// ePId_Input0: Full Res Scene Color
+// ePId_Input1: Full Res Scene Depth
+// ePId_Input2: Full Res velocity input
 // derives from TRenderingCompositePassBase<InputCount, OutputCount> 
 class FRCPassPostProcessVisualizeMotionBlur : public TRenderingCompositePassBase<3, 1>
 {
diff --git a/Engine/Source/Runtime/Renderer/Private/PostProcess/PostProcessing.cpp b/Engine/Source/Runtime/Renderer/Private/PostProcess/PostProcessing.cpp
--- a/Engine/Source/Runtime/Renderer/Private/PostProcess/PostProcessing.cpp
+++ b/Engine/Source/Runtime/Renderer/Private/PostProcess/PostProcessing.cpp
@@ -1214,8 +1214,9 @@ void FPostProcessing::Process(FRHICommandListImmediate& RHICmdList, FViewInfo& V
 
 		bool bVisualizeHDR = View.Family->EngineShowFlags.VisualizeHDR && FeatureLevel >= ERHIFeatureLevel::SM5;
 		bool bVisualizeBloom = View.Family->EngineShowFlags.VisualizeBloom && FeatureLevel >= ERHIFeatureLevel::SM4;
+		bool bVisualizeMotionBlur = View.Family->EngineShowFlags.VisualizeMotionBlur && FeatureLevel >= ERHIFeatureLevel::SM4;
 
-		if(bVisualizeHDR || bVisualizeBloom)
+		if(bVisualizeHDR || bVisualizeBloom || bVisualizeMotionBlur)
 		{
 			bAllowTonemapper = false;
 		}
@@ -1343,20 +1344,19 @@ void FPostProcessing::Process(FRHICommandListImmediate& RHICmdList, FViewInfo& V
 				}
 			}
 
-			if(IsMotionBlurEnabled(View) && VelocityInput.IsValid())
+			if(IsMotionBlurEnabled(View) && VelocityInput.IsValid() && !bVisualizeMotionBlur)
 			{
 				// Motion blur
 
 				if( CVarMotionBlurNew.GetValueOnRenderThread() && FeatureLevel >= ERHIFeatureLevel::SM5 )
 				{
 					FRenderingCompositeOutputRef MaxTileVelocity;
-					FRenderingCompositeOutputRef SceneDepth( Context.SceneDepth );
 
 					{
 						check(!VelocityFlattenPass);
 						VelocityFlattenPass = Context.Graph.RegisterPass( new(FMemStack::Get()) FRCPassPostProcessVelocityFlatten() );
 						VelocityFlattenPass->SetInput( ePId_Input0, VelocityInput );
-						VelocityFlattenPass->SetInput( ePId_Input1, SceneDepth );
+						VelocityFlattenPass->SetInput( ePId_Input1, Context.SceneDepth );
 
 						VelocityInput	= FRenderingCompositeOutputRef( VelocityFlattenPass, ePId_Output0 );
 						MaxTileVelocity	= FRenderingCompositeOutputRef( VelocityFlattenPass, ePId_Output1 );
@@ -1387,7 +1387,7 @@ void FPostProcessing::Process(FRHICommandListImmediate& RHICmdList, FViewInfo& V
 					{
 						FRenderingCompositePass* MotionBlurPass = Context.Graph.RegisterPass( new(FMemStack::Get()) FRCPassPostProcessMotionBlurNew( GetMotionBlurQualityFromCVar(), bTwoPass ? 0 : -1 ) );
 						MotionBlurPass->SetInput( ePId_Input0, Context.FinalOutput );
-						MotionBlurPass->SetInput( ePId_Input1, SceneDepth );
+						MotionBlurPass->SetInput( ePId_Input1, Context.SceneDepth );
 						MotionBlurPass->SetInput( ePId_Input2, VelocityInput );
 						MotionBlurPass->SetInput( ePId_Input3, MaxTileVelocity );
 					
@@ -1398,7 +1398,7 @@ void FPostProcessing::Process(FRHICommandListImmediate& RHICmdList, FViewInfo& V
 					{
 						FRenderingCompositePass* MotionBlurPass = Context.Graph.RegisterPass( new(FMemStack::Get()) FRCPassPostProcessMotionBlurNew( GetMotionBlurQualityFromCVar(), 1 ) );
 						MotionBlurPass->SetInput( ePId_Input0, Context.FinalOutput );
-						MotionBlurPass->SetInput( ePId_Input1, SceneDepth );
+						MotionBlurPass->SetInput( ePId_Input1, Context.SceneDepth );
 						MotionBlurPass->SetInput( ePId_Input2, VelocityInput );
 						MotionBlurPass->SetInput( ePId_Input3, MaxTileVelocity );
 					
@@ -1440,18 +1440,9 @@ void FPostProcessing::Process(FRHICommandListImmediate& RHICmdList, FViewInfo& V
 					// doing the actual motion blur sampling, alpha to mask out the blurred areas
 					FRenderingCompositePass* MotionBlurPass;
 
-					if(View.Family->EngineShowFlags.VisualizeMotionBlur)
-					{
-						MotionBlurPass = Context.Graph.RegisterPass(new(FMemStack::Get()) FRCPassPostProcessVisualizeMotionBlur());
-						bAllowTonemapper = false;
-					}
-					else
-					{
-						int32 MotionBlurQuality = GetMotionBlurQualityFromCVar();
-
-						MotionBlurPass = Context.Graph.RegisterPass(new(FMemStack::Get()) FRCPassPostProcessMotionBlur(MotionBlurQuality));
-					}
+					int32 MotionBlurQuality = GetMotionBlurQualityFromCVar();
 
+					MotionBlurPass = Context.Graph.RegisterPass(new(FMemStack::Get()) FRCPassPostProcessMotionBlur(MotionBlurQuality));
 					MotionBlurPass->SetInput(ePId_Input0, MotionBlurColorDepth);
 
 					if(VelocityInput.IsValid())
@@ -1464,11 +1455,21 @@ void FPostProcessing::Process(FRHICommandListImmediate& RHICmdList, FViewInfo& V
 
 					FRenderingCompositePass* MotionBlurRecombinePass = Context.Graph.RegisterPass(new(FMemStack::Get()) FRCPassPostProcessMotionBlurRecombine());
 					MotionBlurRecombinePass->SetInput(ePId_Input0, Context.FinalOutput);
-					MotionBlurRecombinePass->SetInput(ePId_Input1, FRenderingCompositeOutputRef(MotionBlurPass));
+					MotionBlurRecombinePass->SetInput(ePId_Input1, MotionBlurPass);
 					Context.FinalOutput = FRenderingCompositeOutputRef(MotionBlurRecombinePass);
 				}
 			}
 
+			if(VelocityInput.IsValid() && bVisualizeMotionBlur)
+			{
+				auto VisualizePass = Context.Graph.RegisterPass(new(FMemStack::Get()) FRCPassPostProcessVisualizeMotionBlur());
+				VisualizePass->SetInput(ePId_Input0, Context.FinalOutput);
+				VisualizePass->SetInput(ePId_Input1, Context.SceneDepth);
+				VisualizePass->SetInput(ePId_Input2, VelocityInput);
+
+				Context.FinalOutput = FRenderingCompositeOutputRef(VisualizePass);
+			}
+
 			if(bVisualizeBloom)
 			{
 				AddVisualizeBloomSetup(Context);
diff --git a/Engine/Source/Runtime/Renderer/Private/ReflectionEnvironment.cpp b/Engine/Source/Runtime/Renderer/Private/ReflectionEnvironment.cpp
--- a/Engine/Source/Runtime/Renderer/Private/ReflectionEnvironment.cpp
+++ b/Engine/Source/Runtime/Renderer/Private/ReflectionEnvironment.cpp
@@ -355,7 +355,7 @@ private:
 	FDistanceFieldAOSpecularOcclusionParameters SpecularOcclusionParameters;
 };
 
-template< uint32 bUseLightmaps, uint32 bHasSkyLight, uint32 bBoxCapturesOnly, uint32 bSphereCapturesOnly >
+template< uint32 bUseLightmaps, uint32 bHasSkyLight, uint32 bBoxCapturesOnly, uint32 bSphereCapturesOnly, uint32 bSupportDFAOIndirectOcclusion >
 class TReflectionEnvironmentTiledDeferredCS : public FReflectionEnvironmentTiledDeferredCS
 {
 	DECLARE_SHADER_TYPE(TReflectionEnvironmentTiledDeferredCS, Global);
@@ -374,36 +374,56 @@ public:
 		OutEnvironment.SetDefine(TEXT("HAS_SKYLIGHT"), bHasSkyLight);
 		OutEnvironment.SetDefine(TEXT("HAS_BOX_CAPTURES"), bBoxCapturesOnly);
 		OutEnvironment.SetDefine(TEXT("HAS_SPHERE_CAPTURES"), bSphereCapturesOnly);
+		OutEnvironment.SetDefine(TEXT("SUPPORT_DFAO_INDIRECT_OCCLUSION"), bSupportDFAOIndirectOcclusion);
 	}
 };
 
 // Typedef is necessary because the C preprocessor thinks the comma in the template parameter list is a comma in the macro parameter list.
-#define IMPLEMENT_REFLECTION_COMPUTESHADER_TYPE(A, B, C, D) \
-	typedef TReflectionEnvironmentTiledDeferredCS<A,B,C,D> TReflectionEnvironmentTiledDeferredCS##A##B##C##D; \
-	IMPLEMENT_SHADER_TYPE(template<>,TReflectionEnvironmentTiledDeferredCS##A##B##C##D,TEXT("ReflectionEnvironmentComputeShaders"),TEXT("ReflectionEnvironmentTiledDeferredMain"),SF_Compute)
-
-IMPLEMENT_REFLECTION_COMPUTESHADER_TYPE(0, 0, 0, 0);
-IMPLEMENT_REFLECTION_COMPUTESHADER_TYPE(0, 0, 0, 1);
-IMPLEMENT_REFLECTION_COMPUTESHADER_TYPE(0, 0, 1, 0);
-IMPLEMENT_REFLECTION_COMPUTESHADER_TYPE(0, 0, 1, 1);
-
-IMPLEMENT_REFLECTION_COMPUTESHADER_TYPE(0, 1, 0, 0);
-IMPLEMENT_REFLECTION_COMPUTESHADER_TYPE(0, 1, 0, 1);
-IMPLEMENT_REFLECTION_COMPUTESHADER_TYPE(0, 1, 1, 0);
-IMPLEMENT_REFLECTION_COMPUTESHADER_TYPE(0, 1, 1, 1);
-
-IMPLEMENT_REFLECTION_COMPUTESHADER_TYPE(1, 0, 0, 0);
-IMPLEMENT_REFLECTION_COMPUTESHADER_TYPE(1, 0, 0, 1);
-IMPLEMENT_REFLECTION_COMPUTESHADER_TYPE(1, 0, 1, 0);
-IMPLEMENT_REFLECTION_COMPUTESHADER_TYPE(1, 0, 1, 1);
-
-IMPLEMENT_REFLECTION_COMPUTESHADER_TYPE(1, 1, 0, 0);
-IMPLEMENT_REFLECTION_COMPUTESHADER_TYPE(1, 1, 0, 1);
-IMPLEMENT_REFLECTION_COMPUTESHADER_TYPE(1, 1, 1, 0);
-IMPLEMENT_REFLECTION_COMPUTESHADER_TYPE(1, 1, 1, 1);
-
-
-template< uint32 bSSR, uint32 bReflectionEnv, uint32 bSkylight >
+#define IMPLEMENT_REFLECTION_COMPUTESHADER_TYPE(A, B, C, D, E) \
+	typedef TReflectionEnvironmentTiledDeferredCS<A,B,C,D,E> TReflectionEnvironmentTiledDeferredCS##A##B##C##D##E; \
+	IMPLEMENT_SHADER_TYPE(template<>,TReflectionEnvironmentTiledDeferredCS##A##B##C##D##E,TEXT("ReflectionEnvironmentComputeShaders"),TEXT("ReflectionEnvironmentTiledDeferredMain"),SF_Compute)
+
+IMPLEMENT_REFLECTION_COMPUTESHADER_TYPE(0, 0, 0, 0, 0);
+IMPLEMENT_REFLECTION_COMPUTESHADER_TYPE(0, 0, 0, 1, 0);
+IMPLEMENT_REFLECTION_COMPUTESHADER_TYPE(0, 0, 1, 0, 0);
+IMPLEMENT_REFLECTION_COMPUTESHADER_TYPE(0, 0, 1, 1, 0);
+
+IMPLEMENT_REFLECTION_COMPUTESHADER_TYPE(0, 1, 0, 0, 0);
+IMPLEMENT_REFLECTION_COMPUTESHADER_TYPE(0, 1, 0, 1, 0);
+IMPLEMENT_REFLECTION_COMPUTESHADER_TYPE(0, 1, 1, 0, 0);
+IMPLEMENT_REFLECTION_COMPUTESHADER_TYPE(0, 1, 1, 1, 0);
+
+IMPLEMENT_REFLECTION_COMPUTESHADER_TYPE(1, 0, 0, 0, 0);
+IMPLEMENT_REFLECTION_COMPUTESHADER_TYPE(1, 0, 0, 1, 0);
+IMPLEMENT_REFLECTION_COMPUTESHADER_TYPE(1, 0, 1, 0, 0);
+IMPLEMENT_REFLECTION_COMPUTESHADER_TYPE(1, 0, 1, 1, 0);
+
+IMPLEMENT_REFLECTION_COMPUTESHADER_TYPE(1, 1, 0, 0, 0);
+IMPLEMENT_REFLECTION_COMPUTESHADER_TYPE(1, 1, 0, 1, 0);
+IMPLEMENT_REFLECTION_COMPUTESHADER_TYPE(1, 1, 1, 0, 0);
+IMPLEMENT_REFLECTION_COMPUTESHADER_TYPE(1, 1, 1, 1, 0);
+
+IMPLEMENT_REFLECTION_COMPUTESHADER_TYPE(0, 0, 0, 0, 1);
+IMPLEMENT_REFLECTION_COMPUTESHADER_TYPE(0, 0, 0, 1, 1);
+IMPLEMENT_REFLECTION_COMPUTESHADER_TYPE(0, 0, 1, 0, 1);
+IMPLEMENT_REFLECTION_COMPUTESHADER_TYPE(0, 0, 1, 1, 1);
+
+IMPLEMENT_REFLECTION_COMPUTESHADER_TYPE(0, 1, 0, 0, 1);
+IMPLEMENT_REFLECTION_COMPUTESHADER_TYPE(0, 1, 0, 1, 1);
+IMPLEMENT_REFLECTION_COMPUTESHADER_TYPE(0, 1, 1, 0, 1);
+IMPLEMENT_REFLECTION_COMPUTESHADER_TYPE(0, 1, 1, 1, 1);
+
+IMPLEMENT_REFLECTION_COMPUTESHADER_TYPE(1, 0, 0, 0, 1);
+IMPLEMENT_REFLECTION_COMPUTESHADER_TYPE(1, 0, 0, 1, 1);
+IMPLEMENT_REFLECTION_COMPUTESHADER_TYPE(1, 0, 1, 0, 1);
+IMPLEMENT_REFLECTION_COMPUTESHADER_TYPE(1, 0, 1, 1, 1);
+
+IMPLEMENT_REFLECTION_COMPUTESHADER_TYPE(1, 1, 0, 0, 1);
+IMPLEMENT_REFLECTION_COMPUTESHADER_TYPE(1, 1, 0, 1, 1);
+IMPLEMENT_REFLECTION_COMPUTESHADER_TYPE(1, 1, 1, 0, 1);
+IMPLEMENT_REFLECTION_COMPUTESHADER_TYPE(1, 1, 1, 1, 1);
+
+template< uint32 bSSR, uint32 bReflectionEnv, uint32 bSkylight, uint32 bSupportDFAOIndirectOcclusion >
 class FReflectionApplyPS : public FGlobalShader
 {
 	DECLARE_SHADER_TYPE(FReflectionApplyPS, Global);
@@ -420,6 +440,7 @@ public:
 		OutEnvironment.SetDefine(TEXT("APPLY_SSR"), bSSR);
 		OutEnvironment.SetDefine(TEXT("APPLY_REFLECTION_ENV"), bReflectionEnv);
 		OutEnvironment.SetDefine(TEXT("APPLY_SKYLIGHT"), bSkylight);
+		OutEnvironment.SetDefine(TEXT("SUPPORT_DFAO_INDIRECT_OCCLUSION"), bSupportDFAOIndirectOcclusion);
 	}
 
 	/** Default constructor. */
@@ -487,19 +508,27 @@ private:
 };
 
 // Typedef is necessary because the C preprocessor thinks the comma in the template parameter list is a comma in the macro parameter list.
-#define IMPLEMENT_REFLECTION_APPLY_PIXELSHADER_TYPE(A, B, C) \
-	typedef FReflectionApplyPS<A,B,C> FReflectionApplyPS##A##B##C; \
-	IMPLEMENT_SHADER_TYPE(template<>,FReflectionApplyPS##A##B##C,TEXT("ReflectionEnvironmentShaders"),TEXT("ReflectionApplyPS"),SF_Pixel);
-
-IMPLEMENT_REFLECTION_APPLY_PIXELSHADER_TYPE(0,0,0);
-IMPLEMENT_REFLECTION_APPLY_PIXELSHADER_TYPE(0,0,1);
-IMPLEMENT_REFLECTION_APPLY_PIXELSHADER_TYPE(0,1,0);
-IMPLEMENT_REFLECTION_APPLY_PIXELSHADER_TYPE(0,1,1);
-IMPLEMENT_REFLECTION_APPLY_PIXELSHADER_TYPE(1,0,0);
-IMPLEMENT_REFLECTION_APPLY_PIXELSHADER_TYPE(1,0,1);
-IMPLEMENT_REFLECTION_APPLY_PIXELSHADER_TYPE(1,1,0);
-IMPLEMENT_REFLECTION_APPLY_PIXELSHADER_TYPE(1,1,1);
-
+#define IMPLEMENT_REFLECTION_APPLY_PIXELSHADER_TYPE(A, B, C, D) \
+	typedef FReflectionApplyPS<A,B,C,D> FReflectionApplyPS##A##B##C##D; \
+	IMPLEMENT_SHADER_TYPE(template<>,FReflectionApplyPS##A##B##C##D,TEXT("ReflectionEnvironmentShaders"),TEXT("ReflectionApplyPS"),SF_Pixel);
+
+IMPLEMENT_REFLECTION_APPLY_PIXELSHADER_TYPE(0,0,0,0);
+IMPLEMENT_REFLECTION_APPLY_PIXELSHADER_TYPE(0,0,1,0);
+IMPLEMENT_REFLECTION_APPLY_PIXELSHADER_TYPE(0,1,0,0);
+IMPLEMENT_REFLECTION_APPLY_PIXELSHADER_TYPE(0,1,1,0);
+IMPLEMENT_REFLECTION_APPLY_PIXELSHADER_TYPE(1,0,0,0);
+IMPLEMENT_REFLECTION_APPLY_PIXELSHADER_TYPE(1,0,1,0);
+IMPLEMENT_REFLECTION_APPLY_PIXELSHADER_TYPE(1,1,0,0);
+IMPLEMENT_REFLECTION_APPLY_PIXELSHADER_TYPE(1,1,1,0);
+
+IMPLEMENT_REFLECTION_APPLY_PIXELSHADER_TYPE(0,0,0,1);
+IMPLEMENT_REFLECTION_APPLY_PIXELSHADER_TYPE(0,0,1,1);
+IMPLEMENT_REFLECTION_APPLY_PIXELSHADER_TYPE(0,1,0,1);
+IMPLEMENT_REFLECTION_APPLY_PIXELSHADER_TYPE(0,1,1,1);
+IMPLEMENT_REFLECTION_APPLY_PIXELSHADER_TYPE(1,0,0,1);
+IMPLEMENT_REFLECTION_APPLY_PIXELSHADER_TYPE(1,0,1,1);
+IMPLEMENT_REFLECTION_APPLY_PIXELSHADER_TYPE(1,1,0,1);
+IMPLEMENT_REFLECTION_APPLY_PIXELSHADER_TYPE(1,1,1,1);
 
 class FReflectionCaptureSpecularBouncePS : public FGlobalShader
 {
@@ -754,7 +783,8 @@ void GatherAndSortReflectionCaptures(const FScene* Scene, TArray<FReflectionCapt
 	OutSortData.Sort();	
 }
 
-FReflectionEnvironmentTiledDeferredCS* SelectReflectionEnvironmentTiledDeferredCS(TShaderMap<FGlobalShaderType>* ShaderMap, bool bUseLightmaps, bool bHasSkyLight, bool bHasBoxCaptures, bool bHasSphereCaptures)
+template<bool bSuportDFAOIndirectOcclusion>
+FReflectionEnvironmentTiledDeferredCS* SelectReflectionEnvironmentTiledDeferredCSInner(TShaderMap<FGlobalShaderType>* ShaderMap, bool bUseLightmaps, bool bHasSkyLight, bool bHasBoxCaptures, bool bHasSphereCaptures)
 {
 	FReflectionEnvironmentTiledDeferredCS* ComputeShader = nullptr;
 	if (bUseLightmaps)
@@ -763,38 +793,38 @@ FReflectionEnvironmentTiledDeferredCS* SelectReflectionEnvironmentTiledDeferredC
 		{
 			if (bHasBoxCaptures && bHasSphereCaptures)
 			{
-				ComputeShader = *TShaderMapRef< TReflectionEnvironmentTiledDeferredCS<1, 1, 1, 1> >(ShaderMap);
+				ComputeShader = *TShaderMapRef< TReflectionEnvironmentTiledDeferredCS<1, 1, 1, 1, bSuportDFAOIndirectOcclusion> >(ShaderMap);
 			}
 			else if (bHasBoxCaptures)
 			{
-				ComputeShader = *TShaderMapRef< TReflectionEnvironmentTiledDeferredCS<1, 1, 1, 0> >(ShaderMap);
+				ComputeShader = *TShaderMapRef< TReflectionEnvironmentTiledDeferredCS<1, 1, 1, 0, bSuportDFAOIndirectOcclusion> >(ShaderMap);
 			}
 			else if (bHasSphereCaptures)
 			{
-				ComputeShader = *TShaderMapRef< TReflectionEnvironmentTiledDeferredCS<1, 1, 0, 1> >(ShaderMap);
+				ComputeShader = *TShaderMapRef< TReflectionEnvironmentTiledDeferredCS<1, 1, 0, 1, bSuportDFAOIndirectOcclusion> >(ShaderMap);
 			}
 			else
 			{
-				ComputeShader = *TShaderMapRef< TReflectionEnvironmentTiledDeferredCS<1, 1, 0, 0> >(ShaderMap);
+				ComputeShader = *TShaderMapRef< TReflectionEnvironmentTiledDeferredCS<1, 1, 0, 0, bSuportDFAOIndirectOcclusion> >(ShaderMap);
 			}
 		}
 		else
 		{
 			if (bHasBoxCaptures && bHasSphereCaptures)
 			{
-				ComputeShader = *TShaderMapRef< TReflectionEnvironmentTiledDeferredCS<1, 0, 1, 1> >(ShaderMap);
+				ComputeShader = *TShaderMapRef< TReflectionEnvironmentTiledDeferredCS<1, 0, 1, 1, bSuportDFAOIndirectOcclusion> >(ShaderMap);
 			}
 			else if (bHasBoxCaptures)
 			{
-				ComputeShader = *TShaderMapRef< TReflectionEnvironmentTiledDeferredCS<1, 0, 1, 0> >(ShaderMap);
+				ComputeShader = *TShaderMapRef< TReflectionEnvironmentTiledDeferredCS<1, 0, 1, 0, bSuportDFAOIndirectOcclusion> >(ShaderMap);
 			}
 			else if (bHasSphereCaptures)
 			{
-				ComputeShader = *TShaderMapRef< TReflectionEnvironmentTiledDeferredCS<1, 0, 0, 1> >(ShaderMap);
+				ComputeShader = *TShaderMapRef< TReflectionEnvironmentTiledDeferredCS<1, 0, 0, 1, bSuportDFAOIndirectOcclusion> >(ShaderMap);
 			}
 			else
 			{
-				ComputeShader = *TShaderMapRef< TReflectionEnvironmentTiledDeferredCS<1, 0, 0, 0> >(ShaderMap);
+				ComputeShader = *TShaderMapRef< TReflectionEnvironmentTiledDeferredCS<1, 0, 0, 0, bSuportDFAOIndirectOcclusion> >(ShaderMap);
 			}
 		}		
 	}
@@ -804,38 +834,38 @@ FReflectionEnvironmentTiledDeferredCS* SelectReflectionEnvironmentTiledDeferredC
 		{
 			if (bHasBoxCaptures && bHasSphereCaptures)
 			{
-				ComputeShader = *TShaderMapRef< TReflectionEnvironmentTiledDeferredCS<0, 1, 1, 1> >(ShaderMap);
+				ComputeShader = *TShaderMapRef< TReflectionEnvironmentTiledDeferredCS<0, 1, 1, 1, bSuportDFAOIndirectOcclusion> >(ShaderMap);
 			}
 			else if (bHasBoxCaptures)
 			{
-				ComputeShader = *TShaderMapRef< TReflectionEnvironmentTiledDeferredCS<0, 1, 1, 0> >(ShaderMap);
+				ComputeShader = *TShaderMapRef< TReflectionEnvironmentTiledDeferredCS<0, 1, 1, 0, bSuportDFAOIndirectOcclusion> >(ShaderMap);
 			}
 			else if (bHasSphereCaptures)
 			{
-				ComputeShader = *TShaderMapRef< TReflectionEnvironmentTiledDeferredCS<0, 1, 0, 1> >(ShaderMap);
+				ComputeShader = *TShaderMapRef< TReflectionEnvironmentTiledDeferredCS<0, 1, 0, 1, bSuportDFAOIndirectOcclusion> >(ShaderMap);
 			}
 			else
 			{
-				ComputeShader = *TShaderMapRef< TReflectionEnvironmentTiledDeferredCS<0, 1, 0, 0> >(ShaderMap);
+				ComputeShader = *TShaderMapRef< TReflectionEnvironmentTiledDeferredCS<0, 1, 0, 0, bSuportDFAOIndirectOcclusion> >(ShaderMap);
 			}
 		}
 		else
 		{
 			if (bHasBoxCaptures && bHasSphereCaptures)
 			{
-				ComputeShader = *TShaderMapRef< TReflectionEnvironmentTiledDeferredCS<0, 0, 1, 1> >(ShaderMap);
+				ComputeShader = *TShaderMapRef< TReflectionEnvironmentTiledDeferredCS<0, 0, 1, 1, bSuportDFAOIndirectOcclusion> >(ShaderMap);
 			}
 			else if (bHasBoxCaptures)
 			{
-				ComputeShader = *TShaderMapRef< TReflectionEnvironmentTiledDeferredCS<0, 0, 1, 0> >(ShaderMap);
+				ComputeShader = *TShaderMapRef< TReflectionEnvironmentTiledDeferredCS<0, 0, 1, 0, bSuportDFAOIndirectOcclusion> >(ShaderMap);
 			}
 			else if (bHasSphereCaptures)
 			{
-				ComputeShader = *TShaderMapRef< TReflectionEnvironmentTiledDeferredCS<0, 0, 0, 1> >(ShaderMap);
+				ComputeShader = *TShaderMapRef< TReflectionEnvironmentTiledDeferredCS<0, 0, 0, 1, bSuportDFAOIndirectOcclusion> >(ShaderMap);
 			}
 			else
 			{
-				ComputeShader = *TShaderMapRef< TReflectionEnvironmentTiledDeferredCS<0, 0, 0, 0> >(ShaderMap);
+				ComputeShader = *TShaderMapRef< TReflectionEnvironmentTiledDeferredCS<0, 0, 0, 0, bSuportDFAOIndirectOcclusion> >(ShaderMap);
 			}
 		}		
 	}
@@ -843,6 +873,18 @@ FReflectionEnvironmentTiledDeferredCS* SelectReflectionEnvironmentTiledDeferredC
 	return ComputeShader;
 }
 
+FReflectionEnvironmentTiledDeferredCS* SelectReflectionEnvironmentTiledDeferredCS(TShaderMap<FGlobalShaderType>* ShaderMap, bool bUseLightmaps, bool bHasSkyLight, bool bHasBoxCaptures, bool bHasSphereCaptures, bool bSuportDFAOIndirectOcclusion)
+{
+	if (bSuportDFAOIndirectOcclusion)
+	{
+		return SelectReflectionEnvironmentTiledDeferredCSInner<true>(ShaderMap, bUseLightmaps, bHasSkyLight, bHasBoxCaptures, bHasSphereCaptures);
+	}
+	else
+	{
+		return SelectReflectionEnvironmentTiledDeferredCSInner<false>(ShaderMap, bUseLightmaps, bHasSkyLight, bHasBoxCaptures, bHasSphereCaptures);
+	}
+}
+
 void FDeferredShadingSceneRenderer::RenderTiledDeferredImageBasedReflections(FRHICommandListImmediate& RHICmdList, const TRefCountPtr<IPooledRenderTarget>& DynamicBentNormalAO)
 {
 	FSceneRenderTargets& SceneContext = FSceneRenderTargets::Get(RHICmdList);
@@ -902,10 +944,8 @@ void FDeferredShadingSceneRenderer::RenderTiledDeferredImageBasedReflections(FRH
 					View.ViewRect.Width(), View.ViewRect.Height(), GReflectionEnvironmentTileSizeX, GReflectionEnvironmentTileSizeY,
 					NumBoxCaptures, NumSphereCaptures, bHasSkyLight);
 
-				ComputeShader = SelectReflectionEnvironmentTiledDeferredCS(View.ShaderMap, bUseLightmaps, bHasSkyLight, bHasBoxCaptures, bHasSphereCaptures);
-
+				ComputeShader = SelectReflectionEnvironmentTiledDeferredCS(View.ShaderMap, bUseLightmaps, bHasSkyLight, bHasBoxCaptures, bHasSphereCaptures, DynamicBentNormalAO != NULL);
 
-			
 				//Really we should write this fence where we transition the final depedency for the reflections.  We may add an RHI command just for writing fences if this
 				//can't be done in the general case.  In the meantime, hack this command a bit to write the fence.
 				RHICmdList.TransitionResources(EResourceTransitionAccess::EReadable, EResourceTransitionPipeline::EGfxToCompute, nullptr, 0, ReflectionBeginFence);
@@ -1141,26 +1181,36 @@ void FDeferredShadingSceneRenderer::RenderStandardDeferredImageBasedReflections(
 			// todo: refactor (we abuse another boolean to pass the data through)
 			bReflectionEnv = bReflectionEnv || bEnvironmentMixing;
 
-#define CASE(A,B,C) \
-			case ((A << 2) | (B << 1) | C) : \
+			const bool bSupportDFAOIndirectShadowing = DynamicBentNormalAO != NULL;
+
+#define CASE(A,B,C,D) \
+			case ((A << 3) | (B << 2) | (C << 1) | D) : \
 			{ \
-			TShaderMapRef< FReflectionApplyPS<A, B, C> > PixelShader(View.ShaderMap); \
+			TShaderMapRef< FReflectionApplyPS<A, B, C, D> > PixelShader(View.ShaderMap); \
 			static FGlobalBoundShaderState BoundShaderState; \
 			SetGlobalBoundShaderState(RHICmdList, FeatureLevel, BoundShaderState, GFilterVertexDeclaration.VertexDeclarationRHI, *VertexShader, *PixelShader); \
 			PixelShader->SetParameters(RHICmdList, View, LightAccumulation->GetRenderTargetItem().ShaderResourceTexture, SSROutput->GetRenderTargetItem().ShaderResourceTexture, DynamicBentNormalAO); \
 			}; \
 			break
 
-			switch (((uint32)bSSR << 2) | ((uint32)bReflectionEnv << 1) | (uint32)bSkyLight)
+			switch (((uint32)bSSR << 3) | ((uint32)bReflectionEnv << 2) | ((uint32)bSkyLight << 1) | (uint32)bSupportDFAOIndirectShadowing)
 			{
-				CASE(0, 0, 0);
-				CASE(0, 0, 1);
-				CASE(0, 1, 0);
-				CASE(0, 1, 1);
-				CASE(1, 0, 0);
-				CASE(1, 0, 1);
-				CASE(1, 1, 0);
-				CASE(1, 1, 1);
+				CASE(0, 0, 0, 0);
+				CASE(0, 0, 1, 0);
+				CASE(0, 1, 0, 0);
+				CASE(0, 1, 1, 0);
+				CASE(1, 0, 0, 0);
+				CASE(1, 0, 1, 0);
+				CASE(1, 1, 0, 0);
+				CASE(1, 1, 1, 0);
+				CASE(0, 0, 0, 1);
+				CASE(0, 0, 1, 1);
+				CASE(0, 1, 0, 1);
+				CASE(0, 1, 1, 1);
+				CASE(1, 0, 0, 1);
+				CASE(1, 0, 1, 1);
+				CASE(1, 1, 0, 1);
+				CASE(1, 1, 1, 1);
 			}
 #undef CASE
 
diff --git a/Engine/Source/Runtime/Renderer/Private/SceneRendering.cpp b/Engine/Source/Runtime/Renderer/Private/SceneRendering.cpp
--- a/Engine/Source/Runtime/Renderer/Private/SceneRendering.cpp
+++ b/Engine/Source/Runtime/Renderer/Private/SceneRendering.cpp
@@ -187,6 +187,8 @@ static TAutoConsoleVariable<int32> CVarRHICmdMinDrawsPerParallelCmdList(
 	64,
 	TEXT("The minimum number of draws per cmdlist. If the total number of draws is less than this, then no parallel work will be done at all. This can't always be honored or done correctly. More effective with RHICmdBalanceParallelLists."));
 
+static FParallelCommandListSet* GOutstandingParallelCommandListSet = nullptr;
+
 FParallelCommandListSet::FParallelCommandListSet(TStatId InExecuteStat, const FViewInfo& InView, FRHICommandListImmediate& InParentCmdList, bool bInParallelExecute, bool bInCreateSceneContext)
 	: View(InView)
 	, ParentCmdList(InParentCmdList)
@@ -205,6 +207,8 @@ FParallelCommandListSet::FParallelCommandListSet(TStatId InExecuteStat, const FV
 	CommandLists.Reserve(Width * 8);
 	Events.Reserve(Width * 8);
 	NumDrawsIfKnown.Reserve(Width * 8);
+	check(!GOutstandingParallelCommandListSet);
+	GOutstandingParallelCommandListSet = this;
 }
 
 FRHICommandList* FParallelCommandListSet::AllocCommandList()
@@ -281,6 +285,9 @@ void FParallelCommandListSet::Dispatch(bool bHighPriority)
 
 FParallelCommandListSet::~FParallelCommandListSet()
 {
+	check(GOutstandingParallelCommandListSet == this);
+	GOutstandingParallelCommandListSet = nullptr;
+
 	check(IsInRenderingThread() && FMemStack::Get().GetNumMarks() == 1); // we do not want this popped before the end of the scene and it better be the scene allocator
 	checkf(CommandLists.Num() == 0, TEXT("Derived class of FParallelCommandListSet did not call Dispatch in virtual destructor"));
 	checkf(NumAlloc == 0, TEXT("Derived class of FParallelCommandListSet did not call Dispatch in virtual destructor"));
@@ -314,6 +321,30 @@ void FParallelCommandListSet::AddParallelCommandList(FRHICommandList* CmdList, F
 	NumDrawsIfKnown.Add(InNumDrawsIfKnown);
 }
 
+void FParallelCommandListSet::WaitForTasks()
+{
+	check(GOutstandingParallelCommandListSet);
+	GOutstandingParallelCommandListSet->WaitForTasksInternal();
+}
+
+void FParallelCommandListSet::WaitForTasksInternal()
+{
+	check(IsInRenderingThread());
+	QUICK_SCOPE_CYCLE_COUNTER(STAT_FParallelCommandListSet_WaitForTasks);
+	FGraphEventArray WaitOutstandingTasks;
+	for (int32 Index = 0; Index < Events.Num(); Index++)
+	{
+		if (!Events[Index]->IsComplete())
+		{
+			WaitOutstandingTasks.Add(Events[Index]);
+		}
+	}
+	if (WaitOutstandingTasks.Num())
+	{
+		check(!FTaskGraphInterface::Get().IsThreadProcessingTasks(ENamedThreads::RenderThread_Local));
+		FTaskGraphInterface::Get().WaitUntilTasksComplete(WaitOutstandingTasks, ENamedThreads::RenderThread_Local);
+	}
+}
 
 
 /*-----------------------------------------------------------------------------
diff --git a/Engine/Source/Runtime/Renderer/Private/SceneRendering.h b/Engine/Source/Runtime/Renderer/Private/SceneRendering.h
--- a/Engine/Source/Runtime/Renderer/Private/SceneRendering.h
+++ b/Engine/Source/Runtime/Renderer/Private/SceneRendering.h
@@ -362,6 +362,9 @@ public:
 	{
 
 	}
+	static void WaitForTasks();
+private:
+	void WaitForTasksInternal();
 };
 
 class FVolumeUpdateRegion
diff --git a/Engine/Source/Runtime/Renderer/Private/ShadowRendering.cpp b/Engine/Source/Runtime/Renderer/Private/ShadowRendering.cpp
--- a/Engine/Source/Runtime/Renderer/Private/ShadowRendering.cpp
+++ b/Engine/Source/Runtime/Renderer/Private/ShadowRendering.cpp
@@ -2304,7 +2304,7 @@ void FProjectedShadowInfo::RenderProjection(FRHICommandListImmediate& RHICmdList
 		// Note that self-shadow pre-shadows still mask by receiver elements.
 		const TArray<FMeshBatchAndRelevance,SceneRenderingAllocator>& DynamicMeshElements = bPreShadow ? DynamicReceiverMeshElements : DynamicSubjectMeshElements;
 
-		FDepthDrawingPolicyFactory::ContextType Context(DDM_AllOccluders);
+		FDepthDrawingPolicyFactory::ContextType Context(DDM_AllOccluders, false);
 
 		for (int32 MeshBatchIndex = 0; MeshBatchIndex < DynamicMeshElements.Num(); MeshBatchIndex++)
 		{
@@ -2337,7 +2337,7 @@ void FProjectedShadowInfo::RenderProjection(FRHICommandListImmediate& RHICmdList
 							FDepthDrawingPolicyFactory::DrawStaticMesh(
 								RHICmdList, 
 								*View,
-								FDepthDrawingPolicyFactory::ContextType(DDM_AllOccluders),
+								FDepthDrawingPolicyFactory::ContextType(DDM_AllOccluders, false),
 								StaticMesh,
 								StaticMesh.Elements.Num() == 1 ? 1 : View->StaticMeshBatchVisibility[StaticMesh.Id],
 								true,
@@ -2361,7 +2361,7 @@ void FProjectedShadowInfo::RenderProjection(FRHICommandListImmediate& RHICmdList
 				FDepthDrawingPolicyFactory::DrawStaticMesh(
 					RHICmdList, 
 					*View,
-					FDepthDrawingPolicyFactory::ContextType(DDM_AllOccluders),
+					FDepthDrawingPolicyFactory::ContextType(DDM_AllOccluders, false),
 					StaticMesh,
 					StaticMesh.Elements.Num() == 1 ? 1 : View->StaticMeshBatchVisibility[StaticMesh.Id],
 					true,
@@ -2497,7 +2497,7 @@ void FProjectedShadowInfo::RenderProjection(FRHICommandListImmediate& RHICmdList
 				0xff, 0xff
 			>::GetRHI(), 0);
 
-			FDepthDrawingPolicyFactory::ContextType Context(DDM_AllOccluders);
+			FDepthDrawingPolicyFactory::ContextType Context(DDM_AllOccluders, false);
 			for (int32 MeshBatchIndex = 0; MeshBatchIndex < DynamicSubjectMeshElements.Num(); MeshBatchIndex++)
 			{
 				const FMeshBatchAndRelevance& MeshBatchAndRelevance = DynamicSubjectMeshElements[MeshBatchIndex];
diff --git a/Engine/Source/Runtime/Renderer/Private/ShadowSetup.cpp b/Engine/Source/Runtime/Renderer/Private/ShadowSetup.cpp
--- a/Engine/Source/Runtime/Renderer/Private/ShadowSetup.cpp
+++ b/Engine/Source/Runtime/Renderer/Private/ShadowSetup.cpp
@@ -757,7 +757,18 @@ void FProjectedShadowInfo::AddSubjectPrimitive(FPrimitiveSceneInfo* PrimitiveSce
 			*(PrimitiveSceneInfo->ComponentLastRenderTime) = CurrentWorldTime;
 			if (PrimitiveSceneInfo->NeedsLazyUpdateForRendering())
 			{
-				PrimitiveSceneInfo->ConditionalLazyUpdateForRendering(FRHICommandListExecutor::GetImmediateCommandList());
+				if (GDrawListsLocked && PrimitiveSceneInfo->NeedsUpdateStaticMeshes())
+				{
+					QUICK_SCOPE_CYCLE_COUNTER(STAT_FProjectedShadowInfo_AddSubjectPrimitive_FlushPrepass);
+					FRHICommandListExecutor::GetImmediateCommandList().ImmediateFlush(EImmediateFlushType::WaitForOutstandingTasksOnly);
+					FParallelCommandListSet::WaitForTasks();
+					TGuardValue<bool> LockDrawLists(GDrawListsLocked, false);
+					PrimitiveSceneInfo->ConditionalLazyUpdateForRendering(FRHICommandListExecutor::GetImmediateCommandList());
+				}
+				else
+				{
+					PrimitiveSceneInfo->ConditionalLazyUpdateForRendering(FRHICommandListExecutor::GetImmediateCommandList());
+				}
 			}
 		}
 
diff --git a/Engine/Source/Runtime/Renderer/Private/StaticMeshDrawList.h b/Engine/Source/Runtime/Renderer/Private/StaticMeshDrawList.h
--- a/Engine/Source/Runtime/Renderer/Private/StaticMeshDrawList.h
+++ b/Engine/Source/Runtime/Renderer/Private/StaticMeshDrawList.h
@@ -7,6 +7,7 @@
 #ifndef __STATICMESHDRAWLIST_H__
 #define __STATICMESHDRAWLIST_H__
 
+extern ENGINE_API bool GDrawListsLocked;
 /** View state for instanced stereo rendering. */
 struct StereoPair
 {
diff --git a/Engine/Source/Runtime/Renderer/Private/StaticMeshDrawList.inl b/Engine/Source/Runtime/Renderer/Private/StaticMeshDrawList.inl
--- a/Engine/Source/Runtime/Renderer/Private/StaticMeshDrawList.inl
+++ b/Engine/Source/Runtime/Renderer/Private/StaticMeshDrawList.inl
@@ -16,6 +16,7 @@
 template<typename DrawingPolicyType>
 void TStaticMeshDrawList<DrawingPolicyType>::FElementHandle::Remove()
 {
+	check(!GDrawListsLocked);
 	// Make a copy of this handle's variables on the stack, since the call to Elements.RemoveSwap deletes the handle.
 	TStaticMeshDrawList* const LocalDrawList = StaticMeshDrawList;
 	FDrawingPolicyLink* const LocalDrawingPolicyLink = &LocalDrawList->DrawingPolicySet[SetId];
@@ -170,6 +171,8 @@ void TStaticMeshDrawList<DrawingPolicyType>::AddMesh(
 	ERHIFeatureLevel::Type InFeatureLevel
 	)
 {
+	check(!GDrawListsLocked);
+
 	// Check for an existing drawing policy matching the mesh's drawing policy.
 	FDrawingPolicyLink* DrawingPolicyLink = DrawingPolicySet.Find(InDrawingPolicy);
 	if(!DrawingPolicyLink)
diff --git a/Engine/Source/Runtime/Renderer/Private/VelocityRendering.cpp b/Engine/Source/Runtime/Renderer/Private/VelocityRendering.cpp
--- a/Engine/Source/Runtime/Renderer/Private/VelocityRendering.cpp
+++ b/Engine/Source/Runtime/Renderer/Private/VelocityRendering.cpp
@@ -557,7 +557,7 @@ bool IsMotionBlurEnabled(const FViewInfo& View)
 
 void FDeferredShadingSceneRenderer::RenderDynamicVelocitiesMeshElementsInner(FRHICommandList& RHICmdList, const FViewInfo& View, int32 FirstIndex, int32 LastIndex)
 {
-	FVelocityDrawingPolicyFactory::ContextType Context(DDM_AllOccluders);
+	FVelocityDrawingPolicyFactory::ContextType Context(DDM_AllOccluders, false);
 
 	for (int32 MeshBatchIndex = FirstIndex; MeshBatchIndex <= LastIndex; MeshBatchIndex++)
 	{
diff --git a/Engine/Source/Runtime/Renderer/Private/VelocityRendering.h b/Engine/Source/Runtime/Renderer/Private/VelocityRendering.h
--- a/Engine/Source/Runtime/Renderer/Private/VelocityRendering.h
+++ b/Engine/Source/Runtime/Renderer/Private/VelocityRendering.h
@@ -65,7 +65,7 @@ private:
 class FVelocityDrawingPolicyFactory : public FDepthDrawingPolicyFactory
 {
 public:
-	static void AddStaticMesh(FScene* Scene, FStaticMesh* StaticMesh, ContextType = ContextType(DDM_AllOccluders));
+	static void AddStaticMesh(FScene* Scene, FStaticMesh* StaticMesh, ContextType = ContextType(DDM_AllOccluders, false));
 	static bool DrawDynamicMesh(	
 		FRHICommandList& RHICmdList, 
 		const FViewInfo& View,
diff --git a/Engine/Source/Runtime/Windows/D3D11RHI/Private/Windows/WindowsD3D11Device.cpp b/Engine/Source/Runtime/Windows/D3D11RHI/Private/Windows/WindowsD3D11Device.cpp
--- a/Engine/Source/Runtime/Windows/D3D11RHI/Private/Windows/WindowsD3D11Device.cpp
+++ b/Engine/Source/Runtime/Windows/D3D11RHI/Private/Windows/WindowsD3D11Device.cpp
@@ -20,7 +20,7 @@ static TAutoConsoleVariable<int32> CVarGraphicsAdapter(
 	TEXT("r.GraphicsAdapter"),
 	-1,
 	TEXT("User request to pick a specific graphics adapter (e.g. when using a integrated graphics card with a discrete one)\n")
-	TEXT("At the moment this only works on Direct3D 11.\n")
+	TEXT("At the moment this only works on Direct3D 11. Unless a specific adapter is chosen we reject Microsoft adapters because we don't want the software emulation.\n")
 	TEXT(" -2: Take the first one that fulfills the criteria\n")
 	TEXT(" -1: Favour non integrated because there are usually faster (default)\n")
 	TEXT("  0: Adpater #0\n")
@@ -297,16 +297,18 @@ void FD3D11DynamicRHIModule::FindAdapter()
 					GetFeatureLevelString(ActualFeatureLevel)
 					);
 				UE_LOG(LogD3D11RHI, Log,
-					TEXT("      %uMB of dedicated video memory, %uMB of dedicated system memory, and %uMB of shared system memory, %d output[s]"),
+					TEXT("      %u/%u/%u MB DedicatedVideo/DedicatedSystem/SharedSystem, Outputs:%d, VendorId:0x%x"),
 					(uint32)(AdapterDesc.DedicatedVideoMemory / (1024*1024)),
 					(uint32)(AdapterDesc.DedicatedSystemMemory / (1024*1024)),
 					(uint32)(AdapterDesc.SharedSystemMemory / (1024*1024)),
-					OutputCount
+					OutputCount,
+					AdapterDesc.VendorId
 					);
 
 				bool bIsAMD = AdapterDesc.VendorId == 0x1002;
 				bool bIsIntel = AdapterDesc.VendorId == 0x8086;
 				bool bIsNVIDIA = AdapterDesc.VendorId == 0x10DE;
+				bool bIsMicrosoft = AdapterDesc.VendorId == 0x1414;
 
 				if(bIsAMD) bIsAnyAMD = true;
 				if(bIsIntel) bIsAnyIntel = true;
@@ -319,10 +321,11 @@ void FD3D11DynamicRHIModule::FindAdapter()
 
 				FD3D11Adapter CurrentAdapter(AdapterIndex, ActualFeatureLevel);
 
-				if(!OutputCount)
+				if(bIsMicrosoft && CVarValue < 0)
 				{
-					// This device has no outputs. Reject it, 
-					// http://msdn.microsoft.com/en-us/library/windows/desktop/bb205075%28v=vs.85%29.aspx#WARP_new_for_Win8
+					// To reject the software emulation, unless the cvar wants it.
+					// https://msdn.microsoft.com/en-us/library/windows/desktop/bb205075(v=vs.85).aspx#WARP_new_for_Win8
+					// Before we tested for no output devices but that failed where a laptop had a Intel (with output) and NVidia (with no output)
 					continue;
 				}
 
