Copying //UE4/Ocean-Staging to //UE4/Main

========================
MAJOR FEATURES + CHANGES
========================

Change 2851350 on 2016/02/01 by [REDACTED_PERSON]

	#Ocean - Initial setup of automation scripts for the Ocean Main stream

Change 2858910 on 2016/02/08 by [REDACTED_PERSON]

	Fixup email notifications now that Ocean is in its own branch.

Change 2859312 on 2016/02/08 by [REDACTED_PERSON]

	Add a branch hacker for Ocean, which disables things like building documentation and the UE4 binary release.

Change 2869290 on 2016/02/16 by [REDACTED_PERSON]

	StreamableManager - Don't make ExistingStreamables static as this function is sometimes re-entrant and this causes issues.

Change 2872198 on 2016/02/18 by [REDACTED_PERSON]

	#engine - Made AnimNotify_PlayParticleEffect and AnimNotify_PlaySound part of ENGINE_API so it can be subclassed.

Change 2874305 on 2016/02/19 by [REDACTED_PERSON]

	Workaround for crash on loading a EGPLogoShapesSavePB Shape property

Change 2876829 on 2016/02/22 by [REDACTED_PERSON]

	Adding Ocean branch hack back in

diff --git a/Engine/Source/Runtime/Engine/Classes/Animation/AnimNotifies/AnimNotify_PlayParticleEffect.h b/Engine/Source/Runtime/Engine/Classes/Animation/AnimNotifies/AnimNotify_PlayParticleEffect.h
--- a/Engine/Source/Runtime/Engine/Classes/Animation/AnimNotifies/AnimNotify_PlayParticleEffect.h
+++ b/Engine/Source/Runtime/Engine/Classes/Animation/AnimNotifies/AnimNotify_PlayParticleEffect.h
@@ -5,8 +5,8 @@
 #include "AnimNotify.h"
 #include "AnimNotify_PlayParticleEffect.generated.h"
 
-UCLASS(const, hidecategories=Object, collapsecategories, meta=(DisplayName="Play Particle Effect"), MinimalAPI)
-class UAnimNotify_PlayParticleEffect : public UAnimNotify
+UCLASS(const, hidecategories=Object, collapsecategories, meta=(DisplayName="Play Particle Effect"))
+class ENGINE_API UAnimNotify_PlayParticleEffect : public UAnimNotify
 {
 	GENERATED_BODY()
 
diff --git a/Engine/Source/Runtime/Engine/Classes/Animation/AnimNotifies/AnimNotify_PlaySound.h b/Engine/Source/Runtime/Engine/Classes/Animation/AnimNotifies/AnimNotify_PlaySound.h
--- a/Engine/Source/Runtime/Engine/Classes/Animation/AnimNotifies/AnimNotify_PlaySound.h
+++ b/Engine/Source/Runtime/Engine/Classes/Animation/AnimNotifies/AnimNotify_PlaySound.h
@@ -5,8 +5,8 @@
 #include "AnimNotify.h"
 #include "AnimNotify_PlaySound.generated.h"
 
-UCLASS(const, hidecategories=Object, collapsecategories, meta=(DisplayName="Play Sound"), MinimalAPI)
-class UAnimNotify_PlaySound : public UAnimNotify
+UCLASS(const, hidecategories=Object, collapsecategories, meta=(DisplayName="Play Sound"))
+class ENGINE_API UAnimNotify_PlaySound : public UAnimNotify
 {
 	GENERATED_BODY()
 
diff --git a/Engine/Source/Runtime/Engine/Private/StreamableManager.cpp b/Engine/Source/Runtime/Engine/Private/StreamableManager.cpp
--- a/Engine/Source/Runtime/Engine/Private/StreamableManager.cpp
+++ b/Engine/Source/Runtime/Engine/Private/StreamableManager.cpp
@@ -249,7 +249,7 @@ void FStreamableManager::RequestAsyncLoad(const TArray<FStringAssetReference>& T
 	TSharedRef<FStreamableRequest> NewRequest = MakeShareable(new FStreamableRequest());
 	NewRequest->CompletionDelegate = DelegateToCall;
 
-	static TArray<FStreamable*> ExistingStreamables;
+	TArray<FStreamable *> ExistingStreamables;
 	ExistingStreamables.Reserve(TargetsToStream.Num());
 
 	for (int32 i = 0; i < TargetsToStream.Num(); i++)
@@ -273,8 +273,6 @@ void FStreamableManager::RequestAsyncLoad(const TArray<FStringAssetReference>& T
 			CheckCompletedRequests(TargetsToStream[i], Existing);
 		}
 	}
-
-	ExistingStreamables.Reset();
 }
 
 void FStreamableManager::RequestAsyncLoad(const FStringAssetReference& TargetToStream, FStreamableDelegate DelegateToCall, TAsyncLoadPriority Priority)
