Making sure material expressions are never stripped when running the editor with -game.

diff --git a/Engine/Source/Runtime/Engine/Classes/Materials/MaterialExpression.h b/Engine/Source/Runtime/Engine/Classes/Materials/MaterialExpression.h
--- a/Engine/Source/Runtime/Engine/Classes/Materials/MaterialExpression.h
+++ b/Engine/Source/Runtime/Engine/Classes/Materials/MaterialExpression.h
@@ -193,7 +193,11 @@ class ENGINE_API UMaterialExpression : public UObject
 	virtual bool Modify( bool bAlwaysMarkDirty=true ) override;
 	virtual void Serialize( FArchive& Ar ) override;
 	virtual bool NeedsLoadForClient() const override;
-	virtual bool NeedsLoadForServer() const override;	
+	virtual bool NeedsLoadForServer() const override;
+	virtual bool NeedsLoadForEditorGame() const override
+	{
+		return true;
+	}
 	//~ End UObject Interface.
 
 #if WITH_EDITOR
