Copying //UE4/Dev-Networking to Dev-Main (//UE4/Dev-Main)

==========================
MAJOR FEATURES + CHANGES
==========================

Change 2814968 on 2016/01/04 by [REDACTED_PERSON]
	Remove warning per connection when same warning happens at net driver level

Change 2815271 on 2016/01/04 by [REDACTED_PERSON]
	Fix header uploads in replays showing wrong size

	#tests Compiled, confirmed correct size

Change 2830613 on 2016/01/15 by [REDACTED_PERSON]
	Added better ability for unit tests to test socket-level netcode

	Full changes:
	- Added flag for auto-reconnecting on connection close
	- Added hook for socket level sends
	- Added ability to block socket-level sends and pre-PacketHandler LowLevelSend's
	- Added ability for net connections to force-enable packet handlers, even if presently off in the engine

Change 2830617 on 2016/01/15 by [REDACTED_PERSON][REDACTED_HANDLE]

	Enabled NetcodeUnitTest in Linux, and other minor updates.

Change 2830621 on 2016/01/15 by [REDACTED_PERSON][REDACTED_HANDLE]

	Added 'PacketLimitTest' and 'PacketLimitTest_Oodle' unit tests, for testing/verifying that MaxPacket limits are respected, through testing with Bunches and with raw sends passed through PacketHandlers (such as Oodle).

Change 2832402 on 2016/01/18 by [REDACTED_PERSON][REDACTED_HANDLE]

	Always use the physical platform file when writing out steam_appid.txt, because it has to live next to executable that's being run (CotF builds will output it to the wrong location).

	#tests checked that steam_appid.txt was written to the correct place with both CotF and cooked builds

Change 2833728 on 2016/01/19 by [REDACTED_PERSON][REDACTED_HANDLE]

	Fix for UE-20017 - "Http server headers implemented as TMaps and are overridden".  Just changing TMaps to TMultiMaps.  Haven't modified the behavior of GetHeader(), which will still return the first header that matches.

Change 2844880 on 2016/01/26 by [REDACTED_PERSON]
	Fix for UE-20017 Http server headers implemented as TMaps and are overridden

	Previously had submitted a fix that changes the request and response headers TMap into a MultiMap, but the consensus was that it would be better to leave them as TMaps, but allow appending onto the existing entry for a given request header with an "AppendToHeader" function, and to concatenate entries for the same header name in the response headers.  In both cases, entries for the same header name are concatenated with ", ".

	Mac and iOS already had the desired behavior for responses, just added the AppendToHeader function to their request classes.

	HTML5 still has no implementation of response headers, filed as separate JIRA issue UE-26047.

Change 2859343 on 2016/02/08 by [REDACTED_PERSON]
	Default MaxSearchResults in FOnlineSessionSearch objects to 1 instead of 0 - otherwise, the FOnlineAsyncTaskSteamFindServerBase tasks will complete right away without searching for sessions. Also print a warning in the task if this will happen.

Change 2862352 on 2016/02/10 by [REDACTED_PERSON]
	Fix demo.SkipTime cvar to work with negative values

Change 2866624 on 2016/02/14 by [REDACTED_PERSON]
	Fix plugin commandlet crash.

Change 2866627 on 2016/02/14 by [REDACTED_PERSON]
	Added NetDriver/PacketHandler ability to send/receive connectionless packets - i.e. packets not associated with a UNetConnection.

	Removed 'bEnabled' from PacketHandler - PacketHandler is now a default always-on part of the Engine.

	Added 'CountBits' to LowLevelSend, to support bit-level additions to packets from PacketHandler's, and refactored related code.

Change 2866628 on 2016/02/14 by [REDACTED_PERSON]
	Added an HMAC (Hash-based Message Authentication Code) generating function, to the SHA-1 code - for use with performing secure net connection handshakes.

Change 2866629 on 2016/02/14 by [REDACTED_PERSON]
	Added 'stateless handshake' to UE4, allowing a handshake to be performed without any open UNetConnection - making it more difficult for DoS attacks to consume memory resources.

Change 2866630 on 2016/02/14 by [REDACTED_PERSON]
	Updated Oodle to support new stateless handshake.

Change 2866631 on 2016/02/14 by [REDACTED_PERSON]
	Updated NetcodeUnitTest to work with new connectionless/stateless handshake changes. Compiles, but netcode presently needs fixing.

Change 2868367 on 2016/02/16 by [REDACTED_PERSON]
	Updated NetcodeUnitTest to work with stateless handshakes (now fixed connecting).

	Also fixed further commandlet problems.

Change 2878025 on 2016/02/23 by [REDACTED_PERSON]
	Immediately stop processing zero-size packets, to prevent CPU utilization during DDoS.

	Log IP of new connections, post-stateless-challenge (first time we know IP is verified).

Change 2883776 on 2016/02/26 by [REDACTED_PERSON]
	Fixed beacons to work with stateless handshake.

Change 2886773 on 2016/02/29 by [REDACTED_PERSON]
	Made PacketHandler bit overhead align upwards to the next byte -  this wastes more bits than necessary, but is required as a temporary fix.

Change 2887801 on 2016/03/01 by [REDACTED_PERSON]
	Use correct default values when initializing shadow state for replication.

	Fixes issues with certain components and actors that have default values that differ from CDO.

diff --git a/Engine/Plugins/Experimental/HTML5Networking/Source/HTML5Networking/Classes/WebSocketConnection.h b/Engine/Plugins/Experimental/HTML5Networking/Source/HTML5Networking/Classes/WebSocketConnection.h
--- a/Engine/Plugins/Experimental/HTML5Networking/Source/HTML5Networking/Classes/WebSocketConnection.h
+++ b/Engine/Plugins/Experimental/HTML5Networking/Source/HTML5Networking/Classes/WebSocketConnection.h
@@ -16,7 +16,7 @@ class HTML5NETWORKING_API UWebSocketConnection : public UNetConnection
 	virtual void InitBase(UNetDriver* InDriver, class FSocket* InSocket, const FURL& InURL, EConnectionState InState, int32 InMaxPacket = 0, int32 InPacketOverhead = 0) override;
 	virtual void InitRemoteConnection(UNetDriver* InDriver, class FSocket* InSocket, const FURL& InURL, const class FInternetAddr& InRemoteAddr, EConnectionState InState, int32 InMaxPacket = 0, int32 InPacketOverhead = 0) override;
 	virtual void InitLocalConnection(UNetDriver* InDriver, class FSocket* InSocket, const FURL& InURL, EConnectionState InState, int32 InMaxPacket = 0, int32 InPacketOverhead = 0) override;
-	virtual void LowLevelSend(void* Data, int32 Count) override;
+	virtual void LowLevelSend(void* Data, int32 CountBytes, int32 CountBits) override;
 	FString LowLevelGetRemoteAddress(bool bAppendPort = false) override;
 	FString LowLevelDescribe() override;
 	virtual void Tick();
diff --git a/Engine/Plugins/Experimental/HTML5Networking/Source/HTML5Networking/Private/WebsocketConnection.cpp b/Engine/Plugins/Experimental/HTML5Networking/Source/HTML5Networking/Private/WebsocketConnection.cpp
--- a/Engine/Plugins/Experimental/HTML5Networking/Source/HTML5Networking/Private/WebsocketConnection.cpp
+++ b/Engine/Plugins/Experimental/HTML5Networking/Source/HTML5Networking/Private/WebsocketConnection.cpp
@@ -64,10 +64,12 @@ void UWebSocketConnection::InitRemoteConnection(UNetDriver* InDriver, class FSoc
 	SetExpectedClientLoginMsgType(NMT_Hello);
 }
 
-void UWebSocketConnection::LowLevelSend(void* Data, int32 Count)
+void UWebSocketConnection::LowLevelSend(void* Data, int32 CountBytes, int32 CountBits)
 {
+	// @todo: PacketHandler's
+
 	int32 BytesSent = 0;
-	WebSocket->Send((uint8*)Data, Count);
+	WebSocket->Send((uint8*)Data, CountBytes);
 }
 
 FString UWebSocketConnection::LowLevelGetRemoteAddress(bool bAppendPort)
diff --git a/Engine/Plugins/NetcodeUnitTest/NUTUnrealEngine4/NUTUnrealEngine4.uplugin b/Engine/Plugins/NetcodeUnitTest/NUTUnrealEngine4/NUTUnrealEngine4.uplugin
--- a/Engine/Plugins/NetcodeUnitTest/NUTUnrealEngine4/NUTUnrealEngine4.uplugin
+++ b/Engine/Plugins/NetcodeUnitTest/NUTUnrealEngine4/NUTUnrealEngine4.uplugin
@@ -23,8 +23,9 @@
 			"WhitelistPlatforms" :
 			[
 				"Win64",
-				"Win32"
+				"Win32",
+				"Linux"
 			]
 		}
 	]
-}
\ No newline at end of file
+}
diff --git a/Engine/Plugins/NetcodeUnitTest/NUTUnrealEngine4/Source/NUTUnrealEngine4/Classes/UnitTests/PacketLimitTest.h b/Engine/Plugins/NetcodeUnitTest/NUTUnrealEngine4/Source/NUTUnrealEngine4/Classes/UnitTests/PacketLimitTest.h
new file mode 100644
--- /dev/null
+++ b/Engine/Plugins/NetcodeUnitTest/NUTUnrealEngine4/Source/NUTUnrealEngine4/Classes/UnitTests/PacketLimitTest.h
@@ -0,0 +1,89 @@
+// Copyright 1998-2016 Epic Games, Inc. All Rights Reserved.
+
+#pragma once
+
+#include "UnitTest.h"
+
+#include "PacketLimitTest.generated.h"
+
+
+/**
+ * The current stage of testing for the unit test
+ */
+enum class ELimitTestStage : uint8
+{
+	LTS_LowLevel_AtLimit	= 0,
+	LTS_LowLevel_OverLimit	= 1,
+	LTS_Bunch_AtLimit		= 2,
+	LTS_Bunch_OverLimit		= 3,
+
+	LTS_MAX					= 4
+};
+
+/**
+ * Used to get name values for the above enum
+ */
+inline FString GetLimitTestStageName(ELimitTestStage Stage)
+{
+	#define ELTS_CASE(x) case ELimitTestStage::x : return TEXT(#x)
+
+	switch (Stage)
+	{
+		ELTS_CASE(LTS_LowLevel_AtLimit);
+		ELTS_CASE(LTS_LowLevel_OverLimit);
+		ELTS_CASE(LTS_Bunch_AtLimit);
+		ELTS_CASE(LTS_Bunch_OverLimit);
+		ELTS_CASE(LTS_MAX);
+
+	default:
+		return FString::Printf(TEXT("Unknown 0x%08X"), (uint32)Stage);
+	}
+
+	#undef ELTS_CASE
+}
+
+/**
+ * Unit test for testing packet bit-size limits, and ensuring edge cases don't trigger warnings or send failures
+ */
+UCLASS()
+class UPacketLimitTest : public UClientUnitTest
+{
+	GENERATED_UCLASS_BODY()
+
+protected:
+	/** Whether or not to test with Oodle enabled. */
+	bool bUseOodle;
+
+	/** The current stage of testing */
+	ELimitTestStage TestStage;
+
+	/** The last attempted send size passed to LowLevelSend (before PacketHandlers) */
+	int32 LastSendSize;
+
+	/** The last attempted send size passed to the connection socket (after PacketHandlers) */
+	int32 LastSocketSendSize;
+
+	/** The size the final packet needs to be, if it is to be sent - all packets not matching this size, are blocked */
+	int32 TargetSocketSendSize;
+
+public:
+	virtual void InitializeEnvironmentSettings() override;
+
+	virtual bool ValidateUnitTestSettings(bool bCDOCheck) override;
+
+	virtual bool ConnectFakeClient(FUniqueNetIdRepl* InNetID) override;
+
+	virtual void ExecuteClientUnitTest() override;
+
+	virtual void NotifySendRawPacket(void* Data, int32 Count, bool& bBlockSend) override;
+
+	virtual void NotifySocketSendRawPacket(void* Data, int32 Count, bool& bBlockSend) override;
+
+	virtual void NotifyProcessLog(TWeakPtr<FUnitTestProcess> InProcess, const TArray<FString>& InLogLines) override;
+
+protected:
+	/**
+	 * Advanced to the next test stage
+	 */
+	void NextTestStage();
+};
\ No newline at end of file
diff --git a/Engine/Plugins/NetcodeUnitTest/NUTUnrealEngine4/Source/NUTUnrealEngine4/Classes/UnitTests/PacketLimitTest_Oodle.h b/Engine/Plugins/NetcodeUnitTest/NUTUnrealEngine4/Source/NUTUnrealEngine4/Classes/UnitTests/PacketLimitTest_Oodle.h
new file mode 100644
--- /dev/null
+++ b/Engine/Plugins/NetcodeUnitTest/NUTUnrealEngine4/Source/NUTUnrealEngine4/Classes/UnitTests/PacketLimitTest_Oodle.h
@@ -0,0 +1,16 @@
+// Copyright 1998-2016 Epic Games, Inc. All Rights Reserved.
+
+#pragma once
+
+#include "UnitTests/PacketLimitTest.h"
+
+#include "PacketLimitTest_Oodle.generated.h"
+
+/**
+ * Version of UPacketLimitTest, which tests with Oodle enabled
+ */
+UCLASS()
+class UPacketLimitTest_Oodle : public UPacketLimitTest
+{
+	GENERATED_UCLASS_BODY()
+};
diff --git a/Engine/Plugins/NetcodeUnitTest/NUTUnrealEngine4/Source/NUTUnrealEngine4/Private/UnitTests/PacketLimitTest.cpp b/Engine/Plugins/NetcodeUnitTest/NUTUnrealEngine4/Source/NUTUnrealEngine4/Private/UnitTests/PacketLimitTest.cpp
new file mode 100644
--- /dev/null
+++ b/Engine/Plugins/NetcodeUnitTest/NUTUnrealEngine4/Source/NUTUnrealEngine4/Private/UnitTests/PacketLimitTest.cpp
@@ -0,0 +1,332 @@
+// Copyright 1998-2016 Epic Games, Inc. All Rights Reserved.
+
+#include "NUTUnrealEngine4PCH.h"
+
+#include "UnitTests/PacketLimitTest.h"
+
+#include "UnitTestEnvironment.h"
+#include "Net/UnitTestNetConnection.h"
+
+#include "IpConnection.h"
+
+
+/**
+ * UPacketLimitTest
+ */
+
+UPacketLimitTest::UPacketLimitTest(const FObjectInitializer& ObjectInitializer)
+	: Super(ObjectInitializer)
+	, bUseOodle(false)
+	, TestStage(ELimitTestStage::LTS_LowLevel_AtLimit)
+	, LastSendSize(0)
+	, LastSocketSendSize(0)
+	, TargetSocketSendSize(0)
+{
+	UnitTestName = TEXT("PacketLimitTest");
+	UnitTestType = TEXT("Test");
+
+	UnitTestDate = FDateTime(2015, 12, 23);
+
+	bWorkInProgress = true;
+
+	// @todo #JohnBExploitCL: Bugtracking/changelist notes
+
+	ExpectedResult.Add(TEXT("ShooterGame"), EUnitTestVerification::VerifiedNotFixed);
+
+	UnitTestTimeout = 60;
+
+
+	UnitTestFlags |= (EUnitTestFlags::LaunchServer | EUnitTestFlags::SkipControlJoin | EUnitTestFlags::AutoReconnect |
+						EUnitTestFlags::RequirePing | EUnitTestFlags::CaptureSendRaw);
+}
+
+void UPacketLimitTest::InitializeEnvironmentSettings()
+{
+	BaseServerURL = UnitEnv->GetDefaultMap(UnitTestFlags);
+	BaseServerParameters = UnitEnv->GetDefaultServerParameters();
+
+	if (bUseOodle)
+	{
+		BaseServerURL += TEXT(" -PacketHandler -Oodle");
+	}
+}
+
+bool UPacketLimitTest::ValidateUnitTestSettings(bool bCDOCheck)
+{
+	bool bReturnVal = Super::ValidateUnitTestSettings(bCDOCheck);
+
+#if UE_BUILD_SHIPPING
+	UNIT_LOG(ELogType::StatusFailure, TEXT("The 'PacketLimitTest' unit test, can only be run in non-shipping mode."));
+
+	bReturnVal = false;
+#endif
+
+	return bReturnVal;
+}
+
+bool UPacketLimitTest::ConnectFakeClient(FUniqueNetIdRepl* InNetID)
+{
+	bool bReturnVal = false;
+	bool bForceEnableDefault = UUnitTestNetConnection::bForceEnableHandler;
+
+	if (bUseOodle)
+	{
+		UUnitTestNetConnection::bForceEnableHandler = true;
+		GEngine->Exec(NULL, TEXT("Oodle ForceEnable On"));
+	}
+
+	bReturnVal = Super::ConnectFakeClient(InNetID);
+
+	if (bUseOodle)
+	{
+		GEngine->Exec(NULL, TEXT("Oodle ForceEnable Default"));
+		UUnitTestNetConnection::bForceEnableHandler = bForceEnableDefault;
+	}
+
+
+	return bReturnVal;
+}
+
+void UPacketLimitTest::ExecuteClientUnitTest()
+{
+	bool bLowLevelSend = TestStage == ELimitTestStage::LTS_LowLevel_AtLimit || TestStage == ELimitTestStage::LTS_LowLevel_OverLimit;
+	bool bBunchSend = TestStage == ELimitTestStage::LTS_Bunch_AtLimit || TestStage == ELimitTestStage::LTS_Bunch_OverLimit;
+
+	// You can't access LowLevelSend from UNetConnection, but you can from UIpConnection, as it's exported there
+	UIpConnection* IpConn = Cast<UIpConnection>(UnitConn);
+
+	if (IpConn != nullptr)
+	{
+		int32 PacketLimit = UnitConn->MaxPacket;
+		int32 SocketLimit = UnitConn->MaxPacket;
+		TArray<uint8> PacketData;
+
+		if (bBunchSend)
+		{
+			int64 FreeBits = UnitConn->SendBuffer.GetMaxBits() - MAX_BUNCH_HEADER_BITS + MAX_PACKET_TRAILER_BITS;
+
+			PacketLimit = FreeBits / 8;
+
+			check(PacketLimit > 0)
+		}
+
+		if (TestStage == ELimitTestStage::LTS_LowLevel_OverLimit || TestStage == ELimitTestStage::LTS_Bunch_OverLimit)
+		{
+			// Nudge the packet over the MaxPacket limit (accurate for LowLevel, approximate for Bunch)
+			PacketLimit++;
+			SocketLimit++;
+		}
+
+		PacketData.AddZeroed(PacketLimit);
+
+		// Randomize the packet data (except for last byte), to thwart any compression, which would cause infinite recursion below
+		// (e.g. recursively adding just zero's, would mean the same compressed size almost every time)
+		for (int32 i=0; i<PacketData.Num()-1; i++)
+		{
+			PacketData[i] = FMath::Rand() % 255;
+		}
+
+
+		// Iteratively run 'test' sends, where the packet is passed through all the netcode but not sent,
+		// unless the final (post-PacketHandler) packet size matches SocketLimit.
+		bool bPacketAtLimit = false;
+		int32 TryCount = 0;
+		int32 SendDelta = 0;
+
+		// Blocks all socket sends not matching SocketLimit
+		TargetSocketSendSize = SocketLimit;
+
+		while (!bPacketAtLimit && TryCount < 16)
+		{
+			if (bLowLevelSend)
+			{
+				IpConn->LowLevelSend(PacketData.GetData(), PacketData.Num(), PacketData.Num() * 8);
+			}
+			else if (bBunchSend)
+			{
+				UUnitTestNetConnection* UnitTestConn = CastChecked<UUnitTestNetConnection>(UnitConn);
+
+				// If the bunch is to go over the limit, disable asserts
+				bool bBunchOverLimit = TestStage == ELimitTestStage::LTS_Bunch_OverLimit;
+
+				UnitTestConn->bDisableValidateSend = bBunchOverLimit;
+
+				UnitConn->FlushNet();
+
+
+				int32 ControlBunchSequence = 0;
+				FOutBunch* TestBunch = NUTNet::CreateChannelBunch(ControlBunchSequence, UnitConn, CHTYPE_Control, 0);
+
+				TestBunch->Serialize(PacketData.GetData(), PacketData.Num());
+
+				UnitConn->SendRawBunch(*TestBunch, false);
+
+
+				if (bBunchOverLimit)
+				{
+					// For a successful test, the bunch must cause a send error
+					if (UnitConn->SendBuffer.IsError())
+					{
+						bPacketAtLimit = true;
+
+						UNIT_LOG(ELogType::StatusSuccess, TEXT("Detected successful bunch overflow. Moving to next test."));
+						NextTestStage();
+					}
+					else
+					{
+						UNIT_LOG(ELogType::StatusFailure, TEXT("Failed to detect bunch overflow, when one was expected."));
+						VerificationState = EUnitTestVerification::VerifiedNeedsUpdate;
+					}
+
+					break;
+				}
+
+
+				UnitConn->FlushNet();
+
+				UnitTestConn->bDisableValidateSend = false;
+			}
+
+
+			// If PacketHandlers have increased/decreased final packet size away from SocketLimit, trim/pad the packet and retry
+			SendDelta = FMath::Max(1, (SendDelta == 0 ? FMath::Abs(LastSocketSendSize - SocketLimit) : SendDelta / 2));
+
+			if (LastSocketSendSize > SocketLimit)
+			{
+				PacketData.RemoveAt(PacketData.Num()-1-SendDelta, SendDelta, false);
+			}
+			else if (LastSocketSendSize < SocketLimit)
+			{
+				PacketData.InsertZeroed(PacketData.Num()-1, SendDelta);
+
+				for (int32 i=PacketData.Num()-1-SendDelta; i<SendDelta; i++)
+				{
+					PacketData[i] = FMath::Rand() % 255;
+				}
+			}
+			else // if (LastSocketSendSize == SocketLimit)
+			{
+				// Packet successfully sent
+				bPacketAtLimit = true;
+			}
+
+			TryCount++;
+		}
+
+
+		// Re-enable sending packets
+		TargetSocketSendSize = 0;
+
+		if (!bPacketAtLimit)
+		{
+			UNIT_LOG(ELogType::StatusFailure, TEXT("Failed to send packet - reached packet testing iteration limit."));
+			VerificationState = EUnitTestVerification::VerifiedUnreliable;
+		}
+	}
+}
+
+void UPacketLimitTest::NotifySendRawPacket(void* Data, int32 Count, bool& bBlockSend)
+{
+	LastSendSize = Count;
+}
+
+void UPacketLimitTest::NotifySocketSendRawPacket(void* Data, int32 Count, bool& bBlockSend)
+{
+	LastSocketSendSize = Count;
+
+	if (TargetSocketSendSize > 0)
+	{
+		if (Count == TargetSocketSendSize)
+		{
+			UNIT_LOG(, TEXT("Packet passed size filter of '%i' bytes."), TargetSocketSendSize);
+		}
+		else if (!bBlockSend)
+		{
+			bBlockSend = true;
+		}
+	}
+}
+
+void UPacketLimitTest::NotifyProcessLog(TWeakPtr<FUnitTestProcess> InProcess, const TArray<FString>& InLogLines)
+{
+	Super::NotifyProcessLog(InProcess, InLogLines);
+
+	bool bMoveToNextStage = false;
+
+	if (InProcess.HasSameObject(ServerHandle.Pin().Get()))
+	{
+		for (auto CurLine : InLogLines)
+		{
+			const bool bPacketOverflow = CurLine.Contains(TEXT("LogNet:Warning: UDP recvfrom error: 12 (SE_EMSGSIZE) from "));
+			const bool bPacketReceived = CurLine.Contains(TEXT(" Malformed_Packet: Received packet with 0's in last byte of packet"));
+			const bool bBadControlMsg = CurLine.Contains(TEXT("LogNet:Error: Received unknown control channel message"));
+
+			bool bUnexpectedResult = false;
+
+			if (TestStage == ELimitTestStage::LTS_LowLevel_AtLimit)
+			{
+				bUnexpectedResult = bPacketOverflow || bBadControlMsg;
+
+				if (bPacketReceived)
+				{
+					UNIT_LOG(ELogType::StatusSuccess, TEXT("Detected successful packet send at limit. Moving to next test."));
+					bMoveToNextStage = true;
+				}
+			}
+			else if (TestStage == ELimitTestStage::LTS_LowLevel_OverLimit)
+			{
+				bUnexpectedResult = bPacketReceived || bBadControlMsg;
+
+				if (bPacketOverflow)
+				{
+					UNIT_LOG(ELogType::StatusSuccess, TEXT("Detected successful packet overflow. Moving to next test."));
+					bMoveToNextStage = true;
+				}
+			}
+			else if (TestStage == ELimitTestStage::LTS_Bunch_AtLimit)
+			{
+				bUnexpectedResult = bPacketOverflow || bPacketReceived;
+
+				if (bBadControlMsg)
+				{
+					UNIT_LOG(ELogType::StatusSuccess, TEXT("Detected successful bunch send at limit. Moving to next test."));
+					bMoveToNextStage = true;
+				}
+			}
+
+
+			if (bUnexpectedResult)
+			{
+				UNIT_LOG(ELogType::StatusFailure, TEXT("Detected unexpected log result for test stage '%s'."),
+							*GetLimitTestStageName(TestStage));
+
+				UNIT_LOG(ELogType::StatusFailure, TEXT("Values: bPacketOverflow: %i, bPacketReceived: %i, bBadControlMsg: %i"),
+							(int32)bPacketOverflow, (int32)bPacketReceived, (int32)bBadControlMsg);
+			}
+		}
+	}
+
+	if (bMoveToNextStage)
+	{
+		NextTestStage();
+	}
+}
+
+void UPacketLimitTest::NextTestStage()
+{
+	((uint8&)TestStage)++;
+
+	if (TestStage >= ELimitTestStage::LTS_MAX)
+	{
+		UNIT_LOG(ELogType::StatusImportant, TEXT("Testing complete."));
+
+		VerificationState = EUnitTestVerification::VerifiedFixed;
+	}
+	else
+	{
+		UNIT_LOG(, TEXT("Advancing TestStage to: %s"), *GetLimitTestStageName(TestStage));
+
+		TriggerAutoReconnect();
+	}
+}
+
diff --git a/Engine/Plugins/NetcodeUnitTest/NUTUnrealEngine4/Source/NUTUnrealEngine4/Private/UnitTests/PacketLimitTest_Oodle.cpp b/Engine/Plugins/NetcodeUnitTest/NUTUnrealEngine4/Source/NUTUnrealEngine4/Private/UnitTests/PacketLimitTest_Oodle.cpp
new file mode 100644
--- /dev/null
+++ b/Engine/Plugins/NetcodeUnitTest/NUTUnrealEngine4/Source/NUTUnrealEngine4/Private/UnitTests/PacketLimitTest_Oodle.cpp
@@ -0,0 +1,27 @@
+// Copyright 1998-2016 Epic Games, Inc. All Rights Reserved.
+
+#include "NUTUnrealEngine4PCH.h"
+
+#include "UnitTests/PacketLimitTest_Oodle.h"
+
+#include "UnitTestEnvironment.h"
+
+/**
+ * UPacketLimitTest_Oodle
+ */
+UPacketLimitTest_Oodle::UPacketLimitTest_Oodle(const FObjectInitializer& ObjectInitializer)
+	: Super(ObjectInitializer)
+{
+	UnitTestName = TEXT("PacketLimitTest_Oodle");
+
+	UnitTestDate = FDateTime(2016, 1, 10);
+
+	bWorkInProgress = true;
+
+	// @todo #JohnBExploitCL: Bugtracking/changelist notes
+
+	ExpectedResult.Add(TEXT("ShooterGame"), EUnitTestVerification::VerifiedNotFixed);
+
+
+	bUseOodle = true;
+}
diff --git a/Engine/Plugins/NetcodeUnitTest/NetcodeUnitTest/NetcodeUnitTest.uplugin b/Engine/Plugins/NetcodeUnitTest/NetcodeUnitTest/NetcodeUnitTest.uplugin
--- a/Engine/Plugins/NetcodeUnitTest/NetcodeUnitTest/NetcodeUnitTest.uplugin
+++ b/Engine/Plugins/NetcodeUnitTest/NetcodeUnitTest/NetcodeUnitTest.uplugin
@@ -23,8 +23,10 @@
 			"WhitelistPlatforms" :
 			[
 				"Win64",
-				"Win32"
+				"Win32",
+				"Linux"
 			]
 		}
 	]
-}
\ No newline at end of file
+}
+
diff --git a/Engine/Plugins/NetcodeUnitTest/NetcodeUnitTest/Source/NetcodeUnitTest/Classes/ClientUnitTest.h b/Engine/Plugins/NetcodeUnitTest/NetcodeUnitTest/Source/NetcodeUnitTest/Classes/ClientUnitTest.h
--- a/Engine/Plugins/NetcodeUnitTest/NetcodeUnitTest/Source/NetcodeUnitTest/Classes/ClientUnitTest.h
+++ b/Engine/Plugins/NetcodeUnitTest/NetcodeUnitTest/Source/NetcodeUnitTest/Classes/ClientUnitTest.h
@@ -43,14 +43,15 @@ enum class EUnitTestFlags : uint32 // NOTE: If you change from uint32, you need
 	SendRPCs				= 0x00000020,	// Whether or not to allow RPC sending (NOT blocked by default, @todo JohnB: Add send hook)
 	SkipControlJoin			= 0x00000040,	// Whether or not to skip sending NMT_Join upon connect (or NMT_BeaconJoin for beacons)
 	BeaconConnect			= 0x00000080,	// Whether or not to connect to the servers beacon (greatly limits the connection)
+	AutoReconnect			= 0x00000100,	// Whether or not to auto-reconnect on server disconnect (NOTE: Won't catch all disconnects)
 
 	/** Unit test state-setup/requirements/prerequisites */
-	RequirePlayerController	= 0x00000100,	// Whether or not to wait for the PlayerController, before triggering ExecuteClientUnitTest
-	RequirePawn				= 0x00000200,	// Whether or not to wait for the PlayerController's pawn, before ExecuteClientUnitTest
-	RequirePing				= 0x00000400,	// Whether or not to wait for a ping round-trip, before triggering ExecuteClientUnitTest
-	RequireNUTActor			= 0x00000800,	// Whether or not to wait for the NUTActor, before triggering ExecuteClientUnitTest
-	RequireBeacon			= 0x00001000,	// Whether or not to wait for beacon replication, before triggering ExecuteClientUnitTest
-	RequireCustom			= 0x00002000,	// Whether or not ExecuteClientUnitTest will be executed manually, within the unit test
+	RequirePlayerController	= 0x00000200,	// Whether or not to wait for the PlayerController, before triggering ExecuteClientUnitTest
+	RequirePawn				= 0x00000400,	// Whether or not to wait for the PlayerController's pawn, before ExecuteClientUnitTest
+	RequirePing				= 0x00000800,	// Whether or not to wait for a ping round-trip, before triggering ExecuteClientUnitTest
+	RequireNUTActor			= 0x00001000,	// Whether or not to wait for the NUTActor, before triggering ExecuteClientUnitTest
+	RequireBeacon			= 0x00002000,	// Whether or not to wait for beacon replication, before triggering ExecuteClientUnitTest
+	RequireCustom			= 0x00004000,	// Whether or not ExecuteClientUnitTest will be executed manually, within the unit test
 
 	RequirementsMask		= RequirePlayerController | RequirePawn | RequirePing | RequireNUTActor | RequireBeacon | RequireCustom,
 
@@ -331,10 +332,25 @@ public:
 	 * If EUnitTestFlags::CaptureSendRaw is set, this is triggered for every packet sent to the server
 	 * NOTE: Don't consider data safe to modify (will need to modify the implementation, if that is desired)
 	 *
+	 * @param Data			The raw data/packet being sent
+	 * @param Count			The amount of data being sent
+	 * @param bBlockSend	Whether or not to block the send (defaults to false)
+	 */
+	virtual void NotifySendRawPacket(void* Data, int32 Count, bool& bBlockSend);
+
+	/**
+	 * If EUnitTestFlags::CaptureSendRaw is set, this is triggered for every packet sent to the server
+	 * IMPORTANT: This is like the above function, except this occurs AFTER PacketHandler's have had a chance to modify packet data
+	 *
+	 * NOTE: Don't consider data safe to modify (will need to modify the implementation, if that is desired)
+	 *
 	 * @param Data		The raw data/packet being sent
 	 * @param Count		The amount of data being sent
+	 * @param bBlockSend	Whether or not to block the send (defaults to false)
 	 */
-	virtual void NotifySendRawPacket(void* Data, int32 Count);
+	virtual void NotifySocketSendRawPacket(void* Data, int32 Count, bool& bBlockSend)
+	{
+	}
 
 	/**
 	 * Bunches received on the control channel. These need to be parsed manually,
@@ -507,6 +523,11 @@ protected:
 	 */
 	virtual void CleanupFakeClient();
 
+	/**
+	 * Triggers an auto-reconnect (disconnect/reconnect) of the fake client
+	 */
+	void TriggerAutoReconnect();
+
 
 	/**
 	 * Starts the server process for a particular unit test
diff --git a/Engine/Plugins/NetcodeUnitTest/NetcodeUnitTest/Source/NetcodeUnitTest/Classes/Net/UnitTestNetConnection.h b/Engine/Plugins/NetcodeUnitTest/NetcodeUnitTest/Source/NetcodeUnitTest/Classes/Net/UnitTestNetConnection.h
--- a/Engine/Plugins/NetcodeUnitTest/NetcodeUnitTest/Source/NetcodeUnitTest/Classes/Net/UnitTestNetConnection.h
+++ b/Engine/Plugins/NetcodeUnitTest/NetcodeUnitTest/Source/NetcodeUnitTest/Classes/Net/UnitTestNetConnection.h
@@ -3,6 +3,9 @@
 #pragma once
 
 #include "IpConnection.h"
+
+#include "NUTUtilNet.h"
+
 #include "UnitTestNetConnection.generated.h"
 
 
@@ -10,7 +13,7 @@
  * A net connection for enabling unit testing through barebones/minimal client connections
  */
 UCLASS(transient)
-class UUnitTestNetConnection : public UIpConnection
+class NETCODEUNITTEST_API UUnitTestNetConnection : public UIpConnection
 {
 	GENERATED_UCLASS_BODY()
 
@@ -25,8 +28,12 @@ class UUnitTestNetConnection : public UIpConnection
 								int32 InMaxPacket=0) override;
 #endif
 
+	virtual void InitHandler() override;
+
+
+	virtual void LowLevelSend(void* Data, int32 CountBytes, int32 CountBits) override;
 
-	virtual void LowLevelSend(void* Data, int32 Count) override;
+	virtual void ValidateSendBuffer() override;
 
 	virtual void ReceivedRawPacket(void* Data, int32 Count) override;
 
@@ -36,10 +43,11 @@ class UUnitTestNetConnection : public UIpConnection
 	/**
 	 * Delegate for hooking 'LowLevelSend'
 	 *
-	 * @param Data		The data being sent
-	 * @param Count		The number of bytes being sent
+	 * @param Data			The data being sent
+	 * @param Count			The number of bytes being sent
+	 * @param bBlockSend	Whether or not to block the send (defaults to false)
 	 */
-	DECLARE_DELEGATE_TwoParams(FLowLevelSendDel, void* /*Data*/, int32 /*Count*/);
+	DECLARE_DELEGATE_ThreeParams(FLowLevelSendDel, void* /*Data*/, int32 /*Count*/, bool& /*bBlockSend*/);
 
 	/**
 	 * Delegate for hooking 'ReceivedRawPacket'
@@ -59,6 +67,7 @@ class UUnitTestNetConnection : public UIpConnection
 	DECLARE_DELEGATE_RetVal_TwoParams(bool, FOnReplicatedActorSpawn, UClass* /*ActorClass*/, bool /*bActorChannel*/);
 
 
+public:
 	/** Delegate for hooking LowLevelSend */
 	FLowLevelSendDel		LowLevelSendDel;
 
@@ -67,6 +76,17 @@ class UUnitTestNetConnection : public UIpConnection
 
 	/** Delegate for notifying on replicated actor creation */
 	FOnReplicatedActorSpawn	ReplicatedActorSpawnDel;
+
+	/** Socket hook - for hooking socket-level events */
+	FSocketHook				SocketHook;
+
+	/** Whether or not to override error detection within ValidateSendBuffer */
+	bool					bDisableValidateSend;
+
+
+public:
+	/** Whether or not newly-created instances of this class, should force-enable packet handlers */
+	static bool				bForceEnableHandler;
 };
 
 
diff --git a/Engine/Plugins/NetcodeUnitTest/NetcodeUnitTest/Source/NetcodeUnitTest/NetcodeUnitTest.Build.cs b/Engine/Plugins/NetcodeUnitTest/NetcodeUnitTest/Source/NetcodeUnitTest/NetcodeUnitTest.Build.cs
--- a/Engine/Plugins/NetcodeUnitTest/NetcodeUnitTest/Source/NetcodeUnitTest/NetcodeUnitTest.Build.cs
+++ b/Engine/Plugins/NetcodeUnitTest/NetcodeUnitTest/Source/NetcodeUnitTest/NetcodeUnitTest.Build.cs
@@ -15,12 +15,21 @@ namespace UnrealBuildTool.Rules
 					"Core",
 					"CoreUObject",
 					"Engine",
-					"EngineSettings",
-					"OnlineSubsystem",
 					"OnlineSubsystemUtils",
+					"Sockets"
+				}
+			);
+
+			PrivateDependencyModuleNames.AddRange
+			(
+				new string[]
+				{
+					"EngineSettings",
 					"InputCore",
-					"SlateCore",
-					"Slate"
+					"OnlineSubsystem",
+					"PacketHandler",
+					"Slate",
+					"SlateCore"
 				}
 			);
 
diff --git a/Engine/Plugins/NetcodeUnitTest/NetcodeUnitTest/Source/NetcodeUnitTest/Private/ClientUnitTest.cpp b/Engine/Plugins/NetcodeUnitTest/NetcodeUnitTest/Source/NetcodeUnitTest/Private/ClientUnitTest.cpp
--- a/Engine/Plugins/NetcodeUnitTest/NetcodeUnitTest/Source/NetcodeUnitTest/Private/ClientUnitTest.cpp
+++ b/Engine/Plugins/NetcodeUnitTest/NetcodeUnitTest/Source/NetcodeUnitTest/Private/ClientUnitTest.cpp
@@ -207,34 +207,43 @@ void UClientUnitTest::NotifyNetActor(UActorChannel* ActorChannel, AActor* Actor)
 
 void UClientUnitTest::NotifyNetworkFailure(ENetworkFailure::Type FailureType, const FString& ErrorString)
 {
-	// Only process this error, if a result has not already been returned
-	if (VerificationState == EUnitTestVerification::Unverified)
+	if (!!(UnitTestFlags & EUnitTestFlags::AutoReconnect))
 	{
-		FString LogMsg = FString::Printf(TEXT("Got network failure of type '%s' (%s)"), ENetworkFailure::ToString(FailureType),
-											*ErrorString);
+		UNIT_LOG(ELogType::StatusImportant, TEXT("Detected fake client disconnect when AutoReconnect is enabled. Reconnecting."));
 
-		if (!(UnitTestFlags & EUnitTestFlags::IgnoreDisconnect))
+		TriggerAutoReconnect();
+	}
+	else
+	{
+		// Only process this error, if a result has not already been returned
+		if (VerificationState == EUnitTestVerification::Unverified)
 		{
-			LogMsg += TEXT(", marking unit test as needing update.");
+			FString LogMsg = FString::Printf(TEXT("Got network failure of type '%s' (%s)"), ENetworkFailure::ToString(FailureType),
+												*ErrorString);
 
-			UNIT_LOG(ELogType::StatusFailure | ELogType::StyleBold, TEXT("%s"), *LogMsg);
-			UNIT_STATUS_LOG(ELogType::StatusFailure | ELogType::StatusVerbose | ELogType::StyleBold, TEXT("%s"), *LogMsg);
+			if (!(UnitTestFlags & EUnitTestFlags::IgnoreDisconnect))
+			{
+				LogMsg += TEXT(", marking unit test as needing update.");
 
-			VerificationState = EUnitTestVerification::VerifiedNeedsUpdate;
-		}
-		else
-		{
-			LogMsg += TEXT(".");
+				UNIT_LOG(ELogType::StatusFailure | ELogType::StyleBold, TEXT("%s"), *LogMsg);
+				UNIT_STATUS_LOG(ELogType::StatusFailure | ELogType::StatusVerbose | ELogType::StyleBold, TEXT("%s"), *LogMsg);
+
+				VerificationState = EUnitTestVerification::VerifiedNeedsUpdate;
+			}
+			else
+			{
+				LogMsg += TEXT(".");
 
-			UNIT_LOG(ELogType::StatusWarning, TEXT("%s"), *LogMsg);
-			UNIT_STATUS_LOG(ELogType::StatusWarning | ELogType::StatusVerbose, TEXT("%s"), *LogMsg);
+				UNIT_LOG(ELogType::StatusWarning, TEXT("%s"), *LogMsg);
+				UNIT_STATUS_LOG(ELogType::StatusWarning | ELogType::StatusVerbose, TEXT("%s"), *LogMsg);
+			}
 		}
-	}
 
-	// Shut down the fake client now (relevant for developer mode)
-	if (VerificationState != EUnitTestVerification::Unverified)
-	{
-		CleanupFakeClient();
+		// Shut down the fake client now (relevant for developer mode)
+		if (VerificationState != EUnitTestVerification::Unverified)
+		{
+			CleanupFakeClient();
+		}
 	}
 }
 
@@ -250,7 +259,7 @@ void UClientUnitTest::NotifyReceivedRawPacket(void* Data, int32& Count)
 	}
 }
 
-void UClientUnitTest::NotifySendRawPacket(void* Data, int32 Count)
+void UClientUnitTest::NotifySendRawPacket(void* Data, int32 Count, bool& bBlockSend)
 {
 	if (!!(UnitTestFlags & EUnitTestFlags::DumpSendRaw))
 	{
@@ -595,6 +604,7 @@ void UClientUnitTest::NotifySuspendRequest()
 	// @todo #JohnBFeature: Currently on suspends the server, suspend the client as well, and also add more granularity
 	//						(deciding which to suspend)
 
+#if PLATFORM_WINDOWS
 	TSharedPtr<FUnitTestProcess> CurProcess = (ServerHandle.IsValid() ? ServerHandle.Pin() : NULL);
 
 	if (CurProcess.IsValid())
@@ -666,6 +676,9 @@ void UClientUnitTest::NotifySuspendRequest()
 			}
 		}
 	}
+#else
+	UNIT_LOG(ELogType::StatusImportant, TEXT("Suspend/Resume is only supported in Windows."));
+#endif
 }
 
 void UClientUnitTest::NotifyProcessSuspendState(TWeakPtr<FUnitTestProcess> InProcess, ESuspendState InSuspendState)
@@ -1226,6 +1239,7 @@ bool UClientUnitTest::ConnectFakeClient(FUniqueNetIdRepl* InNetID/*=NULL*/)
 					if (!!(UnitTestFlags & EUnitTestFlags::CaptureSendRaw))
 					{
 						CurUnitConn->LowLevelSendDel.BindUObject(this, &UClientUnitTest::NotifySendRawPacket);
+						CurUnitConn->SocketHook.SendToDel.BindUObject(this, &UClientUnitTest::NotifySocketSendRawPacket);
 					}
 
 					CurUnitConn->ReplicatedActorSpawnDel.BindUObject(this, &UClientUnitTest::NotifyAllowNetActor);
@@ -1308,7 +1322,12 @@ void UClientUnitTest::CleanupFakeClient()
 	UnitNetDriver = NULL;
 	UnitPC = NULL;
 	UnitConn = NULL;
+	bUnitPawnSetup = false;
 	UnitNUTActor = NULL;
+	UnitBeacon = nullptr;
+	bReceivedPong = false;
+	ControlBunchSequence = 0;
+	PendingNetActorChans.Empty();
 
 	if (UnitNotify != NULL)
 	{
@@ -1333,6 +1352,14 @@ void UClientUnitTest::CleanupFakeClient()
 	}
 }
 
+void UClientUnitTest::TriggerAutoReconnect()
+{
+	UNIT_LOG(ELogType::StatusImportant, TEXT("Performing Auto-Reconnect."))
+
+	CleanupFakeClient();
+	ConnectFakeClient();
+}
+
 
 void UClientUnitTest::StartUnitTestServer()
 {
diff --git a/Engine/Plugins/NetcodeUnitTest/NetcodeUnitTest/Source/NetcodeUnitTest/Private/NUTActor.cpp b/Engine/Plugins/NetcodeUnitTest/NetcodeUnitTest/Source/NetcodeUnitTest/Private/NUTActor.cpp
--- a/Engine/Plugins/NetcodeUnitTest/NetcodeUnitTest/Source/NetcodeUnitTest/Private/NUTActor.cpp
+++ b/Engine/Plugins/NetcodeUnitTest/NetcodeUnitTest/Source/NetcodeUnitTest/Private/NUTActor.cpp
@@ -229,6 +229,7 @@ bool ANUTActor::NotifyControlMessage(UNetConnection* Connection, uint8 MessageTy
 		// Suspend the game, until a resume request is received (used for giving time, to attach a debugger)
 		else if (CmdType == ENUTControlCommand::SuspendProcess)
 		{
+#if PLATFORM_WINDOWS
 			UE_LOG(LogUnitTest, Log, TEXT("Suspend start."));
 
 
@@ -273,6 +274,9 @@ bool ANUTActor::NotifyControlMessage(UNetConnection* Connection, uint8 MessageTy
 			ResumePipe.Destroy();
 
 			UE_LOG(LogUnitTest, Log, TEXT("Suspend end."));
+#else
+			UE_LOG(LogUnitTest, Log, TEXT("Suspend/Resume is only supported in Windows."));
+#endif
 		}
 
 		bHandledMessage = true;
diff --git a/Engine/Plugins/NetcodeUnitTest/NetcodeUnitTest/Source/NetcodeUnitTest/Private/NUTUtilReflection.cpp b/Engine/Plugins/NetcodeUnitTest/NetcodeUnitTest/Source/NetcodeUnitTest/Private/NUTUtilReflection.cpp
--- a/Engine/Plugins/NetcodeUnitTest/NetcodeUnitTest/Source/NetcodeUnitTest/Private/NUTUtilReflection.cpp
+++ b/Engine/Plugins/NetcodeUnitTest/NetcodeUnitTest/Source/NetcodeUnitTest/Private/NUTUtilReflection.cpp
@@ -78,8 +78,8 @@ FVMReflection::FVMReflection(const FVMReflection& ToCopy)
 	, bVerifiedFieldType(ToCopy.bVerifiedFieldType)
 	, bSetArrayElement(ToCopy.bSetArrayElement)
 	, bNextActionMustBeCast(ToCopy.bNextActionMustBeCast)
-	, NextActionError(ToCopy.NextActionError)
 	, bIsError(ToCopy.bIsError)
+	, NextActionError(ToCopy.NextActionError)
 	, bOutError(NULL)
 	, History()
 	, OutHistoryPtr(NULL)
@@ -365,7 +365,7 @@ FVMReflection& FVMReflection::operator [](const ANSICHAR* InFieldType)
 			#define FIELD_TYPE_CHECK_FIRST(TypeString, TypeProperty) \
 				if (ExpectedFieldType == TypeString) \
 				{ \
-					if (ActualFieldType->IsA(TypeProperty##::StaticClass())) \
+					if (ActualFieldType->IsA(TypeProperty::StaticClass())) \
 					{ \
 						bTypeValid = true; \
 					} \
diff --git a/Engine/Plugins/NetcodeUnitTest/NetcodeUnitTest/Source/NetcodeUnitTest/Private/Net/NUTUtilNet.cpp b/Engine/Plugins/NetcodeUnitTest/NetcodeUnitTest/Source/NetcodeUnitTest/Private/Net/NUTUtilNet.cpp
--- a/Engine/Plugins/NetcodeUnitTest/NetcodeUnitTest/Source/NetcodeUnitTest/Private/Net/NUTUtilNet.cpp
+++ b/Engine/Plugins/NetcodeUnitTest/NetcodeUnitTest/Source/NetcodeUnitTest/Private/Net/NUTUtilNet.cpp
@@ -30,6 +30,163 @@ TArray<UWorld*> PendingUnitWorldCleanup;
 TArray<FWorldTickHook*> ActiveTickHooks;
 
 
+/**
+ * FSocketHook
+ */
+
+bool FSocketHook::Close()
+{
+	return HookedSocket->Close();
+}
+
+bool FSocketHook::Bind(const FInternetAddr& Addr)
+{
+	return HookedSocket->Bind(Addr);
+}
+
+bool FSocketHook::Connect(const FInternetAddr& Addr)
+{
+	return HookedSocket->Connect(Addr);
+}
+
+bool FSocketHook::Listen(int32 MaxBacklog)
+{
+	return HookedSocket->Listen(MaxBacklog);
+}
+
+bool FSocketHook::HasPendingConnection(bool& bHasPendingConnection)
+{
+	return HookedSocket->HasPendingConnection(bHasPendingConnection);
+}
+
+bool FSocketHook::HasPendingData(uint32& PendingDataSize)
+{
+	return HookedSocket->HasPendingData(PendingDataSize);
+}
+
+class FSocket* FSocketHook::Accept(const FString& InSocketDescription)
+{
+	return HookedSocket->Accept(InSocketDescription);
+}
+
+class FSocket* FSocketHook::Accept(FInternetAddr& OutAddr, const FString& InSocketDescription)
+{
+	return HookedSocket->Accept(OutAddr, InSocketDescription);
+}
+
+bool FSocketHook::SendTo(const uint8* Data, int32 Count, int32& BytesSent, const FInternetAddr& Destination)
+{
+	bool bReturnVal = false;
+	bool bBlockSend = false;
+
+	SendToDel.ExecuteIfBound((void*)Data, Count, bBlockSend);
+
+	if (!bBlockSend)
+	{
+		bReturnVal = HookedSocket->SendTo(Data, Count, BytesSent, Destination);
+	}
+
+	bLastSendToBlocked = bBlockSend;
+
+	return bReturnVal;
+}
+
+bool FSocketHook::Send(const uint8* Data, int32 Count, int32& BytesSent)
+{
+	return HookedSocket->Send(Data, Count, BytesSent);
+}
+
+bool FSocketHook::RecvFrom(uint8* Data, int32 BufferSize, int32& BytesRead, FInternetAddr& Source, ESocketReceiveFlags::Type Flags)
+{
+	return HookedSocket->RecvFrom(Data, BufferSize, BytesRead, Source, Flags);
+}
+
+bool FSocketHook::Recv(uint8* Data, int32 BufferSize, int32& BytesRead, ESocketReceiveFlags::Type Flags)
+{
+	return HookedSocket->Recv(Data, BufferSize, BytesRead, Flags);
+}
+
+bool FSocketHook::Wait(ESocketWaitConditions::Type Condition, FTimespan WaitTime)
+{
+	return HookedSocket->Wait(Condition, WaitTime);
+}
+
+ESocketConnectionState FSocketHook::GetConnectionState()
+{
+	return HookedSocket->GetConnectionState();
+}
+
+void FSocketHook::GetAddress(FInternetAddr& OutAddr)
+{
+	HookedSocket->GetAddress(OutAddr);
+}
+
+bool FSocketHook::GetPeerAddress(FInternetAddr& OutAddr)
+{
+	return HookedSocket->GetPeerAddress(OutAddr);
+}
+
+bool FSocketHook::SetNonBlocking(bool bIsNonBlocking)
+{
+	return HookedSocket->SetNonBlocking(bIsNonBlocking);
+}
+
+bool FSocketHook::SetBroadcast(bool bAllowBroadcast)
+{
+	return HookedSocket->SetBroadcast(bAllowBroadcast);
+}
+
+bool FSocketHook::JoinMulticastGroup(const FInternetAddr& GroupAddress)
+{
+	return HookedSocket->JoinMulticastGroup(GroupAddress);
+}
+
+bool FSocketHook::LeaveMulticastGroup(const FInternetAddr& GroupAddress)
+{
+	return HookedSocket->LeaveMulticastGroup(GroupAddress);
+}
+
+bool FSocketHook::SetMulticastLoopback(bool bLoopback)
+{
+	return HookedSocket->SetMulticastLoopback(bLoopback);
+}
+
+bool FSocketHook::SetMulticastTtl(uint8 TimeToLive)
+{
+	return HookedSocket->SetMulticastTtl(TimeToLive);
+}
+
+bool FSocketHook::SetReuseAddr(bool bAllowReuse)
+{
+	return HookedSocket->SetReuseAddr(bAllowReuse);
+}
+
+bool FSocketHook::SetLinger(bool bShouldLinger, int32 Timeout)
+{
+	return HookedSocket->SetLinger(bShouldLinger, Timeout);
+}
+
+bool FSocketHook::SetRecvErr(bool bUseErrorQueue)
+{
+	return HookedSocket->SetRecvErr(bUseErrorQueue);
+}
+
+bool FSocketHook::SetSendBufferSize(int32 Size, int32& NewSize)
+{
+	return HookedSocket->SetSendBufferSize(Size, NewSize);
+}
+
+bool FSocketHook::SetReceiveBufferSize(int32 Size, int32& NewSize)
+{
+	return HookedSocket->SetReceiveBufferSize(Size, NewSize);
+}
+
+int32 FSocketHook::GetPortNo()
+{
+	return HookedSocket->GetPortNo();
+}
+
+
 /**
  * FNetworkNotifyHook
  */
@@ -344,6 +501,13 @@ bool NUTNet::CreateFakePlayer(UWorld* InWorld, UNetDriver*& InNetDriver, FString
 
 			UE_LOG(LogUnitTest, Log, TEXT("Successfully kicked off connect to IP '%s'"), *ServerIP);
 
+#if TARGET_UE4_CL >= CL_STATELESSCONNECT
+			if (TargetConn->StatelessConnectComponent.IsValid())
+			{
+				TargetConn->StatelessConnectComponent.Pin()->SendInitialConnect();
+			}
+#endif
+
 
 			int ControlBunchSequence = 0;
 
@@ -588,6 +752,11 @@ void NUTNet::CleanupUnitTestWorlds()
 
 	PendingUnitWorldCleanup.Empty();
 
+
+	// @todo #JohnB: There are some actor components that get marked as having begun play, and you need to clean them up here before GC,
+	//					otherwise a crash is triggered
+
+
 	// Immediately garbage collect remaining objects, to finish net driver cleanup
 	CollectGarbage(GARBAGE_COLLECTION_KEEPFLAGS, true);
 }
@@ -602,7 +771,7 @@ bool NUTNet::IsSteamNetDriverAvailable()
 	bool bReturnVal = false;
 	UGameEngine* GameEngine = Cast<UGameEngine>(GEngine);
 
-	if (GameEngine != NULL)
+	if (GameEngine != nullptr)
 	{
 		bool bFoundSteamDriver = false;
 		const TCHAR* SteamDriverClassName = TEXT("OnlineSubsystemSteam.SteamNetDriver");
@@ -622,18 +791,17 @@ bool NUTNet::IsSteamNetDriverAvailable()
 
 		if (bFoundSteamDriver)
 		{
-			UClass* SteamNetDriverClass = StaticLoadClass(UNetDriver::StaticClass(), NULL, SteamDriverClassName, NULL, LOAD_Quiet);
+			UClass* SteamNetDriverClass = StaticLoadClass(UNetDriver::StaticClass(), nullptr, SteamDriverClassName, nullptr,
+															LOAD_Quiet);
 
-			if (SteamDriverClassName != NULL)
+			if (SteamDriverClassName != nullptr && SteamNetDriverClass != nullptr)
 			{
 				UNetDriver* SteamNetDriverDef = Cast<UNetDriver>(SteamNetDriverClass->GetDefaultObject());
 
-				bReturnVal = SteamNetDriverDef != NULL && SteamNetDriverDef->IsAvailable();
+				bReturnVal = SteamNetDriverDef != nullptr && SteamNetDriverDef->IsAvailable();
 			}
 		}
 	}
 
 	return bReturnVal;
 }
-
-
diff --git a/Engine/Plugins/NetcodeUnitTest/NetcodeUnitTest/Source/NetcodeUnitTest/Private/Net/UnitTestNetConnection.cpp b/Engine/Plugins/NetcodeUnitTest/NetcodeUnitTest/Source/NetcodeUnitTest/Private/Net/UnitTestNetConnection.cpp
--- a/Engine/Plugins/NetcodeUnitTest/NetcodeUnitTest/Source/NetcodeUnitTest/Private/Net/UnitTestNetConnection.cpp
+++ b/Engine/Plugins/NetcodeUnitTest/NetcodeUnitTest/Source/NetcodeUnitTest/Private/Net/UnitTestNetConnection.cpp
@@ -12,6 +12,14 @@
 #include "Net/UnitTestPackageMap.h"
 #include "Net/UnitTestActorChannel.h"
 
+
+/**
+ * Static variables
+ */
+
+bool UUnitTestNetConnection::bForceEnableHandler = false;
+
+
 /**
  * Default constructor
  */
@@ -19,6 +27,7 @@ UUnitTestNetConnection::UUnitTestNetConnection(const FObjectInitializer& ObjectI
 	: Super(ObjectInitializer)
 	, LowLevelSendDel()
 	, ReceivedRawPacketDel()
+	, bDisableValidateSend(false)
 {
 }
 
@@ -55,13 +64,59 @@ void UUnitTestNetConnection::InitConnection(UNetDriver* InDriver, EConnectionSta
 #endif
 }
 
-void UUnitTestNetConnection::LowLevelSend(void* Data, int32 Count)
+void UUnitTestNetConnection::InitHandler()
 {
-	LowLevelSendDel.ExecuteIfBound(Data, Count);
+	// @todo #JohnB: Broken/redundant, after addition of stateless challenge?
+#if 0
+	// Copy-pasted from UNetConnection
+	if (bForceEnableHandler)
+	{
+		check(!Handler.IsValid());
+
+		Handler = MakeUnique<PacketHandler>();
 
-	Super::LowLevelSend(Data, Count);
+		Handler->bEnabled = true;
 
-	GSentBunch = true;
+		if (Handler.IsValid())
+		{
+			Handler::Mode Mode = Driver->ServerConnection != nullptr ? Handler::Mode::Client : Handler::Mode::Server;
+			Handler->Initialize(Mode);
+
+			MaxPacketHandlerBits = Handler->GetTotalPacketOverheadBits();
+		}
+	}
+	else
+#endif
+	{
+		Super::InitHandler();
+	}
+}
+
+void UUnitTestNetConnection::LowLevelSend(void* Data, int32 CountBytes, int32 CountBits)
+{
+	bool bBlockSend = false;
+
+	LowLevelSendDel.ExecuteIfBound(Data, CountBytes, bBlockSend);
+
+	if (!bBlockSend)
+	{
+		// Only attach the SocketHook when it is needed - not all the time
+		SocketHook.Attach(Socket);
+
+		Super::LowLevelSend(Data, CountBytes, CountBits);
+
+		SocketHook.Detach(Socket);
+
+		GSentBunch = !SocketHook.bLastSendToBlocked;
+	}
+}
+
+void UUnitTestNetConnection::ValidateSendBuffer()
+{
+	if (!bDisableValidateSend)
+	{
+		Super::ValidateSendBuffer();
+	}
 }
 
 void UUnitTestNetConnection::ReceivedRawPacket(void* Data, int32 Count)
@@ -116,4 +171,3 @@ void UUnitTestNetConnection::HandleClientPlayer(class APlayerController* PC, cla
 }
 
 
-
diff --git a/Engine/Plugins/NetcodeUnitTest/NetcodeUnitTest/Source/NetcodeUnitTest/Private/Net/UnitTestNetDriver.cpp b/Engine/Plugins/NetcodeUnitTest/NetcodeUnitTest/Source/NetcodeUnitTest/Private/Net/UnitTestNetDriver.cpp
--- a/Engine/Plugins/NetcodeUnitTest/NetcodeUnitTest/Source/NetcodeUnitTest/Private/Net/UnitTestNetDriver.cpp
+++ b/Engine/Plugins/NetcodeUnitTest/NetcodeUnitTest/Source/NetcodeUnitTest/Private/Net/UnitTestNetDriver.cpp
@@ -14,9 +14,6 @@ UUnitTestNetDriver::UUnitTestNetDriver(const FObjectInitializer& ObjectInitializ
 {
 }
 
-/**
- * @todo #JohnBDoc
- */
 bool UUnitTestNetDriver::InitConnectionClass()
 {
 	bool bReturnVal = false;
@@ -27,11 +24,6 @@ bool UUnitTestNetDriver::InitConnectionClass()
 		bReturnVal = true;
 	}
 
-	if (NetConnectionClass == NULL)
-	{
-		bReturnVal = Super::InitConnectionClass();
-	}
-
 	return bReturnVal;
 }
 
diff --git a/Engine/Plugins/NetcodeUnitTest/NetcodeUnitTest/Source/NetcodeUnitTest/Private/ProcessUnitTest.cpp b/Engine/Plugins/NetcodeUnitTest/NetcodeUnitTest/Source/NetcodeUnitTest/Private/ProcessUnitTest.cpp
--- a/Engine/Plugins/NetcodeUnitTest/NetcodeUnitTest/Source/NetcodeUnitTest/Private/ProcessUnitTest.cpp
+++ b/Engine/Plugins/NetcodeUnitTest/NetcodeUnitTest/Source/NetcodeUnitTest/Private/ProcessUnitTest.cpp
@@ -19,8 +19,8 @@
 UProcessUnitTest::UProcessUnitTest(const FObjectInitializer& ObjectInitializer)
 	: Super(ObjectInitializer)
 	, ActiveProcesses()
-	, OnSuspendStateChange()
 	, LastBlockingProcessCheck(0)
+	, OnSuspendStateChange()
 {
 }
 
@@ -523,7 +523,7 @@ bool UProcessUnitTest::IsBlockingProcessPresent(bool bLogIfFound/*=false*/)
 		{
 			if (It.Value() == SearchParentId)
 			{
-				ChildProcs.Add(It.Key());
+				ChildProcs.AddUnique(It.Key());
 			}
 		}
 	}
diff --git a/Engine/Plugins/NetcodeUnitTest/NetcodeUnitTest/Source/NetcodeUnitTest/Private/UI/SLogWidget.cpp b/Engine/Plugins/NetcodeUnitTest/NetcodeUnitTest/Source/NetcodeUnitTest/Private/UI/SLogWidget.cpp
--- a/Engine/Plugins/NetcodeUnitTest/NetcodeUnitTest/Source/NetcodeUnitTest/Private/UI/SLogWidget.cpp
+++ b/Engine/Plugins/NetcodeUnitTest/NetcodeUnitTest/Source/NetcodeUnitTest/Private/UI/SLogWidget.cpp
@@ -887,8 +887,9 @@ TSharedRef<SDockTab> SLogWidget::SpawnLogTab(const FSpawnTabArgs& InSpawnTabArgs
 									.OnClicked_Lambda(
 										[&]()
 										{
-											auto CurTabInfo = GetActiveTabInfo();
-											ScrollToText(CurTabInfo.ToSharedRef(), CurTabInfo->FindBox->GetText().ToString(), true);
+											auto ActiveTabInfo = GetActiveTabInfo();
+											ScrollToText(ActiveTabInfo.ToSharedRef(), ActiveTabInfo->FindBox->GetText().ToString(),
+															true);
 
 											return FReply::Handled();
 										})
@@ -910,8 +911,8 @@ TSharedRef<SDockTab> SLogWidget::SpawnLogTab(const FSpawnTabArgs& InSpawnTabArgs
 									.OnClicked_Lambda(
 										[&]()
 										{
-											auto CurTabInfo = GetActiveTabInfo();
-											ScrollToText(CurTabInfo.ToSharedRef(), CurTabInfo->FindBox->GetText().ToString());
+											auto ActiveTabInfo = GetActiveTabInfo();
+											ScrollToText(ActiveTabInfo.ToSharedRef(), ActiveTabInfo->FindBox->GetText().ToString());
 
 											return FReply::Handled();
 										})
@@ -945,8 +946,9 @@ TSharedRef<SDockTab> SLogWidget::SpawnLogTab(const FSpawnTabArgs& InSpawnTabArgs
 										{
 											if (InCommitType == ETextCommit::OnEnter)
 											{
-												auto CurTabInfo = GetActiveTabInfo();
-												ScrollToText(CurTabInfo.ToSharedRef(), InText.ToString(), CurTabInfo->bLastFindWasUp);
+												auto ActiveTabInfo = GetActiveTabInfo();
+												ScrollToText(ActiveTabInfo.ToSharedRef(), InText.ToString(),
+																ActiveTabInfo->bLastFindWasUp);
 											}
 										})
 								]
@@ -968,9 +970,9 @@ TSharedRef<SDockTab> SLogWidget::SpawnLogTab(const FSpawnTabArgs& InSpawnTabArgs
 				{
 					TSharedRef<SButton> CurButton = StaticCastSharedRef<SButton>(InWidget);
 					const FButtonStyle* ButtonStyle = GET_PRIVATE(SButton, CurButton, Style);
-					const FDockTabStyle& TabStyle = CALL_PROTECTED(SDockTab, ReturnVal, GetCurrentStyle)();
+					const FDockTabStyle& CurTabStyle = CALL_PROTECTED(SDockTab, ReturnVal, GetCurrentStyle)();
 
-					bFound = ButtonStyle == &TabStyle.CloseButtonStyle;
+					bFound = ButtonStyle == &CurTabStyle.CloseButtonStyle;
 				}
 
 				return bFound;
diff --git a/Engine/Plugins/NetcodeUnitTest/NetcodeUnitTest/Source/NetcodeUnitTest/Private/UnitTestCommandlet.cpp b/Engine/Plugins/NetcodeUnitTest/NetcodeUnitTest/Source/NetcodeUnitTest/Private/UnitTestCommandlet.cpp
--- a/Engine/Plugins/NetcodeUnitTest/NetcodeUnitTest/Source/NetcodeUnitTest/Private/UnitTestCommandlet.cpp
+++ b/Engine/Plugins/NetcodeUnitTest/NetcodeUnitTest/Source/NetcodeUnitTest/Private/UnitTestCommandlet.cpp
@@ -62,6 +62,18 @@ int32 UUnitTestCommandlet::Main(const FString& Params)
 #else
 	GIsRequestingExit = false;
 
+
+	// Load some core modules that LaunchEngineLoop.cpp no longer loads for commandlets
+	if (FModuleManager::Get().ModuleExists(TEXT("XMPP")))
+	{
+		FModuleManager::Get().LoadModule(TEXT("XMPP"));
+	}
+
+	FModuleManager::Get().LoadModule(TEXT("HTTP"));
+	FModuleManager::Get().LoadModule(TEXT("OnlineSubsystem"));
+	FModuleManager::Get().LoadModule(TEXT("OnlineSubsystemUtils"));
+
+
 	// @todo #JohnBLowPri: Steam detection doesn't seem to work this early on, but does work further down the line;
 	//				try to find a way, to detect it as early as possible
 
diff --git a/Engine/Plugins/NetcodeUnitTest/NetcodeUnitTest/Source/NetcodeUnitTest/Private/UnitTests/VMReflection.cpp b/Engine/Plugins/NetcodeUnitTest/NetcodeUnitTest/Source/NetcodeUnitTest/Private/UnitTests/VMReflection.cpp
--- a/Engine/Plugins/NetcodeUnitTest/NetcodeUnitTest/Source/NetcodeUnitTest/Private/UnitTests/VMReflection.cpp
+++ b/Engine/Plugins/NetcodeUnitTest/NetcodeUnitTest/Source/NetcodeUnitTest/Private/UnitTests/VMReflection.cpp
@@ -43,7 +43,7 @@ bool UVMReflection::ExecuteUnitTest()
 		TestObjA->AObjectRef = TestObjB;
 
 		bool bOriginalError = false;
-		(UObject*)(FVMReflection(TestObjA)->*"AObjectRef", &bOriginalError);
+		(void)(UObject*)(FVMReflection(TestObjA)->*"AObjectRef", &bOriginalError);
 
 		bool bError = false;
 		(FString)(FVMReflection(TestObjA)->*"AObjectRef", &bError);
@@ -808,7 +808,7 @@ bool UVMReflection::ExecuteUnitTest()
 		UVMTestClassA* TestObj = NewObject<UVMTestClassA>();
 		bool bError = false;
 
-		(uint8)((FVMReflection(TestObj)->*"BytePropArray")["uint8"], &bError);
+		(void)(uint8)((FVMReflection(TestObj)->*"BytePropArray")["uint8"], &bError);
 
 		TestResults.Add(TEXT("Bad static array access (no element)"), bError);
 	}
@@ -818,7 +818,7 @@ bool UVMReflection::ExecuteUnitTest()
 		UVMTestClassA* TestObj = NewObject<UVMTestClassA>();
 		bool bError = false;
 
-		(uint8)(FVMReflection(TestObj)->*"BytePropArray", &bError);
+		(void)(uint8)(FVMReflection(TestObj)->*"BytePropArray", &bError);
 
 		TestResults.Add(TEXT("Bad static array access (no type verification)"), bError);
 	}
@@ -915,7 +915,7 @@ bool UVMReflection::ExecuteUnitTest()
 		UVMTestClassA* TestObj = NewObject<UVMTestClassA>();
 		bool bError = false;
 
-		(uint8)((FVMReflection(TestObj)->*"DynBytePropArray")["uint8"], &bError);
+		(void)(uint8)((FVMReflection(TestObj)->*"DynBytePropArray")["uint8"], &bError);
 
 		TestResults.Add(TEXT("Bad dynamic array access (no element)"), bError);
 	}
@@ -925,7 +925,7 @@ bool UVMReflection::ExecuteUnitTest()
 		UVMTestClassA* TestObj = NewObject<UVMTestClassA>();
 		bool bError = false;
 
-		(uint8)(FVMReflection(TestObj)->*"DynBytePropArray", &bError);
+		(void)(uint8)(FVMReflection(TestObj)->*"DynBytePropArray", &bError);
 
 		TestResults.Add(TEXT("Bad dynamic array access (no type verification)"), bError);
 	}
diff --git a/Engine/Plugins/NetcodeUnitTest/NetcodeUnitTest/Source/NetcodeUnitTest/Public/NUTUtilDebug.h b/Engine/Plugins/NetcodeUnitTest/NetcodeUnitTest/Source/NetcodeUnitTest/Public/NUTUtilDebug.h
--- a/Engine/Plugins/NetcodeUnitTest/NetcodeUnitTest/Source/NetcodeUnitTest/Public/NUTUtilDebug.h
+++ b/Engine/Plugins/NetcodeUnitTest/NetcodeUnitTest/Source/NetcodeUnitTest/Public/NUTUtilDebug.h
@@ -436,9 +436,9 @@ public:
 
 
 		auto MatchLogLine =
-			[&](const FLogTraceEntry& CurEntry)
+			[&](const FLogTraceEntry& InEntry)
 			{
-				return CurEntry.LogLine == LogLine;
+				return InEntry.LogLine == LogLine;
 			};
 
 		if (bPartial)
diff --git a/Engine/Plugins/NetcodeUnitTest/NetcodeUnitTest/Source/NetcodeUnitTest/Public/Net/NUTUtilNet.h b/Engine/Plugins/NetcodeUnitTest/NetcodeUnitTest/Source/NetcodeUnitTest/Public/Net/NUTUtilNet.h
--- a/Engine/Plugins/NetcodeUnitTest/NetcodeUnitTest/Source/NetcodeUnitTest/Public/Net/NUTUtilNet.h
+++ b/Engine/Plugins/NetcodeUnitTest/NetcodeUnitTest/Source/NetcodeUnitTest/Public/Net/NUTUtilNet.h
@@ -2,11 +2,156 @@
 
 #pragma once
 
+#include "Sockets.h"
+
 
 // Forward declarations
 class UUnitTestChannel;
 class UUnitTestNetDriver;
 
+
+/**
+ * A transparent FSocket-hook, for capturing lowest-level socket operations (use Attach/Detach functions for hooking)
+ *
+ * NOTE: Not compatible with all net subsystems, e.g. some subsystems expect a specific FSocket-subclass, and statically cast as such.
+ * NOTE: Best if only temporarily attached, or the original socket may leak (no attempt is made to handle destruction properly).
+ */
+class FSocketHook : public FSocket
+{
+public:
+	/**
+	 * Base constructor
+	 */
+	FSocketHook()
+		: bLastSendToBlocked(false)
+		, HookedSocket(nullptr)
+	{
+	}
+
+	/**
+	 * Base constructor
+	 */
+	FSocketHook(FSocket* InHookedSocket)
+		: FSocket((InHookedSocket != nullptr ? InHookedSocket->GetSocketType() : ESocketType::SOCKTYPE_Datagram),
+					(InHookedSocket != nullptr ? InHookedSocket->GetDescription() + TEXT(" (hooked)") : TEXT("(hooked)")))
+		, bLastSendToBlocked(false)
+		, HookedSocket(InHookedSocket)
+	{
+	}
+
+	/**
+	 * Attach the hook to the specified FSocket* variable
+	 *
+	 * @param InSocketVar	The socket var to hook
+	 */
+	FORCEINLINE void Attach(FSocket*& InSocketVar)
+	{
+		check(InSocketVar != this);
+		check(HookedSocket == nullptr);
+
+		HookedSocket = InSocketVar;
+		InSocketVar = this;
+	}
+
+	/**
+	 * Detach the hook from the specified FSocket* variable
+	 *
+	 * @param InSocketVar	The socket var to detach from 
+	 */
+	FORCEINLINE void Detach(FSocket*& InSocketVar)
+	{
+		check(InSocketVar == this);
+		check(HookedSocket != nullptr);
+
+		InSocketVar = HookedSocket;
+		HookedSocket = nullptr;
+	}
+
+	virtual ~FSocketHook()
+	{
+	}
+
+
+	virtual bool Close() override;
+
+	virtual bool Bind(const FInternetAddr& Addr) override;
+
+	virtual bool Connect(const FInternetAddr& Addr) override;
+
+	virtual bool Listen(int32 MaxBacklog) override;
+
+	virtual bool HasPendingConnection(bool& bHasPendingConnection) override;
+
+	virtual bool HasPendingData(uint32& PendingDataSize) override;
+
+	virtual class FSocket* Accept(const FString& InSocketDescription) override;
+
+	virtual class FSocket* Accept(FInternetAddr& OutAddr, const FString& InSocketDescription) override;
+
+	virtual bool SendTo(const uint8* Data, int32 Count, int32& BytesSent, const FInternetAddr& Destination) override;
+
+	virtual bool Send(const uint8* Data, int32 Count, int32& BytesSent) override;
+
+	virtual bool RecvFrom(uint8* Data, int32 BufferSize, int32& BytesRead, FInternetAddr& Source,
+							ESocketReceiveFlags::Type Flags) override;
+
+	virtual bool Recv(uint8* Data, int32 BufferSize, int32& BytesRead, ESocketReceiveFlags::Type Flags) override;
+
+	virtual bool Wait(ESocketWaitConditions::Type Condition, FTimespan WaitTime) override;
+
+	virtual ESocketConnectionState GetConnectionState() override;
+
+	virtual void GetAddress(FInternetAddr& OutAddr) override;
+
+	virtual bool GetPeerAddress(FInternetAddr& OutAddr) override;
+
+	virtual bool SetNonBlocking(bool bIsNonBlocking) override;
+
+	virtual bool SetBroadcast(bool bAllowBroadcast) override;
+
+	virtual bool JoinMulticastGroup(const FInternetAddr& GroupAddress) override;
+
+	virtual bool LeaveMulticastGroup(const FInternetAddr& GroupAddress) override;
+
+	virtual bool SetMulticastLoopback(bool bLoopback) override;
+
+	virtual bool SetMulticastTtl(uint8 TimeToLive) override;
+
+	virtual bool SetReuseAddr(bool bAllowReuse) override;
+
+	virtual bool SetLinger(bool bShouldLinger, int32 Timeout) override;
+
+	virtual bool SetRecvErr(bool bUseErrorQueue) override;
+
+	virtual bool SetSendBufferSize(int32 Size, int32& NewSize) override;
+
+	virtual bool SetReceiveBufferSize(int32 Size, int32& NewSize) override;
+
+	virtual int32 GetPortNo() override;
+
+
+	/**
+	 * Delegate for hooking 'SendTo'
+	 *
+	 * @param Data			The data being sent
+	 * @param Count			The number of bytes being sent
+	 * @param bBlockSend	Whether or not to block the send (defaults to false)
+	 */
+	DECLARE_DELEGATE_ThreeParams(FSendToDel, void* /*Data*/, int32 /*Count*/, bool& /*bBlockSend*/);
+
+
+	/** Delegate for hooking 'SendTo' */
+	FSendToDel SendToDel;
+
+	/** Whether or not the last 'SendTo' call was blocked */
+	bool bLastSendToBlocked;
+
+
+public:
+	/** The original socket which has been hooked */
+	FSocket* HookedSocket;
+};
+
 /**
  * A delegate network notify class, to allow for easy inline-hooking.
  * 
diff --git a/Engine/Plugins/NetcodeUnitTest/NetcodeUnitTest/Source/NetcodeUnitTest/Public/NetcodeUnitTest.h b/Engine/Plugins/NetcodeUnitTest/NetcodeUnitTest/Source/NetcodeUnitTest/Public/NetcodeUnitTest.h
--- a/Engine/Plugins/NetcodeUnitTest/NetcodeUnitTest/Source/NetcodeUnitTest/Public/NetcodeUnitTest.h
+++ b/Engine/Plugins/NetcodeUnitTest/NetcodeUnitTest/Source/NetcodeUnitTest/Public/NetcodeUnitTest.h
@@ -33,6 +33,7 @@
  * List of mainline UE4 engine changelists, that required a NetcodeUnitTest compatibility adjustment (newest first)
  * Every time you update NetcodeUnitTest code, to toggle based on a changelist, add a define for it with the CL number here
  */
+#define CL_STATELESSCONNECT	2866629	// @todo #JohnB: Refers to Dev-Networking, update so it refers to /UE4//Main merge CL
 #define CL_FENGINEVERSION	2655102
 #define CL_INITCONNPARAM	2567692
 #define CL_CONSTUNIQUEID	2540329
@@ -56,7 +57,7 @@
  *
  * If in doubt, set to the top CL from list above ('List of mainline UE4 engine changelists'), and work your way down until it compiles.
  */
-#define TARGET_UE4_CL 2664876	// IMPORTANT: If you're hovering over this because compiling failed, you need to adjust this value.
+#define TARGET_UE4_CL 2866631	// IMPORTANT: If you're hovering over this because compiling failed, you need to adjust this value.
 
 
 /**
@@ -238,25 +239,25 @@ inline ELogType OptionalFlags(ELogType InFlags=ELogType::None)
 	UNIT_LOG_BEGIN(UnitTestObj, UnitLogTypeFlags); \
 		if ((OptionalFlags(UnitLogTypeFlags) & ELogType::StatusError) == ELogType::StatusError) \
 		{ \
-			UE_LOG(LogUnitTest, Error, Format, __VA_ARGS__); \
+			UE_LOG(LogUnitTest, Error, Format, ##__VA_ARGS__); \
 		} \
 		else if ((OptionalFlags(UnitLogTypeFlags) & ELogType::StatusWarning) == ELogType::StatusWarning) \
 		{ \
-			UE_LOG(LogUnitTest, Warning, Format, __VA_ARGS__); \
+			UE_LOG(LogUnitTest, Warning, Format, ##__VA_ARGS__); \
 		} \
 		else \
 		{ \
-			UE_LOG(LogUnitTest, Log, Format, __VA_ARGS__); \
+			UE_LOG(LogUnitTest, Log, Format, ##__VA_ARGS__); \
 		} \
 	UNIT_LOG_END(); \
 	UNIT_LOG_VOID_START(); \
 	if ((OptionalFlags(UnitLogTypeFlags) & ELogType::StatusError) == ELogType::StatusError) \
 	{ \
-		UE_LOG(LogUnitTest, Error, TEXT("%s: %s"), *UnitTestObj->GetUnitTestName(), *FString::Printf(Format, __VA_ARGS__)); \
+		UE_LOG(LogUnitTest, Error, TEXT("%s: %s"), *UnitTestObj->GetUnitTestName(), *FString::Printf(Format, ##__VA_ARGS__)); \
 	} \
 	else if ((OptionalFlags(UnitLogTypeFlags) & ELogType::StatusWarning) == ELogType::StatusWarning) \
 	{ \
-		UE_LOG(LogUnitTest, Warning, TEXT("%s: %s"), *UnitTestObj->GetUnitTestName(), *FString::Printf(Format, __VA_ARGS__)); \
+		UE_LOG(LogUnitTest, Warning, TEXT("%s: %s"), *UnitTestObj->GetUnitTestName(), *FString::Printf(Format, ##__VA_ARGS__)); \
 	} \
 	UNIT_LOG_VOID_END();
 
@@ -264,7 +265,7 @@ inline ELogType OptionalFlags(ELogType InFlags=ELogType::None)
 // More-concise version of the above macro
 
 #define UNIT_LOG(UnitLogTypeFlags, Format, ...) \
-	UNIT_LOG_OBJ(this, UnitLogTypeFlags, Format, __VA_ARGS__)
+	UNIT_LOG_OBJ(this, UnitLogTypeFlags, Format, ##__VA_ARGS__)
 
 
 // @todo #JohnBDoc: Document these macros
@@ -302,12 +303,12 @@ inline ELogType OptionalFlags(ELogType InFlags=ELogType::None)
 		GActiveLogTypeFlags = ELogType::GlobalStatus | OptionalFlags(UnitLogTypeFlags); \
 		UUnitTestManager* Manager = UUnitTestManager::Get(); \
 		Manager->SetStatusLog(true); \
-		Manager->Logf(Format, __VA_ARGS__); \
+		Manager->Logf(Format, ##__VA_ARGS__); \
 		Manager->SetStatusLog(false); \
 		\
 		if ((GActiveLogTypeFlags & ELogType::StatusAutomation) == ELogType::StatusAutomation && GIsAutomationTesting) \
 		{ \
-			GWarn->Logf(Format, __VA_ARGS__); \
+			GWarn->Logf(Format, ##__VA_ARGS__); \
 		} \
 		\
 		GActiveLogTypeFlags = ELogType::None; \
@@ -322,11 +323,11 @@ inline ELogType OptionalFlags(ELogType InFlags=ELogType::None)
 		{ \
 			if (SourceUnitTest != NULL) \
 			{ \
-				UE_LOG(LogUnitTest, Error, TEXT("%s: %s"), *SourceUnitTest->GetUnitTestName(), *FString::Printf(Format, __VA_ARGS__)); \
+				UE_LOG(LogUnitTest, Error, TEXT("%s: %s"), *SourceUnitTest->GetUnitTestName(), *FString::Printf(Format, ##__VA_ARGS__)); \
 			} \
 			else \
 			{ \
-				UE_LOG(LogUnitTest, Error, Format, __VA_ARGS__) \
+				UE_LOG(LogUnitTest, Error, Format, ##__VA_ARGS__) \
 			} \
 		} \
 		else if ((GActiveLogTypeFlags & ELogType::StatusWarning) == ELogType::StatusWarning) \
@@ -334,34 +335,34 @@ inline ELogType OptionalFlags(ELogType InFlags=ELogType::None)
 			if (SourceUnitTest != NULL) \
 			{ \
 				UE_LOG(LogUnitTest, Warning, TEXT("%s: %s"), *SourceUnitTest->GetUnitTestName(), \
-						*FString::Printf(Format, __VA_ARGS__)); \
+						*FString::Printf(Format, ##__VA_ARGS__)); \
 			} \
 			else \
 			{ \
-				UE_LOG(LogUnitTest, Warning, Format, __VA_ARGS__) \
+				UE_LOG(LogUnitTest, Warning, Format, ##__VA_ARGS__) \
 			} \
 		} \
 		else \
 		{ \
-			UE_LOG(LogUnitTest, Log, Format, __VA_ARGS__) \
+			UE_LOG(LogUnitTest, Log, Format, ##__VA_ARGS__) \
 		} \
 		UNIT_EVENT_RESTORE; \
-		STATUS_LOG_BASE(UnitLogTypeFlags, Format, __VA_ARGS__) \
+		STATUS_LOG_BASE(UnitLogTypeFlags, Format, ##__VA_ARGS__) \
 		GActiveLogTypeFlags = ELogType::None; \
 	}
 
-#define STATUS_LOG(UnitLogTypeFlags, Format, ...) STATUS_LOG_OBJ(NULL, UnitLogTypeFlags, Format, __VA_ARGS__)
+#define STATUS_LOG(UnitLogTypeFlags, Format, ...) STATUS_LOG_OBJ(NULL, UnitLogTypeFlags, Format, ##__VA_ARGS__)
 
 /**
  * Version of the above, for unit test status window entries, which are from specific unit tests
  */
 #define UNIT_STATUS_LOG_OBJ(UnitTestObj, UnitLogTypeFlags, Format, ...) \
 	UNIT_LOG_BEGIN(UnitTestObj, UnitLogTypeFlags); \
-	STATUS_LOG_OBJ(UnitTestObj, UnitLogTypeFlags, Format, __VA_ARGS__) \
+	STATUS_LOG_OBJ(UnitTestObj, UnitLogTypeFlags, Format, ##__VA_ARGS__) \
 	UNIT_LOG_END();
 
 #define UNIT_STATUS_LOG(UnitLogTypeFlags, Format, ...) \
-	UNIT_STATUS_LOG_OBJ(this, UnitLogTypeFlags, Format, __VA_ARGS__)
+	UNIT_STATUS_LOG_OBJ(this, UnitLogTypeFlags, Format, ##__VA_ARGS__)
 
 
 /**
diff --git a/Engine/Plugins/NetcodeUnitTest/NetcodeUnitTest/Source/NetcodeUnitTest/Public/UI/SMultiSelectTableRow.h b/Engine/Plugins/NetcodeUnitTest/NetcodeUnitTest/Source/NetcodeUnitTest/Public/UI/SMultiSelectTableRow.h
--- a/Engine/Plugins/NetcodeUnitTest/NetcodeUnitTest/Source/NetcodeUnitTest/Public/UI/SMultiSelectTableRow.h
+++ b/Engine/Plugins/NetcodeUnitTest/NetcodeUnitTest/Source/NetcodeUnitTest/Public/UI/SMultiSelectTableRow.h
@@ -22,13 +22,13 @@ public:
 	{
 		FReply Reply = STableRow<ItemType>::OnDragDetected(MyGeometry, MouseEvent);
 
-		const TSharedPtr<ITypedTableView<ItemType>> OwnerWidget = OwnerTablePtr.Pin();
+		auto OwnerWidget = this->OwnerTablePtr.Pin();
 
 		// When the user starts dragging a log line, treat that as the start of selecting multiple-lines
 #if TARGET_UE4_CL < CL_GETSELECTIONMODE
 		if (OwnerWidget->Private_GetSelectionMode() == ESelectionMode::Multi)
 #else
-		if (GetSelectionMode() == ESelectionMode::Multi)
+		if (this->GetSelectionMode() == ESelectionMode::Multi)
 #endif
 		{
 			const ItemType* MyItem = OwnerWidget->Private_ItemFromWidget(this);
@@ -50,13 +50,13 @@ public:
 	{
 		STableRow<ItemType>::OnDragEnter(MyGeometry, DragDropEvent);
 
-		const TSharedPtr<ITypedTableView<ItemType>> OwnerWidget = OwnerTablePtr.Pin();
+		auto OwnerWidget = this->OwnerTablePtr.Pin();
 
 		// Each time the user mouses-over another log line when dragging, select all lines from the dragged line to the current line
 #if TARGET_UE4_CL < CL_GETSELECTIONMODE
 		if (OwnerWidget->Private_GetSelectionMode() == ESelectionMode::Multi && OwnerWidget->Private_GetNumSelectedItems() > 0)
 #else
-		if (GetSelectionMode() == ESelectionMode::Multi && OwnerWidget->Private_GetNumSelectedItems() > 0)
+		if (this->GetSelectionMode() == ESelectionMode::Multi && OwnerWidget->Private_GetNumSelectedItems() > 0)
 #endif
 		{
 			const ItemType* MyItem = OwnerWidget->Private_ItemFromWidget(this);
diff --git a/Engine/Plugins/NetcodeUnitTest/NetcodeUnitTest/VAXSnippets/Unit Test Subclass Header.txt b/Engine/Plugins/NetcodeUnitTest/NetcodeUnitTest/VAXSnippets/Unit Test Subclass Header.txt
new file mode 100644
--- /dev/null
+++ b/Engine/Plugins/NetcodeUnitTest/NetcodeUnitTest/VAXSnippets/Unit Test Subclass Header.txt	
@@ -0,0 +1,21 @@
+Title: Unit Test Subclass Header
+Shortcut:
+Description: Creates a header for a unit test subclass.
+
+Source:
+// Copyright 1998-2016 Epic Games, Inc. All Rights Reserved.
+
+#pragma once
+
+#include "UnitTests/$SubclassUnitTest$.h"
+
+#include "$FILE_BASE$.generated.h"
+
+/**
+ * @todo #JohnB: $end$
+ */
+UCLASS()
+class U$FILE_BASE$ : public U$SubclassUnitTest$
+{
+	GENERATED_UCLASS_BODY()
+};
diff --git a/Engine/Plugins/NetcodeUnitTest/NetcodeUnitTest/VAXSnippets/Unit Test Subclass Source.txt b/Engine/Plugins/NetcodeUnitTest/NetcodeUnitTest/VAXSnippets/Unit Test Subclass Source.txt
new file mode 100644
--- /dev/null
+++ b/Engine/Plugins/NetcodeUnitTest/NetcodeUnitTest/VAXSnippets/Unit Test Subclass Source.txt	
@@ -0,0 +1,29 @@
+Title: Unit Test Subclass Source
+Shortcut:
+Description: Creates the source file for a unit test subclass.
+
+Source:
+// Copyright 1998-2016 Epic Games, Inc. All Rights Reserved.
+
+#include "$UnitTestProject$PCH.h"
+
+#include "UnitTests/$FILE_BASE$.h"
+
+#include "UnitTestEnvironment.h"
+
+/**
+ * U$FILE_BASE$
+ */
+U$FILE_BASE$::U$FILE_BASE$(const FObjectInitializer& ObjectInitializer)
+	: Super(ObjectInitializer)
+{
+	UnitTestName = TEXT("$UnitTestName$");
+
+	UnitTestDate = FDateTime($YEAR$, $MONTH$, $DAY$);
+
+	bWorkInProgress = true;
+
+	// @todo #JohnBExploitCL: Bugtracking/changelist notes
+
+	ExpectedResult.Add(TEXT("$TargetGame$"), EUnitTestVerification::VerifiedNotFixed);
+}
diff --git a/Engine/Plugins/Runtime/PacketHandlers/CompressionComponents/Oodle/Source/OodleHandlerComponent/Private/OodleHandlerComponent.cpp b/Engine/Plugins/Runtime/PacketHandlers/CompressionComponents/Oodle/Source/OodleHandlerComponent/Private/OodleHandlerComponent.cpp
--- a/Engine/Plugins/Runtime/PacketHandlers/CompressionComponents/Oodle/Source/OodleHandlerComponent/Private/OodleHandlerComponent.cpp
+++ b/Engine/Plugins/Runtime/PacketHandlers/CompressionComponents/Oodle/Source/OodleHandlerComponent/Private/OodleHandlerComponent.cpp
@@ -8,10 +8,6 @@
 DEFINE_LOG_CATEGORY(OodleHandlerComponentLog);
 
 
-// @todo #JohnB: IMPORTANT: Do not add any placeholders for the compression-toggling into the protocol yet. Needs careful planning,
-//					might not be packet-based compression, but bunch-based.
-
-
 // @todo #JohnB: Deprecate the 'training' mode name from this handler eventually, as it is misnamed/confusing
 
 // @todo #JohnB: You need context-sensitive data capturing; e.g. you do NOT want to be capturing voice channel data
@@ -705,6 +701,9 @@ bool OodleHandlerComponent::IsValid() const
 
 void OodleHandlerComponent::Incoming(FBitReader& Packet)
 {
+	// Oodle must be the first HandlerComponent to process incoming packets, so does not support bit-shifted reads
+	check(Packet.GetPosBits() == 0);
+
 	if (bEnableOodle)
 	{
 		switch(Mode)
@@ -715,9 +714,13 @@ void OodleHandlerComponent::Incoming(FBitReader& Packet)
 				if (Handler->Mode == Handler::Mode::Server)
 				{
 					uint32 SizeOfPacket = Packet.GetNumBytes();
-					void* PacketData = Packet.GetData();
 
-					InPacketLog->SerializePacket(PacketData, SizeOfPacket);
+					if (SizeOfPacket > 0)
+					{
+						void* PacketData = Packet.GetData();
+
+						InPacketLog->SerializePacket(PacketData, SizeOfPacket);
+					}
 				}
 				break;
 			}
@@ -752,9 +755,14 @@ void OodleHandlerComponent::Incoming(FBitReader& Packet)
 															CompressedLength, DecompressedData, DecompressedLength);
 						}
 
+
+
 						if (!bSuccess)
 						{
 							// @todo #JohnB: Set packet archive error and allow in shipping - see note below
+
+							// @todo #JohnB: I don't think this is always an error. This may occur when Oodle decides not to compress,
+							//					in which case it should fall-through below (i.e. 'Packet' should remain unmodified)
 #if !UE_BUILD_SHIPPING
 							UE_LOG(OodleHandlerComponentLog, Error, TEXT("Error decoding Oodle network data."));
 #endif
@@ -817,7 +825,10 @@ void OodleHandlerComponent::Outgoing(FBitWriter& Packet)
 				{
 					uint32 SizeOfPacket = Packet.GetNumBytes();
 
-					OutPacketLog->SerializePacket((void*)Packet.GetData(), SizeOfPacket);
+					if (SizeOfPacket > 0)
+					{
+						OutPacketLog->SerializePacket((void*)Packet.GetData(), SizeOfPacket);
+					}
 				}
 				break;
 			}
@@ -906,6 +917,16 @@ int32 OodleHandlerComponent::GetPacketOverheadBits()
 	return ReturnVal;
 }
 
+uint8 OodleHandlerComponent::GetBitAlignment()
+{
+	return (GetPacketOverheadBits() % 8);
+}
+
+bool OodleHandlerComponent::DoesResetBitAlignment()
+{
+	return true;
+}
+
 
 #if !UE_BUILD_SHIPPING
 /**
@@ -918,6 +939,7 @@ static bool OodleExec(UWorld* InWorld, const TCHAR* Cmd, FOutputDevice& Ar)
 	if (FParse::Command(&Cmd, TEXT("Oodle")))
 	{
 		// Used by unit testing code, to enable/disable Oodle during a unit test
+		// NOTE: Do not use while a NetConnection is using Oodle, as this will cause it to break. Debugging/Testing only.
 		if (FParse::Command(&Cmd, TEXT("ForceEnable")))
 		{
 			bool bTurnOn = false;
diff --git a/Engine/Plugins/Runtime/PacketHandlers/CompressionComponents/Oodle/Source/OodleHandlerComponent/Public/OodleHandlerComponent.h b/Engine/Plugins/Runtime/PacketHandlers/CompressionComponents/Oodle/Source/OodleHandlerComponent/Public/OodleHandlerComponent.h
--- a/Engine/Plugins/Runtime/PacketHandlers/CompressionComponents/Oodle/Source/OodleHandlerComponent/Public/OodleHandlerComponent.h
+++ b/Engine/Plugins/Runtime/PacketHandlers/CompressionComponents/Oodle/Source/OodleHandlerComponent/Public/OodleHandlerComponent.h
@@ -265,10 +265,10 @@ public:
 class OODLEHANDLERCOMPONENT_API OodleHandlerComponent : public HandlerComponent
 {
 public:
-	/* Initializes default data */
+	/** Initializes default data */
 	OodleHandlerComponent();
 
-	/* Default Destructor */
+	/** Default Destructor */
 	~OodleHandlerComponent();
 
 	/**
@@ -315,41 +315,45 @@ public:
 	bool FindFallbackDictionaries(FString& OutServerDictionary, FString& OutClientDictionary, bool bTestOnly=false);
 #endif
 
-	/* Initializes the handler component */
 	virtual void Initialize() override;
 
-	/* Whether the handler component is valid */
 	virtual bool IsValid() const override;
 
-	/* Handles any incoming packets */
 	virtual void Incoming(FBitReader& Packet) override;
 
-	/* Handles any outgoing packets */
 	virtual void Outgoing(FBitWriter& Packet) override;
 
+	virtual void IncomingConnectionless(FString Address, FBitReader& Packet) override
+	{
+	}
+
+	virtual void OutgoingConnectionless(FString Address, FBitWriter& Packet) override
+	{
+	}
+
 	virtual int32 GetPacketOverheadBits() override;
 
+	virtual uint8 GetBitAlignment() override;
+
+	virtual bool DoesResetBitAlignment() override;
+
+
 protected:
 	/** Whether or not Oodle is enabled */
 	bool bEnableOodle;
 
 #if !UE_BUILD_SHIPPING || OODLE_DEV_SHIPPING
-	/* File to log input packets to */
+	/** File to log input packets to */
 	FPacketCaptureArchive* InPacketLog;
 
-	/* File to log output packets to */
+	/** File to log output packets to */
 	FPacketCaptureArchive* OutPacketLog;
 
 	/** Search for dictionary files and use them if present - switching mode to Release in process - don't use in shipping */
 	bool bUseDictionaryIfPresent;
 #endif
 
-	/* 
-	* Modes of the component
-	* Specify this in .ini by
-	* [OodleHandlerComponent]
-	* Mode=Capturing
-	*/
+	/** Modes of the component  */
 	EOodleHandlerMode Mode;
 
 	/** Server (Outgoing) dictionary data */
diff --git a/Engine/Source/Runtime/Core/Private/Misc/SecureHash.cpp b/Engine/Source/Runtime/Core/Private/Misc/SecureHash.cpp
--- a/Engine/Source/Runtime/Core/Private/Misc/SecureHash.cpp
+++ b/Engine/Source/Runtime/Core/Private/Misc/SecureHash.cpp
@@ -731,6 +731,64 @@ void FSHA1::HashBuffer(const void* Data, uint32 DataSize, uint8* OutHash)
 	Sha.GetHash(OutHash);
 }
 
+void FSHA1::HMACBuffer(const void* Key, uint32 KeySize, const void* Data, uint32 DataSize, uint8* OutHash)
+{
+	const uint8 BlockSize = 64;
+	const uint8 HashSize = 20;
+	uint8 FinalKey[BlockSize];
+
+	// Fit 'Key' into a BlockSize-aligned 'FinalKey' value
+	if (KeySize > BlockSize)
+	{
+		HashBuffer(Key, KeySize, FinalKey);
+
+		FMemory::Memzero(FinalKey + HashSize, BlockSize - HashSize);
+	}
+	else if (KeySize < BlockSize)
+	{
+		FMemory::Memcpy(FinalKey, Key, KeySize);
+		FMemory::Memzero(FinalKey + KeySize, BlockSize - KeySize);
+	}
+	else
+	{
+		FMemory::Memcpy(FinalKey, Key, KeySize);
+	}
+
+
+	uint8 OKeyPad[BlockSize];
+	uint8 IKeyPad[BlockSize];
+
+	for (int32 i=0; i<BlockSize; i++)
+	{
+		OKeyPad[i] = 0x5C ^ FinalKey[i];
+		IKeyPad[i] = 0x36 ^ FinalKey[i];
+	}
+
+
+	// Start concatenating/hashing the pads/data etc: Hash(OKeyPad + Hash(IKeyPad + Data))
+	uint8* IKeyPad_Data = new uint8[ARRAY_COUNT(IKeyPad) + DataSize];
+
+	FMemory::Memcpy(IKeyPad_Data, IKeyPad, ARRAY_COUNT(IKeyPad));
+	FMemory::Memcpy(IKeyPad_Data + ARRAY_COUNT(IKeyPad), Data, DataSize);
+
+
+	uint8 IKeyPad_Data_Hash[HashSize];
+
+	HashBuffer(IKeyPad_Data, ARRAY_COUNT(IKeyPad) + DataSize, IKeyPad_Data_Hash);
+
+	delete[] IKeyPad_Data;
+
+
+	uint8 OKeyPad_IHash[ARRAY_COUNT(OKeyPad) + HashSize];
+
+	FMemory::Memcpy(OKeyPad_IHash, OKeyPad, ARRAY_COUNT(OKeyPad));
+	FMemory::Memcpy(OKeyPad_IHash + ARRAY_COUNT(OKeyPad), IKeyPad_Data_Hash, HashSize);
+
+
+	// Output the final hash
+	HashBuffer(OKeyPad_IHash, ARRAY_COUNT(OKeyPad_IHash), OutHash);
+}
+
 
 /**
  * Shared hashes.sha reading code (each platform gets a buffer to the data,
diff --git a/Engine/Source/Runtime/Core/Public/Misc/SecureHash.h b/Engine/Source/Runtime/Core/Public/Misc/SecureHash.h
--- a/Engine/Source/Runtime/Core/Public/Misc/SecureHash.h
+++ b/Engine/Source/Runtime/Core/Public/Misc/SecureHash.h
@@ -238,6 +238,18 @@ public:
 	 */
 	static void HashBuffer(const void* Data, uint32 DataSize, uint8* OutHash);
 
+	/**
+	 * Generate the HMAC (Hash-based Message Authentication Code) for a block of data.
+	 * https://en.wikipedia.org/wiki/Hash-based_message_authentication_code
+	 *
+	 * @param Key		The secret key to be used when generating the HMAC
+	 * @param KeySize	The size of the key
+	 * @param Data		Input data to hash
+	 * @param DataSize	Size of the Data block
+	 * @param OutHash	Resulting hash value (20 byte buffer)
+	 */
+	static void HMACBuffer(const void* Key, uint32 KeySize, const void* Data, uint32 DataSize, uint8* OutHash);
+
 	/**
 	 * Shared hashes.sha reading code (each platform gets a buffer to the data,
 	 * then passes it to this function for processing)
diff --git a/Engine/Source/Runtime/Engine/Classes/Engine/ChildConnection.h b/Engine/Source/Runtime/Engine/Classes/Engine/ChildConnection.h
--- a/Engine/Source/Runtime/Engine/Classes/Engine/ChildConnection.h
+++ b/Engine/Source/Runtime/Engine/Classes/Engine/ChildConnection.h
@@ -36,7 +36,7 @@ public:
 		return Parent->LowLevelDescribe();
 	}
 
-	virtual void LowLevelSend( void* Data, int32 Count ) override
+	virtual void LowLevelSend(void* Data, int32 CountBytes, int32 CountBits) override
 	{
 	}
 
diff --git a/Engine/Source/Runtime/Engine/Classes/Engine/DemoNetConnection.h b/Engine/Source/Runtime/Engine/Classes/Engine/DemoNetConnection.h
--- a/Engine/Source/Runtime/Engine/Classes/Engine/DemoNetConnection.h
+++ b/Engine/Source/Runtime/Engine/Classes/Engine/DemoNetConnection.h
@@ -7,7 +7,20 @@
 
 struct FQueuedDemoPacket
 {
-	TArray< uint8 > Data;
+	/** The packet data to send */
+	TArray<uint8> Data;
+
+	/** The size of the packet in bits */
+	int32 SizeBits;
+
+public:
+	FORCEINLINE FQueuedDemoPacket(uint8* InData, int32 InSizeBytes, int32 InSizeBits)
+		: Data()
+		, SizeBits(InSizeBits)
+	{
+		Data.AddUninitialized(InSizeBytes);
+		FMemory::Memcpy(Data.GetData(), InData, InSizeBytes);
+	}
 };
 
 
@@ -27,7 +40,7 @@ public:
 	virtual void InitConnection( class UNetDriver* InDriver, EConnectionState InState, const FURL& InURL, int32 InConnectionSpeed = 0, int32 InMaxPacket=0) override;
 	virtual FString LowLevelGetRemoteAddress( bool bAppendPort = false ) override;
 	virtual FString LowLevelDescribe() override;
-	virtual void LowLevelSend( void* Data, int32 Count ) override;
+	virtual void LowLevelSend(void* Data, int32 CountBytes, int32 CountBits) override;
 	virtual int32 IsNetReady( bool Saturate ) override;
 	virtual void FlushNet( bool bIgnoreSimulation = false ) override;
 	virtual void HandleClientPlayer( APlayerController* PC, class UNetConnection* NetConnection ) override;
diff --git a/Engine/Source/Runtime/Engine/Classes/Engine/NetConnection.h b/Engine/Source/Runtime/Engine/Classes/Engine/NetConnection.h
--- a/Engine/Source/Runtime/Engine/Classes/Engine/NetConnection.h
+++ b/Engine/Source/Runtime/Engine/Classes/Engine/NetConnection.h
@@ -113,13 +113,29 @@ namespace EClientLoginState
 };
 
 #if DO_ENABLE_NET_TEST
-//
-// A lagged packet
-//
+/**
+ * An artificially lagged packet
+ */
 struct DelayedPacket
 {
+	/** The packet data to send */
 	TArray<uint8> Data;
+
+	/** The size of the packet in bits */
+	int32 SizeBits;
+
+	/** The time at which to send the packet */
 	double SendTime;
+
+public:
+	FORCEINLINE DelayedPacket(uint8* InData, int32 InSizeBytes, int32 InSizeBits)
+		: Data()
+		, SizeBits(InSizeBits)
+		, SendTime(0.0)
+	{
+		Data.AddUninitialized(InSizeBytes);
+		FMemory::Memcpy(Data.GetData(), InData, InSizeBytes);
+	}
 };
 #endif
 
@@ -203,9 +219,14 @@ public:
 	
 	uint32 bPendingDestroy:1;    // when true, playercontroller or beaconclient is being destroyed
 
-	// Packet Handler
+
+	/** PacketHandler, for managing layered handler components, which modify packets as they are sent/received */
 	TUniquePtr<PacketHandler> Handler;
 
+	/** Reference to the PacketHandler component, for managing stateless connection handshakes */
+	TWeakPtr<StatelessConnectHandlerComponent> StatelessConnectComponent;
+
+
 	/** Whether this channel needs to byte swap all data or not */
 	bool			bNeedsByteSwapping;
 	/** Net id of remote player on this connection. Only valid on client connections. */
@@ -416,10 +437,13 @@ public:
 	/**
 	 * Sends a byte stream to the remote endpoint using the underlying socket
 	 *
-	 * @param Data the byte stream to send
-	 * @param Count the length of the stream to send
+	 * @param Data			The byte stream to send
+	 * @param CountBytes	The length of the stream to send, in bytes
+	 * @param CountBits		The length of the stream to send, in bits (to support bit-level additions to packets, from PacketHandler's)
 	 */
-	virtual void LowLevelSend( void* Data, int32 Count ) PURE_VIRTUAL(UNetConnection::LowLevelSend,); //!! "Looks like an FArchive"
+	// @todo: Deprecate 'CountBytes' eventually
+	ENGINE_API virtual void LowLevelSend(void* Data, int32 CountBytes, int32 CountBits)
+		PURE_VIRTUAL(UNetConnection::LowLevelSend,);
 
 	/** Validates the FBitWriter to make sure it's not in an error state */
 	ENGINE_API virtual void ValidateSendBuffer();
diff --git a/Engine/Source/Runtime/Engine/Classes/Engine/NetDriver.h b/Engine/Source/Runtime/Engine/Classes/Engine/NetDriver.h
--- a/Engine/Source/Runtime/Engine/Classes/Engine/NetDriver.h
+++ b/Engine/Source/Runtime/Engine/Classes/Engine/NetDriver.h
@@ -5,6 +5,9 @@
 #pragma once
 
 #include "EngineTypes.h"
+#include "Runtime/PacketHandlers/PacketHandler/Public/PacketHandler.h"
+#include "StatelessConnectHandlerComponent.h"
+
 #include "NetDriver.generated.h"
 
 class FRepChangedPropertyTracker;
@@ -181,6 +184,14 @@ public:
 	UPROPERTY()
 	TArray<class UNetConnection*> ClientConnections;
 
+
+	/** Serverside PacketHandler for managing connectionless packets */
+	TUniquePtr<PacketHandler> ConnectionlessHandler;
+
+	/** Reference to the PacketHandler component, for managing stateless connection handshakes */
+	TWeakPtr<StatelessConnectHandlerComponent> StatelessConnectComponent;
+
+
 	/** World this net driver is associated with */
 	UPROPERTY()
 	class UWorld* World;
@@ -416,6 +427,18 @@ public:
 	 */
 	ENGINE_API virtual bool InitListen(class FNetworkNotify* InNotify, FURL& ListenURL, bool bReuseAddressAndPort, FString& Error) PURE_VIRTUAL( UNetDriver::InitListen, return true;);
 
+	/**
+	 * Initialize a PacketHandler for serverside net drivers, for handling connectionless packets
+	 * NOTE: Only triggered by net driver subclasses that support it - from within InitListen.
+	 */
+	ENGINE_API virtual void InitConnectionlessHandler();
+
+	/**
+	 * Flushes all packets queued by the connectionless PacketHandler
+	 * NOTE: This should be called shortly after all calls to PacketHandler::IncomingConnectionless, to minimize packet buffer buildup.
+	 */
+	ENGINE_API virtual void FlushHandler();
+
 
 	/** Initializes the net connection class to use for new connections */
 	ENGINE_API virtual bool InitConnectionClass(void);
@@ -469,6 +492,18 @@ public:
 	/** PostTick actions */
 	ENGINE_API virtual void PostTickFlush();
 
+
+	/**
+	 * Sends a 'connectionless' (not associated with a UNetConection) packet, to the specified address.
+	 * NOTE: Address is an abstract format defined by subclasses. Anything calling this, must use an address supplied by the net driver.
+	 *
+	 * @param Address		The address the packet should be sent to (format is abstract, determined by net driver subclasses)
+	 * @param Data			The packet data
+	 * @param CountBits		The size of the packet data, in bits
+	 */
+	ENGINE_API virtual void LowLevelSend(FString Address, void* Data, int32 CountBits)
+		PURE_VIRTUAL(UNetDriver::LowLevelSend,);
+
 	/**
 	 * Process any local talker packets that need to be sent to clients
 	 */
diff --git a/Engine/Source/Runtime/Engine/Classes/PacketHandlers/EngineHandlerComponentFactory.h b/Engine/Source/Runtime/Engine/Classes/PacketHandlers/EngineHandlerComponentFactory.h
new file mode 100644
--- /dev/null
+++ b/Engine/Source/Runtime/Engine/Classes/PacketHandlers/EngineHandlerComponentFactory.h
@@ -0,0 +1,19 @@
+// Copyright 1998-2016 Epic Games, Inc. All Rights Reserved.
+
+#pragma once
+
+#include "HandlerComponentFactory.h"
+
+#include "EngineHandlerComponentFactory.generated.h"
+
+/**
+ * Factory class for loading HandlerComponent's contained within Engine
+ */
+UCLASS()
+class UEngineHandlerComponentFactory : public UHandlerComponentFactory
+{
+	GENERATED_UCLASS_BODY()
+
+public:
+	virtual TSharedPtr<HandlerComponent> CreateComponentInstance(FString& Options) override;
+};
diff --git a/Engine/Source/Runtime/Engine/Engine.Build.cs b/Engine/Source/Runtime/Engine/Engine.Build.cs
--- a/Engine/Source/Runtime/Engine/Engine.Build.cs
+++ b/Engine/Source/Runtime/Engine/Engine.Build.cs
@@ -76,6 +76,7 @@ public class Engine : ModuleRules
 				"SynthBenchmark",
                 "AIModule",
 				"DatabaseSupport",
+                "PacketHandler"
 			}
 		);
 
@@ -100,7 +101,6 @@ public class Engine : ModuleRules
 				"Projects",
 				"Niagara",
                 "Internationalization",
-                "PacketHandler",
                 "MaterialShaderQualitySettings",
 			}
         );
diff --git a/Engine/Source/Runtime/Engine/Private/DataReplication.cpp b/Engine/Source/Runtime/Engine/Private/DataReplication.cpp
--- a/Engine/Source/Runtime/Engine/Private/DataReplication.cpp
+++ b/Engine/Source/Runtime/Engine/Private/DataReplication.cpp
@@ -196,7 +196,7 @@ void FObjectReplicator::InitWithObject( UObject* InObject, UNetConnection * InCo
 	RepLayout = Connection->Driver->GetObjectClassRepLayout( ObjectClass );
 
 	// Make a copy of the net properties
-	uint8* Source = bUseDefaultState ? (uint8*)GetObject()->GetClass()->GetDefaultObject() : (uint8*)InObject;
+	uint8* Source = bUseDefaultState ? (uint8*)GetObject()->GetArchetype() : (uint8*)InObject;
 
 	InitRecentProperties( Source );
 
diff --git a/Engine/Source/Runtime/Engine/Private/DemoNetDriver.cpp b/Engine/Source/Runtime/Engine/Private/DemoNetDriver.cpp
--- a/Engine/Source/Runtime/Engine/Private/DemoNetDriver.cpp
+++ b/Engine/Source/Runtime/Engine/Private/DemoNetDriver.cpp
@@ -4,6 +4,11 @@
 	UDemoNetDriver.cpp: Simulated network driver for recording and playing back game sessions.
 =============================================================================*/
 
+
+// @todo: LowLevelSend now includes the packet size in bits, but this is ignored locally.
+//			Tracking of this must be added, if demos are to support PacketHandler's in the future (not presently needed).
+
+
 #include "EnginePrivate.h"
 #include "Engine/DemoNetDriver.h"
 #include "Engine/DemoNetConnection.h"
@@ -1355,7 +1360,9 @@ void UDemoNetDriver::TickDemoRecord( float DeltaSeconds )
 
 	for ( int32 i = 0; i < ClientDemoConnection->QueuedDemoPackets.Num(); i++ )
 	{
-		ClientDemoConnection->LowLevelSend( (char*)&ClientDemoConnection->QueuedDemoPackets[i].Data[0], ClientDemoConnection->QueuedDemoPackets[i].Data.Num() );
+		FQueuedDemoPacket& CurPacket = ClientDemoConnection->QueuedDemoPackets[i];
+
+		ClientDemoConnection->LowLevelSend((char*)&CurPacket.Data[0], CurPacket.Data.Num(), CurPacket.SizeBits);
 	}
 
 	ClientDemoConnection->QueuedDemoPackets.Empty();
@@ -1743,11 +1750,11 @@ void UDemoNetDriver::TickDemoPlayback( float DeltaSeconds )
 		CVarGotoTimeInSeconds.AsVariable()->Set( TEXT( "-1" ), ECVF_SetByConsole );
 	}
 
-	if ( CVarDemoSkipTime.GetValueOnGameThread() > 0 )
+	if (FMath::Abs(CVarDemoSkipTime.GetValueOnGameThread()) > 0.0f)
 	{
 		// Just overwrite existing value, cvar wins in this case
-		SkipTime( CVarDemoSkipTime.GetValueOnGameThread() );
-		CVarDemoSkipTime.AsVariable()->Set( TEXT( "0" ), ECVF_SetByConsole );
+		GotoTimeInSeconds(DemoCurrentTime + CVarDemoSkipTime.GetValueOnGameThread());
+		CVarDemoSkipTime.AsVariable()->Set(TEXT("0"), ECVF_SetByConsole);
 	}
 
 	if ( !ProcessReplayTasks() )
@@ -2326,7 +2333,7 @@ FString UDemoNetConnection::LowLevelGetRemoteAddress( bool bAppendPort )
 	return TEXT( "UDemoNetConnection" );
 }
 
-void UDemoNetConnection::LowLevelSend( void* Data, int32 Count )
+void UDemoNetConnection::LowLevelSend(void* Data, int32 CountBytes, int32 CountBits)
 {
 	if ( GetDriver()->bSavingCheckpoint )
 	{
@@ -2337,22 +2344,22 @@ void UDemoNetConnection::LowLevelSend( void* Data, int32 Count )
 			return;
 		}
 	
-		*CheckpointArchive << Count;
-		CheckpointArchive->Serialize( Data, Count );
+		*CheckpointArchive << CountBytes;
+		CheckpointArchive->Serialize(Data, CountBytes);
 
-		TrackSendForProfiler( Data, Count );
+		TrackSendForProfiler(Data, CountBytes);
 		return;
 	}
 
-	if ( Count == 0 )
+	if ( CountBytes == 0 )
 	{
 		UE_LOG( LogDemo, Warning, TEXT( "UDemoNetConnection::LowLevelSend: Ignoring empty packet." ) );
 		return;
 	}
 
-	if ( Count > MAX_DEMO_READ_WRITE_BUFFER )
+	if ( CountBytes > MAX_DEMO_READ_WRITE_BUFFER )
 	{
-		UE_LOG( LogDemo, Fatal, TEXT( "UDemoNetConnection::LowLevelSend: Count > MAX_DEMO_READ_WRITE_BUFFER." ) );
+		UE_LOG( LogDemo, Fatal, TEXT( "UDemoNetConnection::LowLevelSend: CountBytes > MAX_DEMO_READ_WRITE_BUFFER." ) );
 	}
 
 	FArchive* FileAr = GetDriver()->ReplayStreamer->GetStreamingArchive();
@@ -2362,19 +2369,18 @@ void UDemoNetConnection::LowLevelSend( void* Data, int32 Count )
 		// If we're outside of an official demo frame, we need to queue this up or it will throw off the stream
 		if ( !GetDriver()->bIsRecordingDemoFrame )
 		{
-			FQueuedDemoPacket & B = *( new( QueuedDemoPackets )FQueuedDemoPacket );
-			B.Data.AddUninitialized( Count );
-			FMemory::Memcpy( B.Data.GetData(), Data, Count );
+			new( QueuedDemoPackets )FQueuedDemoPacket((uint8*)Data, CountBytes, CountBits);
+
 			return;
 		}
 
-		*FileAr << Count;
-		FileAr->Serialize( Data, Count );
+		*FileAr << CountBytes;
+		FileAr->Serialize( Data, CountBytes);
 		
-		TrackSendForProfiler( Data, Count );
+		TrackSendForProfiler( Data, CountBytes );
 		
 #if DEMO_CHECKSUMS == 1
-		uint32 Checksum = FCrc::MemCrc32( Data, Count, 0 );
+		uint32 Checksum = FCrc::MemCrc32( Data, CountBytes, 0 );
 		*FileAr << Checksum;
 #endif
 	}
diff --git a/Engine/Source/Runtime/Engine/Private/GameInstance.cpp b/Engine/Source/Runtime/Engine/Private/GameInstance.cpp
--- a/Engine/Source/Runtime/Engine/Private/GameInstance.cpp
+++ b/Engine/Source/Runtime/Engine/Private/GameInstance.cpp
@@ -51,13 +51,17 @@ void UGameInstance::Init()
 {
 	ReceiveInit();
 
-	const auto OnlineSub = IOnlineSubsystem::Get();
-	if (OnlineSub != nullptr)
+	if (!IsRunningCommandlet())
 	{
-		IOnlineSessionPtr SessionInt = OnlineSub->GetSessionInterface();
-		if (SessionInt.IsValid())
+		const auto OnlineSub = IOnlineSubsystem::Get();
+		if (OnlineSub != nullptr)
 		{
-			SessionInt->AddOnSessionUserInviteAcceptedDelegate_Handle(FOnSessionUserInviteAcceptedDelegate::CreateUObject(this, &UGameInstance::HandleSessionUserInviteAccepted));
+			IOnlineSessionPtr SessionInt = OnlineSub->GetSessionInterface();
+			if (SessionInt.IsValid())
+			{
+				SessionInt->AddOnSessionUserInviteAcceptedDelegate_Handle(
+					FOnSessionUserInviteAcceptedDelegate::CreateUObject(this, &UGameInstance::HandleSessionUserInviteAccepted));
+			}
 		}
 	}
 
@@ -73,13 +77,16 @@ void UGameInstance::Shutdown()
 {
 	ReceiveShutdown();
 
-	const auto OnlineSub = IOnlineSubsystem::Get();
-	if (OnlineSub != nullptr)
+	if (!IsRunningCommandlet())
 	{
-		IOnlineSessionPtr SessionInt = OnlineSub->GetSessionInterface();
-		if (SessionInt.IsValid())
+		const auto OnlineSub = IOnlineSubsystem::Get();
+		if (OnlineSub != nullptr)
 		{
-			SessionInt->ClearOnSessionUserInviteAcceptedDelegate_Handle(OnSessionUserInviteAcceptedDelegateHandle);
+			IOnlineSessionPtr SessionInt = OnlineSub->GetSessionInterface();
+			if (SessionInt.IsValid())
+			{
+				SessionInt->ClearOnSessionUserInviteAcceptedDelegate_Handle(OnSessionUserInviteAcceptedDelegateHandle);
+			}
 		}
 	}
 
diff --git a/Engine/Source/Runtime/Engine/Private/NetConnection.cpp b/Engine/Source/Runtime/Engine/Private/NetConnection.cpp
--- a/Engine/Source/Runtime/Engine/Private/NetConnection.cpp
+++ b/Engine/Source/Runtime/Engine/Private/NetConnection.cpp
@@ -47,6 +47,8 @@ UNetConnection::UNetConnection(const FObjectInitializer& ObjectInitializer)
 ,	InternalAck			( false )
 ,	MaxPacketHandlerBits ( 0 )
 ,	State				( USOCK_Invalid )
+,	Handler()
+,	StatelessConnectComponent()
 ,	PacketOverhead		( 0 )
 ,	ResponseId			( 0 )
 
@@ -210,8 +212,24 @@ void UNetConnection::InitHandler()
 	if (Handler.IsValid())
 	{
 		Handler::Mode Mode = Driver->ServerConnection != nullptr ? Handler::Mode::Client : Handler::Mode::Server;
+
 		Handler->Initialize(Mode);
 
+
+		// Add handling for the stateless connect handshake, for connectionless packets, as the outermost layer
+		TSharedPtr<HandlerComponent> NewComponent =
+			Handler->AddHandler(TEXT("Engine.EngineHandlerComponentFactory(StatelessConnectHandlerComponent)"), true);
+
+		StatelessConnectComponent = StaticCastSharedPtr<StatelessConnectHandlerComponent>(NewComponent);
+
+		if (StatelessConnectComponent.IsValid())
+		{
+			StatelessConnectComponent.Pin()->SetDriver(Driver);
+		}
+
+
+		Handler->InitializeComponents();
+
 		MaxPacketHandlerBits = Handler->GetTotalPacketOverheadBits();
 	}
 }
@@ -499,7 +517,7 @@ void UNetConnection::ReceivedRawPacket( void* InData, int32 Count )
 	{
 		const ProcessedPacket UnProcessedPacket = Handler->Incoming(Data, Count);
 
-		Count = UnProcessedPacket.Count;
+		Count = FMath::DivideAndRoundUp(UnProcessedPacket.CountBits, 8);
 
 		if (Count > 0)
 		{
@@ -585,14 +603,12 @@ void UNetConnection::FlushNet(bool bIgnoreSimulation)
 			// Checked in FlushNet() so each child class doesn't have to implement this
 			if (Driver->IsNetResourceValid())
 			{
-				LowLevelSend(SendBuffer.GetData(), SendBuffer.GetNumBytes());
+				LowLevelSend(SendBuffer.GetData(), SendBuffer.GetNumBytes(), SendBuffer.GetNumBits());
 			}
 		}
 		else if( PacketSimulationSettings.PktOrder )
 		{
-			DelayedPacket& B = *(new(Delayed)DelayedPacket);
-			B.Data.AddUninitialized( SendBuffer.GetNumBytes() );
-			FMemory::Memcpy( B.Data.GetData(), SendBuffer.GetData(), SendBuffer.GetNumBytes() );
+			DelayedPacket& B = *(new(Delayed)DelayedPacket(SendBuffer.GetData(), SendBuffer.GetNumBytes(), SendBuffer.GetNumBits()));
 
 			for( int32 i=Delayed.Num()-1; i>=0; i-- )
 			{
@@ -603,7 +619,7 @@ void UNetConnection::FlushNet(bool bIgnoreSimulation)
 						// Checked in FlushNet() so each child class doesn't have to implement this
 						if (Driver->IsNetResourceValid())
 						{
-							LowLevelSend( (char*)&Delayed[i].Data[0], Delayed[i].Data.Num() );
+							LowLevelSend( (char*)&Delayed[i].Data[0], Delayed[i].Data.Num(), Delayed[i].SizeBits );
 						}
 					}
 					Delayed.RemoveAt( i );
@@ -614,9 +630,8 @@ void UNetConnection::FlushNet(bool bIgnoreSimulation)
 		{
 			if( !PacketSimulationSettings.PktLoss || FMath::FRand()*100.f > PacketSimulationSettings.PktLoss )
 			{
-				DelayedPacket& B = *(new(Delayed)DelayedPacket);
-				B.Data.AddUninitialized( SendBuffer.GetNumBytes() );
-				FMemory::Memcpy( B.Data.GetData(), SendBuffer.GetData(), SendBuffer.GetNumBytes() );
+				DelayedPacket& B = *(new(Delayed)DelayedPacket(SendBuffer.GetData(), SendBuffer.GetNumBytes(), SendBuffer.GetNumBits()));
+
 				B.SendTime = FPlatformTime::Seconds() + (double(PacketSimulationSettings.PktLag)  + 2.0f * (FMath::FRand() - 0.5f) * double(PacketSimulationSettings.PktLagVariance))/ 1000.f;
 			}
 		}
@@ -626,7 +641,7 @@ void UNetConnection::FlushNet(bool bIgnoreSimulation)
 			// Checked in FlushNet() so each child class doesn't have to implement this
 			if (Driver->IsNetResourceValid())
 			{
-				LowLevelSend( SendBuffer.GetData(), SendBuffer.GetNumBytes() );
+				LowLevelSend(SendBuffer.GetData(), SendBuffer.GetNumBytes(), SendBuffer.GetNumBits());
 			}
 #if DO_ENABLE_NET_TEST
 			if( PacketSimulationSettings.PktDup && FMath::FRand()*100.f < PacketSimulationSettings.PktDup )
@@ -634,7 +649,7 @@ void UNetConnection::FlushNet(bool bIgnoreSimulation)
 				// Checked in FlushNet() so each child class doesn't have to implement this
 				if (Driver->IsNetResourceValid())
 				{
-					LowLevelSend( (char*)SendBuffer.GetData(), SendBuffer.GetNumBytes() );
+					LowLevelSend((char*)SendBuffer.GetData(), SendBuffer.GetNumBytes(), SendBuffer.GetNumBits());
 				}
 			}
 		}
@@ -1436,7 +1451,7 @@ void UNetConnection::Tick()
 		{
 			if( FPlatformTime::Seconds() > Delayed[i].SendTime )
 			{
-				LowLevelSend( (char*)&Delayed[i].Data[0], Delayed[i].Data.Num() );
+				LowLevelSend((char*)&Delayed[i].Data[0], Delayed[i].Data.Num(), Delayed[i].SizeBits);
 				Delayed.RemoveAt( i );
 				i--;
 			}
@@ -1511,15 +1526,6 @@ void UNetConnection::Tick()
 	const float DeltaTime	= Driver->Time - LastTickTime;
 	LastTickTime			= Driver->Time;
 
-	// Check to see if too much time is passing between ticks
-	// Setting this to somewhat large value for now, but small enough to catch blocking calls that are causing timeouts
-	const float TickLogThreshold = 5.0f;
-
-	if ( DeltaTime > TickLogThreshold || FrameTime > TickLogThreshold )
-	{
-		UE_LOG( LogNet, Log, TEXT( "UNetConnection::Tick: Very long time between ticks. DeltaTime: %2.2f, Realtime: %2.2f %s" ), DeltaTime, FrameTime, *Describe() );
-	}
-
 	// Handle timeouts.
 	const float Timeout = GetTimeoutValue();
 
@@ -1622,7 +1628,7 @@ void UNetConnection::Tick()
 		/* Send all queued packets */
 		while(QueuedPacket != nullptr)
 		{
-			LowLevelSend(QueuedPacket->Data, QueuedPacket->BytesCount);
+			LowLevelSend(QueuedPacket->Data, FMath::DivideAndRoundUp(QueuedPacket->CountBits, 8u), QueuedPacket->CountBits);
 			delete QueuedPacket;
 			QueuedPacket = Handler->GetQueuedPacket();
 		}
diff --git a/Engine/Source/Runtime/Engine/Private/NetworkDriver.cpp b/Engine/Source/Runtime/Engine/Private/NetworkDriver.cpp
--- a/Engine/Source/Runtime/Engine/Private/NetworkDriver.cpp
+++ b/Engine/Source/Runtime/Engine/Private/NetworkDriver.cpp
@@ -128,6 +128,8 @@ UNetDriver::UNetDriver(const FObjectInitializer& ObjectInitializer)
 ,   bNoTimeouts(false)
 ,   ServerConnection(nullptr)
 ,	ClientConnections()
+,	ConnectionlessHandler()
+,	StatelessConnectComponent()
 ,   World(nullptr)
 ,   Notify(nullptr)
 ,	Time( 0.f )
@@ -615,6 +617,13 @@ void UNetDriver::TickFlush(float DeltaSeconds)
 		ClientConnections[i]->Tick();
 	}
 
+	if (ConnectionlessHandler.IsValid())
+	{
+		ConnectionlessHandler->Tick(DeltaSeconds);
+
+		FlushHandler();
+	}
+
 	if (CVarNetDormancyDraw.GetValueOnGameThread() > 0)
 	{
 		DrawNetDriverDebug();
@@ -861,10 +870,58 @@ bool UNetDriver::InitBase(bool bInitAsClient, FNetworkNotify* InNotify, const FU
 {
 	LastTickDispatchRealtime = FPlatformTime::Seconds();
 	bool bSuccess = InitConnectionClass();
+
+	if (!bInitAsClient)
+	{
+		ConnectionlessHandler.Reset(nullptr);
+	}
+
 	Notify = InNotify;
+
 	return bSuccess;
 }
 
+void UNetDriver::InitConnectionlessHandler()
+{
+	check(!ConnectionlessHandler.IsValid());
+
+	ConnectionlessHandler = MakeUnique<PacketHandler>();
+
+	if (ConnectionlessHandler.IsValid())
+	{
+		ConnectionlessHandler->Initialize(Handler::Mode::Server, true);
+
+
+		// Add handling for the stateless connect handshake, for connectionless packets, as the outermost layer
+		TSharedPtr<HandlerComponent> NewComponent =
+			ConnectionlessHandler->AddHandler(TEXT("Engine.EngineHandlerComponentFactory(StatelessConnectHandlerComponent)"), true);
+
+		StatelessConnectComponent = StaticCastSharedPtr<StatelessConnectHandlerComponent>(NewComponent);
+
+		if (StatelessConnectComponent.IsValid())
+		{
+			StatelessConnectComponent.Pin()->SetDriver(this);
+		}
+
+
+		ConnectionlessHandler->InitializeComponents();
+	}
+}
+
+void UNetDriver::FlushHandler()
+{
+	BufferedPacket* QueuedPacket = ConnectionlessHandler->GetQueuedConnectionlessPacket();
+
+	while (QueuedPacket != nullptr)
+	{
+		LowLevelSend(QueuedPacket->Address, QueuedPacket->Data, QueuedPacket->CountBits);
+
+		delete QueuedPacket;
+
+		QueuedPacket = ConnectionlessHandler->GetQueuedConnectionlessPacket();
+	}
+}
+
 ENetMode UNetDriver::GetNetMode() const
 {
 	// Special case for PIE - forcing dedicated server behavior
@@ -939,6 +996,8 @@ void UNetDriver::Shutdown()
 		}
 	}
 
+	ConnectionlessHandler.Reset(nullptr);
+
 #if DO_ENABLE_NET_TEST
 	PacketSimulationSettings.UnregisterCommands();
 #endif
diff --git a/Engine/Source/Runtime/Engine/Private/PacketHandlers/EngineHandlerComponentFactory.cpp b/Engine/Source/Runtime/Engine/Private/PacketHandlers/EngineHandlerComponentFactory.cpp
new file mode 100644
--- /dev/null
+++ b/Engine/Source/Runtime/Engine/Private/PacketHandlers/EngineHandlerComponentFactory.cpp
@@ -0,0 +1,23 @@
+// Copyright 1998-2016 Epic Games, Inc. All Rights Reserved.
+
+#include "EnginePrivate.h"
+
+#include "PacketHandlers/EngineHandlerComponentFactory.h"
+
+/**
+ * UEngineHandlerComponentFactor
+ */
+UEngineHandlerComponentFactory::UEngineHandlerComponentFactory(const FObjectInitializer& ObjectInitializer)
+	: Super(ObjectInitializer)
+{
+}
+
+TSharedPtr<HandlerComponent> UEngineHandlerComponentFactory::CreateComponentInstance(FString& Options)
+{
+	if (Options == TEXT("StatelessConnectHandlerComponent"))
+	{
+		return MakeShareable(new StatelessConnectHandlerComponent);
+	}
+
+	return nullptr;
+}
diff --git a/Engine/Source/Runtime/Engine/Private/PacketHandlers/StatelessConnectHandlerComponent.cpp b/Engine/Source/Runtime/Engine/Private/PacketHandlers/StatelessConnectHandlerComponent.cpp
new file mode 100644
--- /dev/null
+++ b/Engine/Source/Runtime/Engine/Private/PacketHandlers/StatelessConnectHandlerComponent.cpp
@@ -0,0 +1,664 @@
+// Copyright 1998-2016 Epic Games, Inc. All Rights Reserved.
+
+#include "EnginePrivate.h"
+
+#include "StatelessConnectHandlerComponent.h"
+
+#include "Engine.h"
+#include "SecureHash.h"
+
+
+DEFINE_LOG_CATEGORY(LogHandshake);
+
+
+
+// @todo #JohnB: It is important that banning functionality gets implemented here. This is the earliest/best place,
+//					to reject connections (however, format of 'Address', is defined by net driver - making this tricky, e.g. Steam)
+
+// @todo #JohnB: Do profiling of the handshake check, plus the entire PacketHandler 'Incoming'/'Outgoing' chain,
+//					when under a DoS attack.
+
+// @todo #JohnB: Consider adding an increasing-cost challenge at some stage. Not strictly necessary, but may be nice to have.
+
+// @todo #JohnB: See if you can generalize the bit-alignment code in this class, up to a base class;
+//					perhaps to a generic parent, expressly for implementing 'packet switches' (a bit specifying a unique packet type,
+//					followed by implementing a packet of completely custom format)
+
+
+
+/**
+ * Purpose:
+ *
+ * UDP connections are vulnerable to various types of DoS attacks, particularly spoofing the IP address in UDP packets,
+ * and to protect against this a handshake is needed to verify that the IP is really owned by the client.
+ *
+ * This handshake can be implemented in two ways:
+ *	Stateful:	Here the server stores connection state information (i.e. maintains a UNetConnection) while the handshake is underway,
+ *				allowing spoofed packets to allocate server memory space, prior to handshake verification.
+ *
+ *	Stateless:	Here the server does not store any connection state information, until the handshake completes,
+ *				preventing spoofed packets from allocating server memory space until after the handshake.
+ *
+ * Stateful handshakes are vulnerable to DoS attacks through server memory usage, whereas stateless handshakes are not,
+ * so this implementation uses stateless handshakes.
+ *
+ *
+ *
+ * Handshake Process/Protocol:
+ *
+ * The protocol for the handshake involves the client sending an initial packet to the server,
+ * and the server responding with a unique 'Cookie' value, which the client has to respond with.
+ *
+ * Client - Initial Connect:
+ * [HandshakeBit][SecretIdBit][24:PacketSizeFiller][AlignPad]
+ *																--->
+ *																		Server - Stateless Handshake Challenge:
+ *																		[HandshakeBit][SecretIdBit][4:Timestamp][20:Cookie][AlignPad]
+ *																<---
+ * Client - Stateless Challenge Response:
+ * [HandshakeBit][SecretIdBit][4:Timestamp][20:Cookie][AlignPad]
+ *																--->
+ *																		Server:
+ *																		Ignore, or create UNetConnection.
+ *
+ *
+ *	- HandshakeBit:		Bit signifying whether a packet is a handshake packet. Applied to all game packets.
+ *	- SecretIdBit:		For handshake packets, specifies which HandshakeSecret array was used to generate Cookie.
+ *	- Timestamp:		Server timestamp, from the moment the handshake challenge was sent.
+ *	- Cookie:			Cookie generated by the server, which the client must reply with.
+ *	- AlignPad:			Handshake packets and PacketHandler's in general, require complex padding of packets. See ParseHandshakePacket.
+ *
+ *	- PacketSizeFiller:	Pads the client packet with blank information, so that the initial client packet,
+ *						is the same size as the server response packet.
+ *
+ *						The server will ignore initial packets below/above this length. This prevents hijacking of game servers,
+ *						for use in 'DRDoS' reflection amplification attacks.
+ *
+ *
+ *
+ * Game Protocol Changes:
+ *
+ * Every game (as opposed to handshake) packet starts with an extra bit, represented by [HandshakeBit], and game packets set this to 0.
+ * This is the only change to game packets. When HandshakeBit is set to 1, the separate protocol above is used for handshake packets.
+ *
+ *
+ *
+ * HandshakeSecret/Cookie:
+ *
+ * The Cookie value is used to uniquely identify and perform a handshake with a connecting client,
+ * but only the server can generate and recognize valid cookies, and the server must do this without storing any connection state data.
+ *
+ * To do this, the server stores 2 large random HandshakeSecret values, that only the server knows,
+ * and combines that with data unique to the client connection (IP and Port), plus a server Timestamp, as part of generating the cookie.
+ *
+ * This data is then combined using a special HMAC hashing function, used specifically for authentication, to generate the cookie:
+ *	Cookie = HMAC(HandshakeSecret, Timestamp + Client IP + Client Port)
+ *
+ * When the client responds to the handshake challenge, sending back TimeStamp and the Cookie,
+ * the server will be able to collect all the remaining information it needs from the client packet (Client IP, Client Port),
+ * plus the HandshakeSecret, to be able to regenerate the Cookie from scratch, and verify that the regenerated cookie,
+ * is the same as the one the client sent back.
+ *
+ * No connection state data needs to be stored in order to do this, so this allows a stateless handshake.
+ *
+ *
+ * In addition, HandshakeSecret updates every 15 + Rand(0,5) seconds (with previous value being stored/accepted for same amount of time)
+ * in order to limit packet replay attacks, where a valid cookie can be reused multiple times.
+ *
+ * Checks on the handshake Timestamp, especially when combined with 5 second variance above, compliment this in limiting replay attacks.
+ */
+
+
+/**
+ * Debug Defines
+ */
+
+// Used to test adding extra dud bits to handshake packets (from 0-7), to test packet alignment padding.
+#if !UE_BUILD_SHIPPING
+	#define PAD_TEST_BITS			0
+#else
+	#define PAD_TEST_BITS			0	// Don't touch - adjust above for testing
+#endif
+
+// Enables packetloss testing, which should be tested by connecting/reconnecting to a server a couple dozen times.
+// Every such connection attempt should eventually succeed/recover automatically - if any fail, something's broken.
+#define PACKETLOSS_TEST 0
+
+
+/**
+ * Defines
+ */
+
+#define HANDSHAKE_PACKET_SIZE_BITS	(194 + PAD_TEST_BITS)
+#define HANDSHAKE_PACKET_ALIGNMENT	(HANDSHAKE_PACKET_SIZE_BITS % 8)
+
+// The maximum number of bits that may be added, for aligning a handshake packet
+#define HANDSHAKE_PACKET_ALIGN_MAX	15
+
+
+// The number of seconds between secret value updates, and the random variance applied to this
+#define SECRET_UPDATE_TIME			15.f
+#define SECRET_UPDATE_TIME_VARIANCE	5.f
+
+// The maximum allowed lifetime (in seconds) of any one handshake cookie
+#define MAX_COOKIE_LIFETIME			((SECRET_UPDATE_TIME + SECRET_UPDATE_TIME_VARIANCE) * (float)SECRET_COUNT)
+
+
+/**
+ * StatelessConnectHandlerComponent
+ */
+
+StatelessConnectHandlerComponent::StatelessConnectHandlerComponent()
+	: HandlerComponent()
+	, Driver(nullptr)
+	, HandshakeSecret()
+	, ActiveSecret(255)
+	, LastSecretUpdateTimestamp(0.f)
+	, LastChallengeSuccessAddress(TEXT(""))
+	, LastClientSendTimestamp(0.0)
+	, bConnectConfirmed(false)
+	, LastSecretId(0)
+	, LastTimestamp(0.f)
+{
+	SetActive(true);
+}
+
+void StatelessConnectHandlerComponent::SendInitialConnect()
+{
+	UNetConnection* ServerConn = (Driver != nullptr ? Driver->ServerConnection : NULL);
+
+	if (ServerConn != nullptr)
+	{
+		FBitWriter InitialPacket(HANDSHAKE_PACKET_SIZE_BITS + HANDSHAKE_PACKET_ALIGN_MAX);
+		uint8 bHandshakePacket = 1;
+
+		InitialPacket.WriteBit(bHandshakePacket);
+
+
+		// In order to prevent DRDoS reflection amplification attacks, clients must pad the packet to match server handshake packet size
+		uint8 SecretIdPad = 0;
+		uint8 PacketSizeFiller[24];
+
+		InitialPacket.WriteBit(SecretIdPad);
+
+#if !UE_BUILD_SHIPPING && PAD_TEST_BITS > 0
+		UE_LOG(LogHandshake, Log, TEXT("Handshake PAD_TEST_BITS: %i"), PAD_TEST_BITS);
+
+		InitialPacket.SerializeBits(&bHandshakePacket, PAD_TEST_BITS);
+#endif
+
+		FMemory::Memzero(PacketSizeFiller, ARRAY_COUNT(PacketSizeFiller));
+		InitialPacket.Serialize(PacketSizeFiller, ARRAY_COUNT(PacketSizeFiller));
+
+
+
+		AlignHandshakePacket(InitialPacket);
+
+
+		// Disable PacketHandler parsing, and send the raw packet
+		Handler->SetRawSend(true);
+
+#if !UE_BUILD_SHIPPING && PACKETLOSS_TEST
+		bool bRandFail = FMath::RandBool();
+
+		if (bRandFail)
+		{
+			UE_LOG(LogHandshake, Log, TEXT("Triggering random initial connect packet fail."));
+		}
+
+		if (!bRandFail)
+#endif
+		{
+			ServerConn->LowLevelSend(InitialPacket.GetData(), InitialPacket.GetNumBytes(), InitialPacket.GetNumBits());
+		}
+
+		Handler->SetRawSend(false);
+
+		LastClientSendTimestamp = FPlatformTime::Seconds();
+	}
+	else
+	{
+		UE_LOG(LogHandshake, Error, TEXT("Tried to send handshake connect packet without a server connection."));
+	}
+}
+
+void StatelessConnectHandlerComponent::SendConnectChallenge(FString ClientAddress)
+{
+	if (Driver != nullptr)
+	{
+		FBitWriter ChallengePacket(HANDSHAKE_PACKET_SIZE_BITS + HANDSHAKE_PACKET_ALIGN_MAX);
+		uint8 bHandshakePacket = 1;
+		float Timestamp = (Driver != nullptr ? Driver->Time : -1.f);
+		uint8 Cookie[20];
+
+		GenerateCookie(ClientAddress, ActiveSecret, Timestamp, Cookie);
+
+		ChallengePacket.WriteBit(bHandshakePacket);
+		ChallengePacket.WriteBit(ActiveSecret);
+
+
+#if !UE_BUILD_SHIPPING && PAD_TEST_BITS > 0
+		ChallengePacket.SerializeBits(&bHandshakePacket, PAD_TEST_BITS);
+#endif
+
+		ChallengePacket << Timestamp;
+		ChallengePacket.Serialize(Cookie, ARRAY_COUNT(Cookie));
+
+		UE_LOG( LogHandshake, Log, TEXT( "SendConnectChallenge. Timestamp: %f, Cookie: %s" ), Timestamp, *FString::FromBlob( Cookie, ARRAY_COUNT( Cookie ) ) );
+
+		AlignHandshakePacket(ChallengePacket);
+
+		
+#if !UE_BUILD_SHIPPING && PACKETLOSS_TEST
+		bool bRandFail = FMath::RandBool();
+
+		if (bRandFail)
+		{
+			UE_LOG(LogHandshake, Log, TEXT("Triggering random connect challenge packet fail."));
+		}
+
+		if (!bRandFail)
+#endif
+		{
+			// 'SetRawSend' is not used here, as that only affects UNetConnections
+			Driver->LowLevelSend(ClientAddress, ChallengePacket.GetData(), ChallengePacket.GetNumBits());
+		}
+	}
+	else
+	{
+#if !UE_BUILD_SHIPPING
+		UE_LOG(LogHandshake, Error, TEXT("Tried to send handshake challenge packet without a net driver."));
+#endif
+	}
+}
+
+void StatelessConnectHandlerComponent::SendChallengeResponse(uint8 InSecretId, float InTimestamp, uint8 InCookie[20])
+{
+	UNetConnection* ServerConn = (Driver != nullptr ? Driver->ServerConnection : NULL);
+
+	if (ServerConn != nullptr)
+	{
+		FBitWriter ResponsePacket(HANDSHAKE_PACKET_SIZE_BITS + HANDSHAKE_PACKET_ALIGN_MAX);
+		uint8 bHandshakePacket = 1;
+
+		ResponsePacket.WriteBit(bHandshakePacket);
+		ResponsePacket.WriteBit(InSecretId);
+
+#if !UE_BUILD_SHIPPING && PAD_TEST_BITS > 0
+		ResponsePacket.SerializeBits(&bHandshakePacket, PAD_TEST_BITS);
+#endif
+
+		ResponsePacket << InTimestamp;
+		ResponsePacket.Serialize(InCookie, 20);
+
+		UE_LOG( LogHandshake, Log, TEXT( "SendChallengeResponse. Timestamp: %f, Cookie: %s" ), InTimestamp, *FString::FromBlob( InCookie, 20 ) );
+
+		AlignHandshakePacket(ResponsePacket);
+
+
+		// Disable PacketHandler parsing, and send the raw packet
+		Handler->SetRawSend(true);
+
+#if !UE_BUILD_SHIPPING && PACKETLOSS_TEST
+		bool bRandFail = FMath::RandBool();
+
+		if (bRandFail)
+		{
+			UE_LOG(LogHandshake, Log, TEXT("Triggering random challenge response packet fail."));
+		}
+
+		if (!bRandFail)
+#endif
+		{
+			ServerConn->LowLevelSend(ResponsePacket.GetData(), ResponsePacket.GetNumBytes(), ResponsePacket.GetNumBits());
+		}
+
+		Handler->SetRawSend(false);
+
+		LastClientSendTimestamp = FPlatformTime::Seconds();
+		LastTimestamp = InTimestamp;
+		FMemory::Memcpy(LastCookie, InCookie, ARRAY_COUNT(LastCookie));
+	}
+	else
+	{
+		UE_LOG(LogHandshake, Error, TEXT("Tried to send handshake response packet without a server connection."));
+	}
+}
+
+void StatelessConnectHandlerComponent::AlignHandshakePacket(FBitWriter& HandshakePacket)
+{
+	check(HandshakePacket.GetNumBits() == HANDSHAKE_PACKET_SIZE_BITS);
+
+	// If the handshake packet overshoots normal packet bit alignment, overshooting bits will get chopped, unless another byte is added.
+	bool bOvershootsAlignment = HANDSHAKE_PACKET_ALIGNMENT > Handler->GetPacketBitAlignment();
+	const uint8 AlignBitCount = (8 - HANDSHAKE_PACKET_ALIGNMENT) + (bOvershootsAlignment ? 8 : 0);
+	uint16 AlignPad = 0;
+
+	if (AlignBitCount > 0)
+	{
+		HandshakePacket.SerializeBits(&AlignPad, AlignBitCount);
+	}
+
+	check(!HandshakePacket.IsError());
+	check((HandshakePacket.GetNumBits() % 8) == 0);
+}
+
+void StatelessConnectHandlerComponent::SetDriver(UNetDriver* InDriver)
+{
+	Driver = InDriver;
+
+	if (Handler->Mode == Handler::Mode::Server && Driver->StatelessConnectComponent.HasSameObject(this))
+	{
+		UpdateSecret();
+	}
+}
+
+void StatelessConnectHandlerComponent::Initialize()
+{
+	// On the server, initializes immediately. Clientside doesn't initialize until handshake completes.
+	if (Handler->Mode == Handler::Mode::Server)
+	{
+		Initialized();
+	}
+}
+
+void StatelessConnectHandlerComponent::Incoming(FBitReader& Packet)
+{
+	bool bHandshakePacket = !!Packet.ReadBit() && !Packet.IsError();
+
+	if (bHandshakePacket)
+	{
+		uint8 SecretId = 0;
+		float Timestamp = 1.f;
+		uint8 Cookie[20];
+
+		bHandshakePacket = ParseHandshakePacket(Packet, SecretId, Timestamp, Cookie);
+
+
+		if (bHandshakePacket)
+		{
+			if (Handler->Mode == Handler::Mode::Client)
+			{
+				if (State == Handler::Component::State::UnInitialized)
+				{
+					SendChallengeResponse(SecretId, Timestamp, Cookie);
+
+					// Now finish initializing the handler - flushing the queued packet buffer in the process.
+					Initialized();
+				}
+				else
+				{
+					// Ignore, could be a dupe/out-of-order challenge packet
+				}
+			}
+			else //if (Handler->Mode == Handler::Mode::Server)
+			{
+				// It's possible for a client to send HandshakeResponsePacket1, and - after a timeout - HandshakeResponsePacket2,
+				// where HandshakeResponsePacket2 triggers connection creation and HandshakeResponsePacket1 arrives out-of-order/lagged.
+				// Just ignore these packets.
+			}
+		}
+#if !UE_BUILD_SHIPPING
+		else
+		{
+			UE_LOG(LogHandshake, Log, TEXT("Error reading handshake packet."));
+		}
+#endif
+	}
+#if !UE_BUILD_SHIPPING
+	else if (Packet.IsError())
+	{
+		UE_LOG(LogHandshake, Log, TEXT("Error reading handshake bit from packet."));
+	}
+#endif
+}
+
+void StatelessConnectHandlerComponent::Outgoing(FBitWriter& Packet)
+{
+	// All UNetConnection packets must specify a zero bHandshakePacket value
+	FBitWriter NewPacket(Packet.GetNumBits()+1);
+	uint8 bHandshakePacket = 0;
+
+	NewPacket.WriteBit(bHandshakePacket);
+	NewPacket.SerializeBits(Packet.GetData(), Packet.GetNumBits());
+
+	Packet.Reset();
+	Packet.SerializeBits(NewPacket.GetData(), NewPacket.GetNumBits());
+}
+
+void StatelessConnectHandlerComponent::IncomingConnectionless(FString Address, FBitReader& Packet)
+{
+	bool bHandshakePacket = !!Packet.ReadBit() && !Packet.IsError();
+
+	LastChallengeSuccessAddress.Empty();
+
+	if (bHandshakePacket)
+	{
+		uint8 SecretId = 0;
+		float Timestamp = 1.f;
+		uint8 Cookie[20];
+
+		bHandshakePacket = ParseHandshakePacket(Packet, SecretId, Timestamp, Cookie);
+
+		if (bHandshakePacket)
+		{
+			if (Handler->Mode == Handler::Mode::Server)
+			{
+				bool bInitialConnect = Timestamp == 0.f;
+
+				if (bInitialConnect)
+				{
+					SendConnectChallenge(Address);
+				}
+				// Challenge response
+				else if (Driver != nullptr)
+				{
+					bool bChallengeSuccess = false;
+					float CookieDelta = Driver->Time - Timestamp;
+					float SecretDelta = Timestamp - LastSecretUpdateTimestamp;
+					bool bValidCookieLifetime = CookieDelta > 0.0 && (MAX_COOKIE_LIFETIME - CookieDelta) > 0.f;
+					bool bValidSecretIdTimestamp = (SecretId == ActiveSecret) ? (SecretDelta >= 0.f) : (SecretDelta <= 0.f);
+
+					if (bValidCookieLifetime && bValidSecretIdTimestamp)
+					{
+						// Regenerate the cookie from the packet info, and see if the received cookie matches the regenerated one
+						uint8 RegenCookie[20];
+
+						GenerateCookie(Address, SecretId, Timestamp, RegenCookie);
+
+						bChallengeSuccess = FMemory::Memcmp(Cookie, RegenCookie, 20) == 0;
+					}
+
+					if (bChallengeSuccess)
+					{
+						LastChallengeSuccessAddress = Address;
+					}
+				}
+			}
+		}
+#if !UE_BUILD_SHIPPING
+		else
+		{
+			UE_LOG(LogHandshake, Log, TEXT("Error reading handshake packet."));
+		}
+#endif
+	}
+#if !UE_BUILD_SHIPPING
+	else if (Packet.IsError())
+	{
+		UE_LOG(LogHandshake, Log, TEXT("Error reading handshake bit from packet."));
+	}
+#endif
+}
+
+bool StatelessConnectHandlerComponent::ParseHandshakePacket(FBitReader& Packet, uint8& OutSecretId, float& OutTimestamp,
+															uint8 (&OutCookie)[20])
+{
+	bool bValidPacket = false;
+
+	/**
+	 * Handshake packet bit padding differs, depending on whether handshake packets undershoot or overshoot normal packet bit alignment.
+	 *
+	 * '1' = end-of-packet/desired-alignment
+	 * '*' = bits cut off by PacketHandler, based on GetPacketBitAlignment
+	 * 'x' = padding bits, added so that PacketHandler doesn't cut out any handshake packet bits
+	 *
+	 * Matched alignment:
+	 *		Normal packet alignment:			0000 1***
+	 *		Handshake packet alignment:			0000 1***
+	 *		Padding bits:						0
+	 *
+	 * Undershoot alignment:
+	 *		Normal packet alignment:			0001 ****
+	 *		Handshake packet alignment:			001x ****
+	 *		Padding bits:						1
+	 *
+	 * Overshoot alignment:
+	 *		Normal packet alignment:			001* ****
+	 *		Handshake packet alignment:			0001 xxxx	xxx* ****
+	 *		Padding bits:						7
+	 */
+	uint8 PaddingBits = (8 + ((int32)Handler->GetPacketBitAlignment() - (int32)HANDSHAKE_PACKET_ALIGNMENT)) % 8;
+
+
+	// Only accept handshake packets of precisely the right size
+	if (Packet.GetBitsLeft() == (HANDSHAKE_PACKET_SIZE_BITS - 1 + PaddingBits))
+	{
+		OutSecretId = (uint8)Packet.ReadBit();
+
+#if !UE_BUILD_SHIPPING && PAD_TEST_BITS > 0
+		uint8 EatTestBits;
+
+		Packet.SerializeBits(&EatTestBits, PAD_TEST_BITS);
+#endif
+
+		Packet << OutTimestamp;
+
+		Packet.Serialize(OutCookie, 20);
+
+		if (PaddingBits > 0)
+		{
+			uint8 EatPad;
+
+			Packet.SerializeBits(&EatPad, PaddingBits);
+		}
+
+		bValidPacket = !Packet.IsError();
+	}
+
+	return bValidPacket;
+}
+
+void StatelessConnectHandlerComponent::GenerateCookie(FString ClientAddress, uint8 SecretId, float Timestamp, uint8 (&OutCookie)[20])
+{
+	// @todo #JohnB: Add cpu stats tracking, like what Oodle does upon compression
+	//					NOTE: Being serverside, will only show up in .uprof, not on any 'stat' commands. Still necessary though.
+
+	TArray<uint8> CookieData;
+	FMemoryWriter CookieArc(CookieData);
+
+	CookieArc << Timestamp;
+	CookieArc << ClientAddress;
+
+	FSHA1::HMACBuffer(HandshakeSecret[!!SecretId].GetData(), SECRET_BYTE_SIZE, CookieData.GetData(), CookieData.Num(), OutCookie);
+}
+
+void StatelessConnectHandlerComponent::UpdateSecret()
+{
+	LastSecretUpdateTimestamp = Driver != nullptr ? Driver->Time : 0.f;
+
+	// On first update, update both secrets
+	if (ActiveSecret == 255)
+	{
+		// NOTE: The size of this may be excessive.
+		HandshakeSecret[0].AddUninitialized(SECRET_BYTE_SIZE);
+		HandshakeSecret[1].AddUninitialized(SECRET_BYTE_SIZE);
+
+		TArray<uint8>& CurArray = HandshakeSecret[1];
+
+		for (int32 i=0; i<SECRET_BYTE_SIZE; i++)
+		{
+			CurArray[i] = FMath::Rand() % 255;
+		}
+
+		ActiveSecret = 0;
+	}
+	else
+	{
+		ActiveSecret = (uint8)!ActiveSecret;
+	}
+
+	TArray<uint8>& CurArray = HandshakeSecret[ActiveSecret];
+
+	for (int32 i=0; i<SECRET_BYTE_SIZE; i++)
+	{
+		CurArray[i] = FMath::Rand() % 255;
+	}
+}
+
+int32 StatelessConnectHandlerComponent::GetPacketOverheadBits()
+{
+	return 1;
+}
+
+uint8 StatelessConnectHandlerComponent::GetBitAlignment()
+{
+	// Only return the bit alignment of bHandshakeBit, as this is the only bit that affects all game packets.
+	// Alignment handling for handshake packets is done locally, as alignment may differ from main game packets.
+	return 1;
+}
+
+bool StatelessConnectHandlerComponent::DoesResetBitAlignment()
+{
+	return false;
+}
+
+void StatelessConnectHandlerComponent::Tick(float DeltaTime)
+{
+	if (Handler->Mode == Handler::Mode::Client)
+	{
+		if (!bConnectConfirmed && LastClientSendTimestamp != 0.0)
+		{
+			UNetConnection* ServerConn = Driver != nullptr ? Driver->ServerConnection : nullptr;
+			double LastSendTimeDiff = FPlatformTime::Seconds() - LastClientSendTimestamp;
+
+			// Confirm a connection, by checking to see if any packets have been acked
+			bConnectConfirmed = ServerConn != nullptr && ServerConn->OutAckPacketId != -1;
+
+			if (!bConnectConfirmed && LastSendTimeDiff > 1.0)
+			{
+				if (State == Handler::Component::State::UnInitialized)
+				{
+					UE_LOG(LogHandshake, Verbose, TEXT("Initial handshake packet timeout - resending."));
+
+					SendInitialConnect();
+				}
+				else if (LastTimestamp != 0.f)
+				{
+					UE_LOG(LogHandshake, Verbose, TEXT("Challenge response packet timeout - resending."));
+
+					SendChallengeResponse(LastSecretId, LastTimestamp, LastCookie);
+				}
+			}
+		}
+	}
+	else // if (Handler->Mode == Handler::Mode::Server)
+	{
+		bool bConnectionlessHandler = Driver != nullptr && Driver->StatelessConnectComponent.HasSameObject(this);
+
+		if (bConnectionlessHandler)
+		{
+			static float CurVariance = FMath::FRandRange(0.f, SECRET_UPDATE_TIME_VARIANCE);
+
+			// Update the secret value periodically, to reduce replay attacks. Also adds a bit of randomness to the timing of this,
+			// so that handshake Timestamp checking as an added method of reducing replay attacks, is more effective.
+			if (((Driver->Time - LastSecretUpdateTimestamp) - (SECRET_UPDATE_TIME + CurVariance)) > 0.f)
+			{
+				CurVariance = FMath::FRandRange(0.f, SECRET_UPDATE_TIME_VARIANCE);
+
+				UpdateSecret();
+			}
+		}
+	}
+}
+
diff --git a/Engine/Source/Runtime/Engine/Private/PendingNetGame.cpp b/Engine/Source/Runtime/Engine/Private/PendingNetGame.cpp
--- a/Engine/Source/Runtime/Engine/Private/PendingNetGame.cpp
+++ b/Engine/Source/Runtime/Engine/Private/PendingNetGame.cpp
@@ -38,15 +38,24 @@ void UPendingNetGame::InitNetDriver()
 
 		if( NetDriver->InitConnect( this, URL, ConnectionError ) )
 		{
+			UNetConnection* ServerConn = NetDriver->ServerConnection;
+
+			// Kick off the connection handshake
+			if (ServerConn->StatelessConnectComponent.IsValid())
+			{
+				ServerConn->StatelessConnectComponent.Pin()->SendInitialConnect();
+			}
+
+
 			// Send initial message.
 			uint8 IsLittleEndian = uint8(PLATFORM_LITTLE_ENDIAN);
 			check(IsLittleEndian == !!IsLittleEndian); // should only be one or zero
 			
 			uint32 LocalNetworkVersion = FNetworkVersion::GetLocalNetworkVersion();
 
-			FNetControlMessage<NMT_Hello>::Send( NetDriver->ServerConnection, IsLittleEndian, LocalNetworkVersion );
+			FNetControlMessage<NMT_Hello>::Send(ServerConn, IsLittleEndian, LocalNetworkVersion);
 
-			NetDriver->ServerConnection->FlushNet();
+			ServerConn->FlushNet();
 		}
 		else
 		{
diff --git a/Engine/Source/Runtime/Engine/Public/PacketHandlers/StatelessConnectHandlerComponent.h b/Engine/Source/Runtime/Engine/Public/PacketHandlers/StatelessConnectHandlerComponent.h
new file mode 100644
--- /dev/null
+++ b/Engine/Source/Runtime/Engine/Public/PacketHandlers/StatelessConnectHandlerComponent.h
@@ -0,0 +1,183 @@
+// Copyright 1998-2016 Epic Games, Inc. All Rights Reserved.
+
+#pragma once
+
+#include "PacketHandler.h"
+
+
+DECLARE_LOG_CATEGORY_EXTERN(LogHandshake, Log, All);
+
+
+/**
+ * Forward Declarations
+ */
+
+class UNetConnection;
+class UNetDriver;
+
+
+/**
+ * Defines
+ */
+
+#define SECRET_BYTE_SIZE 64
+#define SECRET_COUNT 2
+
+
+/**
+ * PacketHandler component for implementing a stateless (non-memory-consuming) connection handshake
+ *
+ * Partially based on the Datagram Transport Layer Security protocol.
+ */
+class StatelessConnectHandlerComponent : public HandlerComponent
+{
+public:
+	/**
+	 * Base constructor
+	 */
+	StatelessConnectHandlerComponent();
+
+
+	/**
+	 * Constructs and sends the initial connect packet, from the client to the server
+	 */
+	ENGINE_API void SendInitialConnect();
+
+private:
+	/**
+	 * Constructs and sends the server response to the initial connect packet, from the server to the client.
+	 *
+	 * @param ClientAddress		The address of the client to send the challenge to.
+	 */
+	void SendConnectChallenge(FString ClientAddress);
+
+	/**
+	 * Constructs and sends the final handshake challenge response packet, from the client to the server
+	 *
+	 * @param InSecretId		Which of the two server HandshakeSecret values this uses
+	 * @param InTimestamp		The timestamp value to send
+	 * @param InCookie			The cookie value to send
+	 */
+	void SendChallengeResponse(uint8 InSecretId, float InTimestamp, uint8 InCookie[20]);
+
+
+	/**
+	 * Pads the handshake packet, to match the PacketBitAlignment of the PacketHandler, so that it will parse correctly.
+	 *
+	 * @param HandshakePacket	The handshake packet to be aligned.
+	 */
+	void AlignHandshakePacket(FBitWriter& HandshakePacket);
+
+public:
+	/**
+	 * Whether or not the specified connection address, has just passed the connection handshake challenge.
+	 *
+	 * @param Address	The address (including port, for UIpNetDriver) being checked
+	 */
+	FORCEINLINE bool HasPassedChallenge(FString Address)
+	{
+		return LastChallengeSuccessAddress == Address;
+	}
+
+
+	/**
+	 * Sets the net driver this handler is associated with
+	 *
+	 * @param InDriver	The net driver to set
+	 */
+	void SetDriver(UNetDriver* InDriver);
+
+
+protected:
+	virtual void Initialize() override;
+
+	virtual void Incoming(FBitReader& Packet) override;
+
+	virtual void Outgoing(FBitWriter& Packet) override;
+
+	virtual void IncomingConnectionless(FString Address, FBitReader& Packet) override;
+
+	virtual void OutgoingConnectionless(FString Address, FBitWriter& Packet) override
+	{
+	}
+
+	virtual bool CanReadUnaligned() override
+	{
+		return true;
+	}
+
+	virtual int32 GetPacketOverheadBits() override;
+
+	virtual uint8 GetBitAlignment() override;
+
+	virtual bool DoesResetBitAlignment() override;
+
+	virtual void Tick(float DeltaTime) override;
+
+private:
+	/**
+	 * Parses an incoming handshake packet (does not parse the handshake bit though)
+	 *
+	 * @param Packet		The packet the handshake is being parsed from
+	 * @param OutSecretId	Which of the two serverside HandshakeSecret values this is based on
+	 * @param OutTimestamp	The server timestamp, from the moment the challenge was sent (or 0.f if from the client)
+	 * @param OutCookie		A unique identifier, generated by the server, which the client must reply with (or 0, for initial packet)
+	 * @return				Whether or not the handshake packet was parsed successfully
+	 */
+	bool ParseHandshakePacket(FBitReader& Packet, uint8& OutSecretId, float& OutTimestamp, uint8 (&OutCookie)[20]);
+
+	/**
+	 * Takes the client address plus server timestamp, and outputs a deterministic cookie value
+	 *
+	 * @param ClientAddress		The address of the client, including port
+	 * @param SecretId			Which of the two HandshakeSecret values the cookie will be based on
+	 * @param TimeStamp			The serverside timestamp
+	 * @param OutCookie			Outputs the generated cookie value.
+	 */
+	void GenerateCookie(FString ClientAddress, uint8 SecretId, float Timestamp, uint8 (&OutCookie)[20]);
+
+	/**
+	 * Generates a new HandshakeSecret value
+	 */
+	void UpdateSecret();
+
+
+private:
+	/** The net driver associated with this handler - for performing connectionless sends */
+	UNetDriver* Driver;
+
+
+	/** Serverside variables */
+
+	/** The serverside-only 'secret' value, used to help with generating cookies. */
+	TArray<uint8> HandshakeSecret[SECRET_COUNT];
+
+	/** Which of the two secret values above is active (values are changed frequently, to limit replay attacks) */
+	uint8 ActiveSecret;
+
+	/** The time of the last secret value update */
+	float LastSecretUpdateTimestamp;
+
+	/** The last address to successfully complete the handshake challenge */
+	FString LastChallengeSuccessAddress;
+
+
+	/** Clientside variables */
+
+	/** The last time a handshake packet was sent - used for detecting failed sends. */
+	double LastClientSendTimestamp;
+
+	/** Whether or not the connection has been confirmed as active. */
+	bool bConnectConfirmed;
+
+
+	/** The SecretId value of the last challenge response sent */
+	uint8 LastSecretId;
+
+	/** The Timestamp value of the last challenge response sent */
+	float LastTimestamp;
+
+	/** The Cookie value of the last challenge response sent */
+	uint8 LastCookie[20];
+};
+
diff --git a/Engine/Source/Runtime/NetworkReplayStreaming/HttpNetworkReplayStreaming/Private/HttpNetworkReplayStreaming.cpp b/Engine/Source/Runtime/NetworkReplayStreaming/HttpNetworkReplayStreaming/Private/HttpNetworkReplayStreaming.cpp
--- a/Engine/Source/Runtime/NetworkReplayStreaming/HttpNetworkReplayStreaming/Private/HttpNetworkReplayStreaming.cpp
+++ b/Engine/Source/Runtime/NetworkReplayStreaming/HttpNetworkReplayStreaming/Private/HttpNetworkReplayStreaming.cpp
@@ -456,7 +456,7 @@ void FHttpNetworkReplayStreamer::UploadHeader()
 	check( StreamChunkIndex == 0 );
 
 	// First upload the header
-	UE_LOG( LogHttpReplay, Log, TEXT( "FHttpNetworkReplayStreamer::UploadHeader. Header. StreamChunkIndex: %i, Size: %i" ), StreamChunkIndex, StreamArchive.Buffer.Num() );
+	UE_LOG( LogHttpReplay, Log, TEXT( "FHttpNetworkReplayStreamer::UploadHeader. Header. StreamChunkIndex: %i, Size: %i" ), StreamChunkIndex, HeaderArchive.Buffer.Num() );
 
 	// Create the Http request and add to pending request list
 	TSharedRef<class IHttpRequest> HttpRequest = FHttpModule::Get().CreateRequest();
diff --git a/Engine/Source/Runtime/Online/HTTP/Private/Curl/CurlHttp.cpp b/Engine/Source/Runtime/Online/HTTP/Private/Curl/CurlHttp.cpp
--- a/Engine/Source/Runtime/Online/HTTP/Private/Curl/CurlHttp.cpp
+++ b/Engine/Source/Runtime/Online/HTTP/Private/Curl/CurlHttp.cpp
@@ -198,6 +198,22 @@ void FCurlHttpRequest::SetHeader(const FString& HeaderName, const FString& Heade
 	Headers.Add(HeaderName, HeaderValue);
 }
 
+void FCurlHttpRequest::AppendToHeader(const FString& HeaderName, const FString& AdditionalHeaderValue)
+{
+	if (!HeaderName.IsEmpty() && !AdditionalHeaderValue.IsEmpty())
+	{
+		FString* PreviousValue = Headers.Find(HeaderName);
+		FString NewValue;
+		if (PreviousValue != nullptr && !PreviousValue->IsEmpty())
+		{
+			NewValue = (*PreviousValue) + TEXT(", ");
+		}
+		NewValue += AdditionalHeaderValue;
+
+		SetHeader(HeaderName, NewValue);
+	}
+}
+
 FString FCurlHttpRequest::GetVerb()
 {
 	return Verb;
@@ -269,15 +285,25 @@ size_t FCurlHttpRequest::ReceiveResponseHeaderCallback(void* Ptr, size_t SizeInB
 
 			UE_LOG(LogHttp, Verbose, TEXT("%p: Received response header '%s'."), this, *Header);
 
-			FString HeaderName, Param;
-			if (Header.Split(TEXT(": "), &HeaderName, &Param))
+			FString HeaderKey, HeaderValue;
+			if (Header.Split(TEXT(": "), &HeaderKey, &HeaderValue))
 			{
-				Response->Headers.Add(HeaderName, Param);
-
-				//Store the content length so OnRequestProgress() delegates have something to work with
-				if (HeaderName == TEXT("Content-Length"))
+				if (!HeaderKey.IsEmpty() && !HeaderValue.IsEmpty())
 				{
-					Response->ContentLength = FCString::Atoi(*Param);
+					FString* PreviousValue = Response->Headers.Find(HeaderKey);
+					FString NewValue;
+					if (PreviousValue != nullptr && !PreviousValue->IsEmpty())
+					{
+						NewValue = (*PreviousValue) + TEXT(", ");
+					}
+					NewValue += HeaderValue;
+					Response->Headers.Add(HeaderKey, NewValue);
+
+					//Store the content length so OnRequestProgress() delegates have something to work with
+					if (HeaderKey == TEXT("Content-Length"))
+					{
+						Response->ContentLength = FCString::Atoi(*HeaderValue);
+					}
 				}
 			}
 			return HeaderSize;
diff --git a/Engine/Source/Runtime/Online/HTTP/Private/Curl/CurlHttp.h b/Engine/Source/Runtime/Online/HTTP/Private/Curl/CurlHttp.h
--- a/Engine/Source/Runtime/Online/HTTP/Private/Curl/CurlHttp.h
+++ b/Engine/Source/Runtime/Online/HTTP/Private/Curl/CurlHttp.h
@@ -123,6 +123,7 @@ public:
 	virtual void SetContent(const TArray<uint8>& ContentPayload) override;
 	virtual void SetContentAsString(const FString& ContentString) override;
 	virtual void SetHeader(const FString& HeaderName, const FString& HeaderValue) override;
+	virtual void AppendToHeader(const FString& HeaderName, const FString& AdditionalHeaderValue) override;
 	virtual bool ProcessRequest() override;
 	virtual FHttpRequestCompleteDelegate& OnProcessRequestComplete() override;
 	virtual FHttpRequestProgressDelegate& OnRequestProgress() override;
diff --git a/Engine/Source/Runtime/Online/HTTP/Private/GenericPlatform/GenericPlatformHttp.cpp b/Engine/Source/Runtime/Online/HTTP/Private/GenericPlatform/GenericPlatformHttp.cpp
--- a/Engine/Source/Runtime/Online/HTTP/Private/GenericPlatform/GenericPlatformHttp.cpp
+++ b/Engine/Source/Runtime/Online/HTTP/Private/GenericPlatform/GenericPlatformHttp.cpp
@@ -29,6 +29,7 @@ public:
 	virtual void SetContent(const TArray<uint8>& ContentPayload) override {}
 	virtual void SetContentAsString(const FString& ContentString) override {}
 	virtual void SetHeader(const FString& HeaderName, const FString& HeaderValue) override {}
+	virtual void AppendToHeader(const FString& HeaderName, const FString& AdditionalHeaderValue) override {};
 	virtual bool ProcessRequest() override { return false; }
 	virtual FHttpRequestCompleteDelegate& OnProcessRequestComplete() override { static FHttpRequestCompleteDelegate RequestCompleteDelegate; return RequestCompleteDelegate; }
 	virtual FHttpRequestProgressDelegate& OnRequestProgress() override { static FHttpRequestProgressDelegate RequestProgressDelegate; return RequestProgressDelegate; }
diff --git a/Engine/Source/Runtime/Online/HTTP/Private/HTML5/HTML5HTTP.cpp b/Engine/Source/Runtime/Online/HTTP/Private/HTML5/HTML5HTTP.cpp
--- a/Engine/Source/Runtime/Online/HTTP/Private/HTML5/HTML5HTTP.cpp
+++ b/Engine/Source/Runtime/Online/HTTP/Private/HTML5/HTML5HTTP.cpp
@@ -144,6 +144,22 @@ void FHTML5HttpRequest::SetContentAsString(const FString& ContentString)
 }
 
 
+void FHTML5HttpRequest::AppendToHeader(const FString& HeaderName, const FString& AdditionalHeaderValue)
+{
+	if (!HeaderName.IsEmpty() && !AdditionalHeaderValue.IsEmpty())
+	{
+		FString* PreviousValue = Headers.Find(HeaderName);
+		FString NewValue;
+		if (PreviousValue != nullptr && !PreviousValue->IsEmpty())
+		{
+			NewValue = (*PreviousValue) + TEXT(", ");
+		}
+		NewValue += AdditionalHeaderValue;
+
+		SetHeader(HeaderName, NewValue);
+	}
+}
+
 FString FHTML5HttpRequest::GetVerb()
 {
 	return Verb;
diff --git a/Engine/Source/Runtime/Online/HTTP/Private/HTML5/HTML5HTTP.h b/Engine/Source/Runtime/Online/HTTP/Private/HTML5/HTML5HTTP.h
--- a/Engine/Source/Runtime/Online/HTTP/Private/HTML5/HTML5HTTP.h
+++ b/Engine/Source/Runtime/Online/HTTP/Private/HTML5/HTML5HTTP.h
@@ -53,7 +53,7 @@ public:
 	virtual void SetURL(const FString& InURL) override;
 	virtual void SetContent(const TArray<uint8>& ContentPayload) override;
 	virtual void SetContentAsString(const FString& ContentString) override;
-	virtual void SetHeader(const FString& HeaderName, const FString& HeaderValue) override;
+	virtual void AppendToHeader(const FString& HeaderName, const FString& AdditionalHeaderValue) override;
 	virtual bool ProcessRequest() override;
 	virtual FHttpRequestCompleteDelegate& OnProcessRequestComplete() override;
 	virtual FHttpRequestProgressDelegate& OnRequestProgress() override;
diff --git a/Engine/Source/Runtime/Online/HTTP/Private/IOS/IOSHTTP.cpp b/Engine/Source/Runtime/Online/HTTP/Private/IOS/IOSHTTP.cpp
--- a/Engine/Source/Runtime/Online/HTTP/Private/IOS/IOSHTTP.cpp
+++ b/Engine/Source/Runtime/Online/HTTP/Private/IOS/IOSHTTP.cpp
@@ -84,6 +84,23 @@ void FIOSHttpRequest::SetHeader(const FString& HeaderName, const FString& Header
 	[Request setValue:[NSString stringWithFString:HeaderValue] forHTTPHeaderField:[NSString stringWithFString:HeaderName]];
 }
 
+void FIOSHttpRequest::AppendToHeader(const FString& HeaderName, const FString& AdditionalHeaderValue)
+{
+    if (!HeaderName.IsEmpty() && !AdditionalHeaderValue.IsEmpty())
+    {
+        NSDictionary* Headers = [Request allHTTPHeaderFields];
+        NSString* PreviousHeaderValuePtr = [Headers objectForKey: HeaderName.GetNSString()];
+        FString PreviousValue(PreviousHeaderValuePtr);
+        FString NewValue;
+        if (PreviousValue != nullptr && !PreviousValue.IsEmpty())
+        {
+            NewValue = PreviousValue + TEXT(", ");
+        }
+        NewValue += AdditionalHeaderValue;
+        
+        SetHeader(HeaderName, NewValue);
+    }
+}
 
 TArray<FString> FIOSHttpRequest::GetAllHeaders() 
 {
diff --git a/Engine/Source/Runtime/Online/HTTP/Private/IOS/IOSHTTP.h b/Engine/Source/Runtime/Online/HTTP/Private/IOS/IOSHTTP.h
--- a/Engine/Source/Runtime/Online/HTTP/Private/IOS/IOSHTTP.h
+++ b/Engine/Source/Runtime/Online/HTTP/Private/IOS/IOSHTTP.h
@@ -38,6 +38,7 @@ public:
 	virtual void SetContent(const TArray<uint8>& ContentPayload) override;
 	virtual void SetContentAsString(const FString& ContentString) override;
 	virtual void SetHeader(const FString& HeaderName, const FString& HeaderValue) override;
+    virtual void AppendToHeader(const FString& HeaderName, const FString& AdditionalHeaderValue) override;
 	virtual bool ProcessRequest() override;
 	virtual FHttpRequestCompleteDelegate& OnProcessRequestComplete() override;
 	virtual FHttpRequestProgressDelegate& OnRequestProgress() override;
diff --git a/Engine/Source/Runtime/Online/HTTP/Private/Mac/MacHTTP.cpp b/Engine/Source/Runtime/Online/HTTP/Private/Mac/MacHTTP.cpp
--- a/Engine/Source/Runtime/Online/HTTP/Private/Mac/MacHTTP.cpp
+++ b/Engine/Source/Runtime/Online/HTTP/Private/Mac/MacHTTP.cpp
@@ -81,6 +81,23 @@ void FMacHttpRequest::SetHeader(const FString& HeaderName, const FString& Header
 	[Request setValue: HeaderValue.GetNSString() forHTTPHeaderField: HeaderName.GetNSString()];
 }
 
+void FMacHttpRequest::AppendToHeader(const FString& HeaderName, const FString& AdditionalHeaderValue)
+{
+    if (!HeaderName.IsEmpty() && !AdditionalHeaderValue.IsEmpty())
+    {
+        NSDictionary* Headers = [Request allHTTPHeaderFields];
+        NSString* PreviousHeaderValuePtr = [Headers objectForKey: HeaderName.GetNSString()];
+        FString PreviousValue(PreviousHeaderValuePtr);
+		FString NewValue;
+		if (PreviousValue != nullptr && !PreviousValue.IsEmpty())
+		{
+			NewValue = PreviousValue + TEXT(", ");
+		}
+		NewValue += AdditionalHeaderValue;
+
+        SetHeader(HeaderName, NewValue);
+	}
+}
 
 TArray<FString> FMacHttpRequest::GetAllHeaders()
 {
diff --git a/Engine/Source/Runtime/Online/HTTP/Private/Mac/MacHTTP.h b/Engine/Source/Runtime/Online/HTTP/Private/Mac/MacHTTP.h
--- a/Engine/Source/Runtime/Online/HTTP/Private/Mac/MacHTTP.h
+++ b/Engine/Source/Runtime/Online/HTTP/Private/Mac/MacHTTP.h
@@ -38,6 +38,7 @@ public:
 	virtual void SetContent(const TArray<uint8>& ContentPayload) override;
 	virtual void SetContentAsString(const FString& ContentString) override;
 	virtual void SetHeader(const FString& HeaderName, const FString& HeaderValue) override;
+	virtual void AppendToHeader(const FString& HeaderName, const FString& AdditionalHeaderValue) override;
 	virtual bool ProcessRequest() override;
 	virtual FHttpRequestCompleteDelegate& OnProcessRequestComplete() override;
 	virtual FHttpRequestProgressDelegate& OnRequestProgress() override;
diff --git a/Engine/Source/Runtime/Online/HTTP/Private/NullHttp.cpp b/Engine/Source/Runtime/Online/HTTP/Private/NullHttp.cpp
--- a/Engine/Source/Runtime/Online/HTTP/Private/NullHttp.cpp
+++ b/Engine/Source/Runtime/Online/HTTP/Private/NullHttp.cpp
@@ -82,6 +82,22 @@ void FNullHttpRequest::SetHeader(const FString& HeaderName, const FString& Heade
 	Headers.Add(HeaderName, HeaderValue);
 }
 
+void FNullHttpRequest::AppendToHeader(const FString& HeaderName, const FString& AdditionalHeaderValue)
+{
+	if (HeaderName.Len() > 0 && AdditionalHeaderValue.Len() > 0)
+	{
+		FString* PreviousValue = Headers.Find(HeaderName);
+		FString NewValue;
+		if (PreviousValue != NULL && PreviousValue->Len() > 0)
+		{
+			NewValue = (*PreviousValue) + TEXT(", ");
+		}
+		NewValue += AdditionalHeaderValue;
+
+		SetHeader(HeaderName, NewValue);
+	}
+}
+
 bool FNullHttpRequest::ProcessRequest()
 {
 	ElapsedTime = 0;
diff --git a/Engine/Source/Runtime/Online/HTTP/Private/NullHttp.h b/Engine/Source/Runtime/Online/HTTP/Private/NullHttp.h
--- a/Engine/Source/Runtime/Online/HTTP/Private/NullHttp.h
+++ b/Engine/Source/Runtime/Online/HTTP/Private/NullHttp.h
@@ -24,6 +24,7 @@ public:
 	virtual void SetContent(const TArray<uint8>& ContentPayload) override;
 	virtual void SetContentAsString(const FString& ContentString) override;
 	virtual void SetHeader(const FString& HeaderName, const FString& HeaderValue) override;
+	virtual void AppendToHeader(const FString& HeaderName, const FString& AdditionalHeaderValue) override;
 	virtual bool ProcessRequest() override;
 	virtual FHttpRequestCompleteDelegate& OnProcessRequestComplete() override;
 	virtual FHttpRequestProgressDelegate& OnRequestProgress() override;
diff --git a/Engine/Source/Runtime/Online/HTTP/Private/Windows/HttpWinInet.cpp b/Engine/Source/Runtime/Online/HTTP/Private/Windows/HttpWinInet.cpp
--- a/Engine/Source/Runtime/Online/HTTP/Private/Windows/HttpWinInet.cpp
+++ b/Engine/Source/Runtime/Online/HTTP/Private/Windows/HttpWinInet.cpp
@@ -439,6 +439,22 @@ void FHttpRequestWinInet::SetHeader(const FString& HeaderName, const FString& He
 	RequestHeaders.Add(HeaderName, HeaderValue);
 }
 
+void FHttpRequestWinInet::AppendToHeader(const FString& HeaderName, const FString& AdditionalHeaderValue)
+{
+	if (!HeaderName.IsEmpty() && !AdditionalHeaderValue.IsEmpty())
+	{
+		FString* PreviousValue = RequestHeaders.Find(HeaderName);
+		FString NewValue;
+		if (PreviousValue != nullptr && !PreviousValue->IsEmpty())
+		{
+			NewValue = (*PreviousValue) + TEXT(", ");
+		}
+		NewValue += AdditionalHeaderValue;
+
+		SetHeader(HeaderName, NewValue);
+	}
+}
+
 bool FHttpRequestWinInet::ProcessRequest()
 {
 	bool bStarted = false;
@@ -1056,9 +1072,16 @@ void FHttpResponseWinInet::ProcessResponseHeaders()
 			FString HeaderKey,HeaderValue;
 			if (HeaderLine.Split(TEXT(":"), &HeaderKey, &HeaderValue, ESearchCase::CaseSensitive))
 			{
-				if (!HeaderKey.IsEmpty())
+				if (!HeaderKey.IsEmpty() && !HeaderValue.IsEmpty())
 				{
-					ResponseHeaders.Add(HeaderKey, HeaderValue.Trim());
+					FString* PreviousValue = ResponseHeaders.Find(HeaderKey);
+					FString NewValue;
+					if (PreviousValue != nullptr && !PreviousValue->IsEmpty())
+					{
+						NewValue = (*PreviousValue) + TEXT(", ");
+					}
+					NewValue += HeaderValue;
+					ResponseHeaders.Add(HeaderKey, NewValue);
 				}
 			}
 			HeaderPtr = DelimiterPtr + 2;
diff --git a/Engine/Source/Runtime/Online/HTTP/Private/Windows/HttpWinInet.h b/Engine/Source/Runtime/Online/HTTP/Private/Windows/HttpWinInet.h
--- a/Engine/Source/Runtime/Online/HTTP/Private/Windows/HttpWinInet.h
+++ b/Engine/Source/Runtime/Online/HTTP/Private/Windows/HttpWinInet.h
@@ -86,6 +86,7 @@ public:
 	virtual void SetContent(const TArray<uint8>& ContentPayload) override;
 	virtual void SetContentAsString(const FString& ContentString) override;
 	virtual void SetHeader(const FString& HeaderName, const FString& HeaderValue) override;
+	virtual void AppendToHeader(const FString& HeaderName, const FString& AdditionalHeaderValue) override;
 	virtual bool ProcessRequest() override;
 	virtual FHttpRequestCompleteDelegate& OnProcessRequestComplete() override;
 	virtual FHttpRequestProgressDelegate& OnRequestProgress() override;
diff --git a/Engine/Source/Runtime/Online/HTTP/Public/HttpRequestAdapter.h b/Engine/Source/Runtime/Online/HTTP/Public/HttpRequestAdapter.h
--- a/Engine/Source/Runtime/Online/HTTP/Public/HttpRequestAdapter.h
+++ b/Engine/Source/Runtime/Online/HTTP/Public/HttpRequestAdapter.h
@@ -30,6 +30,7 @@ public:
 	virtual void                          SetContent(const TArray<uint8>& ContentPayload) override                 { HttpRequest->SetContent(ContentPayload); }
 	virtual void                          SetContentAsString(const FString& ContentString) override                 { HttpRequest->SetContentAsString(ContentString); }
 	virtual void                          SetHeader(const FString& HeaderName, const FString& HeaderValue) override { HttpRequest->SetHeader(HeaderName, HeaderValue); }
+	virtual void                          AppendToHeader(const FString& HeaderName, const FString& AdditionalHeaderValue) override { HttpRequest->AppendToHeader(HeaderName, AdditionalHeaderValue); }
 	virtual const FHttpResponsePtr        GetResponse() const override                                             { return HttpRequest->GetResponse(); }
 	virtual float                         GetElapsedTime() override                                                { return HttpRequest->GetElapsedTime(); }
 	virtual EHttpRequestStatus::Type	  GetStatus() override { return HttpRequest->GetStatus(); }
diff --git a/Engine/Source/Runtime/Online/HTTP/Public/Interfaces/IHttpRequest.h b/Engine/Source/Runtime/Online/HTTP/Public/Interfaces/IHttpRequest.h
--- a/Engine/Source/Runtime/Online/HTTP/Public/Interfaces/IHttpRequest.h
+++ b/Engine/Source/Runtime/Online/HTTP/Public/Interfaces/IHttpRequest.h
@@ -125,6 +125,8 @@ public:
 
 	/**
 	 * Sets optional header info.
+	 * SetHeader for a given HeaderName will overwrite any previous values
+	 * Use AppendToHeader to append more values for the same header
 	 * Content-Length is the only header set for you.
 	 * Required headers depends on the request itself.
 	 * Eg. "multipart/form-data" needed for a form post
@@ -134,6 +136,19 @@ public:
 	 */
 	virtual void SetHeader(const FString& HeaderName, const FString& HeaderValue) = 0;
 
+	/**
+	* Appends to the value already set in the header. 
+	* If there is already content in that header, a comma delimiter is used.
+	* If the header is as of yet unset, the result is the same as calling SetHeader
+	* Content-Length is the only header set for you.
+	* Also see: SetHeader()
+	*
+	* @param HeaderName - Name of the header (ie, Content-Type)
+	* @param AdditionalHeaderValue - Value to add to the existing contents of the specified header.
+	*	comma is inserted between old value and new value, per HTTP specifications
+	*/
+	virtual void AppendToHeader(const FString& HeaderName, const FString& AdditionalHeaderValue) = 0;
+
 	/**
 	 * Called to begin processing the request.
 	 * OnProcessRequestComplete delegate is always called when the request completes or on error if it is bound.
diff --git a/Engine/Source/Runtime/Online/OnlineSubsystem/Public/OnlineSessionSettings.h b/Engine/Source/Runtime/Online/OnlineSubsystem/Public/OnlineSessionSettings.h
--- a/Engine/Source/Runtime/Online/OnlineSubsystem/Public/OnlineSessionSettings.h
+++ b/Engine/Source/Runtime/Online/OnlineSubsystem/Public/OnlineSessionSettings.h
@@ -571,7 +571,7 @@ public:
 	/** Constructor */
 	FOnlineSessionSearch() :
 		SearchState(EOnlineAsyncTaskState::NotStarted),
-		MaxSearchResults(0),
+		MaxSearchResults(1),
 		bIsLanQuery(false),
 		PingBucketSize(0),
 		TimeoutInSeconds(0.0f)
diff --git a/Engine/Source/Runtime/Online/OnlineSubsystemSteam/Private/OnlineSessionAsyncServerSteam.cpp b/Engine/Source/Runtime/Online/OnlineSubsystemSteam/Private/OnlineSessionAsyncServerSteam.cpp
--- a/Engine/Source/Runtime/Online/OnlineSubsystemSteam/Private/OnlineSessionAsyncServerSteam.cpp
+++ b/Engine/Source/Runtime/Online/OnlineSubsystemSteam/Private/OnlineSessionAsyncServerSteam.cpp
@@ -1016,6 +1016,11 @@ void FOnlineAsyncTaskSteamFindServerBase::Tick()
 		}
 #endif
 
+		if (SearchSettings->MaxSearchResults <= 0)
+		{
+			UE_LOG_ONLINE(Warning, TEXT("FOnlineAsyncTaskSteamFindServerBase::Tick - SearchSettings->MaxSearchResults should be greater than 0, but it is currently %d. No search results will be found."), SearchSettings->MaxSearchResults);
+		}
+
 		ServerListRequestHandle = SteamMatchmakingServersPtr->RequestInternetServerList(Subsystem->GetSteamAppId(), &Filters, NumFilters, this);
 		if (ServerListRequestHandle == NULL)
 		{
diff --git a/Engine/Source/Runtime/Online/OnlineSubsystemSteam/Private/OnlineSubsystemSteam.cpp b/Engine/Source/Runtime/Online/OnlineSubsystemSteam/Private/OnlineSubsystemSteam.cpp
--- a/Engine/Source/Runtime/Online/OnlineSubsystemSteam/Private/OnlineSubsystemSteam.cpp
+++ b/Engine/Source/Runtime/Online/OnlineSubsystemSteam/Private/OnlineSubsystemSteam.cpp
@@ -89,16 +89,24 @@ static void WriteSteamAppIdToDisk(int32 SteamAppId)
 	{
 		// Turn off sandbox temporarily to make sure file is where it's always expected
 		FScopeSandboxContext ScopedSandbox(false);
-		FArchive* AppIdFile = IFileManager::Get().CreateFileWriter(*GetSteamAppIdFilename());
-		
-		if (AppIdFile)
+
+		// Access the physical file writer directly so that we still write next to the executable in CotF builds.
+		FString SteamAppIdFilename = GetSteamAppIdFilename();
+		IFileHandle* Handle = IPlatformFile::GetPlatformPhysical().OpenWrite(*SteamAppIdFilename, false, false);
+		if (!Handle)
+		{
+			UE_LOG_ONLINE(Fatal, TEXT("Failed to create file: %s"), *SteamAppIdFilename);
+		}
+		else
 		{
 			FString AppId = FString::Printf(TEXT("%d"), SteamAppId);
-			AppIdFile->Serialize((void*)TCHAR_TO_ANSI(*AppId), AppId.Len());
 
-			// Close the file
-			delete AppIdFile;
-			AppIdFile = nullptr;
+			FBufferArchive Archive;
+			Archive.Serialize((void*)TCHAR_TO_ANSI(*AppId), AppId.Len());
+
+			Handle->Write(Archive.GetData(), Archive.Num());
+			delete Handle;
+			Handle = nullptr;
 		}
 	}
 }
diff --git a/Engine/Source/Runtime/Online/OnlineSubsystemUtils/Classes/IpConnection.h b/Engine/Source/Runtime/Online/OnlineSubsystemUtils/Classes/IpConnection.h
--- a/Engine/Source/Runtime/Online/OnlineSubsystemUtils/Classes/IpConnection.h
+++ b/Engine/Source/Runtime/Online/OnlineSubsystemUtils/Classes/IpConnection.h
@@ -20,7 +20,7 @@ class ONLINESUBSYSTEMUTILS_API UIpConnection : public UNetConnection
 	virtual void InitBase(UNetDriver* InDriver, class FSocket* InSocket, const FURL& InURL, EConnectionState InState, int32 InMaxPacket = 0, int32 InPacketOverhead = 0) override;
 	virtual void InitRemoteConnection(UNetDriver* InDriver, class FSocket* InSocket, const FURL& InURL, const class FInternetAddr& InRemoteAddr, EConnectionState InState, int32 InMaxPacket = 0, int32 InPacketOverhead = 0) override;
 	virtual void InitLocalConnection(UNetDriver* InDriver, class FSocket* InSocket, const FURL& InURL, EConnectionState InState, int32 InMaxPacket = 0, int32 InPacketOverhead = 0) override;
-	virtual void LowLevelSend(void* Data,int32 Count) override;
+	virtual void LowLevelSend(void* Data, int32 CountBytes, int32 CountBits) override;
 	FString LowLevelGetRemoteAddress(bool bAppendPort=false) override;
 	FString LowLevelDescribe() override;
 	virtual int32 GetAddrAsInt(void) override;
diff --git a/Engine/Source/Runtime/Online/OnlineSubsystemUtils/Classes/IpNetDriver.h b/Engine/Source/Runtime/Online/OnlineSubsystemUtils/Classes/IpNetDriver.h
--- a/Engine/Source/Runtime/Online/OnlineSubsystemUtils/Classes/IpNetDriver.h
+++ b/Engine/Source/Runtime/Online/OnlineSubsystemUtils/Classes/IpNetDriver.h
@@ -37,6 +37,7 @@ class ONLINESUBSYSTEMUTILS_API UIpNetDriver : public UNetDriver
 	virtual bool InitListen( FNetworkNotify* InNotify, FURL& LocalURL, bool bReuseAddressAndPort, FString& Error ) override;
 	virtual void ProcessRemoteFunction(class AActor* Actor, class UFunction* Function, void* Parameters, struct FOutParmRec* OutParms, struct FFrame* Stack, class UObject* SubObject = NULL) override;
 	virtual void TickDispatch( float DeltaTime ) override;
+	virtual void LowLevelSend(FString Address, void* Data, int32 CountBits) override;
 	virtual FString LowLevelGetNetworkNumber() override;
 	virtual void LowLevelDestroy() override;
 	virtual class ISocketSubsystem* GetSocketSubsystem() override;
diff --git a/Engine/Source/Runtime/Online/OnlineSubsystemUtils/Private/IpConnection.cpp b/Engine/Source/Runtime/Online/OnlineSubsystemUtils/Private/IpConnection.cpp
--- a/Engine/Source/Runtime/Online/OnlineSubsystemUtils/Private/IpConnection.cpp
+++ b/Engine/Source/Runtime/Online/OnlineSubsystemUtils/Private/IpConnection.cpp
@@ -98,7 +98,7 @@ void UIpConnection::InitRemoteConnection(UNetDriver* InDriver, class FSocket* In
 	SetExpectedClientLoginMsgType( NMT_Hello );
 }
 
-void UIpConnection::LowLevelSend( void* Data, int32 Count )
+void UIpConnection::LowLevelSend(void* Data, int32 CountBytes, int32 CountBits)
 {
 	const uint8* DataToSend = reinterpret_cast<uint8*>(Data);
 
@@ -132,26 +132,28 @@ void UIpConnection::LowLevelSend( void* Data, int32 Count )
 	}
 	 
 	// Process any packet modifiers
-	if(Handler.IsValid())
+	if (Handler.IsValid() && !Handler->GetRawSend())
 	{
-		const ProcessedPacket ProcessedData = Handler->Outgoing(reinterpret_cast<uint8*>(Data), Count);
+		const ProcessedPacket ProcessedData = Handler->Outgoing(reinterpret_cast<uint8*>(Data), CountBits);
+
 		DataToSend = ProcessedData.Data;
-		Count = ProcessedData.Count;
+		CountBytes = FMath::DivideAndRoundUp(ProcessedData.CountBits, 8);
+		CountBits = ProcessedData.CountBits;
 	}
 
 	// Send to remote.
 	int32 BytesSent = 0;
 	CLOCK_CYCLES(Driver->SendCycles);
 
-	if ( Count > MaxPacket )
+	if ( CountBytes > MaxPacket )
 	{
-		UE_LOG( LogNet, Warning, TEXT( "UIpConnection::LowLevelSend: Count > MaxPacketSize! Count: %i, MaxPacket: %i %s" ), Count, MaxPacket, *Describe() );
+		UE_LOG( LogNet, Warning, TEXT( "UIpConnection::LowLevelSend: CountBytes > MaxPacketSize! Count: %i, MaxPacket: %i %s" ), CountBytes, MaxPacket, *Describe() );
 	}
 
-	Socket->SendTo(DataToSend, Count, BytesSent, *RemoteAddr);
+	Socket->SendTo(DataToSend, CountBytes, BytesSent, *RemoteAddr);
 	UNCLOCK_CYCLES(Driver->SendCycles);
 	NETWORK_PROFILER(GNetworkProfiler.FlushOutgoingBunches(this));
-	NETWORK_PROFILER(GNetworkProfiler.TrackSocketSendTo(Socket->GetDescription(),Data,BytesSent,NumPacketIdBits,NumBunchBits,NumAckBits,NumPaddingBits,this));
+	NETWORK_PROFILER(GNetworkProfiler.TrackSocketSendTo(Socket->GetDescription(),DataToSend,BytesSent,NumPacketIdBits,NumBunchBits,NumAckBits,NumPaddingBits,this));
 }
 
 FString UIpConnection::LowLevelGetRemoteAddress(bool bAppendPort)
diff --git a/Engine/Source/Runtime/Online/OnlineSubsystemUtils/Private/IpNetDriver.cpp b/Engine/Source/Runtime/Online/OnlineSubsystemUtils/Private/IpNetDriver.cpp
--- a/Engine/Source/Runtime/Online/OnlineSubsystemUtils/Private/IpNetDriver.cpp
+++ b/Engine/Source/Runtime/Online/OnlineSubsystemUtils/Private/IpNetDriver.cpp
@@ -174,6 +174,9 @@ bool UIpNetDriver::InitListen( FNetworkNotify* InNotify, FURL& LocalURL, bool bR
 		return false;
 	}
 
+
+	InitConnectionlessHandler();
+
 	// Update result URL.
 	//LocalURL.Host = LocalAddr->ToString(false);
 	LocalURL.Port = LocalAddr->GetPort();
@@ -193,7 +196,10 @@ void UIpNetDriver::TickDispatch( float DeltaTime )
 
 	// Process all incoming packets.
 	uint8 Data[MAX_PACKET_SIZE];
+	uint8* DataRef = Data;
+	bool bIgnorePacket = false;
 	TSharedRef<FInternetAddr> FromAddr = SocketSubsystem->CreateInternetAddr();
+
 	for( ; Socket != NULL; )
 	{
 		{
@@ -207,12 +213,21 @@ void UIpNetDriver::TickDispatch( float DeltaTime )
 		}
 
 		int32 BytesRead = 0;
+
 		// Get data, if any.
 		CLOCK_CYCLES(RecvCycles);
 		bool bOk = Socket->RecvFrom(Data, sizeof(Data), BytesRead, *FromAddr);
 		UNCLOCK_CYCLES(RecvCycles);
-		// Handle result.
-		if( bOk == false )
+
+		if (bOk)
+		{
+			// Immediately stop processing, for empty packets (usually a DDoS)
+			if (BytesRead == 0)
+			{
+				break;
+			}
+		}
+		else
 		{
 			ESocketErrors Error = SocketSubsystem->GetLastErrorCode();
 			if(Error == SE_EWOULDBLOCK ||
@@ -317,22 +332,59 @@ void UIpNetDriver::TickDispatch( float DeltaTime )
 			if( !Connection )
 			{
 				// Determine if allowing for client/server connections
-				const bool bAcceptingConnection = Notify->NotifyAcceptingConnection() == EAcceptConnection::Accept;
+				const bool bAcceptingConnection = Notify != nullptr && Notify->NotifyAcceptingConnection() == EAcceptConnection::Accept;
 
 				if (bAcceptingConnection)
 				{
-					Connection = NewObject<UIpConnection>(GetTransientPackage(), NetConnectionClass);
-                    check(Connection);
-					Connection->InitRemoteConnection( this, Socket,  FURL(), *FromAddr, USOCK_Open);
-					Notify->NotifyAcceptedConnection( Connection );
-					AddClientConnection(Connection);
+					bool bPassedChallenge = false;
+
+					bIgnorePacket = true;
+
+					if (ConnectionlessHandler.IsValid() && StatelessConnectComponent.IsValid())
+					{
+						TSharedPtr<StatelessConnectHandlerComponent> StatelessConnect = StatelessConnectComponent.Pin();
+						FString IncomingAddress = FromAddr->ToString(true);
+
+						const ProcessedPacket UnProcessedPacket =
+												ConnectionlessHandler->IncomingConnectionless(IncomingAddress, DataRef, BytesRead);
+
+						bPassedChallenge = StatelessConnect->HasPassedChallenge(IncomingAddress);
+
+						if (bPassedChallenge)
+						{
+							BytesRead = FMath::DivideAndRoundUp(UnProcessedPacket.CountBits, 8);
+
+							if (BytesRead > 0)
+							{
+								DataRef = UnProcessedPacket.Data;
+								bIgnorePacket = false;
+							}
+						}
+					}
+					else
+					{
+						UE_LOG(LogNet, Log,
+								TEXT("Invalid ConnectionlessHandler (%i) or StatelessConnectComponent (%i); can't accept connections."),
+								(int32)(ConnectionlessHandler.IsValid()), (int32)(StatelessConnectComponent.IsValid()));
+					}
+
+					if (bPassedChallenge)
+					{
+						UE_LOG(LogNet, Log, TEXT("Server accepting post-challenge connection from: %s"), *FromAddr->ToString(true));
+
+						Connection = NewObject<UIpConnection>(GetTransientPackage(), NetConnectionClass);
+						check(Connection);
+						Connection->InitRemoteConnection( this, Socket,  FURL(), *FromAddr, USOCK_Open);
+						Notify->NotifyAcceptedConnection( Connection );
+						AddClientConnection(Connection);
+					}
 				}
 			}
 
 			// Send the packet to the connection for processing.
-			if( Connection )
+			if (Connection && !bIgnorePacket)
 			{
-				Connection->ReceivedRawPacket( Data, BytesRead );
+				Connection->ReceivedRawPacket( DataRef, BytesRead );
 			}
 		}
 	}
@@ -346,6 +398,48 @@ void UIpNetDriver::TickDispatch( float DeltaTime )
 	}
 }
 
+void UIpNetDriver::LowLevelSend(FString Address, void* Data, int32 CountBits)
+{
+	bool bValidAddress = !Address.IsEmpty();
+	TSharedRef<FInternetAddr> RemoteAddr = GetSocketSubsystem()->CreateInternetAddr();
+
+	if (bValidAddress)
+	{
+		RemoteAddr->SetIp(*Address, bValidAddress);
+	}
+
+	if (bValidAddress)
+	{
+		const uint8* DataToSend = reinterpret_cast<uint8*>(Data);
+
+		if (ConnectionlessHandler.IsValid())
+		{
+			const ProcessedPacket ProcessedData =
+					ConnectionlessHandler->OutgoingConnectionless(Address, (uint8*)DataToSend, CountBits);
+
+			DataToSend = ProcessedData.Data;
+			CountBits = ProcessedData.CountBits;
+		}
+
+
+		int32 BytesSent = 0;
+
+		CLOCK_CYCLES(SendCycles);
+		Socket->SendTo(DataToSend, FMath::DivideAndRoundUp(CountBits, 8), BytesSent, *RemoteAddr);
+		UNCLOCK_CYCLES(SendCycles);
+
+
+		// @todo: Can't implement these profiling events (require UNetConnections)
+		//NETWORK_PROFILER(GNetworkProfiler.FlushOutgoingBunches(/* UNetConnection */));
+		//NETWORK_PROFILER(GNetworkProfiler.TrackSocketSendTo(Socket->GetDescription(),Data,BytesSent,NumPacketIdBits,NumBunchBits,
+							//NumAckBits,NumPaddingBits, /* UNetConnection */));
+	}
+	else
+	{
+		UE_LOG(LogNet, Warning, TEXT("UIpNetDriver::LowLevelSend: Invalid send address '%s'"), *Address);
+	}
+}
+
 void UIpNetDriver::ProcessRemoteFunction(class AActor* Actor, UFunction* Function, void* Parameters, FOutParmRec* OutParms, FFrame* Stack, class UObject* SubObject )
 {
 	bool bIsServer = IsServer();
diff --git a/Engine/Source/Runtime/Online/OnlineSubsystemUtils/Private/OnlineBeaconClient.cpp b/Engine/Source/Runtime/Online/OnlineSubsystemUtils/Private/OnlineBeaconClient.cpp
--- a/Engine/Source/Runtime/Online/OnlineSubsystemUtils/Private/OnlineBeaconClient.cpp
+++ b/Engine/Source/Runtime/Online/OnlineSubsystemUtils/Private/OnlineBeaconClient.cpp
@@ -75,6 +75,14 @@ bool AOnlineBeaconClient::InitClient(FURL& URL)
 			FString Error;
 			if (NetDriver->InitConnect(this, URL, Error))
 			{
+				BeaconConnection = NetDriver->ServerConnection;
+
+				// Kick off the connection handshake
+				if (BeaconConnection->StatelessConnectComponent.IsValid())
+				{
+					BeaconConnection->StatelessConnectComponent.Pin()->SendInitialConnect();
+				}
+
 				SetConnectionState(EBeaconConnectionState::Pending);
 
 				NetDriver->SetWorld(GetWorld());
@@ -85,7 +93,6 @@ bool AOnlineBeaconClient::InitClient(FURL& URL)
 				// Send initial message.
 				uint8 IsLittleEndian = uint8(PLATFORM_LITTLE_ENDIAN);
 				check(IsLittleEndian == !!IsLittleEndian); // should only be one or zero
-				BeaconConnection = NetDriver->ServerConnection;
 
 				uint32 LocalNetworkVersion = FNetworkVersion::GetLocalNetworkVersion();
 				
diff --git a/Engine/Source/Runtime/PacketHandlers/PacketHandler/Classes/HandlerComponentFactory.h b/Engine/Source/Runtime/PacketHandlers/PacketHandler/Classes/HandlerComponentFactory.h
new file mode 100644
--- /dev/null
+++ b/Engine/Source/Runtime/PacketHandlers/PacketHandler/Classes/HandlerComponentFactory.h
@@ -0,0 +1,20 @@
+// Copyright 1998-2016 Epic Games, Inc. All Rights Reserved.
+
+#pragma once
+
+#include "PacketHandler.h"
+
+#include "HandlerComponentFactory.generated.h"
+
+/**
+ * A UObject alternative for loading HandlerComponents without strict module dependency
+ */
+UCLASS()
+class PACKETHANDLER_API UHandlerComponentFactory : public UObject
+{
+	GENERATED_UCLASS_BODY()
+
+public:
+	virtual TSharedPtr<HandlerComponent> CreateComponentInstance(FString& Options)
+		PURE_VIRTUAL(UHandlerComponentFactory::CreateComponentInstance, return nullptr;);
+};
diff --git a/Engine/Source/Runtime/PacketHandlers/PacketHandler/PacketHandler.Build.cs b/Engine/Source/Runtime/PacketHandlers/PacketHandler/PacketHandler.Build.cs
--- a/Engine/Source/Runtime/PacketHandlers/PacketHandler/PacketHandler.Build.cs
+++ b/Engine/Source/Runtime/PacketHandlers/PacketHandler/PacketHandler.Build.cs
@@ -7,9 +7,14 @@ public class PacketHandler : ModuleRules
 {
     public PacketHandler(TargetInfo Target)
     {
-        PublicDependencyModuleNames.AddRange(
-            new string[] {
+		PrivateIncludePaths.Add("PacketHandler/Private");
+
+        PublicDependencyModuleNames.AddRange
+		(
+            new string[]
+			{
 				"Core",
+				"CoreUObject",
                 "ReliabilityHandlerComponent",
             }
         );
diff --git a/Engine/Source/Runtime/PacketHandlers/PacketHandler/Private/HandlerComponentFactory.cpp b/Engine/Source/Runtime/PacketHandlers/PacketHandler/Private/HandlerComponentFactory.cpp
new file mode 100644
--- /dev/null
+++ b/Engine/Source/Runtime/PacketHandlers/PacketHandler/Private/HandlerComponentFactory.cpp
@@ -0,0 +1,13 @@
+// Copyright 1998-2016 Epic Games, Inc. All Rights Reserved.
+
+#include "PacketHandlerPCH.h"
+
+#include "HandlerComponentFactory.h"
+
+/**
+ * UHandlerComponentFactory
+ */
+UHandlerComponentFactory::UHandlerComponentFactory(const FObjectInitializer& ObjectInitializer)
+	: Super(ObjectInitializer)
+{
+}
diff --git a/Engine/Source/Runtime/PacketHandlers/PacketHandler/Private/PacketHandler.cpp b/Engine/Source/Runtime/PacketHandlers/PacketHandler/Private/PacketHandler.cpp
--- a/Engine/Source/Runtime/PacketHandlers/PacketHandler/Private/PacketHandler.cpp
+++ b/Engine/Source/Runtime/PacketHandlers/PacketHandler/Private/PacketHandler.cpp
@@ -1,8 +1,24 @@
 // Copyright 1998-2016 Epic Games, Inc. All Rights Reserved.
 
+#include "PacketHandlerPCH.h"
+
 #include "PacketHandler.h"
+#include "HandlerComponentFactory.h"
 #include "ReliabilityHandlerComponent.h"
 
+
+// @todo #JohnB: Currently, when HandlerComponent's are not byte-aligned, they waste a lot of bits (up to a byte),
+//					so it would be good to optimize this in the future, while avoiding overcomplicating the code.
+//
+//					One possibility for optimization, which will work well with Oodle, is having a hard-coded bit alignment,
+//					and then trimming all trailing zero's from the packet, then padding-up zeros up to the bit alignment, upon receive.
+//					This may require using a bit, to signal that padding is needed post-receive. May overcomplicate the code though.
+
+
+// @todo #JohnB: There is quite a lot of inefficient copying of packet data going on.
+//					Redo the whole packet parsing/modification pipeline.
+
+
 IMPLEMENT_MODULE(FPacketHandlerComponentModuleInterface, PacketHandler);
 
 DEFINE_LOG_CATEGORY(PacketHandlerLog);
@@ -11,36 +27,26 @@ DEFINE_LOG_CATEGORY(PacketHandlerLog);
 // @todo #JohnB: For the moment, disable the reliability handler, while it is causing network trouble
 #define DISABLE_RELIABILITY_HANDLER 1
 
-/**
- * BufferedPacket
- */
-
-BufferedPacket::BufferedPacket()
-	: Data(nullptr)
-	, BytesCount(0)
-	, SendTime(0.f)
-	, Id(0)
-{
-}
-
-BufferedPacket::~BufferedPacket()
-{
-	if (Data != nullptr)
-	{
-		delete Data;
-	}
-}
 
 /**
  * PacketHandler
  */
+
 PacketHandler::PacketHandler()
-	: Time(0.f)
+	: Mode(Handler::Mode::Client)
+	, Time(0.f)
+	, OutgoingPacket()
+	, IncomingPacket()
+	, HandlerComponents()
+	, PacketBitAlignment(0)
 	, State(Handler::State::Uninitialized)
+	, BufferedPackets()
+	, QueuedPackets()
+	, BufferedConnectionlessPackets()
+	, QueuedConnectionlessPackets()
 	, ReliabilityComponent(nullptr)
+	, bRawSend(false)
 {
-	GConfig->GetBool(TEXT("PacketHandlerComponents"), TEXT("bEnabled"), bEnabled, GEngineIni);
-
 	OutgoingPacket.SetAllowResize(true);
 	OutgoingPacket.AllowAppend(true);
 }
@@ -49,277 +55,406 @@ void PacketHandler::Tick(float DeltaTime)
 {
 	Time += DeltaTime;
 
-	switch (State)
+	for (int32 i=0; i<HandlerComponents.Num(); ++i)
+	{
+		HandlerComponents[i]->Tick(DeltaTime);
+	}
+}
+
+void PacketHandler::Initialize(Handler::Mode InMode, bool bConnectionlessOnly/*=false*/)
+{
+	Mode = InMode;
+
+	// @todo #JohnB: Redo this, so you don't load from the .ini at all, have it hardcoded elsewhere - do not want this in shipping.
+
+	// Only UNetConnection's will load the .ini components, for now.
+	if (!bConnectionlessOnly)
 	{
-		case Handler::State::InitializingProcessors:
+		TArray<FString> Components;
+
+		GConfig->GetArray(TEXT("PacketHandlerComponents"), TEXT("Components"), Components, GEngineIni);
+
+		for (FString CurComponent : Components)
 		{
-			break;
+			AddHandler(CurComponent, true);
 		}
-		case Handler::State::Initialized:
+	}
+}
+
+void PacketHandler::InitializeComponents()
+{
+	if (State == Handler::State::Uninitialized)
+	{
+		if (HandlerComponents.Num() > 0)
 		{
-			break;
+			SetState(Handler::State::InitializingComponents);
+		}
+		else
+		{
+			HandlerInitialized();
 		}
 	}
 
-	for (int32 i=0; i<HandlerComponents.Num(); ++i)
+	// Trigger delayed-initialization for HandlerComponents
+	for (int32 i=0; i<HandlerComponents.Num(); i++)
 	{
-		HandlerComponents[i]->Tick(DeltaTime);
+		HandlerComponent& CurComponent = *HandlerComponents[i];
+
+		if (!CurComponent.IsInitialized())
+		{
+			CurComponent.Initialize();
+		}
+
+		if (CurComponent.DoesResetBitAlignment())
+		{
+			PacketBitAlignment = CurComponent.GetBitAlignment();
+		}
+		else
+		{
+			PacketBitAlignment = (PacketBitAlignment + CurComponent.GetBitAlignment()) % 8;
+		}
 	}
+
+	// Called early, to ensure that all handlers report a valid packet overhead (triggers an assert if not)
+	GetTotalPacketOverheadBits();
 }
 
-void PacketHandler::Initialize(Handler::Mode InMode)
+void PacketHandler::AddHandler(TSharedPtr<HandlerComponent> NewHandler, bool bDeferInitialize/*=false*/)
 {
-	if (!bEnabled && FParse::Param(FCommandLine::Get(), TEXT("PacketHandler")))
+	// This is never valid. Can end up silently changing maximum allow packet size, which could cause failure to send packets.
+	if (State != Handler::State::Uninitialized)
 	{
-		UE_LOG(PacketHandlerLog, Log, TEXT("Force-enabling packet handler from commandline."));
-		bEnabled = true;
+		LowLevelFatalError(TEXT("Handler added during runtime."));
+		return;
 	}
 
-	if (bEnabled)
+	// This should always be fatal, as an unexpectedly missing handler, may break net compatibility with the remote server/client
+	if (!NewHandler.IsValid())
 	{
-		struct ComponentAndOptions
-		{
-			FString ComponentName;
-			FString Options;
-		};
+		LowLevelFatalError(TEXT("Failed to add handler - invalid instance."));
+		return;
+	}
 
-		Mode = InMode;
+	HandlerComponents.Add(NewHandler);
+	NewHandler->Handler = this;
 
+	if (!bDeferInitialize)
+	{
+		NewHandler->Initialize();
+	}
+}
 
-		TArray<FString> Components;
-		TArray<ComponentAndOptions> ComponentsArray;
+TSharedPtr<HandlerComponent> PacketHandler::AddHandler(FString ComponentStr, bool bDeferInitialize/*=false*/)
+{
+	TSharedPtr<HandlerComponent> ReturnVal = nullptr;
 
-		GConfig->GetArray(TEXT("PacketHandlerComponents"), TEXT("Components"), Components, GEngineIni);
+	if (!ComponentStr.IsEmpty())
+	{
+		FString ComponentName;
+		FString ComponentOptions;
 
-		for (FString CurComponent : Components)
+		for (int32 i=0; i<ComponentStr.Len(); i++)
 		{
-			if (CurComponent.IsEmpty())
-			{
-				continue;
-			}
+			TCHAR c = ComponentStr[i];
 
-
-			int CurIndex = ComponentsArray.Add(ComponentAndOptions());
-
-			for (int32 i=0; i<CurComponent.Len(); i++)
+			// Parsing Options
+			if (c == '(')
 			{
-				TCHAR c = CurComponent[i];
+				// Skip '('
+				++i;
 
-				// Parsing Options
-				if (c == '(')
+				// Parse until end of options
+				for (; i<ComponentStr.Len(); i++)
 				{
-					// Skip '('
-					++i;
+					c = ComponentStr[i];
 
-					// Parse until end of options
-					for (; i<CurComponent.Len(); i++)
+					// End of options
+					if (c == ')')
 					{
-						c = CurComponent[i];
-
-						// End of options
-						if (c == ')')
-						{
-							break;
-						}
-						// Append char to options
-						else
-						{
-							ComponentsArray[CurIndex].Options.AppendChar(c);
-						}
+						break;
+					}
+					// Append char to options
+					else
+					{
+						ComponentOptions.AppendChar(c);
 					}
 				}
-				// Append char to component name if not whitespace
-				else if (c != ' ')
-				{
-					ComponentsArray[CurIndex].ComponentName.AppendChar(c);
-				}
+			}
+			// Append char to component name if not whitespace
+			else if (c != ' ')
+			{
+				ComponentName.AppendChar(c);
 			}
 		}
 
-		for (int32 i=0; i<ComponentsArray.Num(); i++)
+		if (ComponentName != TEXT("ReliabilityHandlerComponent"))
 		{
-			// Skip adding reliability component as it is added automatically in the handler
-			if(ComponentsArray[i].ComponentName == TEXT("ReliabilityHandlerComponent"))
+			int32 FactoryComponentDelim = ComponentName.Find(TEXT("."));
+
+			if (FactoryComponentDelim != INDEX_NONE)
 			{
-				continue;
-			}
+				// Every HandlerComponentFactory type has one instance, loaded as a named singleton
+				FString SingletonName = ComponentName.Mid(FactoryComponentDelim + 1) + TEXT("_Singleton");
+				UHandlerComponentFactory* Factory = FindObject<UHandlerComponentFactory>(ANY_PACKAGE, *SingletonName);
+
+				if (Factory == nullptr)
+				{
+					UClass* FactoryClass = StaticLoadClass(UHandlerComponentFactory::StaticClass(), nullptr, *ComponentName);
 
+					if (FactoryClass != nullptr)
+					{
+						Factory = NewObject<UHandlerComponentFactory>(GetTransientPackage(), FactoryClass, *SingletonName);
+					}
+				}
 
-			TSharedPtr<IModuleInterface> Interface = FModuleManager::Get().LoadModule(FName(*ComponentsArray[i].ComponentName));
 
-			if(Interface.IsValid())
+				if (Factory != nullptr)
+				{
+					ReturnVal = Factory->CreateComponentInstance(ComponentOptions);
+				}
+				else
+				{
+					UE_LOG(PacketHandlerLog, Warning, TEXT("Unable to load HandlerComponent factory: %s"), *ComponentName);
+				}
+			}
+			// @todo #JohnB: Deprecate non-factory components eventually
+			else
 			{
-				TSharedPtr<FPacketHandlerComponentModuleInterface> PacketHandlerInterface =
-					StaticCastSharedPtr<FPacketHandlerComponentModuleInterface>(Interface);
+				TSharedPtr<IModuleInterface> Interface = FModuleManager::Get().LoadModule(FName(*ComponentName));
 
-				if(PacketHandlerInterface.IsValid())
+				if (Interface.IsValid())
 				{
-					UE_LOG(PacketHandlerLog, Log, TEXT("Loading PacketHandler component: %s (%s)"), *ComponentsArray[i].ComponentName,
-							*ComponentsArray[i].Options);
+					TSharedPtr<FPacketHandlerComponentModuleInterface> PacketHandlerInterface =
+						StaticCastSharedPtr<FPacketHandlerComponentModuleInterface>(Interface);
 
-					Add(PacketHandlerInterface->CreateComponentInstance(ComponentsArray[i].Options));
+					if (PacketHandlerInterface.IsValid())
+					{
+						ReturnVal = PacketHandlerInterface->CreateComponentInstance(ComponentOptions);
+					}
+				}
+				else
+				{
+					UE_LOG(PacketHandlerLog, Warning, TEXT("Unable to Load Module: %s"), *ComponentName);
 				}
 			}
-			else
+
+
+			if (ReturnVal.IsValid())
 			{
-				UE_LOG(PacketHandlerLog, Warning, TEXT("Unable to Load Module: %s"), *ComponentsArray[i].ComponentName);
+				UE_LOG(PacketHandlerLog, Log, TEXT("Loaded PacketHandler component: %s (%s)"), *ComponentName,
+						*ComponentOptions);
+
+				AddHandler(ReturnVal, bDeferInitialize);
 			}
 		}
+		else
+		{
+			UE_LOG(PacketHandlerLog, Warning, TEXT("PacketHandlerComponent 'ReliabilityHandlerComponent' is internal-only."));
+		}
+	}
+
+	return ReturnVal;
+}
+
+const ProcessedPacket PacketHandler::Outgoing(uint8* Packet, int32 CountBits)
+{
+	if (!bRawSend)
+	{
+		OutgoingPacket.Reset();
+
+		if (State == Handler::State::Uninitialized)
+		{
+			UpdateInitialState();
+		}
+
+
+		if (State == Handler::State::Initialized)
+		{
+			OutgoingPacket.SerializeBits(Packet, CountBits);
+		}
+		// Buffer any packets being sent from game code until processors are initialized
+		else if (State == Handler::State::InitializingComponents && CountBits > 0)
+		{
+			BufferedPackets.Add(new BufferedPacket(Packet, CountBits));
+
+			Packet = nullptr;
+			CountBits = 0;
+		}
 
 
-		// Called early, to ensure that all handlers report a valid packet overhead (triggers an assert if not)
-		GetTotalPacketOverheadBits();
+		// Queue packet for resending before handling
+		if (CountBits > 0 && ReliabilityComponent.IsValid())
+		{
+			ReliabilityComponent->QueuePacketForResending(Packet, CountBits);
+		}
+
+		for (int32 i=0; i<HandlerComponents.Num(); ++i)
+		{
+			if (HandlerComponents[i]->IsActive())
+			{
+				HandlerComponents[i]->Outgoing(OutgoingPacket);
+			}
+		}
+
+		return ProcessedPacket(OutgoingPacket.GetData(), OutgoingPacket.GetNumBits());
+	}
+	else
+	{
+		return ProcessedPacket(Packet, CountBits);
 	}
 }
 
-void PacketHandler::Add(TSharedPtr<HandlerComponent> NewHandler)
+const ProcessedPacket PacketHandler::Incoming(uint8* Packet, int32 CountBytes)
 {
-	// This is never valid. Can end up silently changing maximum allow packet size, which could cause failure to send packets.
-	if (State != Handler::State::Uninitialized)
+	int32 BitAlignment = (PacketBitAlignment > 0) ? PacketBitAlignment : 8;
+	int32 CountBits = (CountBytes > 0 ? ((CountBytes - 1) * 8) + BitAlignment : 0);
+	FBitReader ProcessedPacketReader(Packet, CountBits);
+
+	if (State == Handler::State::Uninitialized)
 	{
-		LowLevelFatalError(TEXT("Handler added during runtime."));
-		return;
+		UpdateInitialState();
 	}
 
-	// This should always be fatal, as an unexpectedly missing handler, may break net compatibility with the remote server/client
-	if (!NewHandler.IsValid())
+	// Handle
+	for (int32 i=HandlerComponents.Num() - 1; i>=0; --i)
 	{
-		LowLevelFatalError(TEXT("Failed to add handler - invalid instance."));
-		return;
+		HandlerComponent& CurComponent = *HandlerComponents[i];
+
+		if (CurComponent.IsActive() && ProcessedPacketReader.GetBitsLeft() > 0)
+		{
+			// Realign the packet, so the packet data starts at position 0, if necessary
+			if (ProcessedPacketReader.GetPosBits() != 0 && !CurComponent.CanReadUnaligned())
+			{
+				RealignPacket(ProcessedPacketReader);
+			}
+
+			CurComponent.Incoming(ProcessedPacketReader);
+		}
 	}
 
-	HandlerComponents.Add(NewHandler);
-	NewHandler->Handler = this;
-	NewHandler->Initialize();
+	ReplaceIncomingPacket(ProcessedPacketReader);
+
+	return ProcessedPacket(IncomingPacket.GetData(), IncomingPacket.GetBitsLeft());
 }
 
-const ProcessedPacket PacketHandler::Outgoing(uint8* Packet, int32 Count)
+const ProcessedPacket PacketHandler::IncomingConnectionless(FString Address, uint8* Packet, int32 CountBytes)
 {
-	OutgoingPacket.Reset();
+	// @todo #JohnB: Try to optimize this function more, seeing as it will be a common codepath DoS attacks pass through
 
-	switch (State)
+	int32 BitAlignment = (PacketBitAlignment > 0) ? PacketBitAlignment : 8;
+	int32 CountBits = (CountBytes > 0 ? ((CountBytes - 1) * 8) + BitAlignment : 0);
+	FBitReader ProcessedPacketReader(Packet, CountBits);
+
+	if (State == Handler::State::Uninitialized)
+	{
+		UpdateInitialState();
+	}
+
+
+	for (int32 i=HandlerComponents.Num() - 1; i>=0; --i)
 	{
-		case Handler::State::Uninitialized:
+		HandlerComponent& CurComponent = *HandlerComponents[i];
+
+		if (CurComponent.IsActive() && ProcessedPacketReader.GetBitsLeft() > 0)
 		{
-#if !DISABLE_RELIABILITY_HANDLER
-			if (!ReliabilityComponent.IsValid())
+			// Realign the packet, so the packet data starts at position 0, if necessary
+			if (ProcessedPacketReader.GetPosBits() != 0 && !CurComponent.CanReadUnaligned())
 			{
-				ReliabilityComponent = MakeShareable(new ReliabilityHandlerComponent);
-				Add(ReliabilityComponent);
+				RealignPacket(ProcessedPacketReader);
 			}
-#endif
 
-			// Have handlers to initialize other than reliability
-			if (HandlerComponents.Num() > 1)
-			{
-				State = Handler::State::InitializingProcessors;
-			} 
-			else
-			{
-				State = Handler::State::Initialized;
-				break;
-			}
+			CurComponent.IncomingConnectionless(Address, ProcessedPacketReader);
 		}
-		case Handler::State::InitializingProcessors:
-		{
-			// Buffer any packets being sent from game code
-			// until processors are initialized
-			if (Count > 0)
-			{
-				BufferedPacket* NewBufferedPacket = new BufferedPacket;
-				BufferedPackets.Add(NewBufferedPacket);
+	}
 
-				uint8* Data = new uint8[Count];
-				memcpy(Data, Packet, Count);
-				NewBufferedPacket->Data = Data;
-				NewBufferedPacket->BytesCount = Count;
+	ReplaceIncomingPacket(ProcessedPacketReader);
 
-				Count = 0;
-				Packet = nullptr;
-			}
+	return ProcessedPacket(IncomingPacket.GetData(), IncomingPacket.GetBitsLeft());
+}
 
-			break;
-		}
-		case Handler::State::Initialized:
-		{
-			OutgoingPacket.Serialize(Packet, Count);
-			break;
-		}
-		default:
-		{
-			break;
-		}
+const ProcessedPacket PacketHandler::OutgoingConnectionless(FString Address, uint8* Packet, int32 CountBits)
+{
+	OutgoingPacket.Reset();
+
+	if (State == Handler::State::Uninitialized)
+	{
+		UpdateInitialState();
 	}
 
-	// Queue packet for resending before handling
-	if (Count > 0 && ReliabilityComponent.IsValid())
+
+	if (State == Handler::State::Initialized)
 	{
-		ReliabilityComponent->QueuePacketForResending(Packet, Count);
+		OutgoingPacket.SerializeBits(Packet, CountBits);
 	}
+	// Buffer any packets being sent from game code until processors are initialized
+	else if (State == Handler::State::InitializingComponents && CountBits > 0)
+	{
+		BufferedConnectionlessPackets.Add(new BufferedPacket(Address, Packet, CountBits));
+
+		Packet = nullptr;
+		CountBits = 0;
+	}
+
 
 	// Handle
 	for (int32 i=0; i<HandlerComponents.Num(); ++i)
 	{
 		if (HandlerComponents[i]->IsActive())
 		{
-			HandlerComponents[i]->Outgoing(OutgoingPacket);
+			HandlerComponents[i]->OutgoingConnectionless(Address, OutgoingPacket);
 		}
 	}
 
-	return ProcessedPacket(OutgoingPacket.GetData(), OutgoingPacket.GetNumBytes());
+	return ProcessedPacket(OutgoingPacket.GetData(), OutgoingPacket.GetNumBits());
 }
 
-const ProcessedPacket PacketHandler::Incoming(uint8* Packet, int32 Count)
+void PacketHandler::ReplaceIncomingPacket(FBitReader& ReplacementPacket)
 {
-	FBitReader ProcessedPacketReader(Packet, Count * 8);
-
-	switch (State)
+	if (ReplacementPacket.GetPosBits() == 0 || ReplacementPacket.GetBitsLeft() == 0)
 	{
-		case Handler::State::Uninitialized:
-		{
-#if !DISABLE_RELIABILITY_HANDLER
-			if (!ReliabilityComponent.IsValid())
-			{
-				ReliabilityComponent = MakeShareable(new ReliabilityHandlerComponent);
-				Add(ReliabilityComponent);
-			}
-#endif
+		IncomingPacket = ReplacementPacket;
+	}
+	else
+	{
+		// @todo #JohnB: Make this directly adjust and write into IncomingPacket's buffer, instead of copying - very inefficient
+		TArray<uint8> NewPacketData;
+		int32 NewPacketSizeBits = ReplacementPacket.GetBitsLeft();
 
-			if (HandlerComponents.Num() > 1)
-			{
-				State = Handler::State::InitializingProcessors;
-			}
-			else
-			{
-				State = Handler::State::Initialized;
-			}
-		}
-		case Handler::State::InitializingProcessors:
-		{
-			break;
-		}
-		case Handler::State::Initialized:
-		{
-			break;
-		}
-		default:
-		{
-			break;
-		}
+		NewPacketData.AddUninitialized(ReplacementPacket.GetBytesLeft());
+		NewPacketData[NewPacketData.Num()-1] = 0;
+
+		ReplacementPacket.SerializeBits(NewPacketData.GetData(), NewPacketSizeBits);
+
+
+		FBitReader NewPacket(NewPacketData.GetData(), NewPacketSizeBits);
+
+		IncomingPacket = NewPacket;
 	}
+}
 
-	// Handle
-	for (int32 i=HandlerComponents.Num() - 1; i>=0; --i)
+void PacketHandler::RealignPacket(FBitReader& Packet)
+{
+	if (Packet.GetPosBits() != 0)
 	{
-		if (HandlerComponents[i]->IsActive() && ProcessedPacketReader.GetNumBytes() > 0)
+		uint32 BitsLeft = Packet.GetBitsLeft();
+
+		if (BitsLeft > 0)
 		{
-			HandlerComponents[i]->Incoming(ProcessedPacketReader);
-		}
-	}
+			// @todo #JohnB: Based on above - when you optimize above, optimize this too
+			TArray<uint8> NewPacketData;
+			int32 NewPacketSizeBits = BitsLeft;
 
-	IncomingPacket = ProcessedPacketReader;
+			NewPacketData.AddUninitialized(Packet.GetBytesLeft());
+			NewPacketData[NewPacketData.Num()-1] = 0;
 
-	return ProcessedPacket(IncomingPacket.GetData(), IncomingPacket.GetBytesLeft());
+			Packet.SerializeBits(NewPacketData.GetData(), NewPacketSizeBits);
+
+
+			FBitReader NewPacket(NewPacketData.GetData(), NewPacketSizeBits);
+
+			Packet = NewPacket;
+		}
+	}
 }
 
 void PacketHandler::SetState(Handler::State InState)
@@ -334,6 +469,32 @@ void PacketHandler::SetState(Handler::State InState)
 	}
 }
 
+void PacketHandler::UpdateInitialState()
+{
+	if (State == Handler::State::Uninitialized)
+	{
+#if !DISABLE_RELIABILITY_HANDLER
+		if (!ReliabilityComponent.IsValid())
+		{
+			ReliabilityComponent = MakeShareable(new ReliabilityHandlerComponent);
+			AddHandler(ReliabilityComponent);
+		}
+
+		// Have handlers to initialize other than reliability
+		if (HandlerComponents.Num() > 1)
+#else
+		if (HandlerComponents.Num() > 0)
+#endif
+		{
+			InitializeComponents();
+		} 
+		else
+		{
+			HandlerInitialized();
+		}
+	}
+}
+
 void PacketHandler::HandlerInitialized()
 {
 	// If any buffered packets, add to queue
@@ -343,6 +504,17 @@ void PacketHandler::HandlerInitialized()
 		BufferedPackets[i] = nullptr;
 	}
 
+	BufferedPackets.Empty();
+
+	for (int32 i=0; i<BufferedConnectionlessPackets.Num(); ++i)
+	{
+		QueuedConnectionlessPackets.Enqueue(BufferedConnectionlessPackets[i]);
+		BufferedConnectionlessPackets[i] = nullptr;
+	}
+
+	BufferedConnectionlessPackets.Empty();
+
+
 	SetState(Handler::State::Initialized);
 }
 
@@ -372,9 +544,18 @@ BufferedPacket* PacketHandler::GetQueuedPacket()
 	return QueuedPacket;
 }
 
-void PacketHandler::QueuePacketForSending(BufferedPacket* PacketToQueue)
+BufferedPacket* PacketHandler::GetQueuedConnectionlessPacket()
 {
+	BufferedPacket* QueuedConnectionlessPacket = nullptr;
+
+	QueuedConnectionlessPackets.Dequeue(QueuedConnectionlessPacket);
+
+	return QueuedConnectionlessPacket;
+}
 
+void PacketHandler::QueuePacketForSending(BufferedPacket* PacketToQueue)
+{
+	// @todo #JohnB: Deprecate?
 }
 
 int32 PacketHandler::GetTotalPacketOverheadBits()
@@ -395,6 +576,13 @@ int32 PacketHandler::GetTotalPacketOverheadBits()
 		ReturnVal += CurOverhead;
 	}
 
+	// @todo #JohnB: Fix byte alignment notes in NetConnection.cpp, then remove this code
+#if 1
+	// Byte-align the result
+	ReturnVal = FMath::DivideAndRoundUp(ReturnVal, 8) * 8;
+#endif
+
+
 	return ReturnVal;
 }
 
@@ -404,7 +592,8 @@ int32 PacketHandler::GetTotalPacketOverheadBits()
  */
 
 HandlerComponent::HandlerComponent()
-	: State(Handler::Component::State::UnInitialized)
+	: Handler(nullptr)
+	, State(Handler::Component::State::UnInitialized)
 	, bActive(false)
 	, bInitialized(false)
 {
@@ -435,14 +624,3 @@ bool HandlerComponent::IsInitialized() const
 {
 	return bInitialized;
 }
-
-
-/**
- * ProcessedPacket
- */
-
-ProcessedPacket::ProcessedPacket(uint8* InData, int32 InCount)
-{
-	Data = InData;
-	Count = InCount;
-}
diff --git a/Engine/Source/Runtime/PacketHandlers/PacketHandler/Private/PacketHandlerPCH.h b/Engine/Source/Runtime/PacketHandlers/PacketHandler/Private/PacketHandlerPCH.h
new file mode 100644
--- /dev/null
+++ b/Engine/Source/Runtime/PacketHandlers/PacketHandler/Private/PacketHandlerPCH.h
@@ -0,0 +1,4 @@
+// Copyright 1998-2016 Epic Games, Inc. All Rights Reserved.
+
+#include "Core.h"
+#include "CoreUObject.h"
diff --git a/Engine/Source/Runtime/PacketHandlers/PacketHandler/Public/PacketHandler.h b/Engine/Source/Runtime/PacketHandlers/PacketHandler/Public/PacketHandler.h
--- a/Engine/Source/Runtime/PacketHandlers/PacketHandler/Public/PacketHandler.h
+++ b/Engine/Source/Runtime/PacketHandlers/PacketHandler/Public/PacketHandler.h
@@ -3,37 +3,45 @@
 #pragma once
 
 #include "ModuleManager.h"
-#include "Core.h"
 
 DECLARE_LOG_CATEGORY_EXTERN(PacketHandlerLog, Log, All);
 
 class HandlerComponent;
 class ReliabilityHandlerComponent;
 
-/* Enums related to the Packet Handler */
+
+/**
+ * Enums related to the PacketHandler
+ */
 namespace Handler
 {
-	/* State of Packet Handler */
+	/**
+	 * State of PacketHandler
+	 */
 	enum State
 	{
-		Uninitialized,
-		InitializingProcessors,
-		Initialized
+		Uninitialized,			// PacketHandler is uninitialized
+		InitializingComponents,	// PacketHandler is initializing HandlerComponents
+		Initialized				// PacketHandler and all HandlerComponents (if any) are initialized
 	};
 
-	/* Mode of Packet Handler */
+	/**
+	 * Mode of Packet Handler
+	 */
 	enum Mode
 	{
-		Client,
-		Server
+		Client,					// Clientside PacketHandler
+		Server					// Serverside PacketHandler
 	};
 
 	namespace Component
 	{
-		/* Packet Handler State */
+		/**
+		 * HandlerComponent State
+		 */
 		enum State
 		{
-			UnInitialized,
+			UnInitialized,		// HandlerComponent not yet initialized
 			InitializedOnLocal, // Initialized on local instance
 			InitializeOnRemote, // Initialized on remote instance, not on local instance
 			Initialized         // Initialized on both local and remote instances
@@ -41,66 +49,217 @@ namespace Handler
 	}
 }
 
-/* The result of calling Incoming and Outgoing in the PacketHandler */
+/**
+ * The result of calling Incoming and Outgoing in the PacketHandler
+ */
 struct PACKETHANDLER_API ProcessedPacket
 {
-	ProcessedPacket(uint8* Data, int32 Count);
+	/** Pointer to the returned packet data */
 	uint8* Data;
-	int32 Count;
+
+	/** Size of the returned packet data in bits */
+	int32 CountBits;
+
+public:
+	ProcessedPacket()
+		: Data()
+		, CountBits(0)
+	{
+	}
+
+	/**
+	 * Base constructor
+	 */
+	ProcessedPacket(uint8* InData, int32 InCountBits)
+		: Data(InData)
+		, CountBits(InCountBits)
+	{
+	}
 };
 
-/*
-* Packet Handler will buffer packets, this struct is used to buffer such packets
-* while handler components are initialized
-*/
+/**
+ * PacketHandler will buffer packets, this struct is used to buffer such packets while handler components are initialized
+ */
 struct PACKETHANDLER_API BufferedPacket
 {
-	BufferedPacket();
-	~BufferedPacket();
-
+	/** Buffered packet data */
 	uint8* Data;
-	uint32 BytesCount;
-	float SendTime;
+
+	/** Size of buffered packet in bits */
+	uint32 CountBits;
+
+	/** Used by ReliabilityHandlerComponent, to mark a packet for resending */
+	float ResendTime;
+
+	/** Used by ReliabilityHandlerComponent, to track packet id's */
 	uint32 Id;
+
+	/** For connectionless packets, the address to send to (format is abstract, determined by active net driver) */
+	FString Address;
+
+
+public:
+	/**
+	 * Base constructor
+	 */
+	BufferedPacket()
+		: Data(nullptr)
+		, CountBits(0)
+		, ResendTime(0.f)
+		, Id(0)
+		, Address()
+	{
+	}
+
+	BufferedPacket(uint8* InCopyData, uint32 InCountBits, float InResendTime=0.f, uint32 InId=0)
+		: CountBits(InCountBits)
+		, ResendTime(InResendTime)
+		, Id(InId)
+		, Address()
+	{
+		check(InCopyData != nullptr);
+
+		Data = new uint8[FMath::DivideAndRoundUp(InCountBits, 8u)];
+		FMemory::Memcpy(Data, InCopyData, FMath::DivideAndRoundUp(InCountBits, 8u));
+	}
+
+	BufferedPacket(FString InAddress, uint8* InCopyData, uint32 InCountBits, float InResendTime=0.f, uint32 InId=0)
+		: BufferedPacket(InCopyData, InCountBits, InResendTime, InId)
+	{
+		Address = InAddress;
+	}
+
+	/**
+	 * Base destructor
+	 */
+	~BufferedPacket()
+	{
+		if (Data != nullptr)
+		{
+			delete Data;
+		}
+	}
 };
 
-/*
-* This class maintains an array of all Packet Handler Components and
-* forwards incoming and outgoing packets the each component
-*/
+/**
+ * This class maintains an array of all PacketHandler Components and forwards incoming and outgoing packets the each component
+ */
 class PACKETHANDLER_API PacketHandler
 {
 public:
-	/* Default initialization of data */
+	/**
+	 * Base constructor
+	 */
 	PacketHandler();
 
-	/* Destructor */
-	virtual ~PacketHandler() { };
+	/**
+	 * Base Destructor
+	 */
+	virtual ~PacketHandler()
+	{
+	};
 
-	/* Handles initialization of manager */
-	virtual void Initialize(Handler::Mode Mode);
+	/**
+	 * Handles initialization of manager
+	 *
+	 * @param Mode					The mode the manager should be initialized in
+	 * @param bConnectionlessOnly	Whether or not this is a connectionless-only manager (ignores .ini components)
+	 */
+	virtual void Initialize(Handler::Mode Mode, bool bConnectionlessOnly=false);
+
+	/**
+	 * Triggers initialization of HandlerComponents.
+	 */
+	virtual void InitializeComponents();
 
-	/* Tick per frame */
 	virtual void Tick(float DeltaTime);
 
-	/* For adding handler components */
-	virtual void Add(TSharedPtr<HandlerComponent> NewHandler);
+	/**
+	 * Adds a HandlerComponent to the pipeline, prior to initialization (none can be added after initialization)
+	 *
+	 * @param NewHandler		The HandlerComponent to add
+	 * @param bDeferInitialize	Whether or not to defer triggering Initialize (for batch-adds - code calling this, triggers it instead)
+	 */
+	virtual void AddHandler(TSharedPtr<HandlerComponent> NewHandler, bool bDeferInitialize=false);
+
+	/**
+	 * As above, but initializes from a string specifying the component module, and (optionally) additional options
+	 *
+	 * @param ComponentStr		The handler component to load
+	 * @param bDeferInitialize	Whether or not to defer triggering Initialize (for batch-adds - code calling this, triggers it instead)
+	 */
+	virtual TSharedPtr<HandlerComponent> AddHandler(FString ComponentStr, bool bDeferInitialize=false);
+
+
+	/**
+	 * Processes incoming UNetConnection packets, through all HandlerComponents, and returns the final packet for local handling
+	 *
+	 * @param Packet		The packet data to be processed
+	 * @param CountBytes	The size of the packet data in bytes
+	 * @return				Returns the final packet
+	 */
+	virtual const ProcessedPacket Incoming(uint8* Packet, int32 CountBytes);
+
+	/**
+	 * Processes outgoing UNetConnection packets, through all HandlerComponents, and returns the final packet for sending
+	 *
+	 * @param Packet		The packet data to be processed
+	 * @param CountBits		The size of the packet data in bits
+	 * @return				Returns the final packet
+	 */
+	virtual const ProcessedPacket Outgoing(uint8* Packet, int32 CountBits);
+
+	/**
+	 * Processes incoming packets without a UNetConnection, in the same manner as 'Incoming' above
+	 * IMPORTANT: Net drivers triggering this, should call 'UNetDriver::FlushHandler' shortly afterwards, to minimize packet buffering
+	 * NOTE: Connectionless packets are unreliable.
+	 *
+	 * @param Address		The address the packet was received from (format is abstract, determined by active net driver)
+	 * @param Packet		The packet data to be processed
+	 * @param CountBytes	The size of the packet data in bytes
+	 * @return				Returns the final packet
+	 */
+	virtual const ProcessedPacket IncomingConnectionless(FString Address, uint8* Packet, int32 CountBytes);
 
-	/* Handles any incoming packets */
-	virtual const ProcessedPacket Incoming(uint8* Packet, int32 Count);
+	/**
+	 * Processes outgoing packets without a UNetConnection, in the same manner as 'Outgoing' above
+	 * NOTE: Connectionless packets are unreliable.
+	 *
+	 * @param Address		The address the packet is being sent to (format is abstract, determined by active net driver)
+	 * @param Packet		The packet data to be processed
+	 * @param CountBits		The size of the packet data in bits
+	 * @return				Returns the final packet
+	 */
+	virtual const ProcessedPacket OutgoingConnectionless(FString Address, uint8* Packet, int32 CountBits);
 
-	/* Handles any outgoing packets */
-	virtual const ProcessedPacket Outgoing(uint8* Packet, int32 Count);
 
-	/* Indicates that a handler was initialized */
+	/**
+	 * Triggered when a child HandlerComponen has been initialized
+	 */
 	void HandlerComponentInitialized();
 
-	/* Queue's a packet to be sent when the handler is ticked */
+	/**
+	 * Queue's a packet to be sent when the handler is ticked
+	 * NOTE: Unimplemented.
+	 *
+	 * @param PacketToQueue		The packet to be queued
+	 */
 	void QueuePacketForSending(BufferedPacket* PacketToQueue);
 
-	/* Gets a packet from the buffered packet queue for sending */
+	/**
+	 * Gets a packet from the buffered packet queue for sending
+	 *
+	 * @return		The packet to be sent, or nullptr if none are to be sent
+	 */
 	BufferedPacket* GetQueuedPacket();
 
+	/**
+	 * Gets a packet from the buffered connectionless packet queue for sending
+	 *
+	 * @return		The packet to be sent, or nullptr if none are to be sent
+	 */
+	BufferedPacket* GetQueuedConnectionlessPacket();
+
 	/**
 	 * Gets the combined packet/protocol overhead from all handlers, for reserving space in the parent connections packets
 	 *
@@ -109,131 +268,273 @@ public:
 	int32 GetTotalPacketOverheadBits();
 
 
-	/** Whether or not the PacketHandler system is enabled */
-	bool bEnabled;
+	/**
+	 * Sets whether or not outgoing packets should bypass this handler - used when raw packet sends are necessary
+	 * (such as for the stateless handshake)
+	 *
+	 * @param bInEnabled	Whether or not raw sends are enabled
+	 */
+	FORCEINLINE void SetRawSend(bool bInEnabled)
+	{
+		bRawSend = bInEnabled;
+	}
 
-	/* Mode of the handler, Client or Server */
-	Handler::Mode Mode;
+	/**
+	 * Whether or not raw packet sends are enabled
+	 */
+	FORCEINLINE bool GetRawSend()
+	{
+		return bRawSend;
+	}
+
+	FORCEINLINE uint8 GetPacketBitAlignment()
+	{
+		return PacketBitAlignment;
+	}
 
-	/* Time, updated by Tick */
-	float Time;
 
 private:
-	/* Set state of handler */
+	/**
+	 * Set state of handler
+	 *
+	 * @param InState	The new state for the handler
+	 */
 	void SetState(Handler::State InState);
 
-	/* Called when handler is initialized */
+	/**
+	 * Called when net send/receive functions are triggered, when the handler is still uninitialized - to set a valid initial state
+	 */
+	void UpdateInitialState();
+
+	/**
+	 * Called when handler is finished initializing
+	 */
 	void HandlerInitialized();
 
-	/* Used for packing outgoing packets */
+	/**
+	 * Replaces IncomingPacket with all unread data from ReplacementPacket
+	 *
+	 * @param ReplacementPacket		The packet whose unread data should replace IncomingPacket
+	 */
+	void ReplaceIncomingPacket(FBitReader& ReplacementPacket);
+
+	/**
+	 * Takes a Packet whose position is not at bit 0, and shifts/aligns all packet data to place the current bit at position 0
+	 *
+	 * @param Packet	The packet to realign
+	 */
+	void RealignPacket(FBitReader& Packet);
+
+
+public:
+	/** Mode of the handler, Client or Server */
+	Handler::Mode Mode;
+
+	/** Time, updated by Tick */
+	float Time;
+
+
+private:
+	/** Used for packing outgoing packets */
 	FBitWriter OutgoingPacket;
 
-	/* Used for unpacking incoming packets */
+	/** Used for unpacking incoming packets */
 	FBitReader IncomingPacket;
 
-	/* Array of Handler components */
+	/** The HandlerComponent pipeline, for processing incoming/outgoing packets */
 	TArray<TSharedPtr<HandlerComponent>> HandlerComponents;
 
-	/* State of the handler */
+	/** The overall packet bit alignment (0-7), for all HandlerComponents - used to determine the bit-size of incoming packets/ */
+	uint8 PacketBitAlignment;
+
+	/** State of the handler */
 	Handler::State State;
 	
-	/*
-	* Packets that are buffered while handler components
-	* are being initialized
-	*/
+	/** Packets that are buffered while HandlerComponents are being initialized */
 	TArray<BufferedPacket*> BufferedPackets;
 
-	/*
-	* Packets that are queued to be sent when handler is ticked
-	*/
+	/** Packets that are queued to be sent when handler is ticked */
 	TQueue<BufferedPacket*> QueuedPackets;
 
-	/* Reliability Handler Component */
+	/** Packets that are buffered while HandlerComponents are being initialized */
+	TArray<BufferedPacket*> BufferedConnectionlessPackets;
+
+	/** Packets that are queued to be sent when handler is ticked */
+	TQueue<BufferedPacket*> QueuedConnectionlessPackets;
+
+	/** Reliability Handler Component */
+	// @todo #JohnB: Deprecate?
 	TSharedPtr<ReliabilityHandlerComponent> ReliabilityComponent;
+
+	/** Whether or not outgoing packets bypass the handler */
+	bool bRawSend;
 };
 
-/*
-* This class appends or modifies incoming and outgoing packets
-* on a connection
-*/
+/**
+ * This class appends or modifies incoming and outgoing packets on a connection
+ */
 class PACKETHANDLER_API HandlerComponent
 {
 public:
-	/* Default initialization of data */
+	/**
+	 * Base constructor
+	 */
 	HandlerComponent();
 
-	/** Base destructor */
-	virtual ~HandlerComponent() {}
+	/**
+	 * Base destructor
+	 */
+	virtual ~HandlerComponent()
+	{
+	}
 
-	/* Returns whether this handler is currently active */
+	/**
+	 * Returns whether this handler is currently active
+	 */
 	virtual bool IsActive() const;
 
-	/* Return whether this handler is valid */
+	/**
+	 * Return whether this handler is valid
+	 */
 	virtual bool IsValid() const PURE_VIRTUAL(Handler::Component::IsValid, return false;);
 
-	/* Returns whether this handler is initialized */
+	/**
+	 * Returns whether this handler is initialized
+	 */
 	bool IsInitialized() const;
 
-	/* Returns whether this handler requires a tick every frame */
+	/**
+	 * Returns whether this handler requires a tick every frame
+	 */
 	bool DoesTick() const;
 
-	/* Handles any incoming packets */
-	virtual void Incoming(FBitReader& Packet) = 0;
 
-	/* Handles any outgoing packets */
-	virtual void Outgoing(FBitWriter& Packet) = 0;
+	/**
+	 * Handles incoming packets
+	 *
+	 * @param Packet	The packet to be handled
+	 */
+	virtual void Incoming(FBitReader& Packet) PURE_VIRTUAL(HandlerComponent::Incoming,);
+
+	/**
+	 * Handles any outgoing packets
+	 *
+	 * @param Packet	The packet to be handled
+	 */
+	virtual void Outgoing(FBitWriter& Packet) PURE_VIRTUAL(HandlerComponent::Outgoing,);
+
+	/**
+	 * Handles incoming packets not associated with a UNetConnection
+	 *
+	 * @param Address	The address the packet was received from (format is abstract, determined by active net driver)
+	 * @param Packet	The packet to be handled
+	 */
+	virtual void IncomingConnectionless(FString Address, FBitReader& Packet) PURE_VIRTUAL(HandlerComponent::IncomingConnectionless,);
+
+	/**
+	 * Handles any outgoing packets not associated with a UNetConnection
+	 *
+	 * @param Address	The address the packet is being sent to (format is abstract, determined by active net driver)
+	 * @param Packet	The packet to be handled
+	 */
+	virtual void OutgoingConnectionless(FString Address, FBitWriter& Packet) PURE_VIRTUAL(HandlerComponent::OutgoingConnectionless,);
+
+
+	/**
+	 * Whether or not the Incoming/IncomingConnectionless implementations, support reading Packets that aren't aligned at bit position 0
+	 * (i.e. whether or not this handler supports bit-level, rather than byte-level, reads)
+	 *
+	 * @return	Whether or not the above is supported
+	 */
+	virtual bool CanReadUnaligned()
+	{
+		return false;
+	}
+
 
-	/* Initialization functionality should be placed here */
-	virtual void Initialize() = 0;
+	/**
+	 * Initialization functionality should be placed here
+	 */
+	virtual void Initialize() PURE_VIRTUAL(HandlerComponent::Initialize,);
 
-	/* Tick functionality should be placed here */
-	virtual void Tick(float DeltaTime){};
+	/**
+	 * Tick functionality should be placed here
+	 */
+	virtual void Tick(float DeltaTime)
+	{
+	}
 
-	/* Whether this handler is active */
+	/**
+	 * Sets whether this handler is currently active
+	 *
+	 * @param Active	Whether or not the handled should be active
+	 */
 	virtual void SetActive(bool Active);
 
 	/**
-	 * Returns the amount of packet/protocol overhead expected from this handler.
+	 * Returns the amount of packet/protocol overhead expected from this component.
 	 *
-	 * IMPORTANT: This MUST be accurate, and should represent the worst-case overhead expected from the handler.
+	 * IMPORTANT: This MUST be accurate, and should represent the worst-case overhead expected from the component.
 	 *				If this is inaccurate, packets will randomly fail to send, in rare cases which are extremely hard to trace.
 	 *
-	 * @return	The worst-case packet overhead for the handler
+	 * @return	The worst-case packet overhead for the component
 	 */
 	virtual int32 GetPacketOverheadBits() PURE_VIRTUAL(Handler::Component::GetPacketOverheadBits, return -1;);
 
-protected:
+	/**
+	 * Determines the bit alignment (0-7) of this component, so the PacketHandler can determine the bit-size of incoming packets,
+	 * when they are not byte-aligned.
+	 *
+	 * NOTE: The return value must be static, determined at compile time. It must not change at runtime.
+	 *
+	 * @return	Returns the bit alignment of this component
+	 */
+	virtual uint8 GetBitAlignment() PURE_VIRTUAL(Handler::Component::GetBitAlignment, return 0;);
+
+	/**
+	 * Whether or not the changes this component makes to the packet, resets all bit alignment changes to the packet,
+	 * from previous components in the chain (with the new packet alignment, matching this components bit alignment).
+	 *
+	 * NOTE: The return value must be static, determined at compile time. It must not change at runtime.
+	 *
+	 * @return	Returns whether or not this component resets packet bit alignment.
+	 */
+	virtual bool DoesResetBitAlignment() PURE_VIRTUAL(Handler::Component::DoesResetBitAlignment, return false;);
+
 
-	/* Sets the state of the handler */
+protected:
+	/**
+	 * Sets the state of the handler
+	 *
+	 * @param State		The new state for the handler
+	 */
 	void SetState(Handler::Component::State State);
 
-	/* Should be called when the handler is fully initialized on both remote and local */
+	/**
+	 * Should be called when the handler is fully initialized on both remote and local
+	 */
 	void Initialized();
 
 
 public:
-	/* The manager of the handler, set in initialization */
+	/** The manager of the handler, set in initialization */
 	PacketHandler* Handler; 
 
 protected:
-	/* The state of this handler */
+	/** The state of this handler */
 	Handler::Component::State State;
 
 private:
-	/* 
-	* Whether this handler is active, which dictates whether it will receive
-	* incoming and outgoing packets.
-	*/
+	/** Whether this handler is active, which dictates whether it will receive incoming and outgoing packets. */
 	bool bActive;
 
-	/* Whether this handler is fully initialized on both remote and local */
+	/** Whether this handler is fully initialized on both remote and local */
 	bool bInitialized;
-
-	/* Allows manager access to setting the Handler variable */
-	friend PacketHandler;
 };
 
-/** Packet Handler Module Interface */
+/**
+ * PacketHandler Module Interface
+ */
 class FPacketHandlerComponentModuleInterface : public IModuleInterface
 {
 public:
diff --git a/Engine/Source/Runtime/PacketHandlers/ReliabilityHandlerComponent/Private/ReliabilityHandlerComponent.cpp b/Engine/Source/Runtime/PacketHandlers/ReliabilityHandlerComponent/Private/ReliabilityHandlerComponent.cpp
--- a/Engine/Source/Runtime/PacketHandlers/ReliabilityHandlerComponent/Private/ReliabilityHandlerComponent.cpp
+++ b/Engine/Source/Runtime/PacketHandlers/ReliabilityHandlerComponent/Private/ReliabilityHandlerComponent.cpp
@@ -56,8 +56,8 @@ void ReliabilityHandlerComponent::Outgoing(FBitWriter& Packet)
 					memcpy(Data, Packet.GetData(), Packet.GetNumBytes());
 					ReliablePacket->Data = Data;
 					ReliablePacket->Id = LocalPacketID;
-					ReliablePacket->BytesCount = Packet.GetNumBytes();
-					ReliablePacket->SendTime = Handler->Time + ResendResolutionTime;
+					ReliablePacket->CountBits = Packet.GetNumBits();
+					ReliablePacket->ResendTime = Handler->Time + ResendResolutionTime;
 					BufferedPackets.Enqueue(ReliablePacket);
 				}
 
@@ -139,7 +139,7 @@ void ReliabilityHandlerComponent::Tick(float DeltaTime)
 		BufferedPackets.Peek(Packet);
 
 		// Resend
-		if (Packet->SendTime > CurrentTime)
+		if (Packet->ResendTime > CurrentTime)
 		{
 			BufferedPackets.Dequeue(Packet);
 
@@ -168,16 +168,9 @@ void ReliabilityHandlerComponent::Tick(float DeltaTime)
 	}
 }
 
-void ReliabilityHandlerComponent::QueuePacketForResending(uint8* Packet, int32 Count)
+void ReliabilityHandlerComponent::QueuePacketForResending(uint8* Packet, int32 CountBits)
 {
-	BufferedPacket* Buffered = new BufferedPacket;
-	Buffered->Data = new uint8[Count];
-	memcpy(Buffered->Data, Packet, Count);
-	Buffered->Id = LocalPacketID;
-	Buffered->SendTime = Handler->Time + ResendResolutionTime;
-	Buffered->BytesCount = Count;
-
-	BufferedPackets.Enqueue(Buffered);
+	BufferedPackets.Enqueue(new BufferedPacket(Packet, CountBits, Handler->Time + ResendResolutionTime, LocalPacketID));
 }
 
 // MODULE INTERFACE
diff --git a/Engine/Source/Runtime/PacketHandlers/ReliabilityHandlerComponent/Public/ReliabilityHandlerComponent.h b/Engine/Source/Runtime/PacketHandlers/ReliabilityHandlerComponent/Public/ReliabilityHandlerComponent.h
--- a/Engine/Source/Runtime/PacketHandlers/ReliabilityHandlerComponent/Public/ReliabilityHandlerComponent.h
+++ b/Engine/Source/Runtime/PacketHandlers/ReliabilityHandlerComponent/Public/ReliabilityHandlerComponent.h
@@ -29,7 +29,7 @@ public:
 	virtual void Outgoing(FBitWriter& Packet) override;
 
 	/* Queues a packet for resending */
-	virtual void QueuePacketForResending(uint8* Packet, int32 Count);
+	virtual void QueuePacketForResending(uint8* Packet, int32 CountBits);
 
 protected:
 	/* Buffered Packets in case they need to be resent */
