Merging //UE4/Release-4.11 to //UE4/Main (Up to CL#2897738)

diff --git a/Engine/Source/Runtime/Engine/Classes/Engine/World.h b/Engine/Source/Runtime/Engine/Classes/Engine/World.h
--- a/Engine/Source/Runtime/Engine/Classes/Engine/World.h
+++ b/Engine/Source/Runtime/Engine/Classes/Engine/World.h
@@ -2775,10 +2775,16 @@ public:
 		return (OwningGameInstance ? OwningGameInstance->GetTimerManager() : *TimerManager);
 	}
 
-	/** Returns LatentActionManager instance for this world. */
+	/**
+	 * Returns LatentActionManager instance, preferring the one allocated by the game instance if a game instance is associated with this.
+	 *
+	 * This pattern is a little bit of a kludge to allow UWorld clients (for instance, preview world in the Blueprint Editor
+ 	 * to not worry about replacing features from GameInstance. Alternatively we could mandate that they implement a game instance
+	 * for their scene.
+	 */
 	inline FLatentActionManager& GetLatentActionManager()
 	{
-		return LatentActionManager;
+		return (OwningGameInstance ? OwningGameInstance->GetLatentActionManager() : LatentActionManager);
 	}
 
 	/** Sets the owning game instance for this world */
diff --git a/Engine/Source/Runtime/Engine/Classes/GameFramework/InputSettings.h b/Engine/Source/Runtime/Engine/Classes/GameFramework/InputSettings.h
--- a/Engine/Source/Runtime/Engine/Classes/GameFramework/InputSettings.h
+++ b/Engine/Source/Runtime/Engine/Classes/GameFramework/InputSettings.h
@@ -2,8 +2,8 @@
 
 #pragma once
 #include "PlayerInput.h"
-#include "InputSettings.generated.h"
 
+#include "InputSettings.generated.h"
 
 /**
  * Project wide settings for input handling
@@ -43,6 +43,10 @@ class ENGINE_API UInputSettings
 	UPROPERTY(config, EditAnywhere, Category="MouseProperties", AdvancedDisplay)
 	float DoubleClickTime;
 
+	/** The default mouse capture mode for the game viewport */
+	UPROPERTY(config, EditAnywhere, Category = "ViewportProperties")
+	TEnumAsByte<EMouseCaptureMode> DefaultViewportMouseCaptureMode;
+
 	/** List of Action Mappings */
 	UPROPERTY(config, EditAnywhere, Category="Bindings")
 	TArray<struct FInputActionKeyMapping> ActionMappings;
diff --git a/Engine/Source/Runtime/Engine/Classes/GameFramework/PlayerController.h b/Engine/Source/Runtime/Engine/Classes/GameFramework/PlayerController.h
--- a/Engine/Source/Runtime/Engine/Classes/GameFramework/PlayerController.h
+++ b/Engine/Source/Runtime/Engine/Classes/GameFramework/PlayerController.h
@@ -121,10 +121,16 @@ protected:
 /** Data structure used to setup an input mode that allows only the player input / player controller to respond to user input. */
 struct ENGINE_API FInputModeGameOnly : public FInputModeDataBase
 {
+	/** Whether the mouse down that causes capture should be consumed, and not passed to player input processing */
+	FInputModeGameOnly& SetConsumeCaptureMouseDown(bool InConsumeCaptureMouseDown) { bConsumeCaptureMouseDown = InConsumeCaptureMouseDown; return *this; }
+
 	FInputModeGameOnly()
+		: bConsumeCaptureMouseDown(true)
 	{}
 
 protected:
+	bool bConsumeCaptureMouseDown;
+
 	virtual void ApplyInputMode(FReply& SlateOperations, class UGameViewportClient& GameViewportClient) const override;
 };
 
@@ -902,6 +908,13 @@ public:
 	*/
 	UFUNCTION(BlueprintCallable, Category = "Game|Feedback")
 	void SetHapticsByValue(const float Frequency, const float Amplitude, TEnumAsByte<EControllerHand> Hand);
+	
+	/**
+	* Sets the light color of the player's controller
+	* @param	Color					The color for the light to be
+	*/
+	UFUNCTION(BlueprintCallable, Category = "Game|Feedback")
+	void SetControllerLightColor(FColor Color);
 
 	/**
 	 * Travel to a different map or IP address. Calls the PreClientTravel event before doing anything.
diff --git a/Engine/Source/Runtime/Engine/Classes/Vehicles/WheeledVehicleMovementComponent.h b/Engine/Source/Runtime/Engine/Classes/Vehicles/WheeledVehicleMovementComponent.h
--- a/Engine/Source/Runtime/Engine/Classes/Vehicles/WheeledVehicleMovementComponent.h
+++ b/Engine/Source/Runtime/Engine/Classes/Vehicles/WheeledVehicleMovementComponent.h
@@ -537,6 +537,9 @@ protected:
 	
 	/** remaining time of avoidance velocity lock */
 	float AvoidanceLockTimer;
+
+	/** Handle for delegate registered on mesh component */
+	FDelegateHandle MeshOnPhysicsStateChangeHandle;
 	
 	/** BEGIN IRVOAvoidanceInterface */
 	virtual void SetRVOAvoidanceUID(int32 UID) override;
diff --git a/Engine/Source/Runtime/Engine/Private/AI/Navigation/NavLinkProxy.cpp b/Engine/Source/Runtime/Engine/Private/AI/Navigation/NavLinkProxy.cpp
--- a/Engine/Source/Runtime/Engine/Private/AI/Navigation/NavLinkProxy.cpp
+++ b/Engine/Source/Runtime/Engine/Private/AI/Navigation/NavLinkProxy.cpp
@@ -96,6 +96,12 @@ void ANavLinkProxy::PostEditChangeProperty(FPropertyChangedEvent& PropertyChange
 }
 #endif // WITH_EDITOR
 
+void ANavLinkProxy::PostInitializeComponents()
+{
+	Super::PostInitializeComponents();
+	SmartLinkComp->SetNavigationRelevancy(bSmartLinkIsRelevant);
+}
+
 #if ENABLE_VISUAL_LOG
 void ANavLinkProxy::BeginPlay()
 {
diff --git a/Engine/Source/Runtime/Engine/Private/AI/Navigation/NavRelevantComponent.cpp b/Engine/Source/Runtime/Engine/Private/AI/Navigation/NavRelevantComponent.cpp
--- a/Engine/Source/Runtime/Engine/Private/AI/Navigation/NavRelevantComponent.cpp
+++ b/Engine/Source/Runtime/Engine/Private/AI/Navigation/NavRelevantComponent.cpp
@@ -5,6 +5,7 @@
 
 UNavRelevantComponent::UNavRelevantComponent(const FObjectInitializer& ObjectInitializer) : Super(ObjectInitializer)
 {
+	bCanEverAffectNavigation = true;
 	bNavigationRelevant = true;
 	bAttachToOwnersRoot = true;
 	bBoundsInitialized = false;
diff --git a/Engine/Source/Runtime/Engine/Private/ActorConstruction.cpp b/Engine/Source/Runtime/Engine/Private/ActorConstruction.cpp
--- a/Engine/Source/Runtime/Engine/Private/ActorConstruction.cpp
+++ b/Engine/Source/Runtime/Engine/Private/ActorConstruction.cpp
@@ -91,6 +91,29 @@ void AActor::ResetPropertiesForConstruction()
 	}
 }
 
+ int32 CalcComponentAttachDepth(UActorComponent* InComp, TMap<UActorComponent*, int32>& ComponentDepthMap) 
+ {
+	int32 ComponentDepth = 0;
+	int32* CachedComponentDepth = ComponentDepthMap.Find(InComp);
+	if (CachedComponentDepth)
+	{
+		ComponentDepth = *CachedComponentDepth;
+	}
+	else
+	{
+		if (USceneComponent* SC = Cast<USceneComponent>(InComp))
+		{
+			if (SC->AttachParent && SC->AttachParent->GetOwner() == InComp->GetOwner())
+			{
+				ComponentDepth = CalcComponentAttachDepth(SC->AttachParent, ComponentDepthMap) + 1;
+			}
+		}
+		ComponentDepthMap.Add(InComp, ComponentDepth);
+	}
+
+	return ComponentDepth;
+ }
+
 //* Destroys the constructed components.
 void AActor::DestroyConstructedComponents()
 {
@@ -98,22 +121,24 @@ void AActor::DestroyConstructedComponents()
 	TInlineComponentArray<UActorComponent*> PreviouslyAttachedComponents;
 	GetComponents(PreviouslyAttachedComponents);
 
-	// We need the hierarchy to be torn down in attachment order, so do a quick sort
-	PreviouslyAttachedComponents.Remove(nullptr);
-	PreviouslyAttachedComponents.Sort([](UActorComponent& A, UActorComponent& B)
+	TMap<UActorComponent*, int32> ComponentDepthMap;
+
+	for (UActorComponent* Component : PreviouslyAttachedComponents)
 	{
-		if (USceneComponent* BSC = Cast<USceneComponent>(&B))
+		if (Component)
 		{
-			if (BSC->GetAttachParent() == &A)
-			{
-				return false;
-			}
+			CalcComponentAttachDepth(Component, ComponentDepthMap);
 		}
-		return true;
+	}
+
+	ComponentDepthMap.ValueSort([](const int32& A, const int32& B)
+	{
+		return A > B;
 	});
 
-	for (UActorComponent* Component : PreviouslyAttachedComponents)
+	for (const TPair<UActorComponent*,int32>& ComponentAndDepth : ComponentDepthMap)
 	{
+		UActorComponent* Component = ComponentAndDepth.Key;
 		if (Component)
 		{
 			bool bDestroyComponent = false;
diff --git a/Engine/Source/Runtime/Engine/Private/Animation/AnimInstance.cpp b/Engine/Source/Runtime/Engine/Private/Animation/AnimInstance.cpp
--- a/Engine/Source/Runtime/Engine/Private/Animation/AnimInstance.cpp
+++ b/Engine/Source/Runtime/Engine/Private/Animation/AnimInstance.cpp
@@ -507,7 +507,7 @@ bool UAnimInstance::ParallelCanEvaluate(const USkeletalMesh* InSkeletalMesh) con
 	return Proxy.GetRequiredBones().IsValid() && (Proxy.GetRequiredBones().GetAsset() == InSkeletalMesh);
 }
 
-void UAnimInstance::ParallelEvaluateAnimation(bool bForceRefPose, const USkeletalMesh* InSkeletalMesh, TArray<FTransform>& OutLocalAtoms, TArray<FActiveVertexAnim>& OutVertexAnims, FBlendedHeapCurve& OutCurve)
+void UAnimInstance::ParallelEvaluateAnimation(bool bForceRefPose, const USkeletalMesh* InSkeletalMesh, TArray<FTransform>& OutLocalAtoms, FBlendedHeapCurve& OutCurve)
 {
 	FMemMark Mark(FMemStack::Get());
 	FAnimInstanceProxy& Proxy = GetProxyOnAnyThread<FAnimInstanceProxy>();
@@ -543,8 +543,6 @@ void UAnimInstance::ParallelEvaluateAnimation(bool bForceRefPose, const USkeleta
 	{
 		FAnimationRuntime::FillWithRefPose(OutLocalAtoms, Proxy.GetRequiredBones());
 	}
-
-	OutVertexAnims = FAnimationRuntime::UpdateActiveVertexAnims(InSkeletalMesh, Proxy.GetMorphTargetCurves(), Proxy.GetVertexAnims());
 }
 
 void UAnimInstance::PostEvaluateAnimation()
@@ -1083,6 +1081,12 @@ void UAnimInstance::TickSyncGroupWriteIndex()
 	GetProxyOnGameThread<FAnimInstanceProxy>().TickSyncGroupWriteIndex();
 }
 
+void UAnimInstance::RefreshCurves(USkeletalMeshComponent* Component)
+{
+	// update curves to component
+	GetProxyOnGameThread<FAnimInstanceProxy>().UpdateCurvesToComponents(Component);
+}
+
 void UAnimInstance::UpdateCurves(const FBlendedHeapCurve& InCurve)
 {
 	SCOPE_CYCLE_COUNTER(STAT_UpdateCurves);
@@ -1120,6 +1124,9 @@ void UAnimInstance::UpdateCurves(const FBlendedHeapCurve& InCurve)
 	{
 		AddCurveValue(ParamsToClearCopy[i], 0.0f, ACF_DrivesMaterial);
 	}
+
+	// update curves to component
+	Proxy.UpdateCurvesToComponents();
 }
 
 bool UAnimInstance::HasMorphTargetCurves() const
@@ -2606,27 +2613,11 @@ void UAnimInstance::DestroyAnimInstanceProxy(FAnimInstanceProxy* InProxy)
 	delete InProxy;
 }
 
-void UAnimInstance::UpdateMorphTargetCurves(const TMap<FName, float>& InMorphTargetCurves)
-{
-	GetProxyOnGameThread<FAnimInstanceProxy>().UpdateMorphTargetCurves(InMorphTargetCurves);
-}
-
-void UAnimInstance::UpdateComponentsMaterialParameters(UPrimitiveComponent* Component)
-{
-	GetProxyOnGameThread<FAnimInstanceProxy>().UpdateComponentsMaterialParameters(Component);
-}
-
 void UAnimInstance::RecordStateWeight(const int32& InMachineClassIndex, const int32& InStateIndex, const float& InStateWeight)
 {
 	GetProxyOnAnyThread<FAnimInstanceProxy>().RecordStateWeight(InMachineClassIndex, InStateIndex, InStateWeight);
 }
 
-TArray<FActiveVertexAnim> UAnimInstance::UpdateActiveVertexAnims(const USkeletalMesh* SkeletalMesh)
-{
-	FAnimInstanceProxy& Proxy = GetProxyOnAnyThread<FAnimInstanceProxy>();	// this can be called from CreateRenderState_Concurrent
-	return FAnimationRuntime::UpdateActiveVertexAnims(SkeletalMesh, Proxy.GetMorphTargetCurves(), Proxy.GetVertexAnims());
-}
-
 FBoneContainer& UAnimInstance::GetRequiredBones()
 {
 	return GetProxyOnGameThread<FAnimInstanceProxy>().GetRequiredBones();
diff --git a/Engine/Source/Runtime/Engine/Private/Animation/AnimInstanceProxy.cpp b/Engine/Source/Runtime/Engine/Private/Animation/AnimInstanceProxy.cpp
--- a/Engine/Source/Runtime/Engine/Private/Animation/AnimInstanceProxy.cpp
+++ b/Engine/Source/Runtime/Engine/Private/Animation/AnimInstanceProxy.cpp
@@ -677,40 +677,15 @@ void FAnimInstanceProxy::EvaluateAnimation(FPoseContext& Output)
 	}
 }
 
-void FAnimInstanceProxy::UpdateMorphTargetCurves(const TMap<FName, float>& InMorphTargetCurves)
+void FAnimInstanceProxy::UpdateCurvesToComponents(USkeletalMeshComponent* Component)
 {
-	QUICK_SCOPE_CYCLE_COUNTER(STAT_AnimInstanceProxy_UpdateMorphTargetCurves);
-
-	// now all tick/trigger/kismet is done
-	// add MorphTarget Curves from Kismet driven or any other source
-	// and overwrite if it exists
-	// Tick always should maintain this list, not Evaluate
-	for( auto Iter = InMorphTargetCurves.CreateConstIterator(); Iter; ++Iter )
+	if (Component)
 	{
-		float *CurveValPtr = MorphTargetCurves.Find(Iter.Key());
-		if ( CurveValPtr )
-		{
-			// override the value if Kismet request was made
-			*CurveValPtr = Iter.Value();
-		}
-		else
-		{
-			MorphTargetCurves.Add(Iter.Key(), Iter.Value());
-		}
+		Component->ApplyAnimationCurvesToComponent(&MaterialParameterCurves, &MorphTargetCurves);
 	}
-}
-
-void FAnimInstanceProxy::UpdateComponentsMaterialParameters(UPrimitiveComponent* Component)
-{
-	if(MaterialParameterCurves.Num() > 0)
+	else
 	{
-		QUICK_SCOPE_CYCLE_COUNTER(STAT_FAnimInstanceProxy_UpdateComponentsMaterialParameters);
-		for(auto Iter = MaterialParameterCurves.CreateConstIterator(); Iter; ++Iter)
-		{
-			FName ParameterName = Iter.Key();
-			float ParameterValue = Iter.Value();
-			SkeletalMeshComponent->SetScalarParameterValueOnMaterials(ParameterName, ParameterValue);
-		}
+		SkeletalMeshComponent->ApplyAnimationCurvesToComponent(&MaterialParameterCurves, &MorphTargetCurves);
 	}
 }
 
diff --git a/Engine/Source/Runtime/Engine/Private/Animation/AnimMontage.cpp b/Engine/Source/Runtime/Engine/Private/Animation/AnimMontage.cpp
--- a/Engine/Source/Runtime/Engine/Private/Animation/AnimMontage.cpp
+++ b/Engine/Source/Runtime/Engine/Private/Animation/AnimMontage.cpp
@@ -2101,7 +2101,6 @@ void FAnimMontageInstance::PreviewMatineeSetAnimPositionInner(FName SlotName, US
 	}
 
 	// Update space bases so new animation position has an effect.
-	SkeletalMeshComponent->UpdateMaterialParameters();
 	SkeletalMeshComponent->RefreshBoneTransforms();
 	SkeletalMeshComponent->RefreshSlaveComponents();
 	SkeletalMeshComponent->UpdateComponentToWorld();
diff --git a/Engine/Source/Runtime/Engine/Private/Animation/AnimSequence.cpp b/Engine/Source/Runtime/Engine/Private/Animation/AnimSequence.cpp
--- a/Engine/Source/Runtime/Engine/Private/Animation/AnimSequence.cpp
+++ b/Engine/Source/Runtime/Engine/Private/Animation/AnimSequence.cpp
@@ -25,6 +25,7 @@ DECLARE_CYCLE_STAT(TEXT("AnimSeq GetBonePose"), STAT_AnimSeq_GetBonePose, STATGR
 DECLARE_CYCLE_STAT(TEXT("Build Anim Track Pairs"), STAT_BuildAnimTrackPairs, STATGROUP_Anim);
 DECLARE_CYCLE_STAT(TEXT("Extract Pose From Anim Data"), STAT_ExtractPoseFromAnimData, STATGROUP_Anim);
 
+TAutoConsoleVariable<int32> CVarUseBakedAdditiveAnimations(TEXT("a.UseBakedAdditiveAnimations"), 0, TEXT("If 1, additive delta calculations will be baked into the animation during cook. This improves runtime performance at the expense of significantly increasing cook times"));
 
 /////////////////////////////////////////////////////
 // FRawAnimSequenceTrackNativeDeprecated
@@ -326,7 +327,7 @@ void UAnimSequence::Serialize(FArchive& Ar)
 #if WITH_EDITORONLY_DATA
 	FScopedAnimSequenceCompressedDataCache CompressedDataCache;
 
-	if (!bCompressedOnCook && Ar.IsCooking() && Ar.IsSaving() && CanBakeAdditive())
+	if (!bCompressedOnCook && Ar.IsCooking() && Ar.IsSaving() && CanBakeAdditive() && CVarUseBakedAdditiveAnimations.GetValueOnGameThread() == 1)
 	{
 		CompressedDataCache.InitFrom(this);
 
diff --git a/Engine/Source/Runtime/Engine/Private/Animation/AnimationRuntime.cpp b/Engine/Source/Runtime/Engine/Private/Animation/AnimationRuntime.cpp
--- a/Engine/Source/Runtime/Engine/Private/Animation/AnimationRuntime.cpp
+++ b/Engine/Source/Runtime/Engine/Private/Animation/AnimationRuntime.cpp
@@ -1367,25 +1367,8 @@ static int32 FindVertexAnim(const TArray<FActiveVertexAnim>& ActiveAnims, UVerte
 	return INDEX_NONE;
 }
 
-TArray<FActiveVertexAnim> FAnimationRuntime::UpdateActiveVertexAnims(const USkeletalMesh* InSkeletalMesh, const TMap<FName, float>& MorphCurveAnims, const TArray<FActiveVertexAnim>& ActiveAnims)
+void FAnimationRuntime::AppendActiveVertexAnims(const USkeletalMesh* InSkeletalMesh, const TMap<FName, float>& MorphCurveAnims, TArray<FActiveVertexAnim>& InOutActiveAnims)
 {
-	TArray<struct FActiveVertexAnim> OutVertexAnims;
-
-	// First copy ActiveAnims
-	for(int32 AnimIdx=0; AnimIdx < ActiveAnims.Num(); AnimIdx++)
-	{
-		const FActiveVertexAnim& ActiveAnim = ActiveAnims[AnimIdx];
-		const float ActiveAnimAbsWeight = FMath::Abs(ActiveAnim.Weight);
-
-		// Check it has valid weight, and works on this SkeletalMesh
-		if (	ActiveAnimAbsWeight > MinVertexAnimBlendWeight &&
-			ActiveAnim.VertAnim != NULL &&
-			ActiveAnim.VertAnim->BaseSkelMesh == InSkeletalMesh)
-		{
-			OutVertexAnims.Add(ActiveAnim);
-		}
-		// @TODO Need to check for duplicates here?
-	}
 
 	// Then go over the CurveKeys finding morph targets by name
 	for(auto CurveIter=MorphCurveAnims.CreateConstIterator(); CurveIter; ++CurveIter)
@@ -1401,23 +1384,21 @@ TArray<FActiveVertexAnim> FAnimationRuntime::UpdateActiveVertexAnims(const USkel
 			if(Target != NULL)				
 			{
 				// See if this morph target already has an entry
-				int32 AnimIndex = FindVertexAnim(OutVertexAnims, Target);
+				int32 AnimIndex = FindVertexAnim(InOutActiveAnims, Target);
 				// If not, add it
 				if(AnimIndex == INDEX_NONE)
 				{
-					OutVertexAnims.Add(FActiveVertexAnim(Target, Weight));
+					InOutActiveAnims.Add(FActiveVertexAnim(Target, Weight));
 				}
 				// If it does, use the max weight
 				else
 				{
-					const float CurrentWeight = OutVertexAnims[AnimIndex].Weight;
-					OutVertexAnims[AnimIndex].Weight = FMath::Max<float>(CurrentWeight, Weight);
+					const float CurrentWeight = InOutActiveAnims[AnimIndex].Weight;
+					InOutActiveAnims[AnimIndex].Weight = FMath::Max<float>(CurrentWeight, Weight);
 				}
 			}
 		}
 	}
-
-	return OutVertexAnims;
 }
 
 /////////////////////////////////////////////////////////////////////////////////////////
diff --git a/Engine/Source/Runtime/Engine/Private/AudioDevice.cpp b/Engine/Source/Runtime/Engine/Private/AudioDevice.cpp
--- a/Engine/Source/Runtime/Engine/Private/AudioDevice.cpp
+++ b/Engine/Source/Runtime/Engine/Private/AudioDevice.cpp
@@ -1296,11 +1296,12 @@ void FAudioDevice::RecurseIntoSoundClasses(USoundClass* CurrentClass, FSoundClas
 
 void FAudioDevice::UpdateHighestPriorityReverb()
 {
-	for (auto It = ActivatedReverbs.CreateConstIterator(); It; ++It)
+	HighestPriorityReverb = nullptr;
+	for (const TPair<FName, FActivatedReverb>& ActivatedReverb : ActivatedReverbs)
 	{
-		if (!HighestPriorityReverb || It.Value().Priority > HighestPriorityReverb->Priority)
+		if (HighestPriorityReverb == nullptr || ActivatedReverb.Value.Priority > HighestPriorityReverb->Priority)
 		{
-			HighestPriorityReverb = &It.Value();
+			HighestPriorityReverb = &ActivatedReverb.Value;
 		}
 	}
 }
@@ -2156,49 +2157,21 @@ void FAudioDevice::ClearSoundMixModifiers()
 
 void FAudioDevice::ActivateReverbEffect(UReverbEffect* ReverbEffect, FName TagName, float Priority, float Volume, float FadeTime)
 {
-	FActivatedReverb* ExistingReverb = ActivatedReverbs.Find(TagName);
+	FActivatedReverb& ActivatedReverb = ActivatedReverbs.FindOrAdd(TagName);
 
-	if (ExistingReverb)
-	{
-		float OldPriority = ExistingReverb->Priority;
-		ExistingReverb->ReverbSettings.ReverbEffect = ReverbEffect;
-		ExistingReverb->ReverbSettings.Volume = Volume;
-		ExistingReverb->ReverbSettings.FadeTime = FadeTime;
-		ExistingReverb->Priority = Priority;
-
-		if (OldPriority != Priority)
-		{
-			UpdateHighestPriorityReverb();
-		}
-	}
-	else
-	{
-		ExistingReverb = &ActivatedReverbs.Add(TagName, FActivatedReverb());
-		ExistingReverb->ReverbSettings.ReverbEffect = ReverbEffect;
-		ExistingReverb->ReverbSettings.Volume = Volume;
-		ExistingReverb->ReverbSettings.FadeTime = FadeTime;
-		ExistingReverb->Priority = Priority;
+	ActivatedReverb.ReverbSettings.ReverbEffect = ReverbEffect;
+	ActivatedReverb.ReverbSettings.Volume = Volume;
+	ActivatedReverb.ReverbSettings.FadeTime = FadeTime;
+	ActivatedReverb.Priority = Priority;
 
-		UpdateHighestPriorityReverb();
-	}
+	UpdateHighestPriorityReverb();
 }
 
 void FAudioDevice::DeactivateReverbEffect(FName TagName)
 {
-	FActivatedReverb* ExistingReverb = ActivatedReverbs.Find(TagName);
-	if (ExistingReverb)
+	if (ActivatedReverbs.Remove(TagName) > 0)
 	{
-		if (ExistingReverb == HighestPriorityReverb)
-		{
-			HighestPriorityReverb = nullptr;
-		}
-
-		ActivatedReverbs.Remove(TagName);
-
-		if (HighestPriorityReverb == nullptr)
-		{
-			UpdateHighestPriorityReverb();
-		}
+		UpdateHighestPriorityReverb();
 	}
 }
 
diff --git a/Engine/Source/Runtime/Engine/Private/Components/DestructibleComponent.cpp b/Engine/Source/Runtime/Engine/Private/Components/DestructibleComponent.cpp
--- a/Engine/Source/Runtime/Engine/Private/Components/DestructibleComponent.cpp
+++ b/Engine/Source/Runtime/Engine/Private/Components/DestructibleComponent.cpp
@@ -310,8 +310,11 @@ void UDestructibleComponent::CreatePhysicsState()
 	verify( NxParameterized::setParamU32(*ActorParams,"p3ShapeDescTemplate.queryFilterData.word3", PQueryFilterData.word3 ) );
 
 	// Set the PhysX material in the shape descriptor
-	PxMaterial* PMaterial = PhysMat->GetPhysXMaterial();
-	verify( NxParameterized::setParamU64(*ActorParams,"p3ShapeDescTemplate.material", (physx::PxU64)PMaterial) );
+	if(PxMaterial* PMaterial = PhysMat->GetPhysXMaterial())
+	{
+		verify(NxParameterized::setParamU64(*ActorParams, "p3ShapeDescTemplate.material", (physx::PxU64)PMaterial));
+	}
+
 
 	// Set the rest depth to match the skin width in the shape descriptor
 	const physx::PxCookingParams& CookingParams = GApexSDK->getCookingInterface()->getParams();
@@ -337,7 +340,14 @@ void UDestructibleComponent::CreatePhysicsState()
 	SleepEnergyThreshold *= BodyInstance.GetSleepThresholdMultiplier();
 	verify( NxParameterized::setParamF32(*ActorParams,"p3BodyDescTemplate.sleepThreshold", SleepEnergyThreshold) );
 //	NxParameterized::setParamF32(*ActorParams,"bodyDescTemplate.sleepDamping", SleepDamping );
-	verify( NxParameterized::setParamF32(*ActorParams,"p3BodyDescTemplate.density", 0.001f*PhysMat->Density) );	// Convert from g/cm^3 to kg/cm^3
+	
+	float DensityPerCubicCM = 1.0f;
+	if(PhysMat)
+	{
+		DensityPerCubicCM = PhysMat->Density;
+	}
+	verify( NxParameterized::setParamF32(*ActorParams,"p3BodyDescTemplate.density", 0.001f * DensityPerCubicCM) );	// Convert from g/cm^3 to kg/cm^3
+
 	// Enable CCD if requested
 	verify( NxParameterized::setParamBool(*ActorParams,"p3BodyDescTemplate.flags.eENABLE_CCD", BodyInstance.bUseCCD != 0) );
 	// Ask the actor to create chunk events, for more efficient visibility updates
@@ -1551,12 +1561,14 @@ void UDestructibleComponent::SetMaterial(int32 ElementIndex, UMaterialInterface*
 		if(ApexDestructibleActor->getPhysX3Template(*Template))
 		{
 			UPhysicalMaterial* SimpleMaterial = GetBodyInstance()->GetSimplePhysicalMaterial();
-			check(SimpleMaterial);
-			PxMaterial* PhysxMat = SimpleMaterial->GetPhysXMaterial();
-
-			Template->setMaterials(&PhysxMat, 1);
+			
+			if(SimpleMaterial)
+			{
+				PxMaterial* PhysxMat = SimpleMaterial->GetPhysXMaterial();
 
-			ApexDestructibleActor->setPhysX3Template(Template);
+				Template->setMaterials(&PhysxMat, 1);
+				ApexDestructibleActor->setPhysX3Template(Template);
+			}
 		}
 		Template->release();
 	}
diff --git a/Engine/Source/Runtime/Engine/Private/Components/InterpToMovementComponent.cpp b/Engine/Source/Runtime/Engine/Private/Components/InterpToMovementComponent.cpp
--- a/Engine/Source/Runtime/Engine/Private/Components/InterpToMovementComponent.cpp
+++ b/Engine/Source/Runtime/Engine/Private/Components/InterpToMovementComponent.cpp
@@ -500,6 +500,10 @@ void UInterpToMovementComponent::AddControlPointPosition(FVector Pos,bool bPosit
 
 void UInterpToMovementComponent::FinaliseControlPoints()
 {
+	if (UpdatedComponent == nullptr)
+	{
+		return;
+	}
 	StartLocation = UpdatedComponent->GetComponentLocation();
 	TimeMultiplier = 1.0f / Duration;
 	if (ControlPoints.Num() != 0)
diff --git a/Engine/Source/Runtime/Engine/Private/Components/SceneComponent.cpp b/Engine/Source/Runtime/Engine/Private/Components/SceneComponent.cpp
--- a/Engine/Source/Runtime/Engine/Private/Components/SceneComponent.cpp
+++ b/Engine/Source/Runtime/Engine/Private/Components/SceneComponent.cpp
@@ -399,9 +399,9 @@ void USceneComponent::UpdateComponentToWorldWithParent(USceneComponent* Parent,F
 	bool bHasChanged;
 	{
 		//QUICK_SCOPE_CYCLE_COUNTER(STAT_USceneComponent_UpdateComponentToWorldWithParent_HasChanged);
-		bHasChanged = ComponentToWorld.Equals(NewTransform, SMALL_NUMBER);
+		bHasChanged = !ComponentToWorld.Equals(NewTransform, SMALL_NUMBER);
 	}
-	if (!bHasChanged)
+	if (bHasChanged)
 	{
 		//QUICK_SCOPE_CYCLE_COUNTER(STAT_USceneComponent_UpdateComponentToWorldWithParent_Changed);
 		// Update transform
@@ -773,13 +773,13 @@ void USceneComponent::OnComponentDestroyed(bool bDestroyingHierarchy)
 					bool bNeedsDetach = true;
 					if (ExternalAttachParent)
 					{
-						bNeedsDetach = (Child->AttachTo(ExternalAttachParent) == false);
+						bNeedsDetach = (Child->AttachTo(ExternalAttachParent, NAME_None, EAttachLocation::KeepWorldPosition) == false);
 					}
 					if (bNeedsDetach)
 					{
 						if (Child->AttachParent && Child->AttachParent == this)
 						{
-							Child->DetachFromParent();
+							Child->DetachFromParent(true);
 						}
 						else
 						{
@@ -821,13 +821,13 @@ void USceneComponent::OnComponentDestroyed(bool bDestroyingHierarchy)
 					bool bNeedsDetach = true;
 					if (AttachParent)
 					{
-						bNeedsDetach = (Child->AttachTo(AttachParent) == false);
+						bNeedsDetach = (Child->AttachTo(AttachParent, NAME_None, EAttachLocation::KeepWorldPosition) == false);
 					}
 					if (bNeedsDetach)
 					{
 						if (Child->AttachParent && Child->AttachParent == this)
 						{
-							Child->DetachFromParent();
+							Child->DetachFromParent(true);
 						}
 						else
 						{
@@ -860,7 +860,7 @@ void USceneComponent::OnComponentDestroyed(bool bDestroyingHierarchy)
 		if (AttachParent && (!bDestroyingHierarchy || AttachParent->GetOwner() != MyOwner))
 		{
 			// Ensure we are detached before destroying
-			DetachFromParent();
+			DetachFromParent(true);
 		}
 	}
 }
@@ -1454,7 +1454,7 @@ bool USceneComponent::AttachTo(class USceneComponent* Parent, FName InSocketName
 
 			if (PrimitiveComponent && PrimitiveComponent->BodyInstance.bSimulatePhysics && !bWeldSimulatedBodies && GetWorld() && GetWorld()->IsGameWorld())
 			{
-				 if(!GetWorld()->bIsRunningConstructionScript)
+				 if(!GetWorld()->bIsRunningConstructionScript && (GetOwner()->HasActorBegunPlay() || GetOwner()->IsActorBeginningPlay()))
 				 {
 					 //Since the object is physically simulated it can't be the case that it's a child of object A and being attached to object B (at runtime)
 					 bDisableDetachmentUpdateOverlaps = true;
@@ -1474,7 +1474,8 @@ bool USceneComponent::AttachTo(class USceneComponent* Parent, FName InSocketName
 					 }
 
 					 return false;
-				 }else
+				 }
+				 else
 				 {
 					//A simulated object needs to be detached at runtime. We are in the construction script so we can't do it here. However, we want to make sure it is done in BeginPlay.
 					bWantsBeginPlay = true;
diff --git a/Engine/Source/Runtime/Engine/Private/Components/SkeletalMeshComponent.cpp b/Engine/Source/Runtime/Engine/Private/Components/SkeletalMeshComponent.cpp
--- a/Engine/Source/Runtime/Engine/Private/Components/SkeletalMeshComponent.cpp
+++ b/Engine/Source/Runtime/Engine/Private/Components/SkeletalMeshComponent.cpp
@@ -437,7 +437,10 @@ bool USkeletalMeshComponent::InitializeAnimScriptInstance(bool bForceReinit)
 		{
 			AnimScriptInstance->InitializeAnimation();
 			bCalledInitialize = true;
-		}
+		}		
+
+		// refresh vertex animation - this can happen when re-registration happens
+		RefreshActiveVertexAnims();
 	}
 	return bCalledInitialize;
 }
@@ -588,28 +591,10 @@ void USkeletalMeshComponent::TickAnimation(float DeltaTime, bool bNeedsValidRoot
 		{
 			// Tick the animation
 			AnimScriptInstance->UpdateAnimation(DeltaTime * GlobalAnimRateScale, bNeedsValidRootMotion);
-
-			// TODO @LinaH - I've hit access violations due to AnimScriptInstance being NULL after this, probably due to
-			// AnimNotifies?  Please take a look and fix as we discussed.  Temporary fix:
-			if (AnimScriptInstance != NULL)
-			{
-				AnimScriptInstance->UpdateMorphTargetCurves(MorphTargetCurves);
-
-				//Update material parameters
-				AnimScriptInstance->UpdateComponentsMaterialParameters(this);
-			}
 		}
 	}
 }
 
-void USkeletalMeshComponent::UpdateMaterialParameters()
-{
-	if(AnimScriptInstance != nullptr)
-	{
-		AnimScriptInstance->UpdateComponentsMaterialParameters(this);
-	}
-}
-
 bool USkeletalMeshComponent::UpdateLODStatus()
 {
 	if ( Super::UpdateLODStatus() )
@@ -704,6 +689,13 @@ void USkeletalMeshComponent::TickComponent(float DeltaTime, enum ELevelTick Tick
 	UpdatePostPhysicsTickRegisteredState();
 	UpdateClothTickRegisteredState();
 
+	// clear and add morphtarget curves that are added via SetMorphTarget
+	ActiveVertexAnims.Reset();
+	if (SkeletalMesh && MorphTargetCurves.Num() > 0)
+	{
+		FAnimationRuntime::AppendActiveVertexAnims(SkeletalMesh, MorphTargetCurves, ActiveVertexAnims);
+	}
+
 	Super::TickComponent(DeltaTime, TickType, ThisTickFunction);
 
 	// Update bOldForceRefPose
@@ -1033,7 +1025,7 @@ void USkeletalMeshComponent::RecalcRequiredBones(int32 LODIndex)
 	CachedSpaceBases.Empty();
 }
 
-void USkeletalMeshComponent::EvaluateAnimation(const USkeletalMesh* InSkeletalMesh, UAnimInstance* InAnimInstance, TArray<FTransform>& OutLocalAtoms, TArray<FActiveVertexAnim>& OutVertexAnims, FVector& OutRootBoneTranslation, FBlendedHeapCurve& OutCurve) const
+void USkeletalMeshComponent::EvaluateAnimation(const USkeletalMesh* InSkeletalMesh, UAnimInstance* InAnimInstance, TArray<FTransform>& OutLocalAtoms, FVector& OutRootBoneTranslation, FBlendedHeapCurve& OutCurve) const
 {
 	ANIM_MT_SCOPE_CYCLE_COUNTER(SkeletalComponentAnimEvaluate, IsRunningParallelEvaluation());
 
@@ -1048,18 +1040,11 @@ void USkeletalMeshComponent::EvaluateAnimation(const USkeletalMesh* InSkeletalMe
 		ensure(bRequiredBonesUpToDate) &&
 		InAnimInstance->ParallelCanEvaluate(InSkeletalMesh))
 	{
-		InAnimInstance->ParallelEvaluateAnimation(bForceRefpose, InSkeletalMesh, OutLocalAtoms, OutVertexAnims, OutCurve);
+		InAnimInstance->ParallelEvaluateAnimation(bForceRefpose, InSkeletalMesh, OutLocalAtoms, OutCurve);
 	}
 	else
 	{
 		OutLocalAtoms = InSkeletalMesh->RefSkeleton.GetRefBonePose();
-
-		// if it's only morph, there is no reason to blend
-		if ( MorphTargetCurves.Num() > 0 )
-		{
-			TArray<struct FActiveVertexAnim> EmptyVertexAnims;
-			OutVertexAnims = FAnimationRuntime::UpdateActiveVertexAnims(InSkeletalMesh, MorphTargetCurves, EmptyVertexAnims);
-		}
 	}
 
 	// Remember the root bone's translation so we can move the bounds.
@@ -1070,20 +1055,18 @@ void USkeletalMeshComponent::UpdateSlaveComponent()
 {
 	check (MasterPoseComponent.IsValid());
 
-	if(MasterPoseComponent->IsA(USkeletalMeshComponent::StaticClass()))
+	if (USkeletalMeshComponent* MasterSMC = Cast<USkeletalMeshComponent>(MasterPoseComponent.Get()))
 	{
-		USkeletalMeshComponent* MasterSMC= CastChecked<USkeletalMeshComponent>(MasterPoseComponent.Get());
-
-		if ( MasterSMC->AnimScriptInstance )
+		if (MasterSMC->AnimScriptInstance)
 		{
-			ActiveVertexAnims = MasterSMC->AnimScriptInstance->UpdateActiveVertexAnims(SkeletalMesh);
+			MasterSMC->AnimScriptInstance->RefreshCurves(this);
 		}
 	}
-
+ 
 	Super::UpdateSlaveComponent();
 }
 
-void USkeletalMeshComponent::PerformAnimationEvaluation(const USkeletalMesh* InSkeletalMesh, UAnimInstance* InAnimInstance, TArray<FTransform>& OutSpaceBases, TArray<FTransform>& OutLocalAtoms, TArray<FActiveVertexAnim>& OutVertexAnims, FVector& OutRootBoneTranslation, FBlendedHeapCurve& OutCurve) const
+void USkeletalMeshComponent::PerformAnimationEvaluation(const USkeletalMesh* InSkeletalMesh, UAnimInstance* InAnimInstance, TArray<FTransform>& OutSpaceBases, TArray<FTransform>& OutLocalAtoms, FVector& OutRootBoneTranslation, FBlendedHeapCurve& OutCurve) const
 {
 	ANIM_MT_SCOPE_CYCLE_COUNTER(PerformAnimEvaluation, IsRunningParallelEvaluation());
 
@@ -1101,7 +1084,7 @@ void USkeletalMeshComponent::PerformAnimationEvaluation(const USkeletalMesh* InS
 	}
 
 	// evaluate pure animations, and fill up LocalAtoms
-	EvaluateAnimation(InSkeletalMesh, InAnimInstance, OutLocalAtoms, OutVertexAnims, OutRootBoneTranslation, OutCurve);
+	EvaluateAnimation(InSkeletalMesh, InAnimInstance, OutLocalAtoms, OutRootBoneTranslation, OutCurve);
 	// Fill SpaceBases from LocalAtoms
 	FillSpaceBases(InSkeletalMesh, OutLocalAtoms, OutSpaceBases);
 }
@@ -1243,8 +1226,6 @@ void USkeletalMeshComponent::RefreshBoneTransforms(FActorComponentTickFunction*
 			AnimEvaluationContext.LocalAtoms.Append(LocalAtoms);
 			AnimEvaluationContext.SpaceBases.Reset();
 			AnimEvaluationContext.SpaceBases.Append(GetSpaceBases());
-			AnimEvaluationContext.VertexAnims.Reset();
-			AnimEvaluationContext.VertexAnims.Append(ActiveVertexAnims);
 		}
 
 		// start parallel work
@@ -1268,11 +1249,11 @@ void USkeletalMeshComponent::RefreshBoneTransforms(FActorComponentTickFunction*
 			QUICK_SCOPE_CYCLE_COUNTER(STAT_USkeletalMeshComponent_RefreshBoneTransforms_GamethreadEval);
 			if (AnimEvaluationContext.bDoInterpolation)
 			{
-				PerformAnimationEvaluation(SkeletalMesh, AnimScriptInstance, CachedSpaceBases, CachedLocalAtoms, ActiveVertexAnims, RootBoneTranslation, CachedCurve);
+				PerformAnimationEvaluation(SkeletalMesh, AnimScriptInstance, CachedSpaceBases, CachedLocalAtoms, RootBoneTranslation, CachedCurve);
 			}
 			else
 			{
-				PerformAnimationEvaluation(SkeletalMesh, AnimScriptInstance, GetEditableSpaceBases(), LocalAtoms, ActiveVertexAnims, RootBoneTranslation, AnimEvaluationContext.Curve);
+				PerformAnimationEvaluation(SkeletalMesh, AnimScriptInstance, GetEditableSpaceBases(), LocalAtoms, RootBoneTranslation, AnimEvaluationContext.Curve);
 			}
 		}
 		else
@@ -1377,6 +1358,26 @@ void USkeletalMeshComponent::PostAnimEvaluation(FAnimationEvaluationContext& Eva
 	AnimEvaluationContext.Clear();
 }
 
+void USkeletalMeshComponent::ApplyAnimationCurvesToComponent(const TMap<FName, float>* InMaterialParameterCurves, const TMap<FName, float>* InAnimationMorphCurves)
+{
+	if (InMaterialParameterCurves && InMaterialParameterCurves->Num() > 0)
+	{
+		QUICK_SCOPE_CYCLE_COUNTER(STAT_FAnimInstanceProxy_UpdateComponentsMaterialParameters);
+		for (auto Iter = InMaterialParameterCurves->CreateConstIterator(); Iter; ++Iter)
+		{
+			FName ParameterName = Iter.Key();
+			float ParameterValue = Iter.Value();
+			SetScalarParameterValueOnMaterials(ParameterName, ParameterValue);
+		}
+	}
+
+	if (SkeletalMesh && InAnimationMorphCurves && InAnimationMorphCurves->Num() > 0)
+	{
+		// we want to append to existing curves - i.e. BP driven curves 
+		FAnimationRuntime::AppendActiveVertexAnims(SkeletalMesh, *InAnimationMorphCurves, ActiveVertexAnims);
+	}
+}
+
 FBoxSphereBounds USkeletalMeshComponent::CalcBounds(const FTransform& LocalToWorld) const
 {
 	SCOPE_CYCLE_COUNTER(STAT_CalcSkelMeshBounds);
@@ -1950,7 +1951,6 @@ FRootMotionMovementParams USkeletalMeshComponent::ConsumeRootMotion()
 	return FRootMotionMovementParams();
 }
 
-
 float USkeletalMeshComponent::CalculateMass(FName BoneName)
 {
 	float Mass = 0.0f;
@@ -2249,7 +2249,14 @@ void USkeletalMeshComponent::RefreshActiveVertexAnims()
 	{
 		// as this can be called from any worker thread (i.e. from CreateRenderState_Concurrent) we cant currently be doing parallel evaluation
 		check(!IsRunningParallelEvaluation());
-		ActiveVertexAnims = AnimScriptInstance->UpdateActiveVertexAnims(SkeletalMesh);
+		AnimScriptInstance->RefreshCurves(this);
+	}
+	else if (USkeletalMeshComponent* MasterSMC = Cast<USkeletalMeshComponent>(MasterPoseComponent.Get()))
+	{
+		if (MasterSMC->AnimScriptInstance)
+		{
+			MasterSMC->AnimScriptInstance->RefreshCurves(this);
+		}
 	}
 	else
 	{
@@ -2259,7 +2266,7 @@ void USkeletalMeshComponent::RefreshActiveVertexAnims()
 
 void USkeletalMeshComponent::ParallelAnimationEvaluation() 
 { 
-	PerformAnimationEvaluation(AnimEvaluationContext.SkeletalMesh, AnimEvaluationContext.AnimInstance, AnimEvaluationContext.SpaceBases, AnimEvaluationContext.LocalAtoms, AnimEvaluationContext.VertexAnims, AnimEvaluationContext.RootBoneTranslation, AnimEvaluationContext.Curve); 
+	PerformAnimationEvaluation(AnimEvaluationContext.SkeletalMesh, AnimEvaluationContext.AnimInstance, AnimEvaluationContext.SpaceBases, AnimEvaluationContext.LocalAtoms, AnimEvaluationContext.RootBoneTranslation, AnimEvaluationContext.Curve); 
 }
 
 void USkeletalMeshComponent::CompleteParallelAnimationEvaluation(bool bDoPostAnimEvaluation)
@@ -2273,7 +2280,6 @@ void USkeletalMeshComponent::CompleteParallelAnimationEvaluation(bool bDoPostAni
 
 			Exchange(AnimEvaluationContext.SpaceBases, AnimEvaluationContext.bDoInterpolation ? CachedSpaceBases : GetEditableSpaceBases());
 			Exchange(AnimEvaluationContext.LocalAtoms, AnimEvaluationContext.bDoInterpolation ? CachedLocalAtoms : LocalAtoms);
-			Exchange(AnimEvaluationContext.VertexAnims, ActiveVertexAnims);
 			Exchange(AnimEvaluationContext.RootBoneTranslation, RootBoneTranslation);
 		}
 
@@ -2432,4 +2438,14 @@ void USkeletalMeshComponent::UpdateCachedAnimCurveMappingNameUids()
 TArray<FSmartNameMapping::UID> const * USkeletalMeshComponent::GetCachedAnimCurveMappingNameUids()
 {
 	return &CachedAnimCurveMappingNameUids;
-}
\ No newline at end of file
+}
+
+FDelegateHandle USkeletalMeshComponent::RegisterOnPhysicsCreatedDelegate(const FOnSkelMeshPhysicsCreated& Delegate)
+{
+	return OnSkelMeshPhysicsCreated.Add(Delegate);
+}
+
+void USkeletalMeshComponent::UnregisterOnPhysicsCreatedDelegate(const FDelegateHandle& DelegateHandle)
+{
+	OnSkelMeshPhysicsCreated.Remove(DelegateHandle);
+}
diff --git a/Engine/Source/Runtime/Engine/Private/Components/SkinnedMeshComponent.cpp b/Engine/Source/Runtime/Engine/Private/Components/SkinnedMeshComponent.cpp
--- a/Engine/Source/Runtime/Engine/Private/Components/SkinnedMeshComponent.cpp
+++ b/Engine/Source/Runtime/Engine/Private/Components/SkinnedMeshComponent.cpp
@@ -345,9 +345,6 @@ void USkinnedMeshComponent::CreateRenderState_Concurrent()
 
 		if(MeshObject)
 		{
-			// verifies vertex animations are valid
-			RefreshActiveVertexAnims();
-
 			// Identify current LOD
 			int32 UseLOD;
 			if(MasterPoseComponent.IsValid())
diff --git a/Engine/Source/Runtime/Engine/Private/DataReplication.cpp b/Engine/Source/Runtime/Engine/Private/DataReplication.cpp
--- a/Engine/Source/Runtime/Engine/Private/DataReplication.cpp
+++ b/Engine/Source/Runtime/Engine/Private/DataReplication.cpp
@@ -1162,7 +1162,14 @@ void FObjectReplicator::CallRepNotifies(bool bSkipIfChannelHasQueuedBunches)
 		{
 			//UE_LOG(LogRep, Log,  TEXT("Calling Object->%s with %s"), *RepNotifies(RepNotifyIdx)->RepNotifyFunc.ToString(), *RepNotifies(RepNotifyIdx)->GetName()); 						
 			UProperty* RepProperty = RepNotifies[RepNotifyIdx];
-			UFunction* RepNotifyFunc = Object->FindFunctionChecked(RepProperty->RepNotifyFunc);
+			UFunction* RepNotifyFunc = Object->FindFunction(RepProperty->RepNotifyFunc);
+
+			if (RepNotifyFunc == nullptr)
+			{
+				UE_LOG(LogRep, Warning, TEXT("FObjectReplicator::CallRepNotifies: Can't find RepNotify function %s for property %s on object %s."),
+					*RepProperty->RepNotifyFunc.ToString(), *RepProperty->GetName(), *Object->GetName());
+				continue;
+			}
 
 			if (RepNotifyFunc->NumParms == 0)
 			{
diff --git a/Engine/Source/Runtime/Engine/Private/DestructibleMesh.cpp b/Engine/Source/Runtime/Engine/Private/DestructibleMesh.cpp
--- a/Engine/Source/Runtime/Engine/Private/DestructibleMesh.cpp
+++ b/Engine/Source/Runtime/Engine/Private/DestructibleMesh.cpp
@@ -256,7 +256,11 @@ void UDestructibleMesh::FinishDestroy()
 void FDestructibleDamageParameters::FillDestructibleActorDesc(NxParameterized::Interface* Params, UPhysicalMaterial* PhysMat) const
 {
 	// Damage parameters
-	verify(NxParameterized::setParamF32(*Params, "defaultBehaviorGroup.damageThreshold", DamageThreshold * PhysMat->DestructibleDamageThresholdScale));
+	if(PhysMat)
+	{
+		verify(NxParameterized::setParamF32(*Params, "defaultBehaviorGroup.damageThreshold", DamageThreshold * PhysMat->DestructibleDamageThresholdScale));
+	}
+
 	verify(NxParameterized::setParamF32(*Params, "defaultBehaviorGroup.damageToRadius", DamageSpread));
 	verify(NxParameterized::setParamF32(*Params, "destructibleParameters.forceToDamage", ImpactDamage));
 	verify(NxParameterized::setParamF32(*Params, "defaultBehaviorGroup.materialStrength", ImpactResistanceToAPEX(bCustomImpactResistance, ImpactResistance)));
diff --git a/Engine/Source/Runtime/Engine/Private/GameEngine.cpp b/Engine/Source/Runtime/Engine/Private/GameEngine.cpp
--- a/Engine/Source/Runtime/Engine/Private/GameEngine.cpp
+++ b/Engine/Source/Runtime/Engine/Private/GameEngine.cpp
@@ -396,8 +396,10 @@ void UGameEngine::SwitchGameWindowToUseGameViewport()
 		{
 			FSlateApplication::Get().RegisterGameViewport(GameViewportWidgetRef);
 		}
-		
-		FSlateApplication::Get().SetAllUserFocusToGameViewport();
+		else
+		{
+			FSlateApplication::Get().ActivateGameViewport();
+		}
 	}
 }
 
diff --git a/Engine/Source/Runtime/Engine/Private/GameInstance.cpp b/Engine/Source/Runtime/Engine/Private/GameInstance.cpp
--- a/Engine/Source/Runtime/Engine/Private/GameInstance.cpp
+++ b/Engine/Source/Runtime/Engine/Private/GameInstance.cpp
@@ -9,6 +9,7 @@
 #include "Engine/GameInstance.h"
 #include "Engine/Engine.h"
 #include "Engine/DemoNetDriver.h"
+#include "Engine/LatentActionManager.h"
 #include "Engine/NetworkObjectList.h"
 #include "OnlineSubsystem.h"
 #include "OnlineSessionInterface.h"
@@ -19,10 +20,10 @@
 #include "UnrealEd.h"
 #endif
 
-
 UGameInstance::UGameInstance(const FObjectInitializer& ObjectInitializer)
 : Super(ObjectInitializer)
 , TimerManager(new FTimerManager())
+, LatentActionManager(new FLatentActionManager())
 {
 }
 
@@ -34,6 +35,13 @@ void UGameInstance::FinishDestroy()
 		TimerManager = nullptr;
 	}
 
+	// delete operator should handle null, but maintaining pattern of TimerManager:
+	if (LatentActionManager)
+	{
+		delete LatentActionManager;
+		LatentActionManager = nullptr;
+	}
+
 	Super::FinishDestroy();
 }
 
@@ -758,6 +766,18 @@ void UGameInstance::StartRecordingReplay(const FString& Name, const FString& Fri
 		return;
 	}
 
+	if ( CurrentWorld->WorldType == EWorldType::PIE )
+	{
+		UE_LOG(LogDemo, Warning, TEXT("UGameInstance::StartRecordingReplay: Function called while running a PIE instance, this is disabled."));
+		return;
+	}
+
+	if ( CurrentWorld->DemoNetDriver && CurrentWorld->DemoNetDriver->IsPlaying() )
+	{
+		UE_LOG(LogDemo, Warning, TEXT("UGameInstance::StartRecordingReplay: A replay is already playing, cannot begin recording another one."));
+		return;
+	}
+
 	FURL DemoURL;
 	FString DemoName = Name;
 	
@@ -767,7 +787,7 @@ void UGameInstance::StartRecordingReplay(const FString& Name, const FString& Fri
 	DemoURL.Map = DemoName;
 	DemoURL.AddOption( *FString::Printf( TEXT( "DemoFriendlyName=%s" ), *FriendlyName ) );
 
-	for (const FString& Option : AdditionalOptions)
+	for ( const FString& Option : AdditionalOptions )
 	{
 		DemoURL.AddOption(*Option);
 	}
@@ -824,6 +844,12 @@ void UGameInstance::PlayReplay(const FString& Name, UWorld* WorldOverride, const
 		return;
 	}
 
+	if ( CurrentWorld->WorldType == EWorldType::PIE )
+	{
+		UE_LOG( LogDemo, Warning, TEXT( "UGameInstance::PlayReplay: Function called while running a PIE instance, this is disabled." ) );
+		return;
+	}
+
 	CurrentWorld->DestroyDemoNetDriver();
 
 	FURL DemoURL;
@@ -831,7 +857,7 @@ void UGameInstance::PlayReplay(const FString& Name, UWorld* WorldOverride, const
 
 	DemoURL.Map = Name;
 	
-	for (const FString& Option : AdditionalOptions)
+	for ( const FString& Option : AdditionalOptions )
 	{
 		DemoURL.AddOption(*Option);
 	}
diff --git a/Engine/Source/Runtime/Engine/Private/GameViewportClient.cpp b/Engine/Source/Runtime/Engine/Private/GameViewportClient.cpp
--- a/Engine/Source/Runtime/Engine/Private/GameViewportClient.cpp
+++ b/Engine/Source/Runtime/Engine/Private/GameViewportClient.cpp
@@ -119,6 +119,7 @@ UGameViewportClient::UGameViewportClient(const FObjectInitializer& ObjectInitial
 	, bHideCursorDuringCapture(false)
 	, AudioDeviceHandle(INDEX_NONE)
 	, bHasAudioFocus(false)
+	, bActive(false)
 {
 
 	TitleSafeZone.MaxPercentX = 0.9f;
@@ -177,6 +178,7 @@ UGameViewportClient::UGameViewportClient(FVTableHelper& Helper)
 	, bHideCursorDuringCapture(false)
 	, AudioDeviceHandle(INDEX_NONE)
 	, bHasAudioFocus(false)
+	, bActive(false)
 {
 
 }
@@ -275,6 +277,9 @@ void UGameViewportClient::Init(struct FWorldContext& WorldContext, UGameInstance
 	// remember our game instance
 	GameInstance = OwningGameInstance;
 
+	// Set the projects default viewport mouse capture mode
+	MouseCaptureMode = GetDefault<UInputSettings>()->DefaultViewportMouseCaptureMode;
+
 	// Create the cursor Widgets
 	UUserInterfaceSettings* UISettings = GetMutableDefault<UUserInterfaceSettings>(UUserInterfaceSettings::StaticClass());
 
@@ -591,7 +596,8 @@ bool UGameViewportClient::RequiresUncapturedAxisInput() const
 
 EMouseCursor::Type UGameViewportClient::GetCursor(FViewport* InViewport, int32 X, int32 Y)
 {
-	if (!InViewport->HasMouseCapture() || (ViewportConsole && ViewportConsole->ConsoleActive()))
+	// If the viewport isn't active or the console is active we don't want to override the cursor
+	if (!bActive || (!InViewport->HasMouseCapture() && !InViewport->HasFocus()) || (ViewportConsole && ViewportConsole->ConsoleActive()))
 	{
 		return EMouseCursor::Default;
 	}
@@ -1387,11 +1393,13 @@ bool UGameViewportClient::IsFocused(FViewport* InViewport)
 void UGameViewportClient::Activated(FViewport* InViewport, const FWindowActivateEvent& InActivateEvent)
 {
 	ReceivedFocus(InViewport);
+	bActive = true;
 }
 
 void UGameViewportClient::Deactivated(FViewport* InViewport, const FWindowActivateEvent& InActivateEvent)
 {
 	LostFocus(InViewport);
+	bActive = false;
 }
 
 void UGameViewportClient::CloseRequested(FViewport* InViewport)
diff --git a/Engine/Source/Runtime/Engine/Private/LevelTick.cpp b/Engine/Source/Runtime/Engine/Private/LevelTick.cpp
--- a/Engine/Source/Runtime/Engine/Private/LevelTick.cpp
+++ b/Engine/Source/Runtime/Engine/Private/LevelTick.cpp
@@ -1173,8 +1173,10 @@ void UWorld::Tick( ELevelTick TickType, float DeltaSeconds )
 		&&	!bIsPaused
 		&&	(!NetDriver || !NetDriver->ServerConnection || NetDriver->ServerConnection->State==USOCK_Open);
 
+	FLatentActionManager& CurrentLatentActionManager = GetLatentActionManager();
+
 	// Reset the list of objects the LatentActionManager has processed this frame
-	LatentActionManager.BeginFrame();
+	CurrentLatentActionManager.BeginFrame();
 	// If caller wants time update only, or we are paused, skip the rest.
 	if (bDoingActorTicks)
 	{
@@ -1221,7 +1223,7 @@ void UWorld::Tick( ELevelTick TickType, float DeltaSeconds )
 	if( !bIsPaused )
 	{
 		// This will process any latent actions that have not been processed already
-		LatentActionManager.ProcessLatentActions( NULL, DeltaSeconds );	
+		CurrentLatentActionManager.ProcessLatentActions(NULL, DeltaSeconds);
 	}
 #if 0 // if you need to debug physics drawing in editor, use this. If you type pxvis collision, it will work. 
 	else
diff --git a/Engine/Source/Runtime/Engine/Private/NetworkDriver.cpp b/Engine/Source/Runtime/Engine/Private/NetworkDriver.cpp
--- a/Engine/Source/Runtime/Engine/Private/NetworkDriver.cpp
+++ b/Engine/Source/Runtime/Engine/Private/NetworkDriver.cpp
@@ -666,7 +666,7 @@ void UNetDriver::TickFlush(float DeltaSeconds)
 		}
 
 		LastCleanupTime = CurrentRealtimeSeconds;
-	}
+}
 
 	// Update the lag state
 	UpdateNetworkLagState();
@@ -2217,24 +2217,24 @@ FActorPriority::FActorPriority(class UNetConnection* InConnection, struct FActor
 
 #if WITH_SERVER_CODE
 int32 UNetDriver::ServerReplicateActors_PrepConnections( const float DeltaSeconds )
-{
+	{
 	int32 NumClientsToTick = ClientConnections.Num();
 
 	// by default only throttle update for listen servers unless specified on the commandline
-	static bool bForceClientTickingThrottle = FParse::Param( FCommandLine::Get(), TEXT( "limitclientticks" ) );
-	if ( bForceClientTickingThrottle || GetNetMode() == NM_ListenServer )
+	static bool bForceClientTickingThrottle = FParse::Param(FCommandLine::Get(),TEXT("limitclientticks"));
+	if (bForceClientTickingThrottle || GetNetMode() == NM_ListenServer)
 	{
 		// determine how many clients to tick this frame based on GEngine->NetTickRate (always tick at least one client), double for lan play
 		// FIXME: DeltaTimeOverflow is a static, and will conflict with other running net drivers, we investigate storing it on the driver itself!
 		static float DeltaTimeOverflow = 0.f;
 		// updates are doubled for lan play
-		static bool LanPlay = FParse::Param( FCommandLine::Get(), TEXT( "lanplay" ) );
+		static bool LanPlay = FParse::Param(FCommandLine::Get(),TEXT("lanplay"));
 		//@todo - ideally we wouldn't want to tick more clients with a higher deltatime as that's not going to be good for performance and probably saturate bandwidth in hitchy situations, maybe 
 		// come up with a solution that is greedier with higher framerates, but still won't risk saturating server upstream bandwidth
-		float ClientUpdatesThisFrame = GEngine->NetClientTicksPerSecond * ( DeltaSeconds + DeltaTimeOverflow ) * ( LanPlay ? 2.f : 1.f );
-		NumClientsToTick = FMath::Min<int32>( NumClientsToTick, FMath::TruncToInt( ClientUpdatesThisFrame ) );
+		float ClientUpdatesThisFrame = GEngine->NetClientTicksPerSecond * (DeltaSeconds + DeltaTimeOverflow) * (LanPlay ? 2.f : 1.f);
+		NumClientsToTick = FMath::Min<int32>(NumClientsToTick,FMath::TruncToInt(ClientUpdatesThisFrame));
 		//UE_LOG(LogNet, Log, TEXT("%2.3f: Ticking %d clients this frame, %2.3f/%2.4f"),GetWorld()->GetTimeSeconds(),NumClientsToTick,DeltaSeconds,ClientUpdatesThisFrame);
-		if ( NumClientsToTick == 0 )
+		if (NumClientsToTick == 0)
 		{
 			// if no clients are ticked this frame accumulate the time elapsed for the next frame
 			DeltaTimeOverflow += DeltaSeconds;
@@ -2245,31 +2245,31 @@ int32 UNetDriver::ServerReplicateActors_PrepConnections( const float DeltaSecond
 
 	bool bFoundReadyConnection = false;
 
-	for ( int32 ConnIdx = 0; ConnIdx < ClientConnections.Num(); ConnIdx++ )
+	for( int32 ConnIdx = 0; ConnIdx < ClientConnections.Num(); ConnIdx++ )
 	{
 		UNetConnection* Connection = ClientConnections[ConnIdx];
-		check( Connection );
-		check( Connection->State == USOCK_Pending || Connection->State == USOCK_Open || Connection->State == USOCK_Closed );
-		checkSlow( Connection->GetUChildConnection() == NULL );
+		check(Connection);
+		check(Connection->State == USOCK_Pending || Connection->State == USOCK_Open || Connection->State == USOCK_Closed);
+		checkSlow(Connection->GetUChildConnection() == NULL);
 
 		// Handle not ready channels.
 		//@note: we cannot add Connection->IsNetReady(0) here to check for saturation, as if that's the case we still want to figure out the list of relevant actors
 		//			to reset their NetUpdateTime so that they will get sent as soon as the connection is no longer saturated
 		AActor* OwningActor = Connection->OwningActor;
-		if ( OwningActor != NULL && Connection->State == USOCK_Open && ( Connection->Driver->Time - Connection->LastReceiveTime < 1.5f ) )
+		if (OwningActor != NULL && Connection->State == USOCK_Open && (Connection->Driver->Time - Connection->LastReceiveTime < 1.5f))
 		{
 			check( World == OwningActor->GetWorld() );
 
 			bFoundReadyConnection = true;
-
+			
 			// the view target is what the player controller is looking at OR the owning actor itself when using beacons
 			Connection->ViewTarget = Connection->PlayerController ? Connection->PlayerController->GetViewTarget() : OwningActor;
 
-			for ( int32 ChildIdx = 0; ChildIdx < Connection->Children.Num(); ChildIdx++ )
+			for (int32 ChildIdx = 0; ChildIdx < Connection->Children.Num(); ChildIdx++)
 			{
 				UNetConnection *Child = Connection->Children[ChildIdx];
 				APlayerController* ChildPlayerController = Child->PlayerController;
-				if ( ChildPlayerController != NULL )
+				if (ChildPlayerController != NULL)
 				{
 					Child->ViewTarget = ChildPlayerController->GetViewTarget();
 				}
@@ -2282,7 +2282,7 @@ int32 UNetDriver::ServerReplicateActors_PrepConnections( const float DeltaSecond
 		else
 		{
 			Connection->ViewTarget = NULL;
-			for ( int32 ChildIdx = 0; ChildIdx < Connection->Children.Num(); ChildIdx++ )
+			for (int32 ChildIdx = 0; ChildIdx < Connection->Children.Num(); ChildIdx++)
 			{
 				Connection->Children[ChildIdx]->ViewTarget = NULL;
 			}
@@ -2293,7 +2293,7 @@ int32 UNetDriver::ServerReplicateActors_PrepConnections( const float DeltaSecond
 }
 
 void UNetDriver::ServerReplicateActors_BuildConsiderList( TArray<FNetworkObjectInfo*>& OutConsiderList, const float ServerTickTime )
-{
+	{
 	SCOPE_CYCLE_COUNTER( STAT_NetConsiderActorsTime );
 
 	UE_LOG( LogNetTraffic, Log, TEXT( "ServerReplicateActors_BuildConsiderList, Building ConsiderList %4.2f" ), World->GetTimeSeconds() );
@@ -2309,49 +2309,49 @@ void UNetDriver::ServerReplicateActors_BuildConsiderList( TArray<FNetworkObjectI
 		FNetworkObjectInfo* ActorInfo = ( *ActorIt ).Get();
 		AActor* Actor = ActorInfo->Actor;
 
-		if ( Actor->IsPendingKill() )
-		{
-			ActorIt.RemoveCurrent();
-			continue;
-		}
+			if (Actor->IsPendingKill() )
+			{
+				ActorIt.RemoveCurrent();
+				continue;
+			}
 
-		if ( Actor->GetRemoteRole() == ROLE_None )
-		{
-			ActorIt.RemoveCurrent();
-			continue;
-		}
+			if (Actor->GetRemoteRole()==ROLE_None)
+			{
+				ActorIt.RemoveCurrent();
+				continue;
+			}
 
-		// This actor may belong to a different net driver, make sure this is the correct one
-		// (this can happen when using beacon net drivers for example)
-		if ( Actor->NetDriverName != NetDriverName )
-		{
-			continue;
-		}
+			// This actor may belong to a different net driver, make sure this is the correct one
+			// (this can happen when using beacon net drivers for example)
+			if ( Actor->NetDriverName != NetDriverName )
+			{
+				continue;
+			}
 
-		// Don't send actors that may still be streaming in
-		ULevel* Level = Actor->GetLevel();
-		if ( Level->HasVisibilityRequestPending() || Level->bIsAssociatingLevel )
-		{
-			continue;
-		}
+			// Don't send actors that may still be streaming in
+			ULevel* Level = Actor->GetLevel();
+			if ( Level->HasVisibilityRequestPending() || Level->bIsAssociatingLevel )
+			{
+				continue;
+			}
 
-		if ( Actor->NetDormancy == DORM_Initial && Actor->IsNetStartupActor() )
-		{
-			// This stat isn't that useful in its current form when using NetworkActors list
-			// We'll want to track initially dormant actors some other way to track them with stats
-			SCOPE_CYCLE_COUNTER( STAT_NetInitialDormantCheckTime );
-			NumInitiallyDormant++;
-			ActorIt.RemoveCurrent();
-			//UE_LOG(LogNetTraffic, Log, TEXT("Skipping Actor %s - its initially dormant!"), *Actor->GetName() );
-			continue;
-		}
+			if ( Actor->NetDormancy == DORM_Initial && Actor->IsNetStartupActor() )
+			{
+				// This stat isn't that useful in its current form when using NetworkActors list
+				// We'll want to track initially dormant actors some other way to track them with stats
+				SCOPE_CYCLE_COUNTER(STAT_NetInitialDormantCheckTime);		
+				NumInitiallyDormant++;
+				ActorIt.RemoveCurrent();
+				//UE_LOG(LogNetTraffic, Log, TEXT("Skipping Actor %s - its initially dormant!"), *Actor->GetName() );
+				continue;
+			}
 
-		check( Actor->NeedsLoadForClient() );			// We have no business sending this unless the client can load
+			check( Actor->NeedsLoadForClient() );			// We have no business sending this unless the client can load
 		check( World == Actor->GetWorld() );
 
 		// See if it's time for this actor to update replicated properties
-		if ( ( Actor->bPendingNetUpdate || World->TimeSeconds > Actor->NetUpdateTime ) )
-		{
+			if( (Actor->bPendingNetUpdate || World->TimeSeconds > Actor->NetUpdateTime) ) 
+			{
 			// Set defaults if this actor is replicating for first time
 			if ( ActorInfo->LastNetReplicateTime == 0 )
 			{
@@ -2384,122 +2384,122 @@ void UNetDriver::ServerReplicateActors_BuildConsiderList( TArray<FNetworkObjectI
 			// NOTE - We don't do this if bPendingNetUpdate is true, since this means we're forcing an update due to at least one connection
 			//	that wasn't to replicate previously (due to saturation, etc)
 			// NOTE - This also means all other connections will force an update (even if they just updated, we should look into this)
-			if ( !Actor->bPendingNetUpdate )
-			{
-				UE_LOG( LogNetTraffic, Log, TEXT( "actor %s requesting new net update, time: %2.3f" ), *Actor->GetName(), World->TimeSeconds );
+				if (!Actor->bPendingNetUpdate)
+				{
+					UE_LOG(LogNetTraffic, Log, TEXT("actor %s requesting new net update, time: %2.3f"), *Actor->GetName(), World->TimeSeconds);
 
 				const float NextUpdateDelta = bUseAdapativeNetFrequency ? ActorInfo->OptimalNetUpdateDelta : 1.0f / Actor->NetUpdateFrequency;
 
-				// then set the next update time
+					// then set the next update time
 				Actor->NetUpdateTime = World->TimeSeconds + FMath::SRand() * ServerTickTime + NextUpdateDelta;
 
-				// and mark when the actor first requested an update
-				//@note: using Time because it's compared against UActorChannel.LastUpdateTime which also uses that value
-				Actor->LastNetUpdateTime = Time;
-			}
+					// and mark when the actor first requested an update
+					//@note: using Time because it's compared against UActorChannel.LastUpdateTime which also uses that value
+					Actor->LastNetUpdateTime = Time;
+				}
 
-			// and clear the pending update flag assuming all clients will be able to consider it
-			Actor->bPendingNetUpdate = false;
+				// and clear the pending update flag assuming all clients will be able to consider it
+				Actor->bPendingNetUpdate = false;
 
-			// add it to the list to consider below
-			// For performance reasons, make sure we don't resize the array. It should already be appropriately sized above!
+					// add it to the list to consider below
+					// For performance reasons, make sure we don't resize the array. It should already be appropriately sized above!
 			ensure( OutConsiderList.Num() < OutConsiderList.Max() );
 			OutConsiderList.Add( ActorInfo );
 
 			// Call PreReplication on all actors that will be considered
 			Actor->CallPreReplication( this );
 		}
-	}
-
+				}
+										
 	// Update stats
 	SET_DWORD_STAT( STAT_NumInitiallyDormantActors, NumInitiallyDormant );
 	SET_DWORD_STAT( STAT_NumConsideredActors, OutConsiderList.Num() );
-}
+									}
 
 // Returns true if this actor should replicate to *any* of the passed in connections
 static FORCEINLINE_DEBUGGABLE bool IsActorRelevantToConnection( const AActor* Actor, const TArray<FNetViewer>& ConnectionViewers )
-{
+								{
 	for ( int32 viewerIdx = 0; viewerIdx < ConnectionViewers.Num(); viewerIdx++ )
-	{
+							{
 		if ( Actor->IsNetRelevantFor( ConnectionViewers[viewerIdx].InViewer, ConnectionViewers[viewerIdx].ViewTarget, ConnectionViewers[viewerIdx].ViewLocation ) )
-		{
+								{
 			return true;
-		}
-	}
+								}
+							}
 
 	return false;
-}
+					}
 
 // Returns true if this actor is owned by, and should replicate to *any* of the passed in connections
 static FORCEINLINE_DEBUGGABLE UNetConnection* IsActorOwnedByAndRelevantToConnection( const AActor* Actor, const TArray<FNetViewer>& ConnectionViewers, bool& bOutHasNullViewTarget )
-{
+				{
 	const AActor* ActorOwner = Actor->GetNetOwner();
 
 	bOutHasNullViewTarget = false;
 
 	for ( int i = 0; i < ConnectionViewers.Num(); i++ )
-	{
+			{
 		UNetConnection* ViewerConnection = ConnectionViewers[i].Connection;
 
 		if ( ViewerConnection->ViewTarget == nullptr )
-		{
+				{
 			bOutHasNullViewTarget = true;
-		}
+				}
 
 		if ( ActorOwner == ViewerConnection->PlayerController ||
 			 ( ViewerConnection->PlayerController && ActorOwner == ViewerConnection->PlayerController->GetPawn() ) ||
 			 ViewerConnection->ViewTarget->IsRelevancyOwnerFor( Actor, ActorOwner, ViewerConnection->OwningActor ) )
 		{
 			return ViewerConnection;
+			}
 		}
-	}
 
 	return nullptr;
-}
+	}
 
 // Returns true if this actor is considered dormant (and all properties caught up) to the current connection
 static FORCEINLINE_DEBUGGABLE bool IsActorDormant( const AActor* Actor, const UNetConnection* Connection )
-{
+	{
 	// If actor is already dormant on this channel, then skip replication entirely
 	if ( Connection->DormantActors.Contains( Actor ) )
-	{
+		{
 		// net.DormancyValidate can be set to 2 to validate dormant actor properties on every replicate
 		// (this could be moved to be done every tick instead of every net update if necessary, but seems excessive)
 		if ( CVarNetDormancyValidate.GetValueOnGameThread() == 2 )
-		{
+			{
 			const TSharedRef< FObjectReplicator > * Replicator = Connection->DormantReplicatorMap.Find( Actor );
 
 			if ( Replicator != NULL )
-			{
+				{
 				Replicator->Get().ValidateAgainstState( Actor );
-			}
-		}
+					}
+				}
 
 		return true;
-	}
+			}
 
 	return false;
-}
-
+			}
+			
 // Returns true if this actor wants to go dormant for a particular connection
 static FORCEINLINE_DEBUGGABLE bool ShouldActorGoDormant( AActor* Actor, const TArray<FNetViewer>& ConnectionViewers, UActorChannel* Channel, const float Time, const bool bLowNetBandwidth )
-{
+			{
 	if ( Actor->NetDormancy <= DORM_Awake || !Channel || Channel->bPendingDormancy || Channel->Dormant )
-	{
+				{
 		// Either shouldn't go dormant, or is already dormant
 		return false;
-	}
-
+				}
+				
 	if ( Actor->NetDormancy == DORM_DormantPartial )
 	{
 		for ( int32 viewerIdx = 0; viewerIdx < ConnectionViewers.Num(); viewerIdx++ )
-		{
+				{
 			if ( !Actor->GetNetDormancy( ConnectionViewers[viewerIdx].ViewLocation, ConnectionViewers[viewerIdx].ViewDir, ConnectionViewers[viewerIdx].InViewer, ConnectionViewers[viewerIdx].ViewTarget, Channel, Time, bLowNetBandwidth ) )
-			{
+					{
 				return false;
 			}
-		}
-	}
+					}
+				}
 
 	return true;
 }
@@ -2508,20 +2508,20 @@ int32 UNetDriver::ServerReplicateActors_PrioritizeActors( UNetConnection* Connec
 {
 	SCOPE_CYCLE_COUNTER( STAT_NetPrioritizeActorsTime );
 
-	// Get list of visible/relevant actors.
-
-	NetTag++;
-	Connection->TickCount++;
+				// Get list of visible/relevant actors.
+				
+				NetTag++;
+				Connection->TickCount++;
 
-	// Set up to skip all sent temporary actors
+				// Set up to skip all sent temporary actors
 	for ( int32 j=0; j < Connection->SentTemporaries.Num(); j++ )
-	{
-		Connection->SentTemporaries[j]->NetTag = NetTag;
-	}
-
-	// Make list of all actors to consider.
-	check( World == Connection->OwningActor->GetWorld() );
+				{
+					Connection->SentTemporaries[j]->NetTag = NetTag;
+				}
 
+				// Make list of all actors to consider.
+				check(World == Connection->OwningActor->GetWorld());
+				
 	int32 FinalSortedCount = 0;
 	int32 DeletedCount = 0;
 
@@ -2531,262 +2531,262 @@ int32 UNetDriver::ServerReplicateActors_PrioritizeActors( UNetConnection* Connec
 		OutPriorityList = new( FMemStack::Get(), MaxSortedActors )FActorPriority;
 		OutPriorityActors = new( FMemStack::Get(), MaxSortedActors )FActorPriority*;
 
-		// determine whether we should priority sort the list of relevant actors based on the saturation/bandwidth of the current connection
-		//@note - if the server is currently CPU saturated then do not sort until framerate improves
-		check( World == Connection->ViewTarget->GetWorld() );
-		AGameMode const* const GameMode = World->GetAuthGameMode();
+				// determine whether we should priority sort the list of relevant actors based on the saturation/bandwidth of the current connection
+				//@note - if the server is currently CPU saturated then do not sort until framerate improves
+				check(World == Connection->ViewTarget->GetWorld());
+				AGameMode const* const GameMode = World->GetAuthGameMode();
 		const bool bLowNetBandwidth = !bCPUSaturated && ( Connection->CurrentNetSpeed / float( GameMode->NumPlayers + GameMode->NumBots ) < 500.f );
 
 		for ( FNetworkObjectInfo* ActorInfo : ConsiderList )
-		{
+				{
 			AActor* Actor = ActorInfo->Actor;
 
-			UActorChannel* Channel = Connection->ActorChannels.FindRef( Actor );
+					UActorChannel* Channel = Connection->ActorChannels.FindRef(Actor);
 
 			UNetConnection* PriorityConnection = Connection;
 
 			if ( Actor->bOnlyRelevantToOwner )
-			{
+					{
 				// This actor should be owned by a particular connection, see if that connection is the one passed in
 				bool bHasNullViewTarget = false;
 
 				PriorityConnection = IsActorOwnedByAndRelevantToConnection( Actor, ConnectionViewers, bHasNullViewTarget );
 
 				if ( PriorityConnection == nullptr )
-				{
+						{
 					// Not owned by this connection, if we have a channel, close it, and continue
 					// NOTE - We won't close the channel if any connection has a NULL view target.
 					//	This is to give all connections a chance to own it
 					if ( !bHasNullViewTarget && Channel != NULL && Time - Channel->RelevantTime >= RelevantTimeout )
-					{
+							{
 						Channel->Close();
-					}
+							}
 
 					// This connection doesn't own this actor
-					continue;
-				}
+							continue;
+						}
 			}
 			else if ( CVarSetNetDormancyEnabled.GetValueOnGameThread() != 0 )
-			{
+						{
 				// Skip Actor if dormant
 				if ( IsActorDormant( Actor, Connection ) )
-				{
+							{
 					continue;
-				}
+									}
 
 				// See of actor wants to try and go dormant
 				if ( ShouldActorGoDormant( Actor, ConnectionViewers, Channel, Time, bLowNetBandwidth ) )
-				{
-					// Channel is marked to go dormant now once all properties have been replicated (but is not dormant yet)
-					Channel->StartBecomingDormant();
-				}
-			}
+							{
+								// Channel is marked to go dormant now once all properties have been replicated (but is not dormant yet)
+								Channel->StartBecomingDormant();
+							}
+						}
 
-			// Skip actor if not relevant and theres no channel already.
-			// Historically Relevancy checks were deferred until after prioritization because they were expensive (line traces).
-			// Relevancy is now cheap and we are dealing with larger lists of considered actors, so we want to keep the list of
-			// prioritized actors low.
-			if ( !Channel )
-			{
-				if ( !IsLevelInitializedForActor( Actor, Connection ) )
-				{
-					// If the level this actor belongs to isn't loaded on client, don't bother sending
-					continue;
-				}
+					// Skip actor if not relevant and theres no channel already.
+					// Historically Relevancy checks were deferred until after prioritization because they were expensive (line traces).
+					// Relevancy is now cheap and we are dealing with larger lists of considered actors, so we want to keep the list of
+					// prioritized actors low.
+					if (!Channel)
+					{
+						if ( !IsLevelInitializedForActor(Actor, Connection) )
+						{
+							// If the level this actor belongs to isn't loaded on client, don't bother sending
+							continue;
+						}
 
 				if ( !IsActorRelevantToConnection( Actor, ConnectionViewers ) )
-				{
+						{
 					// If not relevant (and we don't have a channel), skip
-					continue;
-				}
-			}
+							continue;
+						}
+					}
 
 			// Actor is relevant to this connection, add it to the list
 			// NOTE - We use NetTag to make sure SentTemporaries didn't already mark this actor to be skipped
 			if ( Actor->NetTag != NetTag )
-			{
-				UE_LOG( LogNetTraffic, Log, TEXT( "Consider %s alwaysrelevant %d frequency %f " ), *Actor->GetName(), Actor->bAlwaysRelevant, Actor->NetUpdateFrequency );
+					{
+						UE_LOG(LogNetTraffic, Log, TEXT("Consider %s alwaysrelevant %d frequency %f "),*Actor->GetName(), Actor->bAlwaysRelevant, Actor->NetUpdateFrequency);
 
-				Actor->NetTag = NetTag;
+						Actor->NetTag                 = NetTag;
 
 				OutPriorityList[FinalSortedCount]	= FActorPriority( PriorityConnection, Channel, ActorInfo, ConnectionViewers, bLowNetBandwidth );
 				OutPriorityActors[FinalSortedCount] = OutPriorityList + FinalSortedCount;
 
 				FinalSortedCount++;
 
-				if ( DebugRelevantActors )
-				{
-					LastPrioritizedActors.Add( Actor );
+						if (DebugRelevantActors)
+						{
+							LastPrioritizedActors.Add(Actor);
+						}
+					}
 				}
-			}
-		}
 
-		// Add in deleted actors
-		for ( auto It = Connection->DestroyedStartupOrDormantActors.CreateIterator(); It; ++It )
-		{
-			FActorDestructionInfo &DInfo = DestroyedStartupOrDormantActors.FindChecked( *It );
+				// Add in deleted actors
+				for (auto It = Connection->DestroyedStartupOrDormantActors.CreateIterator(); It; ++It)
+				{
+					FActorDestructionInfo &DInfo = DestroyedStartupOrDormantActors.FindChecked(*It);
 			OutPriorityList[FinalSortedCount]	= FActorPriority( Connection, &DInfo, ConnectionViewers );
 			OutPriorityActors[FinalSortedCount] = OutPriorityList + FinalSortedCount;
 			FinalSortedCount++;
-			DeletedCount++;
-		}
+					DeletedCount++;
+				}
 
 		// Sort by priority
 		struct FCompareFActorPriority
-		{
+				{
 			FORCEINLINE bool operator()( const FActorPriority& A, const FActorPriority& B ) const
-			{
+					{
 				return B.Priority < A.Priority;
 			}
 		};
 
 		Sort( OutPriorityActors, FinalSortedCount, FCompareFActorPriority() );
-	}
+					}
 
 	UE_LOG( LogNetTraffic, Log, TEXT( "ServerReplicateActors_PrioritizeActors: Potential %04i ConsiderList %03i FinalSortedCount %03i" ), MaxSortedActors, ConsiderList.Num(), FinalSortedCount );
-	
+
 	// Setup stats
 	SET_DWORD_STAT( STAT_PrioritizedActors, FinalSortedCount );
-	SET_DWORD_STAT( STAT_NumRelevantDeletedActors, DeletedCount );
+				SET_DWORD_STAT(STAT_NumRelevantDeletedActors,DeletedCount);
 
 	return FinalSortedCount;
 }
 
 int32 UNetDriver::ServerReplicateActors_ProcessPrioritizedActors( UNetConnection* Connection, const TArray<FNetViewer>& ConnectionViewers, FActorPriority** PriorityActors, const int32 FinalSortedCount, int32& OutUpdated )
-{
+				{
 	if ( !Connection->IsNetReady( 0 ) )
-	{
+					{
 		// Connection saturated, don't process any actors
 		return 0;
-	}
+					}
 
 	int32 ActorUpdatesThisConnection		= 0;
 	int32 ActorUpdatesThisConnectionSent	= 0;
 	int32 FinalRelevantCount				= 0;
 
 	for ( int32 j = 0; j < FinalSortedCount; j++ )
-	{
-		// Deletion entry
-		if ( PriorityActors[j]->ActorInfo == NULL && PriorityActors[j]->DestructionInfo )
-		{
-			// Make sure client has streaming level loaded
-			if ( PriorityActors[j]->DestructionInfo->StreamingLevelName != NAME_None && !Connection->ClientVisibleLevelNames.Contains( PriorityActors[j]->DestructionInfo->StreamingLevelName ) )
 			{
-				// This deletion entry is for an actor in a streaming level the connection doesn't have loaded, so skip it
-				continue;
-			}
+					// Deletion entry
+		if ( PriorityActors[j]->ActorInfo == NULL && PriorityActors[j]->DestructionInfo )
+					{
+						// Make sure client has streaming level loaded
+						if (PriorityActors[j]->DestructionInfo->StreamingLevelName != NAME_None && !Connection->ClientVisibleLevelNames.Contains(PriorityActors[j]->DestructionInfo->StreamingLevelName))
+						{
+							// This deletion entry is for an actor in a streaming level the connection doesn't have loaded, so skip it
+							continue;
+						}
 
-			UActorChannel* Channel = ( UActorChannel* )Connection->CreateChannel( CHTYPE_Actor, 1 );
-			if ( Channel )
-			{
-				FinalRelevantCount++;
-				UE_LOG( LogNetTraffic, Log, TEXT( "Server replicate actor creating destroy channel for NetGUID <%s,%s> Priority: %d" ), *PriorityActors[j]->DestructionInfo->NetGUID.ToString(), *PriorityActors[j]->DestructionInfo->PathName, PriorityActors[j]->Priority );
+						UActorChannel* Channel = (UActorChannel*)Connection->CreateChannel( CHTYPE_Actor, 1 );
+						if (Channel)
+						{
+							FinalRelevantCount++;
+							UE_LOG(LogNetTraffic, Log, TEXT("Server replicate actor creating destroy channel for NetGUID <%s,%s> Priority: %d"), *PriorityActors[j]->DestructionInfo->NetGUID.ToString(), *PriorityActors[j]->DestructionInfo->PathName, PriorityActors[j]->Priority );
 
-				Channel->SetChannelActorForDestroy( PriorityActors[j]->DestructionInfo ); // Send a close bunch on the new channel
-				Connection->DestroyedStartupOrDormantActors.Remove( PriorityActors[j]->DestructionInfo->NetGUID ); // Remove from connections to-be-destroyed list (close bunch of reliable, so it will make it there)
-			}
-			continue;
-		}
+							Channel->SetChannelActorForDestroy( PriorityActors[j]->DestructionInfo ); // Send a close bunch on the new channel
+							Connection->DestroyedStartupOrDormantActors.Remove( PriorityActors[j]->DestructionInfo->NetGUID ); // Remove from connections to-be-destroyed list (close bunch of reliable, so it will make it there)
+						}
+						continue;
+					}
 
 #if !(UE_BUILD_SHIPPING || UE_BUILD_TEST)
-		static IConsoleVariable* DebugObjectCvar = IConsoleManager::Get().FindConsoleVariable( TEXT( "net.PackageMap.DebugObject" ) );
-		static IConsoleVariable* DebugAllObjectsCvar = IConsoleManager::Get().FindConsoleVariable( TEXT( "net.PackageMap.DebugAll" ) );
+					static IConsoleVariable* DebugObjectCvar = IConsoleManager::Get().FindConsoleVariable(TEXT("net.PackageMap.DebugObject"));
+					static IConsoleVariable* DebugAllObjectsCvar = IConsoleManager::Get().FindConsoleVariable(TEXT("net.PackageMap.DebugAll"));
 		if ( PriorityActors[j]->ActorInfo &&
 				( ( DebugObjectCvar && !DebugObjectCvar->GetString().IsEmpty() && PriorityActors[j]->ActorInfo->Actor->GetName().Contains( DebugObjectCvar->GetString() ) ) ||
-				( DebugAllObjectsCvar && DebugAllObjectsCvar->GetInt() != 0 ) ) )
-		{
+						(DebugAllObjectsCvar && DebugAllObjectsCvar->GetInt() != 0)))
+					{
 			UE_LOG( LogNetPackageMap, Log, TEXT( "Evaluating actor for replication %s" ), *PriorityActors[j]->ActorInfo->Actor->GetName() );
-		}
+					}
 #endif
 
-		// Normal actor replication
-		UActorChannel* Channel = PriorityActors[j]->Channel;
+					// Normal actor replication
+					UActorChannel* Channel     = PriorityActors[j]->Channel;
 		UE_LOG( LogNetTraffic, Log, TEXT( " Maybe Replicate %s" ), *PriorityActors[j]->ActorInfo->Actor->GetName() );
-		if ( !Channel || Channel->Actor ) //make sure didn't just close this channel
-		{
+					if ( !Channel || Channel->Actor ) //make sure didn't just close this channel
+					{ 
 			AActor*		Actor       = PriorityActors[j]->ActorInfo->Actor;
-			bool		bIsRelevant = false;
+						bool		bIsRelevant = false;
 
-			const bool bLevelInitializedForActor = IsLevelInitializedForActor( Actor, Connection );
+						const bool bLevelInitializedForActor = IsLevelInitializedForActor(Actor, Connection);
 
-			// only check visibility on already visible actors every 1.0 + 0.5R seconds
-			// bTearOff actors should never be checked
-			if ( bLevelInitializedForActor )
-			{
-				if ( !Actor->bTearOff && ( !Channel || Time - Channel->RelevantTime > 1.f ) )
-				{
+						// only check visibility on already visible actors every 1.0 + 0.5R seconds
+						// bTearOff actors should never be checked
+						if ( bLevelInitializedForActor )
+						{
+							if (!Actor->bTearOff && (!Channel || Time - Channel->RelevantTime > 1.f))
+							{
 					if ( IsActorRelevantToConnection( Actor, ConnectionViewers ) )
-					{
-						bIsRelevant = true;
-					}
+									{
+										bIsRelevant = true;
+									}
 					else if ( DebugRelevantActors )
-					{
-						LastNonRelevantActors.Add( Actor );
-					}
-				}
-			}
-			else
-			{
-				// Actor is no longer relevant because the world it is/was in is not loaded by client
-				// exception: player controllers should never show up here
-				UE_LOG( LogNetTraffic, Log, TEXT( "- Level not initialized for actor %s" ), *Actor->GetName() );
-			}
-
-			// if the actor is now relevant or was recently relevant
-			const bool bIsRecentlyRelevant = bIsRelevant || ( Channel && Time - Channel->RelevantTime < RelevantTimeout );
-
-			if ( bIsRecentlyRelevant )
-			{
-				FinalRelevantCount++;
-
-				// Find or create the channel for this actor.
-				// we can't create the channel if the client is in a different world than we are
-				// or the package map doesn't support the actor's class/archetype (or the actor itself in the case of serializable actors)
-				// or it's an editor placed actor and the client hasn't initialized the level it's in
-				if ( Channel == NULL && GuidCache->SupportsObject( Actor->GetClass() ) && GuidCache->SupportsObject( Actor->IsNetStartupActor() ? Actor : Actor->GetArchetype() ) )
-				{
-					if ( bLevelInitializedForActor )
-					{
-						// Create a new channel for this actor.
-						Channel = ( UActorChannel* )Connection->CreateChannel( CHTYPE_Actor, 1 );
-						if ( Channel )
-						{
-							Channel->SetChannelActor( Actor );
-						}
-					}
-					// if we couldn't replicate it for a reason that should be temporary, and this Actor is updated very infrequently, make sure we update it again soon
-					else if ( Actor->NetUpdateFrequency < 1.0f )
-					{
-						UE_LOG( LogNetTraffic, Log, TEXT( "Unable to replicate %s" ), *Actor->GetName() );
-						Actor->NetUpdateTime = Actor->GetWorld()->TimeSeconds + 0.2f * FMath::FRand();
-					}
-				}
-
-				if ( Channel )
-				{
-					// if it is relevant then mark the channel as relevant for a short amount of time
-					if ( bIsRelevant )
-					{
-						Channel->RelevantTime = Time + 0.5f * FMath::SRand();
-					}
-					// if the channel isn't saturated
-					if ( Channel->IsNetReady( 0 ) )
-					{
-						// replicate the actor
-						UE_LOG( LogNetTraffic, Log, TEXT( "- Replicate %s. %d" ), *Actor->GetName(), PriorityActors[j]->Priority );
-						if ( DebugRelevantActors )
+									{
+											LastNonRelevantActors.Add(Actor);
+										}
+									}
+								}
+						else
 						{
-							LastRelevantActors.Add( Actor );
+							// Actor is no longer relevant because the world it is/was in is not loaded by client
+							// exception: player controllers should never show up here
+							UE_LOG(LogNetTraffic, Log, TEXT("- Level not initialized for actor %s"), *Actor->GetName());
 						}
-
-						if ( Channel->ReplicateActor() )
-						{
-							ActorUpdatesThisConnectionSent++;
-							if ( DebugRelevantActors )
+						
+						// if the actor is now relevant or was recently relevant
+						const bool bIsRecentlyRelevant = bIsRelevant || (Channel && Time - Channel->RelevantTime < RelevantTimeout);
+
+						if( bIsRecentlyRelevant )
+						{	
+							FinalRelevantCount++;
+
+							// Find or create the channel for this actor.
+							// we can't create the channel if the client is in a different world than we are
+							// or the package map doesn't support the actor's class/archetype (or the actor itself in the case of serializable actors)
+							// or it's an editor placed actor and the client hasn't initialized the level it's in
+				if ( Channel == NULL && GuidCache->SupportsObject( Actor->GetClass() ) && GuidCache->SupportsObject( Actor->IsNetStartupActor() ? Actor : Actor->GetArchetype() ) )
 							{
-								LastSentActors.Add( Actor );
+								if (bLevelInitializedForActor)
+								{
+									// Create a new channel for this actor.
+									Channel = (UActorChannel*)Connection->CreateChannel( CHTYPE_Actor, 1 );
+									if( Channel )
+									{
+										Channel->SetChannelActor( Actor );
+									}
+								}
+								// if we couldn't replicate it for a reason that should be temporary, and this Actor is updated very infrequently, make sure we update it again soon
+								else if (Actor->NetUpdateFrequency < 1.0f)
+								{
+									UE_LOG(LogNetTraffic, Log, TEXT("Unable to replicate %s"),*Actor->GetName());
+									Actor->NetUpdateTime = Actor->GetWorld()->TimeSeconds + 0.2f * FMath::FRand();
+								}
 							}
 
+							if( Channel )
+							{
+								// if it is relevant then mark the channel as relevant for a short amount of time
+								if( bIsRelevant )
+								{
+									Channel->RelevantTime = Time + 0.5f * FMath::SRand();
+								}
+								// if the channel isn't saturated
+								if( Channel->IsNetReady(0) )
+								{
+									// replicate the actor
+									UE_LOG(LogNetTraffic, Log, TEXT("- Replicate %s. %d"),*Actor->GetName(), PriorityActors[j]->Priority);
+									if (DebugRelevantActors)
+									{
+										LastRelevantActors.Add( Actor );
+									}
+
+									if (Channel->ReplicateActor())
+									{
+										ActorUpdatesThisConnectionSent++;
+										if (DebugRelevantActors)
+										{
+											LastSentActors.Add( Actor );
+										}
+
 							// Calculate min delta (max rate actor will upate), and max delta (slowest rate actor will update)
 							const float MinOptimalDelta				= 1.0f / Actor->NetUpdateFrequency;
 							const float MaxOptimalDelta				= FMath::Max( 1.0f / Actor->MinNetUpdateFrequency, MinOptimalDelta );
@@ -2795,50 +2795,50 @@ int32 UNetDriver::ServerReplicateActors_ProcessPrioritizedActors( UNetConnection
 							// Choose an optimal time, we choose 70% of the actual rate to allow frequency to go up if needed
 							PriorityActors[j]->ActorInfo->OptimalNetUpdateDelta = FMath::Clamp( DeltaBetweenReplications * 0.7f, MinOptimalDelta, MaxOptimalDelta );
 							PriorityActors[j]->ActorInfo->LastNetReplicateTime = World->TimeSeconds;
-						}
-						ActorUpdatesThisConnection++;
+									}
+									ActorUpdatesThisConnection++;
 						OutUpdated++;
-					}
-					else
-					{
-						UE_LOG( LogNetTraffic, Log, TEXT( "- Channel saturated, forcing pending update for %s" ), *Actor->GetName() );
-						// otherwise force this actor to be considered in the next tick again
-						Actor->ForceNetUpdate();
-					}
-					// second check for channel saturation
-					if ( !Connection->IsNetReady( 0 ) )
-					{
+								}
+								else
+								{							
+									UE_LOG(LogNetTraffic, Log, TEXT("- Channel saturated, forcing pending update for %s"),*Actor->GetName());
+									// otherwise force this actor to be considered in the next tick again
+									Actor->ForceNetUpdate();
+								}
+								// second check for channel saturation
+								if (!Connection->IsNetReady(0))
+								{
 						// We can bail out now since this connection is saturated, we'll return how far we got though
 						SET_DWORD_STAT( STAT_NumReplicatedActorAttempts, ActorUpdatesThisConnection );
 						SET_DWORD_STAT( STAT_NumReplicatedActors, ActorUpdatesThisConnectionSent );
 						SET_DWORD_STAT( STAT_NumRelevantActors, FinalRelevantCount );
 						return j;
-					}
-				}
-			}
-
+								}
+							}
+						}
+						
 			// If the actor wasn't recently relevant, or if it was torn off, close the actor channel if it exists for this connection
-			if ( ( !bIsRecentlyRelevant || Actor->bTearOff ) && Channel != NULL )
-			{
-				// Non startup (map) actors have their channels closed immediately, which destroys them.
-				// Startup actors get to keep their channels open.
+						if ((!bIsRecentlyRelevant || Actor->bTearOff) && Channel != NULL)
+						{
+							// Non startup (map) actors have their channels closed immediately, which destroys them.
+							// Startup actors get to keep their channels open.
 
-				// Fixme: this should be a setting
-				if ( !bLevelInitializedForActor || !Actor->IsNetStartupActor() )
-				{
-					UE_LOG( LogNetTraffic, Log, TEXT( "- Closing channel for no longer relevant actor %s" ), *Actor->GetName() );
-					Channel->Close();
+							// Fixme: this should be a setting
+							if ( !bLevelInitializedForActor || !Actor->IsNetStartupActor() )
+							{
+								UE_LOG(LogNetTraffic, Log, TEXT("- Closing channel for no longer relevant actor %s"),*Actor->GetName());
+								Channel->Close();
+							}
+						}
+					}
 				}
-			}
-		}
-	}
 
 	SET_DWORD_STAT( STAT_NumReplicatedActorAttempts, ActorUpdatesThisConnection );
 	SET_DWORD_STAT( STAT_NumReplicatedActors, ActorUpdatesThisConnectionSent );
-	SET_DWORD_STAT( STAT_NumRelevantActors, FinalRelevantCount );
+				SET_DWORD_STAT(STAT_NumRelevantActors,FinalRelevantCount);
 
 	return FinalSortedCount;
-}
+			}
 #endif
 
 int32 UNetDriver::ServerReplicateActors(float DeltaSeconds)
@@ -2981,14 +2981,14 @@ int32 UNetDriver::ServerReplicateActors(float DeltaSeconds)
 				else if ( IsActorRelevantToConnection( Actor, ConnectionViewers ) )
 				{
 					// If this actor was relevant but didn't get processed, force another update for next frame
-					UE_LOG( LogNetTraffic, Log, TEXT( " Saturated. Mark %s NetUpdateTime to be checked for next tick" ), *Actor->GetName() );
-					Actor->bPendingNetUpdate = true;
-					if ( Channel != NULL )
-					{
-						Channel->RelevantTime = Time + 0.5f * FMath::SRand();
+							UE_LOG(LogNetTraffic, Log, TEXT(" Saturated. Mark %s NetUpdateTime to be checked for next tick"), *Actor->GetName());
+							Actor->bPendingNetUpdate = true;
+							if (Channel != NULL)
+							{
+								Channel->RelevantTime = Time + 0.5f * FMath::SRand();
+							}
 					}
 				}
-			}
 			RelevantActorMark.Pop();
 		}
 	}
diff --git a/Engine/Source/Runtime/Engine/Private/PhysicsEngine/BodyInstance.cpp b/Engine/Source/Runtime/Engine/Private/PhysicsEngine/BodyInstance.cpp
--- a/Engine/Source/Runtime/Engine/Private/PhysicsEngine/BodyInstance.cpp
+++ b/Engine/Source/Runtime/Engine/Private/PhysicsEngine/BodyInstance.cpp
@@ -407,12 +407,11 @@ void FBodyInstance::UpdatePhysicalMaterials()
 {
 	SCOPE_CYCLE_COUNTER(STAT_UpdatePhysMats);
 	UPhysicalMaterial* SimplePhysMat = GetSimplePhysicalMaterial();
-	check(SimplePhysMat != NULL);
 	TArray<UPhysicalMaterial*> ComplexPhysMats = GetComplexPhysicalMaterials();
 
 #if WITH_PHYSX
-	PxMaterial* PSimpleMat = SimplePhysMat->GetPhysXMaterial();
-	check(PSimpleMat != NULL);
+	PxMaterial* PSimpleMat = SimplePhysMat ? SimplePhysMat->GetPhysXMaterial() : nullptr;
+
 	ExecuteOnPhysicsReadWrite([&]()
 	{
 		ApplyMaterialToInstanceShapes_AssumesLocked(PSimpleMat, ComplexPhysMats);
@@ -422,12 +421,19 @@ void FBodyInstance::UpdatePhysicalMaterials()
 #if WITH_BOX2D
 	if (BodyInstancePtr)
 	{
-		for (b2Fixture* Fixture = BodyInstancePtr->GetFixtureList(); Fixture; Fixture = Fixture->GetNext())
+		if(SimplePhysMat)
 		{
-			Fixture->SetFriction(SimplePhysMat->Friction);
-			Fixture->SetRestitution(SimplePhysMat->Restitution);
+			for(b2Fixture* Fixture = BodyInstancePtr->GetFixtureList(); Fixture; Fixture = Fixture->GetNext())
+			{
+				Fixture->SetFriction(SimplePhysMat->Friction);
+				Fixture->SetRestitution(SimplePhysMat->Restitution);
 
-			//@TODO: BOX2D: Determine if it's feasible to add support for FrictionCombineMode to Box2D
+				//@TODO: BOX2D: Determine if it's feasible to add support for FrictionCombineMode to Box2D
+			}
+		}
+		else
+		{
+			UE_LOG(LogPhysics, Error, TEXT("FBodyInstance::UpdatePhysicalMaterials : No valid physics material found when setting Box2D physical material parameters."));
 		}
 	}
 #endif
@@ -1218,7 +1224,8 @@ struct FInitBodiesHelper
 	{
 		UPhysicalMaterial* SimplePhysMat = Instance->GetSimplePhysicalMaterial();
 		TArray<UPhysicalMaterial*> ComplexPhysMats = Instance->GetComplexPhysicalMaterials();
-		PxMaterial* PSimpleMat = SimplePhysMat->GetPhysXMaterial();
+
+		PxMaterial* PSimpleMat = SimplePhysMat ? SimplePhysMat->GetPhysXMaterial() : nullptr;
 
 		FShapeData ShapeData;
 		Instance->GetFilterData_AssumesLocked(ShapeData);
@@ -2023,9 +2030,9 @@ bool FBodyInstance::Weld(FBodyInstance* TheirBody, const FTransform& TheirTM)
 
 		TheirBody->WeldParent = this;
 
-		UPhysicalMaterial* SimplePhysMat = GetSimplePhysicalMaterial();
-		TArray<UPhysicalMaterial*> ComplexPhysMats = GetComplexPhysicalMaterials();
-		PxMaterial* PSimpleMat = SimplePhysMat->GetPhysXMaterial();
+		UPhysicalMaterial* SimplePhysMat = TheirBody->GetSimplePhysicalMaterial();
+		TArray<UPhysicalMaterial*> ComplexPhysMats = TheirBody->GetComplexPhysicalMaterials();
+		PxMaterial* PSimpleMat =  SimplePhysMat ? SimplePhysMat->GetPhysXMaterial() : nullptr;
 
 		FShapeData ShapeData;
 		GetFilterData_AssumesLocked(ShapeData);
@@ -2148,9 +2155,6 @@ void FBodyInstance::UnWeld(FBodyInstance* TheirBI)
 
 void FBodyInstance::PostShapeChange()
 {
-	// Apply correct physical materials to shape we created.
-	UpdatePhysicalMaterials();
-
 	// Set the filter data on the shapes (call this after setting BodyData because it uses that pointer)
 	UpdatePhysicsFilterData();
 
@@ -3180,11 +3184,16 @@ UPhysicalMaterial* FBodyInstance::GetSimplePhysicalMaterial() const
 
 UPhysicalMaterial* FBodyInstance::GetSimplePhysicalMaterial(const FBodyInstance* BodyInstance, TWeakObjectPtr<UPrimitiveComponent> OwnerComp, TWeakObjectPtr<UBodySetup> BodySetupPtr)
 {
-	check(GEngine->DefaultPhysMaterial != NULL);
+	if(!GEngine || !GEngine->DefaultPhysMaterial)
+	{
+		UE_LOG(LogPhysics, Error, TEXT("FBodyInstance::GetSimplePhysicalMaterial : GEngine not initialized! Cannot call this during native CDO construction, wrap with if(!HasAnyFlags(RF_ClassDefaultObject)) or move out of constructor, material parameters will not be correct."));
+
+		return nullptr;
+	}
 
 	// Find the PhysicalMaterial we need to apply to the physics bodies.
 	// (LOW priority) Engine Mat, Material PhysMat, BodySetup Mat, Component Override, Body Override (HIGH priority)
-
+	
 	UPhysicalMaterial* ReturnPhysMaterial = NULL;
 
 	// BodyInstance override
@@ -3326,7 +3335,6 @@ float gPerCm3ToKgPerCm3(float gPerCm3)
 void FBodyInstance::UpdateMassProperties()
 {
 	UPhysicalMaterial* PhysMat = GetSimplePhysicalMaterial();
-	check(PhysMat);
 
 #if WITH_PHYSX
 	ExecuteOnPxRigidBodyReadWrite(this, [&] (PxRigidBody* PRigidBody)
@@ -3334,7 +3342,14 @@ void FBodyInstance::UpdateMassProperties()
 		if(GetNumSimShapes_AssumesLocked(PRigidBody) > 0)
 	{
 		// physical material - nothing can weigh less than hydrogen (0.09 kg/m^3)
-		float DensityKGPerCubicUU = FMath::Max(KgPerM3ToKgPerCm3(0.09f), gPerCm3ToKgPerCm3(PhysMat->Density));
+		float DensityKGPerCubicUU = 1.0f;
+		float RaiseMassToPower = 0.75f;
+		if(PhysMat)
+		{
+			DensityKGPerCubicUU = FMath::Max(KgPerM3ToKgPerCm3(0.09f), gPerCm3ToKgPerCm3(PhysMat->Density));
+			RaiseMassToPower = PhysMat->RaiseMassToPower;
+		}
+
 		PxRigidBodyExt::updateMassAndInertia(*PRigidBody, DensityKGPerCubicUU);
 
 		//grab OldMass so we can apply new mass while maintaining inertia tensor
@@ -3343,7 +3358,7 @@ void FBodyInstance::UpdateMassProperties()
 
 		if (bOverrideMass == false)
 		{
-			float UsePow = FMath::Clamp<float>(PhysMat->RaiseMassToPower, KINDA_SMALL_NUMBER, 1.f);
+			float UsePow = FMath::Clamp<float>(RaiseMassToPower, KINDA_SMALL_NUMBER, 1.f);
 			NewMass = FMath::Pow(OldMass, UsePow);
 
 			// Apply user-defined mass scaling.
@@ -3384,7 +3399,13 @@ void FBodyInstance::UpdateMassProperties()
 		{
 			// Unreal material density is in g/cm^3, and Box2D density is in kg/m^2
 			// physical material - nothing can weigh less than hydrogen (0.09 kg/m^3)
-			float DensityKGPerCubicCM = FMath::Max(KgPerM3ToKgPerCm3(0.09f), gPerCm3ToKgPerCm3(PhysMat->Density));
+
+			float DensityKGPerCubicCM = 1.0f;
+			if(PhysMat)
+			{
+				DensityKGPerCubicCM = FMath::Max(KgPerM3ToKgPerCm3(0.09f), gPerCm3ToKgPerCm3(PhysMat->Density));
+			}
+
 			const float DensityKGPerCubicM = DensityKGPerCubicCM * 1000.0f;
 			const float DensityKGPerSquareM = DensityKGPerCubicM * 0.1f; //@TODO: BOX2D: Should there be a thickness property for mass calculations?
 			MassScaledDensity = DensityKGPerSquareM * FMath::Clamp<float>(MassScale, 0.01f, 100.0f);
@@ -4626,16 +4647,27 @@ void FBodyInstance::ApplyMaterialToShape_AssumesLocked(PxShape* PShape, PxMateri
 		}
 		else
 		{
-			UE_LOG(LogPhysics, Verbose, TEXT("FBodyInstance::ApplyMaterialToShape_AssumesLocked : PComplexMats is empty - falling back on simple physical material."));
-			PShape->setMaterials(&PSimpleMat, 1);
+			if(PSimpleMat)
+			{
+				UE_LOG(LogPhysics, Verbose, TEXT("FBodyInstance::ApplyMaterialToShape_AssumesLocked : PComplexMats is empty - falling back on simple physical material."));
+				PShape->setMaterials(&PSimpleMat, 1);
+			}
+			else
+			{
+				UE_LOG(LogPhysics, Error, TEXT("FBodyInstance::ApplyMaterialToShape_AssumesLocked : PComplexMats is empty, and we do not have a valid simple material."));
+			}
 		}
 
 	}
 	// Simple shape, 
-	else
+	else if(PSimpleMat)
 	{
 		PShape->setMaterials(&PSimpleMat, 1);
 	}
+	else
+	{
+		UE_LOG(LogPhysics, Error, TEXT("FBodyInstance::ApplyMaterialToShape_AssumesLocked : No valid simple physics material found."));
+	}
 }
 
 void FBodyInstance::ApplyMaterialToInstanceShapes_AssumesLocked(PxMaterial* PSimpleMat, TArray<UPhysicalMaterial*>& ComplexPhysMats)
diff --git a/Engine/Source/Runtime/Engine/Private/PhysicsEngine/BodySetup.cpp b/Engine/Source/Runtime/Engine/Private/PhysicsEngine/BodySetup.cpp
--- a/Engine/Source/Runtime/Engine/Private/PhysicsEngine/BodySetup.cpp
+++ b/Engine/Source/Runtime/Engine/Private/PhysicsEngine/BodySetup.cpp
@@ -328,7 +328,7 @@ struct FAddShapesHelper
 	, NewShapes(InNewShapes)
 	, bShapeSharing(InShapeSharing)
 	{
-		SetupNonUniformHelper(Scale3D, MinScale, MinScaleAbs, Scale3DAbs);
+		SetupNonUniformHelper(Scale3D, MinScale, MinScaleAbs, ShapeScale3DAbs);
 		{
 			float MinScaleRelative;
 			float MinScaleAbsRelative;
@@ -338,9 +338,14 @@ struct FAddShapesHelper
 			SetupNonUniformHelper(Scale3DRelative, MinScaleRelative, MinScaleAbsRelative, Scale3DAbsRelative);
 
 			MinScaleAbs *= MinScaleAbsRelative;
-			Scale3DAbs.X *= Scale3DAbsRelative.X;
-			Scale3DAbs.Y *= Scale3DAbsRelative.Y;
-			Scale3DAbs.Z *= Scale3DAbsRelative.Z;
+			ShapeScale3DAbs.X *= Scale3DAbsRelative.X;
+			ShapeScale3DAbs.Y *= Scale3DAbsRelative.Y;
+			ShapeScale3DAbs.Z *= Scale3DAbsRelative.Z;
+
+			ShapeScale3D = Scale3D;
+			ShapeScale3D.X *= Scale3DAbsRelative.X;
+			ShapeScale3D.Y *= Scale3DAbsRelative.Y;
+			ShapeScale3D.Z *= Scale3DAbsRelative.Z;
 		}
 
 		GetContactOffsetParams(ContactOffsetFactor, MaxContactOffset);
@@ -360,7 +365,8 @@ struct FAddShapesHelper
 
 	float MinScaleAbs;
 	float MinScale;
-	FVector Scale3DAbs;
+	FVector ShapeScale3DAbs;
+	FVector ShapeScale3D;
 
 	float ContactOffsetFactor;
 	float MaxContactOffset;
@@ -423,8 +429,8 @@ public:
 
 	void AddSphylsToRigidActor_AssumesLocked() const
 	{
-		float ScaleRadius = FMath::Max(Scale3DAbs.X, Scale3DAbs.Y);
-		float ScaleLength = Scale3DAbs.Z;
+		float ScaleRadius = FMath::Max(ShapeScale3DAbs.X, ShapeScale3DAbs.Y);
+		float ScaleLength = ShapeScale3DAbs.Z;
 
 		for (int32 i = 0; i < BodySetup->AggGeom.SphylElems.Num(); i++)
 		{
@@ -466,7 +472,7 @@ public:
 
 				PxConvexMeshGeometry PConvexGeom;
 				PConvexGeom.convexMesh = bUseNegX ? ConvexElem.ConvexMeshNegX : ConvexElem.ConvexMesh;
-				PConvexGeom.scale.scale = U2PVector(Scale3DAbs * ConvexElem.GetTransform().GetScale3D().GetAbs());
+				PConvexGeom.scale.scale = U2PVector(ShapeScale3DAbs * ConvexElem.GetTransform().GetScale3D().GetAbs());	//scale shape about the origin
 				FTransform ConvexTransform = ConvexElem.GetTransform();
 				if (ConvexTransform.GetScale3D().X < 0 || ConvexTransform.GetScale3D().Y < 0 || ConvexTransform.GetScale3D().Z < 0)
 				{
@@ -474,6 +480,7 @@ public:
 				}
 				if (ConvexTransform.IsValid())
 				{
+					//Scale the position independent of shape scale. This is because physx transforms have no concept of scale
 					PxTransform PElementTransform = U2PTransform(ConvexTransform * RelativeTM);
 					PLocalPose.q *= PElementTransform.q;
 					PLocalPose.p = PElementTransform.p;
@@ -515,7 +522,7 @@ public:
 		
 			PxTriangleMeshGeometry PTriMeshGeom;
 			PTriMeshGeom.triangleMesh = TriMesh;
-			PTriMeshGeom.scale.scale = U2PVector(Scale3D);
+			PTriMeshGeom.scale.scale = U2PVector(ShapeScale3D); //scale shape about the origin
 			if (BodySetup->bDoubleSidedGeometry)
 			{
 				PTriMeshGeom.meshFlags |= PxMeshGeometryFlag::eDOUBLE_SIDED;
@@ -523,6 +530,7 @@ public:
 
 			if (PTriMeshGeom.isValid())
 			{
+				//Scale the position independent of shape scale. This is because physx transforms have no concept of scale
 				PxTransform PElementTransform = U2PTransform(RelativeTM);
 				PElementTransform.p.x *= Scale3D.X;
 				PElementTransform.p.y *= Scale3D.Y;
diff --git a/Engine/Source/Runtime/Engine/Private/PhysicsEngine/PhysAnim.cpp b/Engine/Source/Runtime/Engine/Private/PhysicsEngine/PhysAnim.cpp
--- a/Engine/Source/Runtime/Engine/Private/PhysicsEngine/PhysAnim.cpp
+++ b/Engine/Source/Runtime/Engine/Private/PhysicsEngine/PhysAnim.cpp
@@ -398,7 +398,6 @@ void USkeletalMeshComponent::BlendInPhysics(FTickFunction& ThisTickFunction)
 			{
 				// Initialize Parallel Task arrays
 				AnimEvaluationContext.SpaceBases = GetSpaceBases();
-				AnimEvaluationContext.VertexAnims = ActiveVertexAnims;
 			}
 
 			AnimEvaluationContext.LocalAtoms.Reset(LocalAtoms.Num());
diff --git a/Engine/Source/Runtime/Engine/Private/PlayerController.cpp b/Engine/Source/Runtime/Engine/Private/PlayerController.cpp
--- a/Engine/Source/Runtime/Engine/Private/PlayerController.cpp
+++ b/Engine/Source/Runtime/Engine/Private/PlayerController.cpp
@@ -3583,6 +3583,11 @@ void APlayerController::SetHapticsByValue(const float Frequency, const float Amp
 		return;
 	}
 
+	if (Player == nullptr)
+	{
+		return;
+	}
+
 	IInputInterface* InputInterface = FSlateApplication::Get().GetInputInterface();
 	if (InputInterface)
 	{
@@ -3593,6 +3598,21 @@ void APlayerController::SetHapticsByValue(const float Frequency, const float Amp
 	}
 }
 
+void APlayerController::SetControllerLightColor(FColor Color)
+{
+	if (Player == nullptr)
+	{
+		return;
+	}
+
+	IInputInterface* InputInterface = FSlateApplication::Get().GetInputInterface();
+	if (InputInterface)
+	{
+		const int32 ControllerId = CastChecked<ULocalPlayer>(Player)->GetControllerId();
+		InputInterface->SetLightColor(ControllerId, Color);
+	}
+}
+
 void APlayerController::ProcessForceFeedbackAndHaptics(const float DeltaTime, const bool bGamePaused)
 {
 	if (Player == NULL)
@@ -4572,7 +4592,7 @@ void FInputModeGameOnly::ApplyInputMode(FReply& SlateOperations, class UGameView
 		SlateOperations.SetUserFocus(ViewportWidgetRef);
 		SlateOperations.LockMouseToWidget(ViewportWidgetRef);
 		GameViewportClient.SetIgnoreInput(false);
-		GameViewportClient.SetCaptureMouseOnClick(EMouseCaptureMode::CapturePermanently);
+		GameViewportClient.SetCaptureMouseOnClick(bConsumeCaptureMouseDown ? EMouseCaptureMode::CapturePermanently : EMouseCaptureMode::CapturePermanently_IncludingInitialMouseDown);
 	}
 }
 
diff --git a/Engine/Source/Runtime/Engine/Private/RepLayout.cpp b/Engine/Source/Runtime/Engine/Private/RepLayout.cpp
--- a/Engine/Source/Runtime/Engine/Private/RepLayout.cpp
+++ b/Engine/Source/Runtime/Engine/Private/RepLayout.cpp
@@ -1558,7 +1558,14 @@ void FRepLayout::CallRepNotifies( FRepState * RepState, UObject* Object ) const
 	{
 		UProperty * RepProperty = RepState->RepNotifies[i];
 
-		UFunction * RepNotifyFunc = Object->FindFunctionChecked( RepProperty->RepNotifyFunc );
+		UFunction * RepNotifyFunc = Object->FindFunction( RepProperty->RepNotifyFunc );
+
+		if (RepNotifyFunc == nullptr)
+		{
+			UE_LOG(LogRep, Warning, TEXT("FRepLayout::CallRepNotifies: Can't find RepNotify function %s for property %s on object %s."),
+				*RepProperty->RepNotifyFunc.ToString(), *RepProperty->GetName(), *Object->GetName());
+			continue;
+		}
 
 		check( RepNotifyFunc->NumParms <= 1 );	// 2 parms not supported yet
 
