Put the DestroyChildActor call back in to OnComponentDestroyed as there are conditions where ChildActor can exist (serialized in) even though the Component never gets registered (thus not getting unregistered which destroys the child actor).

diff --git a/Engine/Source/Runtime/Engine/Private/Components/ChildActorComponent.cpp b/Engine/Source/Runtime/Engine/Private/Components/ChildActorComponent.cpp
--- a/Engine/Source/Runtime/Engine/Private/Components/ChildActorComponent.cpp
+++ b/Engine/Source/Runtime/Engine/Private/Components/ChildActorComponent.cpp
@@ -94,7 +94,8 @@ void UChildActorComponent::OnComponentDestroyed(bool bDestroyingHierarchy)
 {
 	Super::OnComponentDestroyed(bDestroyingHierarchy);
 
-	check(ChildActor == nullptr);
+	const UWorld* const MyWorld = GetWorld();
+	DestroyChildActor(MyWorld && !MyWorld->IsGameWorld());
 }
 
 void UChildActorComponent::OnUnregister()
