Properly handle shape body setup adding to root

diff --git a/Engine/Source/Runtime/Engine/Private/Components/ShapeComponent.cpp b/Engine/Source/Runtime/Engine/Private/Components/ShapeComponent.cpp
--- a/Engine/Source/Runtime/Engine/Private/Components/ShapeComponent.cpp
+++ b/Engine/Source/Runtime/Engine/Private/Components/ShapeComponent.cpp
@@ -107,7 +107,7 @@ void UShapeComponent::CreateShapeBodySetupIfNeeded()
 	if (ShapeBodySetup == nullptr || ShapeBodySetup->IsPendingKill())
 	{
 		ShapeBodySetup = NewObject<UBodySetup>(this);
-		if (IsTemplate() && IsRooted())
+		if (GUObjectArray.IsDisregardForGC(this))
 		{
 			ShapeBodySetup->AddToRoot();
 		}
