Redoing merge of CL#2889076 with force option set to remove bIgnoreCaptureClick variable.

diff --git a/Engine/Source/Runtime/Engine/Private/Slate/SceneViewport.cpp b/Engine/Source/Runtime/Engine/Private/Slate/SceneViewport.cpp
--- a/Engine/Source/Runtime/Engine/Private/Slate/SceneViewport.cpp
+++ b/Engine/Source/Runtime/Engine/Private/Slate/SceneViewport.cpp
@@ -25,7 +25,6 @@ FSceneViewport::FSceneViewport( FViewportClient* InViewportClient, TSharedPtr<SV
 	, MouseDelta( 0, 0 )
 	, bIsCursorVisible( true )
 	, bShouldCaptureMouseOnActivate( true )
-	, bIgnoreCaptureClick( true )
 	, bRequiresVsync( false )
 	, bUseSeparateRenderTarget( InViewportWidget.IsValid() ? !InViewportWidget->ShouldRenderDirectly() : true )
 	, bIsResizing( false )
@@ -385,7 +384,7 @@ FReply FSceneViewport::OnMouseButtonDown( const FGeometry& InGeometry, const FPo
 			(ViewportClient->CaptureMouseOnClick() == EMouseCaptureMode::CaptureDuringRightMouseDown && InMouseEvent.GetEffectingButton() == EKeys::RightMouseButton);
 
 		// Process primary input if we aren't currently a game viewport, we already have capture, or we are permanent capture that dosn't consume the mouse down.
-		const bool bProcessInputPrimary = !IsCurrentlyGameViewport() || !bIgnoreCaptureClick || HasMouseCapture() || (ViewportClient->CaptureMouseOnClick() == EMouseCaptureMode::CapturePermanently_IncludingInitialMouseDown);
+		const bool bProcessInputPrimary = !IsCurrentlyGameViewport() || HasMouseCapture() || (ViewportClient->CaptureMouseOnClick() == EMouseCaptureMode::CapturePermanently_IncludingInitialMouseDown);
 
 		const bool bAnyMenuWasVisible = FSlateApplication::Get().AnyMenusVisible();
 
diff --git a/Engine/Source/Runtime/Engine/Public/Slate/SceneViewport.h b/Engine/Source/Runtime/Engine/Public/Slate/SceneViewport.h
--- a/Engine/Source/Runtime/Engine/Public/Slate/SceneViewport.h
+++ b/Engine/Source/Runtime/Engine/Public/Slate/SceneViewport.h
@@ -330,9 +330,6 @@ private:
 	bool bIsCursorVisible;
 	/** true if we had Capture when deactivated */
 	bool bShouldCaptureMouseOnActivate;
-	/** true if we want to ignore input processing on the click that give the viewport capture.
-	  * This only applies to Permanent Capture mode, the temporary mode always process the input*/
-	bool bIgnoreCaptureClick;
 	/** true if this viewport requires vsync. */
 	bool bRequiresVsync;
 	/** true if this viewport renders to a separate render target.  false to render directly to the windows back buffer */
