Merged CL 2912575 from //UE4/Release-4.11

diff --git a/Engine/Source/Runtime/Engine/Classes/Engine/GameViewportClient.h b/Engine/Source/Runtime/Engine/Classes/Engine/GameViewportClient.h
--- a/Engine/Source/Runtime/Engine/Classes/Engine/GameViewportClient.h
+++ b/Engine/Source/Runtime/Engine/Classes/Engine/GameViewportClient.h
@@ -809,9 +809,6 @@ private:
 
 	/** Whether or not this audio device is in audio-focus */
 	bool bHasAudioFocus;
-
-	/** Whether or not this viewport is active*/
-	bool bActive;
 };
 
 
diff --git a/Engine/Source/Runtime/Engine/Private/GameViewportClient.cpp b/Engine/Source/Runtime/Engine/Private/GameViewportClient.cpp
--- a/Engine/Source/Runtime/Engine/Private/GameViewportClient.cpp
+++ b/Engine/Source/Runtime/Engine/Private/GameViewportClient.cpp
@@ -119,7 +119,6 @@ UGameViewportClient::UGameViewportClient(const FObjectInitializer& ObjectInitial
 	, bHideCursorDuringCapture(false)
 	, AudioDeviceHandle(INDEX_NONE)
 	, bHasAudioFocus(false)
-	, bActive(false)
 {
 
 	TitleSafeZone.MaxPercentX = 0.9f;
@@ -178,7 +177,6 @@ UGameViewportClient::UGameViewportClient(FVTableHelper& Helper)
 	, bHideCursorDuringCapture(false)
 	, AudioDeviceHandle(INDEX_NONE)
 	, bHasAudioFocus(false)
-	, bActive(false)
 {
 
 }
@@ -597,7 +595,7 @@ bool UGameViewportClient::RequiresUncapturedAxisInput() const
 EMouseCursor::Type UGameViewportClient::GetCursor(FViewport* InViewport, int32 X, int32 Y)
 {
 	// If the viewport isn't active or the console is active we don't want to override the cursor
-	if (!bActive || (!InViewport->HasMouseCapture() && !InViewport->HasFocus()) || (ViewportConsole && ViewportConsole->ConsoleActive()))
+	if (!FSlateApplication::Get().IsActive() || (!InViewport->HasMouseCapture() && !InViewport->HasFocus()) || (ViewportConsole && ViewportConsole->ConsoleActive()))
 	{
 		return EMouseCursor::Default;
 	}
@@ -1393,13 +1391,11 @@ bool UGameViewportClient::IsFocused(FViewport* InViewport)
 void UGameViewportClient::Activated(FViewport* InViewport, const FWindowActivateEvent& InActivateEvent)
 {
 	ReceivedFocus(InViewport);
-	bActive = true;
 }
 
 void UGameViewportClient::Deactivated(FViewport* InViewport, const FWindowActivateEvent& InActivateEvent)
 {
 	LostFocus(InViewport);
-	bActive = false;
 }
 
 void UGameViewportClient::CloseRequested(FViewport* InViewport)
