The PortalServices module now adds the WITH_PORTAL_SERVICES definition. The engine now only attempts to load the appropriate Portal modules if the WITH_PORTAL_SERVICES definition is defined true.

Fixes issues with all non-editor builds producing load module warnings.

diff --git a/Engine/Source/Runtime/Engine/Private/UnrealEngine.cpp b/Engine/Source/Runtime/Engine/Private/UnrealEngine.cpp
--- a/Engine/Source/Runtime/Engine/Private/UnrealEngine.cpp
+++ b/Engine/Source/Runtime/Engine/Private/UnrealEngine.cpp
@@ -1639,39 +1639,45 @@ void UEngine::InitializeObjectReferences()
 #if PLATFORM_COMPILER_HAS_VARIADIC_TEMPLATES
 void UEngine::InitializePortalServices()
 {
-	TSharedPtr<IMessagingRpcModule> MessagingRpcModule = StaticCastSharedPtr<IMessagingRpcModule>(FModuleManager::Get().LoadModule("MessagingRpc"));
-	TSharedPtr<IPortalRpcModule> PortalRpcModule = StaticCastSharedPtr<IPortalRpcModule>(FModuleManager::Get().LoadModule("PortalRpc"));
-	TSharedPtr<IPortalServicesModule> PortalServicesModule = StaticCastSharedPtr<IPortalServicesModule>(FModuleManager::Get().LoadModule("PortalServices"));
+	TSharedPtr<IMessagingRpcModule> MessagingRpcModule;
+	TSharedPtr<IPortalRpcModule> PortalRpcModule;
+	TSharedPtr<IPortalServicesModule> PortalServicesModule;
+
+#if WITH_PORTAL_SERVICES
+	MessagingRpcModule = StaticCastSharedPtr<IMessagingRpcModule>(FModuleManager::Get().LoadModule("MessagingRpc"));
+	PortalRpcModule = StaticCastSharedPtr<IPortalRpcModule>(FModuleManager::Get().LoadModule("PortalRpc"));
+	PortalServicesModule = StaticCastSharedPtr<IPortalServicesModule>(FModuleManager::Get().LoadModule("PortalServices"));
+#endif
 
 	if (MessagingRpcModule.IsValid() &&
 		PortalRpcModule.IsValid() &&
 		PortalServicesModule.IsValid())
 	{
-	// Initialize Portal services
+		// Initialize Portal services
 		PortalRpcClient = MessagingRpcModule->CreateRpcClient();
-	{
-		// @todo gmp: catch timeouts?
-	}
+		{
+			// @todo gmp: catch timeouts?
+		}
 
 		PortalRpcLocator = PortalRpcModule->CreateLocator();
-	{
-		PortalRpcLocator->OnServerLocated().BindLambda([=]() { PortalRpcClient->Connect(PortalRpcLocator->GetServerAddress()); });
-		PortalRpcLocator->OnServerLost().BindLambda([=]() { PortalRpcClient->Disconnect(); });
-	}
+		{
+			PortalRpcLocator->OnServerLocated().BindLambda([=]() { PortalRpcClient->Connect(PortalRpcLocator->GetServerAddress()); });
+			PortalRpcLocator->OnServerLost().BindLambda([=]() { PortalRpcClient->Disconnect(); });
+		}
 
-	ServiceDependencies = MakeShareable(new FTypeContainer);
-	{
-		ServiceDependencies->RegisterInstance<IMessageRpcClient>(PortalRpcClient.ToSharedRef());
-	}
+		ServiceDependencies = MakeShareable(new FTypeContainer);
+		{
+			ServiceDependencies->RegisterInstance<IMessageRpcClient>(PortalRpcClient.ToSharedRef());
+		}
 
 		ServiceLocator = PortalServicesModule->CreateLocator(ServiceDependencies.ToSharedRef());
-	{
-		// @todo add any Engine specific Portal services here
-		ServiceLocator->Configure(TEXT("IPortalApplicationWindow"), TEXT("*"), "PortalProxies");
-		ServiceLocator->Configure(TEXT("IPortalUser"), TEXT("*"), "PortalProxies");
-		ServiceLocator->Configure(TEXT("IPortalUserLogin"), TEXT("*"), "PortalProxies");
+		{
+			// @todo add any Engine specific Portal services here
+			ServiceLocator->Configure(TEXT("IPortalApplicationWindow"), TEXT("*"), "PortalProxies");
+			ServiceLocator->Configure(TEXT("IPortalUser"), TEXT("*"), "PortalProxies");
+			ServiceLocator->Configure(TEXT("IPortalUserLogin"), TEXT("*"), "PortalProxies");
+		}
 	}
-}
 	else
 	{
 		class FNullPortalServiceLocator
diff --git a/Engine/Source/Runtime/Portal/Services/PortalServices.Build.cs b/Engine/Source/Runtime/Portal/Services/PortalServices.Build.cs
--- a/Engine/Source/Runtime/Portal/Services/PortalServices.Build.cs
+++ b/Engine/Source/Runtime/Portal/Services/PortalServices.Build.cs
@@ -12,7 +12,9 @@ namespace UnrealBuildTool.Rules
 					"Core",
 					"CoreUObject",
 				}
-			); 
+			);
+
+            Definitions.Add("WITH_PORTAL_SERVICES=1");
 		}
 	}
 }
