Copying //UE4/Dev-Framework to Dev-Main (//UE4/Dev-Main) @ 2926658

diff --git a/Engine/Source/Runtime/CoreUObject/Private/UObject/LinkerLoad.cpp b/Engine/Source/Runtime/CoreUObject/Private/UObject/LinkerLoad.cpp
--- a/Engine/Source/Runtime/CoreUObject/Private/UObject/LinkerLoad.cpp
+++ b/Engine/Source/Runtime/CoreUObject/Private/UObject/LinkerLoad.cpp
@@ -14,6 +14,7 @@
 #include "GatherableTextData.h"
 #include "Serialization/AsyncLoading.h"
 #include "ModuleManager.h"
+#include "LoadTimeTracker.h"
 #include "HAL/ThreadHeartBeat.h"
 
 #define LOCTEXT_NAMESPACE "LinkerLoad"
@@ -111,17 +112,17 @@ void FLinkerLoad::CreateActiveRedirectsMap(const FString& GEngineIniName)
 
 					bool bInstanceOnly = false;
 
-					FParse::Bool( *It.Value(), TEXT("InstanceOnly="), bInstanceOnly );
-					FParse::Value( *It.Value(), TEXT("ObjectName="), ObjectName );
+					FParse::Bool( *It.Value().GetValue(), TEXT("InstanceOnly="), bInstanceOnly );
+					FParse::Value( *It.Value().GetValue(), TEXT("ObjectName="), ObjectName );
 
-					FParse::Value( *It.Value(), TEXT("OldClassName="), OldClassName );
-					FParse::Value( *It.Value(), TEXT("NewClassName="), NewClassName );
+					FParse::Value( *It.Value().GetValue(), TEXT("OldClassName="), OldClassName );
+					FParse::Value( *It.Value().GetValue(), TEXT("NewClassName="), NewClassName );
 
-					FParse::Value( *It.Value(), TEXT("OldSubobjName="), OldSubobjName );
-					FParse::Value( *It.Value(), TEXT("NewSubobjName="), NewSubobjName );
+					FParse::Value( *It.Value().GetValue(), TEXT("OldSubobjName="), OldSubobjName );
+					FParse::Value( *It.Value().GetValue(), TEXT("NewSubobjName="), NewSubobjName );
 
-					FParse::Value( *It.Value(), TEXT("NewClassClass="), NewClassClass );
-					FParse::Value( *It.Value(), TEXT("NewClassPackage="), NewClassPackage );
+					FParse::Value( *It.Value().GetValue(), TEXT("NewClassClass="), NewClassClass );
+					FParse::Value( *It.Value().GetValue(), TEXT("NewClassPackage="), NewClassPackage );
 
 					if (NewSubobjName != NAME_None || OldSubobjName != NAME_None)
 					{
@@ -182,8 +183,8 @@ void FLinkerLoad::CreateActiveRedirectsMap(const FString& GEngineIniName)
 					FName OldGameName = NAME_None;
 					FName NewGameName = NAME_None;
 
-					FParse::Value( *It.Value(), TEXT("OldGameName="), OldGameName );
-					FParse::Value( *It.Value(), TEXT("NewGameName="), NewGameName );
+					FParse::Value( *It.Value().GetValue(), TEXT("OldGameName="), OldGameName );
+					FParse::Value( *It.Value().GetValue(), TEXT("NewGameName="), NewGameName );
 
 					GameNameRedirects.Add(OldGameName, NewGameName);
 				}
@@ -192,8 +193,8 @@ void FLinkerLoad::CreateActiveRedirectsMap(const FString& GEngineIniName)
 					FName OldStructName = NAME_None;
 					FName NewStructName = NAME_None;
 
-					FParse::Value( *It.Value(), TEXT("OldStructName="), OldStructName );
-					FParse::Value( *It.Value(), TEXT("NewStructName="), NewStructName );
+					FParse::Value( *It.Value().GetValue(), TEXT("OldStructName="), OldStructName );
+					FParse::Value( *It.Value().GetValue(), TEXT("NewStructName="), NewStructName );
 
 					StructNameRedirects.Add(OldStructName, NewStructName);
 				}
@@ -202,8 +203,8 @@ void FLinkerLoad::CreateActiveRedirectsMap(const FString& GEngineIniName)
 					FString OldPluginName;
 					FString NewPluginName;
 
-					FParse::Value( *It.Value(), TEXT("OldPluginName="), OldPluginName );
-					FParse::Value( *It.Value(), TEXT("NewPluginName="), NewPluginName );
+					FParse::Value( *It.Value().GetValue(), TEXT("OldPluginName="), OldPluginName );
+					FParse::Value( *It.Value().GetValue(), TEXT("NewPluginName="), NewPluginName );
 
 					OldPluginName = FString(TEXT("/")) + OldPluginName + FString(TEXT("/"));
 					NewPluginName = FString(TEXT("/")) + NewPluginName + FString(TEXT("/"));
@@ -214,7 +215,7 @@ void FLinkerLoad::CreateActiveRedirectsMap(const FString& GEngineIniName)
 				{
 					FName KnownMissingPackage = NAME_None;
 
-					FParse::Value( *It.Value(), TEXT("PackageName="), KnownMissingPackage );
+					FParse::Value( *It.Value().GetValue(), TEXT("PackageName="), KnownMissingPackage );
 
 					KnownMissingPackages.Add(KnownMissingPackage);
 				}
@@ -607,42 +608,49 @@ FLinkerLoad::ELinkerStatus FLinkerLoad::Tick( float InTimeLimit, bool bInUseTime
 			// false is returned until any precaching is complete.
 			if( true )
 			{
+				SCOPED_LOADTIMER(LinkerLoad_CreateLoader);
 				Status = CreateLoader();
 			}
 
 			// Serialize the package file summary and presize the various arrays (name, import & export map)
 			if( Status == LINKER_Loaded )
 			{
+				SCOPED_LOADTIMER(LinkerLoad_SerializePackageFileSummary);
 				Status = SerializePackageFileSummary();
 			}
 
 			// Serialize the name map and register the names.
 			if( Status == LINKER_Loaded )
 			{
+				SCOPED_LOADTIMER(LinkerLoad_SerializeNameMap);
 				Status = SerializeNameMap();
 			}
 
 			// Serialize the gatherable text data map.
 			if( Status == LINKER_Loaded )
 			{
+				SCOPED_LOADTIMER(LinkerLoad_SerializeGatherableTextDataMap);
 				Status = SerializeGatherableTextDataMap();
 			}
 
 			// Serialize the import map.
 			if( Status == LINKER_Loaded )
 			{
+				SCOPED_LOADTIMER(LinkerLoad_SerializeImportMap);
 				Status = SerializeImportMap();
 			}
 
 			// Serialize the export map.
 			if( Status == LINKER_Loaded )
 			{
+				SCOPED_LOADTIMER(LinkerLoad_SerializeExportMap);
 				Status = SerializeExportMap();
 			}
 
 			// Start pre-allocation of texture memory.
 			if( Status == LINKER_Loaded )
 			{
+				SCOPED_LOADTIMER(LinkerLoad_StartTextureAllocation);
 #if WITH_ENGINE
 				Status = StartTextureAllocation();
 #endif		// WITH_ENGINE
@@ -651,41 +659,48 @@ FLinkerLoad::ELinkerStatus FLinkerLoad::Tick( float InTimeLimit, bool bInUseTime
 			// Fix up import map for backward compatible serialization.
 			if( Status == LINKER_Loaded )
 			{	
+				SCOPED_LOADTIMER(LinkerLoad_FixupImportMap);
 				Status = FixupImportMap();
 			}
 
 			if ( Status == LINKER_Loaded )
 			{
+				SCOPED_LOADTIMER(LinkerLoad_RemapImports);
 				Status = RemapImports();
 			}
 
 			// Fix up export map for object class conversion 
 			if( Status == LINKER_Loaded )
 			{	
+				SCOPED_LOADTIMER(LinkerLoad_FixupExportMap);
 				Status = FixupExportMap();
 			}
 
 			// Serialize the dependency map.
 			if( Status == LINKER_Loaded )
 			{
+				SCOPED_LOADTIMER(LinkerLoad_SerializeDependsMap);
 				Status = SerializeDependsMap();
 			}
 
 			// Hash exports.
 			if( Status == LINKER_Loaded )
 			{
+				SCOPED_LOADTIMER(LinkerLoad_CreateExportHash);
 				Status = CreateExportHash();
 			}
 
 			// Find existing objects matching exports and associate them with this linker.
 			if( Status == LINKER_Loaded )
 			{
+				SCOPED_LOADTIMER(LinkerLoad_FindExistingExports);
 				Status = FindExistingExports();
 			}
 
 			// Finalize creation process.
 			if( Status == LINKER_Loaded )
 			{
+				SCOPED_LOADTIMER(LinkerLoad_FinalizeCreation);
 				Status = FinalizeCreation();
 			}
 		}
@@ -1246,6 +1261,8 @@ FLinkerLoad::ELinkerStatus FLinkerLoad::SerializeNameMap()
 
 	while( bFinishedPrecaching && NameMapIndex < Summary.NameCount && !IsTimeLimitExceeded(TEXT("serializing name map"),100) )
 	{
+		SCOPED_LOADTIMER(LinkerLoad_SerializeNameMap_ProcessingEntries);
+
 		// Read the name entry from the file.
 		FNameEntrySerialized NameEntry(ENAME_LinkerConstructor);
 		*this << NameEntry;
diff --git a/Engine/Source/Runtime/CoreUObject/Private/UObject/MetaData.cpp b/Engine/Source/Runtime/CoreUObject/Private/UObject/MetaData.cpp
--- a/Engine/Source/Runtime/CoreUObject/Private/UObject/MetaData.cpp
+++ b/Engine/Source/Runtime/CoreUObject/Private/UObject/MetaData.cpp
@@ -438,8 +438,8 @@ void UMetaData::InitializeRedirectMap()
 						FName OldKey = NAME_None;
 						FName NewKey = NAME_None;
 
-						FParse::Value(*It.Value(), TEXT("OldKey="), OldKey);
-						FParse::Value(*It.Value(), TEXT("NewKey="), NewKey);
+						FParse::Value(*It.Value().GetValue(), TEXT("OldKey="), OldKey);
+						FParse::Value(*It.Value().GetValue(), TEXT("NewKey="), NewKey);
 
 						check(OldKey != NewKey);
 						check(OldKey != NAME_None);
diff --git a/Engine/Source/Runtime/CoreUObject/Private/UObject/Obj.cpp b/Engine/Source/Runtime/CoreUObject/Private/UObject/Obj.cpp
--- a/Engine/Source/Runtime/CoreUObject/Private/UObject/Obj.cpp
+++ b/Engine/Source/Runtime/CoreUObject/Private/UObject/Obj.cpp
@@ -1656,13 +1656,14 @@ void UObject::LoadConfig( UClass* ConfigClass/*=NULL*/, const TCHAR* InFilename/
 			//@Package name transition
 			if( Sec )
 			{
-				TArray<FString> List;
-				Sec->MultiFind(FName(*Key,FNAME_Find),List);
+				TArray<FConfigValue> List;
+				const FName KeyName(*Key, FNAME_Find);
+				Sec->MultiFind(KeyName,List);
 
 				// If we didn't find anything in the first section, try the alternate
 				if ((List.Num() == 0) && AltSec)
 				{
-					AltSec->MultiFind(FName(*Key,FNAME_Find),List);
+					AltSec->MultiFind(KeyName,List);
 				}
 
 				FScriptArrayHelper_InContainer ArrayHelper(Array, this);
@@ -1671,34 +1672,34 @@ void UObject::LoadConfig( UClass* ConfigClass/*=NULL*/, const TCHAR* InFilename/
 				if ( List.Num() > 0 )
 				{
 					ArrayHelper.EmptyAndAddValues(List.Num());
-						for( int32 i=List.Num()-1,c=0; i>=0; i--,c++ )
-						{
-							Array->Inner->ImportText( *List[i], ArrayHelper.GetRawPtr(c), PortFlags, this );
-						}
+					for( int32 i=List.Num()-1,c=0; i>=0; i--,c++ )
+					{
+						Array->Inner->ImportText( *List[i].GetValue(), ArrayHelper.GetRawPtr(c), PortFlags, this );
 					}
+				}
 				else
 				{
 					int32 Index = 0;
-					FString* ElementValue = NULL;
+					const FConfigValue* ElementValue = nullptr;
 					do
 					{
 						// Add array index number to end of key
 						FString IndexedKey = FString::Printf(TEXT("%s[%i]"), *Key, Index);
 
 						// Try to find value of key
-						FName IndexedName(*IndexedKey,FNAME_Find);
+						const FName IndexedName(*IndexedKey,FNAME_Find);
 						if (IndexedName == NAME_None)
 						{
 							break;
 						}
-						ElementValue  = Sec->Find(IndexedName);
+						ElementValue = Sec->Find(IndexedName);
 
 						// If found, import the element
-						if ( ElementValue != NULL )
+						if ( ElementValue != nullptr )
 						{
 							// expand the array if necessary so that Index is a valid element
 							ArrayHelper.ExpandForIndex(Index);
-							Array->Inner->ImportText(**ElementValue, ArrayHelper.GetRawPtr(Index), PortFlags, this);
+							Array->Inner->ImportText(*ElementValue->GetValue(), ArrayHelper.GetRawPtr(Index), PortFlags, this);
 						}
 
 						Index++;
diff --git a/Engine/Source/Runtime/Engine/Classes/AI/Navigation/NavModifierComponent.h b/Engine/Source/Runtime/Engine/Classes/AI/Navigation/NavModifierComponent.h
--- a/Engine/Source/Runtime/Engine/Classes/AI/Navigation/NavModifierComponent.h
+++ b/Engine/Source/Runtime/Engine/Classes/AI/Navigation/NavModifierComponent.h
@@ -13,7 +13,7 @@ class ENGINE_API UNavModifierComponent : public UNavRelevantComponent
 {
 	GENERATED_UCLASS_BODY()
 
-	UPROPERTY(EditAnywhere, Category = Navigation)
+	UPROPERTY(EditAnywhere, BlueprintReadOnly, Category = Navigation)
 	TSubclassOf<UNavArea> AreaClass;
 
 	/** box extent used ONLY when owning actor doesn't have collision component */
@@ -23,6 +23,18 @@ class ENGINE_API UNavModifierComponent : public UNavRelevantComponent
 	virtual void CalcAndCacheBounds() const override;
 	virtual void GetNavigationData(FNavigationRelevantData& Data) const override;
 
+	UFUNCTION(BlueprintCallable, Category = "AI|Navigation")
+	void SetAreaClass(TSubclassOf<UNavArea> NewAreaClass);
+
 protected:
-	mutable FVector ObstacleExtent;
+	struct FRotatedBox
+	{
+		FBox Box;
+		FQuat Quat;
+
+		FRotatedBox() {}
+		FRotatedBox(const FBox& InBox, const FQuat& InQuat) : Box(InBox), Quat(InQuat) {}
+	};
+
+	mutable TArray<FRotatedBox> ComponentBounds;
 };
diff --git a/Engine/Source/Runtime/Engine/Classes/AI/Navigation/NavRelevantComponent.h b/Engine/Source/Runtime/Engine/Classes/AI/Navigation/NavRelevantComponent.h
--- a/Engine/Source/Runtime/Engine/Classes/AI/Navigation/NavRelevantComponent.h
+++ b/Engine/Source/Runtime/Engine/Classes/AI/Navigation/NavRelevantComponent.h
@@ -45,7 +45,7 @@ protected:
 	UPROPERTY()
 	uint32 bAttachToOwnersRoot : 1;
 
-	uint32 bBoundsInitialized : 1;
+	mutable uint32 bBoundsInitialized : 1;
 
 	UPROPERTY(transient)
 	UObject* CachedNavParent;
diff --git a/Engine/Source/Runtime/Engine/Classes/AI/Navigation/NavigationData.h b/Engine/Source/Runtime/Engine/Classes/AI/Navigation/NavigationData.h
--- a/Engine/Source/Runtime/Engine/Classes/AI/Navigation/NavigationData.h
+++ b/Engine/Source/Runtime/Engine/Classes/AI/Navigation/NavigationData.h
@@ -535,8 +535,6 @@ public:
 
 	virtual bool DoesSupportAgent(const FNavAgentProperties& AgentProps) const;
 
-	FORCEINLINE void MarkAsNeedingUpdate() { bWantsUpdate = true; }
-
 	virtual void RestrictBuildingToActiveTiles(bool InRestrictBuildingToActiveTiles) {}
 
 	bool CanBeMainNavData() const { return bCanBeMainNavData; }
@@ -917,8 +915,7 @@ protected:
 	/** was it generated for default agent (SupportedAgents[0]) */
 	uint32 bSupportsDefaultAgent:1;
 
-	/** controls whether this NavigationData will accept navigation building requests resulting from dirty areas
-	 *	true by default, set to false on successful serialization to skip initial rebuild */
+	DEPRECATED(4.12, "This flag is now deprecated, initial rebuild ignore should be handled by discarding dirty areas in UNavigationSystem::ConditionalPopulateNavOctree.")
 	uint32 bWantsUpdate:1;
 
 private:
@@ -929,6 +926,8 @@ private:
 public:
 	DEPRECATED(4.8, "GetRandomPointInRadius is deprecated, please use GetRandomReachablePointInRadius")
 	virtual bool GetRandomPointInRadius(const FVector& Origin, float Radius, FNavLocation& OutResult, FSharedConstNavQueryFilter Filter = NULL, const UObject* Querier = NULL) const PURE_VIRTUAL(ANavigationData::GetRandomPointInRadius, return false;);
+	DEPRECATED(4.12, "This function is now deprecated, initial rebuild ignore should be handled by discarding dirty areas in UNavigationSystem::ConditionalPopulateNavOctree.")
+	void MarkAsNeedingUpdate() { }
 };
 
 struct FAsyncPathFindingQuery : public FPathFindingQuery
diff --git a/Engine/Source/Runtime/Engine/Classes/AI/Navigation/NavigationTestingActor.h b/Engine/Source/Runtime/Engine/Classes/AI/Navigation/NavigationTestingActor.h
--- a/Engine/Source/Runtime/Engine/Classes/AI/Navigation/NavigationTestingActor.h
+++ b/Engine/Source/Runtime/Engine/Classes/AI/Navigation/NavigationTestingActor.h
@@ -81,6 +81,9 @@ public:
 	UPROPERTY(EditAnywhere, Category=Pathfinding)
 	uint32 bGatherDetailedInfo : 1;
 
+	UPROPERTY(EditAnywhere, Category=Query)
+	uint32 bDrawDistanceToWall : 1;
+
 	/** show polys from open (orange) and closed (yellow) sets */
 	UPROPERTY(EditAnywhere, Category=Debug)
 	uint32 bShowNodePool : 1;
@@ -136,6 +139,8 @@ public:
 	UPROPERTY(EditAnywhere, Category=Pathfinding)
 	float OffsetFromCornersDistance;
 
+	FVector ClosestWallLocation;
+
 #if WITH_RECAST && WITH_EDITORONLY_DATA
 	/** detail data gathered from each step of regular A* algorithm */
 	TArray<struct FRecastDebugPathfindingData> DebugSteps;
@@ -185,11 +190,13 @@ public:
 	// constraints/goal evaluators).
 	virtual FPathFindingQuery BuildPathFindingQuery(const ANavigationTestingActor* Goal) const;
 
-public:
 	/** Returns CapsuleComponent subobject **/
 	class UCapsuleComponent* GetCapsuleComponent() const;
 #if WITH_EDITORONLY_DATA
 	/** Returns EdRenderComp subobject **/
 	class UNavTestRenderingComponent* GetEdRenderComp() const;
 #endif
+
+protected:
+	FVector FindClosestWallLocation() const;
 };
diff --git a/Engine/Source/Runtime/Engine/Classes/AI/Navigation/RecastNavMesh.h b/Engine/Source/Runtime/Engine/Classes/AI/Navigation/RecastNavMesh.h
--- a/Engine/Source/Runtime/Engine/Classes/AI/Navigation/RecastNavMesh.h
+++ b/Engine/Source/Runtime/Engine/Classes/AI/Navigation/RecastNavMesh.h
@@ -696,7 +696,6 @@ public:
 	//~ Begin UObject Interface
 	virtual void PostInitProperties() override;
 	virtual void PostLoad() override;
-	virtual SIZE_T GetResourceSize(EResourceSizeMode::Type Mode) override;	
 
 #if WITH_EDITOR
 	virtual void PostEditChangeProperty( struct FPropertyChangedEvent& PropertyChangedEvent) override;
@@ -879,8 +878,9 @@ public:
 	/** Returns nearest navmesh polygon to Loc, or INVALID_NAVMESHREF if Loc is not on the navmesh. */
 	NavNodeRef FindNearestPoly(FVector const& Loc, FVector const& Extent, FSharedConstNavQueryFilter Filter = NULL, const UObject* Querier = NULL) const;
 
-	/** Finds the distance to the closest wall, limited to MaxDistance */
-	float FindDistanceToWall(const FVector& StartLoc, FSharedConstNavQueryFilter Filter = nullptr, float MaxDistance = FLT_MAX) const;
+	/** Finds the distance to the closest wall, limited to MaxDistance
+	 *	[out] OutClosestPointOnWall, if supplied, will be set to closest point on closest wall. Will not be set if no wall in the area (return value 0.f) */
+	float FindDistanceToWall(const FVector& StartLoc, FSharedConstNavQueryFilter Filter = nullptr, float MaxDistance = FLT_MAX, FVector* OutClosestPointOnWall = nullptr) const;
 
 	/** Retrieves center of the specified polygon. Returns false on error. */
 	bool GetPolyCenter(NavNodeRef PolyID, FVector& OutCenter) const;
diff --git a/Engine/Source/Runtime/Engine/Classes/Animation/AnimInstance.h b/Engine/Source/Runtime/Engine/Classes/Animation/AnimInstance.h
--- a/Engine/Source/Runtime/Engine/Classes/Animation/AnimInstance.h
+++ b/Engine/Source/Runtime/Engine/Classes/Animation/AnimInstance.h
@@ -468,6 +468,13 @@ class ENGINE_API UAnimInstance : public UObject
 	UPROPERTY()
 	bool bCanUseParallelUpdateAnimation;
 
+	/**
+	 * Selecting this option will cause the compiler to emit warnings whenever a call into Blueprint
+	 * is made from the animation graph. This can help track down optimizations that need to be made.
+	 */
+	UPROPERTY(Category = Optimization, EditDefaultsOnly)
+	bool bWarnAboutBlueprintUsage;
+
 	/** Flag to check back on the game thread that indicates we need to run PostUpdateAnimation() in the post-eval call */
 	bool bNeedsUpdate;
 
diff --git a/Engine/Source/Runtime/Engine/Classes/Animation/AnimMetaData.h b/Engine/Source/Runtime/Engine/Classes/Animation/AnimMetaData.h
--- a/Engine/Source/Runtime/Engine/Classes/Animation/AnimMetaData.h
+++ b/Engine/Source/Runtime/Engine/Classes/Animation/AnimMetaData.h
@@ -10,8 +10,8 @@
 #pragma once
 #include "AnimMetaData.generated.h"
 
-UCLASS(Blueprintable, abstract, const, MinimalAPI, editinlinenew, hidecategories=Object, collapsecategories)
-class UAnimMetaData : public UObject
+UCLASS(Blueprintable, abstract, const, editinlinenew, hidecategories=Object, collapsecategories)
+class ENGINE_API UAnimMetaData : public UObject
 {
 	GENERATED_UCLASS_BODY()
 };
diff --git a/Engine/Source/Runtime/Engine/Classes/Animation/VertexAnim/VertexAnimation.h b/Engine/Source/Runtime/Engine/Classes/Animation/VertexAnim/VertexAnimation.h
--- a/Engine/Source/Runtime/Engine/Classes/Animation/VertexAnim/VertexAnimation.h
+++ b/Engine/Source/Runtime/Engine/Classes/Animation/VertexAnim/VertexAnimation.h
@@ -36,11 +36,6 @@ class UVertexAnimation : public UVertexAnimBase
 	/** Raw vertex anim data */
 	TArray<FVertexAnimFrame>	VertexAnimData;
 
-	//~ Begin UObject Interface.
-	virtual void Serialize( FArchive& Ar ) override;
-	virtual SIZE_T GetResourceSize(EResourceSizeMode::Type Mode) override;
-	//~ Begin UObject Interface.
-
 	//~ Begin UVertexAnimBase Interface
 	virtual FVertexAnimEvalStateBase* InitEval() override;
 	virtual void TermEval(FVertexAnimEvalStateBase* State) override;
diff --git a/Engine/Source/Runtime/Engine/Classes/Components/ChildActorComponent.h b/Engine/Source/Runtime/Engine/Classes/Components/ChildActorComponent.h
--- a/Engine/Source/Runtime/Engine/Classes/Components/ChildActorComponent.h
+++ b/Engine/Source/Runtime/Engine/Classes/Components/ChildActorComponent.h
@@ -49,9 +49,8 @@ private:
 	UPROPERTY(EditAnywhere, BlueprintReadOnly, Category=ChildActorComponent, meta=(OnlyPlaceable, AllowPrivateAccess="true"))
 	TSubclassOf<AActor>	ChildActorClass;
 
-public:
 	/** The actor that we spawned and own */
-	UPROPERTY(BlueprintReadOnly, Category=ChildActorComponent, TextExportTransient, NonPIEDuplicateTransient)
+	UPROPERTY(BlueprintReadOnly, Category=ChildActorComponent, TextExportTransient, NonPIEDuplicateTransient, meta=(AllowPrivateAccess="true"))
 	AActor*	ChildActor;
 
 	/** We try to keep the child actor's name as best we can, so we store it off here when destroying */
@@ -60,6 +59,8 @@ public:
 	/** Cached copy of the instance data when the ChildActor is destroyed to be available when needed */
 	mutable FChildActorComponentInstanceData* CachedInstanceData;
 
+public:
+
 	//~ Begin Object Interface.
 #if WITH_EDITOR
 	virtual void PostEditChangeProperty(FPropertyChangedEvent& PropertyChangedEvent) override;
@@ -83,6 +84,10 @@ public:
 	/** Create the child actor */
 	void CreateChildActor();
 
+	AActor* GetChildActor() const { return ChildActor; }
+
+	FName GetChildActorName() const { return ChildActorName; }
+
 	/** Kill any currently present child actor */
 	void DestroyChildActor(const bool bRequiresRename = true);
 };
diff --git a/Engine/Source/Runtime/Engine/Classes/Components/HierarchicalInstancedStaticMeshComponent.h b/Engine/Source/Runtime/Engine/Classes/Components/HierarchicalInstancedStaticMeshComponent.h
--- a/Engine/Source/Runtime/Engine/Classes/Components/HierarchicalInstancedStaticMeshComponent.h
+++ b/Engine/Source/Runtime/Engine/Classes/Components/HierarchicalInstancedStaticMeshComponent.h
@@ -126,7 +126,6 @@ public:
 
 	//Begin UObject Interface
 	virtual void Serialize(FArchive& Ar) override;
-	virtual SIZE_T GetResourceSize(EResourceSizeMode::Type Mode) override;
 	virtual FPrimitiveSceneProxy* CreateSceneProxy() override;
 	virtual void PostLoad() override;
 	virtual FBoxSphereBounds CalcBounds(const FTransform& BoundTransform) const override;
diff --git a/Engine/Source/Runtime/Engine/Classes/Components/PrimitiveComponent.h b/Engine/Source/Runtime/Engine/Classes/Components/PrimitiveComponent.h
--- a/Engine/Source/Runtime/Engine/Classes/Components/PrimitiveComponent.h
+++ b/Engine/Source/Runtime/Engine/Classes/Components/PrimitiveComponent.h
@@ -99,15 +99,15 @@ struct FSpriteCategoryInfo
  * Delegate for notification of blocking collision against a specific component.  
  * NormalImpulse will be filled in for physics-simulating bodies, but will be zero for swept-component blocking collisions. 
  */
-DECLARE_DYNAMIC_MULTICAST_DELEGATE_FourParams( FComponentHitSignature, class AActor*, OtherActor, class UPrimitiveComponent*, OtherComp, FVector, NormalImpulse, const FHitResult&, Hit );
+DECLARE_DYNAMIC_MULTICAST_DELEGATE_FiveParams( FComponentHitSignature, UPrimitiveComponent*, HitComponent, AActor*, OtherActor, UPrimitiveComponent*, OtherComp, FVector, NormalImpulse, const FHitResult&, Hit );
 /** Delegate for notification of start of overlap with a specific component */
-DECLARE_DYNAMIC_MULTICAST_DELEGATE_FiveParams( FComponentBeginOverlapSignature,class AActor*, OtherActor, class UPrimitiveComponent*, OtherComp, int32, OtherBodyIndex, bool, bFromSweep, const FHitResult &, SweepResult);
+DECLARE_DYNAMIC_MULTICAST_DELEGATE_SixParams( FComponentBeginOverlapSignature, UPrimitiveComponent*, OverlappedComponent, AActor*, OtherActor, UPrimitiveComponent*, OtherComp, int32, OtherBodyIndex, bool, bFromSweep, const FHitResult &, SweepResult);
 /** Delegate for notification of end of overlap with a specific component */
-DECLARE_DYNAMIC_MULTICAST_DELEGATE_ThreeParams( FComponentEndOverlapSignature, class AActor*, OtherActor, class UPrimitiveComponent*, OtherComp, int32, OtherBodyIndex);
+DECLARE_DYNAMIC_MULTICAST_DELEGATE_FourParams( FComponentEndOverlapSignature, UPrimitiveComponent*, OverlappedComponent, AActor*, OtherActor, UPrimitiveComponent*, OtherComp, int32, OtherBodyIndex);
 /** Delegate for notification when a wake event is fired by physics*/
-DECLARE_DYNAMIC_MULTICAST_DELEGATE_OneParam(FComponentWakeSignature, FName, BoneName);
+DECLARE_DYNAMIC_MULTICAST_DELEGATE_TwoParams(FComponentWakeSignature, UPrimitiveComponent*, WakingComponent, FName, BoneName);
 /** Delegate for notification when a sleep event is fired by physics*/
-DECLARE_DYNAMIC_MULTICAST_DELEGATE_OneParam(FComponentSleepSignature, FName, BoneName);
+DECLARE_DYNAMIC_MULTICAST_DELEGATE_TwoParams(FComponentSleepSignature, UPrimitiveComponent*, SleepingComponent, FName, BoneName);
 
 DECLARE_DYNAMIC_MULTICAST_DELEGATE_OneParam( FComponentBeginCursorOverSignature, UPrimitiveComponent*, TouchedComponent );
 DECLARE_DYNAMIC_MULTICAST_DELEGATE_OneParam( FComponentEndCursorOverSignature, UPrimitiveComponent*, TouchedComponent );
@@ -516,7 +516,8 @@ public:
 	 * Does not affect movement of this component when simulating physics.
 	 * @see IgnoreActorWhenMoving()
 	 */
-	TArray<TWeakObjectPtr<AActor> > MoveIgnoreActors;
+	UPROPERTY(Transient, DuplicateTransient)
+	TArray<AActor*> MoveIgnoreActors;
 
 	/**
 	 * Tells this component whether to ignore collision with all components of a specific Actor when this component is moved.
@@ -535,7 +536,7 @@ public:
 	/**
 	 * Returns the list of actors (as WeakObjectPtr) we currently ignore when moving.
 	 */
-	TArray<TWeakObjectPtr<AActor> > & GetMoveIgnoreActors();
+	const TArray<AActor*>& GetMoveIgnoreActors() const { return MoveIgnoreActors; }
 
 	/**
 	 * Clear the list of actors we ignore when moving.
@@ -621,7 +622,14 @@ public:
 	 * @param ClassFilter			[optional] If set, only returns actors of this class or subclasses
 	 */
 	UFUNCTION(BlueprintCallable, Category="Collision", meta=(UnsafeDuringActorConstruction="true"))
-	void GetOverlappingActors(TArray<AActor*>& OverlappingActors, UClass* ClassFilter=NULL) const;
+	void GetOverlappingActors(TArray<AActor*>& OverlappingActors, TSubclassOf<AActor> ClassFilter=nullptr) const;
+
+	/** 
+	* Returns the set of actors that this component is overlapping.
+	* @param OverlappingActors		[out] Returned list of overlapping actors
+	* @param ClassFilter			[optional] If set, only returns actors of this class or subclasses
+	*/
+	void GetOverlappingActors(TSet<AActor*>& OverlappingActors, TSubclassOf<AActor> ClassFilter=nullptr) const;
 
 	/** Returns list of components this component is overlapping. */
 	UFUNCTION(BlueprintCallable, Category="Collision", meta=(UnsafeDuringActorConstruction="true"))
diff --git a/Engine/Source/Runtime/Engine/Classes/Distributions/DistributionFloat.h b/Engine/Source/Runtime/Engine/Classes/Distributions/DistributionFloat.h
--- a/Engine/Source/Runtime/Engine/Classes/Distributions/DistributionFloat.h
+++ b/Engine/Source/Runtime/Engine/Classes/Distributions/DistributionFloat.h
@@ -48,7 +48,7 @@ public:
 	/**`
 		* Initialize a raw distribution from the original Unreal distribution
 		*/
-	void Initialize();
+	ENGINE_API void Initialize();
 #endif
 			 
 	/**
diff --git a/Engine/Source/Runtime/Engine/Classes/Distributions/DistributionVector.h b/Engine/Source/Runtime/Engine/Classes/Distributions/DistributionVector.h
--- a/Engine/Source/Runtime/Engine/Classes/Distributions/DistributionVector.h
+++ b/Engine/Source/Runtime/Engine/Classes/Distributions/DistributionVector.h
@@ -86,7 +86,7 @@ public:
 	/**
 	* Initialize a raw distribution from the original Unreal distribution
 	*/
-	void Initialize();
+	ENGINE_API void Initialize();
 #endif
 
 	/**
diff --git a/Engine/Source/Runtime/Engine/Classes/Engine/Console.h b/Engine/Source/Runtime/Engine/Classes/Engine/Console.h
--- a/Engine/Source/Runtime/Engine/Classes/Engine/Console.h
+++ b/Engine/Source/Runtime/Engine/Classes/Engine/Console.h
@@ -87,6 +87,12 @@ class ENGINE_API UConsole
 
 	int32 TypedStrPos;    //Current position in TypedStr
 
+	/** The command the user would get if they autocompleted their current input. Empty if no autocompletion entries are available. */
+	FString PrecompletedInputLine;
+
+	/** The most recent input that was autocompleted during this open console session. */
+	FString LastAutoCompletedCommand;
+
 	/**
 	 * Indicates that InputChar events should be captured to prevent them from being passed on to other interactions.  Reset
 	 * when the another keydown event is received.
@@ -255,8 +261,13 @@ private:
 
 	void SetInputLineFromAutoComplete();
 
+	void UpdatePrecompletedInputLine();
+
 	void NormalizeHistoryBuffer();
 
+	// Console settings from BaseInput.ini
+	const UConsoleSettings* ConsoleSettings;
+
 	// Widget that was focused before the console was opened (focus will be restored to this if it's valid after the console closes)
 	TWeakPtr<SWidget> PreviousFocusedWidget;
 };
diff --git a/Engine/Source/Runtime/Engine/Classes/Engine/Engine.h b/Engine/Source/Runtime/Engine/Classes/Engine/Engine.h
--- a/Engine/Source/Runtime/Engine/Classes/Engine/Engine.h
+++ b/Engine/Source/Runtime/Engine/Classes/Engine/Engine.h
@@ -1149,7 +1149,7 @@ public:
 	uint32 bUseFixedFrameRate : 1;
 	
 	/** The fixed framerate to use. */
-	UPROPERTY(config, EditAnywhere, Category = Framerate, meta=(EditCondition="bUseFixedFrameRate"))
+	UPROPERTY(config, EditAnywhere, Category = Framerate, meta=(EditCondition="bUseFixedFrameRate", ClampMin = "15.0"))
 	float FixedFrameRate;
 
 	/** Range of framerates in which smoothing will kick in */
@@ -1301,10 +1301,6 @@ public:
 	UPROPERTY(globalconfig)
 	float NetClientTicksPerSecond;
 
-	/** true if the engine needs to perform a delayed global component reregister (really just for editor) */
-	UPROPERTY(transient)
-	uint32 bHasPendingGlobalReregister:1;
-
 	/** Current display gamma setting */
 	UPROPERTY(config)
 	float DisplayGamma;
diff --git a/Engine/Source/Runtime/Engine/Classes/Engine/EngineTypes.h b/Engine/Source/Runtime/Engine/Classes/Engine/EngineTypes.h
--- a/Engine/Source/Runtime/Engine/Classes/Engine/EngineTypes.h
+++ b/Engine/Source/Runtime/Engine/Classes/Engine/EngineTypes.h
@@ -2178,150 +2178,6 @@ struct FPOV
 };
 
 
-/** The importance of a mesh feature when automatically generating mesh LODs. */
-UENUM()
-namespace EMeshFeatureImportance
-{
-	enum Type
-	{
-		Off,
-		Lowest,
-		Low,
-		Normal,
-		High,
-		Highest
-	};
-}
-
-
-/** Settings used to reduce a mesh. */
-USTRUCT()
-struct FMeshReductionSettings
-{
-	GENERATED_USTRUCT_BODY()
-
-	/** Percentage of triangles to keep. 1.0 = no reduction, 0.0 = no triangles. */
-		UPROPERTY(EditAnywhere, Category = ReductionSettings)
-	float PercentTriangles;
-
-	/** The maximum distance in object space by which the reduced mesh may deviate from the original mesh. */
-	UPROPERTY(EditAnywhere, Category = ReductionSettings)
-	float MaxDeviation;
-
-	/** Threshold in object space at which vertices are welded together. */
-	UPROPERTY(EditAnywhere, Category = ReductionSettings)
-	float WeldingThreshold;
-
-	/** Angle at which a hard edge is introduced between faces. */
-	UPROPERTY(EditAnywhere, Category = ReductionSettings)
-	float HardAngleThreshold;
-
-	/** Higher values minimize change to border edges. */
-	UPROPERTY(EditAnywhere, Category = ReductionSettings)
-	TEnumAsByte<EMeshFeatureImportance::Type> SilhouetteImportance;
-
-	/** Higher values reduce texture stretching. */
-	UPROPERTY(EditAnywhere, Category = ReductionSettings)
-	TEnumAsByte<EMeshFeatureImportance::Type> TextureImportance;
-
-	/** Higher values try to preserve normals better. */
-	UPROPERTY(EditAnywhere, Category = ReductionSettings)
-	TEnumAsByte<EMeshFeatureImportance::Type> ShadingImportance;
-
-	/*UPROPERTY(EditAnywhere, Category = ReductionSettings)
-	bool bActive;*/
-
-	UPROPERTY(EditAnywhere, Category = ReductionSettings)
-	bool bRecalculateNormals;
-
-	UPROPERTY(EditAnywhere, Category = ReductionSettings)
-	int32 BaseLODModel;
-
-	UPROPERTY(EditAnywhere, Category = ReductionSettings)
-		bool bGenerateUniqueLightmapUVs;
-
-	UPROPERTY(EditAnywhere, Category = ReductionSettings)
-		bool bKeepSymmetry;
-
-	UPROPERTY(EditAnywhere, Category = ReductionSettings)
-		bool bVisibilityAided;
-
-	UPROPERTY(EditAnywhere, Category = ReductionSettings)
-		bool bCullOccluded;
-
-	/** Higher values generates fewer samples*/
-	UPROPERTY(EditAnywhere, Category = ReductionSettings)
-		TEnumAsByte<EMeshFeatureImportance::Type> VisibilityAggressiveness;
-
-	/** Higher values minimize change to vertex color data. */
-	UPROPERTY(EditAnywhere, Category = ReductionSettings)
-		TEnumAsByte<EMeshFeatureImportance::Type> VertexColorImportance;
-
-	/** Default settings. */
-	FMeshReductionSettings()
-		: PercentTriangles(1.0f)
-		, MaxDeviation(0.0f)
-		, WeldingThreshold(0.0f)
-		, HardAngleThreshold(80.0f)
-		, SilhouetteImportance(EMeshFeatureImportance::Normal)
-		, TextureImportance(EMeshFeatureImportance::Normal)
-		, ShadingImportance(EMeshFeatureImportance::Normal)
-		, bRecalculateNormals(false)
-		, BaseLODModel(0)
-		, bGenerateUniqueLightmapUVs(false)
-		, bKeepSymmetry(false)
-		, bVisibilityAided(false)
-		, bCullOccluded(false)
-		, VisibilityAggressiveness(EMeshFeatureImportance::Lowest)
-		, VertexColorImportance(EMeshFeatureImportance::Off)
-	{
-	}
-
-	FMeshReductionSettings(const FMeshReductionSettings& Other)
-		: PercentTriangles(Other.PercentTriangles)
-		, MaxDeviation(Other.MaxDeviation)
-		, WeldingThreshold(Other.WeldingThreshold)
-		, HardAngleThreshold(Other.HardAngleThreshold)
-		, SilhouetteImportance(Other.ShadingImportance)
-		, TextureImportance(Other.TextureImportance)
-		, ShadingImportance(Other.ShadingImportance)
-		, bRecalculateNormals(Other.bRecalculateNormals)
-		, BaseLODModel(Other.BaseLODModel)
-		, bGenerateUniqueLightmapUVs(Other.bGenerateUniqueLightmapUVs)
-		, bKeepSymmetry(Other.bKeepSymmetry)
-		, bVisibilityAided(Other.bVisibilityAided)
-		, bCullOccluded(Other.bCullOccluded)
-		, VisibilityAggressiveness(Other.VisibilityAggressiveness)
-		, VertexColorImportance(Other.VertexColorImportance)
-	{
-	}
-
-	/** Equality operator. */
-	bool operator==(const FMeshReductionSettings& Other) const
-	{
-		return PercentTriangles == Other.PercentTriangles
-			&& MaxDeviation == Other.MaxDeviation
-			&& WeldingThreshold == Other.WeldingThreshold
-			&& HardAngleThreshold == Other.HardAngleThreshold
-			&& SilhouetteImportance == Other.SilhouetteImportance
-			&& TextureImportance == Other.TextureImportance
-			&& ShadingImportance == Other.ShadingImportance
-			&& bRecalculateNormals == Other.bRecalculateNormals
-			&& BaseLODModel == Other.BaseLODModel
-			&& bGenerateUniqueLightmapUVs == Other.bGenerateUniqueLightmapUVs
-			&& bKeepSymmetry == Other.bKeepSymmetry
-			&& bVisibilityAided == Other.bVisibilityAided
-			&& bCullOccluded == Other.bCullOccluded
-			&& VisibilityAggressiveness == Other.VisibilityAggressiveness
-			&& VertexColorImportance == Other.VertexColorImportance;
-	}
-
-	/** Inequality. */
-	bool operator!=(const FMeshReductionSettings& Other) const
-	{
-		return !(*this == Other);
-	}
-};
 
 /**
  * Settings applied when building a mesh.
@@ -2442,398 +2298,6 @@ struct FMeshBuildSettings
 	}
 };
 
-// Use FMaterialProxySettings instead
-USTRUCT()
-struct FMaterialSimplificationSettings
-{
-	GENERATED_USTRUCT_BODY()
-
-	// Size of generated BaseColor map
-	UPROPERTY(Category=Material, EditAnywhere)
-	FIntPoint BaseColorMapSize;
-
-	// Whether to generate normal map
-	UPROPERTY(Category=Material, EditAnywhere, meta=(InlineEditConditionToggle))
-	bool bNormalMap;
-	
-	// Size of generated specular map
-	UPROPERTY(Category=Material, EditAnywhere, meta=(editcondition = "bNormalMap"))
-	FIntPoint NormalMapSize;
-
-	// Metallic constant
-	UPROPERTY(Category=Material, EditAnywhere, meta=(ClampMin = "0", ClampMax = "1", UIMin = "0", UIMax = "1"))
-	float MetallicConstant;
-
-	// Whether to generate metallic map
-	UPROPERTY(Category=Material, EditAnywhere, meta=(InlineEditConditionToggle))
-	bool bMetallicMap;
-
-	// Size of generated metallic map
-	UPROPERTY(Category=Material, EditAnywhere, meta=(editcondition = "bMetallicMap"))
-	FIntPoint MetallicMapSize;
-
-	// Roughness constant
-	UPROPERTY(Category=Material, EditAnywhere, meta=(ClampMin = "0", ClampMax = "1", UIMin = "0", UIMax = "1"))
-	float RoughnessConstant;
-
-	// Whether to generate roughness map
-	UPROPERTY(Category=Material, EditAnywhere, meta=(InlineEditConditionToggle))
-	bool bRoughnessMap;
-
-	// Size of generated roughness map
-	UPROPERTY(Category=Material, EditAnywhere, meta=(editcondition = "bRoughnessMap"))
-	FIntPoint RoughnessMapSize;
-		
-	// Specular constant
-	UPROPERTY(Category=Material, EditAnywhere, meta=(ClampMin = "0", ClampMax = "1", UIMin = "0", UIMax = "1"))
-	float SpecularConstant;
-
-	// Whether to generate specular map
-	UPROPERTY(Category=Material, EditAnywhere, meta=(InlineEditConditionToggle))
-	bool bSpecularMap;
-
-	// Size of generated specular map
-	UPROPERTY(Category=Material, EditAnywhere, meta=(editcondition = "bSpecularMap"))
-	FIntPoint SpecularMapSize;
-
-	FMaterialSimplificationSettings()
-		: BaseColorMapSize(1024, 1024)
-		, bNormalMap(true)
-		, NormalMapSize(1024, 1024)
-		, MetallicConstant(0.0f)
-		, bMetallicMap(false)
-		, MetallicMapSize(1024, 1024)
-		, RoughnessConstant(0.5f)
-		, bRoughnessMap(false)
-		, RoughnessMapSize(1024, 1024)
-		, SpecularConstant(0.5f)
-		, bSpecularMap(false)
-		, SpecularMapSize(1024, 1024)
-	{
-	}
-	
-	bool operator == (const FMaterialSimplificationSettings& Other) const
-	{
-		return BaseColorMapSize == Other.BaseColorMapSize
-			&& bNormalMap == Other.bNormalMap
-			&& NormalMapSize == Other.NormalMapSize
-			&& MetallicConstant == Other.MetallicConstant
-			&& bMetallicMap == Other.bMetallicMap
-			&& MetallicMapSize == Other.MetallicMapSize
-			&& RoughnessConstant == Other.RoughnessConstant
-			&& bRoughnessMap == Other.bRoughnessMap
-			&& RoughnessMapSize == Other.RoughnessMapSize
-			&& SpecularConstant == Other.SpecularConstant
-			&& bSpecularMap == Other.bSpecularMap
-			&& SpecularMapSize == Other.SpecularMapSize;
-	}
-};
-
-UENUM()
-enum ETextureSizingType
-{
-	TextureSizingType_UseSingleTextureSize UMETA(DisplayName = "Use TextureSize for all material properties"),
-	TextureSizingType_UseAutomaticBiasedSizes UMETA(DisplayName = "Use automatically biased texture sizes based on TextureSize"),
-	TextureSizingType_UseManualOverrideTextureSize UMETA(DisplayName = "Use per property manually overriden texture sizes"),
-	TextureSizingType_UseSimplygonAutomaticSizing UMETA(DisplayName = "Use Simplygon's automatic texture sizing"),
-	TextureSizingType_MAX,
-};
-
-USTRUCT()
-struct FMaterialProxySettings
-{
-	GENERATED_USTRUCT_BODY()
-		
-	// Size of generated BaseColor map
-	UPROPERTY(Category = Material, EditAnywhere)
-	FIntPoint TextureSize;
-
-	UPROPERTY(Category = Material, EditAnywhere)
-	TEnumAsByte<ETextureSizingType> TextureSizingType;
-
-	UPROPERTY(Category = Material, EditAnywhere)
-	float GutterSpace;
-
-	// Whether to generate normal map
-	UPROPERTY(Category = Material, EditAnywhere)
-	bool bNormalMap;
-	
-	// Whether to generate metallic map
-	UPROPERTY(Category = Material, EditAnywhere)
-	bool bMetallicMap;
-	
-	// Metallic constant
-	UPROPERTY(Category = Material, EditAnywhere, meta = (ClampMin = "0", ClampMax = "1", UIMin = "0", UIMax = "1", editcondition = "!bMetallicMap"))
-	float MetallicConstant;
-
-	// Whether to generate roughness map
-	UPROPERTY(Category = Material, EditAnywhere)
-	bool bRoughnessMap;
-
-	// Roughness constant
-	UPROPERTY(Category = Material, EditAnywhere, meta = (ClampMin = "0", ClampMax = "1", UIMin = "0", UIMax = "1", editcondition = "!bRoughnessMap"))
-	float RoughnessConstant;
-		
-	// Whether to generate specular map
-	UPROPERTY(Category = Material, EditAnywhere)
-	bool bSpecularMap;	
-
-	// Specular constant
-	UPROPERTY(Category = Material, EditAnywhere, meta = (ClampMin = "0", ClampMax = "1", UIMin = "0", UIMax = "1", editcondition = "!bSpecularMap"))
-	float SpecularConstant;
-
-	// Whether to generate emissive map
-	UPROPERTY(Category = Material, EditAnywhere)
-	bool bEmissiveMap;
-
-	// Whether to generate opacity map
-	UPROPERTY(Category = Material, EditAnywhere)
-	bool bOpacityMap;
-
-	// Override diffuse map size
-	UPROPERTY(Category = Material, AdvancedDisplay, EditAnywhere)
-	FIntPoint DiffuseTextureSize;
-
-	// Override normal map size
-	UPROPERTY(Category = Material, AdvancedDisplay, EditAnywhere)
-	FIntPoint NormalTextureSize;
-
-	// Override metallic map size
-	UPROPERTY(Category = Material, AdvancedDisplay, EditAnywhere)
-	FIntPoint MetallicTextureSize;
-
-	// Override roughness map size
-	UPROPERTY(Category = Material, AdvancedDisplay, EditAnywhere)
-	FIntPoint RoughnessTextureSize;
-
-	// Override specular map size
-	UPROPERTY(Category = Material, AdvancedDisplay, EditAnywhere)
-	FIntPoint SpecularTextureSize;				    
-						
-	// Override emissive map size
-	UPROPERTY(Category = Material, AdvancedDisplay, EditAnywhere)
-	FIntPoint EmissiveTextureSize;				    
-						
-	// Override opacity map size
-	UPROPERTY(Category = Material, AdvancedDisplay, EditAnywhere)
-	FIntPoint OpacityTextureSize;
-	
-	FMaterialProxySettings()
-		: TextureSize(1024, 1024)
-		, TextureSizingType(TextureSizingType_UseSingleTextureSize)
-		, GutterSpace(4.0f)
-		, bNormalMap(true)
-		, bMetallicMap(false)
-		, MetallicConstant(0.0f)
-		, bRoughnessMap(false)
-		, RoughnessConstant(0.5f)				
-		, bSpecularMap(false)		
-		, SpecularConstant(0.5f)
-		, bEmissiveMap(false)
-		, bOpacityMap(false)
-		, DiffuseTextureSize(1024, 1024)
-		, NormalTextureSize(1024, 1024)
-		, MetallicTextureSize(1024, 1024)
-		, RoughnessTextureSize(1024, 1024)
-		, EmissiveTextureSize(1024, 1024)
-		, OpacityTextureSize(1024, 1024)
-	{
-	}
-
-	bool operator == (const FMaterialProxySettings& Other) const
-	{
-		return TextureSize == Other.TextureSize
-			&& TextureSizingType == Other.TextureSizingType
-			&& GutterSpace == Other.GutterSpace
-			&& bNormalMap == Other.bNormalMap
-			&& MetallicConstant == Other.MetallicConstant
-			&& bMetallicMap == Other.bMetallicMap
-			&& RoughnessConstant == Other.RoughnessConstant
-			&& bRoughnessMap == Other.bRoughnessMap
-			&& SpecularConstant == Other.SpecularConstant
-			&& bSpecularMap == Other.bSpecularMap
-			&& bEmissiveMap == Other.bEmissiveMap
-			&& bOpacityMap == Other.bOpacityMap
-			&& DiffuseTextureSize == Other.DiffuseTextureSize
-			&& NormalTextureSize == Other.NormalTextureSize
-			&& MetallicTextureSize == Other.MetallicTextureSize
-			&& RoughnessTextureSize == Other.RoughnessTextureSize
-			&& EmissiveTextureSize == Other.EmissiveTextureSize
-			&& OpacityTextureSize == Other.OpacityTextureSize;
-	}
-};
-
-USTRUCT()
-struct FMeshProxySettings
-{
-	GENERATED_USTRUCT_BODY()
-	/** Screen size of the resulting proxy mesh in pixel size*/
-	UPROPERTY(EditAnywhere, Category=ProxySettings)
-	int32 ScreenSize;
-
-	/** Material simplification */
-	UPROPERTY(EditAnywhere, Category = ProxySettings)
-	FMaterialProxySettings MaterialSettings;
-
-	UPROPERTY()
-	int32 TextureWidth_DEPRECATED;
-	UPROPERTY()
-	int32 TextureHeight_DEPRECATED;
-
-	UPROPERTY()
-	bool bExportNormalMap_DEPRECATED;
-
-	UPROPERTY()
-	bool bExportMetallicMap_DEPRECATED;
-
-	UPROPERTY()
-	bool bExportRoughnessMap_DEPRECATED;
-
-	UPROPERTY()
-	bool bExportSpecularMap_DEPRECATED;
-
-	/** Material simplification */
-	UPROPERTY()
-	FMaterialSimplificationSettings Material_DEPRECATED;
-
-
-	/** Determines whether or not the correct LOD models should be calculated given the source meshes and transition size */
-	UPROPERTY(EditAnywhere, Category = ProxySettings)
-	bool bCalculateLODSourceModels;
-
-	/** Distance at which meshes should be merged together */
-	UPROPERTY(EditAnywhere, Category = ProxySettings)
-	float MergeDistance;
-
-	/** Angle at which a hard edge is introduced between faces */
-	UPROPERTY(EditAnywhere, Category = ProxySettings, meta = (DisplayName = "Hard Edge Angle"))
-	float HardAngleThreshold;
-	
-	/** Lightmap resolution */
-	UPROPERTY(EditAnywhere, Category=ProxySettings)
-	int32 LightMapResolution;
-
-	/** Whether Simplygon should recalculate normals, otherwise the normals channel will be sampled from the original mesh */
-	UPROPERTY(EditAnywhere, Category=ProxySettings)
-	bool bRecalculateNormals;
-		
-	UPROPERTY()
-	bool bBakeVertexData_DEPRECATED;
-
-	/** Default settings. */
-	FMeshProxySettings()
-		: ScreenSize(300)
-		, TextureWidth_DEPRECATED(512)
-		, TextureHeight_DEPRECATED(512)
-		, bExportNormalMap_DEPRECATED(true)
-		, bExportMetallicMap_DEPRECATED(false)
-		, bExportRoughnessMap_DEPRECATED(false)
-		, bExportSpecularMap_DEPRECATED(false)
-		, MergeDistance(4)
-		, HardAngleThreshold(80.0f)
-		, LightMapResolution(256)
-		, bRecalculateNormals(true)
-	{ }
-
-	/** Equality operator. */
-	bool operator==(const FMeshProxySettings& Other) const
-	{
-		return ScreenSize == Other.ScreenSize
-			&& MaterialSettings == Other.MaterialSettings
-			&& bRecalculateNormals == Other.bRecalculateNormals
-			&& HardAngleThreshold == Other.HardAngleThreshold
-			&& MergeDistance == Other.MergeDistance;
-	}
-
-	/** Inequality. */
-	bool operator!=(const FMeshProxySettings& Other) const
-	{
-		return !(*this == Other);
-	}
-
-	/** Handles deprecated properties */
-	void PostLoadDeprecated();
-};
-
-/**
- * Mesh merging settings
- */
-USTRUCT()
-struct FMeshMergingSettings
-{
-	GENERATED_USTRUCT_BODY()
-
-	/** Whether to generate lightmap UVs for a merged mesh*/
-	UPROPERTY(EditAnywhere, Category=FMeshMergingSettings)
-	bool bGenerateLightMapUV;
-	
-	/** Target UV channel in a merged mesh for a lightmap */
-	UPROPERTY(EditAnywhere, Category=FMeshMergingSettings)
-	int32 TargetLightMapUVChannel;
-
-	/** Target lightmap resolution */
-	UPROPERTY(EditAnywhere, Category=FMeshMergingSettings)
-	int32 TargetLightMapResolution;
-		
-	/** Whether we should import vertex colors into merged mesh */
-	UPROPERTY()
-	bool bImportVertexColors_DEPRECATED;
-	
-	/** Whether merged mesh should have pivot at world origin, or at first merged component otherwise */
-	UPROPERTY(EditAnywhere, Category=FMeshMergingSettings)
-	bool bPivotPointAtZero;
-
-	/** Whether to merge physics data (collision primitives)*/
-	UPROPERTY(EditAnywhere, Category=FMeshMergingSettings)
-	bool bMergePhysicsData;
-
-	/** Whether to merge source materials into one flat material */
-	UPROPERTY(EditAnywhere, Category=MeshMerge)
-	bool bMergeMaterials;
-
-	/** Material simplification */
-	UPROPERTY(EditAnywhere, Category = MeshMerge, meta = (editcondition = "bMergeMaterials"))
-	FMaterialProxySettings MaterialSettings;
-
-	UPROPERTY(EditAnywhere, Category = MeshMerge)
-	bool bBakeVertexData;
-
-	/** Whether to export normal maps for material merging */
-	UPROPERTY()
-	bool bExportNormalMap_DEPRECATED;
-	/** Whether to export metallic maps for material merging */
-	UPROPERTY()
-	bool bExportMetallicMap_DEPRECATED;
-	/** Whether to export roughness maps for material merging */
-	UPROPERTY()
-	bool bExportRoughnessMap_DEPRECATED;
-	/** Whether to export specular maps for material merging */
-	UPROPERTY()
-	bool bExportSpecularMap_DEPRECATED;
-	/** Merged material texture atlas resolution */
-	UPROPERTY()
-	int32 MergedMaterialAtlasResolution_DEPRECATED;
-		
-	/** Default settings. */
-	FMeshMergingSettings()
-		: bGenerateLightMapUV(false)
-		, TargetLightMapUVChannel(1)
-		, TargetLightMapResolution(256)
-		, bImportVertexColors_DEPRECATED(false)
-		, bPivotPointAtZero(false)
-		, bMergePhysicsData(false)
-		, bMergeMaterials(false)
-		, bExportNormalMap_DEPRECATED(true)
-		, bExportMetallicMap_DEPRECATED(false)
-		, bExportRoughnessMap_DEPRECATED(false)
-		, bExportSpecularMap_DEPRECATED(false)
-		, MergedMaterialAtlasResolution_DEPRECATED(1024)
-	{
-	}
-
-	/** Handles deprecated properties */
-	void PostLoadDeprecated();
-};
 
 
 USTRUCT()
@@ -3063,21 +2527,24 @@ struct FTimerHandle
 {
 	GENERATED_BODY()
 
+	friend class FTimerManager;
+
 	FTimerHandle()
-	: Handle(INDEX_NONE)
+	: Handle(0)
 	{
 	}
 
 	bool IsValid() const
 	{
-		return Handle != INDEX_NONE;
+		return Handle != 0;
 	}
 
 	void Invalidate()
 	{
-		Handle = INDEX_NONE;
+		Handle = 0;
 	}
 
+	DEPRECATED(4.12, "This function is deprecated to avoid problems with timer wraparound. Please call FTimerManager::ValidateHandle.")
 	void MakeValid();
 
 	bool operator==(const FTimerHandle& Other) const
@@ -3096,8 +2563,8 @@ struct FTimerHandle
 	}
 
 private:
-	UPROPERTY(transient)
-	int32 Handle;
+	UPROPERTY(Transient)
+	uint64 Handle;
 };
 
 UENUM()
diff --git a/Engine/Source/Runtime/Engine/Classes/Engine/Level.h b/Engine/Source/Runtime/Engine/Classes/Engine/Level.h
--- a/Engine/Source/Runtime/Engine/Classes/Engine/Level.h
+++ b/Engine/Source/Runtime/Engine/Classes/Engine/Level.h
@@ -2,6 +2,7 @@
 
 #pragma once
 #include "Engine/World.h"
+#include "MaterialMerging.h"
 #include "Level.generated.h"
 
 class ALevelBounds;
diff --git a/Engine/Source/Runtime/Engine/Classes/Engine/MaterialMerging.h b/Engine/Source/Runtime/Engine/Classes/Engine/MaterialMerging.h
new file mode 100644
--- /dev/null
+++ b/Engine/Source/Runtime/Engine/Classes/Engine/MaterialMerging.h
@@ -0,0 +1,227 @@
+// Copyright 1998-2016 Epic Games, Inc. All Rights Reserved.
+
+#pragma once
+
+#include "CoreUObject.h"
+#include "MaterialMerging.generated.h"
+
+// Use FMaterialProxySettings instead
+USTRUCT()
+struct FMaterialSimplificationSettings
+{
+	GENERATED_USTRUCT_BODY()
+
+		// Size of generated BaseColor map
+		UPROPERTY(Category = Material, EditAnywhere)
+		FIntPoint BaseColorMapSize;
+
+	// Whether to generate normal map
+	UPROPERTY(Category = Material, EditAnywhere, meta = (InlineEditConditionToggle))
+		bool bNormalMap;
+
+	// Size of generated specular map
+	UPROPERTY(Category = Material, EditAnywhere, meta = (editcondition = "bNormalMap"))
+		FIntPoint NormalMapSize;
+
+	// Metallic constant
+	UPROPERTY(Category = Material, EditAnywhere, meta = (ClampMin = "0", ClampMax = "1", UIMin = "0", UIMax = "1"))
+		float MetallicConstant;
+
+	// Whether to generate metallic map
+	UPROPERTY(Category = Material, EditAnywhere, meta = (InlineEditConditionToggle))
+		bool bMetallicMap;
+
+	// Size of generated metallic map
+	UPROPERTY(Category = Material, EditAnywhere, meta = (editcondition = "bMetallicMap"))
+		FIntPoint MetallicMapSize;
+
+	// Roughness constant
+	UPROPERTY(Category = Material, EditAnywhere, meta = (ClampMin = "0", ClampMax = "1", UIMin = "0", UIMax = "1"))
+		float RoughnessConstant;
+
+	// Whether to generate roughness map
+	UPROPERTY(Category = Material, EditAnywhere, meta = (InlineEditConditionToggle))
+		bool bRoughnessMap;
+
+	// Size of generated roughness map
+	UPROPERTY(Category = Material, EditAnywhere, meta = (editcondition = "bRoughnessMap"))
+		FIntPoint RoughnessMapSize;
+
+	// Specular constant
+	UPROPERTY(Category = Material, EditAnywhere, meta = (ClampMin = "0", ClampMax = "1", UIMin = "0", UIMax = "1"))
+		float SpecularConstant;
+
+	// Whether to generate specular map
+	UPROPERTY(Category = Material, EditAnywhere, meta = (InlineEditConditionToggle))
+		bool bSpecularMap;
+
+	// Size of generated specular map
+	UPROPERTY(Category = Material, EditAnywhere, meta = (editcondition = "bSpecularMap"))
+		FIntPoint SpecularMapSize;
+
+	FMaterialSimplificationSettings()
+		: BaseColorMapSize(1024, 1024)
+		, bNormalMap(true)
+		, NormalMapSize(1024, 1024)
+		, MetallicConstant(0.0f)
+		, bMetallicMap(false)
+		, MetallicMapSize(1024, 1024)
+		, RoughnessConstant(0.5f)
+		, bRoughnessMap(false)
+		, RoughnessMapSize(1024, 1024)
+		, SpecularConstant(0.5f)
+		, bSpecularMap(false)
+		, SpecularMapSize(1024, 1024)
+	{
+	}
+
+	bool operator == (const FMaterialSimplificationSettings& Other) const
+	{
+		return BaseColorMapSize == Other.BaseColorMapSize
+			&& bNormalMap == Other.bNormalMap
+			&& NormalMapSize == Other.NormalMapSize
+			&& MetallicConstant == Other.MetallicConstant
+			&& bMetallicMap == Other.bMetallicMap
+			&& MetallicMapSize == Other.MetallicMapSize
+			&& RoughnessConstant == Other.RoughnessConstant
+			&& bRoughnessMap == Other.bRoughnessMap
+			&& RoughnessMapSize == Other.RoughnessMapSize
+			&& SpecularConstant == Other.SpecularConstant
+			&& bSpecularMap == Other.bSpecularMap
+			&& SpecularMapSize == Other.SpecularMapSize;
+	}
+};
+
+UENUM()
+enum ETextureSizingType
+{
+	TextureSizingType_UseSingleTextureSize UMETA(DisplayName = "Use TextureSize for all material properties"),
+	TextureSizingType_UseAutomaticBiasedSizes UMETA(DisplayName = "Use automatically biased texture sizes based on TextureSize"),
+	TextureSizingType_UseManualOverrideTextureSize UMETA(DisplayName = "Use per property manually overriden texture sizes"),
+	TextureSizingType_UseSimplygonAutomaticSizing UMETA(DisplayName = "Use Simplygon's automatic texture sizing"),
+	TextureSizingType_MAX,
+};
+
+USTRUCT()
+struct FMaterialProxySettings
+{
+	GENERATED_USTRUCT_BODY()
+
+		// Size of generated BaseColor map
+		UPROPERTY(Category = Material, EditAnywhere)
+		FIntPoint TextureSize;
+
+	UPROPERTY(Category = Material, EditAnywhere)
+		TEnumAsByte<ETextureSizingType> TextureSizingType;
+
+	UPROPERTY(Category = Material, EditAnywhere)
+		float GutterSpace;
+
+	// Whether to generate normal map
+	UPROPERTY(Category = Material, EditAnywhere)
+		bool bNormalMap;
+
+	// Whether to generate metallic map
+	UPROPERTY(Category = Material, EditAnywhere)
+		bool bMetallicMap;
+
+	// Metallic constant
+	UPROPERTY(Category = Material, EditAnywhere, meta = (ClampMin = "0", ClampMax = "1", UIMin = "0", UIMax = "1", editcondition = "!bMetallicMap"))
+		float MetallicConstant;
+
+	// Whether to generate roughness map
+	UPROPERTY(Category = Material, EditAnywhere)
+		bool bRoughnessMap;
+
+	// Roughness constant
+	UPROPERTY(Category = Material, EditAnywhere, meta = (ClampMin = "0", ClampMax = "1", UIMin = "0", UIMax = "1", editcondition = "!bRoughnessMap"))
+		float RoughnessConstant;
+
+	// Whether to generate specular map
+	UPROPERTY(Category = Material, EditAnywhere)
+		bool bSpecularMap;
+
+	// Specular constant
+	UPROPERTY(Category = Material, EditAnywhere, meta = (ClampMin = "0", ClampMax = "1", UIMin = "0", UIMax = "1", editcondition = "!bSpecularMap"))
+		float SpecularConstant;
+
+	// Whether to generate emissive map
+	UPROPERTY(Category = Material, EditAnywhere)
+		bool bEmissiveMap;
+
+	// Whether to generate opacity map
+	UPROPERTY(Category = Material, EditAnywhere)
+		bool bOpacityMap;
+
+	// Override diffuse map size
+	UPROPERTY(Category = Material, AdvancedDisplay, EditAnywhere)
+		FIntPoint DiffuseTextureSize;
+
+	// Override normal map size
+	UPROPERTY(Category = Material, AdvancedDisplay, EditAnywhere)
+		FIntPoint NormalTextureSize;
+
+	// Override metallic map size
+	UPROPERTY(Category = Material, AdvancedDisplay, EditAnywhere)
+		FIntPoint MetallicTextureSize;
+
+	// Override roughness map size
+	UPROPERTY(Category = Material, AdvancedDisplay, EditAnywhere)
+		FIntPoint RoughnessTextureSize;
+
+	// Override specular map size
+	UPROPERTY(Category = Material, AdvancedDisplay, EditAnywhere)
+		FIntPoint SpecularTextureSize;
+
+	// Override emissive map size
+	UPROPERTY(Category = Material, AdvancedDisplay, EditAnywhere)
+		FIntPoint EmissiveTextureSize;
+
+	// Override opacity map size
+	UPROPERTY(Category = Material, AdvancedDisplay, EditAnywhere)
+		FIntPoint OpacityTextureSize;
+
+	FMaterialProxySettings()
+		: TextureSize(1024, 1024)
+		, TextureSizingType(TextureSizingType_UseSingleTextureSize)
+		, GutterSpace(4.0f)
+		, bNormalMap(true)
+		, bMetallicMap(false)
+		, MetallicConstant(0.0f)
+		, bRoughnessMap(false)
+		, RoughnessConstant(0.5f)
+		, bSpecularMap(false)
+		, SpecularConstant(0.5f)
+		, bEmissiveMap(false)
+		, bOpacityMap(false)
+		, DiffuseTextureSize(1024, 1024)
+		, NormalTextureSize(1024, 1024)
+		, MetallicTextureSize(1024, 1024)
+		, RoughnessTextureSize(1024, 1024)
+		, EmissiveTextureSize(1024, 1024)
+		, OpacityTextureSize(1024, 1024)
+	{
+	}
+
+	bool operator == (const FMaterialProxySettings& Other) const
+	{
+		return TextureSize == Other.TextureSize
+			&& TextureSizingType == Other.TextureSizingType
+			&& GutterSpace == Other.GutterSpace
+			&& bNormalMap == Other.bNormalMap
+			&& MetallicConstant == Other.MetallicConstant
+			&& bMetallicMap == Other.bMetallicMap
+			&& RoughnessConstant == Other.RoughnessConstant
+			&& bRoughnessMap == Other.bRoughnessMap
+			&& SpecularConstant == Other.SpecularConstant
+			&& bSpecularMap == Other.bSpecularMap
+			&& bEmissiveMap == Other.bEmissiveMap
+			&& bOpacityMap == Other.bOpacityMap
+			&& DiffuseTextureSize == Other.DiffuseTextureSize
+			&& NormalTextureSize == Other.NormalTextureSize
+			&& MetallicTextureSize == Other.MetallicTextureSize
+			&& RoughnessTextureSize == Other.RoughnessTextureSize
+			&& EmissiveTextureSize == Other.EmissiveTextureSize
+			&& OpacityTextureSize == Other.OpacityTextureSize;
+	}
+};
diff --git a/Engine/Source/Runtime/Engine/Classes/Engine/MeshMerging.h b/Engine/Source/Runtime/Engine/Classes/Engine/MeshMerging.h
new file mode 100644
--- /dev/null
+++ b/Engine/Source/Runtime/Engine/Classes/Engine/MeshMerging.h
@@ -0,0 +1,325 @@
+// Copyright 1998-2016 Epic Games, Inc. All Rights Reserved.
+
+#pragma once
+
+#include "CoreUObject.h"
+#include "MaterialMerging.h"
+#include "MeshMerging.generated.h"
+
+/** The importance of a mesh feature when automatically generating mesh LODs. */
+UENUM()
+namespace EMeshFeatureImportance
+{
+	enum Type
+	{
+		Off,
+		Lowest,
+		Low,
+		Normal,
+		High,
+		Highest
+	};
+}
+
+
+/** Settings used to reduce a mesh. */
+USTRUCT()
+struct FMeshReductionSettings
+{
+	GENERATED_USTRUCT_BODY()
+
+		/** Percentage of triangles to keep. 1.0 = no reduction, 0.0 = no triangles. */
+		UPROPERTY(EditAnywhere, Category = ReductionSettings)
+		float PercentTriangles;
+
+	/** The maximum distance in object space by which the reduced mesh may deviate from the original mesh. */
+	UPROPERTY(EditAnywhere, Category = ReductionSettings)
+		float MaxDeviation;
+
+	/** Threshold in object space at which vertices are welded together. */
+	UPROPERTY(EditAnywhere, Category = ReductionSettings)
+		float WeldingThreshold;
+
+	/** Angle at which a hard edge is introduced between faces. */
+	UPROPERTY(EditAnywhere, Category = ReductionSettings)
+		float HardAngleThreshold;
+
+	/** Higher values minimize change to border edges. */
+	UPROPERTY(EditAnywhere, Category = ReductionSettings)
+		TEnumAsByte<EMeshFeatureImportance::Type> SilhouetteImportance;
+
+	/** Higher values reduce texture stretching. */
+	UPROPERTY(EditAnywhere, Category = ReductionSettings)
+		TEnumAsByte<EMeshFeatureImportance::Type> TextureImportance;
+
+	/** Higher values try to preserve normals better. */
+	UPROPERTY(EditAnywhere, Category = ReductionSettings)
+		TEnumAsByte<EMeshFeatureImportance::Type> ShadingImportance;
+
+	/*UPROPERTY(EditAnywhere, Category = ReductionSettings)
+	bool bActive;*/
+
+	UPROPERTY(EditAnywhere, Category = ReductionSettings)
+		bool bRecalculateNormals;
+
+	UPROPERTY(EditAnywhere, Category = ReductionSettings)
+		int32 BaseLODModel;
+
+	UPROPERTY(EditAnywhere, Category = ReductionSettings)
+		bool bGenerateUniqueLightmapUVs;
+
+	UPROPERTY(EditAnywhere, Category = ReductionSettings)
+		bool bKeepSymmetry;
+
+	UPROPERTY(EditAnywhere, Category = ReductionSettings)
+		bool bVisibilityAided;
+
+	UPROPERTY(EditAnywhere, Category = ReductionSettings)
+		bool bCullOccluded;
+
+	/** Higher values generates fewer samples*/
+	UPROPERTY(EditAnywhere, Category = ReductionSettings)
+		TEnumAsByte<EMeshFeatureImportance::Type> VisibilityAggressiveness;
+
+	/** Higher values minimize change to vertex color data. */
+	UPROPERTY(EditAnywhere, Category = ReductionSettings)
+		TEnumAsByte<EMeshFeatureImportance::Type> VertexColorImportance;
+
+	/** Default settings. */
+	FMeshReductionSettings()
+		: PercentTriangles(1.0f)
+		, MaxDeviation(0.0f)
+		, WeldingThreshold(0.0f)
+		, HardAngleThreshold(80.0f)
+		, SilhouetteImportance(EMeshFeatureImportance::Normal)
+		, TextureImportance(EMeshFeatureImportance::Normal)
+		, ShadingImportance(EMeshFeatureImportance::Normal)
+		, bRecalculateNormals(false)
+		, BaseLODModel(0)
+		, bGenerateUniqueLightmapUVs(false)
+		, bKeepSymmetry(false)
+		, bVisibilityAided(false)
+		, bCullOccluded(false)
+		, VisibilityAggressiveness(EMeshFeatureImportance::Lowest)
+		, VertexColorImportance(EMeshFeatureImportance::Off)
+	{
+	}
+
+	FMeshReductionSettings(const FMeshReductionSettings& Other)
+		: PercentTriangles(Other.PercentTriangles)
+		, MaxDeviation(Other.MaxDeviation)
+		, WeldingThreshold(Other.WeldingThreshold)
+		, HardAngleThreshold(Other.HardAngleThreshold)
+		, SilhouetteImportance(Other.ShadingImportance)
+		, TextureImportance(Other.TextureImportance)
+		, ShadingImportance(Other.ShadingImportance)
+		, bRecalculateNormals(Other.bRecalculateNormals)
+		, BaseLODModel(Other.BaseLODModel)
+		, bGenerateUniqueLightmapUVs(Other.bGenerateUniqueLightmapUVs)
+		, bKeepSymmetry(Other.bKeepSymmetry)
+		, bVisibilityAided(Other.bVisibilityAided)
+		, bCullOccluded(Other.bCullOccluded)
+		, VisibilityAggressiveness(Other.VisibilityAggressiveness)
+		, VertexColorImportance(Other.VertexColorImportance)
+	{
+	}
+
+	/** Equality operator. */
+	bool operator==(const FMeshReductionSettings& Other) const
+	{
+		return PercentTriangles == Other.PercentTriangles
+			&& MaxDeviation == Other.MaxDeviation
+			&& WeldingThreshold == Other.WeldingThreshold
+			&& HardAngleThreshold == Other.HardAngleThreshold
+			&& SilhouetteImportance == Other.SilhouetteImportance
+			&& TextureImportance == Other.TextureImportance
+			&& ShadingImportance == Other.ShadingImportance
+			&& bRecalculateNormals == Other.bRecalculateNormals
+			&& BaseLODModel == Other.BaseLODModel
+			&& bGenerateUniqueLightmapUVs == Other.bGenerateUniqueLightmapUVs
+			&& bKeepSymmetry == Other.bKeepSymmetry
+			&& bVisibilityAided == Other.bVisibilityAided
+			&& bCullOccluded == Other.bCullOccluded
+			&& VisibilityAggressiveness == Other.VisibilityAggressiveness
+			&& VertexColorImportance == Other.VertexColorImportance;
+	}
+
+	/** Inequality. */
+	bool operator!=(const FMeshReductionSettings& Other) const
+	{
+		return !(*this == Other);
+	}
+};
+
+
+USTRUCT()
+struct FMeshProxySettings
+{
+	GENERATED_USTRUCT_BODY()
+		/** Screen size of the resulting proxy mesh in pixel size*/
+		UPROPERTY(EditAnywhere, Category = ProxySettings)
+		int32 ScreenSize;
+
+	/** Material simplification */
+	UPROPERTY(EditAnywhere, Category = ProxySettings)
+		FMaterialProxySettings MaterialSettings;
+
+	UPROPERTY()
+		int32 TextureWidth_DEPRECATED;
+	UPROPERTY()
+		int32 TextureHeight_DEPRECATED;
+
+	UPROPERTY()
+		bool bExportNormalMap_DEPRECATED;
+
+	UPROPERTY()
+		bool bExportMetallicMap_DEPRECATED;
+
+	UPROPERTY()
+		bool bExportRoughnessMap_DEPRECATED;
+
+	UPROPERTY()
+		bool bExportSpecularMap_DEPRECATED;
+
+	/** Material simplification */
+	UPROPERTY()
+		FMaterialSimplificationSettings Material_DEPRECATED;
+
+
+	/** Determines whether or not the correct LOD models should be calculated given the source meshes and transition size */
+	UPROPERTY(EditAnywhere, Category = ProxySettings)
+		bool bCalculateLODSourceModels;
+
+	/** Distance at which meshes should be merged together */
+	UPROPERTY(EditAnywhere, Category = ProxySettings)
+		float MergeDistance;
+
+	/** Angle at which a hard edge is introduced between faces */
+	UPROPERTY(EditAnywhere, Category = ProxySettings, meta = (DisplayName = "Hard Edge Angle"))
+		float HardAngleThreshold;
+
+	/** Lightmap resolution */
+	UPROPERTY(EditAnywhere, Category = ProxySettings)
+		int32 LightMapResolution;
+
+	/** Whether Simplygon should recalculate normals, otherwise the normals channel will be sampled from the original mesh */
+	UPROPERTY(EditAnywhere, Category = ProxySettings)
+		bool bRecalculateNormals;
+
+	UPROPERTY()
+		bool bBakeVertexData_DEPRECATED;
+
+	/** Default settings. */
+	FMeshProxySettings()
+		: ScreenSize(300)
+		, TextureWidth_DEPRECATED(512)
+		, TextureHeight_DEPRECATED(512)
+		, bExportNormalMap_DEPRECATED(true)
+		, bExportMetallicMap_DEPRECATED(false)
+		, bExportRoughnessMap_DEPRECATED(false)
+		, bExportSpecularMap_DEPRECATED(false)
+		, MergeDistance(4)
+		, HardAngleThreshold(80.0f)
+		, LightMapResolution(256)
+		, bRecalculateNormals(true)
+	{ }
+
+	/** Equality operator. */
+	bool operator==(const FMeshProxySettings& Other) const
+	{
+		return ScreenSize == Other.ScreenSize
+			&& MaterialSettings == Other.MaterialSettings
+			&& bRecalculateNormals == Other.bRecalculateNormals
+			&& HardAngleThreshold == Other.HardAngleThreshold
+			&& MergeDistance == Other.MergeDistance;
+	}
+
+	/** Inequality. */
+	bool operator!=(const FMeshProxySettings& Other) const
+	{
+		return !(*this == Other);
+	}
+
+	/** Handles deprecated properties */
+	void PostLoadDeprecated();
+};
+
+/**
+* Mesh merging settings
+*/
+USTRUCT()
+struct FMeshMergingSettings
+{
+	GENERATED_USTRUCT_BODY()
+
+		/** Whether to generate lightmap UVs for a merged mesh*/
+		UPROPERTY(EditAnywhere, Category = FMeshMergingSettings)
+		bool bGenerateLightMapUV;
+
+	/** Target UV channel in a merged mesh for a lightmap */
+	UPROPERTY(EditAnywhere, Category = FMeshMergingSettings)
+		int32 TargetLightMapUVChannel;
+
+	/** Target lightmap resolution */
+	UPROPERTY(EditAnywhere, Category = FMeshMergingSettings)
+		int32 TargetLightMapResolution;
+
+	/** Whether we should import vertex colors into merged mesh */
+	UPROPERTY()
+		bool bImportVertexColors_DEPRECATED;
+
+	/** Whether merged mesh should have pivot at world origin, or at first merged component otherwise */
+	UPROPERTY(EditAnywhere, Category = FMeshMergingSettings)
+		bool bPivotPointAtZero;
+
+	/** Whether to merge physics data (collision primitives)*/
+	UPROPERTY(EditAnywhere, Category = FMeshMergingSettings)
+		bool bMergePhysicsData;
+
+	/** Whether to merge source materials into one flat material */
+	UPROPERTY(EditAnywhere, Category = MeshMerge)
+		bool bMergeMaterials;
+
+	/** Material simplification */
+	UPROPERTY(EditAnywhere, Category = MeshMerge, meta = (editcondition = "bMergeMaterials"))
+		FMaterialProxySettings MaterialSettings;
+
+	UPROPERTY(EditAnywhere, Category = MeshMerge)
+		bool bBakeVertexData;
+
+	/** Whether to export normal maps for material merging */
+	UPROPERTY()
+		bool bExportNormalMap_DEPRECATED;
+	/** Whether to export metallic maps for material merging */
+	UPROPERTY()
+		bool bExportMetallicMap_DEPRECATED;
+	/** Whether to export roughness maps for material merging */
+	UPROPERTY()
+		bool bExportRoughnessMap_DEPRECATED;
+	/** Whether to export specular maps for material merging */
+	UPROPERTY()
+		bool bExportSpecularMap_DEPRECATED;
+	/** Merged material texture atlas resolution */
+	UPROPERTY()
+		int32 MergedMaterialAtlasResolution_DEPRECATED;
+
+	/** Default settings. */
+	FMeshMergingSettings()
+		: bGenerateLightMapUV(false)
+		, TargetLightMapUVChannel(1)
+		, TargetLightMapResolution(256)
+		, bImportVertexColors_DEPRECATED(false)
+		, bPivotPointAtZero(false)
+		, bMergePhysicsData(false)
+		, bMergeMaterials(false)
+		, bExportNormalMap_DEPRECATED(true)
+		, bExportMetallicMap_DEPRECATED(false)
+		, bExportRoughnessMap_DEPRECATED(false)
+		, bExportSpecularMap_DEPRECATED(false)
+		, MergedMaterialAtlasResolution_DEPRECATED(1024)
+	{
+	}
+
+	/** Handles deprecated properties */
+	void PostLoadDeprecated();
+};
\ No newline at end of file
diff --git a/Engine/Source/Runtime/Engine/Classes/Engine/SkeletalMesh.h b/Engine/Source/Runtime/Engine/Classes/Engine/SkeletalMesh.h
--- a/Engine/Source/Runtime/Engine/Classes/Engine/SkeletalMesh.h
+++ b/Engine/Source/Runtime/Engine/Classes/Engine/SkeletalMesh.h
@@ -599,8 +599,52 @@ public:
 	UPROPERTY(Category=Mesh, AssetRegistrySearchable, VisibleAnywhere, BlueprintReadOnly)
 	USkeleton* Skeleton;
 
+private:
+	/** Original imported mesh bounds */
+	UPROPERTY(transient, duplicatetransient)
+	FBoxSphereBounds ImportedBounds;
+
+	/** Bounds extended by user values below */
 	UPROPERTY(transient, duplicatetransient)
-	FBoxSphereBounds Bounds;
+	FBoxSphereBounds ExtendedBounds;
+
+protected:
+	// The properties below are protected to force the use of the Set* methods for this data
+	// in code so we can keep the extended bounds up to date after changing the data.
+	// Property editors will trigger property events to correctly recalculate the extended bounds.
+
+	/** Bound extension values in the positive direction of XYZ, positive value increases bound size */
+	UPROPERTY(EditAnywhere, BlueprintReadOnly, Category = Mesh)
+	FVector PositiveBoundsExtension;
+
+	/** Bound extension values in the negative direction of XYZ, positive value increases bound size */
+	UPROPERTY(EditAnywhere, BlueprintReadOnly, Category = Mesh)
+	FVector NegativeBoundsExtension;
+
+public:
+
+	/** Get the extended bounds of this mesh (imported bounds plus bounds extension) */
+	UFUNCTION(BlueprintCallable, Category = Mesh)
+	ENGINE_API FBoxSphereBounds GetBounds();
+
+	/** Get the original imported bounds of the skel mesh */
+	UFUNCTION(BlueprintCallable, Category = Mesh)
+	ENGINE_API FBoxSphereBounds GetImportedBounds();
+
+	/** Set the original imported bounds of the skel mesh, will recalculate extended bounds */
+	ENGINE_API void SetImportedBounds(const FBoxSphereBounds& InBounds);
+
+	/** Set bound extension values in the positive direction of XYZ, positive value increases bound size */
+	ENGINE_API void SetPositiveBoundsExtension(const FVector& InExtension);
+
+	/** Set bound extension values in the negative direction of XYZ, positive value increases bound size */
+	ENGINE_API void SetNegativeBoundsExtension(const FVector& InExtension);
+
+	/** Calculate the extended bounds based on the imported bounds and the extension values */
+	void CalculateExtendedBounds();
+
+	/** Alters the bounds extension values to fit correctly into the current bounds (so negative values never extend the bounds etc.) */
+	void ValidateBoundsExtension();
 
 	/** List of materials applied to this mesh. */
 	UPROPERTY(EditAnywhere, BlueprintReadOnly, transient, duplicatetransient, Category=SkeletalMesh)
@@ -788,6 +832,7 @@ public:
 #if WITH_EDITOR
 	virtual void PreEditChange(UProperty* PropertyAboutToChange) override;
 	virtual void PostEditChangeProperty(FPropertyChangedEvent& PropertyChangedEvent) override;
+
 	virtual void PostEditUndo() override;
 	virtual void GetAssetRegistryTagMetadata(TMap<FName, FAssetRegistryTagMetadata>& OutMetadata) const override;
 #endif // WITH_EDITOR
@@ -971,6 +1016,7 @@ private:
 	* Test whether all the flags in an array are identical (could be moved to Array.h?)
 	*/
 	bool AreAllFlagsIdentical( const TArray<bool>& BoolArray ) const;
+
 };
 
 
diff --git a/Engine/Source/Runtime/Engine/Classes/Engine/StaticMesh.h b/Engine/Source/Runtime/Engine/Classes/Engine/StaticMesh.h
--- a/Engine/Source/Runtime/Engine/Classes/Engine/StaticMesh.h
+++ b/Engine/Source/Runtime/Engine/Classes/Engine/StaticMesh.h
@@ -3,6 +3,7 @@
 #pragma once
 #include "Interfaces/Interface_CollisionDataProvider.h"
 #include "Interfaces/Interface_AssetUserData.h"
+#include "MeshMerging.h"
 #include "StaticMesh.generated.h"
 
 /** The maximum number of static mesh LODs allowed. */
diff --git a/Engine/Source/Runtime/Engine/Classes/GameFramework/Actor.h b/Engine/Source/Runtime/Engine/Classes/GameFramework/Actor.h
--- a/Engine/Source/Runtime/Engine/Classes/GameFramework/Actor.h
+++ b/Engine/Source/Runtime/Engine/Classes/GameFramework/Actor.h
@@ -21,23 +21,23 @@ ENGINE_API DECLARE_LOG_CATEGORY_EXTERN(LogActor, Log, Warning);
  
 
 // Delegate signatures
-DECLARE_DYNAMIC_MULTICAST_DELEGATE_FourParams( FTakeAnyDamageSignature, float, Damage, const class UDamageType*, DamageType, class AController*, InstigatedBy, AActor*, DamageCauser );
-DECLARE_DYNAMIC_MULTICAST_DELEGATE_EightParams( FTakePointDamageSignature, float, Damage, class AController*, InstigatedBy, FVector, HitLocation, class UPrimitiveComponent*, FHitComponent, FName, BoneName, FVector, ShotFromDirection, const class UDamageType*, DamageType, AActor*, DamageCauser );
-DECLARE_DYNAMIC_MULTICAST_DELEGATE_OneParam( FActorBeginOverlapSignature, AActor*, OtherActor );
-DECLARE_DYNAMIC_MULTICAST_DELEGATE_OneParam( FActorEndOverlapSignature, AActor*, OtherActor );
+DECLARE_DYNAMIC_MULTICAST_DELEGATE_FiveParams( FTakeAnyDamageSignature, AActor*, DamagedActor, float, Damage, const class UDamageType*, DamageType, class AController*, InstigatedBy, AActor*, DamageCauser );
+DECLARE_DYNAMIC_MULTICAST_DELEGATE_NineParams( FTakePointDamageSignature, AActor*, DamagedActor, float, Damage, class AController*, InstigatedBy, FVector, HitLocation, class UPrimitiveComponent*, FHitComponent, FName, BoneName, FVector, ShotFromDirection, const class UDamageType*, DamageType, AActor*, DamageCauser );
+DECLARE_DYNAMIC_MULTICAST_DELEGATE_TwoParams( FActorBeginOverlapSignature, AActor*, OverlappedActor, AActor*, OtherActor );
+DECLARE_DYNAMIC_MULTICAST_DELEGATE_TwoParams( FActorEndOverlapSignature, AActor*, OverlappedActor, AActor*, OtherActor );
 DECLARE_DYNAMIC_MULTICAST_DELEGATE_FourParams( FActorHitSignature, AActor*, SelfActor, AActor*, OtherActor, FVector, NormalImpulse, const FHitResult&, Hit );
 
-DECLARE_DYNAMIC_MULTICAST_DELEGATE( FActorBeginCursorOverSignature );
-DECLARE_DYNAMIC_MULTICAST_DELEGATE( FActorEndCursorOverSignature );
-DECLARE_DYNAMIC_MULTICAST_DELEGATE_OneParam( FActorOnClickedSignature, FKey, ButtonPressed );
-DECLARE_DYNAMIC_MULTICAST_DELEGATE_OneParam( FActorOnReleasedSignature, FKey, ButtonReleased );
-DECLARE_DYNAMIC_MULTICAST_DELEGATE_OneParam( FActorOnInputTouchBeginSignature, ETouchIndex::Type, FingerIndex );
-DECLARE_DYNAMIC_MULTICAST_DELEGATE_OneParam( FActorOnInputTouchEndSignature, ETouchIndex::Type, FingerIndex );
-DECLARE_DYNAMIC_MULTICAST_DELEGATE_OneParam( FActorBeginTouchOverSignature, ETouchIndex::Type, FingerIndex );
-DECLARE_DYNAMIC_MULTICAST_DELEGATE_OneParam( FActorEndTouchOverSignature, ETouchIndex::Type, FingerIndex );
+DECLARE_DYNAMIC_MULTICAST_DELEGATE_OneParam( FActorBeginCursorOverSignature, AActor*, TouchedActor );
+DECLARE_DYNAMIC_MULTICAST_DELEGATE_OneParam( FActorEndCursorOverSignature, AActor*, TouchedActor );
+DECLARE_DYNAMIC_MULTICAST_DELEGATE_TwoParams( FActorOnClickedSignature, AActor*, TouchedActor , FKey, ButtonPressed );
+DECLARE_DYNAMIC_MULTICAST_DELEGATE_TwoParams( FActorOnReleasedSignature, AActor*, TouchedActor , FKey, ButtonReleased );
+DECLARE_DYNAMIC_MULTICAST_DELEGATE_TwoParams( FActorOnInputTouchBeginSignature, ETouchIndex::Type, FingerIndex, AActor*, TouchedActor );
+DECLARE_DYNAMIC_MULTICAST_DELEGATE_TwoParams( FActorOnInputTouchEndSignature, ETouchIndex::Type, FingerIndex, AActor*, TouchedActor );
+DECLARE_DYNAMIC_MULTICAST_DELEGATE_TwoParams( FActorBeginTouchOverSignature, ETouchIndex::Type, FingerIndex, AActor*, TouchedActor );
+DECLARE_DYNAMIC_MULTICAST_DELEGATE_TwoParams( FActorEndTouchOverSignature, ETouchIndex::Type, FingerIndex, AActor*, TouchedActor );
 
-DECLARE_DYNAMIC_MULTICAST_DELEGATE(FActorDestroyedSignature);
-DECLARE_DYNAMIC_MULTICAST_DELEGATE_OneParam(FActorEndPlaySignature, EEndPlayReason::Type, EndPlayReason);
+DECLARE_DYNAMIC_MULTICAST_DELEGATE_OneParam(FActorDestroyedSignature, AActor*, DestroyedActor );
+DECLARE_DYNAMIC_MULTICAST_DELEGATE_TwoParams(FActorEndPlaySignature, AActor*, Actor , EEndPlayReason::Type, EndPlayReason);
 
 DECLARE_DELEGATE_SixParams(FMakeNoiseDelegate, AActor*, float /*Loudness*/, class APawn*, const FVector&, float /*MaxRange*/, FName /*Tag*/);
 
@@ -562,7 +562,7 @@ public:
 	UPROPERTY()
 	uint8 bActorSeamlessTraveled:1;
 
-	/** Whether this actor should no be affected by world origin shifting. */
+	/** Whether this actor should not be affected by world origin shifting. */
 	UPROPERTY(EditAnywhere, AdvancedDisplay, Category=Actor)
 	uint8 bIgnoresOriginShifting:1;
 	
@@ -1166,6 +1166,10 @@ public:
 	UFUNCTION(BlueprintCallable, meta=(DeprecatedFunction, DeprecationMessage="Use PrimitiveComponent.CreateAndSetMaterialInstanceDynamic instead.", BlueprintProtected = "true"), Category="Rendering|Material")
 	class UMaterialInstanceDynamic* MakeMIDForMaterial(class UMaterialInterface* Parent);
 
+	/** The number of seconds (in game time) since this Actor was created, relative to Get Game Time In Seconds. */
+	UFUNCTION(BlueprintPure, Category = Actor)
+	float GetGameTimeSinceCreation();
+
 	//~=============================================================================
 	// AI functions.
 	
@@ -1307,12 +1311,22 @@ public:
 	 * @param ClassFilter			[optional] If set, only returns actors of this class or subclasses
 	 */
 	UFUNCTION(BlueprintCallable, Category="Collision", meta=(UnsafeDuringActorConstruction="true"))
-	void GetOverlappingActors(TArray<AActor*>& OverlappingActors, UClass* ClassFilter=NULL) const;
+	void GetOverlappingActors(TArray<AActor*>& OverlappingActors, TSubclassOf<AActor> ClassFilter=nullptr) const;
+
+	/** 
+	 * Returns set of actors this actor is overlapping (any component overlapping any component). Does not return itself.
+	 * @param OverlappingActors		[out] Returned list of overlapping actors
+	 * @param ClassFilter			[optional] If set, only returns actors of this class or subclasses
+	 */
+	void GetOverlappingActors(TSet<AActor*>& OverlappingActors, TSubclassOf<AActor> ClassFilter=nullptr) const;
 
 	/** Returns list of components this actor is overlapping. */
 	UFUNCTION(BlueprintCallable, Category="Collision", meta=(UnsafeDuringActorConstruction="true"))
 	void GetOverlappingComponents(TArray<UPrimitiveComponent*>& OverlappingComponents) const;
 
+	/** Returns set of components this actor is overlapping. */
+	void GetOverlappingComponents(TSet<UPrimitiveComponent*>& OverlappingComponents) const;
+
 	/** 
 	 * Event when this actor bumps into a blocking object, or blocks another actor that bumps into it.
 	 * This could happen due to things like Character movement, using Set Location with 'sweep' enabled, or physics simulation.
@@ -1391,7 +1405,6 @@ public:
 	virtual void BeginDestroy() override;
 	virtual bool IsReadyForFinishDestroy() override;
 	virtual bool Rename( const TCHAR* NewName=NULL, UObject* NewOuter=NULL, ERenameFlags Flags=REN_None ) override;
-	static void AddReferencedObjects(UObject* InThis, FReferenceCollector& Collector);
 #if WITH_EDITOR
 	virtual void PreEditChange(UProperty* PropertyThatWillChange) override;
 	virtual void PostEditChangeProperty(FPropertyChangedEvent& PropertyChangedEvent) override;
@@ -2046,6 +2059,10 @@ public:
 	UFUNCTION(BlueprintCallable, Category="Actor")
 	bool IsChildActor() const;
 
+	/** Returns a list of all child actors, including children of children */
+	UFUNCTION(BlueprintCallable, Category="Actor")
+	void GetAllChildActors(TArray<AActor*>& ChildActors, bool bIncludeDescendants = true) const;
+
 	UFUNCTION(BlueprintCallable, Category="Actor")
 	UChildActorComponent* GetParentComponent() const;
 
@@ -2577,11 +2594,11 @@ public:
 	}
 
 	/**
-	 * Get a direct reference to the Components array rather than a copy with the null pointers removed.
+	 * Get a direct reference to the Components set rather than a copy with the null pointers removed.
 	 * WARNING: anything that could cause the component to change ownership or be destroyed will invalidate
-	 * this array, so use caution when iterating this list!
+	 * this array, so use caution when iterating this set!
 	 */
-	const TArray<UActorComponent*>& GetComponents() const
+	const TSet<UActorComponent*>& GetComponents() const
 	{
 		return OwnedComponents;
 	}
@@ -2615,19 +2632,22 @@ public:
 	 */
 	void UpdateAllReplicatedComponents();
 
-	/** Returns a constant reference to the replicated components array
+	/** Returns a constant reference to the replicated components set
 	 */
-	const TArray<UActorComponent*>& GetReplicatedComponents() const;
+	const TSet<UActorComponent*>& GetReplicatedComponents() const 
+	{ 
+		return ReplicatedComponents; 
+	}
 
 private:
 	/**
 	 * All ActorComponents owned by this Actor.
 	 * @see GetComponents()
 	 */
-	TArray<UActorComponent*> OwnedComponents;
+	TSet<UActorComponent*> OwnedComponents;
 
-	/** List of replicated components. */
-	TArray<UActorComponent*> ReplicatedComponents;
+	/** Set of replicated components. */
+	TSet<UActorComponent*> ReplicatedComponents;
 
 #if WITH_EDITOR
 	/** Maps natively-constructed components to properties that reference them. */
diff --git a/Engine/Source/Runtime/Engine/Classes/GameFramework/CharacterMovementComponent.h b/Engine/Source/Runtime/Engine/Classes/GameFramework/CharacterMovementComponent.h
--- a/Engine/Source/Runtime/Engine/Classes/GameFramework/CharacterMovementComponent.h
+++ b/Engine/Source/Runtime/Engine/Classes/GameFramework/CharacterMovementComponent.h
@@ -1648,7 +1648,7 @@ protected:
 
 	/** Called when the collision capsule touches another primitive component */
 	UFUNCTION()
-	virtual void CapsuleTouched(AActor* Other, UPrimitiveComponent* OtherComp, int32 OtherBodyIndex, bool bFromSweep, const FHitResult& SweepResult);
+	virtual void CapsuleTouched(UPrimitiveComponent* OverlappedComp, AActor* Other, UPrimitiveComponent* OtherComp, int32 OtherBodyIndex, bool bFromSweep, const FHitResult& SweepResult);
 
 	// Enum used to control GetPawnCapsuleExtent behavior
 	enum EShrinkCapsuleExtent
diff --git a/Engine/Source/Runtime/Engine/Classes/GameFramework/CheatManager.h b/Engine/Source/Runtime/Engine/Classes/GameFramework/CheatManager.h
--- a/Engine/Source/Runtime/Engine/Classes/GameFramework/CheatManager.h
+++ b/Engine/Source/Runtime/Engine/Classes/GameFramework/CheatManager.h
@@ -323,6 +323,10 @@ class ENGINE_API UCheatManager : public UObject
 	UFUNCTION(exec)
 	void InvertMouse();
 
+	/** Executes commands listed in CheatScript.ScriptName ini section of DefaultGame.ini */
+	UFUNCTION(exec)
+	void CheatScript(FString ScriptName);
+
 	/**
 	 * This will move the player and set their rotation to the passed in values.
 	 * This actually does the location / rotation setting.  Additionally it will set you as ghost as the level may have
diff --git a/Engine/Source/Runtime/Engine/Classes/GameFramework/MovementComponent.h b/Engine/Source/Runtime/Engine/Classes/GameFramework/MovementComponent.h
--- a/Engine/Source/Runtime/Engine/Classes/GameFramework/MovementComponent.h
+++ b/Engine/Source/Runtime/Engine/Classes/GameFramework/MovementComponent.h
@@ -54,13 +54,13 @@ class ENGINE_API UMovementComponent : public UActorComponent
 	 * If this is null at startup and bAutoRegisterUpdatedComponent is true, the owning Actor's root component will automatically be set as our UpdatedComponent at startup.
 	 * @see bAutoRegisterUpdatedComponent, SetUpdatedComponent(), UpdatedPrimitive
 	 */
-	UPROPERTY(BlueprintReadOnly, DuplicateTransient, Category=MovementComponent)
+	UPROPERTY(BlueprintReadOnly, Transient, DuplicateTransient, Category=MovementComponent)
 	USceneComponent* UpdatedComponent;
 
 	/**
 	 * UpdatedComponent, cast as a UPrimitiveComponent. May be invalid if UpdatedComponent was null or not a UPrimitiveComponent.
 	 */
-	UPROPERTY(BlueprintReadOnly, DuplicateTransient, Category=MovementComponent)
+	UPROPERTY(BlueprintReadOnly, Transient, DuplicateTransient, Category=MovementComponent)
 	UPrimitiveComponent* UpdatedPrimitive;
 
 	/**
diff --git a/Engine/Source/Runtime/Engine/Classes/Kismet/GameplayStatics.h b/Engine/Source/Runtime/Engine/Classes/Kismet/GameplayStatics.h
--- a/Engine/Source/Runtime/Engine/Classes/Kismet/GameplayStatics.h
+++ b/Engine/Source/Runtime/Engine/Classes/Kismet/GameplayStatics.h
@@ -587,6 +587,12 @@ class ENGINE_API UGameplayStatics : public UBlueprintFunctionLibrary
 	UFUNCTION(BlueprintCallable, Category="Audio", meta=(WorldContext = "WorldContextObject"))
 	static void DeactivateReverbEffect(UObject* WorldContextObject, FName TagName);
 
+	/** 
+	 * Returns the highest priority reverb settings currently active from any source (volumes or manual setting).
+	 */
+	UFUNCTION(BlueprintCallable, Category = "Audio", meta = (WorldContext = "WorldContextObject"))
+	static class UReverbEffect* GetCurrentReverbEffect(UObject* WorldContextObject);
+
 	// --- Decal functions ------------------------------
 
 	/** Spawns a decal at the given location and rotation, fire and forget. Does not replicate.
diff --git a/Engine/Source/Runtime/Engine/Classes/Kismet/KismetSystemLibrary.h b/Engine/Source/Runtime/Engine/Classes/Kismet/KismetSystemLibrary.h
--- a/Engine/Source/Runtime/Engine/Classes/Kismet/KismetSystemLibrary.h
+++ b/Engine/Source/Runtime/Engine/Classes/Kismet/KismetSystemLibrary.h
@@ -792,8 +792,8 @@ class ENGINE_API UKismetSystemLibrary : public UBlueprintFunctionLibrary
 	 * @param OutHit		Properties of the trace hit.
 	 * @return				True if there was a hit, false otherwise.
 	 */
-	UFUNCTION(BlueprintCallable, Category="Collision", meta=(bIgnoreSelf="true", WorldContext="WorldContextObject", AutoCreateRefTerm="ActorsToIgnore", DisplayName="LineTraceByChannel", Keywords="raycast"))
-	static bool LineTraceSingle_NEW(UObject* WorldContextObject, const FVector Start, const FVector End, ETraceTypeQuery TraceChannel, bool bTraceComplex, const TArray<AActor*>& ActorsToIgnore, EDrawDebugTrace::Type DrawDebugType, FHitResult& OutHit, bool bIgnoreSelf);
+	UFUNCTION(BlueprintCallable, Category="Collision", meta=(bIgnoreSelf="true", WorldContext="WorldContextObject", AutoCreateRefTerm="ActorsToIgnore", DisplayName="LineTraceByChannel", AdvancedDisplay="TraceColor,TraceHitColor,DrawTime", Keywords="raycast"))
+	static bool LineTraceSingle_NEW(UObject* WorldContextObject, const FVector Start, const FVector End, ETraceTypeQuery TraceChannel, bool bTraceComplex, const TArray<AActor*>& ActorsToIgnore, EDrawDebugTrace::Type DrawDebugType, FHitResult& OutHit, bool bIgnoreSelf, FLinearColor TraceColor = FLinearColor::Red, FLinearColor TraceHitColor = FLinearColor::Green, float DrawTime = 5.0f);
 	
 	/**
 	 * Does a collision trace along the given line and returns all hits encountered up to and including the first blocking hit.
@@ -807,8 +807,8 @@ class ENGINE_API UKismetSystemLibrary : public UBlueprintFunctionLibrary
 	 * @param OutHit		Properties of the trace hit.
 	 * @return				True if there was a blocking hit, false otherwise.
 	 */
-	UFUNCTION(BlueprintCallable, Category="Collision", meta=(bIgnoreSelf="true", WorldContext="WorldContextObject", AutoCreateRefTerm="ActorsToIgnore", DisplayName = "MultiLineTraceByChannel", Keywords="raycast"))
-	static bool LineTraceMulti_NEW(UObject* WorldContextObject, const FVector Start, const FVector End, ETraceTypeQuery TraceChannel, bool bTraceComplex, const TArray<AActor*>& ActorsToIgnore, EDrawDebugTrace::Type DrawDebugType, TArray<FHitResult>& OutHits, bool bIgnoreSelf);
+	UFUNCTION(BlueprintCallable, Category="Collision", meta=(bIgnoreSelf="true", WorldContext="WorldContextObject", AutoCreateRefTerm="ActorsToIgnore", DisplayName = "MultiLineTraceByChannel", AdvancedDisplay="TraceColor,TraceHitColor,DrawTime", Keywords="raycast"))
+	static bool LineTraceMulti_NEW(UObject* WorldContextObject, const FVector Start, const FVector End, ETraceTypeQuery TraceChannel, bool bTraceComplex, const TArray<AActor*>& ActorsToIgnore, EDrawDebugTrace::Type DrawDebugType, TArray<FHitResult>& OutHits, bool bIgnoreSelf, FLinearColor TraceColor = FLinearColor::Red, FLinearColor TraceHitColor = FLinearColor::Green, float DrawTime = 5.0f);
 
 	/**
 	 * Sweeps a sphere along the given line and returns the first blocking hit encountered.
@@ -920,8 +920,8 @@ class ENGINE_API UKismetSystemLibrary : public UBlueprintFunctionLibrary
 	 * @param OutHit		Properties of the trace hit.
 	 * @return				True if there was a hit, false otherwise.
 	 */
-	UFUNCTION(BlueprintCallable, Category="Collision", meta=(bIgnoreSelf="true", WorldContext="WorldContextObject", AutoCreateRefTerm="ActorsToIgnore", DisplayName = "LineTraceForObjects", Keywords="raycast"))
-	static bool LineTraceSingleForObjects(UObject* WorldContextObject, const FVector Start, const FVector End, const TArray<TEnumAsByte<EObjectTypeQuery> > & ObjectTypes, bool bTraceComplex, const TArray<AActor*>& ActorsToIgnore, EDrawDebugTrace::Type DrawDebugType, FHitResult& OutHit, bool bIgnoreSelf);
+	UFUNCTION(BlueprintCallable, Category="Collision", meta=(bIgnoreSelf="true", WorldContext="WorldContextObject", AutoCreateRefTerm="ActorsToIgnore", DisplayName = "LineTraceForObjects", AdvancedDisplay="TraceColor,TraceHitColor,DrawTime", Keywords="raycast"))
+	static bool LineTraceSingleForObjects(UObject* WorldContextObject, const FVector Start, const FVector End, const TArray<TEnumAsByte<EObjectTypeQuery> > & ObjectTypes, bool bTraceComplex, const TArray<AActor*>& ActorsToIgnore, EDrawDebugTrace::Type DrawDebugType, FHitResult& OutHit, bool bIgnoreSelf, FLinearColor TraceColor = FLinearColor::Red, FLinearColor TraceHitColor = FLinearColor::Green, float DrawTime = 5.0f );
 	
 	/**
 	 * Does a collision trace along the given line and returns all hits encountered.
@@ -935,8 +935,8 @@ class ENGINE_API UKismetSystemLibrary : public UBlueprintFunctionLibrary
 	 * @param OutHit		Properties of the trace hit.
 	 * @return				True if there was a hit, false otherwise.
 	 */
-	UFUNCTION(BlueprintCallable, Category="Collision", meta=(bIgnoreSelf="true", WorldContext="WorldContextObject", AutoCreateRefTerm="ActorsToIgnore", DisplayName = "MultiLineTraceForObjects", Keywords="raycast"))
-	static bool LineTraceMultiForObjects(UObject* WorldContextObject, const FVector Start, const FVector End, const TArray<TEnumAsByte<EObjectTypeQuery> > & ObjectTypes, bool bTraceComplex, const TArray<AActor*>& ActorsToIgnore, EDrawDebugTrace::Type DrawDebugType, TArray<FHitResult>& OutHits, bool bIgnoreSelf);
+	UFUNCTION(BlueprintCallable, Category="Collision", meta=(bIgnoreSelf="true", WorldContext="WorldContextObject", AutoCreateRefTerm="ActorsToIgnore", DisplayName = "MultiLineTraceForObjects", AdvancedDisplay="TraceColor,TraceHitColor,DrawTime", Keywords="raycast"))
+	static bool LineTraceMultiForObjects(UObject* WorldContextObject, const FVector Start, const FVector End, const TArray<TEnumAsByte<EObjectTypeQuery> > & ObjectTypes, bool bTraceComplex, const TArray<AActor*>& ActorsToIgnore, EDrawDebugTrace::Type DrawDebugType, TArray<FHitResult>& OutHits, bool bIgnoreSelf, FLinearColor TraceColor = FLinearColor::Red, FLinearColor TraceHitColor = FLinearColor::Green, float DrawTime = 5.0f);
 
 	/**
 	 * Sweeps a sphere along the given line and returns the first hit encountered.
@@ -1049,8 +1049,8 @@ class ENGINE_API UKismetSystemLibrary : public UBlueprintFunctionLibrary
 	 * @param OutHit		Properties of the trace hit.
 	 * @return				True if there was a hit, false otherwise.
 	 */
-	UFUNCTION(BlueprintCallable, Category="Collision", meta=(DeprecatedFunction, DeprecationMessage = "Use new SingleLineTraceByChannel", bIgnoreSelf="true", WorldContext="WorldContextObject", AutoCreateRefTerm="ActorsToIgnore", DisplayName = "SingleLineTraceByChannelDeprecated"))
-	static bool LineTraceSingle_DEPRECATED(UObject* WorldContextObject, const FVector Start, const FVector End, ECollisionChannel TraceChannel, bool bTraceComplex, const TArray<AActor*>& ActorsToIgnore, EDrawDebugTrace::Type DrawDebugType, FHitResult& OutHit, bool bIgnoreSelf);
+	UFUNCTION(BlueprintCallable, Category="Collision", meta=(DeprecatedFunction, DeprecationMessage = "Use new SingleLineTraceByChannel", bIgnoreSelf="true", WorldContext="WorldContextObject", AutoCreateRefTerm="ActorsToIgnore", AdvancedDisplay="TraceColor,TraceHitColor,DrawTime", DisplayName = "SingleLineTraceByChannelDeprecated"))
+	static bool LineTraceSingle_DEPRECATED(UObject* WorldContextObject, const FVector Start, const FVector End, ECollisionChannel TraceChannel, bool bTraceComplex, const TArray<AActor*>& ActorsToIgnore, EDrawDebugTrace::Type DrawDebugType, FHitResult& OutHit, bool bIgnoreSelf, FLinearColor TraceColor = FLinearColor::Red, FLinearColor TraceHitColor = FLinearColor::Green, float DrawTime = 5.0f);
 	
 	/**
 	 * Does a collision trace along the given line and returns all hits encountered up to and including the first blocking hit.
@@ -1064,8 +1064,8 @@ class ENGINE_API UKismetSystemLibrary : public UBlueprintFunctionLibrary
 	 * @param OutHit		Properties of the trace hit.
 	 * @return				True if there was a hit, false otherwise.
 	 */
-	UFUNCTION(BlueprintCallable, Category="Collision", meta=(DeprecatedFunction, DeprecationMessage = "Use new MultiLineTraceByChannel", bIgnoreSelf="true", WorldContext="WorldContextObject", AutoCreateRefTerm="ActorsToIgnore", DisplayName = "MultiLineTraceByChannelDeprecated"))
-	static bool LineTraceMulti_DEPRECATED(UObject* WorldContextObject, const FVector Start, const FVector End, ECollisionChannel TraceChannel, bool bTraceComplex, const TArray<AActor*>& ActorsToIgnore, EDrawDebugTrace::Type DrawDebugType, TArray<FHitResult>& OutHits, bool bIgnoreSelf);
+	UFUNCTION(BlueprintCallable, Category="Collision", meta=(DeprecatedFunction, DeprecationMessage = "Use new MultiLineTraceByChannel", bIgnoreSelf="true", WorldContext="WorldContextObject", AutoCreateRefTerm="ActorsToIgnore", AdvancedDisplay="TraceColor,TraceHitColor,DrawTime", DisplayName = "MultiLineTraceByChannelDeprecated"))
+	static bool LineTraceMulti_DEPRECATED(UObject* WorldContextObject, const FVector Start, const FVector End, ECollisionChannel TraceChannel, bool bTraceComplex, const TArray<AActor*>& ActorsToIgnore, EDrawDebugTrace::Type DrawDebugType, TArray<FHitResult>& OutHits, bool bIgnoreSelf, FLinearColor TraceColor = FLinearColor::Red, FLinearColor TraceHitColor = FLinearColor::Green, float DrawTime = 5.0f);
 
 	/**
 	 * Sweeps a sphere along the given line and returns the first blocking hit encountered.
@@ -1144,8 +1144,8 @@ class ENGINE_API UKismetSystemLibrary : public UBlueprintFunctionLibrary
 	 * @param OutHit		Properties of the trace hit.
 	 * @return				True if there was a hit, false otherwise.
 	 */
-	UFUNCTION(BlueprintCallable, Category="Collision", meta=(DeprecatedFunction, DeprecationMessage = "Use new SingleLineTraceForObjects", bIgnoreSelf="true", WorldContext="WorldContextObject", AutoCreateRefTerm="ActorsToIgnore", DisplayName = "SingleLineTraceByObjectDeprecated"))
-	static bool LineTraceSingleByObject_DEPRECATED(UObject* WorldContextObject, const FVector Start, const FVector End, const TArray<TEnumAsByte<ECollisionChannel> > & ObjectsToTrace, bool bTraceComplex, const TArray<AActor*>& ActorsToIgnore, EDrawDebugTrace::Type DrawDebugType, FHitResult& OutHit, bool bIgnoreSelf);
+	UFUNCTION(BlueprintCallable, Category="Collision", meta=(DeprecatedFunction, DeprecationMessage = "Use new SingleLineTraceForObjects", bIgnoreSelf="true", WorldContext="WorldContextObject", AutoCreateRefTerm="ActorsToIgnore", AdvancedDisplay="TraceColor,TraceHitColor,DrawTime", DisplayName = "SingleLineTraceByObjectDeprecated"))
+	static bool LineTraceSingleByObject_DEPRECATED(UObject* WorldContextObject, const FVector Start, const FVector End, const TArray<TEnumAsByte<ECollisionChannel> > & ObjectsToTrace, bool bTraceComplex, const TArray<AActor*>& ActorsToIgnore, EDrawDebugTrace::Type DrawDebugType, FHitResult& OutHit, bool bIgnoreSelf, FLinearColor TraceColor = FLinearColor::Red, FLinearColor TraceHitColor = FLinearColor::Green, float DrawTime = 5.0f);
 	
 	/**
 	 * Does a collision trace along the given line and returns all hits encountered.
@@ -1159,8 +1159,8 @@ class ENGINE_API UKismetSystemLibrary : public UBlueprintFunctionLibrary
 	 * @param OutHit		Properties of the trace hit.
 	 * @return				True if there was a hit, false otherwise.
 	 */
-	UFUNCTION(BlueprintCallable, Category="Collision", meta=(DeprecatedFunction, DeprecationMessage = "Use new MultiLineTraceForObjects", bIgnoreSelf="true", WorldContext="WorldContextObject", AutoCreateRefTerm="ActorsToIgnore", DisplayName = "MultiLineTraceByObjectDeprecated"))
-	static bool LineTraceMultiByObject_DEPRECATED(UObject* WorldContextObject, const FVector Start, const FVector End, const TArray<TEnumAsByte<ECollisionChannel> > & ObjectsToTrace, bool bTraceComplex, const TArray<AActor*>& ActorsToIgnore, EDrawDebugTrace::Type DrawDebugType, TArray<FHitResult>& OutHits, bool bIgnoreSelf);
+	UFUNCTION(BlueprintCallable, Category="Collision", meta=(DeprecatedFunction, DeprecationMessage = "Use new MultiLineTraceForObjects", bIgnoreSelf="true", WorldContext="WorldContextObject", AutoCreateRefTerm="ActorsToIgnore", AdvancedDisplay="TraceColor,TraceHitColor,DrawTime", DisplayName = "MultiLineTraceByObjectDeprecated"))
+	static bool LineTraceMultiByObject_DEPRECATED(UObject* WorldContextObject, const FVector Start, const FVector End, const TArray<TEnumAsByte<ECollisionChannel> > & ObjectsToTrace, bool bTraceComplex, const TArray<AActor*>& ActorsToIgnore, EDrawDebugTrace::Type DrawDebugType, TArray<FHitResult>& OutHits, bool bIgnoreSelf, FLinearColor TraceColor = FLinearColor::Red, FLinearColor TraceHitColor = FLinearColor::Green, float DrawTime = 5.0f);
 
 	/**
 	 * Sweeps a sphere along the given line and returns the first hit encountered.
diff --git a/Engine/Source/Runtime/Engine/Classes/Particles/ParticleSystem.h b/Engine/Source/Runtime/Engine/Classes/Particles/ParticleSystem.h
--- a/Engine/Source/Runtime/Engine/Classes/Particles/ParticleSystem.h
+++ b/Engine/Source/Runtime/Engine/Classes/Particles/ParticleSystem.h
@@ -353,7 +353,6 @@ public:
 #endif // WITH_EDITOR
 	virtual void PreSave() override;
 	virtual void PostLoad() override;
-	virtual SIZE_T GetResourceSize(EResourceSizeMode::Type Mode) override;
 	virtual void GetAssetRegistryTags(TArray<FAssetRegistryTag>& OutTags) const override;
 	bool UsesCPUCollision() const;
 	virtual bool CanBeClusterRoot() const override;
diff --git a/Engine/Source/Runtime/Engine/Classes/PhysicsEngine/AggregateGeom.h b/Engine/Source/Runtime/Engine/Classes/PhysicsEngine/AggregateGeom.h
--- a/Engine/Source/Runtime/Engine/Classes/PhysicsEngine/AggregateGeom.h
+++ b/Engine/Source/Runtime/Engine/Classes/PhysicsEngine/AggregateGeom.h
@@ -19,16 +19,16 @@ struct ENGINE_API FKAggregateGeom
 {
 	GENERATED_USTRUCT_BODY()
 
-	UPROPERTY(EditAnywhere, editfixedsize, Category = KAggregateGeom)
+	UPROPERTY(EditAnywhere, editfixedsize, Category = "Aggregate Geometry", meta = (DisplayName = "Spheres"))
 	TArray<FKSphereElem> SphereElems;
 
-	UPROPERTY(EditAnywhere, editfixedsize, Category = KAggregateGeom)
+	UPROPERTY(EditAnywhere, editfixedsize, Category = "Aggregate Geometry", meta = (DisplayName = "Boxes"))
 	TArray<FKBoxElem> BoxElems;
 
-	UPROPERTY(EditAnywhere, editfixedsize, Category = KAggregateGeom)
+	UPROPERTY(EditAnywhere, editfixedsize, Category = "Aggregate Geometry", meta = (DisplayName = "Capsules"))
 	TArray<FKSphylElem> SphylElems;
 
-	UPROPERTY(EditAnywhere, editfixedsize, Category = KAggregateGeom)
+	UPROPERTY(EditAnywhere, editfixedsize, Category = "Aggregate Geometry", meta = (DisplayName = "Convex Elements"))
 	TArray<FKConvexElem> ConvexElems;
 
 	class FKConvexGeomRenderInfo* RenderInfo;
diff --git a/Engine/Source/Runtime/Engine/Classes/PhysicsEngine/BodySetup.h b/Engine/Source/Runtime/Engine/Classes/PhysicsEngine/BodySetup.h
--- a/Engine/Source/Runtime/Engine/Classes/PhysicsEngine/BodySetup.h
+++ b/Engine/Source/Runtime/Engine/Classes/PhysicsEngine/BodySetup.h
@@ -27,7 +27,7 @@ class UBodySetup : public UObject
 	GENERATED_UCLASS_BODY()
 
 	/** Simplified collision representation of this  */
-	UPROPERTY()
+	UPROPERTY(EditAnywhere, Category = BodySetup, meta=(DisplayName = "Primitives"))
 	struct FKAggregateGeom AggGeom;
 
 	/** Used in the PhysicsAsset case. Associates this Body with Bone in a skeletal mesh. */
