Copying //UE4/Dev-Framework to Dev-Main (//UE4/Dev-Main) @ 2926658

diff --git a/Engine/Source/Runtime/Engine/Classes/PhysicsEngine/BoxElem.h b/Engine/Source/Runtime/Engine/Classes/PhysicsEngine/BoxElem.h
--- a/Engine/Source/Runtime/Engine/Classes/PhysicsEngine/BoxElem.h
+++ b/Engine/Source/Runtime/Engine/Classes/PhysicsEngine/BoxElem.h
@@ -16,20 +16,24 @@ struct FKBoxElem : public FKShapeElem
 	UPROPERTY()
 	FMatrix TM_DEPRECATED;
 
-	UPROPERTY(Category=KBoxElem, VisibleAnywhere)
+	/** Position of the box's origin */
+	UPROPERTY(Category=Box, EditAnywhere)
 	FVector Center;
 
-	UPROPERTY(Category=KBoxElem, VisibleAnywhere)
+	/** Orientation of the box */
+	UPROPERTY(Category= Box, EditAnywhere)
 	FQuat Orientation;
 
-	UPROPERTY(Category=KBoxElem, VisibleAnywhere)
+	/** Extent of the box along the y-axis */
+	UPROPERTY(Category= Box, EditAnywhere, meta =(DisplayName = "X Extent"))
 	float X;
 
-	UPROPERTY(Category=KBoxElem, VisibleAnywhere)
+	/** Extent of the box along the y-axis */
+	UPROPERTY(Category= Box, EditAnywhere, meta = (DisplayName = "Y Extent"))
 	float Y;
 
-	/** length (not radius) */
-	UPROPERTY(Category=KBoxElem, VisibleAnywhere)
+	/** Extent of the box along the z-axis */
+	UPROPERTY(Category= Box, EditAnywhere, meta = (DisplayName = "Z Extent"))
 	float Z;
 
 
diff --git a/Engine/Source/Runtime/Engine/Classes/PhysicsEngine/PhysicsCollisionHandler.h b/Engine/Source/Runtime/Engine/Classes/PhysicsEngine/PhysicsCollisionHandler.h
--- a/Engine/Source/Runtime/Engine/Classes/PhysicsEngine/PhysicsCollisionHandler.h
+++ b/Engine/Source/Runtime/Engine/Classes/PhysicsEngine/PhysicsCollisionHandler.h
@@ -34,22 +34,8 @@ class ENGINE_API UPhysicsCollisionHandler : public UObject
 		return !HasAnyFlags(RF_ClassDefaultObject) ? CastChecked<UWorld>(GetOuter()) : nullptr;
 	}
 
-	/** Gives game-specific ability to handle all physics collisions in one place. This is a good place to play sounds and spawn effects, as it does not require special object-specific code. */
-	DEPRECATED(4.8, "Please call HandlePhysicsCollisions_AssumesLocked and make sure you obtain the appropriate PhysX scene locks")
-	virtual void HandlePhysicsCollisions(const TArray<FCollisionNotifyInfo>& PendingCollisionNotifies)
-	{
-		HandlePhysicsCollisions_AssumesLocked(PendingCollisionNotifies);
-	}
-
-	/** Gives game-specific ability to handle all physics collisions in one place. This is a good place to play sounds and spawn effects, as it does not require special object-specific code. */
-	virtual void HandlePhysicsCollisions_AssumesLocked(const TArray<FCollisionNotifyInfo>& PendingCollisionNotifies);
-
-	/** Handle a single */
-	DEPRECATED(4.8, "Please call DefaultHandleCollision_AssumesLocked and make sure you obtain the appropriate PhysX scene locks")
-	void DefaultHandleCollision(const FRigidBodyCollisionInfo& MyInfo, const FRigidBodyCollisionInfo& OtherInfo, const FCollisionImpactData& RigidCollisionData)
-	{
-		DefaultHandleCollision_AssumesLocked(MyInfo, OtherInfo, RigidCollisionData);
-	}
+	/** Gives game-specific ability to handle and filter all physics collisions in one place. This is a good place to play sounds and spawn effects, as it does not require special object-specific code. */
+	virtual void HandlePhysicsCollisions_AssumesLocked(TArray<FCollisionNotifyInfo>& PendingCollisionNotifies);
 
 	/** Handle a single */
 	void DefaultHandleCollision_AssumesLocked(const FRigidBodyCollisionInfo& MyInfo, const FRigidBodyCollisionInfo& OtherInfo, const FCollisionImpactData& RigidCollisionData);
diff --git a/Engine/Source/Runtime/Engine/Classes/PhysicsEngine/PhysicsSettings.h b/Engine/Source/Runtime/Engine/Classes/PhysicsEngine/PhysicsSettings.h
--- a/Engine/Source/Runtime/Engine/Classes/PhysicsEngine/PhysicsSettings.h
+++ b/Engine/Source/Runtime/Engine/Classes/PhysicsEngine/PhysicsSettings.h
@@ -144,6 +144,18 @@ class ENGINE_API UPhysicsSettings : public UDeveloperSettings
 	UPROPERTY(config, EditAnywhere, Category = Simulation)
 	float MaxDepenetrationVelocity;
 
+	/** Contact offset multiplier. When creating a physics shape we look at its bounding volume and multiply its minimum value by this multiplier. A bigger number will generate contact points earlier which results in higher stability at the cost of performance. */
+	UPROPERTY(config, EditAnywhere, Category = Simulation, meta = (ClampMin = "0.001", UIMin = "0.001"))
+	float ContactOffsetMultiplier;
+
+	/** Min Contact offset. */
+	UPROPERTY(config, EditAnywhere, Category = Simulation, meta = (ClampMin = "0.0001", UIMin = "0.0001"))
+	float MinContactOffset;
+
+	/** Max Contact offset. */
+	UPROPERTY(config, EditAnywhere, Category = Simulation, meta = (ClampMin = "0.001", UIMin = "0.001"))
+	float MaxContactOffset;
+
 	/**
 	*  If true, simulate physics for this component on a dedicated server.
 	*  This should be set if simulating physics and replicating with a dedicated server.
diff --git a/Engine/Source/Runtime/Engine/Classes/PhysicsEngine/SphereElem.h b/Engine/Source/Runtime/Engine/Classes/PhysicsEngine/SphereElem.h
--- a/Engine/Source/Runtime/Engine/Classes/PhysicsEngine/SphereElem.h
+++ b/Engine/Source/Runtime/Engine/Classes/PhysicsEngine/SphereElem.h
@@ -15,10 +15,12 @@ struct FKSphereElem : public FKShapeElem
 	UPROPERTY()
 	FMatrix TM_DEPRECATED;
 
-	UPROPERTY(Category=KSphereElem, VisibleAnywhere)
+	/** Position of the sphere's origin */
+	UPROPERTY(Category=Sphere, EditAnywhere)
 	FVector Center;
 
-	UPROPERTY(Category=KSphereElem, VisibleAnywhere)
+	/** Radius of the sphere */
+	UPROPERTY(Category=Sphere, EditAnywhere)
 	float Radius;
 
 	FKSphereElem() 
diff --git a/Engine/Source/Runtime/Engine/Classes/PhysicsEngine/SphylElem.h b/Engine/Source/Runtime/Engine/Classes/PhysicsEngine/SphylElem.h
--- a/Engine/Source/Runtime/Engine/Classes/PhysicsEngine/SphylElem.h
+++ b/Engine/Source/Runtime/Engine/Classes/PhysicsEngine/SphylElem.h
@@ -18,17 +18,20 @@ struct FKSphylElem : public FKShapeElem
 	UPROPERTY()
 	FMatrix TM_DEPRECATED;
 
-	UPROPERTY(Category=KBoxElem, VisibleAnywhere)
+	/** Position of the capsule's origin */
+	UPROPERTY(Category=Capsule, EditAnywhere)
 	FVector Center;
 
-	UPROPERTY(Category=KBoxElem, VisibleAnywhere)
+	/** Orientation of the capsule */
+	UPROPERTY(Category= Capsule, EditAnywhere)
 	FQuat Orientation;
 
-	UPROPERTY(Category=KSphylElem, VisibleAnywhere)
+	/** Radius of the capsule */
+	UPROPERTY(Category= Capsule, EditAnywhere)
 	float Radius;
 
 	/** This is of line-segment ie. add Radius to both ends to find total length. */
-	UPROPERTY(Category=KSphylElem, VisibleAnywhere)
+	UPROPERTY(Category= Capsule, EditAnywhere)
 	float Length;
 
 	FKSphylElem()
diff --git a/Engine/Source/Runtime/Engine/Classes/Sound/SoundConcurrency.h b/Engine/Source/Runtime/Engine/Classes/Sound/SoundConcurrency.h
--- a/Engine/Source/Runtime/Engine/Classes/Sound/SoundConcurrency.h
+++ b/Engine/Source/Runtime/Engine/Classes/Sound/SoundConcurrency.h
@@ -22,7 +22,7 @@ namespace EMaxConcurrentResolutionRule
 		/** When Max Concurrent sounds are active stop the furthest sound.  If all sounds are the same distance then stop the oldest. */
 		StopFarthestThenOldest,
 
-		/** Stop the lowest volume-weighted priority in the group. */
+		/** Stop the lowest volume-weighted priority sound in the group. If all sounds are the same priority, then it will stop the oldest sound in the group. */
 		StopLowestPriority,
 
 		/** Stop the sound that is quietest in the group. */
diff --git a/Engine/Source/Runtime/Engine/Private/AI/Navigation/NavModifierComponent.cpp b/Engine/Source/Runtime/Engine/Private/AI/Navigation/NavModifierComponent.cpp
--- a/Engine/Source/Runtime/Engine/Private/AI/Navigation/NavModifierComponent.cpp
+++ b/Engine/Source/Runtime/Engine/Private/AI/Navigation/NavModifierComponent.cpp
@@ -16,22 +16,103 @@ void UNavModifierComponent::CalcAndCacheBounds() const
 	const AActor* MyOwner = GetOwner();
 	if (MyOwner)
 	{
-		const float Radius = MyOwner->GetSimpleCollisionRadius();
-		const float HalfHeght = MyOwner->GetSimpleCollisionHalfHeight();
-		
-		ObstacleExtent = FVector(Radius, Radius, HalfHeght);
-		if (ObstacleExtent.IsNearlyZero())
+		TInlineComponentArray<UPrimitiveComponent*> PrimComponents;
+		MyOwner->GetComponents(PrimComponents);
+
+		Bounds = FBox(ForceInit);
+		ComponentBounds.Reset();
+		for (int32 Idx = 0; Idx < PrimComponents.Num(); Idx++)
 		{
-			const FVector ActorScale3D = MyOwner->GetActorScale3D();
-			ObstacleExtent = FailsafeExtent * ActorScale3D;
+			UPrimitiveComponent* PrimComp = PrimComponents[Idx];
+			if (PrimComp && PrimComp->IsRegistered() && PrimComp->IsCollisionEnabled())
+			{
+				UBodySetup* BodySetup = PrimComp->GetBodySetup();
+				if (BodySetup)
+				{
+					FTransform ParentTM = PrimComp->ComponentToWorld;
+					const FVector Scale3D = ParentTM.GetScale3D();
+					ParentTM.RemoveScaling();
+					Bounds += PrimComp->Bounds.GetBox();
+
+					for (int32 SphereIdx = 0; SphereIdx < BodySetup->AggGeom.SphereElems.Num(); SphereIdx++)
+					{
+						const FKSphereElem& ElemInfo = BodySetup->AggGeom.SphereElems[SphereIdx];
+						FTransform ElemTM = ElemInfo.GetTransform();
+						ElemTM.ScaleTranslation(Scale3D);
+						ElemTM *= ParentTM;
+
+						const FBox SphereBounds = FBox::BuildAABB(ElemTM.GetLocation(), ElemInfo.Radius * Scale3D);
+						ComponentBounds.Add(FRotatedBox(SphereBounds, ElemTM.GetRotation()));
+					}
+
+					for (int32 BoxIdx = 0; BoxIdx < BodySetup->AggGeom.BoxElems.Num(); BoxIdx++)
+					{
+						const FKBoxElem& ElemInfo = BodySetup->AggGeom.BoxElems[BoxIdx];
+						FTransform ElemTM = ElemInfo.GetTransform();
+						ElemTM.ScaleTranslation(Scale3D);
+						ElemTM *= ParentTM;
+
+						const FBox BoxBounds = FBox::BuildAABB(ElemTM.GetLocation(), FVector(ElemInfo.X, ElemInfo.Y, ElemInfo.Z) * Scale3D * 0.5f);
+						ComponentBounds.Add(FRotatedBox(BoxBounds, ElemTM.GetRotation()));
+					}
+
+					for (int32 SphylIdx = 0; SphylIdx < BodySetup->AggGeom.SphylElems.Num(); SphylIdx++)
+					{
+						const FKSphylElem& ElemInfo = BodySetup->AggGeom.SphylElems[SphylIdx];
+						FTransform ElemTM = ElemInfo.GetTransform();
+						ElemTM.ScaleTranslation(Scale3D);
+						ElemTM *= ParentTM;
+
+						const FBox SphylBounds = FBox::BuildAABB(ElemTM.GetLocation(), FVector(ElemInfo.Radius, ElemInfo.Radius, ElemInfo.Length) * Scale3D);
+						ComponentBounds.Add(FRotatedBox(SphylBounds, ElemTM.GetRotation()));
+					}
+
+					for (int32 ConvexIdx = 0; ConvexIdx < BodySetup->AggGeom.ConvexElems.Num(); ConvexIdx++)
+					{
+						const FKConvexElem& ElemInfo = BodySetup->AggGeom.ConvexElems[ConvexIdx];
+						FTransform ElemTM = ElemInfo.GetTransform();
+						ElemTM.ScaleTranslation(Scale3D);
+						ElemTM *= FTransform(ParentTM.GetLocation());
+
+						const FBox ConvexBounds = ElemInfo.CalcAABB(ElemTM, Scale3D);
+						ComponentBounds.Add(FRotatedBox(ConvexBounds, ElemTM.GetRotation() * ParentTM.GetRotation()));
+					}
+				}
+			}
 		}
 
-		Bounds = FBox::BuildAABB(MyOwner->GetActorLocation(), ObstacleExtent);
+		if (ComponentBounds.Num() == 0)
+		{
+			Bounds = FBox::BuildAABB(MyOwner->GetActorLocation(), FailsafeExtent);
+			ComponentBounds.Add(FRotatedBox(Bounds, MyOwner->GetActorQuat()));
+		}
+		else
+		{
+			for (int32 Idx = 0; Idx < ComponentBounds.Num(); Idx++)
+			{
+				const FVector BoxOrigin = ComponentBounds[Idx].Box.GetCenter();
+				const FVector BoxExtent = ComponentBounds[Idx].Box.GetExtent();
+				
+				const FVector NavModBoxOrigin = FTransform(ComponentBounds[Idx].Quat).InverseTransformPosition(BoxOrigin);
+				ComponentBounds[Idx].Box = FBox::BuildAABB(NavModBoxOrigin, BoxExtent);
+			}
+		}
 	}
 }
 
 void UNavModifierComponent::GetNavigationData(FNavigationRelevantData& Data) const
 {
-	CalcAndCacheBounds();
-	Data.Modifiers.Add(FAreaNavModifier(Bounds, FTransform::Identity, AreaClass));
+	for (int32 Idx = 0; Idx < ComponentBounds.Num(); Idx++)
+	{
+		Data.Modifiers.Add(FAreaNavModifier(ComponentBounds[Idx].Box, FTransform(ComponentBounds[Idx].Quat), AreaClass));
+	}
+}
+
+void UNavModifierComponent::SetAreaClass(TSubclassOf<UNavArea> NewAreaClass)
+{
+	if (AreaClass != NewAreaClass)
+	{
+		AreaClass = NewAreaClass;
+		RefreshNavigationModifiers();
+	}
 }
diff --git a/Engine/Source/Runtime/Engine/Private/AI/Navigation/NavRelevantComponent.cpp b/Engine/Source/Runtime/Engine/Private/AI/Navigation/NavRelevantComponent.cpp
--- a/Engine/Source/Runtime/Engine/Private/AI/Navigation/NavRelevantComponent.cpp
+++ b/Engine/Source/Runtime/Engine/Private/AI/Navigation/NavRelevantComponent.cpp
@@ -49,11 +49,6 @@ void UNavRelevantComponent::OnRegister()
 			}
 		}
 	}
-	else if (!bBoundsInitialized)
-	{
-		bBoundsInitialized = true;
-		UpdateNavigationBounds();
-	}
 
 	UNavigationSystem::OnComponentRegistered(this);
 }
@@ -67,6 +62,12 @@ void UNavRelevantComponent::OnUnregister()
 
 FBox UNavRelevantComponent::GetNavigationBounds() const
 {
+	if (!bBoundsInitialized)
+	{
+		CalcAndCacheBounds();
+		bBoundsInitialized = true;
+	}
+
 	return Bounds;
 }
 
@@ -78,6 +79,7 @@ bool UNavRelevantComponent::IsNavigationRelevant() const
 void UNavRelevantComponent::UpdateNavigationBounds()
 {
 	CalcAndCacheBounds();
+	bBoundsInitialized = true;
 }
 
 UObject* UNavRelevantComponent::GetNavigationParent() const
diff --git a/Engine/Source/Runtime/Engine/Private/AI/Navigation/NavTestRenderingComponent.cpp b/Engine/Source/Runtime/Engine/Private/AI/Navigation/NavTestRenderingComponent.cpp
--- a/Engine/Source/Runtime/Engine/Private/AI/Navigation/NavTestRenderingComponent.cpp
+++ b/Engine/Source/Runtime/Engine/Private/AI/Navigation/NavTestRenderingComponent.cpp
@@ -58,6 +58,8 @@ public:
 		bShowBestPath = NavTestActor->bShowBestPath;
 		bShowDiff = NavTestActor->bShowDiffWithPreviousStep;
 
+		ClosestWallLocation = NavTestActor->bDrawDistanceToWall ? NavTestActor->ClosestWallLocation : FNavigationSystem::InvalidLocation;
+
 		GatherPathPoints();
 		GatherPathStep();
 	}
@@ -99,6 +101,7 @@ public:
 					const FVector ActorLocation = NavTestActor->GetActorLocation();
 					const FVector ProjectedLocation = NavTestActor->ProjectedLocation + NavMeshDrawOffset;
 					const FColor ProjectedColor = NavTestActor->bProjectedLocationValid ? FColor(0, 255, 0, 120) : FColor(255, 0, 0, 120);
+					const FColor ClosestWallColor = FColorList::Orange;
 					const FVector BoxExtent(20, 20, 20);
 
 					FMaterialRenderProxy* const ColoredMeshInstance = new(FMemStack::Get()) FColoredMaterialRenderProxy(GEngine->DebugMeshMaterial->GetRenderProxy(false), ProjectedColor);
@@ -113,6 +116,11 @@ public:
 
 					// draw query extent
 					DrawWireBox(PDI, FBox(ActorLocation - NavTestActor->QueryingExtent, ActorLocation + NavTestActor->QueryingExtent), FColor::Blue, false);
+
+					if (FNavigationSystem::IsValidLocation(ClosestWallLocation))
+					{
+						PDI->DrawLine(ClosestWallLocation, ActorLocation, ClosestWallColor, SDPG_World, 2.5);
+					}
 				}
 
 				// draw path
@@ -418,6 +426,8 @@ private:
 	TSet<FNodeDebugData> NodeDebug;
 	FSetElementId BestNodeId;
 
+	FVector ClosestWallLocation;
+
 	uint32 bShowBestPath : 1;
 	uint32 bShowNodePool : 1;
 	uint32 bShowDiff : 1;
diff --git a/Engine/Source/Runtime/Engine/Private/AI/Navigation/NavigationData.cpp b/Engine/Source/Runtime/Engine/Private/AI/Navigation/NavigationData.cpp
--- a/Engine/Source/Runtime/Engine/Private/AI/Navigation/NavigationData.cpp
+++ b/Engine/Source/Runtime/Engine/Private/AI/Navigation/NavigationData.cpp
@@ -120,7 +120,6 @@ ANavigationData::ANavigationData(const FObjectInitializer& ObjectInitializer)
 	, FindPathImplementation(NULL)
 	, FindHierarchicalPathImplementation(NULL)
 	, bRegistered(false)
-	, bWantsUpdate(true)
 	, NavDataUniqueID(GetNextUniqueID())
 {
 	PrimaryActorTick.bCanEverTick = true;
@@ -180,12 +179,6 @@ void ANavigationData::PostInitializeComponents()
 	{
 		CleanUpAndMarkPendingKill();
 	}
-	else
-	{
-		// note: this is not a final fix for world composition's issues with navmesh generation
-		// but it's good for now, and navmesh creation is going to get a face-lift soon anyway
-		bWantsUpdate |= MyWorld->GetWorldSettings()->bEnableWorldComposition;
-	}
 }
 
 void ANavigationData::PostLoad() 
@@ -338,11 +331,13 @@ void ANavigationData::InstantiateAndRegisterRenderingComponent()
 		if (RenderingComp)
 		{
 			// rename the old rendering component out of the way
-			RenderingComp->Rename(NULL, NULL, REN_DontCreateRedirectors | REN_ForceGlobalUnique | REN_DoNotDirty | REN_NonTransactional);
+			RenderingComp->Rename(NULL, NULL, REN_DontCreateRedirectors | REN_ForceGlobalUnique | REN_DoNotDirty | REN_NonTransactional | REN_ForceNoResetLoaders);
 		}
 
 		RenderingComp = ConstructRenderingComponent();
-		if (RenderingComp != NULL)
+
+		UWorld* World = GetWorld();
+		if (World && World->bIsWorldInitialized && RenderingComp)
 		{
 			RenderingComp->RegisterComponent();
 		}
@@ -505,14 +500,10 @@ void ANavigationData::TickAsyncBuild(float DeltaSeconds)
 
 void ANavigationData::RebuildDirtyAreas(const TArray<FNavigationDirtyArea>& DirtyAreas)
 {
-	// the 'bWantsUpdate' mechanics allows us to skip first requested update after data is loaded
-	// Can be also used to manually control navigation rebuilding, by for example forcing bWantsUpdate 
-	//	to false in Tick function effectively suppressing rebuilding
-	if (bWantsUpdate == true && NavDataGenerator.IsValid())
+	if (NavDataGenerator.IsValid())
 	{
 		NavDataGenerator->RebuildDirtyAreas(DirtyAreas);
 	}
-	bWantsUpdate = true;
 }
 
 TArray<FBox> ANavigationData::GetNavigableBounds() const
diff --git a/Engine/Source/Runtime/Engine/Private/AI/Navigation/NavigationSystem.cpp b/Engine/Source/Runtime/Engine/Private/AI/Navigation/NavigationSystem.cpp
--- a/Engine/Source/Runtime/Engine/Private/AI/Navigation/NavigationSystem.cpp
+++ b/Engine/Source/Runtime/Engine/Private/AI/Navigation/NavigationSystem.cpp
@@ -708,21 +708,6 @@ void UNavigationSystem::OnWorldInitDone(FNavigationSystemRunMode Mode)
 				}
 			}
 		}
-
-		if (!World->IsGameWorld() && World->GetLevels().Num() > 1)
-		{
-			// this is a bit of a @hack for the time being until we reorganize navigation data instances (4.9?)
-			// the point of this hack is to force editor-time navmesh rebuilding if we have any sublevels
-			// since currently there's no information regarding whether we already have all navigation built 
-			// for all the actors in sublevels
-			for (ANavigationData* NavData : NavDataSet)
-			{
-				if (NavData)
-				{
-					NavData->MarkAsNeedingUpdate();
-				}
-			}
-		}
 	}
 
 	if (Mode == FNavigationSystemRunMode::EditorMode && bGenerateNavigationOnlyAroundNavigationInvokers)
@@ -2500,7 +2485,7 @@ void UNavigationSystem::UpdateComponentInNavOctree(UActorComponent& Comp)
 		if (NavElement)
 		{
 			AActor* OwnerActor = Comp.GetOwner();
-			if (ensure(OwnerActor))
+			if (OwnerActor)
 			{
 				UNavigationSystem* NavSys = UNavigationSystem::GetCurrent(OwnerActor->GetWorld());
 				if (NavSys)
diff --git a/Engine/Source/Runtime/Engine/Private/AI/Navigation/NavigationTestingActor.cpp b/Engine/Source/Runtime/Engine/Private/AI/Navigation/NavigationTestingActor.cpp
--- a/Engine/Source/Runtime/Engine/Private/AI/Navigation/NavigationTestingActor.cpp
+++ b/Engine/Source/Runtime/Engine/Private/AI/Navigation/NavigationTestingActor.cpp
@@ -44,6 +44,8 @@ ANavigationTestingActor::ANavigationTestingActor(const FObjectInitializer& Objec
 	bShouldBeVisibleInGame = false;
 	TextCanvasOffset = FVector2D::ZeroVector;
 	bGatherDetailedInfo = true;
+	bDrawDistanceToWall = false;
+	ClosestWallLocation = FNavigationSystem::InvalidLocation;
 	OffsetFromCornersDistance = 0.f;
 
 	QueryingExtent = FVector(DEFAULT_NAV_QUERY_EXTENT_HORIZONTAL, DEFAULT_NAV_QUERY_EXTENT_HORIZONTAL, DEFAULT_NAV_QUERY_EXTENT_VERTICAL);
@@ -138,6 +140,19 @@ void ANavigationTestingActor::PostEditChangeProperty(FPropertyChangedEvent& Prop
 		{
 #if WITH_EDITORONLY_DATA
 			EdRenderComp->MarkRenderStateDirty();
+#endif
+		}
+		else if (ChangedCategory == TEXT("Query"))
+		{
+			if (bDrawDistanceToWall)
+			{
+				ClosestWallLocation = FindClosestWallLocation();
+			}
+#if WITH_EDITORONLY_DATA
+			else
+			{
+				EdRenderComp->MarkRenderStateDirty();
+			}
 #endif
 		}
 		else if (ChangedCategory == TEXT("Pathfinding"))
@@ -204,6 +219,11 @@ void ANavigationTestingActor::PostEditMove(bool bFinished)
 		{
 			UpdatePathfinding();
 		}
+
+		if (bDrawDistanceToWall)
+		{
+			ClosestWallLocation = FindClosestWallLocation();
+		}
 	}
 }
 
@@ -327,6 +347,26 @@ void ANavigationTestingActor::UpdatePathfinding()
 	}
 }
 
+FVector ANavigationTestingActor::FindClosestWallLocation() const
+{
+#if WITH_EDITORONLY_DATA
+	if (EdRenderComp)
+	{
+		EdRenderComp->MarkRenderStateDirty();
+	}
+#endif // WITH_EDITORONLY_DATA
+
+	ARecastNavMesh* AsRecastNavMesh = Cast<ARecastNavMesh>(MyNavData);
+	if (AsRecastNavMesh)
+	{
+		FVector TmpOutLocation = FNavigationSystem::InvalidLocation;
+		const float Distance = AsRecastNavMesh->FindDistanceToWall(GetActorLocation(), UNavigationQueryFilter::GetQueryFilter(*MyNavData, FilterClass), FLT_MAX, &TmpOutLocation);
+		return TmpOutLocation;
+	}
+	
+	return FNavigationSystem::InvalidLocation;
+}
+
 void ANavigationTestingActor::SearchPathTo(ANavigationTestingActor* Goal)
 {
 #if WITH_EDITORONLY_DATA
diff --git a/Engine/Source/Runtime/Engine/Private/AI/Navigation/PImplRecastNavMesh.cpp b/Engine/Source/Runtime/Engine/Private/AI/Navigation/PImplRecastNavMesh.cpp
--- a/Engine/Source/Runtime/Engine/Private/AI/Navigation/PImplRecastNavMesh.cpp
+++ b/Engine/Source/Runtime/Engine/Private/AI/Navigation/PImplRecastNavMesh.cpp
@@ -2678,6 +2678,11 @@ void FPImplRecastNavMesh::RemoveTileCacheLayer(int32 TileX, int32 TileY, int32 L
 		if (ExistingLayersList->IsValidIndex(LayerIdx))
 		{
 			ExistingLayersList->RemoveAt(LayerIdx);
+
+			for (int32 Idx = LayerIdx; Idx < ExistingLayersList->Num(); Idx++)
+			{
+				(*ExistingLayersList)[Idx].LayerIndex = Idx;
+			}
 		}
 		
 		if (ExistingLayersList->Num() == 0)
diff --git a/Engine/Source/Runtime/Engine/Private/AI/Navigation/RecastNavMesh.cpp b/Engine/Source/Runtime/Engine/Private/AI/Navigation/RecastNavMesh.cpp
--- a/Engine/Source/Runtime/Engine/Private/AI/Navigation/RecastNavMesh.cpp
+++ b/Engine/Source/Runtime/Engine/Private/AI/Navigation/RecastNavMesh.cpp
@@ -721,7 +721,6 @@ void ARecastNavMesh::Serialize( FArchive& Ar )
 		else if (RecastNavMeshSizeBytes > 4)
 		{
 			SerializeRecastNavMesh(Ar, RecastNavMeshImpl, NavMeshVersion);
-			bWantsUpdate = bForceRebuildOnLoad == true || HasValidNavmesh() == false;
 #if !(UE_BUILD_SHIPPING)
 			RequestDrawingUpdate();
 #endif //!(UE_BUILD_SHIPPING)
@@ -1215,7 +1214,7 @@ NavNodeRef ARecastNavMesh::FindNearestPoly(FVector const& Loc, FVector const& Ex
 	return PolyRef;
 }
 
-float ARecastNavMesh::FindDistanceToWall(const FVector& StartLoc, FSharedConstNavQueryFilter Filter, float MaxDistance) const
+float ARecastNavMesh::FindDistanceToWall(const FVector& StartLoc, FSharedConstNavQueryFilter Filter, float MaxDistance, FVector* OutClosestPointOnWall) const
 {
 	if (HasValidNavmesh() == false)
 	{
@@ -1250,6 +1249,10 @@ float ARecastNavMesh::FindDistanceToWall(const FVector& StartLoc, FSharedConstNa
 
 		if (dtStatusSucceed(RaycastStatus))
 		{
+			if (OutClosestPointOnWall)
+			{
+				*OutClosestPointOnWall = Recast2UnrealPoint(TmpHitPos);
+			}
 			return DistanceToWall;
 		}
 	}
@@ -1474,9 +1477,13 @@ void ARecastNavMesh::RequestDrawingUpdate(bool bForce)
 #if !UE_BUILD_SHIPPING
 	if (bForce || UNavMeshRenderingComponent::IsNavigationShowFlagSet(GetWorld()))
 	{
-		if (bForce && Cast<UNavMeshRenderingComponent>(RenderingComp))
+		if (bForce)
 		{
-			Cast<UNavMeshRenderingComponent>(RenderingComp)->ForceUpdate();
+			UNavMeshRenderingComponent* NavRenderingComp = Cast<UNavMeshRenderingComponent>(RenderingComp);
+			if (NavRenderingComp)
+			{
+				NavRenderingComp->ForceUpdate();
+			}
 		}
 
 		DECLARE_CYCLE_STAT(TEXT("FSimpleDelegateGraphTask.Requesting navmesh redraw"),
@@ -1683,11 +1690,6 @@ uint32 ARecastNavMesh::LogMemUsed() const
 
 #endif // !UE_BUILD_SHIPPING
 
-SIZE_T ARecastNavMesh::GetResourceSize(EResourceSizeMode::Type Mode)
-{
-	return 0;
-}
-
 uint16 ARecastNavMesh::GetDefaultForbiddenFlags() const
 {
 	return FPImplRecastNavMesh::GetFilterForbiddenFlags((const FRecastQueryFilter*)DefaultQueryFilter->GetImplementation());
@@ -1738,8 +1740,6 @@ void ARecastNavMesh::OnStreamingLevelAdded(ULevel* InLevel, UWorld* InWorld)
 {
 	QUICK_SCOPE_CYCLE_COUNTER(STAT_RecastNavMesh_OnStreamingLevelAdded);
 	
-	bWantsUpdate = true;
-
 	if (SupportsStreaming() && RecastNavMeshImpl)
 	{
 		URecastNavMeshDataChunk* NavDataChunk = GetNavigationDataChunk(InLevel);
@@ -2311,7 +2311,6 @@ void ARecastNavMesh::RemoveTiles(const TArray<FIntPoint>& Tiles)
 {
 	if (Tiles.Num() > 0)
 	{
-		bWantsUpdate = true;
 		FRecastNavMeshGenerator* MyGenerator = static_cast<FRecastNavMeshGenerator*>(GetGenerator());
 		if (MyGenerator)
 		{
@@ -2324,7 +2323,6 @@ void ARecastNavMesh::RebuildTile(const TArray<FIntPoint>& Tiles)
 {
 	if (Tiles.Num() > 0)
 	{
-		bWantsUpdate = true;
 		FRecastNavMeshGenerator* MyGenerator = static_cast<FRecastNavMeshGenerator*>(GetGenerator());
 		if (MyGenerator)
 		{
diff --git a/Engine/Source/Runtime/Engine/Private/AI/Navigation/RecastNavMeshGenerator.cpp b/Engine/Source/Runtime/Engine/Private/AI/Navigation/RecastNavMeshGenerator.cpp
--- a/Engine/Source/Runtime/Engine/Private/AI/Navigation/RecastNavMeshGenerator.cpp
+++ b/Engine/Source/Runtime/Engine/Private/AI/Navigation/RecastNavMeshGenerator.cpp
@@ -3073,7 +3073,6 @@ FRecastNavMeshGenerator::FRecastNavMeshGenerator(ARecastNavMesh& InDestNavMesh)
 	{
 		// recreate navmesh from scratch if no data was loaded
 		ConstructTiledNavMesh();
-		InDestNavMesh.MarkAsNeedingUpdate();
 	}
 	else
 	{
diff --git a/Engine/Source/Runtime/Engine/Private/Actor.cpp b/Engine/Source/Runtime/Engine/Private/Actor.cpp
--- a/Engine/Source/Runtime/Engine/Private/Actor.cpp
+++ b/Engine/Source/Runtime/Engine/Private/Actor.cpp
@@ -201,8 +201,7 @@ void AActor::ResetOwnedComponents()
 	// Identify any natively-constructed components referenced by properties that either failed to serialize or came in as NULL.
 	if(HasAnyFlags(RF_WasLoaded) && NativeConstructedComponentToPropertyMap.Num() > 0)
 	{
-		TInlineComponentArray<UActorComponent*> ComponentsToDestroy;
-		for (auto Component : OwnedComponents)
+		for (UActorComponent* Component : OwnedComponents)
 		{
 			// Only consider native components
 			if (Component && Component->CreationMethod == EComponentCreationMethod::Native)
@@ -212,7 +211,7 @@ void AActor::ResetOwnedComponents()
 				NativeConstructedComponentToPropertyMap.MultiFind(Component->GetFName(), Properties);
 
 				// Determine if the property or properties are no longer valid references (either it got serialized out that way or something failed during load)
-				for (auto ObjProp : Properties)
+				for (UObjectProperty* ObjProp : Properties)
 				{
 					check(ObjProp != nullptr);
 					UActorComponent* ActorComponent = Cast<UActorComponent>(ObjProp->GetObjectPropertyValue_InContainer(this));
@@ -231,8 +230,8 @@ void AActor::ResetOwnedComponents()
 #endif
 
 	TArray<UObject*> ActorChildren;
-	OwnedComponents.Empty();
-	ReplicatedComponents.Empty();
+	OwnedComponents.Reset();
+	ReplicatedComponents.Reset();
 	GetObjectsWithOuter(this, ActorChildren, true, RF_NoFlags, EInternalObjectFlags::PendingKill);
 
 	for (UObject* Child : ActorChildren)
@@ -264,13 +263,6 @@ void AActor::PostInitProperties()
 	}
 }
 
-void AActor::AddReferencedObjects(UObject* InThis, FReferenceCollector& Collector)
-{
-	AActor* This = CastChecked<AActor>(InThis);
-	Collector.AddReferencedObjects(This->OwnedComponents);
-	Super::AddReferencedObjects(InThis, Collector);
-}
-
 UWorld* AActor::GetWorld() const
 {
 	// CDO objects do not belong to a world
@@ -479,7 +471,8 @@ void AActor::Serialize(FArchive& Ar)
 	{
 		if (const UBlueprintGeneratedClass* BPGC = Cast<UBlueprintGeneratedClass>(GetClass()))
 		{
-			NativeConstructedComponentToPropertyMap.Empty(OwnedComponents.Num());
+			NativeConstructedComponentToPropertyMap.Reset();
+			NativeConstructedComponentToPropertyMap.Reserve(OwnedComponents.Num());
 			for(TFieldIterator<UObjectProperty> PropertyIt(BPGC, EFieldIteratorFlags::IncludeSuper); PropertyIt; ++PropertyIt)
 			{
 				UObjectProperty* ObjProp = *PropertyIt;
@@ -533,7 +526,7 @@ void AActor::PostLoad()
 		TInlineComponentArray<UChildActorComponent*> ParentChildActorComponents(ParentActor);
 		for (UChildActorComponent* ChildActorComponent : ParentChildActorComponents)
 		{
-			if (ChildActorComponent->ChildActor == this)
+			if (ChildActorComponent->GetChildActor() == this)
 			{
 				ParentComponent = ChildActorComponent;
 				break;
@@ -626,8 +619,7 @@ void AActor::ApplyWorldOffset(const FVector& InOffset, bool bWorldShift)
 	}
 
 	// Shift UActorComponent derived components, but not USceneComponents
-	const TArray<UActorComponent*>& AllActorComponents = GetComponents();
- 	for (UActorComponent* ActorComponent : AllActorComponents)
+ 	for (UActorComponent* ActorComponent : GetComponents())
  	{
  		if (IsValid(ActorComponent) && !ActorComponent->IsA<USceneComponent>())
  		{
@@ -881,10 +873,10 @@ void AActor::CallPreReplication(UNetDriver* NetDriver)
 	PreReplication(*NetDriver->FindOrCreateRepChangedPropertyTracker(this).Get());
 
 	// Call PreReplication on all owned components that are replicated
-	for (UActorComponent* Component : OwnedComponents)
+	for (UActorComponent* Component : ReplicatedComponents)
 	{
 		// Only call on components that aren't pending kill
-		if (Component && !Component->IsPendingKill() && Component->GetIsReplicated())
+		if (Component && !Component->IsPendingKill())
 		{
 			Component->PreReplication(*NetDriver->FindOrCreateRepChangedPropertyTracker(Component).Get());
 		}
@@ -1129,11 +1121,25 @@ bool AActor::IsOverlappingActor(const AActor* Other) const
 	return false;
 }
 
-void AActor::GetOverlappingActors(TArray<AActor*>& OutOverlappingActors, UClass* ClassFilter) const
+void AActor::GetOverlappingActors(TArray<AActor*>& OutOverlappingActors, TSubclassOf<AActor> ClassFilter) const
+{
+	// prepare output
+	TSet<AActor*> OverlappingActors;
+	GetOverlappingActors(OverlappingActors, ClassFilter);
+
+	OutOverlappingActors.Reset(OverlappingActors.Num());
+
+	for (AActor* OverlappingActor : OverlappingActors)
+	{
+		OutOverlappingActors.Add(OverlappingActor);
+	}
+}
+
+void AActor::GetOverlappingActors(TSet<AActor*>& OutOverlappingActors, TSubclassOf<AActor> ClassFilter) const
 {
 	// prepare output
 	OutOverlappingActors.Reset();
-	TArray<AActor*> OverlappingActorsForCurrentComponent;
+	TSet<AActor*> OverlappingActorsForCurrentComponent;
 
 	for(UActorComponent* OwnedComp : OwnedComponents)
 	{
@@ -1141,13 +1147,14 @@ void AActor::GetOverlappingActors(TArray<AActor*>& OutOverlappingActors, UClass*
 		{
 			PrimComp->GetOverlappingActors(OverlappingActorsForCurrentComponent, ClassFilter);
 
+			OutOverlappingActors.Reserve(OutOverlappingActors.Num() + OverlappingActorsForCurrentComponent.Num());
+
 			// then merge it into our final list
-			for (auto CompIt = OverlappingActorsForCurrentComponent.CreateIterator(); CompIt; ++CompIt)
+			for (AActor* OverlappingActor : OverlappingActorsForCurrentComponent)
 			{
-				AActor* OverlappingActor = *CompIt;
-				if(OverlappingActor != this)
+				if (OverlappingActor != this)
 				{
-					OutOverlappingActors.AddUnique(OverlappingActor);
+					OutOverlappingActors.Add(OverlappingActor);
 				}
 			}
 		}
@@ -1155,6 +1162,19 @@ void AActor::GetOverlappingActors(TArray<AActor*>& OutOverlappingActors, UClass*
 }
 
 void AActor::GetOverlappingComponents(TArray<UPrimitiveComponent*>& OutOverlappingComponents) const
+{
+	TSet<UPrimitiveComponent*> OverlappingComponents;
+	GetOverlappingComponents(OverlappingComponents);
+
+	OutOverlappingComponents.Reset(OverlappingComponents.Num());
+
+	for (UPrimitiveComponent* OverlappingComponent : OverlappingComponents)
+	{
+		OutOverlappingComponents.Add(OverlappingComponent);
+	}
+}
+
+void AActor::GetOverlappingComponents(TSet<UPrimitiveComponent*>& OutOverlappingComponents) const
 {
 	OutOverlappingComponents.Reset();
 	TArray<UPrimitiveComponent*> OverlappingComponentsForCurrentComponent;
@@ -1166,16 +1186,17 @@ void AActor::GetOverlappingComponents(TArray<UPrimitiveComponent*>& OutOverlappi
 			// get list of components from the component
 			PrimComp->GetOverlappingComponents(OverlappingComponentsForCurrentComponent);
 
+			OutOverlappingComponents.Reserve(OutOverlappingComponents.Num() + OverlappingComponentsForCurrentComponent.Num());
+
 			// then merge it into our final list
-			for (auto CompIt = OverlappingComponentsForCurrentComponent.CreateIterator(); CompIt; ++CompIt)
+			for (UPrimitiveComponent* OverlappingComponent : OverlappingComponentsForCurrentComponent)
 			{
-				OutOverlappingComponents.AddUnique(*CompIt);
+				OutOverlappingComponents.Add(OverlappingComponent);
 			}
 		}
 	}
 }
 
-
 void AActor::NotifyActorBeginOverlap(AActor* OtherActor)
 {
 	// call BP handler
@@ -1703,7 +1724,7 @@ void AActor::EndPlay(const EEndPlayReason::Type EndPlayReason)
 
 		// Dispatch the blueprint events
 		ReceiveEndPlay(EndPlayReason);
-		OnEndPlay.Broadcast(EndPlayReason);
+		OnEndPlay.Broadcast(this, EndPlayReason);
 
 		TInlineComponentArray<UActorComponent*> Components;
 		GetComponents(Components);
@@ -1771,7 +1792,7 @@ void AActor::Destroyed()
 	RouteEndPlay(EEndPlayReason::Destroyed);
 
 	ReceiveDestroyed();
-	OnDestroyed.Broadcast();
+	OnDestroyed.Broadcast(this);
 	UWorld* ActorWorld = GetWorld();
 
 	if( ActorWorld )
@@ -1852,7 +1873,7 @@ float AActor::TakeDamage(float DamageAmount, FDamageEvent const& DamageEvent, AC
 		if (ActualDamage != 0.f)
 		{
 			ReceivePointDamage(ActualDamage, DamageTypeCDO, PointDamageEvent->HitInfo.ImpactPoint, PointDamageEvent->HitInfo.ImpactNormal, PointDamageEvent->HitInfo.Component.Get(), PointDamageEvent->HitInfo.BoneName, PointDamageEvent->ShotDirection, EventInstigator, DamageCauser);
-			OnTakePointDamage.Broadcast(ActualDamage, EventInstigator, PointDamageEvent->HitInfo.ImpactPoint, PointDamageEvent->HitInfo.Component.Get(), PointDamageEvent->HitInfo.BoneName, PointDamageEvent->ShotDirection, DamageTypeCDO, DamageCauser);
+			OnTakePointDamage.Broadcast(this, ActualDamage, EventInstigator, PointDamageEvent->HitInfo.ImpactPoint, PointDamageEvent->HitInfo.Component.Get(), PointDamageEvent->HitInfo.BoneName, PointDamageEvent->ShotDirection, DamageTypeCDO, DamageCauser);
 
 			// Notify the component
 			UPrimitiveComponent* const PrimComp = PointDamageEvent->HitInfo.Component.Get();
@@ -1892,7 +1913,7 @@ float AActor::TakeDamage(float DamageAmount, FDamageEvent const& DamageEvent, AC
 	if (ActualDamage != 0.f)
 	{
 		ReceiveAnyDamage(ActualDamage, DamageTypeCDO, EventInstigator, DamageCauser);
-		OnTakeAnyDamage.Broadcast(ActualDamage, DamageTypeCDO, EventInstigator, DamageCauser);
+		OnTakeAnyDamage.Broadcast(this, ActualDamage, DamageTypeCDO, EventInstigator, DamageCauser);
 		if (EventInstigator != NULL)
 		{
 			EventInstigator->InstigatedAnyDamage(ActualDamage, DamageTypeCDO, this, DamageCauser);
@@ -1967,7 +1988,7 @@ void AActor::InternalDispatchBlockingHit(UPrimitiveComponent* MyComp, UPrimitive
 		// If component is still alive, call delegate on component
 		if(!MyComp->IsPendingKill())
 		{
-			MyComp->OnComponentHit.Broadcast(OtherActor, OtherComp, FVector(0,0,0), Hit);
+			MyComp->OnComponentHit.Broadcast(MyComp, OtherActor, OtherComp, FVector(0,0,0), Hit);
 		}
 	}
 }
@@ -2036,18 +2057,15 @@ void AActor::DisplayDebug(UCanvas* Canvas, const FDebugDisplayInfo& DebugDisplay
 
 		T = FString(TEXT("Overlapping "));
 
-		TArray<AActor*> TouchingActors;
+		TSet<AActor*> TouchingActors;
 		GetOverlappingActors(TouchingActors);
 		bool bFoundAnyOverlaps = false;
-		for (int32 iTouching = 0; iTouching<TouchingActors.Num(); ++iTouching)
+		for (AActor* const TestActor : TouchingActors)
 		{
-			AActor* const TestActor = TouchingActors[iTouching];
 			if (TestActor &&
-				!TestActor->IsPendingKill() &&
-				TestActor->IsA<AActor>())
+				!TestActor->IsPendingKill())
 			{
-				AActor* A = TestActor;
-				T = T + A->GetName() + " ";
+				T = T + TestActor->GetName() + " ";
 				bFoundAnyOverlaps = true;
 			}
 		}
@@ -2230,51 +2248,40 @@ enum ECollisionResponse AActor::GetComponentsCollisionResponseToChannel(enum ECo
 void AActor::AddOwnedComponent(UActorComponent* Component)
 {
 	check(Component->GetOwner() == this);
-	OwnedComponents.AddUnique(Component);
 
-	if (Component->GetIsReplicated())
-	{
-		ReplicatedComponents.AddUnique(Component);
-	}
+	bool bAlreadyInSet = false;
+	OwnedComponents.Add(Component, &bAlreadyInSet);
 
-	if (Component->IsCreatedByConstructionScript())
-	{
-		BlueprintCreatedComponents.AddUnique(Component);
-	}
-	else if (Component->CreationMethod == EComponentCreationMethod::Instance)
+	if (!bAlreadyInSet)
 	{
-		InstanceComponents.AddUnique(Component);
-	}
-}
+		if (Component->GetIsReplicated())
+		{
+			ReplicatedComponents.Add(Component);
+		}
 
-void AActor::RemoveOwnedComponent(UActorComponent* Component)
-{
-	if (OwnedComponents.RemoveSwap(Component) > 0)
-	{
-		ReplicatedComponents.RemoveSwap(Component);
 		if (Component->IsCreatedByConstructionScript())
 		{
-			BlueprintCreatedComponents.RemoveSwap(Component);
+			BlueprintCreatedComponents.Add(Component);
 		}
 		else if (Component->CreationMethod == EComponentCreationMethod::Instance)
 		{
-			InstanceComponents.RemoveSwap(Component);
+			InstanceComponents.Add(Component);
 		}
 	}
-	else
-	{
-		// If we didn't remove something we expected to then it probably got NULL through the
-		// property system so take the time to pull them out now
-		OwnedComponents.RemoveSwap(nullptr);
+}
 
-		ReplicatedComponents.RemoveSwap(nullptr);
+void AActor::RemoveOwnedComponent(UActorComponent* Component)
+{
+	if (OwnedComponents.Remove(Component) > 0)
+	{
+		ReplicatedComponents.Remove(Component);
 		if (Component->IsCreatedByConstructionScript())
 		{
-			BlueprintCreatedComponents.RemoveSwap(nullptr);
+			BlueprintCreatedComponents.RemoveSingleSwap(Component);
 		}
 		else if (Component->CreationMethod == EComponentCreationMethod::Instance)
 		{
-			InstanceComponents.RemoveSwap(nullptr);
+			InstanceComponents.RemoveSingleSwap(Component);
 		}
 	}
 }
@@ -2286,27 +2293,22 @@ bool AActor::OwnsComponent(UActorComponent* Component) const
 }
 #endif
 
-const TArray<UActorComponent*>& AActor::GetReplicatedComponents() const
-{
-	return ReplicatedComponents;
-}
-
 void AActor::UpdateReplicatedComponent(UActorComponent* Component)
 {
 	checkf(Component->GetOwner() == this, TEXT("UE-9568: Component %s being updated for Actor %s"), *Component->GetPathName(), *GetPathName() );
 	if (Component->GetIsReplicated())
 	{
-		ReplicatedComponents.AddUnique(Component);
+		ReplicatedComponents.Add(Component);
 	}
 	else
 	{
-		ReplicatedComponents.RemoveSwap(Component);
+		ReplicatedComponents.Remove(Component);
 	}
 }
 
 void AActor::UpdateAllReplicatedComponents()
 {
-	ReplicatedComponents.Empty();
+	ReplicatedComponents.Reset();
 
 	for (UActorComponent* Component : OwnedComponents)
 	{
@@ -2348,7 +2350,7 @@ void AActor::ClearInstanceComponents(const bool bDestroyComponents)
 	}
 	else
 	{
-		InstanceComponents.Empty();
+		InstanceComponents.Reset();
 	}
 }
 
@@ -3644,14 +3646,15 @@ void AActor::DispatchPhysicsCollisionHit(const FRigidBodyCollisionInfo& MyInfo,
 
 	// Execute delegates if bound
 
-	if(OnActorHit.IsBound())
+	if (OnActorHit.IsBound())
 	{
 		OnActorHit.Broadcast(this, OtherInfo.Actor.Get(), RigidCollisionData.TotalNormalImpulse, Result);
 	}
 
-	if(MyInfo.Component.IsValid() && MyInfo.Component.Get()->OnComponentHit.IsBound())
+	UPrimitiveComponent* MyInfoComponent = MyInfo.Component.Get();
+	if (MyInfoComponent && MyInfoComponent->OnComponentHit.IsBound())
 	{
-		MyInfo.Component.Get()->OnComponentHit.Broadcast(OtherInfo.Actor.Get(), OtherInfo.Component.Get(), RigidCollisionData.TotalNormalImpulse, Result);
+		MyInfoComponent->OnComponentHit.Broadcast(MyInfoComponent, OtherInfo.Actor.Get(), OtherInfo.Component.Get(), RigidCollisionData.TotalNormalImpulse, Result);
 	}
 }
 
@@ -3685,6 +3688,25 @@ UChildActorComponent* AActor::GetParentComponent() const
 	return ParentComponent.Get();
 }
 
+void AActor::GetAllChildActors(TArray<AActor*>& ChildActors, bool bIncludeDescendants) const
+{
+	TInlineComponentArray<UChildActorComponent*> ChildActorComponents(this);
+
+	ChildActors.Reserve(ChildActors.Num() + ChildActorComponents.Num());
+	for (UChildActorComponent* CAC : ChildActorComponents)
+	{
+		if (AActor* ChildActor = CAC->GetChildActor())
+		{
+			ChildActors.Add(ChildActor);
+			if (bIncludeDescendants)
+			{
+				ChildActor->GetAllChildActors(ChildActors, true);
+			}
+		}
+	}
+}
+
+
 // COMPONENTS
 
 void AActor::UnregisterAllComponents()
@@ -3880,9 +3902,9 @@ void AActor::MarkComponentsRenderStateDirty()
 			ActorComp->MarkRenderStateDirty();
 			if (UChildActorComponent* ChildActorComponent = Cast<UChildActorComponent>(ActorComp))
 			{
-				if (ChildActorComponent->ChildActor)
+				if (ChildActorComponent->GetChildActor())
 				{
-					ChildActorComponent->ChildActor->MarkComponentsRenderStateDirty();
+					ChildActorComponent->GetChildActor()->MarkComponentsRenderStateDirty();
 				}
 			}
 		}
@@ -4245,6 +4267,17 @@ void AActor::K2_SetActorRelativeTransform(const FTransform& NewRelativeTransform
 	SetActorRelativeTransform(NewRelativeTransform, bSweep, (bSweep ? &SweepHitResult : nullptr), TeleportFlagToEnum(bTeleport));
 }
 
-
+float AActor::GetGameTimeSinceCreation()
+{
+	if(GetWorld() != nullptr)
+	{
+		return GetWorld()->GetTimeSeconds() - CreationTime;		
+	}
+	// return 0.f if GetWorld return's null
+	else
+	{
+		return 0.f;
+	}
+}
 
 #undef LOCTEXT_NAMESPACE
diff --git a/Engine/Source/Runtime/Engine/Private/Animation/AnimPhysicsSolver.cpp b/Engine/Source/Runtime/Engine/Private/Animation/AnimPhysicsSolver.cpp
--- a/Engine/Source/Runtime/Engine/Private/Animation/AnimPhysicsSolver.cpp
+++ b/Engine/Source/Runtime/Engine/Private/Animation/AnimPhysicsSolver.cpp
@@ -791,7 +791,7 @@ void FAnimPhys::CreateSpring(TArray<FAnimPhysSpring>& SpringContainer, FAnimPhys
 	SpringContainer.Add(NewSpring);
 }
 
-void FAnimPhys::InitializeBodyVelocity(float DeltaTime, FAnimPhysRigidBody *InBody)
+void FAnimPhys::InitializeBodyVelocity(float DeltaTime, FAnimPhysRigidBody *InBody, const FVector& GravityDirection)
 {
 	// Gather weak forces being applied to body at beginning of timestep
 	// Forward euler update of the velocity and rotation/spin 
@@ -806,7 +806,7 @@ void FAnimPhys::InitializeBodyVelocity(float DeltaTime, FAnimPhysRigidBody *InBo
 	InBody->LinearMomentum *= LinearDampingLeftOver;
 	InBody->AngularMomentum *= AngularDampingLeftOver;
 
-	FVector Force = FVector(0.0f, 0.0f, UPhysicsSettings::Get()->DefaultGravityZ) * InBody->Mass * InBody->GravityScale;
+	FVector Force = GravityDirection * FMath::Abs(UPhysicsSettings::Get()->DefaultGravityZ) * InBody->Mass * InBody->GravityScale;
 	FVector Torque(0.0f, 0.0f, 0.0f);
 
 	// Add wind forces
@@ -865,13 +865,22 @@ void FAnimPhys::UpdatePose(FAnimPhysRigidBody* InBody)
 	InBody->InverseWorldSpaceTensor = OrientAsMatrix * InBody->InverseTensorWithoutMass * InBody->InverseMass * OrientAsMatrix.GetTransposed();
 }
 
-void FAnimPhys::PhysicsUpdate(float DeltaTime, TArray<FAnimPhysRigidBody*>& Bodies, TArray<FAnimPhysLinearLimit>& LinearLimits, TArray<FAnimPhysAngularLimit>& AngularLimits, TArray<FAnimPhysSpring>& Springs, int32 NumPreIterations, int32 NumPostIterations)
+void FAnimPhys::PhysicsUpdate(float DeltaTime, TArray<FAnimPhysRigidBody*>& Bodies, TArray<FAnimPhysLinearLimit>& LinearLimits, TArray<FAnimPhysAngularLimit>& AngularLimits, TArray<FAnimPhysSpring>& Springs, const FVector& GravityDirection, const FVector& ExternalForce, int32 NumPreIterations, int32 NumPostIterations)
 {
 	SCOPE_CYCLE_COUNTER(STAT_AnimDynamicsUpdate);
 
 	for (FAnimPhysRigidBody* Body : Bodies)
 	{
-		InitializeBodyVelocity(DeltaTime, Body);	// based on previous and current force/torque
+		InitializeBodyVelocity(DeltaTime, Body, GravityDirection);	// based on previous and current force/torque
+	}
+
+	// Apply any external forces.
+	if(!ExternalForce.IsNearlyZero())
+	{
+		for(FAnimPhysRigidBody* Body : Bodies)
+		{
+			Body->LinearMomentum += ExternalForce * DeltaTime;
+		}
 	}
 
 	for (FAnimPhysSpring& Spring : Springs)
diff --git a/Engine/Source/Runtime/Engine/Private/Animation/AnimSequence.cpp b/Engine/Source/Runtime/Engine/Private/Animation/AnimSequence.cpp
--- a/Engine/Source/Runtime/Engine/Private/Animation/AnimSequence.cpp
+++ b/Engine/Source/Runtime/Engine/Private/Animation/AnimSequence.cpp
@@ -2569,8 +2569,8 @@ void UAnimSequence::RemapTracksToNewSkeleton( USkeleton* NewSkeleton, bool bConv
 			{
 				// first calculate component space ref pose to get the relative transform between
 				// two ref poses. It is very important update ref pose before getting here. 
-				TArray<FTransform> NewRotations, NewSpaceBases, OldRotations;
-				TArray<FTransform> NewTranslations, OldSpaceBases, OldTranslations;
+				TArray<FTransform> NewRotations, OldRotations, NewSpaceBases, OldSpaceBases;
+				TArray<FTransform> NewTranslations, OldTranslations;
 				TArray<bool> NewTranslationParentFlags, OldTranslationParentFlags;
 				// get the spacebases transform
 				FillUpTransformBasedOnRig(NewSkeleton, NewSpaceBases, NewRotations, NewTranslations, NewTranslationParentFlags);
@@ -3362,7 +3362,7 @@ int32 UAnimSequence::GetSpaceBasedAnimationData(TArray< TArray<FTransform> >& An
 
 						FQuat ComponentRotation;
 						FTransform ComponentTranslation;
-						FVector ComponentScale = FVector(1.f);
+						FVector ComponentScale;
 
 						// rotation first
 						// this is easy since we just make sure it's evaluated or not
@@ -3415,8 +3415,11 @@ int32 UAnimSequence::GetSpaceBasedAnimationData(TArray< TArray<FTransform> >& An
 								{
 									for(int32 Key = 0; Key < NumKeys; ++Key)
 									{
-										ComponentTranslation = FTransform(RiggingAnimationData->AnimationTracks[NodeIndex].PosKeys[Key]) * AnimationDataInComponentSpace[ParentBoneIndex][Key];
+										const FTransform& AnimCompSpace = AnimationDataInComponentSpace[ParentBoneIndex][Key];
+										ComponentTranslation = FTransform(RiggingAnimationData->AnimationTracks[NodeIndex].PosKeys[Key]) * AnimCompSpace;
 										AnimationDataInComponentSpace[BoneIndex][Key].SetTranslation(ComponentTranslation.GetTranslation());
+
+										ComponentScale = AnimCompSpace.GetScale3D() * RiggingAnimationData->AnimationTracks[NodeIndex].ScaleKeys[Key];
 										AnimationDataInComponentSpace[BoneIndex][Key].SetScale3D(ComponentScale);
 									}
 								}
@@ -3433,6 +3436,8 @@ int32 UAnimSequence::GetSpaceBasedAnimationData(TArray< TArray<FTransform> >& An
 								{
 									ComponentTranslation = FTransform(RiggingAnimationData->AnimationTracks[NodeIndex].PosKeys[Key]);
 									AnimationDataInComponentSpace[BoneIndex][Key].SetTranslation(ComponentTranslation.GetTranslation());
+									
+									ComponentScale = RiggingAnimationData->AnimationTracks[NodeIndex].ScaleKeys[Key];
 									AnimationDataInComponentSpace[BoneIndex][Key].SetScale3D(ComponentScale);
 								}
 							}
diff --git a/Engine/Source/Runtime/Engine/Private/Animation/VertexAnimation.cpp b/Engine/Source/Runtime/Engine/Private/Animation/VertexAnimation.cpp
--- a/Engine/Source/Runtime/Engine/Private/Animation/VertexAnimation.cpp
+++ b/Engine/Source/Runtime/Engine/Private/Animation/VertexAnimation.cpp
@@ -9,19 +9,6 @@ UVertexAnimation::UVertexAnimation(const FObjectInitializer& ObjectInitializer)
 {
 }
 
-void UVertexAnimation::Serialize( FArchive& Ar )
-{
-	Super::Serialize( Ar );
-
-	//Ar << VertexAnimData;
-}
-
-
-SIZE_T UVertexAnimation::GetResourceSize(EResourceSizeMode::Type Mode)
-{
-	return 0;
-}
-
 int32 UVertexAnimation::GetNumFrames()
 {
 	return VertexAnimData.Num();
diff --git a/Engine/Source/Runtime/Engine/Private/AudioDevice.cpp b/Engine/Source/Runtime/Engine/Private/AudioDevice.cpp
--- a/Engine/Source/Runtime/Engine/Private/AudioDevice.cpp
+++ b/Engine/Source/Runtime/Engine/Private/AudioDevice.cpp
@@ -912,7 +912,7 @@ bool FAudioDevice::HandleAudioMemoryInfo(const TCHAR* Cmd, FOutputDevice& Ar)
 	{
 		for (FConfigSectionMap::TIterator It(*TrackedFolders); It; ++It)
 		{
-			FString SoundFolder = *It.Value();
+			const FString& SoundFolder = *It.Value().GetValue();
 			SoundWaveGroupSizes.Add(SoundFolder, FSoundWaveGroupInfo());
 			SoundWaveGroupFolders.Add(SoundFolder);
 		}
@@ -2186,6 +2186,10 @@ void FAudioDevice::SetReverbSettings(class AAudioVolume* Volume, const FReverbSe
 	Effects->SetReverbSettings(*ActivatedReverb);
 }
 
+UReverbEffect* FAudioDevice::GetCurrentReverbEffect()
+{
+	return Effects->GetCurrentReverbEffect();
+}
 
 void* FAudioDevice::InitEffect(FSoundSource* Source)
 {
diff --git a/Engine/Source/Runtime/Engine/Private/BufferVisualizationData.cpp b/Engine/Source/Runtime/Engine/Private/BufferVisualizationData.cpp
--- a/Engine/Source/Runtime/Engine/Private/BufferVisualizationData.cpp
+++ b/Engine/Source/Runtime/Engine/Private/BufferVisualizationData.cpp
@@ -35,7 +35,7 @@ void FBufferVisualizationData::Initialize()
 				for (FConfigSection::TIterator It(*MaterialSection); It; ++It)
 				{
 					FString MaterialName;
-					if( FParse::Value( *It.Value(), TEXT("Material="), MaterialName, true ) )
+					if( FParse::Value( *It.Value().GetValue(), TEXT("Material="), MaterialName, true ) )
 					{
 						UMaterial* Material = LoadObject<UMaterial>(NULL, *MaterialName);
 			
@@ -46,7 +46,7 @@ void FBufferVisualizationData::Initialize()
 							Rec.Name = It.Key().GetPlainNameString();
 							Rec.Material = Material;
 							FText DisplayName;
-							FParse::Value( *It.Value(), TEXT("Name="), DisplayName, TEXT("Engine.BufferVisualizationMaterials") );
+							FParse::Value( *It.Value().GetValue(), TEXT("Name="), DisplayName, TEXT("Engine.BufferVisualizationMaterials") );
 							Rec.DisplayName = DisplayName;
 						}
 					}
diff --git a/Engine/Source/Runtime/Engine/Private/Camera/CameraAnim.cpp b/Engine/Source/Runtime/Engine/Private/Camera/CameraAnim.cpp
--- a/Engine/Source/Runtime/Engine/Private/Camera/CameraAnim.cpp
+++ b/Engine/Source/Runtime/Engine/Private/Camera/CameraAnim.cpp
@@ -146,7 +146,7 @@ void UCameraAnim::CalcLocalAABB()
 
 SIZE_T UCameraAnim::GetResourceSize(EResourceSizeMode::Type Mode)
 {
-	int32 ResourceSize = 0;
+	int32 ResourceSize = Super::GetResourceSize(Mode);
 
 	if (Mode == EResourceSizeMode::Inclusive && CameraInterpGroup)
 	{
diff --git a/Engine/Source/Runtime/Engine/Private/Character.cpp b/Engine/Source/Runtime/Engine/Private/Character.cpp
--- a/Engine/Source/Runtime/Engine/Private/Character.cpp
+++ b/Engine/Source/Runtime/Engine/Private/Character.cpp
@@ -14,6 +14,7 @@
 #include "Animation/AnimMontage.h"
 #include "Components/CapsuleComponent.h"
 #include "GameFramework/DamageType.h"
+#include "Engine/SkeletalMeshSocket.h"
 
 DEFINE_LOG_CATEGORY_STATIC(LogCharacter, Log, All);
 DEFINE_LOG_CATEGORY_STATIC(LogAvatar, Log, All);
@@ -561,10 +562,22 @@ namespace MovementBaseUtility
 						return true;
 					}
 
+					const USkeletalMeshSocket* Socket = SkeletalBase->GetSocketByName(BoneName);
+					if (Socket != nullptr)
+					{
+						const FTransform SocketTransform = Socket->GetSocketTransform(SkeletalBase);
+						OutLocation = SocketTransform.GetLocation();
+						OutQuat = SocketTransform.GetRotation();
+						return true;
+					}
+
 					UE_LOG(LogCharacter, Warning, TEXT("GetMovementBaseTransform(): Invalid bone '%s' for SkeletalMeshComponent base %s"), *BoneName.ToString(), *GetPathNameSafe(MovementBase));
 					return false;
 				}
-				// TODO: warn if not a skeletal mesh but providing bone index.
+				else
+				{
+					UE_LOG(LogCharacter, Warning, TEXT("GetMovementBaseTransform(): Requested bone '%s' for non-SkeletalMeshComponent base %s"), *BoneName.ToString(), *GetPathNameSafe(MovementBase));
+				}
 			}
 
 			// No bone supplied
diff --git a/Engine/Source/Runtime/Engine/Private/CheatManager.cpp b/Engine/Source/Runtime/Engine/Private/CheatManager.cpp
--- a/Engine/Source/Runtime/Engine/Private/CheatManager.cpp
+++ b/Engine/Source/Runtime/Engine/Private/CheatManager.cpp
@@ -7,6 +7,7 @@
 #include "OnlineSubsystemUtils.h"
 #include "VisualLogger/VisualLogger.h"
 #include "GameFramework/Character.h"
+#include "Engine/Console.h"
 
 #if !UE_BUILD_SHIPPING
 #include "ISlateReflectorModule.h"
@@ -338,7 +339,7 @@ void UCheatManager::DestroyPawns(TSubclassOf<APawn> aClass)
 
 void UCheatManager::Summon( const FString& ClassName )
 {
-	UE_LOG(LogCheatManager, Log,  TEXT("Fabricate %s"), *ClassName );
+	UE_LOG(LogCheatManager, Log, TEXT("Fabricate %s"), *ClassName );
 
 	bool bIsValidClassName = true;
 	FString FailureReason;
@@ -1176,6 +1177,41 @@ void UCheatManager::InvertMouse()
 	}
 }
 
+void UCheatManager::CheatScript(FString ScriptName)
+{
+	APlayerController* const PlayerController = GetOuterAPlayerController();
+	ULocalPlayer* const LocalPlayer = PlayerController ? Cast<ULocalPlayer>(PlayerController->Player) : nullptr;
+
+	if (LocalPlayer)
+	{
+		// Run commands from the ini
+		FConfigSection const* const CommandsToRun = GConfig->GetSectionPrivate(*FString::Printf(TEXT("CheatScript.%s"), *ScriptName), 0, 1, GGameIni);
+
+		if (CommandsToRun)
+		{
+			for (FConfigSectionMap::TConstIterator It(*CommandsToRun); It; ++It)
+			{
+				// show user what commands ran
+				if (LocalPlayer->ViewportClient && LocalPlayer->ViewportClient->ViewportConsole)
+				{
+					FString const S = FString::Printf(TEXT("> %s"), *It.Value().GetValue());
+					LocalPlayer->ViewportClient->ViewportConsole->OutputText(S);
+				}
+
+				LocalPlayer->Exec(GetWorld(), *It.Value().GetValue(), *GLog);
+			}
+		}
+		else
+		{
+			UE_LOG(LogCheatManager, Warning, TEXT("Can't find section 'CheatScript.%s' in DefaultGame.ini"), *ScriptName);
+		}
+	}
+	else
+	{
+		UE_LOG(LogCheatManager, Warning, TEXT("Can't find local player!"));
+	}
+}
+
 void UCheatManager::LogOutBugItGoToLogFile( const FString& InScreenShotDesc, const FString& InGoString, const FString& InLocString )
 {
 #if ALLOW_DEBUG_FILES
@@ -1220,4 +1256,5 @@ void UCheatManager::LogOutBugItGoToLogFile( const FString& InScreenShotDesc, con
 #endif // ALLOW_DEBUG_FILES
 }
 
+
 #undef LOCTEXT_NAMESPACE
diff --git a/Engine/Source/Runtime/Engine/Private/Components/ActorComponent.cpp b/Engine/Source/Runtime/Engine/Private/Components/ActorComponent.cpp
--- a/Engine/Source/Runtime/Engine/Private/Components/ActorComponent.cpp
+++ b/Engine/Source/Runtime/Engine/Private/Components/ActorComponent.cpp
@@ -43,7 +43,7 @@ FGlobalComponentReregisterContext::FGlobalComponentReregisterContext()
 	FlushRenderingCommands();
 
 	// Detach all actor components.
-	for(auto* Component : TObjectRange<UActorComponent>())
+	for(UActorComponent* Component : TObjectRange<UActorComponent>())
 	{
 		new(ComponentContexts) FComponentReregisterContext(Component);
 	}
@@ -57,7 +57,7 @@ FGlobalComponentReregisterContext::FGlobalComponentReregisterContext(const TArra
 	FlushRenderingCommands();
 
 	// Detach only actor components that are not in the excluded list
-	for (auto* Component : TObjectRange<UActorComponent>())
+	for (UActorComponent* Component : TObjectRange<UActorComponent>())
 	{
 		bool bShouldReregister=true;
 		for (UClass* ExcludeClass : ExcludeComponents)
@@ -76,14 +76,6 @@ FGlobalComponentReregisterContext::FGlobalComponentReregisterContext(const TArra
 	}
 }
 
-FGlobalComponentReregisterContext::FGlobalComponentReregisterContext(const TArray<AActor*>& InParentActors)
-{
-	ActiveGlobalReregisterContextCount++;
-
-	// wait until resources are released
-	FlushRenderingCommands();
-}
-
 FGlobalComponentReregisterContext::~FGlobalComponentReregisterContext()
 {
 	check(ActiveGlobalReregisterContextCount > 0);
@@ -98,7 +90,7 @@ FGlobalComponentRecreateRenderStateContext::FGlobalComponentRecreateRenderStateC
 	FlushRenderingCommands();
 
 	// recreate render state for all components.
-	for (auto* Component : TObjectRange<UActorComponent>())
+	for (UActorComponent* Component : TObjectRange<UActorComponent>())
 	{
 		new(ComponentContexts) FComponentRecreateRenderStateContext(Component);
 	}
@@ -868,25 +860,33 @@ void UActorComponent::RegisterComponentWithWorld(UWorld* InWorld)
 	ExecuteRegisterEvents();
 
 	// If not in a game world register ticks now, otherwise defer until BeginPlay. If no owner we won't trigger BeginPlay either so register now in that case as well.
-	if (MyOwner == nullptr || !InWorld->IsGameWorld())
+	if (!InWorld->IsGameWorld())
 	{
 		RegisterAllComponentTickFunctions(true);
 	}
-
-	if (MyOwner == nullptr || MyOwner->IsActorInitialized())
+	else if (MyOwner == nullptr)
 	{
 		if (!bHasBeenInitialized && bWantsInitializeComponent)
 		{
 			InitializeComponent();
 		}
-	}
 
-	if (MyOwner && (MyOwner->HasActorBegunPlay() || MyOwner->IsActorBeginningPlay()))
-	{
 		RegisterAllComponentTickFunctions(true);
-		if (bWantsBeginPlay && !bHasBegunPlay)
+	}
+	else
+	{
+		if (!bHasBeenInitialized && bWantsInitializeComponent && MyOwner->IsActorInitialized())
 		{
-			BeginPlay();
+			InitializeComponent();
+		}
+
+		if (MyOwner->HasActorBegunPlay() || MyOwner->IsActorBeginningPlay())
+		{
+			RegisterAllComponentTickFunctions(true);
+			if (bWantsBeginPlay && !bHasBegunPlay)
+			{
+				BeginPlay();
+			}
 		}
 	}
 
diff --git a/Engine/Source/Runtime/Engine/Private/Components/CharacterMovementComponent.cpp b/Engine/Source/Runtime/Engine/Private/Components/CharacterMovementComponent.cpp
--- a/Engine/Source/Runtime/Engine/Private/Components/CharacterMovementComponent.cpp
+++ b/Engine/Source/Runtime/Engine/Private/Components/CharacterMovementComponent.cpp
@@ -8209,7 +8209,7 @@ void UCharacterMovementComponent::ClientAckGoodMove_Implementation(float TimeSta
 	ClientData->AckMove(MoveIndex);
 }
 
-void UCharacterMovementComponent::CapsuleTouched( AActor* Other, UPrimitiveComponent* OtherComp, int32 OtherBodyIndex, bool bFromSweep, const FHitResult& SweepResult )
+void UCharacterMovementComponent::CapsuleTouched(UPrimitiveComponent* OverlappedComp, AActor* Other, UPrimitiveComponent* OtherComp, int32 OtherBodyIndex, bool bFromSweep, const FHitResult& SweepResult )
 {
 	if (!bEnablePhysicsInteraction)
 	{
diff --git a/Engine/Source/Runtime/Engine/Private/Components/ChildActorComponent.cpp b/Engine/Source/Runtime/Engine/Private/Components/ChildActorComponent.cpp
--- a/Engine/Source/Runtime/Engine/Private/Components/ChildActorComponent.cpp
+++ b/Engine/Source/Runtime/Engine/Private/Components/ChildActorComponent.cpp
@@ -108,12 +108,12 @@ void UChildActorComponent::OnUnregister()
 
 FChildActorComponentInstanceData::FChildActorComponentInstanceData(const UChildActorComponent* Component)
 	: FSceneComponentInstanceData(Component)
-	, ChildActorName(Component->ChildActorName)
+	, ChildActorName(Component->GetChildActorName())
 	, ComponentInstanceData(nullptr)
 {
-	if (Component->ChildActor)
+	if (Component->GetChildActor())
 	{
-		ComponentInstanceData = new FComponentInstanceDataCache(Component->ChildActor);
+		ComponentInstanceData = new FComponentInstanceDataCache(Component->GetChildActor());
 		// If it is empty dump it
 		if (!ComponentInstanceData->HasInstanceData())
 		{
@@ -121,7 +121,7 @@ FChildActorComponentInstanceData::FChildActorComponentInstanceData(const UChildA
 			ComponentInstanceData = nullptr;
 		}
 
-		USceneComponent* ChildRootComponent = Component->ChildActor->GetRootComponent();
+		USceneComponent* ChildRootComponent = Component->GetChildActor()->GetRootComponent();
 		if (ChildRootComponent)
 		{
 			for (USceneComponent* AttachedComponent : ChildRootComponent->AttachChildren)
@@ -129,7 +129,7 @@ FChildActorComponentInstanceData::FChildActorComponentInstanceData(const UChildA
 				if (AttachedComponent)
 				{
 					AActor* AttachedActor = AttachedComponent->GetOwner();
-					if (AttachedActor != Component->ChildActor)
+					if (AttachedActor != Component->GetChildActor())
 					{
 						FAttachedActorInfo Info;
 						Info.Actor = AttachedActor;
@@ -377,4 +377,4 @@ void UChildActorComponent::DestroyChildActor(const bool bRequiresRename)
 	}
 
 	ChildActor = nullptr;
-}
+}
\ No newline at end of file
diff --git a/Engine/Source/Runtime/Engine/Private/Components/DestructibleComponent.cpp b/Engine/Source/Runtime/Engine/Private/Components/DestructibleComponent.cpp
--- a/Engine/Source/Runtime/Engine/Private/Components/DestructibleComponent.cpp
+++ b/Engine/Source/Runtime/Engine/Private/Components/DestructibleComponent.cpp
@@ -83,15 +83,16 @@ void UDestructibleComponent::PostEditChangeProperty( struct FPropertyChangedEven
 FBoxSphereBounds UDestructibleComponent::CalcBounds(const FTransform& LocalToWorld) const
 {
 #if WITH_APEX
-	if( ApexDestructibleActor == NULL )
+	if(ApexDestructibleActor == NULL || ApexDestructibleActor->getBounds().isEmpty())
 	{
-		// Fallback if we don't have physics
+		// Fallback if we don't have physics, or we have empty bounds (all chunks inactive/not visible)
 		return Super::CalcBounds(LocalToWorld);
 	}
 
 	const PxBounds3& PBounds = ApexDestructibleActor->getBounds();
 
-	return FBoxSphereBounds( FBox( P2UVector(PBounds.minimum), P2UVector(PBounds.maximum) ) );
+	return FBoxSphereBounds(FBox(P2UVector(PBounds.minimum), P2UVector(PBounds.maximum)));
+
 #else	// #if WITH_APEX
 	return Super::CalcBounds(LocalToWorld);
 #endif	// #if WITH_APEX
diff --git a/Engine/Source/Runtime/Engine/Private/Components/PrimitiveComponent.cpp b/Engine/Source/Runtime/Engine/Private/Components/PrimitiveComponent.cpp
--- a/Engine/Source/Runtime/Engine/Private/Components/PrimitiveComponent.cpp
+++ b/Engine/Source/Runtime/Engine/Private/Components/PrimitiveComponent.cpp
@@ -529,15 +529,19 @@ void UPrimitiveComponent::CreatePhysicsState()
 		if(BodySetup)
 		{
 			// Create new BodyInstance at given location.
-			if(FMath::Abs(ComponentToWorld.GetDeterminant()) < SMALL_NUMBER)
+			FTransform BodyTransform = ComponentToWorld;
+
+			// Here we make sure we don't have zero scale. This still results in a body being made and placed in
+			// world (very small) but is consistent with a body scaled to zero.
+			const FVector BodyScale = BodyTransform.GetScale3D();
+			if(BodyScale.IsNearlyZero())
 			{
-				UE_LOG(LogPrimitiveComponent, Log, TEXT("Zero scaling not supported (%s)"), *GetPathName());
-				return;
+				BodyTransform.SetScale3D(FVector(KINDA_SMALL_NUMBER));
 			}
 
 #if UE_WITH_PHYSICS
 			// Create the body.
-			BodyInstance.InitBody(BodySetup, ComponentToWorld, this, World->GetPhysicsScene());
+			BodyInstance.InitBody(BodySetup, BodyTransform, this, World->GetPhysicsScene());
 #endif //UE_WITH_PHYSICS
 
 #if WITH_EDITOR
@@ -1840,10 +1844,10 @@ void UPrimitiveComponent::DispatchWakeEvents(int32 WakeEvent, FName BoneName)
 		{
 			if (WakeEvent == SleepEvent::SET_Wakeup)
 			{
-				OnComponentWake.Broadcast(BoneName);
+				OnComponentWake.Broadcast(this, BoneName);
 			}else
 			{
-				OnComponentSleep.Broadcast(BoneName);
+				OnComponentSleep.Broadcast(this, BoneName);
 			}
 		}
 	}
@@ -2082,13 +2086,13 @@ void UPrimitiveComponent::BeginComponentOverlap(const FOverlapInfo& OtherOverlap
 				// first execute component delegates
 				if (!IsPendingKill())
 				{
-					OnComponentBeginOverlap.Broadcast(OtherActor, OtherComp, OtherOverlap.GetBodyIndex(), OtherOverlap.bFromSweep, OtherOverlap.OverlapInfo);
+					OnComponentBeginOverlap.Broadcast(this, OtherActor, OtherComp, OtherOverlap.GetBodyIndex(), OtherOverlap.bFromSweep, OtherOverlap.OverlapInfo);
 				}
 
 				if (!OtherComp->IsPendingKill())
 				{
 					// Reverse normals for other component. When it's a sweep, we are the one that moved.
-					OtherComp->OnComponentBeginOverlap.Broadcast(MyActor, this, INDEX_NONE, OtherOverlap.bFromSweep, OtherOverlap.bFromSweep ? FHitResult::GetReversedHit(OtherOverlap.OverlapInfo) : OtherOverlap.OverlapInfo);
+					OtherComp->OnComponentBeginOverlap.Broadcast(OtherComp, MyActor, this, INDEX_NONE, OtherOverlap.bFromSweep, OtherOverlap.bFromSweep ? FHitResult::GetReversedHit(OtherOverlap.OverlapInfo) : OtherOverlap.OverlapInfo);
 				}
 
 				// then execute actor notification if this is a new actor touch
@@ -2108,12 +2112,12 @@ void UPrimitiveComponent::BeginComponentOverlap(const FOverlapInfo& OtherOverlap
 					// Then level-script delegates
 					if (IsActorValidToNotify(MyActor))
 					{
-						MyActor->OnActorBeginOverlap.Broadcast(OtherActor);
+						MyActor->OnActorBeginOverlap.Broadcast(MyActor, OtherActor);
 					}
 
 					if (IsActorValidToNotify(OtherActor))
 					{
-						OtherActor->OnActorBeginOverlap.Broadcast(MyActor);
+						OtherActor->OnActorBeginOverlap.Broadcast(OtherActor, MyActor);
 					}
 				}
 			}
@@ -2151,12 +2155,12 @@ void UPrimitiveComponent::EndComponentOverlap(const FOverlapInfo& OtherOverlap,
 			{
 				if (!bSkipNotifySelf && IsPrimCompValidAndAlive(this))
 				{
-					OnComponentEndOverlap.Broadcast(OtherActor, OtherComp, OtherOverlap.GetBodyIndex());
+					OnComponentEndOverlap.Broadcast(this, OtherActor, OtherComp, OtherOverlap.GetBodyIndex());
 				}
 
 				if (IsPrimCompValidAndAlive(OtherComp))
 				{
-					OtherComp->OnComponentEndOverlap.Broadcast(MyActor, this, INDEX_NONE);
+					OtherComp->OnComponentEndOverlap.Broadcast(OtherComp, MyActor, this, INDEX_NONE);
 				}
 	
 				// if this was the last touch on the other actor by this actor, notify that we've untouched the actor as well
@@ -2165,13 +2169,13 @@ void UPrimitiveComponent::EndComponentOverlap(const FOverlapInfo& OtherOverlap,
 					if (IsActorValidToNotify(MyActor))
 					{
 						MyActor->NotifyActorEndOverlap(OtherActor);
-						MyActor->OnActorEndOverlap.Broadcast(OtherActor);
+						MyActor->OnActorEndOverlap.Broadcast(MyActor, OtherActor);
 					}
 
 					if (IsActorValidToNotify(OtherActor))
 					{
 						OtherActor->NotifyActorEndOverlap(MyActor);
-						OtherActor->OnActorEndOverlap.Broadcast(MyActor);
+						OtherActor->OnActorEndOverlap.Broadcast(OtherActor, MyActor);
 					}
 				}
 			}
@@ -2179,38 +2183,48 @@ void UPrimitiveComponent::EndComponentOverlap(const FOverlapInfo& OtherOverlap,
 	}
 }
 
+void UPrimitiveComponent::GetOverlappingActors(TArray<AActor*>& OutOverlappingActors, TSubclassOf<AActor> ClassFilter) const
+{
+	TSet<AActor*> OverlappingActors;
+	GetOverlappingActors(OverlappingActors, ClassFilter);
+
+	OutOverlappingActors.Reset(OverlappingActors.Num());
 
+	for (AActor* OverlappingActor : OverlappingActors)
+	{
+		OutOverlappingActors.Add(OverlappingActor);
+	}
 
-void UPrimitiveComponent::GetOverlappingActors(TArray<AActor*>& OutOverlappingActors, UClass* ClassFilter) const
+}
+
+void UPrimitiveComponent::GetOverlappingActors(TSet<AActor*>& OutOverlappingActors, TSubclassOf<AActor> ClassFilter) const
 {
 	OutOverlappingActors.Reset();
+	OutOverlappingActors.Reserve(OverlappingComponents.Num());
 
-	for (auto CompIt = OverlappingComponents.CreateConstIterator(); CompIt; ++CompIt)
+	for (const FOverlapInfo& OtherOverlap : OverlappingComponents)
 	{
-		const FOverlapInfo& OtherOvelap = *CompIt;
-		
-		if (OtherOvelap.OverlapInfo.Component.IsValid())
+		if (UPrimitiveComponent* OtherComponent = OtherOverlap.OverlapInfo.Component.Get())
 		{
-			AActor* OtherActor = OtherOvelap.OverlapInfo.Component->GetOwner();
+			AActor* OtherActor = OtherComponent->GetOwner();
 			if (OtherActor)
 			{
-				if ( !ClassFilter || OtherActor->IsA(ClassFilter) )
+				if ( (*ClassFilter) == nullptr || OtherActor->IsA(ClassFilter) )
 				{
-					OutOverlappingActors.AddUnique(OtherActor);
+					OutOverlappingActors.Add(OtherActor);
 				}
 			}
 		}
 	}
 }
 
-
 void UPrimitiveComponent::GetOverlappingComponents(TArray<UPrimitiveComponent*>& OutOverlappingComponents) const
 {
 	OutOverlappingComponents.Reset( OverlappingComponents.Num() );
 
-	for (auto CompIt = OverlappingComponents.CreateConstIterator(); CompIt; ++CompIt)
+	for (const FOverlapInfo& OtherOverlap : OverlappingComponents)
 	{
-		UPrimitiveComponent* const OtherComp = (*CompIt).OverlapInfo.Component.Get();
+		UPrimitiveComponent* const OtherComp = OtherOverlap.OverlapInfo.Component.Get();
 		if (OtherComp)
 		{
 			OutOverlappingComponents.Add(OtherComp);
@@ -2218,7 +2232,6 @@ void UPrimitiveComponent::GetOverlappingComponents(TArray<UPrimitiveComponent*>&
 	}
 }
 
-
 const TArray<FOverlapInfo>* UPrimitiveComponent::ConvertSweptOverlapsToCurrentOverlaps(
 	TArray<FOverlapInfo>& OverlapsAtEndLocation, const TArray<FOverlapInfo>& SweptOverlaps, int32 SweptOverlapsIndex,
 	const FVector& EndLocation, const FQuat& EndRotationQuat)
@@ -2363,7 +2376,7 @@ void UPrimitiveComponent::BeginPlay()
 void UPrimitiveComponent::IgnoreActorWhenMoving(AActor* Actor, bool bShouldIgnore)
 {
 	// Clean up stale references
-	MoveIgnoreActors.RemoveSwap(NULL);
+	MoveIgnoreActors.RemoveSwap(nullptr);
 
 	// Add/Remove the actor from the list
 	if (Actor)
@@ -2381,28 +2394,14 @@ void UPrimitiveComponent::IgnoreActorWhenMoving(AActor* Actor, bool bShouldIgnor
 
 TArray<AActor*> UPrimitiveComponent::CopyArrayOfMoveIgnoreActors()
 {
-	TArray<AActor*> TempMoveIgnoreActors;
-	for (int32 Idx = 0; Idx < MoveIgnoreActors.Num();)
+	for (int32 Index = MoveIgnoreActors.Num() - 1; Index >=0; --Index)
 	{
-		AActor* Actor = MoveIgnoreActors[Idx].Get();
-		if (Actor)
+		const AActor* const MoveIgnoreActor = MoveIgnoreActors[Index];
+		if (MoveIgnoreActor == nullptr || MoveIgnoreActor->IsPendingKill())
 		{
-			TempMoveIgnoreActors.Add(Actor);
-			Idx++;
-		}
-		else
-		{
-			MoveIgnoreActors.RemoveAtSwap(Idx, 1, false);
+			MoveIgnoreActors.RemoveAtSwap(Index,1,false);
 		}
 	}
-	MoveIgnoreActors.Shrink();
-	return TempMoveIgnoreActors;
-}
-
-TArray<TWeakObjectPtr<AActor> > & UPrimitiveComponent::GetMoveIgnoreActors()
-{
-	// Clean up stale references
-	MoveIgnoreActors.RemoveSwap(NULL);
 	return MoveIgnoreActors;
 }
 
@@ -2411,7 +2410,6 @@ void UPrimitiveComponent::ClearMoveIgnoreActors()
 	MoveIgnoreActors.Empty();
 }
 
-
 void UPrimitiveComponent::UpdateOverlaps(const TArray<FOverlapInfo>* NewPendingOverlaps, bool bDoNotifies, const TArray<FOverlapInfo>* OverlapsAtEndLocation)
 {
 	SCOPE_CYCLE_COUNTER(STAT_UpdateOverlaps); 
@@ -2664,7 +2662,7 @@ void UPrimitiveComponent::DispatchMouseOverEvents(UPrimitiveComponent* CurrentCo
 					CurrentOwner->NotifyActorEndCursorOver();
 					if (IsActorValidToNotify(CurrentOwner))
 					{
-						CurrentOwner->OnEndCursorOver.Broadcast();
+						CurrentOwner->OnEndCursorOver.Broadcast(CurrentOwner);
 					}
 				}
 			}
@@ -2677,7 +2675,7 @@ void UPrimitiveComponent::DispatchMouseOverEvents(UPrimitiveComponent* CurrentCo
 				NewOwner->NotifyActorBeginCursorOver();
 				if (IsActorValidToNotify(NewOwner))
 				{
-					NewOwner->OnBeginCursorOver.Broadcast();
+					NewOwner->OnBeginCursorOver.Broadcast(NewOwner);
 				}
 			}
 			if (!NewComponent->IsPendingKill())
@@ -2700,7 +2698,7 @@ void UPrimitiveComponent::DispatchMouseOverEvents(UPrimitiveComponent* CurrentCo
 			CurrentOwner->NotifyActorEndCursorOver();
 			if (IsActorValidToNotify(CurrentOwner))
 			{
-				CurrentOwner->OnEndCursorOver.Broadcast();
+				CurrentOwner->OnEndCursorOver.Broadcast(CurrentOwner);
 			}
 		}
 	}
@@ -2735,7 +2733,7 @@ void UPrimitiveComponent::DispatchTouchOverEvents(ETouchIndex::Type FingerIndex,
 					CurrentOwner->NotifyActorOnInputTouchLeave(FingerIndex);
 					if (IsActorValidToNotify(CurrentOwner))
 					{
-						CurrentOwner->OnInputTouchLeave.Broadcast(FingerIndex);
+						CurrentOwner->OnInputTouchLeave.Broadcast(FingerIndex, CurrentOwner);
 					}
 				}
 			}
@@ -2748,7 +2746,7 @@ void UPrimitiveComponent::DispatchTouchOverEvents(ETouchIndex::Type FingerIndex,
 				NewOwner->NotifyActorOnInputTouchEnter(FingerIndex);
 				if (IsActorValidToNotify(NewOwner))
 				{
-					NewOwner->OnInputTouchEnter.Broadcast(FingerIndex);
+					NewOwner->OnInputTouchEnter.Broadcast(FingerIndex, NewOwner);
 				}
 			}
 			if (!NewComponent->IsPendingKill())
@@ -2771,7 +2769,7 @@ void UPrimitiveComponent::DispatchTouchOverEvents(ETouchIndex::Type FingerIndex,
 			CurrentOwner->NotifyActorOnInputTouchLeave(FingerIndex);
 			if (IsActorValidToNotify(CurrentOwner))
 			{
-				CurrentOwner->OnInputTouchLeave.Broadcast(FingerIndex);
+				CurrentOwner->OnInputTouchLeave.Broadcast(FingerIndex, CurrentOwner);
 			}
 		}
 	}
@@ -2784,7 +2782,7 @@ void UPrimitiveComponent::DispatchOnClicked(FKey ButtonPressed)
 		GetOwner()->NotifyActorOnClicked(ButtonPressed);
 		if (IsActorValidToNotify(GetOwner()))
 		{
-			GetOwner()->OnClicked.Broadcast(ButtonPressed);
+			GetOwner()->OnClicked.Broadcast(GetOwner(), ButtonPressed);
 		}
 	}
 
@@ -2801,7 +2799,7 @@ void UPrimitiveComponent::DispatchOnReleased(FKey ButtonReleased)
 		GetOwner()->NotifyActorOnReleased(ButtonReleased);
 		if (IsActorValidToNotify(GetOwner()))
 		{
-			GetOwner()->OnReleased.Broadcast(ButtonReleased);
+			GetOwner()->OnReleased.Broadcast(GetOwner(), ButtonReleased);
 		}
 	}
 
@@ -2818,7 +2816,7 @@ void UPrimitiveComponent::DispatchOnInputTouchBegin(const ETouchIndex::Type Fing
 		GetOwner()->NotifyActorOnInputTouchBegin(FingerIndex);
 		if (IsActorValidToNotify(GetOwner()))
 		{
-			GetOwner()->OnInputTouchBegin.Broadcast(FingerIndex);
+			GetOwner()->OnInputTouchBegin.Broadcast(FingerIndex, GetOwner());
 		}
 	}
 
@@ -2835,7 +2833,7 @@ void UPrimitiveComponent::DispatchOnInputTouchEnd(const ETouchIndex::Type Finger
 		GetOwner()->NotifyActorOnInputTouchEnd(FingerIndex);
 		if (IsActorValidToNotify(GetOwner()))
 		{
-			GetOwner()->OnInputTouchEnd.Broadcast(FingerIndex);
+			GetOwner()->OnInputTouchEnd.Broadcast(FingerIndex, GetOwner());
 		}
 	}
 
@@ -2847,7 +2845,7 @@ void UPrimitiveComponent::DispatchOnInputTouchEnd(const ETouchIndex::Type Finger
 
 SIZE_T UPrimitiveComponent::GetResourceSize( EResourceSizeMode::Type Mode )
 {
-	SIZE_T ResSize = 0;
+	SIZE_T ResSize = Super::GetResourceSize(Mode);
 
 	if (BodyInstance.IsValidBodyInstance())
 	{
diff --git a/Engine/Source/Runtime/Engine/Private/Components/SceneComponent.cpp b/Engine/Source/Runtime/Engine/Private/Components/SceneComponent.cpp
--- a/Engine/Source/Runtime/Engine/Private/Components/SceneComponent.cpp
+++ b/Engine/Source/Runtime/Engine/Private/Components/SceneComponent.cpp
@@ -2923,10 +2923,10 @@ void USceneComponent::UpdateNavigationData()
 {
 	SCOPE_CYCLE_COUNTER(STAT_ComponentUpdateNavData);
 
-	if (UNavigationSystem::ShouldUpdateNavOctreeOnComponentChange() &&
-		IsRegistered() && World && World->IsGameWorld() &&
-		World->GetNetMode() < ENetMode::NM_Client)
+	if (UNavigationSystem::ShouldUpdateNavOctreeOnComponentChange() && IsRegistered() && 
+		((World == nullptr) || !World->IsGameWorld() || World->GetNetMode() < ENetMode::NM_Client))
 	{
+		// use propagated component's transform update in editor OR server game with additional navsys check
 		UNavigationSystem::UpdateComponentInNavOctree(*this);
 	}
 }
@@ -2935,10 +2935,10 @@ void USceneComponent::PostUpdateNavigationData()
 {
 	SCOPE_CYCLE_COUNTER(STAT_ComponentPostUpdateNavData);
 
-	if (UNavigationSystem::ShouldUpdateNavOctreeOnComponentChange() &&
-		IsRegistered() && World && World->IsGameWorld() &&
-		World->GetNetMode() < ENetMode::NM_Client)
+	if (UNavigationSystem::ShouldUpdateNavOctreeOnComponentChange() && IsRegistered() &&
+		((World == nullptr) || !World->IsGameWorld() || World->GetNetMode() < ENetMode::NM_Client))
 	{
+		// use propagated component's transform update in editor OR server game with additional navsys check
 		UNavigationSystem::UpdateNavOctreeAfterMove(this);
 	}
 }
diff --git a/Engine/Source/Runtime/Engine/Private/Components/SkinnedMeshComponent.cpp b/Engine/Source/Runtime/Engine/Private/Components/SkinnedMeshComponent.cpp
--- a/Engine/Source/Runtime/Engine/Private/Components/SkinnedMeshComponent.cpp
+++ b/Engine/Source/Runtime/Engine/Private/Components/SkinnedMeshComponent.cpp
@@ -705,13 +705,13 @@ FBoxSphereBounds USkinnedMeshComponent::CalcMeshBound(const FVector& RootOffset,
 	// if not visible, or we were told to use fixed bounds, use skelmesh bounds
 	if ( (!bVisible || bComponentUseFixedSkelBounds) && SkeletalMesh ) 
 	{
-		FBoxSphereBounds RootAdjustedBounds = SkeletalMesh->Bounds;
+		FBoxSphereBounds RootAdjustedBounds = SkeletalMesh->GetBounds();
 		RootAdjustedBounds.Origin += RootOffset; // Adjust bounds by root bone translation
 		NewBounds = RootAdjustedBounds.TransformBy(LocalToWorld);
 	}
 	else if(MasterPoseComponentInst && MasterPoseComponentInst->SkeletalMesh && MasterPoseComponentInst->bComponentUseFixedSkelBounds)
 	{
-		FBoxSphereBounds RootAdjustedBounds = MasterPoseComponentInst->SkeletalMesh->Bounds;
+		FBoxSphereBounds RootAdjustedBounds = MasterPoseComponentInst->SkeletalMesh->GetBounds();
 		RootAdjustedBounds.Origin += RootOffset; // Adjust bounds by root bone translation
 		NewBounds = RootAdjustedBounds.TransformBy(LocalToWorld);
 	}
@@ -740,7 +740,7 @@ FBoxSphereBounds USkinnedMeshComponent::CalcMeshBound(const FVector& RootOffset,
 	// Fallback is to use the one from the skeletal mesh. Usually pretty bad in terms of Accuracy of where the SkelMesh Bounds are located (i.e. usually bigger than it needs to be)
 	else if( SkeletalMesh )
 	{
-		FBoxSphereBounds RootAdjustedBounds = SkeletalMesh->Bounds;
+		FBoxSphereBounds RootAdjustedBounds = SkeletalMesh->GetBounds();
 
 		// Adjust bounds by root bone translation
 		RootAdjustedBounds.Origin += RootOffset;
diff --git a/Engine/Source/Runtime/Engine/Private/EngineUtils.cpp b/Engine/Source/Runtime/Engine/Private/EngineUtils.cpp
--- a/Engine/Source/Runtime/Engine/Private/EngineUtils.cpp
+++ b/Engine/Source/Runtime/Engine/Private/EngineUtils.cpp
@@ -26,7 +26,7 @@ FContentComparisonHelper::FContentComparisonHelper()
 	{
 		for( FConfigSectionMap::TIterator It(*RefTypes); It; ++It )
 		{
-			const FString RefType = It.Value();
+			const FString& RefType = It.Value().GetValue();
 			ReferenceClassesOfInterest.Add(RefType, true);
 			UE_LOG(LogEngineUtils, Log, TEXT("Adding class of interest: %s"), *RefType);
 		}
diff --git a/Engine/Source/Runtime/Engine/Private/GameViewportClient.cpp b/Engine/Source/Runtime/Engine/Private/GameViewportClient.cpp
--- a/Engine/Source/Runtime/Engine/Private/GameViewportClient.cpp
+++ b/Engine/Source/Runtime/Engine/Private/GameViewportClient.cpp
@@ -2426,11 +2426,12 @@ void UGameViewportClient::ToggleShowCollision()
 				{
 					PC->ClientMessage(FString::Printf(TEXT("Base %s"), *PC->GetPawn()->GetMovementBase()->GetName()));
 				}
-				TArray<AActor*> Touching;
-				PC->GetPawn()->GetOverlappingActors(Touching);
-				for (int32 i = 0; i < Touching.Num(); i++)
+				TSet<AActor*> TouchingActors;
+				PC->GetPawn()->GetOverlappingActors(TouchingActors);
+				int32 i = 0;
+				for (AActor* TouchingActor : TouchingActors)
 				{
-					PC->ClientMessage(FString::Printf(TEXT("Touching %d: %s"), i, *Touching[i]->GetName()));
+					PC->ClientMessage(FString::Printf(TEXT("Touching %d: %s"), i++, *TouchingActor->GetName()));
 				}
 			}
 		}
diff --git a/Engine/Source/Runtime/Engine/Private/GameplayStatics.cpp b/Engine/Source/Runtime/Engine/Private/GameplayStatics.cpp
--- a/Engine/Source/Runtime/Engine/Private/GameplayStatics.cpp
+++ b/Engine/Source/Runtime/Engine/Private/GameplayStatics.cpp
@@ -896,6 +896,7 @@ void UGameplayStatics::PlaySound2D(UObject* WorldContextObject, class USoundBase
 		NewActiveSound.RequestedStartTime = FMath::Max(0.f, StartTime);
 
 		NewActiveSound.bIsUISound = true;
+		NewActiveSound.bAllowSpatialization = false;
 		NewActiveSound.ConcurrencySettings = ConcurrencySettings;
 		NewActiveSound.Priority = Sound->Priority;
 
@@ -1239,6 +1240,26 @@ void UGameplayStatics::DeactivateReverbEffect(UObject* WorldContextObject, FName
 	}
 }
 
+class UReverbEffect* UGameplayStatics::GetCurrentReverbEffect(UObject* WorldContextObject)
+{
+	if (GEngine == nullptr || !GEngine->UseSound())
+	{
+		return nullptr;
+	}
+
+	UWorld* ThisWorld = GEngine->GetWorldFromContextObject(WorldContextObject);
+	if (!ThisWorld || !ThisWorld->bAllowAudioPlayback)
+	{
+		return nullptr;
+	}
+
+	if (FAudioDevice* AudioDevice = ThisWorld->GetAudioDevice())
+	{
+		return AudioDevice->GetCurrentReverbEffect();
+	}
+	return nullptr;
+}
+
 UDecalComponent* CreateDecalComponent(class UMaterialInterface* DecalMaterial, FVector DecalSize, UWorld* World, AActor* Actor, float LifeSpan)
 {
 	UDecalComponent* DecalComp = NewObject<UDecalComponent>((Actor ? Actor : (UObject*)World));
diff --git a/Engine/Source/Runtime/Engine/Private/HierarchicalInstancedStaticMesh.cpp b/Engine/Source/Runtime/Engine/Private/HierarchicalInstancedStaticMesh.cpp
--- a/Engine/Source/Runtime/Engine/Private/HierarchicalInstancedStaticMesh.cpp
+++ b/Engine/Source/Runtime/Engine/Private/HierarchicalInstancedStaticMesh.cpp
@@ -1714,15 +1714,6 @@ void UHierarchicalInstancedStaticMeshComponent::Serialize(FArchive& Ar)
 	}
 }
 
-SIZE_T UHierarchicalInstancedStaticMeshComponent::GetResourceSize( EResourceSizeMode::Type Mode )
-{
-	SIZE_T ResSize = Super::GetResourceSize(Mode);
-
-	ResSize += SortedInstances.GetAllocatedSize();
-
-	return ResSize;
-}
-
 void UHierarchicalInstancedStaticMeshComponent::RemoveInstanceInternal(int32 InstanceIndex)
 {
 	PartialNavigationUpdate(InstanceIndex);
diff --git a/Engine/Source/Runtime/Engine/Private/KismetSystemLibrary.cpp b/Engine/Source/Runtime/Engine/Private/KismetSystemLibrary.cpp
--- a/Engine/Source/Runtime/Engine/Private/KismetSystemLibrary.cpp
+++ b/Engine/Source/Runtime/Engine/Private/KismetSystemLibrary.cpp
@@ -1337,12 +1337,12 @@ bool UKismetSystemLibrary::ComponentOverlapComponents_NEW(UPrimitiveComponent* C
 static const float KISMET_TRACE_DEBUG_DRAW_DURATION = 5.f;
 static const float KISMET_TRACE_DEBUG_IMPACTPOINT_SIZE = 16.f;
 
-bool UKismetSystemLibrary::LineTraceSingle_NEW(UObject* WorldContextObject, const FVector Start, const FVector End, ETraceTypeQuery TraceChannel, bool bTraceComplex, const TArray<AActor*>& ActorsToIgnore, EDrawDebugTrace::Type DrawDebugType, FHitResult& OutHit, bool bIgnoreSelf)
+bool UKismetSystemLibrary::LineTraceSingle_NEW(UObject* WorldContextObject, const FVector Start, const FVector End, ETraceTypeQuery TraceChannel, bool bTraceComplex, const TArray<AActor*>& ActorsToIgnore, EDrawDebugTrace::Type DrawDebugType, FHitResult& OutHit, bool bIgnoreSelf, FLinearColor TraceColor, FLinearColor TraceHitColor, float DrawTime)
 {
-	return LineTraceSingle_DEPRECATED(WorldContextObject, Start, End, UEngineTypes::ConvertToCollisionChannel(TraceChannel), bTraceComplex, ActorsToIgnore, DrawDebugType, OutHit, bIgnoreSelf);
+	return LineTraceSingle_DEPRECATED(WorldContextObject, Start, End, UEngineTypes::ConvertToCollisionChannel(TraceChannel), bTraceComplex, ActorsToIgnore, DrawDebugType, OutHit, bIgnoreSelf, TraceColor, TraceHitColor, DrawTime);
 }
 
-bool UKismetSystemLibrary::LineTraceSingle_DEPRECATED(UObject* WorldContextObject, const FVector Start, const FVector End, ECollisionChannel TraceChannel, bool bTraceComplex, const TArray<AActor*>& ActorsToIgnore, EDrawDebugTrace::Type DrawDebugType, FHitResult& OutHit, bool bIgnoreSelf)
+bool UKismetSystemLibrary::LineTraceSingle_DEPRECATED(UObject* WorldContextObject, const FVector Start, const FVector End, ECollisionChannel TraceChannel, bool bTraceComplex, const TArray<AActor*>& ActorsToIgnore, EDrawDebugTrace::Type DrawDebugType, FHitResult& OutHit, bool bIgnoreSelf, FLinearColor TraceColor, FLinearColor TraceHitColor, float DrawTime)
 {
 	static const FName LineTraceSingleName(TEXT("LineTraceSingle"));
 
@@ -1380,32 +1380,32 @@ bool UKismetSystemLibrary::LineTraceSingle_DEPRECATED(UObject* WorldContextObjec
 	if (DrawDebugType != EDrawDebugTrace::None)
 	{
 		bool bPersistent = DrawDebugType == EDrawDebugTrace::Persistent;
-		float LifeTime = (DrawDebugType == EDrawDebugTrace::ForDuration)? KISMET_TRACE_DEBUG_DRAW_DURATION: 0.f;
+		float LifeTime = (DrawDebugType == EDrawDebugTrace::ForDuration)? DrawTime : 0.f;
 
 		// @fixme, draw line with thickness = 2.f?
 		if (bHit && OutHit.bBlockingHit)
 		{
 			// Red up to the blocking hit, green thereafter
-			::DrawDebugLine(World, Start, OutHit.ImpactPoint, FColor::Red, bPersistent, LifeTime);
-			::DrawDebugLine(World, OutHit.ImpactPoint, End, FColor::Green, bPersistent, LifeTime);
-			::DrawDebugPoint(World, OutHit.ImpactPoint, KISMET_TRACE_DEBUG_IMPACTPOINT_SIZE, FColor::Red, bPersistent, LifeTime);
+			::DrawDebugLine(World, Start, OutHit.ImpactPoint, TraceColor.ToFColor(true), bPersistent, LifeTime);
+			::DrawDebugLine(World, OutHit.ImpactPoint, End, TraceHitColor.ToFColor(true), bPersistent, LifeTime);
+			::DrawDebugPoint(World, OutHit.ImpactPoint, KISMET_TRACE_DEBUG_IMPACTPOINT_SIZE, TraceColor.ToFColor(true), bPersistent, LifeTime);
 		}
 		else
 		{
 			// no hit means all red
-			::DrawDebugLine(World, Start, End, FColor::Red, bPersistent, LifeTime);
+			::DrawDebugLine(World, Start, End, TraceColor.ToFColor(true), bPersistent, LifeTime);
 		}
 	}
 
 	return bHit;
 }
 
-bool UKismetSystemLibrary::LineTraceMulti_NEW(UObject* WorldContextObject, const FVector Start, const FVector End, ETraceTypeQuery TraceChannel, bool bTraceComplex, const TArray<AActor*>& ActorsToIgnore, EDrawDebugTrace::Type DrawDebugType, TArray<FHitResult>& OutHits, bool bIgnoreSelf)
+bool UKismetSystemLibrary::LineTraceMulti_NEW(UObject* WorldContextObject, const FVector Start, const FVector End, ETraceTypeQuery TraceChannel, bool bTraceComplex, const TArray<AActor*>& ActorsToIgnore, EDrawDebugTrace::Type DrawDebugType, TArray<FHitResult>& OutHits, bool bIgnoreSelf, FLinearColor TraceColor, FLinearColor TraceHitColor, float DrawTime)
 {
-	return LineTraceMulti_DEPRECATED(WorldContextObject, Start, End, UEngineTypes::ConvertToCollisionChannel(TraceChannel), bTraceComplex, ActorsToIgnore, DrawDebugType, OutHits, bIgnoreSelf);
+	return LineTraceMulti_DEPRECATED(WorldContextObject, Start, End, UEngineTypes::ConvertToCollisionChannel(TraceChannel), bTraceComplex, ActorsToIgnore, DrawDebugType, OutHits, bIgnoreSelf, TraceColor, TraceHitColor, DrawTime);
 }
 
-bool UKismetSystemLibrary::LineTraceMulti_DEPRECATED(UObject* WorldContextObject, const FVector Start, const FVector End, ECollisionChannel TraceChannel, bool bTraceComplex, const TArray<AActor*>& ActorsToIgnore, EDrawDebugTrace::Type DrawDebugType, TArray<FHitResult>& OutHits, bool bIgnoreSelf)
+bool UKismetSystemLibrary::LineTraceMulti_DEPRECATED(UObject* WorldContextObject, const FVector Start, const FVector End, ECollisionChannel TraceChannel, bool bTraceComplex, const TArray<AActor*>& ActorsToIgnore, EDrawDebugTrace::Type DrawDebugType, TArray<FHitResult>& OutHits, bool bIgnoreSelf, FLinearColor TraceColor, FLinearColor TraceHitColor, float DrawTime)
 {
 	static const FName LineTraceMultiName(TEXT("LineTraceMulti"));
 
@@ -1443,27 +1443,27 @@ bool UKismetSystemLibrary::LineTraceMulti_DEPRECATED(UObject* WorldContextObject
 	if (DrawDebugType != EDrawDebugTrace::None)
 	{
 		bool bPersistent = DrawDebugType == EDrawDebugTrace::Persistent;
-		float LifeTime = (DrawDebugType == EDrawDebugTrace::ForDuration)? KISMET_TRACE_DEBUG_DRAW_DURATION: 0.f;
+		float LifeTime = (DrawDebugType == EDrawDebugTrace::ForDuration)? DrawTime : 0.f;
 
 		// @fixme, draw line with thickness = 2.f?
 		if (bHit && OutHits.Last().bBlockingHit)
 		{
 			// Red up to the blocking hit, green thereafter
 			FVector const BlockingHitPoint = OutHits.Last().ImpactPoint;
-			::DrawDebugLine(World, Start, BlockingHitPoint , FColor::Red, bPersistent, LifeTime);
-			::DrawDebugLine(World, BlockingHitPoint, End, FColor::Green, bPersistent, LifeTime);
+			::DrawDebugLine(World, Start, BlockingHitPoint , TraceColor.ToFColor(true), bPersistent, LifeTime);
+			::DrawDebugLine(World, BlockingHitPoint, End, TraceHitColor.ToFColor(true), bPersistent, LifeTime);
 		}
 		else
 		{
 			// no hit means all red
-			::DrawDebugLine(World, Start, End, FColor::Red, bPersistent, LifeTime);
+			::DrawDebugLine(World, Start, End, TraceColor.ToFColor(true), bPersistent, LifeTime);
 		}
 
 		// draw hits
 		for (int32 HitIdx=0; HitIdx<OutHits.Num(); ++HitIdx)
 		{
 			FHitResult const& Hit = OutHits[HitIdx];
-			::DrawDebugPoint(World, Hit.ImpactPoint, KISMET_TRACE_DEBUG_IMPACTPOINT_SIZE, (Hit.bBlockingHit ? FColor::Red : FColor::Green), bPersistent, LifeTime);
+			::DrawDebugPoint(World, Hit.ImpactPoint, KISMET_TRACE_DEBUG_IMPACTPOINT_SIZE, (Hit.bBlockingHit ? TraceColor.ToFColor(true) : TraceHitColor.ToFColor(true)), bPersistent, LifeTime);
 		}
 	}
 
@@ -1907,7 +1907,7 @@ bool UKismetSystemLibrary::CapsuleTraceMulti_DEPRECATED(UObject* WorldContextObj
 }
 
 /** Object Query functions **/
-bool UKismetSystemLibrary::LineTraceSingleForObjects(UObject* WorldContextObject, const FVector Start, const FVector End, const TArray<TEnumAsByte<EObjectTypeQuery> > & ObjectTypes, bool bTraceComplex, const TArray<AActor*>& ActorsToIgnore, EDrawDebugTrace::Type DrawDebugType, FHitResult& OutHit, bool bIgnoreSelf)
+bool UKismetSystemLibrary::LineTraceSingleForObjects(UObject* WorldContextObject, const FVector Start, const FVector End, const TArray<TEnumAsByte<EObjectTypeQuery> > & ObjectTypes, bool bTraceComplex, const TArray<AActor*>& ActorsToIgnore, EDrawDebugTrace::Type DrawDebugType, FHitResult& OutHit, bool bIgnoreSelf, FLinearColor TraceColor, FLinearColor TraceHitColor, float DrawTime)
 {
 
 	TArray<TEnumAsByte<ECollisionChannel>> CollisionObjectTraces;
@@ -1918,10 +1918,10 @@ bool UKismetSystemLibrary::LineTraceSingleForObjects(UObject* WorldContextObject
 		CollisionObjectTraces[Iter.GetIndex()] = UEngineTypes::ConvertToCollisionChannel(*Iter);
 	}
 
-	return LineTraceSingleByObject_DEPRECATED(WorldContextObject, Start, End, CollisionObjectTraces, bTraceComplex, ActorsToIgnore, DrawDebugType, OutHit, bIgnoreSelf);
+	return LineTraceSingleByObject_DEPRECATED(WorldContextObject, Start, End, CollisionObjectTraces, bTraceComplex, ActorsToIgnore, DrawDebugType, OutHit, bIgnoreSelf, TraceColor, TraceHitColor, DrawTime);
 }
 
-bool UKismetSystemLibrary::LineTraceSingleByObject_DEPRECATED(UObject* WorldContextObject, const FVector Start, const FVector End, const TArray<TEnumAsByte<ECollisionChannel> > & ObjectsToTrace, bool bTraceComplex, const TArray<AActor*>& ActorsToIgnore, EDrawDebugTrace::Type DrawDebugType, FHitResult& OutHit, bool bIgnoreSelf)
+bool UKismetSystemLibrary::LineTraceSingleByObject_DEPRECATED(UObject* WorldContextObject, const FVector Start, const FVector End, const TArray<TEnumAsByte<ECollisionChannel> > & ObjectsToTrace, bool bTraceComplex, const TArray<AActor*>& ActorsToIgnore, EDrawDebugTrace::Type DrawDebugType, FHitResult& OutHit, bool bIgnoreSelf, FLinearColor TraceColor, FLinearColor TraceHitColor, float DrawTime)
 {
 	static const FName LineTraceSingleName(TEXT("LineTraceSingle"));
 
@@ -1979,27 +1979,27 @@ bool UKismetSystemLibrary::LineTraceSingleByObject_DEPRECATED(UObject* WorldCont
 	if (DrawDebugType != EDrawDebugTrace::None)
 	{
 		bool bPersistent = DrawDebugType == EDrawDebugTrace::Persistent;
-		float LifeTime = (DrawDebugType == EDrawDebugTrace::ForDuration)? KISMET_TRACE_DEBUG_DRAW_DURATION: 0.f;
+		float LifeTime = (DrawDebugType == EDrawDebugTrace::ForDuration)? DrawTime : 0.f;
 
 		// @fixme, draw line with thickness = 2.f?
 		if (bHit && OutHit.bBlockingHit)
 		{
 			// Red up to the blocking hit, green thereafter
-			::DrawDebugLine(World, Start, OutHit.ImpactPoint, FColor::Red, bPersistent, LifeTime);
-			::DrawDebugLine(World, OutHit.ImpactPoint, End, FColor::Green, bPersistent, LifeTime);
-			::DrawDebugPoint(World, OutHit.ImpactPoint, KISMET_TRACE_DEBUG_IMPACTPOINT_SIZE, FColor::Red, bPersistent, LifeTime);
+			::DrawDebugLine(World, Start, OutHit.ImpactPoint, TraceColor.ToFColor(true), bPersistent, LifeTime);
+			::DrawDebugLine(World, OutHit.ImpactPoint, End, TraceHitColor.ToFColor(true), bPersistent, LifeTime);
+			::DrawDebugPoint(World, OutHit.ImpactPoint, KISMET_TRACE_DEBUG_IMPACTPOINT_SIZE, TraceColor.ToFColor(true), bPersistent, LifeTime);
 		}
 		else
 		{
 			// no hit means all red
-			::DrawDebugLine(World, Start, End, FColor::Red, bPersistent, LifeTime);
+			::DrawDebugLine(World, Start, End, TraceColor.ToFColor(true), bPersistent, LifeTime);
 		}
 	}
 
 	return bHit;
 }
 
-bool UKismetSystemLibrary::LineTraceMultiForObjects(UObject* WorldContextObject, const FVector Start, const FVector End, const TArray<TEnumAsByte<EObjectTypeQuery> > & ObjectTypes, bool bTraceComplex, const TArray<AActor*>& ActorsToIgnore, EDrawDebugTrace::Type DrawDebugType, TArray<FHitResult>& OutHits, bool bIgnoreSelf)
+bool UKismetSystemLibrary::LineTraceMultiForObjects(UObject* WorldContextObject, const FVector Start, const FVector End, const TArray<TEnumAsByte<EObjectTypeQuery> > & ObjectTypes, bool bTraceComplex, const TArray<AActor*>& ActorsToIgnore, EDrawDebugTrace::Type DrawDebugType, TArray<FHitResult>& OutHits, bool bIgnoreSelf, FLinearColor TraceColor, FLinearColor TraceHitColor, float DrawTime)
 {
 	TArray<TEnumAsByte<ECollisionChannel>> CollisionObjectTraces;
 	CollisionObjectTraces.AddUninitialized(ObjectTypes.Num());
@@ -2009,10 +2009,10 @@ bool UKismetSystemLibrary::LineTraceMultiForObjects(UObject* WorldContextObject,
 		CollisionObjectTraces[Iter.GetIndex()] = UEngineTypes::ConvertToCollisionChannel(*Iter);
 	}
 
-	return LineTraceMultiByObject_DEPRECATED(WorldContextObject, Start, End, CollisionObjectTraces, bTraceComplex, ActorsToIgnore, DrawDebugType, OutHits, bIgnoreSelf);
+	return LineTraceMultiByObject_DEPRECATED(WorldContextObject, Start, End, CollisionObjectTraces, bTraceComplex, ActorsToIgnore, DrawDebugType, OutHits, bIgnoreSelf, TraceColor, TraceHitColor, DrawTime);
 }
 
-bool UKismetSystemLibrary::LineTraceMultiByObject_DEPRECATED(UObject* WorldContextObject, const FVector Start, const FVector End, const TArray<TEnumAsByte<ECollisionChannel> > & ObjectsToTrace, bool bTraceComplex, const TArray<AActor*>& ActorsToIgnore, EDrawDebugTrace::Type DrawDebugType, TArray<FHitResult>& OutHits, bool bIgnoreSelf)
+bool UKismetSystemLibrary::LineTraceMultiByObject_DEPRECATED(UObject* WorldContextObject, const FVector Start, const FVector End, const TArray<TEnumAsByte<ECollisionChannel> > & ObjectsToTrace, bool bTraceComplex, const TArray<AActor*>& ActorsToIgnore, EDrawDebugTrace::Type DrawDebugType, TArray<FHitResult>& OutHits, bool bIgnoreSelf, FLinearColor TraceColor, FLinearColor TraceHitColor, float DrawTime)
 {
 	static const FName LineTraceMultiName(TEXT("LineTraceMulti"));
 
@@ -2070,27 +2070,27 @@ bool UKismetSystemLibrary::LineTraceMultiByObject_DEPRECATED(UObject* WorldConte
 	if (DrawDebugType != EDrawDebugTrace::None)
 	{
 		bool bPersistent = DrawDebugType == EDrawDebugTrace::Persistent;
-		float LifeTime = (DrawDebugType == EDrawDebugTrace::ForDuration)? KISMET_TRACE_DEBUG_DRAW_DURATION: 0.f;
+		float LifeTime = (DrawDebugType == EDrawDebugTrace::ForDuration)? DrawTime : 0.f;
 
 		// @fixme, draw line with thickness = 2.f?
 		if (bHit && OutHits.Last().bBlockingHit)
 		{
 			// Red up to the blocking hit, green thereafter
 			FVector const BlockingHitPoint = OutHits.Last().ImpactPoint;
-			::DrawDebugLine(World, Start, BlockingHitPoint , FColor::Red, bPersistent, LifeTime);
-			::DrawDebugLine(World, BlockingHitPoint, End, FColor::Green, bPersistent, LifeTime);
+			::DrawDebugLine(World, Start, BlockingHitPoint , TraceColor.ToFColor(true), bPersistent, LifeTime);
+			::DrawDebugLine(World, BlockingHitPoint, End, TraceHitColor.ToFColor(true), bPersistent, LifeTime);
 		}
 		else
 		{
 			// no hit means all red
-			::DrawDebugLine(World, Start, End, FColor::Red, bPersistent, LifeTime);
+			::DrawDebugLine(World, Start, End, TraceColor.ToFColor(true), bPersistent, LifeTime);
 		}
 
 		// draw hits
 		for (int32 HitIdx=0; HitIdx<OutHits.Num(); ++HitIdx)
 		{
 			FHitResult const& Hit = OutHits[HitIdx];
-			::DrawDebugPoint(World, Hit.ImpactPoint, KISMET_TRACE_DEBUG_IMPACTPOINT_SIZE, (Hit.bBlockingHit ? FColor::Red : FColor::Green), bPersistent, LifeTime);
+			::DrawDebugPoint(World, Hit.ImpactPoint, KISMET_TRACE_DEBUG_IMPACTPOINT_SIZE, (Hit.bBlockingHit ? TraceColor.ToFColor(true) : TraceHitColor.ToFColor(true)), bPersistent, LifeTime);
 		}
 	}
 
