Copying //UE4/Dev-Framework to Dev-Main (//UE4/Dev-Main) @ 2926658

diff --git a/Engine/Source/Runtime/Engine/Private/LevelActor.cpp b/Engine/Source/Runtime/Engine/Private/LevelActor.cpp
--- a/Engine/Source/Runtime/Engine/Private/LevelActor.cpp
+++ b/Engine/Source/Runtime/Engine/Private/LevelActor.cpp
@@ -563,9 +563,9 @@ bool UWorld::DestroyActor( AActor* ThisActor, bool bNetForce, bool bShouldModify
 			AActor* ChildActor = *AttachedActorIt;
 			if (ChildActor != NULL)
 			{
-				for (auto SceneCompIter = SceneComponents.CreateIterator(); SceneCompIter; ++SceneCompIter)
+				for (USceneComponent* SceneComponent : SceneComponents)
 				{
-					ChildActor->DetachSceneComponentsFromParent(*SceneCompIter, true);
+					ChildActor->DetachSceneComponentsFromParent(SceneComponent, true);
 				}
 #if WITH_EDITOR
 				if( GIsEditor )
@@ -756,8 +756,10 @@ bool UWorld::FindTeleportSpot(AActor* TestActor, FVector& TestLocation, FRotator
 	return !EncroachingBlockingGeometry(TestActor, TestLocation, TestRotation, &Adjust);
 }
 
+static FName NAME_ComponentEncroachesBlockingGeometry_NoAdjustment = FName(TEXT("ComponentEncroachesBlockingGeometry_NoAdjustment"));
+
 /** Tests shape components more efficiently than the with-adjustment case, but does less-efficient ppr-poly collision for meshes. */
-static bool ComponentEncroachesBlockingGeometry_NoAdjustment(UWorld const* World, AActor const* TestActor, UPrimitiveComponent const* PrimComp, FTransform const& TestWorldTransform)
+static bool ComponentEncroachesBlockingGeometry_NoAdjustment(UWorld const* World, AActor const* TestActor, UPrimitiveComponent const* PrimComp, FTransform const& TestWorldTransform, const TArray<AActor*>& IgnoreActors)
 {	
 	float const Epsilon = CVarEncroachEpsilon.GetValueOnGameThread();
 	
@@ -765,7 +767,6 @@ static bool ComponentEncroachesBlockingGeometry_NoAdjustment(UWorld const* World
 	{
 		bool bFoundBlockingHit = false;
 		
-		static FName NAME_ComponentEncroachesBlockingGeometry_NoAdjustment = FName(TEXT("ComponentEncroachesBlockingGeometry_NoAdjustment"));
 		ECollisionChannel const BlockingChannel = PrimComp->GetCollisionObjectType();
 		FCollisionShape const CollisionShape = PrimComp->GetCollisionShape(-Epsilon);
 
@@ -780,6 +781,7 @@ static bool ComponentEncroachesBlockingGeometry_NoAdjustment(UWorld const* World
 				FComponentQueryParams Params(NAME_ComponentEncroachesBlockingGeometry_NoAdjustment, TestActor);
 				FCollisionResponseParams ResponseParams;
 				PrimComp->InitSweepCollisionParams(Params, ResponseParams);
+				Params.AddIgnoredActors(IgnoreActors);
 				return World->ComponentOverlapMultiByChannel(Overlaps, PrimComp, TestWorldTransform.GetLocation(), TestWorldTransform.GetRotation(), BlockingChannel, Params);
 			}
 			else
@@ -793,6 +795,7 @@ static bool ComponentEncroachesBlockingGeometry_NoAdjustment(UWorld const* World
 			FCollisionQueryParams Params(NAME_ComponentEncroachesBlockingGeometry_NoAdjustment, false, TestActor);
 			FCollisionResponseParams ResponseParams;
 			PrimComp->InitSweepCollisionParams(Params, ResponseParams);
+			Params.AddIgnoredActors(IgnoreActors);
 			return World->OverlapBlockingTestByChannel(TestWorldTransform.GetLocation(), TestWorldTransform.GetRotation(), BlockingChannel, CollisionShape, Params, ResponseParams);
 		}
 	}
@@ -800,8 +803,10 @@ static bool ComponentEncroachesBlockingGeometry_NoAdjustment(UWorld const* World
 	return false;
 }
 
+static FName NAME_ComponentEncroachesBlockingGeometry_WithAdjustment = FName(TEXT("ComponentEncroachesBlockingGeometry_WithAdjustment"));
+
 /** Tests shape components less efficiently than the no-adjustment case, but does quicker aabb collision for meshes. */
-static bool ComponentEncroachesBlockingGeometry_WithAdjustment(UWorld const* World, AActor const* TestActor, UPrimitiveComponent const* PrimComp, FTransform const& TestWorldTransform, FVector& OutProposedAdjustment)
+static bool ComponentEncroachesBlockingGeometry_WithAdjustment(UWorld const* World, AActor const* TestActor, UPrimitiveComponent const* PrimComp, FTransform const& TestWorldTransform, FVector& OutProposedAdjustment, const TArray<AActor*>& IgnoreActors)
 {
 	// init our output
 	OutProposedAdjustment = FVector::ZeroVector;
@@ -814,7 +819,6 @@ static bool ComponentEncroachesBlockingGeometry_WithAdjustment(UWorld const* Wor
 		bool bComputePenetrationAdjustment = true;
 		
 		TArray<FOverlapResult> Overlaps;
-		static FName NAME_ComponentEncroachesBlockingGeometry_WithAdjustment = FName(TEXT("ComponentEncroachesBlockingGeometry_WithAdjustment"));
 		ECollisionChannel const BlockingChannel = PrimComp->GetCollisionObjectType();
 		FCollisionShape const CollisionShape = PrimComp->GetCollisionShape(-Epsilon);
 
@@ -829,6 +833,7 @@ static bool ComponentEncroachesBlockingGeometry_WithAdjustment(UWorld const* Wor
 				FComponentQueryParams Params(NAME_ComponentEncroachesBlockingGeometry_WithAdjustment, TestActor);
 				FCollisionResponseParams ResponseParams;
 				PrimComp->InitSweepCollisionParams(Params, ResponseParams);
+				Params.AddIgnoredActors(IgnoreActors);
 				bFoundBlockingHit = World->ComponentOverlapMultiByChannel(Overlaps, PrimComp, TestWorldTransform.GetLocation(), TestWorldTransform.GetRotation(), BlockingChannel, Params);
 				bComputePenetrationAdjustment = false;
 			}
@@ -843,6 +848,7 @@ static bool ComponentEncroachesBlockingGeometry_WithAdjustment(UWorld const* Wor
 			FCollisionQueryParams Params(NAME_ComponentEncroachesBlockingGeometry_WithAdjustment, false, TestActor);
 			FCollisionResponseParams ResponseParams;
 			PrimComp->InitSweepCollisionParams(Params, ResponseParams);
+			Params.AddIgnoredActors(IgnoreActors);
 			bFoundBlockingHit = World->OverlapMultiByChannel(Overlaps, TestWorldTransform.GetLocation(), TestWorldTransform.GetRotation(), BlockingChannel, CollisionShape, Params, ResponseParams);
 		}
 
@@ -851,12 +857,14 @@ static bool ComponentEncroachesBlockingGeometry_WithAdjustment(UWorld const* Wor
 		{
 			// if encroaching, add up all the MTDs of overlapping shapes
 			FMTDResult MTDResult;
+			uint32 NumBlockingHits = 0;
 			for (int32 HitIdx = 0; HitIdx < Overlaps.Num(); HitIdx++)
 			{
 				UPrimitiveComponent* const OverlapComponent = Overlaps[HitIdx].Component.Get();
 				// first determine closest impact point along each axis
 				if (OverlapComponent && OverlapComponent->GetCollisionResponseToChannel(BlockingChannel) == ECR_Block)
 				{
+					NumBlockingHits++;
 					FCollisionShape const NonShrunkenCollisionShape = PrimComp->GetCollisionShape();
 					bool bSuccess = OverlapComponent->ComputePenetration(MTDResult, NonShrunkenCollisionShape, TestWorldTransform.GetLocation(), TestWorldTransform.GetRotation());
 					if (bSuccess)
@@ -884,14 +892,21 @@ static bool ComponentEncroachesBlockingGeometry_WithAdjustment(UWorld const* Wor
 						}
 						else
 						{
-							UE_LOG(LogPhysics, Log, TEXT("OverlapTest says we are overlapping, yet MTD says we're not (with smaller shape). Something is wrong"));
-							// It's not safe to use a partial result, that could push us out to an invalid location.
-							OutProposedAdjustment = FVector::ZeroVector;
-							return true;
+							// Ignore this overlap.
+							UE_LOG(LogPhysics, Log, TEXT("OverlapTest says we are overlapping, yet MTD says we're not (with smaller shape). Ignoring this overlap."));
+							NumBlockingHits--;
+							continue;
 						}
 					}
 				}
 			}
+
+			// See if we chose to invalidate all of our supposed "blocking hits".
+			if (NumBlockingHits == 0)
+			{
+				OutProposedAdjustment = FVector::ZeroVector;
+				bFoundBlockingHit = false;
+			}
 		}
 
 		return bFoundBlockingHit;
@@ -901,11 +916,11 @@ static bool ComponentEncroachesBlockingGeometry_WithAdjustment(UWorld const* Wor
 }
 
 /** Tests if the given component overlaps any blocking geometry if it were placed at the given world transform, optionally returns a suggested translation to get the component away from its overlaps. */
-static bool ComponentEncroachesBlockingGeometry(UWorld const* World, AActor const* TestActor, UPrimitiveComponent const* PrimComp, FTransform const& TestWorldTransform, FVector* OutProposedAdjustment)
+static bool ComponentEncroachesBlockingGeometry(UWorld const* World, AActor const* TestActor, UPrimitiveComponent const* PrimComp, FTransform const& TestWorldTransform, FVector* OutProposedAdjustment, const TArray<AActor*>& IgnoreActors)
 {
 	return OutProposedAdjustment
-		? ComponentEncroachesBlockingGeometry_WithAdjustment(World, TestActor, PrimComp, TestWorldTransform, *OutProposedAdjustment)
-		: ComponentEncroachesBlockingGeometry_NoAdjustment(World, TestActor, PrimComp, TestWorldTransform);
+		? ComponentEncroachesBlockingGeometry_WithAdjustment(World, TestActor, PrimComp, TestWorldTransform, *OutProposedAdjustment, IgnoreActors)
+		: ComponentEncroachesBlockingGeometry_NoAdjustment(World, TestActor, PrimComp, TestWorldTransform, IgnoreActors);
 }
 
 
@@ -957,7 +972,10 @@ bool UWorld::EncroachingBlockingGeometry(AActor* TestActor, FVector TestLocation
 			FTransform const CompToRoot = MovedPrimComp->GetComponentToWorld() * WorldToOldRoot;
 			FTransform const CompToNewWorld = CompToRoot * TestRootToWorld;
 
-			if (ComponentEncroachesBlockingGeometry(this, TestActor, MovedPrimComp, CompToNewWorld, ProposedAdjustment))
+			TArray<AActor*> ChildActors;
+			TestActor->GetAllChildActors(ChildActors);
+
+			if (ComponentEncroachesBlockingGeometry(this, TestActor, MovedPrimComp, CompToNewWorld, ProposedAdjustment, ChildActors))
 			{
 				if (ProposedAdjustment == nullptr)
 				{
@@ -975,11 +993,17 @@ bool UWorld::EncroachingBlockingGeometry(AActor* TestActor, FVector TestLocation
 	}
 	else
 	{
+		bool bFetchedChildActors = false;
+		TArray<AActor*> ChildActors;
+
 		// This actor does not have a movement component, so we'll assume all components are potentially important to keep out of the world
 		UPrimitiveComponent* const RootPrimComp = Cast<UPrimitiveComponent>(RootComponent);
 		if (RootPrimComp && RootPrimComp->IsQueryCollisionEnabled())
 		{
-			if (ComponentEncroachesBlockingGeometry(this, TestActor, RootPrimComp, TestRootToWorld, ProposedAdjustment))
+			TestActor->GetAllChildActors(ChildActors);
+			bFetchedChildActors = true;
+
+			if (ComponentEncroachesBlockingGeometry(this, TestActor, RootPrimComp, TestRootToWorld, ProposedAdjustment, ChildActors))
 			{
 				if (ProposedAdjustment == nullptr)
 				{
@@ -999,7 +1023,7 @@ bool UWorld::EncroachingBlockingGeometry(AActor* TestActor, FVector TestLocation
 		TArray<USceneComponent*> Children;
 		RootComponent->GetChildrenComponents(true, Children);
 
-		for (auto Child : Children)
+		for (USceneComponent* Child : Children)
 		{
 			if (Child->IsQueryCollisionEnabled())
 			{
@@ -1009,7 +1033,13 @@ bool UWorld::EncroachingBlockingGeometry(AActor* TestActor, FVector TestLocation
 					FTransform const CompToRoot = Child->GetComponentToWorld() * WorldToOldRoot;
 					FTransform const CompToNewWorld = CompToRoot * TestRootToWorld;
 
-					if (ComponentEncroachesBlockingGeometry(this, TestActor, PrimComp, CompToNewWorld, ProposedAdjustment))
+					if (!bFetchedChildActors)
+					{
+						TestActor->GetAllChildActors(ChildActors);
+						bFetchedChildActors = true;
+					}
+
+					if (ComponentEncroachesBlockingGeometry(this, TestActor, PrimComp, CompToNewWorld, ProposedAdjustment, ChildActors))
 					{
 						if (ProposedAdjustment == nullptr)
 						{
@@ -1269,7 +1299,7 @@ FAudioDevice* UWorld::GetAudioDevice()
  */
 void UWorld::SetMapNeedsLightingFullyRebuilt(int32 InNumLightingUnbuiltObjects)
 {
-	static const auto AllowStaticLightingVar = IConsoleManager::Get().FindTConsoleVariableDataInt(TEXT("r.AllowStaticLighting"));
+	static const TConsoleVariableData<int32>* AllowStaticLightingVar = IConsoleManager::Get().FindTConsoleVariableDataInt(TEXT("r.AllowStaticLighting"));
 	const bool bAllowStaticLighting = (!AllowStaticLightingVar || AllowStaticLightingVar->GetValueOnGameThread() != 0);
 
 	AWorldSettings* WorldSettings = GetWorldSettings();
diff --git a/Engine/Source/Runtime/Engine/Private/Materials/Material.cpp b/Engine/Source/Runtime/Engine/Private/Materials/Material.cpp
--- a/Engine/Source/Runtime/Engine/Private/Materials/Material.cpp
+++ b/Engine/Source/Runtime/Engine/Private/Materials/Material.cpp
@@ -33,6 +33,7 @@
 #include "ShaderCompiler.h"
 #include "Materials/MaterialParameterCollection.h"
 #include "MaterialShaderQualitySettings.h"
+#include "LoadTimeTracker.h"
 #if WITH_EDITOR
 #include "MessageLog.h"
 #include "UObjectToken.h"
@@ -415,6 +416,8 @@ static TAutoConsoleVariable<int32> CVarDiscardUnusedQualityLevels(
 
 void SerializeInlineShaderMaps(const TMap<const ITargetPlatform*, TArray<FMaterialResource*>>* PlatformMaterialResourcesToSavePtr, FArchive& Ar, FMaterialResource* (&OutMaterialResourcesLoaded)[EMaterialQualityLevel::Num][ERHIFeatureLevel::Num])
 {
+	SCOPED_LOADTIMER(SerializeInlineShaderMaps);
+
 	if (Ar.IsSaving())
 	{
 		int32 NumResourcesToSave = 0;
@@ -2128,6 +2131,8 @@ const FMaterialResource* UMaterial::GetMaterialResource(ERHIFeatureLevel::Type I
 
 void UMaterial::Serialize(FArchive& Ar)
 {
+	SCOPED_LOADTIMER(MaterialSerializeTime);
+
 	Super::Serialize(Ar);
 
 	if (FPlatformProperties::RequiresCookedData() && Ar.IsLoading())
@@ -2358,6 +2363,8 @@ TMap<FGuid, UMaterialInterface*> LightingGuidFixupMap;
 
 void UMaterial::PostLoad()
 {
+	SCOPED_LOADTIMER(MaterialPostLoad);
+
 	Super::PostLoad();
 
 	if (!IsDefaultMaterial())
@@ -3205,7 +3212,7 @@ void UMaterial::FinishDestroy()
 
 SIZE_T UMaterial::GetResourceSize(EResourceSizeMode::Type Mode)
 {
-	int32 ResourceSize = 0;
+	int32 ResourceSize = Super::GetResourceSize(Mode);
 
 	for (int32 InstanceIndex = 0; InstanceIndex < 3; ++InstanceIndex)
 	{
diff --git a/Engine/Source/Runtime/Engine/Private/Materials/MaterialInstance.cpp b/Engine/Source/Runtime/Engine/Private/Materials/MaterialInstance.cpp
--- a/Engine/Source/Runtime/Engine/Private/Materials/MaterialInstance.cpp
+++ b/Engine/Source/Runtime/Engine/Private/Materials/MaterialInstance.cpp
@@ -15,6 +15,7 @@
 #include "TargetPlatform.h"
 #include "Engine/Font.h"
 #include "Engine/SubsurfaceProfile.h"
+#include "LoadTimeTracker.h"
 
 /**
  * Cache uniform expressions for the given material.
@@ -1830,6 +1831,7 @@ void UMaterialInstance::ClearAllCachedCookedPlatformData()
 
 void UMaterialInstance::Serialize(FArchive& Ar)
 {
+	SCOPED_LOADTIMER(MaterialInstanceSerializeTime);
 	Super::Serialize(Ar);
 
 	// Only serialize the static permutation resource if one exists
@@ -1911,6 +1913,7 @@ void UMaterialInstance::Serialize(FArchive& Ar)
 
 void UMaterialInstance::PostLoad()
 {
+	SCOPED_LOADTIMER(MaterialInstancePostLoad);
 	Super::PostLoad();
 
 	AssertDefaultMaterialsPostLoaded();
@@ -2485,7 +2488,7 @@ bool UMaterialInstance::GetTexturesInPropertyChain(EMaterialProperty InProperty,
 
 SIZE_T UMaterialInstance::GetResourceSize(EResourceSizeMode::Type Mode)
 {
-	SIZE_T ResourceSize = 0;
+	SIZE_T ResourceSize = Super::GetResourceSize(Mode);
 
 	if (bHasStaticPermutationResource)
 	{
diff --git a/Engine/Source/Runtime/Engine/Private/MatineeUtils.cpp b/Engine/Source/Runtime/Engine/Private/MatineeUtils.cpp
--- a/Engine/Source/Runtime/Engine/Private/MatineeUtils.cpp
+++ b/Engine/Source/Runtime/Engine/Private/MatineeUtils.cpp
@@ -180,7 +180,7 @@ private:
 			if (It.Key() == TrackRedirectsName)
 			{
 				FName TargetClassName = NAME_Object;
-				FParse::Value(*It.Value(), TEXT("TargetClassName="), TargetClassName);
+				FParse::Value(*It.Value().GetValue(), TEXT("TargetClassName="), TargetClassName);
 
 				if (UClass* TargetClass = LoadClass<UObject>(NULL, *TargetClassName.ToString(), NULL, LOAD_None, NULL))
 				{
@@ -188,8 +188,8 @@ private:
 
 					FString OldFieldName;
 					FString NewFieldName;
-					FParse::Value(*It.Value(), TEXT("OldFieldName="), OldFieldName);
-					FParse::Value(*It.Value(), TEXT("NewFieldName="), NewFieldName);
+					FParse::Value(*It.Value().GetValue(), TEXT("OldFieldName="), OldFieldName);
+					FParse::Value(*It.Value().GetValue(), TEXT("NewFieldName="), NewFieldName);
 
 					ClassRemapInfo.PrefixMap.Add(OldFieldName, NewFieldName);
 				}
diff --git a/Engine/Source/Runtime/Engine/Private/MemberReference.cpp b/Engine/Source/Runtime/Engine/Private/MemberReference.cpp
--- a/Engine/Source/Runtime/Engine/Private/MemberReference.cpp
+++ b/Engine/Source/Runtime/Engine/Private/MemberReference.cpp
@@ -214,8 +214,8 @@ void FMemberReference::InitFieldRedirectMap()
 					FString OldFieldPathString;
 					FString NewFieldPathString;
 
-					FParse::Value( *It.Value(), TEXT("OldFieldName="), OldFieldPathString );
-					FParse::Value( *It.Value(), TEXT("NewFieldName="), NewFieldPathString );
+					FParse::Value( *It.Value().GetValue(), TEXT("OldFieldName="), OldFieldPathString );
+					FParse::Value( *It.Value().GetValue(), TEXT("NewFieldName="), NewFieldPathString );
 
 					// Handle both cases of just a field being renamed (just one FName), as well as a class and field name (ClassName.FieldName)
 					FFieldRemapInfo OldFieldRemap;
@@ -268,13 +268,13 @@ void FMemberReference::InitFieldRedirectMap()
 					FString NewParamValues;
 					FString CustomValueMapping;
 
-					FParse::Value( *It.Value(), TEXT("NodeName="), NodeName );
-					FParse::Value( *It.Value(), TEXT("NodeTitle="), NodeTitle );
-					FParse::Value( *It.Value(), TEXT("OldParamName="), OldParam );
-					FParse::Value( *It.Value(), TEXT("NewParamName="), NewParam );
-					FParse::Value( *It.Value(), TEXT("OldParamValues="), OldParamValues );
-					FParse::Value( *It.Value(), TEXT("NewParamValues="), NewParamValues );
-					FParse::Value( *It.Value(), TEXT("CustomValueMapping="), CustomValueMapping );
+					FParse::Value( *It.Value().GetValue(), TEXT("NodeName="), NodeName );
+					FParse::Value( *It.Value().GetValue(), TEXT("NodeTitle="), NodeTitle );
+					FParse::Value( *It.Value().GetValue(), TEXT("OldParamName="), OldParam );
+					FParse::Value( *It.Value().GetValue(), TEXT("NewParamName="), NewParam );
+					FParse::Value( *It.Value().GetValue(), TEXT("OldParamValues="), OldParamValues );
+					FParse::Value( *It.Value().GetValue(), TEXT("NewParamValues="), NewParamValues );
+					FParse::Value( *It.Value().GetValue(), TEXT("CustomValueMapping="), CustomValueMapping );
 
 					FParamRemapInfo ParamRemap;
 					ParamRemap.NewParam = NewParam;
diff --git a/Engine/Source/Runtime/Engine/Private/PainCausingVolume.cpp b/Engine/Source/Runtime/Engine/Private/PainCausingVolume.cpp
--- a/Engine/Source/Runtime/Engine/Private/PainCausingVolume.cpp
+++ b/Engine/Source/Runtime/Engine/Private/PainCausingVolume.cpp
@@ -54,16 +54,16 @@ void APainCausingVolume::PainTimer()
 {
 	if (bPainCausing)
 	{
-		TArray<AActor*> TouchingActors;
+		TSet<AActor*> TouchingActors;
 		GetOverlappingActors(TouchingActors, APawn::StaticClass());
 
-		for (int32 iTouching = 0; iTouching < TouchingActors.Num(); ++iTouching)
+		for (AActor* const A : TouchingActors)
 		{
-			AActor* const A = TouchingActors[iTouching];
-			if (A && !A->IsPendingKill())
+			if (A && A->bCanBeDamaged && !A->IsPendingKill())
 			{
 				// @todo physicsVolume This won't work for other actor. Need to fix it properly
-				if ( A->bCanBeDamaged && (Cast<APawn>(A) && Cast<APawn>(A)->GetPawnPhysicsVolume() == this) )
+				APawn* PawnA = Cast<APawn>(A);
+				if (PawnA && PawnA->GetPawnPhysicsVolume() == this)
 				{
 					CausePainTo(A);
 				}
diff --git a/Engine/Source/Runtime/Engine/Private/Particles/ParticleComponents.cpp b/Engine/Source/Runtime/Engine/Private/Particles/ParticleComponents.cpp
--- a/Engine/Source/Runtime/Engine/Private/Particles/ParticleComponents.cpp
+++ b/Engine/Source/Runtime/Engine/Private/Particles/ParticleComponents.cpp
@@ -2386,12 +2386,6 @@ void UParticleSystem::PostLoad()
 	SetupSoloing();
 }
 
-
-SIZE_T UParticleSystem::GetResourceSize(EResourceSizeMode::Type Mode)
-{
-	return 0;
-}
-
 void UParticleSystem::UpdateColorModuleClampAlpha(UParticleModuleColorBase* ColorModule)
 {
 	if (ColorModule)
@@ -3409,7 +3403,7 @@ void UParticleSystemComponent::FinishDestroy()
 SIZE_T UParticleSystemComponent::GetResourceSize(EResourceSizeMode::Type Mode)
 {
 	ForceAsyncWorkCompletion(ENSURE_AND_STALL);
-	int32 ResSize = 0;
+	int32 ResSize = Super::GetResourceSize(Mode);
 	for (int32 EmitterIdx = 0; EmitterIdx < EmitterInstances.Num(); EmitterIdx++)
 	{
 		FParticleEmitterInstance* EmitterInstance = EmitterInstances[EmitterIdx];
diff --git a/Engine/Source/Runtime/Engine/Private/PhysicsEngine/BodyInstance.cpp b/Engine/Source/Runtime/Engine/Private/PhysicsEngine/BodyInstance.cpp
--- a/Engine/Source/Runtime/Engine/Private/PhysicsEngine/BodyInstance.cpp
+++ b/Engine/Source/Runtime/Engine/Private/PhysicsEngine/BodyInstance.cpp
@@ -4454,7 +4454,8 @@ bool FBodyInstance::OverlapPhysX_AssumesLocked(const PxGeometry& PGeom, const Px
 		{
 			PxVec3 POutDirection;
 			float OutDistance;
-			if (PxGeometryQuery::computePenetration(POutDirection, OutDistance, PGeom, ShapePose, PShape->getGeometry().any(), PxShapeExt::getGlobalPose(*PShape, *RigidBody)))
+			//TODO: this 0 check is a temp fix because physx is reporting contacts with distance 0 which gives us nan
+			if (PxGeometryQuery::computePenetration(POutDirection, OutDistance, PGeom, ShapePose, PShape->getGeometry().any(), PxShapeExt::getGlobalPose(*PShape, *RigidBody)) && !FMath::IsNearlyZero(OutDistance))
 			{
 				if(OutMTD)
 				{
@@ -4527,6 +4528,7 @@ void FBodyInstance::LoadProfileData(bool bVerifyProfile)
 				const FBodyInstance& ExternalBodyInstance = ExternalCollisionProfileBodySetup->DefaultInstance;
 				CollisionProfileName = ExternalBodyInstance.CollisionProfileName;
 				ObjectType = ExternalBodyInstance.ObjectType;
+				CollisionEnabled = ExternalBodyInstance.CollisionEnabled;
 				CollisionResponses.SetCollisionResponseContainer(ExternalBodyInstance.CollisionResponses.ResponseToChannels);
 			}
 			else
diff --git a/Engine/Source/Runtime/Engine/Private/PhysicsEngine/BodySetup.cpp b/Engine/Source/Runtime/Engine/Private/PhysicsEngine/BodySetup.cpp
--- a/Engine/Source/Runtime/Engine/Private/PhysicsEngine/BodySetup.cpp
+++ b/Engine/Source/Runtime/Engine/Private/PhysicsEngine/BodySetup.cpp
@@ -27,14 +27,14 @@
 // CVars
 static TAutoConsoleVariable<float> CVarContactOffsetFactor(
 	TEXT("p.ContactOffsetFactor"),
-	0.01f,
-	TEXT("Multiplied by min dimension of object to calculate how close objects get before generating contacts. Default: 0.01"),
+	-1.f,
+	TEXT("Multiplied by min dimension of object to calculate how close objects get before generating contacts. < 0 implies use project settings. Default: 0.01"),
 	ECVF_Default);
 
 static TAutoConsoleVariable<float> CVarMaxContactOffset(
 	TEXT("p.MaxContactOffset"),
-	1.f,
-	TEXT("Max value of contact offset, which controls how close objects get before generating contacts. Default: 1.0"),
+	-1.f,
+	TEXT("Max value of contact offset, which controls how close objects get before generating contacts. < 0 implies use project settings. Default: 1.0"),
 	ECVF_Default);
 
 DEFINE_LOG_CATEGORY(LogPhysics);
@@ -299,11 +299,16 @@ void SetupNonUniformHelper(FVector Scale3D, float& MinScale, float& MinScaleAbs,
 	}
 }
 
-void GetContactOffsetParams(float& ContactOffsetFactor, float& MaxContactOffset)
+void GetContactOffsetParams(float& ContactOffsetFactor, float& MinContactOffset, float& MaxContactOffset)
 {
 	// Get contact offset params
 	ContactOffsetFactor = CVarContactOffsetFactor.GetValueOnGameThread();
 	MaxContactOffset = CVarMaxContactOffset.GetValueOnGameThread();
+
+	ContactOffsetFactor = ContactOffsetFactor < 0.f ? UPhysicsSettings::Get()->ContactOffsetMultiplier : ContactOffsetFactor;
+	MaxContactOffset = MaxContactOffset < 0.f ? UPhysicsSettings::Get()->MaxContactOffset : MaxContactOffset;
+
+	MinContactOffset = UPhysicsSettings::Get()->MinContactOffset;
 }
 
 PxMaterial* GetDefaultPhysMaterial()
@@ -348,7 +353,7 @@ struct FAddShapesHelper
 			ShapeScale3D.Z *= Scale3DAbsRelative.Z;
 		}
 
-		GetContactOffsetParams(ContactOffsetFactor, MaxContactOffset);
+		GetContactOffsetParams(ContactOffsetFactor, MinContactOffset, MaxContactOffset);
 	}
 
 	UBodySetup* BodySetup;
@@ -369,6 +374,7 @@ struct FAddShapesHelper
 	FVector ShapeScale3D;
 
 	float ContactOffsetFactor;
+	float MinContactOffset;
 	float MaxContactOffset;
 
 public:
@@ -387,7 +393,7 @@ public:
 				PxTransform PLocalPose(U2PVector(ScaledSphereElem.Center));
 				ensure(PLocalPose.isValid());
 				{
-					const float ContactOffset = FMath::Min(MaxContactOffset, ContactOffsetFactor * PSphereGeom.radius);
+					const float ContactOffset = FMath::Clamp(ContactOffsetFactor * PSphereGeom.radius, MinContactOffset, MaxContactOffset);
 					PxShape* PShape = AttachShape_AssumesLocked(PSphereGeom, PLocalPose, ContactOffset, SphereElem.GetUserData());
 				}
 			}
@@ -416,7 +422,7 @@ public:
 				PxTransform PLocalPose(U2PTransform(BoxTransform));
 				ensure(PLocalPose.isValid());
 				{
-					const float ContactOffset = FMath::Min(MaxContactOffset, ContactOffsetFactor * PBoxGeom.halfExtents.minElement());
+					const float ContactOffset = FMath::Clamp(ContactOffsetFactor * PBoxGeom.halfExtents.minElement(), MinContactOffset, MaxContactOffset);
 					AttachShape_AssumesLocked(PBoxGeom, PLocalPose, ContactOffset, BoxElem.GetUserData());
 				}
 			}
@@ -448,7 +454,7 @@ public:
 
 				ensure(PLocalPose.isValid());
 				{
-					const float ContactOffset = FMath::Min(MaxContactOffset, ContactOffsetFactor * PCapsuleGeom.radius);
+					const float ContactOffset = FMath::Clamp(ContactOffsetFactor * PCapsuleGeom.radius, MinContactOffset, MaxContactOffset);
 					AttachShape_AssumesLocked(PCapsuleGeom, PLocalPose, ContactOffset, SphylElem.GetUserData());
 				}
 			}
@@ -494,7 +500,7 @@ public:
 
 						ensure(PLocalPose.isValid());
 						{
-							const float ContactOffset = FMath::Min(MaxContactOffset, ContactOffsetFactor * PBoundsExtents.minElement());
+							const float ContactOffset = FMath::Clamp(ContactOffsetFactor * PBoundsExtents.minElement(), MinContactOffset, MaxContactOffset);
 							AttachShape_AssumesLocked(PConvexGeom, PLocalPose, ContactOffset, ConvexElem.GetUserData());
 						}
 					}
@@ -1098,7 +1104,7 @@ void UBodySetup::PostEditUndo()
 
 SIZE_T UBodySetup::GetResourceSize( EResourceSizeMode::Type Mode )
 {
-	SIZE_T ResourceSize = 0;
+	SIZE_T ResourceSize = Super::GetResourceSize(Mode);
 
 #if WITH_PHYSX
 	// Count PhysX trimesh mem usage
@@ -1312,13 +1318,12 @@ FKSphereElem FKSphereElem::GetFinalScaled(const FVector& Scale3D, const FTransfo
 	float MinScale, MinScaleAbs;
 	FVector Scale3DAbs;
 
-	SetupNonUniformHelper(Scale3D, MinScale, MinScaleAbs, Scale3DAbs);
+	SetupNonUniformHelper(Scale3D * RelativeTM.GetScale3D(), MinScale, MinScaleAbs, Scale3DAbs);
 
 	FKSphereElem ScaledSphere = *this;
 	ScaledSphere.Radius *= MinScaleAbs;
 
-	ScaledSphere.Center = RelativeTM.TransformPosition(Center);
-	ScaledSphere.Center *= MinScale;
+	ScaledSphere.Center = RelativeTM.TransformPosition(Center) * Scale3D;
 
 
 	return ScaledSphere;
@@ -1347,7 +1352,7 @@ FKBoxElem FKBoxElem::GetFinalScaled(const FVector& Scale3D, const FTransform& Re
 	float MinScale, MinScaleAbs;
 	FVector Scale3DAbs;
 
-	SetupNonUniformHelper(Scale3D, MinScale, MinScaleAbs, Scale3DAbs);
+	SetupNonUniformHelper(Scale3D * RelativeTM.GetScale3D(), MinScale, MinScaleAbs, Scale3DAbs);
 
 	FKBoxElem ScaledBox = *this;
 	ScaledBox.X *= Scale3DAbs.X;
@@ -1437,7 +1442,7 @@ FKSphylElem FKSphylElem::GetFinalScaled(const FVector& Scale3D, const FTransform
 	float MinScale, MinScaleAbs;
 	FVector Scale3DAbs;
 
-	SetupNonUniformHelper(Scale3D, MinScale, MinScaleAbs, Scale3DAbs);
+	SetupNonUniformHelper(Scale3D * RelativeTM.GetScale3D(), MinScale, MinScaleAbs, Scale3DAbs);
 
 	float ScaleRadius = FMath::Max(Scale3DAbs.X, Scale3DAbs.Y);
 	float ScaleLength = Scale3DAbs.Z;
diff --git a/Engine/Source/Runtime/Engine/Private/PhysicsEngine/ConstraintInstance.cpp b/Engine/Source/Runtime/Engine/Private/PhysicsEngine/ConstraintInstance.cpp
--- a/Engine/Source/Runtime/Engine/Private/PhysicsEngine/ConstraintInstance.cpp
+++ b/Engine/Source/Runtime/Engine/Private/PhysicsEngine/ConstraintInstance.cpp
@@ -489,12 +489,19 @@ void EnsureSleepingActorsStaySleeping_AssumesLocked(PxRigidActor* PActor1, PxRig
 	{
 		if (PActor1 && !IsRigidBodyKinematic_AssumesLocked(PActor1->isRigidDynamic()))
 		{
-			PActor1->isRigidDynamic()->putToSleep();
+			if(PActor1->isRigidDynamic())
+			{
+				PActor1->isRigidDynamic()->putToSleep();
+			}
 		}
 
 		if (PActor2 && !IsRigidBodyKinematic_AssumesLocked(PActor2->isRigidDynamic()))
 		{
-			PActor2->isRigidDynamic()->putToSleep();
+			if(PActor2->isRigidDynamic())
+			{
+				PActor2->isRigidDynamic()->putToSleep();
+			}
+			
 		}
 	}
 }
diff --git a/Engine/Source/Runtime/Engine/Private/PhysicsEngine/PhysScene.cpp b/Engine/Source/Runtime/Engine/Private/PhysicsEngine/PhysScene.cpp
--- a/Engine/Source/Runtime/Engine/Private/PhysicsEngine/PhysScene.cpp
+++ b/Engine/Source/Runtime/Engine/Private/PhysicsEngine/PhysScene.cpp
@@ -624,44 +624,38 @@ void FPhysScene::FlushDeferredCollisionDisableTableQueue()
 	DeferredCollisionDisableTableQueue.Empty();
 }
 
-void GatherPhysXStats(PxScene* PScene, uint32 SceneType)
+void GatherPhysXStats_AssumesLocked(PxScene* PSyncScene, PxScene* PAsyncScene)
 {
 	/** Gather PhysX stats */
-	if (SceneType == 0)
+	if (PSyncScene)
 	{
-		if (PScene)
-		{
-			SCOPED_SCENE_READ_LOCK(PScene);
-			PxSimulationStatistics SimStats;
-			PScene->getSimulationStatistics(SimStats);
+		PxSimulationStatistics SimStats;
+		PSyncScene->getSimulationStatistics(SimStats);
 
-			SET_DWORD_STAT(STAT_NumActiveConstraints, SimStats.nbActiveConstraints);
-			SET_DWORD_STAT(STAT_NumActiveSimulatedBodies, SimStats.nbActiveDynamicBodies);
-			SET_DWORD_STAT(STAT_NumActiveKinematicBodies, SimStats.nbActiveKinematicBodies);
-			SET_DWORD_STAT(STAT_NumStaticBodies, SimStats.nbStaticBodies);
-			SET_DWORD_STAT(STAT_NumMobileBodies, SimStats.nbDynamicBodies);
+		SET_DWORD_STAT(STAT_NumActiveConstraints, SimStats.nbActiveConstraints);
+		SET_DWORD_STAT(STAT_NumActiveSimulatedBodies, SimStats.nbActiveDynamicBodies);
+		SET_DWORD_STAT(STAT_NumActiveKinematicBodies, SimStats.nbActiveKinematicBodies);
+		SET_DWORD_STAT(STAT_NumStaticBodies, SimStats.nbStaticBodies);
+		SET_DWORD_STAT(STAT_NumMobileBodies, SimStats.nbDynamicBodies);
 			
-			SET_DWORD_STAT(STAT_NumBroadphaseAdds, SimStats.getNbBroadPhaseAdds(PxSimulationStatistics::VolumeType::eRIGID_BODY));
-			SET_DWORD_STAT(STAT_NumBroadphaseRemoves, SimStats.getNbBroadPhaseRemoves(PxSimulationStatistics::VolumeType::eRIGID_BODY));
-
-			uint32 NumShapes = 0;
-			for (int32 GeomType = 0; GeomType < PxGeometryType::eGEOMETRY_COUNT; ++GeomType)
-			{
-				NumShapes += SimStats.nbShapes[GeomType];
-			}
-
-			SET_DWORD_STAT(STAT_NumShapes, NumShapes);
+		//SET_DWORD_STAT(STAT_NumBroadphaseAdds, SimStats.getNbBroadPhaseAdds(PxSimulationStatistics::VolumeType::eRIGID_BODY));	//TODO: These do not seem to work
+		//SET_DWORD_STAT(STAT_NumBroadphaseRemoves, SimStats.getNbBroadPhaseRemoves(PxSimulationStatistics::VolumeType::eRIGID_BODY));
 
+		uint32 NumShapes = 0;
+		for (int32 GeomType = 0; GeomType < PxGeometryType::eGEOMETRY_COUNT; ++GeomType)
+		{
+			NumShapes += SimStats.nbShapes[GeomType];
 		}
 
+		SET_DWORD_STAT(STAT_NumShapes, NumShapes);
+
 	}
 
-#if 0	//this is not reliable right now
-	else if (SceneType == 1 & UPhysicsSettings::Get()->bEnableAsyncScene)
+	if(PAsyncScene)
 	{
 		//Having to duplicate because of macros. In theory we can fix this but need to get this quickly
 		PxSimulationStatistics SimStats;
-		PScene->getSimulationStatistics(SimStats);
+		PAsyncScene->getSimulationStatistics(SimStats);
 
 		SET_DWORD_STAT(STAT_NumActiveConstraintsAsync, SimStats.nbActiveConstraints);
 		SET_DWORD_STAT(STAT_NumActiveSimulatedBodiesAsync, SimStats.nbActiveDynamicBodies);
@@ -669,8 +663,8 @@ void GatherPhysXStats(PxScene* PScene, uint32 SceneType)
 		SET_DWORD_STAT(STAT_NumStaticBodiesAsync, SimStats.nbStaticBodies);
 		SET_DWORD_STAT(STAT_NumMobileBodiesAsync, SimStats.nbDynamicBodies);
 
-		SET_DWORD_STAT(STAT_NumBroadphaseAddsAsync, SimStats.getNbBroadPhaseAdds(PxSimulationStatistics::VolumeType::eRIGID_BODY));
-		SET_DWORD_STAT(STAT_NumBroadphaseRemovesAsync, SimStats.getNbBroadPhaseRemoves(PxSimulationStatistics::VolumeType::eRIGID_BODY));
+		//SET_DWORD_STAT(STAT_NumBroadphaseAddsAsync, SimStats.getNbBroadPhaseAdds(PxSimulationStatistics::VolumeType::eRIGID_BODY)); //TODO: These do not seem to work
+		//SET_DWORD_STAT(STAT_NumBroadphaseRemovesAsync, SimStats.getNbBroadPhaseRemoves(PxSimulationStatistics::VolumeType::eRIGID_BODY));
 
 		uint32 NumShapes = 0;
 		for (int32 GeomType = 0; GeomType < PxGeometryType::eGEOMETRY_COUNT; ++GeomType)
@@ -680,7 +674,6 @@ void GatherPhysXStats(PxScene* PScene, uint32 SceneType)
 
 		SET_DWORD_STAT(STAT_NumShapesAsync, NumShapes);
 	}
-#endif
 }
 
 DECLARE_FLOAT_COUNTER_STAT(TEXT("Sync Sim Time (ms)"), STAT_PhysSyncSim, STATGROUP_Physics);
@@ -856,11 +849,6 @@ void FPhysScene::TickPhysScene(uint32 SceneType, FGraphEventRef& InOutCompletion
 		return;
 	}
 
-#if ( WITH_PHYSX  && !(UE_BUILD_SHIPPING || WITH_PHYSX_RELEASE))
-	GatherPhysXStats(GetPhysXScene(SceneType), SceneType);
-#endif
-
-
 	/**
 	* Weight frame time according to PhysScene settings.
 	*/
@@ -1369,6 +1357,10 @@ void FPhysScene::EndFrame(ULineBatchComponent* InLineBatcher)
 	SCOPED_SCENE_WRITE_LOCK(GetPhysXScene(PST_Sync));
 	SCOPED_SCENE_WRITE_LOCK(bAsyncSceneEnabled ? GetPhysXScene(PST_Async) : nullptr);
 
+#if ( WITH_PHYSX  && !(UE_BUILD_SHIPPING || WITH_PHYSX_RELEASE))
+	GatherPhysXStats_AssumesLocked(GetPhysXScene(PST_Sync), HasAsyncScene() ? GetPhysXScene(PST_Async) : nullptr);
+#endif
+
 	if (bAsyncSceneEnabled)
 	{
 		SyncComponentsToBodies_AssumesLocked(PST_Async);
@@ -1394,6 +1386,96 @@ void FPhysScene::EndFrame(ULineBatchComponent* InLineBatcher)
 #endif // !(UE_BUILD_SHIPPING || UE_BUILD_TEST || WITH_PHYSX_RELEASE)
 }
 
+#if WITH_PHYSX
+/** Helper struct that puts all awake actors to sleep and then later wakes them back up */
+struct FHelpEnsureCollisionTreeIsBuilt
+{
+	FHelpEnsureCollisionTreeIsBuilt(PxScene* InPScene)
+	: PScene(InPScene)
+	{
+		if(PScene)
+		{
+			SCOPED_SCENE_WRITE_LOCK(PScene);
+			const int32 NumActors = PScene->getNbActors(PxActorTypeFlag::eRIGID_DYNAMIC);
+
+			if (NumActors)
+			{
+				ActorBuffer.AddUninitialized(NumActors);
+				PScene->getActors(PxActorTypeFlag::eRIGID_DYNAMIC, &ActorBuffer[0], NumActors);
+
+				for (PxActor*& PActor : ActorBuffer)
+				{
+					if (PActor)
+					{
+						if (PxRigidDynamic* PDynamic = PActor->isRigidDynamic())
+						{
+							if (PDynamic->isSleeping() == false)
+							{
+								PDynamic->putToSleep();
+							}
+							else
+							{
+								PActor = nullptr;
+							}
+						}
+					}
+				}
+			}
+		}
+	}
+
+	~FHelpEnsureCollisionTreeIsBuilt()
+	{
+		SCOPED_SCENE_WRITE_LOCK(PScene);
+		for (PxActor* PActor : ActorBuffer)
+		{
+			if (PActor)
+			{
+				if (PxRigidDynamic* PDynamic = PActor->isRigidDynamic())
+				{
+					PDynamic->wakeUp();
+				}
+			}
+		}
+	}
+
+private:
+
+	TArray<PxActor*> ActorBuffer;
+	PxScene* PScene;
+};
+
+#endif
+
+
+DECLARE_CYCLE_STAT(TEXT("EnsureCollisionTreeIsBuilt"), STAT_PhysicsEnsureCollisionTreeIsBuilt, STATGROUP_Physics);
+
+void FPhysScene::EnsureCollisionTreeIsBuilt(UWorld* World)
+{
+	check(IsInGameThread());
+	
+
+	SCOPE_CYCLE_COUNTER(STAT_PhysicsEnsureCollisionTreeIsBuilt);
+	//We have to call fetchResults several times to update the internal data structures. PhysX doesn't have an API for this so we have to make all actors sleep before doing this
+
+	SetIsStaticLoading(true);
+
+#if WITH_PHYSX
+	FHelpEnsureCollisionTreeIsBuilt SyncSceneHelper(GetPhysXScene(PST_Sync));
+	FHelpEnsureCollisionTreeIsBuilt AsyncSceneHelper(HasAsyncScene() ? GetPhysXScene(PST_Async) : nullptr);
+#endif
+
+	for (int Iteration = 0; Iteration < 6; ++Iteration)
+	{
+		World->SetupPhysicsTickFunctions(0.1f);
+		StartFrame();
+		WaitPhysScenes();
+		EndFrame(nullptr);
+	}
+
+	SetIsStaticLoading(false);
+}
+
 void FPhysScene::SetIsStaticLoading(bool bStaticLoading)
 {
 #if WITH_PHYSX
diff --git a/Engine/Source/Runtime/Engine/Private/PhysicsEngine/PhysXSupport.cpp b/Engine/Source/Runtime/Engine/Private/PhysicsEngine/PhysXSupport.cpp
--- a/Engine/Source/Runtime/Engine/Private/PhysicsEngine/PhysXSupport.cpp
+++ b/Engine/Source/Runtime/Engine/Private/PhysicsEngine/PhysXSupport.cpp
@@ -471,13 +471,21 @@ void FPhysXSimEventCallback::onContact(const PxContactPairHeader& PairHeader, co
 		check(Shape1);
 
 		// Get materials
-		PxMaterial* Material0 = NULL;
-		Shape0->getMaterials(&Material0, 1);
-		UPhysicalMaterial* PhysMat0 = (Material0 != NULL) ? FPhysxUserData::Get<UPhysicalMaterial>(Material0->userData) : NULL;
+		PxMaterial* Material0 = nullptr;
+		UPhysicalMaterial* PhysMat0  = nullptr;
+		if(Shape0->getNbMaterials() == 1)	//If we have simple geometry or only 1 material we set it here. Otherwise do it per face
+		{
+			Shape0->getMaterials(&Material0, 1);		
+			PhysMat0 = Material0 ? FPhysxUserData::Get<UPhysicalMaterial>(Material0->userData) : nullptr;
+		}
 
-		PxMaterial* Material1 = NULL;
-		Shape1->getMaterials(&Material1, 1);
-		UPhysicalMaterial* PhysMat1 = (Material1 != NULL) ? FPhysxUserData::Get<UPhysicalMaterial>(Material1->userData) : NULL;
+		PxMaterial* Material1 = nullptr;
+		UPhysicalMaterial* PhysMat1  = nullptr;
+		if (Shape1->getNbMaterials() == 1)	//If we have simple geometry or only 1 material we set it here. Otherwise do it per face
+		{
+			Shape1->getMaterials(&Material1, 1);
+			PhysMat1 = Material1 ? FPhysxUserData::Get<UPhysicalMaterial>(Material1->userData) : nullptr;
+		}
 
 		// Iterate over contact points
 		PxContactPairPoint ContactPointBuffer[16];
@@ -490,6 +498,24 @@ void FPhysXSimEventCallback::onContact(const PxContactPairHeader& PairHeader, co
 			ImpactInfo->TotalNormalImpulse += P2UVector(NormalImpulse);
 			ImpactInfo->TotalFrictionImpulse += P2UVector(Point.impulse - NormalImpulse); // friction is component not along contact normal
 
+			// Get per face materials
+			if(!Material0)	//there is complex geometry or multiple materials so resolve the physical material here
+			{
+				if(PxMaterial* Material0PerFace = Shape0->getMaterialFromInternalFaceIndex(Point.internalFaceIndex0))
+				{
+					PhysMat0 = FPhysxUserData::Get<UPhysicalMaterial>(Material0PerFace->userData);
+				}
+			}
+
+			if (!Material1)	//there is complex geometry or multiple materials so resolve the physical material here
+			{
+				if(PxMaterial* Material1PerFace = Shape1->getMaterialFromInternalFaceIndex(Point.internalFaceIndex1))
+				{
+					PhysMat1 = FPhysxUserData::Get<UPhysicalMaterial>(Material1PerFace->userData);
+				}
+				
+			}
+			
 			new(ImpactInfo->ContactInfos) FRigidBodyContactInfo(
 				P2UVector(Point.position), 
 				P2UVector(Point.normal), 
@@ -845,4 +871,46 @@ PxCollection* MakePhysXCollection(const TArray<UPhysicalMaterial*>& PhysicalMate
 	return PCollection;
 }
 
+void FPhysXErrorCallback::reportError(PxErrorCode::Enum e, const char* message, const char* file, int line)
+{
+	// if not in game, ignore Perf warnings - i.e. Moving Static actor in editor will produce this warning
+	if (GIsEditor && e == PxErrorCode::ePERF_WARNING)
+	{
+		return;
+	}
+
+	// @MASSIVE HACK - muting 'triangle too big' warning :(
+	if (line == 223)
+	{
+		return;
+	}
+
+	// Make string to print out, include physx file/line
+	FString ErrorString = FString::Printf(TEXT("PHYSX: %s (%d) %d : %s"), ANSI_TO_TCHAR(file), line, (int32)e, ANSI_TO_TCHAR(message));
+
+	if (e == PxErrorCode::eOUT_OF_MEMORY || e == PxErrorCode::eINTERNAL_ERROR || e == PxErrorCode::eABORT)
+	{
+		UE_LOG(LogPhysics, Error, TEXT("%s"), *ErrorString);
+		ensureMsgf(false, TEXT("%s"), *ErrorString);
+	}
+	else if (e == PxErrorCode::eINVALID_PARAMETER || e == PxErrorCode::eINVALID_OPERATION)
+	{
+		UE_LOG(LogPhysics, Error, TEXT("%s"), *ErrorString);
+		ensureMsgf(false, TEXT("%s"), *ErrorString);
+	}
+	else if (e == PxErrorCode::ePERF_WARNING)
+	{
+		UE_LOG(LogPhysics, Warning, TEXT("%s"), *ErrorString);
+	}
+#if UE_BUILD_DEBUG
+	else if (e == PxErrorCode::eDEBUG_WARNING)
+	{
+		UE_LOG(LogPhysics, Warning, TEXT("%s"), *ErrorString);
+	}
+#endif
+	else
+	{
+		UE_LOG(LogPhysics, Log, TEXT("%s"), *ErrorString);
+	}
+}
 #endif // WITH_PHYSX
diff --git a/Engine/Source/Runtime/Engine/Private/PhysicsEngine/PhysXSupport.h b/Engine/Source/Runtime/Engine/Private/PhysicsEngine/PhysXSupport.h
--- a/Engine/Source/Runtime/Engine/Private/PhysicsEngine/PhysXSupport.h
+++ b/Engine/Source/Runtime/Engine/Private/PhysicsEngine/PhysXSupport.h
@@ -433,44 +433,7 @@ class FPhysXBroadcastingAllocator : public PxBroadcastingAllocator
 class FPhysXErrorCallback : public PxErrorCallback
 {
 public:
-	virtual void reportError(PxErrorCode::Enum e, const char* message, const char* file, int line) override
-	{
-		// if not in game, ignore Perf warnings - i.e. Moving Static actor in editor will produce this warning
-		if (GIsEditor && e == PxErrorCode::ePERF_WARNING)
-		{
-			return;
-		}
-
-		// @MASSIVE HACK - muting 'triangle too big' warning :(
-		if(line == 223)
-		{
-			return;
-		}
-
-		// Make string to print out, include physx file/line
-		FString ErrorString = FString::Printf( TEXT("PHYSX: %s (%d) %d : %s"), ANSI_TO_TCHAR(file), line, (int32)e, ANSI_TO_TCHAR(message) );
-
-		if(e == PxErrorCode::eOUT_OF_MEMORY ||  e == PxErrorCode::eINTERNAL_ERROR || e == PxErrorCode::eABORT)
-		{
-			//UE_LOG(LogPhysics, Error, *ErrorString);
-			UE_LOG(LogPhysics, Warning, TEXT("%s"), *ErrorString);
-		}
-		else if(e == PxErrorCode::eINVALID_PARAMETER || e == PxErrorCode::eINVALID_OPERATION || e == PxErrorCode::ePERF_WARNING)
-		{
-			UE_LOG(LogPhysics, Warning, TEXT("%s"), *ErrorString);
-		}
-#if UE_BUILD_DEBUG
-		else if (e == PxErrorCode::eDEBUG_WARNING)
-		{
-			UE_LOG(LogPhysics, Warning, TEXT("%s"), *ErrorString);
-		}
-#endif
-		else
-		{
-			UE_LOG(LogPhysics, Log, TEXT("%s"), *ErrorString);
-		}
-
-	}
+	virtual void reportError(PxErrorCode::Enum e, const char* message, const char* file, int line) override;
 };
 
 /** 'Shader' used to filter simulation collisions. Could be called on any thread. */
diff --git a/Engine/Source/Runtime/Engine/Private/PhysicsEngine/PhysicsAsset.cpp b/Engine/Source/Runtime/Engine/Private/PhysicsEngine/PhysicsAsset.cpp
--- a/Engine/Source/Runtime/Engine/Private/PhysicsEngine/PhysicsAsset.cpp
+++ b/Engine/Source/Runtime/Engine/Private/PhysicsEngine/PhysicsAsset.cpp
@@ -185,6 +185,15 @@ FBox UPhysicsAsset::CalcAABB(const USkinnedMeshComponent* MeshComp, const FTrans
 		Box = FBox( LocalToWorld.GetLocation(), LocalToWorld.GetLocation() );
 	}
 
+	const float MinBoundSize = 30.f;
+	const FVector BoxSize = Box.GetSize();
+
+	if(BoxSize.GetMin() < MinBoundSize)
+	{
+		const FVector ExpandByDelta ( FMath::Max(0.f, MinBoundSize - BoxSize.X), FMath::Max(0.f, MinBoundSize - BoxSize.Y), FMath::Max(0.f, MinBoundSize - BoxSize.Z) );
+		Box = Box.ExpandBy(ExpandByDelta * 0.5f);	//expand by applies to both directions with GetSize applies to total size so divide by 2
+	}
+
 	return Box;
 }
 
@@ -410,7 +419,7 @@ void UPhysicsAsset::BodyFindConstraints(int32 BodyIndex, TArray<int32>& Constrai
 
 SIZE_T UPhysicsAsset::GetResourceSize(EResourceSizeMode::Type Mode)
 {
-	SIZE_T ResourceSize = 0;
+	SIZE_T ResourceSize = Super::GetResourceSize(Mode);
 
 	for (const auto& SingleBody : BodySetup)
 	{
diff --git a/Engine/Source/Runtime/Engine/Private/PhysicsEngine/PhysicsCollisionHandler.cpp b/Engine/Source/Runtime/Engine/Private/PhysicsEngine/PhysicsCollisionHandler.cpp
--- a/Engine/Source/Runtime/Engine/Private/PhysicsEngine/PhysicsCollisionHandler.cpp
+++ b/Engine/Source/Runtime/Engine/Private/PhysicsEngine/PhysicsCollisionHandler.cpp
@@ -52,7 +52,7 @@ void UPhysicsCollisionHandler::DefaultHandleCollision_AssumesLocked(const FRigid
 	}
 }
 
-void UPhysicsCollisionHandler::HandlePhysicsCollisions_AssumesLocked(const TArray<FCollisionNotifyInfo>& PendingCollisionNotifies)
+void UPhysicsCollisionHandler::HandlePhysicsCollisions_AssumesLocked(TArray<FCollisionNotifyInfo>& PendingCollisionNotifies)
 {
 	// Fire any collision notifies in the queue.
 	for(int32 i=0; i<PendingCollisionNotifies.Num(); i++)
diff --git a/Engine/Source/Runtime/Engine/Private/PhysicsEngine/PhysicsConstraintComponent.cpp b/Engine/Source/Runtime/Engine/Private/PhysicsEngine/PhysicsConstraintComponent.cpp
--- a/Engine/Source/Runtime/Engine/Private/PhysicsEngine/PhysicsConstraintComponent.cpp
+++ b/Engine/Source/Runtime/Engine/Private/PhysicsEngine/PhysicsConstraintComponent.cpp
@@ -69,14 +69,16 @@ UPrimitiveComponent* UPhysicsConstraintComponent::GetComponentInternal(EConstrai
 			{
 				PrimComp = Cast<UPrimitiveComponent>(Actor->GetRootComponent());
 			}
-			// Name specified, see if we can find that property..
+			// Name specified, see if we can find that component..
 			else
 			{
-				UObjectPropertyBase* ObjProp = FindField<UObjectPropertyBase>(Actor->GetClass(), ComponentName);
-				if(ObjProp != NULL)
+				for(UActorComponent* Comp : Actor->GetComponents())
 				{
-					// .. and return the component that is there
-					PrimComp = Cast<UPrimitiveComponent>(ObjProp->GetObjectPropertyValue_InContainer(Actor));
+					if(Comp->GetFName() == ComponentName)
+					{
+						PrimComp = Cast<UPrimitiveComponent>(Comp);
+						break;
+					}
 				}
 			}
 		}	
diff --git a/Engine/Source/Runtime/Engine/Private/PhysicsEngine/PhysicsSettings.cpp b/Engine/Source/Runtime/Engine/Private/PhysicsEngine/PhysicsSettings.cpp
--- a/Engine/Source/Runtime/Engine/Private/PhysicsEngine/PhysicsSettings.cpp
+++ b/Engine/Source/Runtime/Engine/Private/PhysicsEngine/PhysicsSettings.cpp
@@ -20,6 +20,9 @@ UPhysicsSettings::UPhysicsSettings(const FObjectInitializer& ObjectInitializer)
 	, LockedAxis_DEPRECATED(ESettingsLockedAxis::Invalid)
 	, BounceThresholdVelocity(200.f)
 	, MaxAngularVelocity(3600)	//10 revolutions per second
+	, ContactOffsetMultiplier(0.01f)	//TODO: these values likely need tuning, but for now maintaining what we had in BodySetup
+	, MinContactOffset(0.0001f)
+	, MaxContactOffset(1.f)
 	, bSimulateSkeletalMeshOnDedicatedServer(true)
 	, bDefaultHasComplexCollision(true)
 	, bSuppressFaceRemapTable(false)
diff --git a/Engine/Source/Runtime/Engine/Private/SkeletalMesh.cpp b/Engine/Source/Runtime/Engine/Private/SkeletalMesh.cpp
--- a/Engine/Source/Runtime/Engine/Private/SkeletalMesh.cpp
+++ b/Engine/Source/Runtime/Engine/Private/SkeletalMesh.cpp
@@ -2009,7 +2009,6 @@ USkeletalMesh::USkeletalMesh(const FObjectInitializer& ObjectInitializer)
 	ImportedResource = MakeShareable(new FSkeletalMeshResource());
 }
 
-
 void USkeletalMesh::PostInitProperties()
 {
 #if WITH_EDITORONLY_DATA
@@ -2021,7 +2020,66 @@ void USkeletalMesh::PostInitProperties()
 	Super::PostInitProperties();
 }
 
+FBoxSphereBounds USkeletalMesh::GetBounds()
+{
+	return ExtendedBounds;
+}
+
+FBoxSphereBounds USkeletalMesh::GetImportedBounds()
+{
+	return ImportedBounds;
+}
+
+void USkeletalMesh::SetImportedBounds(const FBoxSphereBounds& InBounds)
+{
+	ImportedBounds = InBounds;
+	CalculateExtendedBounds();
+}
+
+void USkeletalMesh::SetPositiveBoundsExtension(const FVector& InExtension)
+{
+	PositiveBoundsExtension = InExtension;
+	CalculateExtendedBounds();
+}
+
+void USkeletalMesh::SetNegativeBoundsExtension(const FVector& InExtension)
+{
+	NegativeBoundsExtension = InExtension;
+	CalculateExtendedBounds();
+}
+PRAGMA_DISABLE_OPTIMIZATION
+void USkeletalMesh::CalculateExtendedBounds()
+{
+	FBoxSphereBounds CalculatedBounds = ImportedBounds;
+
+	// Convert to Min and Max
+	FVector Min = CalculatedBounds.Origin - CalculatedBounds.BoxExtent;
+	FVector Max = CalculatedBounds.Origin + CalculatedBounds.BoxExtent;
+	// Apply bound extensions
+	Min -= NegativeBoundsExtension;
+	Max += PositiveBoundsExtension;
+	// Convert back to Origin, Extent and update SphereRadius
+	CalculatedBounds.Origin = (Min + Max) / 2;
+	CalculatedBounds.BoxExtent = (Max - Min) / 2;
+	CalculatedBounds.SphereRadius = CalculatedBounds.BoxExtent.GetAbsMax();
+
+	ExtendedBounds = CalculatedBounds;
+}
+
+void USkeletalMesh::ValidateBoundsExtension()
+{
+	FVector HalfExtent = ImportedBounds.BoxExtent;
+
+	PositiveBoundsExtension.X = FMath::Clamp(PositiveBoundsExtension.X, -HalfExtent.X, MAX_flt);
+	PositiveBoundsExtension.Y = FMath::Clamp(PositiveBoundsExtension.Y, -HalfExtent.Y, MAX_flt);
+	PositiveBoundsExtension.Z = FMath::Clamp(PositiveBoundsExtension.Z, -HalfExtent.Z, MAX_flt);
 
+	NegativeBoundsExtension.X = FMath::Clamp(NegativeBoundsExtension.X, -HalfExtent.X, MAX_flt);
+	NegativeBoundsExtension.Y = FMath::Clamp(NegativeBoundsExtension.Y, -HalfExtent.Y, MAX_flt);
+	NegativeBoundsExtension.Z = FMath::Clamp(NegativeBoundsExtension.Z, -HalfExtent.Z, MAX_flt);
+}
+
+PRAGMA_ENABLE_OPTIMIZATION
 void USkeletalMesh::InitResources()
 {
 	ImportedResource->InitResources(bHasVertexColors);
@@ -2358,6 +2416,17 @@ void USkeletalMesh::PostEditChangeProperty(FPropertyChangedEvent& PropertyChange
 		BuildPhysicsData();
 	}
 
+	if(UProperty* MemberProperty = PropertyChangedEvent.MemberProperty)
+	{
+		if(MemberProperty->GetFName() == GET_MEMBER_NAME_CHECKED(USkeletalMesh, PositiveBoundsExtension) ||
+			MemberProperty->GetFName() == GET_MEMBER_NAME_CHECKED(USkeletalMesh, NegativeBoundsExtension))
+		{
+			// If the bounds extensions change, recalculate extended bounds.
+			ValidateBoundsExtension();
+			CalculateExtendedBounds();
+		}
+	}
+
 	Super::PostEditChangeProperty(PropertyChangedEvent);
 }
 
@@ -2461,7 +2530,7 @@ void USkeletalMesh::Serialize( FArchive& Ar )
 
 	FStripDataFlags StripFlags( Ar );
 
-	Ar << Bounds;
+	Ar << ImportedBounds;
 	Ar << Materials;
 
 	Ar << RefSkeleton;
@@ -2917,6 +2986,9 @@ void USkeletalMesh::PostLoad()
 		RetargetBasePose = RefSkeleton.GetRefBonePose();
 	}
 #endif
+
+	// Bounds have been loaded - apply extensions.
+	CalculateExtendedBounds();
 }
 
 void USkeletalMesh::GetAssetRegistryTags(TArray<FAssetRegistryTag>& OutTags) const
diff --git a/Engine/Source/Runtime/Engine/Private/SkeletalMeshMerge.cpp b/Engine/Source/Runtime/Engine/Private/SkeletalMeshMerge.cpp
--- a/Engine/Source/Runtime/Engine/Private/SkeletalMeshMerge.cpp
+++ b/Engine/Source/Runtime/Engine/Private/SkeletalMeshMerge.cpp
@@ -635,7 +635,8 @@ bool FSkeletalMeshMerge::ProcessMergeMesh()
 			if( bNeedsInit )
 			{
 				// initialize the merged mesh with the first src mesh entry used
-				MergeMesh->Bounds = SrcMesh->Bounds;
+				MergeMesh->SetImportedBounds(SrcMesh->GetImportedBounds());
+
 				MergeMesh->SkelMirrorAxis = SrcMesh->SkelMirrorAxis;
 				MergeMesh->SkelMirrorFlipAxis = SrcMesh->SkelMirrorFlipAxis;
 
@@ -645,8 +646,8 @@ bool FSkeletalMeshMerge::ProcessMergeMesh()
 			else
 			{
 				// add bounds
-				MergeMesh->Bounds = MergeMesh->Bounds + SrcMesh->Bounds;
-			}			
+				MergeMesh->SetImportedBounds(MergeMesh->GetImportedBounds() + SrcMesh->GetImportedBounds());
+			}
 		}
 	}
 
diff --git a/Engine/Source/Runtime/Engine/Private/SoundNodeLooping.cpp b/Engine/Source/Runtime/Engine/Private/SoundNodeLooping.cpp
--- a/Engine/Source/Runtime/Engine/Private/SoundNodeLooping.cpp
+++ b/Engine/Source/Runtime/Engine/Private/SoundNodeLooping.cpp
@@ -103,18 +103,18 @@ void USoundNodeLooping::ResetChildren(const UPTRINT NodeWaveInstanceHash, FActiv
 
 	for (int32 ResetNodeIndex = 0; ResetNodeIndex < NodesToReset.Num(); ++ResetNodeIndex)
 	{
-		const FNodeHashPairs& NodeHashPair = NodesToReset[ResetNodeIndex];
+		// cache both fields, ref to NodesToReset[ResetNodeIndex] can become invalid after NodesToReset.Add due to memory reallocation
+		const UPTRINT ResetNodeWaveInstanceHash = NodesToReset[ResetNodeIndex].NodeWaveInstanceHash;
+		USoundNode* ResetNode = NodesToReset[ResetNodeIndex].Node;
 
 		// Reset all child nodes so they are initialized again.
-		uint32* Offset = ActiveSound.SoundNodeOffsetMap.Find(NodeHashPair.NodeWaveInstanceHash);
+		uint32* Offset = ActiveSound.SoundNodeOffsetMap.Find(ResetNodeWaveInstanceHash);
 		if (Offset)
 		{
 			bool* bRequiresInitialization = (bool*)&ActiveSound.SoundNodeData[*Offset];
 			*bRequiresInitialization = true;
 		}
 
-		USoundNode* ResetNode = NodeHashPair.Node;
-
 		if (ResetNode->ChildNodes.Num())
 		{
 			for (int32 ResetChildIndex = 0; ResetChildIndex < ResetNode->ChildNodes.Num(); ++ResetChildIndex)
@@ -122,13 +122,13 @@ void USoundNodeLooping::ResetChildren(const UPTRINT NodeWaveInstanceHash, FActiv
 				USoundNode* ResetChildNode = ResetNode->ChildNodes[ResetChildIndex];
 				if (ResetChildNode)
 				{
-					NodesToReset.Add(FNodeHashPairs(ResetChildNode, GetNodeWaveInstanceHash(NodeHashPair.NodeWaveInstanceHash, ResetChildNode, ResetChildIndex)));
+					NodesToReset.Add(FNodeHashPairs(ResetChildNode, GetNodeWaveInstanceHash(ResetNodeWaveInstanceHash, ResetChildNode, ResetChildIndex)));
 				}
 			}
 		}
 		else if (ResetNode->IsA<USoundNodeWavePlayer>())
 		{
-			FWaveInstance* WaveInstance = ActiveSound.FindWaveInstance(NodeHashPair.NodeWaveInstanceHash);
+			FWaveInstance* WaveInstance = ActiveSound.FindWaveInstance(ResetNodeWaveInstanceHash);
 			if (WaveInstance)
 			{
 				WaveInstance->bAlreadyNotifiedHook = true;
diff --git a/Engine/Source/Runtime/Engine/Private/StaticMesh.cpp b/Engine/Source/Runtime/Engine/Private/StaticMesh.cpp
--- a/Engine/Source/Runtime/Engine/Private/StaticMesh.cpp
+++ b/Engine/Source/Runtime/Engine/Private/StaticMesh.cpp
@@ -872,11 +872,11 @@ void FStaticMeshLODSettings::Initialize(const FConfigFile& IniFile)
 	const FConfigSection* Section = IniFile.Find(IniSection);
 	if (Section)
 	{
-		for (TMultiMap<FName,FString>::TConstIterator It(*Section); It; ++It)
+		for (TMultiMap<FName,FConfigValue>::TConstIterator It(*Section); It; ++It)
 		{
 			FName GroupName = It.Key();
 			FStaticMeshLODGroup& Group = Groups.FindOrAdd(GroupName);
-			ReadEntry(Group, It.Value());
+			ReadEntry(Group, It.Value().GetValue());
 		};
 	}
 
@@ -1713,6 +1713,8 @@ void UStaticMesh::GetAssetRegistryTags(TArray<FAssetRegistryTag>& OutTags) const
 	int32 NumVertices = 0;
 	int32 NumUVChannels = 0;
 	int32 NumLODs = 0;
+	int32 NumSectionsWithCollision = 0;
+
 	if (RenderData && RenderData->LODResources.Num() > 0)
 	{
 		const FStaticMeshLODResources& LOD = RenderData->LODResources[0];
@@ -1720,6 +1722,19 @@ void UStaticMesh::GetAssetRegistryTags(TArray<FAssetRegistryTag>& OutTags) const
 		NumVertices = LOD.VertexBuffer.GetNumVertices();
 		NumUVChannels = LOD.VertexBuffer.GetNumTexCoords();
 		NumLODs = RenderData->LODResources.Num();
+
+#if WITH_EDITORONLY_DATA
+		// Find how many sections have collision enabled
+		const int32 UseLODIndex = FMath::Clamp(LODForCollision, 0, RenderData->LODResources.Num() - 1);
+		FStaticMeshLODResources& CollisionLOD = RenderData->LODResources[UseLODIndex];
+		for (int32 SectionIndex = 0; SectionIndex < CollisionLOD.Sections.Num(); ++SectionIndex)
+		{
+			if (SectionInfoMap.Get(UseLODIndex, SectionIndex).bEnableCollision)
+			{
+				NumSectionsWithCollision++;
+			}
+		}
+#endif
 	}
 
 	int32 NumCollisionPrims = 0;
@@ -1735,6 +1750,13 @@ void UStaticMesh::GetAssetRegistryTags(TArray<FAssetRegistryTag>& OutTags) const
 	}
 	const FString ApproxSizeStr = FString::Printf(TEXT("%dx%dx%d"), FMath::RoundToInt(Bounds.BoxExtent.X * 2.0f), FMath::RoundToInt(Bounds.BoxExtent.Y * 2.0f), FMath::RoundToInt(Bounds.BoxExtent.Z * 2.0f));
 
+	// Get name of default collision profile
+	FName DefaultCollisionName = NAME_None;
+	if(BodySetup != nullptr)
+	{
+		DefaultCollisionName = BodySetup->DefaultInstance.GetCollisionProfileName();
+	}
+
 	OutTags.Add( FAssetRegistryTag("Triangles", FString::FromInt(NumTriangles), FAssetRegistryTag::TT_Numerical) );
 	OutTags.Add( FAssetRegistryTag("Vertices", FString::FromInt(NumVertices), FAssetRegistryTag::TT_Numerical) );
 	OutTags.Add( FAssetRegistryTag("UVChannels", FString::FromInt(NumUVChannels), FAssetRegistryTag::TT_Numerical) );
@@ -1742,6 +1764,8 @@ void UStaticMesh::GetAssetRegistryTags(TArray<FAssetRegistryTag>& OutTags) const
 	OutTags.Add( FAssetRegistryTag("ApproxSize", ApproxSizeStr, FAssetRegistryTag::TT_Dimensional) );
 	OutTags.Add( FAssetRegistryTag("CollisionPrims", FString::FromInt(NumCollisionPrims), FAssetRegistryTag::TT_Numerical));
 	OutTags.Add( FAssetRegistryTag("LODs", FString::FromInt(NumLODs), FAssetRegistryTag::TT_Numerical));
+	OutTags.Add( FAssetRegistryTag("SectionsWithCollision", FString::FromInt(NumSectionsWithCollision), FAssetRegistryTag::TT_Numerical));
+	OutTags.Add( FAssetRegistryTag("DefaultCollision", DefaultCollisionName.ToString(), FAssetRegistryTag::TT_Alphabetical));
 
 #if WITH_EDITORONLY_DATA
 	if (AssetImportData)
@@ -2349,7 +2373,7 @@ bool UStaticMesh::GetPhysicsTriMeshData(struct FTriMeshCollisionData* CollisionD
 	{
 		const FStaticMeshSection& Section = LOD.Sections[SectionIndex];
 
-		if (bInUseAllTriData || SectionInfoMap.Get(0,SectionIndex).bEnableCollision)
+		if (bInUseAllTriData || SectionInfoMap.Get(UseLODIndex,SectionIndex).bEnableCollision)
 		{
 			const uint32 OnePastLastIndex  = Section.FirstIndex + Section.NumTriangles*3;
 
diff --git a/Engine/Source/Runtime/Engine/Private/TimerManager.cpp b/Engine/Source/Runtime/Engine/Private/TimerManager.cpp
--- a/Engine/Source/Runtime/Engine/Private/TimerManager.cpp
+++ b/Engine/Source/Runtime/Engine/Private/TimerManager.cpp
@@ -11,19 +11,13 @@ DECLARE_CYCLE_STAT(TEXT("ClearTimer"), STAT_ClearTimer, STATGROUP_Engine);
 
 void FTimerHandle::MakeValid()
 {
-	static int LastAssignedHandle = -1;
-
-	if (!IsValid())
+	if (GWorld)
 	{
-		++LastAssignedHandle;
-		Handle = LastAssignedHandle;
+		GWorld->GetTimerManager().ValidateHandle(*this);
 	}
-
-	check(IsValid());
 }
 
 
-
 FString FTimerUnifiedDelegate::ToString() const
 {
 	const UObject* Object = nullptr;
@@ -165,7 +159,7 @@ void FTimerManager::InternalSetTimer(FTimerHandle& InOutHandle, FTimerUnifiedDel
 
 	if (InRate > 0.f)
 	{
-		InOutHandle.MakeValid();
+		ValidateHandle(InOutHandle);
 
 		// set up the new timer
 		FTimerData NewTimerData;
@@ -557,6 +551,17 @@ void FTimerManager::ListTimers() const
 	UE_LOG(LogEngine, Log, TEXT("------- %d Total Timers -------"), PendingTimerList.Num() + PausedTimerList.Num() + ActiveTimerHeap.Num());
 }
 
+void FTimerManager::ValidateHandle(FTimerHandle& InOutHandle)
+{
+	if (!InOutHandle.IsValid())
+	{
+		++LastAssignedHandle;
+		InOutHandle.Handle = LastAssignedHandle;
+	}
+
+	checkf(InOutHandle.IsValid(), TEXT("Timer handle has wrapped around to 0!"));
+}
+
 // Handler for ListTimers console command
 static void OnListTimers(UWorld* World)
 {
diff --git a/Engine/Source/Runtime/Engine/Private/TimerManagerTests.cpp b/Engine/Source/Runtime/Engine/Private/TimerManagerTests.cpp
--- a/Engine/Source/Runtime/Engine/Private/TimerManagerTests.cpp
+++ b/Engine/Source/Runtime/Engine/Private/TimerManagerTests.cpp
@@ -59,7 +59,7 @@ bool TimerManagerTest_MissingTimers(UWorld* World, FAutomationTestBase* Test)
 	FTimerManager& TimerManager = World->GetTimerManager();
 	FTimerHandle Handle;
 
-	Handle.MakeValid();
+	TimerManager.ValidateHandle(Handle);
 
 	Test->TestFalse(TIMER_TEST_TEXT("TimerExists called with an invalid handle"), TimerManager.TimerExists(Handle));
 	Test->TestFalse(TIMER_TEST_TEXT("IsTimerActive called with an invalid handle"), TimerManager.IsTimerActive(Handle));
diff --git a/Engine/Source/Runtime/Engine/Private/UnrealEngine.cpp b/Engine/Source/Runtime/Engine/Private/UnrealEngine.cpp
--- a/Engine/Source/Runtime/Engine/Private/UnrealEngine.cpp
+++ b/Engine/Source/Runtime/Engine/Private/UnrealEngine.cpp
@@ -1207,10 +1207,10 @@ void UEngine::UpdateTimeAndHandleMaxTickRate()
 	FApp::UpdateLastTime();
 
 	// Calculate delta time and update time.
-	if( bUseFixedTimeStep || bUseFixedFrameRate )
+	if( bUseFixedTimeStep )
 	{
 		bTimeWasManipulated = true;
-		const float FrameRate = bUseFixedTimeStep ? FApp::GetFixedDeltaTime() : (1.f / FixedFrameRate);
+		const float FrameRate = FApp::GetFixedDeltaTime();
 		FApp::SetDeltaTime(FrameRate);
 		LastTime = FApp::GetCurrentTime();
 		FApp::SetCurrentTime(FApp::GetCurrentTime() + FApp::GetDeltaTime());
@@ -1220,7 +1220,7 @@ void UEngine::UpdateTimeAndHandleMaxTickRate()
 		FApp::SetCurrentTime(FPlatformTime::Seconds());
 		// Did we just switch from a fixed time step to real-time?  If so, then we'll update our
 		// cached 'last time' so our current interval isn't huge (or negative!)
-		if( bTimeWasManipulated )
+		if( bTimeWasManipulated && !bUseFixedFrameRate )
 		{
 			LastTime = FApp::GetCurrentTime() - FApp::GetDeltaTime();
 			bTimeWasManipulated = false;
@@ -1284,6 +1284,14 @@ void UEngine::UpdateTimeAndHandleMaxTickRate()
 			}
 			FApp::SetCurrentTime(FPlatformTime::Seconds());
 		}
+		else if(bUseFixedFrameRate)
+		{
+			//We are doing fixed framerate and the real delta time is bigger than our desired delta time. In this case we start falling behind real time (and that's ok)
+			const float FrameRate = 1.f / FixedFrameRate;
+			FApp::SetDeltaTime(FrameRate);
+			FApp::SetCurrentTime(LastTime + FApp::GetDeltaTime());
+			bTimeWasManipulated = true;
+		}
 
 
 		SET_FLOAT_STAT(STAT_GameTickWantedWaitTime,WaitTime * 1000.f);
@@ -4084,7 +4092,7 @@ bool UEngine::HandleMemReportDeferredCommand( const TCHAR* Cmd, FOutputDevice& A
 	{
 		for (FConfigSectionMap::TIterator It(*CommandsToRun); It; ++It)
 		{
-			Exec( InWorld, *It.Value(), *ReportAr );
+			Exec( InWorld, *It.Value().GetValue(), *ReportAr );
 			ReportAr->Logf( LINE_TERMINATOR );
 		}
 	}
@@ -4097,7 +4105,7 @@ bool UEngine::HandleMemReportDeferredCommand( const TCHAR* Cmd, FOutputDevice& A
 		{
 			for (FConfigSectionMap::TIterator It(*CommandsToRun); It; ++It)
 			{
-				Exec( InWorld, *It.Value(), *ReportAr );
+				Exec( InWorld, *It.Value().GetValue(), *ReportAr );
 				ReportAr->Logf( LINE_TERMINATOR );
 			}
 		}
@@ -6581,6 +6589,11 @@ float UEngine::GetMaxTickRate(float DeltaTime, bool bAllowFrameRateSmoothing) co
 		}
 	}
 
+	if(bUseFixedFrameRate)
+	{
+		MaxTickRate = FixedFrameRate;
+	}
+
 	if (CVarCauseHitches.GetValueOnGameThread())
 	{
 		static float RunningHitchTimer = 0.f;
diff --git a/Engine/Source/Runtime/Engine/Private/UserInterface/Console.cpp b/Engine/Source/Runtime/Engine/Private/UserInterface/Console.cpp
--- a/Engine/Source/Runtime/Engine/Private/UserInterface/Console.cpp
+++ b/Engine/Source/Runtime/Engine/Private/UserInterface/Console.cpp
@@ -16,8 +16,21 @@
 
 static const uint32 MAX_AUTOCOMPLETION_LINES = 20;
 
-static FName NAME_Typing = FName(TEXT("Typing"));
-static FName NAME_Open = FName(TEXT("Open"));
+static const FName NAME_Typing = FName(TEXT("Typing"));
+static const FName NAME_Open = FName(TEXT("Open"));
+
+namespace ConsoleDefs
+{
+	/** Colors */
+	static const FColor BorderColor( 140, 140, 140 );
+	static const FColor CursorColor( 255, 255, 255 );
+	static const FColor AutocompleteBackgroundColor( 0, 0, 0 );
+	static const FColor CursorLineColor( 0, 50, 0 );
+	static const int32 AutocompleteGap = 6;
+
+	/** Text that appears before the user's typed input string that acts as a visual cue for the editable area */
+	static const FString LeadingInputText(TEXT(" > "));
+}
 
 class FConsoleVariableAutoCompleteVisitor 
 {
@@ -37,19 +50,35 @@ public:
 			return;
 		}
 
+		const UConsoleSettings* ConsoleSettings = GetDefault<UConsoleSettings>();
+
 		// can be optimized
 		int32 NewIdx = Sink.AddZeroed(1);
-
 		FAutoCompleteCommand& Cmd = Sink[NewIdx];
-		
 		Cmd.Command = Name;
+
+		IConsoleVariable* CVariable = CVar->AsVariable();
+		if (CVariable)
+		{
+			if (CVar->TestFlags(ECVF_ReadOnly))
+			{
+				Cmd.Color = ConsoleSettings->AutoCompleteFadedColor;
+			}
+			else
+			{
+				Cmd.Color = ConsoleSettings->AutoCompleteCVarColor;
+			}
+		}
+		else
+		{
+			Cmd.Color = ConsoleSettings->AutoCompleteCommandColor;
+		}
 	}
 };
 
 UConsole::UConsole(const FObjectInitializer& ObjectInitializer)
 	: Super(ObjectInitializer)
-//	, HistoryBufferCur(0)
-//	, AutoCompleteCursor(-1)
+	, ConsoleSettings(GetDefault<UConsoleSettings>())
 {
 	// Structure to hold one-time initialization
 	struct FConstructorStatics
@@ -111,14 +140,13 @@ void UConsole::BuildRuntimeAutoCompleteList(bool bForce)
 
 	AutoCompleteTree.ChildNodes.Empty();
 
-	const UConsoleSettings* ConsoleSettings = GetDefault<UConsoleSettings>();
-
 	// copy the manual list first
 	AutoCompleteList.Empty();
 	AutoCompleteList.AddZeroed(ConsoleSettings->ManualAutoCompleteList.Num());
 	for (int32 Idx = 0; Idx < ConsoleSettings->ManualAutoCompleteList.Num(); Idx++)
 	{
 		AutoCompleteList[Idx] = ConsoleSettings->ManualAutoCompleteList[Idx];
+		AutoCompleteList[Idx].Color = ConsoleSettings->AutoCompleteCommandColor;
 	}
 
 	// console variables
@@ -159,6 +187,7 @@ void UConsole::BuildRuntimeAutoCompleteList(bool bForce)
 
 			int32 NewIdx = AutoCompleteList.AddZeroed(1);
 			AutoCompleteList[NewIdx].Command = FuncName;
+			AutoCompleteList[NewIdx].Color = ConsoleSettings->AutoCompleteCommandColor;
 			// build a help string
 			// append each property (and it's type) to the help string
 			for (TFieldIterator<UProperty> PropIt(Func); PropIt && (PropIt->PropertyFlags & CPF_Parm); ++PropIt)
@@ -207,8 +236,11 @@ void UConsole::BuildRuntimeAutoCompleteList(bool bForce)
 				NewIdx = AutoCompleteList.AddZeroed(3);
 			}
 			AutoCompleteList[NewIdx].Command = FString::Printf(TEXT("open %s"),*TrimmedMapName);
+			AutoCompleteList[NewIdx].Color = ConsoleSettings->AutoCompleteCommandColor;
 			AutoCompleteList[NewIdx+1].Command = FString::Printf(TEXT("travel %s"),*TrimmedMapName);
+			AutoCompleteList[NewIdx+1].Color = ConsoleSettings->AutoCompleteCommandColor;
 			AutoCompleteList[NewIdx+2].Command = FString::Printf(TEXT("servertravel %s"),*TrimmedMapName);
+			AutoCompleteList[NewIdx+2].Color = ConsoleSettings->AutoCompleteCommandColor;
 			//MapNames.AddItem(Pkg);
 		}
 	}
@@ -217,6 +249,7 @@ void UConsole::BuildRuntimeAutoCompleteList(bool bForce)
 		int32 NewIdx = AutoCompleteList.AddZeroed(1);
 		AutoCompleteList[NewIdx].Command = FString(TEXT("open 127.0.0.1"));
 		AutoCompleteList[NewIdx].Desc = FString(TEXT("(opens connection to localhost)"));
+		AutoCompleteList[NewIdx].Color = ConsoleSettings->AutoCompleteCommandColor;
 	}
 
 #if STATS
@@ -231,6 +264,7 @@ void UConsole::BuildRuntimeAutoCompleteList(bool bForce)
 			Command += StatGroupName.ToString().RightChop(sizeof("STATGROUP_") - 1);
 
 			AutoCompleteList[NewIdx].Command = Command;
+			AutoCompleteList[NewIdx].Color = ConsoleSettings->AutoCompleteCommandColor;
 			NewIdx++;
 		}
 	}
@@ -254,6 +288,7 @@ void UConsole::BuildRuntimeAutoCompleteList(bool bForce)
 				int32 NewIdx = AutoCompleteList.AddZeroed(1);
 				AutoCompleteList[NewIdx].Command = TEXT("show ") + InName;
 				AutoCompleteList[NewIdx].Desc = FString::Printf(TEXT("(toggles the %s showflag)"),*LocName.ToString());
+				AutoCompleteList[NewIdx].Color = GetDefault<UConsoleSettings>()->AutoCompleteCommandColor;
 				
 				return true;
 			}
@@ -324,7 +359,7 @@ void UConsole::UpdateCompleteIndices()
 		Filter.SetRawFilterText(FText::FromString(Right));
 
 		AutoCompleteIndex = 0;
-		AutoCompleteCursor = -1;
+		AutoCompleteCursor = 0;
 		AutoComplete.Empty();
 
 		for (auto Command : AutoCompleteList)
@@ -340,7 +375,7 @@ void UConsole::UpdateCompleteIndices()
 	}
 
 	AutoCompleteIndex = 0;
-	AutoCompleteCursor = -1;
+	AutoCompleteCursor = 0;
 	AutoComplete.Empty();
 	FAutoCompleteNode *Node = &AutoCompleteTree;
 	FString LowerTypedStr = TypedStr.ToLower();
@@ -415,6 +450,7 @@ void UConsole::SetAutoCompleteFromHistory()
 		FAutoCompleteCommand Cmd;
 
 		Cmd.Command = HistoryBuffer[i]; 
+		Cmd.Color = ConsoleSettings->HistoryColor;
 		Cmd.SetHistory();
 
 		AutoComplete.Add(Cmd);
@@ -474,8 +510,6 @@ void UConsole::ClearOutput()
 
 void UConsole::OutputTextLine(const FString& Text)
 {
-	const UConsoleSettings* ConsoleSettings = GetDefault<UConsoleSettings>();
-
 	// If we are full, delete the first line
 	if (Scrollback.Num() > ConsoleSettings->MaxScrollbackSize)
 	{
@@ -613,6 +647,7 @@ void UConsole::AppendInputText(const FString& Text)
 		}
 	};
 	UpdateCompleteIndices();
+	UpdatePrecompletedInputLine();
 }
 
 bool UConsole::InputChar_Typing( int32 ControllerId, const FString& Unicode )
@@ -677,7 +712,9 @@ bool UConsole::InputKey_InputLine( int32 ControllerId, FKey Key, EInputEvent Eve
 				SetCursorPos(0);
 
 				AutoCompleteIndex = 0;
-				AutoCompleteCursor = -1;
+				AutoCompleteCursor = 0;
+				PrecompletedInputLine = FString(TEXT(""));
+				LastAutoCompletedCommand = nullptr;
 				AutoComplete.Empty();
 				bAutoCompleteLocked = false;
 
@@ -690,15 +727,6 @@ bool UConsole::InputKey_InputLine( int32 ControllerId, FKey Key, EInputEvent Eve
 
 			return true;
 		}
-		else if (Key == EKeys::Tab && Event == IE_Pressed)
-		{
-			if (!bAutoCompleteLocked)
-			{
-				SetInputLineFromAutoComplete();
-			}
-
-			return true;
-		}
 		else if( Key == EKeys::Enter && Event == IE_Released )
 		{
 			if( !TypedStr.IsEmpty() )
@@ -728,6 +756,10 @@ bool UConsole::InputKey_InputLine( int32 ControllerId, FKey Key, EInputEvent Eve
 				FakeGotoState(NAME_None);
 			}
 
+			// A command was executed and/or the console closed, discard the most recent autocomplete info
+			PrecompletedInputLine = FString(TEXT(""));
+			LastAutoCompletedCommand = nullptr;
+
 			return true;
 		}
 		else if( Event != IE_Pressed && Event != IE_Repeat )
@@ -746,52 +778,85 @@ bool UConsole::InputKey_InputLine( int32 ControllerId, FKey Key, EInputEvent Eve
 			{
 				if(AutoComplete.Num())
 				{
-					if(AutoCompleteCursor + 1 < FMath::Min((int32)MAX_AUTOCOMPLETION_LINES, AutoComplete.Num()))
+					if (AutoCompleteCursor > 0)
 					{
 						// move cursor within displayed region
-						++AutoCompleteCursor;
+						--AutoCompleteCursor;
 					}
-					else 
+					else
 					{
-						// can be negative
-						int32 ScrollRegionSize = AutoComplete.Num() - (int32)MAX_AUTOCOMPLETION_LINES;
-
 						// can we scroll?
-						if(AutoCompleteIndex + 1 < ScrollRegionSize)
+						if(AutoCompleteIndex > 0)
 						{
-							++AutoCompleteIndex;
+							--AutoCompleteIndex;
+						}
+						else
+						{
+							// wrap around
+							AutoCompleteIndex = FMath::Max(0, AutoComplete.Num() - (int32)MAX_AUTOCOMPLETION_LINES - 1);
+							if (AutoComplete.Num() <= MAX_AUTOCOMPLETION_LINES)
+							{
+								AutoCompleteCursor = AutoComplete.Num() + AutoCompleteCursor - 1;
+							}
+							else
+							{
+								// skip the "x more matches" line when wrapping
+								AutoCompleteIndex++;
+								AutoCompleteCursor = MAX_AUTOCOMPLETION_LINES + AutoCompleteCursor - 1;
+							}
 						}
+						bAutoCompleteLocked = false;
 					}
-					SetInputLineFromAutoComplete();
 				}
 				else
 				{
 					SetAutoCompleteFromHistory();
-					SetInputLineFromAutoComplete();
 				}
+				SetInputLineFromAutoComplete();
 			}
-			return true;
 		}
-		else if( Key == EKeys::Down )
+		else if( Key == EKeys::Down || Key == EKeys::Tab )
 		{
-			if (!bCtrl && AutoComplete.Num() > 1)
+			if (!bCtrl && AutoComplete.Num())
 			{
-				if(AutoCompleteCursor > 0)
+				bool bScroll = AutoComplete.Num() > 1;
+
+				if (Key == EKeys::Tab)
 				{
-					// move cursor within displayed region
-					--AutoCompleteCursor;
+					bCaptureKeyInput = true;
+
+					// If this is a repeated tab press, we want to scroll. Otherwise complete the current command
+					bScroll = bScroll && LastAutoCompletedCommand == TypedStr;
 				}
-				else
+
+				if (bScroll)
 				{
-					// can we scroll?
-					if(AutoCompleteIndex > 0)
+					if (AutoCompleteCursor + 1 < FMath::Min((int32)MAX_AUTOCOMPLETION_LINES, AutoComplete.Num()))
 					{
-						--AutoCompleteIndex;
+						// move cursor within displayed region
+						++AutoCompleteCursor;
+					}
+					else
+					{
+						// can be negative
+						int32 ScrollRegionSize = AutoComplete.Num() - (int32)MAX_AUTOCOMPLETION_LINES;
+
+						// can we scroll?
+						if (AutoCompleteIndex < ScrollRegionSize)
+						{
+							++AutoCompleteIndex;
+						}
+						else
+						{
+							// wrap around
+							AutoCompleteIndex = AutoCompleteCursor = 0;
+						}
 					}
-					bAutoCompleteLocked = false;
 				}
+
 				SetInputLineFromAutoComplete();
 			}
+			return true;
 		}
 		else if( Key==EKeys::BackSpace )
 		{
@@ -859,42 +924,11 @@ bool UConsole::InputKey_InputLine( int32 ControllerId, FKey Key, EInputEvent Eve
 			SetCursorPos(TypedStr.Len());
 			return true;
 		}
-		else
-		{
-			// auto complete is open and the user adds a key
-			if (AutoComplete.Num() > 0 && !bAutoCompleteLocked && AutoCompleteCursor >= 0)
-			{
-				const FAutoCompleteCommand& Cmd = AutoComplete[AutoCompleteIndex + AutoCompleteCursor];
-
-				TypedStr = Cmd.Command;
-				SetCursorPos(TypedStr.Len());
-				bAutoCompleteLocked = true;
-
-				// no "return true" so that the following key will be added the auto completed string
-			}
-		}
 	}
 
 	return false;
 }
 
-
-namespace ConsoleDefs
-{
-	/** Colors */
-	static const FColor BorderColor( 140, 140, 140 );
-	static const FColor CursorColor( 255, 255, 255 );
-	static const FColor InputTextColor( 220, 220, 220 );
-	static const FColor AutocompleteBackgroundColor( 0, 0, 0 );
-	static const FColor AutocompletePartialSuggestionColor( 100, 100, 100 );
-	static const FColor AutocompleteSuggestionColor( 180, 180, 180 );
-	static const FColor CursorLineColor( 0, 50, 0 );
-	static const int32 AutocompleteGap = 6;
-
-	/** Text that appears before the user's typed input string that acts as a visual cue for the editable area */
-	static const FString LeadingInputText( TEXT( " > " ) );
-}
-
 void UConsole::SetInputLineFromAutoComplete()
 {
 	if (AutoComplete.Num() > 0)
@@ -905,6 +939,25 @@ void UConsole::SetInputLineFromAutoComplete()
 		TypedStr = Cmd.Command;
 		SetCursorPos(TypedStr.Len());
 		bAutoCompleteLocked = true;
+
+		PrecompletedInputLine = Cmd.Command;
+		LastAutoCompletedCommand = Cmd.Command;
+	}
+}
+
+void UConsole::UpdatePrecompletedInputLine()
+{
+	// Set the full command text for the user input if they were to autocomplete it with tab
+	if (AutoComplete.Num() > 0)
+	{
+		const int32 Index = AutoCompleteIndex + (AutoCompleteCursor >= 0 ? AutoCompleteCursor : 0);
+		const FAutoCompleteCommand& Cmd = AutoComplete[Index];
+		PrecompletedInputLine = Cmd.Command;
+	}
+	else
+	{
+		// Input buffer cleared, or the user is typing some nonexistent command
+		PrecompletedInputLine = FString(TEXT(""));
 	}
 }
 
@@ -923,16 +976,6 @@ void UConsole::NormalizeHistoryBuffer()
 
 void UConsole::PostRender_Console_Typing(UCanvas* Canvas)
 {
-	float xl, yl;
-
-	// Blank out a space
-
-	// use the smallest font
-	UFont* Font = GEngine->GetSmallFont();
-	// determine the position for the cursor
-	FString OutStr = FString::Printf(TEXT("%s%s"), *ConsoleDefs::LeadingInputText, *TypedStr);
-	Canvas->StrLen(Font, OutStr,xl,yl);
-
 	float ClipX = Canvas->ClipX;
 	float ClipY = Canvas->ClipY;
 	float LeftPos = 0;
@@ -1049,7 +1092,10 @@ void UConsole::PostRender_Console_Open(UCanvas* Canvas)
 	Canvas->StrLen(Font, TEXT("M"),xl,yl);
 
 	// Background
-	FCanvasTileItem ConsoleTile( FVector2D( LeftPos, 0.0f ), DefaultTexture_Black->Resource, FVector2D( ClipX, Height + TopPos - yl), FVector2D( 0.0f, 0.0f), FVector2D( 1.0f, 1.0f ), FLinearColor::White );
+	FLinearColor BackgroundColor = ConsoleDefs::AutocompleteBackgroundColor.ReinterpretAsLinear();
+	BackgroundColor.A = ConsoleSettings->BackgroundOpacityPercentage / 100.0f;
+	FCanvasTileItem ConsoleTile( FVector2D( LeftPos, 0.0f ), DefaultTexture_Black->Resource, FVector2D( ClipX, Height + TopPos - yl), FVector2D( 0.0f, 0.0f), FVector2D( 1.0f, 1.0f ), BackgroundColor );
+	ConsoleTile.BlendMode = SE_BLEND_Translucent;
 	Canvas->DrawItem( ConsoleTile );
 
 	// figure out which element of the scrollback buffer to should appear first (at the top of the screen)
@@ -1059,7 +1105,7 @@ void UConsole::PostRender_Console_Open(UCanvas* Canvas)
 
 	if(Scrollback.Num())
 	{
-		FCanvasTextItem ConsoleText( FVector2D( LeftPos,TopPos+Height-5-yl ), FText::FromString(TEXT("")), GEngine->GetSmallFont(), ConsoleDefs::InputTextColor );		
+		FCanvasTextItem ConsoleText( FVector2D( LeftPos,TopPos+Height-5-yl ), FText::FromString(TEXT("")), GEngine->GetSmallFont(), ConsoleSettings->InputColor );
 
 		// change the text color to white
 		ConsoleText.SetColor( FLinearColor::White );
@@ -1163,11 +1209,13 @@ void UConsole::PostRender_InputLine(UCanvas* Canvas, FIntPoint UserInputLinePos)
 {
 	float xl, yl;
 
+	const FString TypedInputText = FString::Printf(TEXT("%s%s"), *ConsoleDefs::LeadingInputText, *TypedStr);
+	const FString PrecompletedInputText = !PrecompletedInputLine.IsEmpty() ? PrecompletedInputLine.RightChop(TypedStr.Len()) : FString();
+
 	// use the smallest font
 	UFont* Font = GEngine->GetSmallFont();
-	// determine the position for the cursor
-	FString OutStr = FString::Printf(TEXT("%s%s"), *ConsoleDefs::LeadingInputText, *TypedStr);
-	Canvas->StrLen(Font, OutStr,xl,yl);
+	// determine the size of the input line
+	Canvas->StrLen(Font, TypedInputText, xl, yl);
 
 	float ClipX = Canvas->ClipX;
 	float ClipY = Canvas->ClipY;
@@ -1185,7 +1233,10 @@ void UConsole::PostRender_InputLine(UCanvas* Canvas, FIntPoint UserInputLinePos)
 	}
 
 	// Background
-	FCanvasTileItem ConsoleTile( FVector2D( UserInputLinePos.X,UserInputLinePos.Y-6-yl ), DefaultTexture_Black->Resource, FVector2D( ClipX, yl+6 ), FVector2D( 0.0f, 0.0f), FVector2D( 1.0f, 1.0f ), FLinearColor::White );
+	FLinearColor BackgroundColor = ConsoleDefs::AutocompleteBackgroundColor.ReinterpretAsLinear();
+	BackgroundColor.A = ConsoleSettings->BackgroundOpacityPercentage / 100.0f;
+	FCanvasTileItem ConsoleTile( FVector2D( UserInputLinePos.X,UserInputLinePos.Y-6-yl ), DefaultTexture_Black->Resource, FVector2D( ClipX, yl+6 ), FVector2D( 0.0f, 0.0f), FVector2D( 1.0f, 1.0f ), BackgroundColor );
+	ConsoleTile.BlendMode = SE_BLEND_Translucent;
 	Canvas->DrawItem( ConsoleTile );
 
 	// Separator line
@@ -1195,11 +1246,20 @@ void UConsole::PostRender_InputLine(UCanvas* Canvas, FIntPoint UserInputLinePos)
 	Canvas->DrawItem( ConsoleTile );
 
 	// Currently typed string
-	FText Str = FText::FromString( OutStr );
-	FCanvasTextItem ConsoleText( FVector2D( UserInputLinePos.X,UserInputLinePos.Y-3-yl ), Str , GEngine->GetSmallFont(), ConsoleDefs::InputTextColor );		
+	FText Str = FText::FromString( TypedInputText );
+	FCanvasTextItem ConsoleText( FVector2D( UserInputLinePos.X,UserInputLinePos.Y-3-yl ), Str , GEngine->GetSmallFont(), ConsoleSettings->InputColor );
+	ConsoleText.EnableShadow(FLinearColor::Black);
 	Canvas->DrawItem( ConsoleText );
 
-	// draw the remaining text for matching auto-complete
+	// Precompleted remainder of the typed string (faded out)
+	if (!PrecompletedInputText.IsEmpty())
+	{
+		ConsoleText.SetColor(ConsoleSettings->AutoCompleteFadedColor);
+		ConsoleText.Text = FText::FromString(PrecompletedInputText);
+		Canvas->DrawItem(ConsoleText, UserInputLinePos.X + xl, UserInputLinePos.Y - 3 - yl);
+	}
+
+	// Draw the autocomplete elements
 	if (AutoComplete.Num() > 0)
 	{
 		int32 StartIdx = AutoCompleteIndex;
@@ -1213,7 +1273,9 @@ void UConsole::PostRender_InputLine(UCanvas* Canvas, FIntPoint UserInputLinePos)
 		float y = UserInputLinePos.Y - 6.0f - (yl * 2.0f);
 
 		// Set the background color/texture of the auto-complete section
-		ConsoleTile.SetColor( ConsoleDefs::AutocompleteBackgroundColor );
+		FLinearColor AutoCompleteBackgroundColor = ConsoleDefs::AutocompleteBackgroundColor;
+		AutoCompleteBackgroundColor.A = ConsoleSettings->BackgroundOpacityPercentage / 100.0f;
+		ConsoleTile.SetColor( AutoCompleteBackgroundColor );
 		ConsoleTile.Texture = DefaultTexture_White->Resource;
 
 		// wasteful memory allocations but when typing in a console command this is fine
@@ -1221,119 +1283,122 @@ void UConsole::PostRender_InputLine(UCanvas* Canvas, FIntPoint UserInputLinePos)
 		// to avoid memory many allocations
 		AutoCompleteElements.Empty(MAX_AUTOCOMPLETION_LINES + 1);
 
+		float MaxLeftWidth = 0;
+		float MaxRightWidth = 0;
 		for (int32 MatchIdx = 0; MatchIdx < MAX_AUTOCOMPLETION_LINES && MatchIdx < AutoComplete.Num(); MatchIdx++)
 		{
 			const FAutoCompleteCommand &Cmd = AutoComplete[StartIdx + MatchIdx];
-			OutStr = Cmd.Desc;
-
 			AutoCompleteElements.Add(&Cmd);
+
+			// Find the longest command and the longest description for left-justification of the descriptions
+			float CmdLenX, CmdLenY;
+			Canvas->StrLen(Font, Cmd.GetLeft(), CmdLenX, CmdLenY);
+			MaxLeftWidth = FMath::Max(MaxLeftWidth, CmdLenX);
+			if (!Cmd.Desc.IsEmpty())
+			{
+				float DescLenX, DescLenY;
+				Canvas->StrLen(Font, Cmd.GetRight(), DescLenX, DescLenY);
+				MaxRightWidth = FMath::Max(MaxRightWidth, ConsoleDefs::AutocompleteGap + DescLenX);
+			}
 		}
 
 		// Display a message if there were more matches
-		if (AutoComplete.Num() >= MAX_AUTOCOMPLETION_LINES)
+		if (AutoComplete.Num() > MAX_AUTOCOMPLETION_LINES)
 		{
 			static FAutoCompleteCommand MoreMatchesLine;
 
-			MoreMatchesLine.Desc = FString::Printf(TEXT("[%i more matches]"), (AutoComplete.Num() - MAX_AUTOCOMPLETION_LINES + 1));
+			MoreMatchesLine.Command = FString::Printf(TEXT("[%i more matches]"), (AutoComplete.Num() - MAX_AUTOCOMPLETION_LINES));
+			MoreMatchesLine.Color = ConsoleSettings->AutoCompleteFadedColor;
 			AutoCompleteElements.Add(&MoreMatchesLine);
 		}
 
 		// background rectangle behind auto completion
-		float MaxWidth = 0;
-		{
-			float MaxHeight = 0;
-
-			for(int32 i = 0, Num = AutoCompleteElements.Num(); i < Num; ++i)
-			{
-				const FAutoCompleteCommand& AutoCompleteElement = *AutoCompleteElements[i];
-
-				float Width = 0;
-				{
-					float info_xl, info_yl;
-					Canvas->StrLen(Font, AutoCompleteElement.GetLeft(), info_xl, info_yl);
-					Width += info_xl;
-				}
-				if(!AutoCompleteElement.Desc.IsEmpty())
-				{
-					Width += ConsoleDefs::AutocompleteGap;
-
-					float info_xl, info_yl;
-					Canvas->StrLen(Font, AutoCompleteElement.GetRight(), info_xl, info_yl);
-					Width += info_xl;
-				}
-				MaxWidth = FMath::Max(MaxWidth, Width);
-				MaxHeight += yl;
-			}
-
-			int32 Border = 4;
-
-			// white border
-			ConsoleTile.Size = FVector2D( MaxWidth + 2 * Border, MaxHeight + 2 * Border );
-			ConsoleTile.SetColor( ConsoleDefs::BorderColor );
-			Canvas->DrawItem( ConsoleTile, UserInputLinePos.X + xl - Border, y + yl - MaxHeight - Border );
-
-			--Border;
-
-			// dark inner part
-			ConsoleTile.Size = FVector2D( MaxWidth + 2 * Border, MaxHeight + 2 * Border );
-			ConsoleTile.SetColor( ConsoleDefs::AutocompleteBackgroundColor );
-			Canvas->DrawItem( ConsoleTile, UserInputLinePos.X + xl - Border, y + yl - MaxHeight - Border );
-		}
+		float MaxWidth = MaxLeftWidth + MaxRightWidth;
+		float Height = AutoCompleteElements.Num() * yl;
+		int32 Border = 4;
+
+		// dark inner part
+		ConsoleTile.Size = FVector2D( MaxWidth + 2 * Border, Height + 2 * Border );
+		ConsoleTile.SetColor( AutoCompleteBackgroundColor );
+		ConsoleTile.BlendMode = SE_BLEND_Translucent;
+		Canvas->DrawItem( ConsoleTile, UserInputLinePos.X + xl - Border, y + yl - Height - Border );
+
+		// white border
+		FCanvasBoxItem ConsoleOutline(ConsoleTile.Position, FVector2D(MaxWidth + 2 * Border, Height + 2 * Border));
+		ConsoleOutline.SetColor(ConsoleDefs::BorderColor);
+		ConsoleOutline.BlendMode = SE_BLEND_Opaque;
+		Canvas->DrawItem(ConsoleOutline, UserInputLinePos.X + xl - Border, y + yl - Height - Border);
 
 		// auto completion elements
-		for(int32 i = 0, Num = AutoCompleteElements.Num(); i < Num; ++i)
+		for(int32 Num = AutoCompleteElements.Num(), i = Num - 1; i >= 0; --i)
 		{
 			const FAutoCompleteCommand& AutoCompleteElement = *AutoCompleteElements[i];
 
 			const bool bCursorLineColor = (i == AutoCompleteCursor);
-			const bool bMoreMatches = (Num >= MAX_AUTOCOMPLETION_LINES && i == Num - 1);
+			const bool bMoreMatches = (Num > MAX_AUTOCOMPLETION_LINES && i == Num - 1);
 			const bool bHistory = AutoCompleteElement.IsHistory();
+			int CmdXOffset = 0;
 		
-			FColor LeftC = ConsoleDefs::AutocompleteSuggestionColor;
-			FColor RightC = ConsoleDefs::AutocompletePartialSuggestionColor;
+			FColor LeftC = AutoCompleteElement.Color;
+			FColor RightC = ConsoleSettings->AutoCompleteFadedColor;
 
 			if(bCursorLineColor)
 			{
 				ConsoleTile.Size = FVector2D(MaxWidth, yl);
 				ConsoleTile.SetColor(ConsoleDefs::CursorLineColor);
+				ConsoleTile.BlendMode = SE_BLEND_Opaque;
 				Canvas->DrawItem(ConsoleTile, UserInputLinePos.X + xl, y);
 				LeftC = ConsoleDefs::CursorColor;
 			}
 
-			if(bMoreMatches)
-			{
-				LeftC = RightC = ConsoleDefs::AutocompletePartialSuggestionColor;
-			}
-
 			if(bHistory)
 			{
 				// > HistoryElement has the strings swapped so we need to swap the colors
 				Swap(LeftC, RightC);
 			}
 
+			float CommandWidth, CommandHeight;
+			Canvas->StrLen(Font, AutoCompleteElement.Command, CommandWidth, CommandHeight);
+
+			if (bMoreMatches)
+			{
+				// Center the "x more matches" line, unless that would put it further right than the descriptions
+				CmdXOffset = FMath::Min((MaxWidth / 2) - (CommandWidth / 2), MaxLeftWidth + ConsoleDefs::AutocompleteGap);
+			}
+
 			ConsoleText.SetColor(LeftC);
 			ConsoleText.Text = FText::FromString(AutoCompleteElement.GetLeft());
-			Canvas->DrawItem( ConsoleText, UserInputLinePos.X + xl, y );
-			float info_xl;
+			Canvas->DrawItem( ConsoleText, UserInputLinePos.X + CmdXOffset + xl, y );
+
+			float DescriptionWidth, DescriptionHeight;
+			Canvas->StrLen(Font, AutoCompleteElement.GetRight(), DescriptionWidth, DescriptionHeight);
+			float DescriptionX = UserInputLinePos.X + xl + ConsoleDefs::AutocompleteGap;
+			float DescriptionOverflow = DescriptionX + MaxLeftWidth + DescriptionWidth - Canvas->SizeX;
+
+			if (DescriptionOverflow > 0)
 			{
-				float info_yl;
-				Canvas->StrLen(Font, AutoCompleteElement.GetLeft(), info_xl, info_yl);
+				// Horizontal overflow due to low resolution or an overly long description; forgo justification
+				DescriptionX = FMath::Max(DescriptionX + CommandWidth, DescriptionX + MaxLeftWidth - DescriptionOverflow);
 			}
+			else
+			{
+				DescriptionX += MaxLeftWidth;
+			}
+
 			ConsoleText.SetColor(RightC);
 			ConsoleText.Text = FText::FromString(AutoCompleteElement.GetRight());
-			Canvas->DrawItem( ConsoleText, UserInputLinePos.X + xl + info_xl + ConsoleDefs::AutocompleteGap, y );
+			Canvas->DrawItem( ConsoleText, DescriptionX, y );
 			y -= yl;
 		}
 	}
 
 	// determine the cursor position
-	OutStr = FString::Printf(TEXT("%s%s"), *ConsoleDefs::LeadingInputText, *TypedStr.Left(TypedStrPos));
-	Canvas->StrLen(Font, OutStr,xl,yl);
+	const FString TypedInputTextUpToCursor = FString::Printf(TEXT("%s%s"), *ConsoleDefs::LeadingInputText, *TypedStr.Left(TypedStrPos));
+	Canvas->StrLen(Font, TypedInputTextUpToCursor,xl,yl);
 
 	// draw the cursor
 	ConsoleText.SetColor( ConsoleDefs::CursorColor );
-	OutStr = FString::Printf( TEXT("_") );
-	ConsoleText.Text = FText::FromString( OutStr );
+	ConsoleText.Text = FText::FromString( FString(TEXT("_")) );
 	Canvas->DrawItem( ConsoleText, UserInputLinePos.X + xl, UserInputLinePos.Y-1.0f-yl );
 
 }
diff --git a/Engine/Source/Runtime/Engine/Private/Vehicles/WheeledVehicleMovementComponent4W.cpp b/Engine/Source/Runtime/Engine/Private/Vehicles/WheeledVehicleMovementComponent4W.cpp
--- a/Engine/Source/Runtime/Engine/Private/Vehicles/WheeledVehicleMovementComponent4W.cpp
+++ b/Engine/Source/Runtime/Engine/Private/Vehicles/WheeledVehicleMovementComponent4W.cpp
@@ -198,8 +198,8 @@ static void GetVehicleAutoBoxSetup(const FVehicleTransmissionData& Setup, PxVehi
 	for (int32 i = 0; i < Setup.ForwardGears.Num(); i++)
 	{
 		const FVehicleGearData& GearData = Setup.ForwardGears[i];
-		PxSetup.mUpRatios[i] = GearData.UpRatio;
-		PxSetup.mDownRatios[i] = GearData.DownRatio;
+		PxSetup.mUpRatios[i + PxVehicleGearsData::eFIRST] = GearData.UpRatio;
+		PxSetup.mDownRatios[i + PxVehicleGearsData::eFIRST] = GearData.DownRatio;
 	}
 	PxSetup.mUpRatios[PxVehicleGearsData::eNEUTRAL] = Setup.NeutralGearUpRatio;
 	PxSetup.setLatency(Setup.GearAutoBoxLatency);
diff --git a/Engine/Source/Runtime/Engine/Private/World.cpp b/Engine/Source/Runtime/Engine/Private/World.cpp
--- a/Engine/Source/Runtime/Engine/Private/World.cpp
+++ b/Engine/Source/Runtime/Engine/Private/World.cpp
@@ -44,6 +44,7 @@
 #include "PhysicsEngine/PhysicsConstraintComponent.h"
 
 #include "Materials/MaterialParameterCollectionInstance.h"
+#include "LoadTimeTracker.h"
 
 #if WITH_EDITOR
 	#include "DerivedDataCacheInterface.h"
@@ -1496,17 +1497,7 @@ void UWorld::EnsureCollisionTreeIsBuilt()
 	// Set physics to static loading mode
 	if (PhysicsScene)
 	{
-		PhysicsScene->SetIsStaticLoading(true);
-
-		for (int Iteration = 0; Iteration < 6; ++Iteration)
-		{
-			SetupPhysicsTickFunctions(0.1f);
-			PhysicsScene->StartFrame();
-			PhysicsScene->WaitPhysScenes();
-			PhysicsScene->EndFrame(NULL);
-		}
-
-		PhysicsScene->SetIsStaticLoading(false);
+		PhysicsScene->EnsureCollisionTreeIsBuilt(this);
 	}
 
     bIsBuilt = true;
@@ -4493,6 +4484,7 @@ bool FSeamlessTravelHandler::StartTravel(UWorld* InCurrentWorld, const FURL& InU
 	}
 	else
 	{
+		FLoadTimeTracker::Get().ResetRawLoadTimes();
 		UE_LOG(LogWorld, Log, TEXT("SeamlessTravel to: %s"), *InURL.Map);
 		FString MapName = UWorld::RemovePIEPrefix(InURL.Map);
 		if (!FPackageName::DoesPackageExist(MapName, InGuid.IsValid() ? &InGuid : NULL))
@@ -5078,6 +5070,7 @@ UWorld* FSeamlessTravelHandler::Tick()
 				
 				double TotalSeamlessTravelTime = FPlatformTime::Seconds() - SeamlessTravelStartTime;
 				UE_LOG(LogWorld, Log, TEXT("----SeamlessTravel finished in %.2f seconds ------"), TotalSeamlessTravelTime );
+				FLoadTimeTracker::Get().DumpRawLoadTimes();
 
 				AGameMode* const GameMode = LoadedWorld->GetAuthGameMode();
 				if (GameMode)
diff --git a/Engine/Source/Runtime/Engine/Public/Animation/AnimPhysicsSolver.h b/Engine/Source/Runtime/Engine/Public/Animation/AnimPhysicsSolver.h
--- a/Engine/Source/Runtime/Engine/Public/Animation/AnimPhysicsSolver.h
+++ b/Engine/Source/Runtime/Engine/Public/Animation/AnimPhysicsSolver.h
@@ -460,7 +460,7 @@ public:
 	*  @param NumPreIterations Number of times to iterate the limits before performing the integration
 	*  @param NumPostIterations Number of times to iterae the limits after performing the integration
 	*/
-	static void PhysicsUpdate(float DeltaTime, TArray<FAnimPhysRigidBody*>& Bodies, TArray<FAnimPhysLinearLimit>& LinearLimits, TArray<FAnimPhysAngularLimit>& AngularLimits, TArray<FAnimPhysSpring>& Springs, int32 NumPreIterations = 8, int32 NumPostIterations = 2);
+	static void PhysicsUpdate(float DeltaTime, TArray<FAnimPhysRigidBody*>& Bodies, TArray<FAnimPhysLinearLimit>& LinearLimits, TArray<FAnimPhysAngularLimit>& AngularLimits, TArray<FAnimPhysSpring>& Springs, const FVector& GravityDirection, const FVector& ExternalForce, int32 NumPreIterations = 8, int32 NumPostIterations = 2);
 
 	//////////////////////////////////////////////////////////////////////////
 	// Constraint functions
@@ -552,7 +552,7 @@ private:
 	/** Initialize the velocity for a given body
 	 *  @param InBody Body to initialize
 	 */
-	static void InitializeBodyVelocity(float DeltaTime, FAnimPhysRigidBody *InBody);
+	static void InitializeBodyVelocity(float DeltaTime, FAnimPhysRigidBody *InBody, const FVector& GravityDirection);
 
 	/** Using calculated linear and angular momentum, integrate the position and orientation of a body
 	 *  @param InBody Body to integrate
diff --git a/Engine/Source/Runtime/Engine/Public/AudioDevice.h b/Engine/Source/Runtime/Engine/Public/AudioDevice.h
--- a/Engine/Source/Runtime/Engine/Public/AudioDevice.h
+++ b/Engine/Source/Runtime/Engine/Public/AudioDevice.h
@@ -435,6 +435,11 @@ public:
 	 */
 	void SetReverbSettings(AAudioVolume* Volume, const FReverbSettings& ReverbSettings);
 
+	/** 
+	 * Returns the currently applied reverb effect if there is one.
+	 */
+	class UReverbEffect* GetCurrentReverbEffect();
+
 	/**
 	 * Creates an audio component to handle playing a sound.
 	 * Plays a sound at the given location without creating an audio component.
diff --git a/Engine/Source/Runtime/Engine/Public/ComponentReregisterContext.h b/Engine/Source/Runtime/Engine/Public/ComponentReregisterContext.h
--- a/Engine/Source/Runtime/Engine/Public/ComponentReregisterContext.h
+++ b/Engine/Source/Runtime/Engine/Public/ComponentReregisterContext.h
@@ -186,15 +186,6 @@ public:
 	*/
 	FGlobalComponentReregisterContext(const TArray<UClass*>& ExcludeComponents);
 
-	/**
-	 * Initialization constructor
-	 * Only reregister those components whose replacement primitive is in a direct child of one of the InParentActors
-	 * 
-	 * @param InParentActors - list of actors called out for reregistering
-	 */
-	ENGINE_API FGlobalComponentReregisterContext(const TArray<AActor*>& InParentActors);
-
-
 	/** Destructor */
 	ENGINE_API ~FGlobalComponentReregisterContext();
 
diff --git a/Engine/Source/Runtime/Engine/Public/PhysicsPublic.h b/Engine/Source/Runtime/Engine/Public/PhysicsPublic.h
--- a/Engine/Source/Runtime/Engine/Public/PhysicsPublic.h
+++ b/Engine/Source/Runtime/Engine/Public/PhysicsPublic.h
@@ -405,8 +405,8 @@ public:
 		return PhysicsSceneCompletion;
 	}
 
-	/** Set rather we're doing a static load and want to stall, or are during gameplay and want to distribute over many frames */
-	ENGINE_API void SetIsStaticLoading(bool bStaticLoading);
+	/** Ensures that the collision tree is built. */
+	ENGINE_API void EnsureCollisionTreeIsBuilt(UWorld* World);
 
 	/** Waits for all physics scenes to complete */
 	ENGINE_API void WaitPhysScenes();
@@ -566,6 +566,9 @@ private:
 	/** Task created from TickPhysScene so we can substep without blocking */
 	bool SubstepSimulation(uint32 SceneType, FGraphEventRef& InOutCompletionEvent);
 
+	/** Set whether we're doing a static load and want to stall, or are during gameplay and want to distribute over many frames */
+	void SetIsStaticLoading(bool bStaticLoading);
+
 #if WITH_PHYSX
 	/** User data wrapper passed to physx */
 	struct FPhysxUserData PhysxUserData;
diff --git a/Engine/Source/Runtime/Engine/Public/TimerManager.h b/Engine/Source/Runtime/Engine/Public/TimerManager.h
--- a/Engine/Source/Runtime/Engine/Public/TimerManager.h
+++ b/Engine/Source/Runtime/Engine/Public/TimerManager.h
@@ -156,6 +156,7 @@ public:
 	FTimerManager()
 		: InternalTime(0.0)
 		, LastTickedFrame(static_cast<uint64>(-1))
+		, LastAssignedHandle(0)
 	{}
 
 
@@ -237,15 +238,28 @@ public:
 	}
 
 	/**
-	 * Clears a previously set timer, identical to calling SetTimer() with a <= 0.f rate.
+	 * DEPRECATED: Clears a previously set timer, identical to calling SetTimer() with a <= 0.f rate.
 	 *
 	 * @param InHandle The handle of the timer to clear.
 	 */
-	FORCEINLINE void ClearTimer(FTimerHandle InHandle)
+	DEPRECATED(4.12, "This function is deprecated to ensure that timers that are no longer valid are not persisted. Please call this function with a non-const reference.")
+	FORCEINLINE void ClearTimer(const FTimerHandle& InHandle)
 	{
 		InternalClearTimer(InHandle);
 	}
 
+	/**
+	* Clears a previously set timer, identical to calling SetTimer() with a <= 0.f rate.
+	* Invalidates the timer handle as it should no longer be used.
+	*
+	* @param InHandle The handle of the timer to clear.
+	*/
+	FORCEINLINE void ClearTimer(FTimerHandle& InHandle)
+	{
+		InternalClearTimer(InHandle);
+		InHandle.Invalidate();
+	}
+
 	/** Clears all timers that are bound to functions on the given object. */
 	FORCEINLINE void ClearAllTimersForObject(void const* Object)
 	{
@@ -255,7 +269,6 @@ public:
 		}
 	}
 
-
 	/**
 	 * Pauses a previously set timer.
 	 *
@@ -377,9 +390,12 @@ public:
 	 */
 	FTimerHandle K2_FindDynamicTimerHandle(FTimerDynamicDelegate InDynamicDelegate) const;
 
-	/** Debug command to output info on all timers curently set to the log. */
+	/** Debug command to output info on all timers currently set to the log. */
 	void ListTimers() const;
 
+	/** Get the current last assigned handle */
+	void ValidateHandle(FTimerHandle& InOutHandle);
+
 private:
 	void InternalSetTimer( FTimerHandle& InOutHandle, FTimerUnifiedDelegate const& InDelegate, float InRate, bool InbLoop, float InFirstDelay );
 	void InternalSetTimer( FTimerData& NewTimerData, float InRate, bool InbLoop, float InFirstDelay );
@@ -416,5 +432,8 @@ private:
 
 	/** Set this to GFrameCounter when Timer is ticked. To figure out if Timer has been already ticked or not this frame. */
 	uint64 LastTickedFrame;
+
+	/** The last handle we assigned from this timer manager */
+	uint64 LastAssignedHandle;
 };
 
diff --git a/Engine/Source/Runtime/EngineSettings/Classes/ConsoleSettings.h b/Engine/Source/Runtime/EngineSettings/Classes/ConsoleSettings.h
--- a/Engine/Source/Runtime/EngineSettings/Classes/ConsoleSettings.h
+++ b/Engine/Source/Runtime/EngineSettings/Classes/ConsoleSettings.h
@@ -19,10 +19,15 @@ struct FAutoCompleteCommand
 	UPROPERTY(config, EditAnywhere, Category=Command)
 	FString Desc;
 
+	FColor Color;
+
+	FAutoCompleteCommand()
+		: Color(180, 180, 180)
+	{}
+
 	bool operator<(const FAutoCompleteCommand& rhs) const
 	{
-		// sort them in opposite order for game console UI rendering (bottom up)
-		return Command >= rhs.Command;
+		return Command < rhs.Command;
 	}
 
 	// for game console 
@@ -69,4 +74,28 @@ class ENGINESETTINGS_API UConsoleSettings
 	/** List of relative paths (e.g. Content/Maps) to search for map names for auto-complete usage. Specified in BaseInput.ini. */
 	UPROPERTY(config, EditAnywhere, Category=AutoComplete)
 	TArray<FString> AutoCompleteMapPaths;
+
+	/** Amount of transparency of the console background. */
+	UPROPERTY(config, EditAnywhere, Category=Colors, meta=(UIMin="0", UIMax="100", ClampMin="0", ClampMax="100"))
+	float BackgroundOpacityPercentage;
+
+	/** The color used for text input. */
+	UPROPERTY(config, EditAnywhere, Category=Colors)
+	FColor InputColor;
+
+	/** The color used for the previously typed commands history. */
+	UPROPERTY(config, EditAnywhere, Category=Colors)
+	FColor HistoryColor;
+
+	/** The autocomplete color used for executable commands. */
+	UPROPERTY(config, EditAnywhere, Category=Colors)
+	FColor AutoCompleteCommandColor;
+
+	/** The autocomplete color used for mutable CVars. */
+	UPROPERTY(config, EditAnywhere, Category=Colors)
+	FColor AutoCompleteCVarColor;
+
+	/** The autocomplete color used for command descriptions and read-only CVars. */
+	UPROPERTY(config, EditAnywhere, Category=Colors)
+	FColor AutoCompleteFadedColor;
 };
