Copying //UE4/Dev-Rendering to Dev-Main (Source //UE4/Dev-Rendering@2932636)

diff --git a/Engine/Source/Runtime/Engine/Classes/Engine/TextureStreamingTypes.h b/Engine/Source/Runtime/Engine/Classes/Engine/TextureStreamingTypes.h
new file mode 100644
--- /dev/null
+++ b/Engine/Source/Runtime/Engine/Classes/Engine/TextureStreamingTypes.h
@@ -0,0 +1,228 @@
+// Copyright 1998-2016 Epic Games, Inc. All Rights Reserved.
+
+// Structs and defines used for texture streaming build
+
+#pragma once 
+
+#include "EngineTypes.h"
+#include "TextureStreamingTypes.generated.h"
+
+class UTexture;
+class UTexture2D;
+struct FStreamingTextureBuildInfo;
+
+/** Information about a streaming texture that a primitive uses for rendering. */
+USTRUCT()
+struct FStreamingTexturePrimitiveInfo
+{
+	GENERATED_USTRUCT_BODY()
+
+	UPROPERTY()
+	UTexture* Texture;
+
+	UPROPERTY()
+	FBoxSphereBounds Bounds;
+
+	UPROPERTY()
+	float TexelFactor;
+
+	FStreamingTexturePrimitiveInfo() : 
+		Texture(nullptr), 
+		Bounds(ForceInit), 
+		TexelFactor(1.0f) 
+	{
+	}
+
+
+	/**
+	 *	Set this struct to match the packed params.
+	 *
+	 *	@param	InTexture			The texture as refered by the packed info.
+	 *	@param	RefBounds			The reference bounds used to unpack the relative box.
+	 *	@param	Info				The packed params.
+	 *	@param	bUseRelativeBox		true if the relative box is relevant. Could be irrelevant if a level transform was applied after the texture streaming build.
+	 */
+	ENGINE_API void UnPackFrom(UTexture2D* InTexture, const FBoxSphereBounds& RefBounds, const FStreamingTextureBuildInfo& Info, bool bUseRelativeBox);
+};
+
+/** 
+ * This struct holds the result of TextureStreaming Build for each component texture, as referred by its used materials.
+ * It is possible that the entry referred by this data is not actually relevant in a given quality / target.
+ * It is also possible that some texture are not referred, and will then fall on fallbacks computation.
+ * Because each component holds its precomputed data for each texture, this struct is designed to be as compact as possible.
+ */
+USTRUCT()
+struct FStreamingTextureBuildInfo
+{
+	GENERATED_USTRUCT_BODY()
+
+	/** 
+	 * The relative bounding box for this entry. The relative bounds is a bound equal or smaller than the component bounds and represent
+	 * the merged LOD section bounds of all LOD section referencing the given texture. When the level transform is modified following
+	 * a call to ApplyLevelTransform, this relative bound becomes deprecated as it was computed from the transform at build time.
+	 */
+	UPROPERTY()
+	uint32 PackedRelativeBox;
+
+	/** 
+	 * The level scope identifier of the texture. When building the texture streaming data, each level holds a list of all referred texture Guids.
+	 * This is required to prevent loading textures on platforms which would not require the texture to be loaded, and is a consequence of the texture
+	 * streaming build not being platform specific (the same streaming data is build for every platform target). Could also apply to quality level.
+	 */
+	UPROPERTY()
+	int32 TextureLevelIndex;
+
+	/** 
+	 * The texel factor for this texture. This represent the world size a texture square holding with unit UVs.
+	 * This value is a combination of the TexelFactor from the mesh and also the material scale.
+	 * It does not take into consideration StreamingDistanceMultiplier, or texture group scale.
+	 */
+	UPROPERTY()
+	float TexelFactor;
+
+	FStreamingTextureBuildInfo() : 
+		PackedRelativeBox(0), 
+		TextureLevelIndex(0), 
+		TexelFactor(0) 
+	{
+	}
+
+	/**
+	 *	Set this struct to match the unpacked params.
+	 *
+	 *	@param	LevelTextures	[in,out]	The list of textures referred by all component of a level. The array index maps to UTexture2D::LevelIndex.
+	 *	@param	RefBounds		[in]		The reference bounds used to compute the packed relative box.
+	 *	@param	Info			[in]		The unpacked params.
+	 */
+	ENGINE_API void PackFrom(TArray<UTexture2D*>& LevelTextures, const FBoxSphereBounds& RefBounds, const FStreamingTexturePrimitiveInfo& Info);
+
+};
+
+/** 
+ * This struct holds data about the texture coordinate used to sample a texture within a material.
+ */
+struct FMaterialTexCoordBuildInfo
+{
+	enum 
+	{
+		MAX_NUM_TEX_COORD = 4,			// The max number of texcoords processed in the analysis.
+		INITIAL_GPU_SCALE = 256,		// The initial texcoord scale (must be bigger than actual used values)
+		TILE_RESOLUTION = 32,			// The tile size of the render used when analysis the shader scales.
+		MAX_NUM_TEXTURE_REGISTER = 32	// The max number of texture indices processed in the analysis.
+	};
+
+	FMaterialTexCoordBuildInfo() : Scale(0), Index(INDEX_NONE) {}
+
+	/** The coordinate scale */
+	float Scale;
+
+	/** The coordinate index */
+	int32 Index;
+};
+
+/** 
+ * This struct holds the transient data used in the texture streaming build and is also useful for debugging results.
+ */
+struct FStreamingSectionBuildInfo
+{
+	FStreamingSectionBuildInfo() : LODIndex(INDEX_NONE), ElementIndex(INDEX_NONE), MaterialIndex(INDEX_NONE), BoxOrigin(0, 0, 0), BoxExtent(0, 0, 0)  
+	{
+		FMemory::Memzero(TexelFactors);
+	}
+
+	/** LOD Index related to this data. */
+	int32 LODIndex;
+
+	/** Element (Section) Index related to this data. */
+	int32 ElementIndex;
+
+	/** The material index used for this LOD Section. */
+	int32 MaterialIndex;
+
+	/** The world space box origin of the lod section.*/
+	FVector	BoxOrigin;
+
+	/** The world space box extent of the lod section.*/
+	FVector BoxExtent;
+
+	/** The texture world size associated to the mesh UV mapping.*/
+	TArray<FMaterialTexCoordBuildInfo> TexCoordData;
+
+	/** The texture world size computed from the mesh UV mapping.*/
+	float TexelFactors[FMaterialTexCoordBuildInfo::MAX_NUM_TEX_COORD];
+};
+
+/** 
+ * Context used to resolve FStreamingTextureBuildInfo to FStreamingTexturePrimitiveInfo
+ * The context make sure that build data and each texture is only processed once per component (with constant time).
+ * It manage internally structures used to accelerate the binding between precomputed data and textures,
+ * so that there is only one map lookup per texture per level. 
+ * There is some complexity here because the build data does not reference directly texture objects to avoid hard references
+ * which would load texture when the component is loaded, which could be wrong since the build data is built for a specific
+ * feature level and quality level. The current feature and quality could reference more or less textures.
+ * This requires the logic to not submit a streaming entry for precomputed data, as well as submit fallback data for 
+ * texture that were referenced in the texture streaming build.
+ */
+class FStreamingTextureLevelContext
+{
+	/** Reversed lookup for ULevel::StreamingTextureGuids. */
+	TMap<FGuid, int32> TextureGuidToLevelIndex;
+
+	/** List of texture for which LevelIndex was set. Used to reset the value to INDEX_NONE after the level was processed. */
+	TArray<UTexture2D*> ProcessedTextures;
+
+	/*
+	 * Whether the precomputed relative bounds should be used or not. 
+	 * Will be false if the transform level was rotated  since the last texture streaming build.
+	 */
+	bool bUseRelativeBoxes;
+
+	/** An id used to identify the component. */
+	int32 ComponentTimestamp;
+
+	/** The last bound component texture streaming build data. */
+	const TArray<FStreamingTextureBuildInfo>* ComponentBuildData;
+	/** The last bound component bounds. */
+	FBoxSphereBounds ComponentBounds;
+	/** The last bound component mesh scale. */
+	float ComponentMeshScale;
+	/** The last bound component streaming scale. */
+	float ComponentStreamingScale;
+
+	struct FTextureBoundState
+	{
+		/** The component timestamp that last component referring this texture. */
+		int32 Timestamp;
+		/** The component timestamp that last component having build data for this texture. */
+		int32 BuildDataTimestamp;
+		/** The ComponentBuildData Index referring this texture. */
+		int32 BuildDataIndex;
+	};
+
+	/*
+	 * The component state of the each texture. Used to prevent processing each texture several time.
+	 * Also used to find quickly the build data relating to each texture. 
+	 */
+	TArray<FTextureBoundState> BoundStates;
+
+public:
+
+	FStreamingTextureLevelContext(const ULevel* InLevel = nullptr);
+	~FStreamingTextureLevelContext();
+
+	void BindComponent(const TArray<FStreamingTextureBuildInfo>* BuildData, const FBoxSphereBounds& Bounds, float MeshScale, float StreamingScale);
+	void Process(const TArray<UTexture*>& InTextures, TArray<FStreamingTexturePrimitiveInfo>& OutInfos);
+};
+
+
+/** A Map that gives the (smallest) texture coordinate scale used when sampling each texture register of a shader.
+ * The array index is the register index, and the value, is the coordinate scale.
+ * Since a texture resource can be bound to several texture registers, it can related to different indices.
+ * This is reflected in UMaterialInterface::GetUsedTexturesAndIndices where each texture is bound to 
+ * an array of texture register indices.
+ */
+typedef TMap<UMaterialInterface*, TArray<FMaterialTexCoordBuildInfo> > FTexCoordScaleMap;
+
+ENGINE_API void BuildTextureStreamingShaders(UWorld* InWorld, EMaterialQualityLevel::Type QualityLevel, ERHIFeatureLevel::Type FeatureLevel, OUT FTexCoordScaleMap& TexCoordScales);
+ENGINE_API void BuildTextureStreamingData(UWorld* InWorld, const FTexCoordScaleMap& InTexCoordScales, EMaterialQualityLevel::Type QualityLevel, ERHIFeatureLevel::Type FeatureLevel);
+extern ENGINE_API TAutoConsoleVariable<int32> CVarStreamingUseNewMetrics;
\ No newline at end of file
diff --git a/Engine/Source/Runtime/Engine/Classes/GameFramework/WorldSettings.h b/Engine/Source/Runtime/Engine/Classes/GameFramework/WorldSettings.h
--- a/Engine/Source/Runtime/Engine/Classes/GameFramework/WorldSettings.h
+++ b/Engine/Source/Runtime/Engine/Classes/GameFramework/WorldSettings.h
@@ -375,6 +375,10 @@ class ENGINE_API AWorldSettings : public AInfo, public IInterface_AssetUserData
 	UPROPERTY(EditAnywhere, BlueprintReadOnly, Category=World, AdvancedDisplay)
 	FVector DefaultColorScale;
 
+	/** Max occlusion distance used by mesh distance fields, overridden if there is a movable skylight. */
+	UPROPERTY(EditAnywhere, Category=World, meta=(DisplayName = "Default Max DistanceField Occlusion Distance"))
+	float DefaultMaxDistanceFieldOcclusionDistance;
+
 	/************************************/
 	
 	/** PRECOMPUTED VISIBILITY SETTINGS **/
diff --git a/Engine/Source/Runtime/Engine/Classes/Materials/Material.h b/Engine/Source/Runtime/Engine/Classes/Materials/Material.h
--- a/Engine/Source/Runtime/Engine/Classes/Materials/Material.h
+++ b/Engine/Source/Runtime/Engine/Classes/Materials/Material.h
@@ -649,10 +649,6 @@ public:
 	UPROPERTY(EditAnywhere, BlueprintReadOnly, Category = Mobile, meta = (DisplayName = "High Quality Reflections"))
 	uint32 bUseHQForwardReflections : 1;
 
-	/** Whether the material should composite a planar reflection pass. */
-	UPROPERTY(EditAnywhere, Category = Mobile)
-	uint32 bAcceptsPlanarReflection : 1;
-
 	/** The type of tessellation to apply to this object.  Note D3D11 required for anything except MTM_NoTessellation. */
 	UPROPERTY(EditAnywhere, BlueprintReadOnly, Category=Tessellation)
 	TEnumAsByte<enum EMaterialTessellationMode> D3D11TessellationMode;
@@ -816,6 +812,7 @@ public:
 	ENGINE_API virtual FMaterialRenderProxy* GetRenderProxy(bool Selected, bool bHovered=false) const override;
 	ENGINE_API virtual UPhysicalMaterial* GetPhysicalMaterial() const override;
 	ENGINE_API virtual void GetUsedTextures(TArray<UTexture*>& OutTextures, EMaterialQualityLevel::Type QualityLevel, bool bAllQualityLevels, ERHIFeatureLevel::Type FeatureLevel, bool bAllFeatureLevels) const override;
+	ENGINE_API virtual void GetUsedTexturesAndIndices(TArray<UTexture*>& OutTextures, TArray< TArray<int32> >& OutIndices, EMaterialQualityLevel::Type QualityLevel, ERHIFeatureLevel::Type FeatureLevel) const override;
 	ENGINE_API virtual void OverrideTexture( const UTexture* InTextureToOverride, UTexture* OverrideTexture, ERHIFeatureLevel::Type InFeatureLevel ) override;
 	ENGINE_API virtual void OverrideVectorParameterDefault(FName ParameterName, const FLinearColor& Value, bool bOverride, ERHIFeatureLevel::Type FeatureLevel) override;
 	ENGINE_API virtual void OverrideScalarParameterDefault(FName ParameterName, float Value, bool bOverride, ERHIFeatureLevel::Type FeatureLevel) override;
diff --git a/Engine/Source/Runtime/Engine/Classes/Materials/MaterialInstance.h b/Engine/Source/Runtime/Engine/Classes/Materials/MaterialInstance.h
--- a/Engine/Source/Runtime/Engine/Classes/Materials/MaterialInstance.h
+++ b/Engine/Source/Runtime/Engine/Classes/Materials/MaterialInstance.h
@@ -259,6 +259,7 @@ public:
 	virtual ENGINE_API bool GetTextureParameterOverrideValue(FName ParameterName, class UTexture*& OutValue) const override;
 	virtual ENGINE_API bool GetVectorParameterValue(FName ParameterName, FLinearColor& OutValue) const override;
 	virtual ENGINE_API void GetUsedTextures(TArray<UTexture*>& OutTextures, EMaterialQualityLevel::Type QualityLevel, bool bAllQualityLevels, ERHIFeatureLevel::Type FeatureLevel, bool bAllFeatureLevels) const override;
+	virtual ENGINE_API void GetUsedTexturesAndIndices(TArray<UTexture*>& OutTextures, TArray< TArray<int32> >& OutIndices, EMaterialQualityLevel::Type QualityLevel, ERHIFeatureLevel::Type FeatureLevel) const;
 	virtual ENGINE_API void OverrideTexture(const UTexture* InTextureToOverride, UTexture* OverrideTexture, ERHIFeatureLevel::Type InFeatureLevel) override;
 	virtual ENGINE_API void OverrideVectorParameterDefault(FName ParameterName, const FLinearColor& Value, bool bOverride, ERHIFeatureLevel::Type FeatureLevel) override;
 	virtual ENGINE_API void OverrideScalarParameterDefault(FName ParameterName, float Value, bool bOverride, ERHIFeatureLevel::Type FeatureLevel) override;
@@ -394,7 +395,7 @@ protected:
 
 	ENGINE_API void SetParentInternal(class UMaterialInterface* NewParent, bool RecacheShaders);
 
-	void GetTextureExpressionValues(const FMaterialResource* MaterialResource, TArray<UTexture*>& OutTextures) const;
+	void GetTextureExpressionValues(const FMaterialResource* MaterialResource, TArray<UTexture*>& OutTextures, TArray< TArray<int32> >* OutIndices = nullptr) const;
 
 	/** 
 	 * Updates StaticPermutationMaterialResources based on the value of bHasStaticPermutationResource
diff --git a/Engine/Source/Runtime/Engine/Classes/Materials/MaterialInterface.h b/Engine/Source/Runtime/Engine/Classes/Materials/MaterialInterface.h
--- a/Engine/Source/Runtime/Engine/Classes/Materials/MaterialInterface.h
+++ b/Engine/Source/Runtime/Engine/Classes/Materials/MaterialInterface.h
@@ -274,6 +274,13 @@ public:
 	virtual void GetUsedTextures(TArray<UTexture*>& OutTextures, EMaterialQualityLevel::Type QualityLevel, bool bAllQualityLevels, ERHIFeatureLevel::Type FeatureLevel, bool bAllFeatureLevels) const
 		PURE_VIRTUAL(UMaterialInterface::GetUsedTextures,);
 
+	/** 
+	* Return the textures used to render this material and the material indices bound to each. 
+	* Because material indices can change for each shader, this is limited to a single platform and quality level.
+	* An empty array in OutIndices means the index is undefined.
+	*/
+	ENGINE_API virtual void GetUsedTexturesAndIndices(TArray<UTexture*>& OutTextures, TArray< TArray<int32> >& OutIndices, EMaterialQualityLevel::Type QualityLevel, ERHIFeatureLevel::Type FeatureLevel) const;
+
 	/**
 	 * Override a specific texture (transient)
 	 *
diff --git a/Engine/Source/Runtime/Engine/Classes/Particles/ParticleSystemComponent.h b/Engine/Source/Runtime/Engine/Classes/Particles/ParticleSystemComponent.h
--- a/Engine/Source/Runtime/Engine/Classes/Particles/ParticleSystemComponent.h
+++ b/Engine/Source/Runtime/Engine/Classes/Particles/ParticleSystemComponent.h
@@ -1085,19 +1085,20 @@ protected:
 	 * @param	EmitterInstance		Emitter instance this replay is playing on
 	 * @param	EmitterReplayData	Incoming replay data of any time, cannot be NULL
 	 * @param	bSelected			true if the particle system is currently selected
+	 * @param	InFeatureLevel		The relevant shader feature level
 	 *
 	 * @return	The newly created dynamic data, or NULL on failure
 	 */
-	static FDynamicEmitterDataBase* CreateDynamicDataFromReplay( FParticleEmitterInstance* EmitterInstance, const FDynamicEmitterReplayDataBase* EmitterReplayData, bool bSelected );
+	static FDynamicEmitterDataBase* CreateDynamicDataFromReplay( FParticleEmitterInstance* EmitterInstance, const FDynamicEmitterReplayDataBase* EmitterReplayData, bool bSelected, ERHIFeatureLevel::Type InFeatureLevel );
 
 	/**
 	 * Creates dynamic particle data for rendering the particle system this frame.  This function
 	 * handle creation of dynamic data for regularly simulated particles, but also handles capture
 	 * and playback of particle replay data.
-	 *
+     * @param InFeatureLevel - The relevant shader feature level.
 	 * @return	Returns the dynamic data to render this frame
 	 */
-	FParticleDynamicData* CreateDynamicData();
+	FParticleDynamicData* CreateDynamicData(ERHIFeatureLevel::Type InFeatureLevel);
 
 	/** Orients the Z axis of the ParticleSystemComponent toward the camera while preserving the X axis direction */
 	void OrientZAxisTowardCamera();
diff --git a/Engine/Source/Runtime/Engine/Private/Components/ModelComponent.cpp b/Engine/Source/Runtime/Engine/Private/Components/ModelComponent.cpp
--- a/Engine/Source/Runtime/Engine/Private/Components/ModelComponent.cpp
+++ b/Engine/Source/Runtime/Engine/Private/Components/ModelComponent.cpp
@@ -316,7 +316,7 @@ void UModelComponent::SelectAllSurfaces()
 }
 #endif // WITH_EDITOR
 
-void UModelComponent::GetStreamingTextureInfo(TArray<FStreamingTexturePrimitiveInfo>& OutStreamingTextures) const
+void UModelComponent::GetStreamingTextureInfo(FStreamingTextureLevelContext& LevelContext, TArray<FStreamingTexturePrimitiveInfo>& OutStreamingTextures) const
 {
 	if( Model )
 	{
diff --git a/Engine/Source/Runtime/Engine/Private/Components/PrimitiveComponent.cpp b/Engine/Source/Runtime/Engine/Private/Components/PrimitiveComponent.cpp
--- a/Engine/Source/Runtime/Engine/Private/Components/PrimitiveComponent.cpp
+++ b/Engine/Source/Runtime/Engine/Private/Components/PrimitiveComponent.cpp
@@ -73,7 +73,6 @@ static int32 bEnableFastOverlapCheck = 1;
 static FAutoConsoleVariableRef CVarEnableFastOverlapCheck(TEXT("p.EnableFastOverlapCheck"), bEnableFastOverlapCheck, TEXT("Enable fast overlap check against sweep hits, avoiding UpdateOverlaps (for the swept component)."));
 DECLARE_CYCLE_STAT(TEXT("MoveComponent FastOverlap"), STAT_MoveComponent_FastOverlap, STATGROUP_Game);
 
-
 // Predicate to determine if an overlap is with a certain AActor.
 struct FPredicateOverlapHasSameActor
 {
@@ -242,9 +241,9 @@ bool UPrimitiveComponent::HasStaticLighting() const
 	return ((Mobility == EComponentMobility::Static) || bLightAsIfStatic) && SupportsStaticLighting();
 }
 
-void UPrimitiveComponent::GetStreamingTextureInfoWithNULLRemoval(TArray<struct FStreamingTexturePrimitiveInfo>& OutStreamingTextures) const
+void UPrimitiveComponent::GetStreamingTextureInfoWithNULLRemoval(FStreamingTextureLevelContext& LevelContext, TArray<struct FStreamingTexturePrimitiveInfo>& OutStreamingTextures) const
 {
-	GetStreamingTextureInfo(OutStreamingTextures);
+	GetStreamingTextureInfo(LevelContext, OutStreamingTextures);
 	for (int32 Index = 0; Index < OutStreamingTextures.Num(); Index++)
 	{
 		if (!OutStreamingTextures[Index].Texture || !Cast<UTexture2D>(OutStreamingTextures[Index].Texture))
diff --git a/Engine/Source/Runtime/Engine/Private/Components/SceneCaptureComponent.cpp b/Engine/Source/Runtime/Engine/Private/Components/SceneCaptureComponent.cpp
--- a/Engine/Source/Runtime/Engine/Private/Components/SceneCaptureComponent.cpp
+++ b/Engine/Source/Runtime/Engine/Private/Components/SceneCaptureComponent.cpp
@@ -12,6 +12,10 @@
 #include "Engine/SceneCaptureCube.h"
 #include "Components/SceneCaptureComponentCube.h"
 #include "Components/DrawFrustumComponent.h"
+#include "Engine/PlanarReflection.h"
+#include "Components/PlanarReflectionComponent.h"
+#include "PlanarReflectionSceneProxy.h"
+#include "Components/BoxComponent.h"
 
 ASceneCapture::ASceneCapture(const FObjectInitializer& ObjectInitializer)
 	: Super(ObjectInitializer)
@@ -269,7 +273,6 @@ bool USceneCaptureComponent::GetSettingForShowFlag(FString FlagName, FEngineShow
 
 USceneCaptureComponent2D::USceneCaptureComponent2D(const FObjectInitializer& ObjectInitializer)
 	: Super(ObjectInitializer)
-	, ReflectionPlaneNormal(0.0f, 0.0f, 1.0f) // World space up vector
 {
 	FOVAngle = 90.0f;
 	bAutoActivate = true;
@@ -369,6 +372,164 @@ void USceneCaptureComponent2D::Serialize(FArchive& Ar)
 
 // -----------------------------------------------
 
+APlanarReflection::APlanarReflection(const FObjectInitializer& ObjectInitializer)
+	: Super(ObjectInitializer)
+{
+	PlanarReflectionComponent = CreateDefaultSubobject<UPlanarReflectionComponent>(TEXT("NewPlanarReflectionComponent"));
+	RootComponent = PlanarReflectionComponent;
+
+	UBoxComponent* DrawInfluenceBox = CreateDefaultSubobject<UBoxComponent>(TEXT("DrawBox0"));
+	DrawInfluenceBox->AttachParent = PlanarReflectionComponent;
+	DrawInfluenceBox->bUseEditorCompositing = true;
+	DrawInfluenceBox->SetCollisionProfileName(UCollisionProfile::NoCollision_ProfileName);
+	PlanarReflectionComponent->PreviewBox = DrawInfluenceBox;
+
+	GetMeshComp()->SetWorldRotation(FRotator(0, 0, 0));
+	GetMeshComp()->SetWorldScale3D(FVector(4, 4, 1));
+	GetMeshComp()->AttachParent = PlanarReflectionComponent;
+}
+
+void APlanarReflection::OnInterpToggle(bool bEnable)
+{
+	PlanarReflectionComponent->SetVisibility(bEnable);
+}
+
+void APlanarReflection::PostActorCreated()
+{
+	Super::PostActorCreated();
+
+	// no need load the editor mesh when there is no editor
+#if WITH_EDITOR
+	if(GetMeshComp())
+	{
+		if (!IsRunningCommandlet())
+		{
+			if( !GetMeshComp()->StaticMesh)
+			{
+				UStaticMesh* PlaneMesh = LoadObject<UStaticMesh>(NULL, TEXT("/Engine/EditorMeshes/PlanarReflectionPlane.PlanarReflectionPlane"), NULL, LOAD_None, NULL);
+				GetMeshComp()->SetStaticMesh(PlaneMesh);
+				UMaterial* PlaneMaterial = LoadObject<UMaterial>(NULL, TEXT("/Engine/EditorMeshes/ColorCalibrator/M_ChromeBall.M_ChromeBall"), NULL, LOAD_None, NULL);
+				GetMeshComp()->SetMaterial(0, PlaneMaterial);
+			}
+		}
+	}
+#endif
+}
+
+#if WITH_EDITOR
+void APlanarReflection::EditorApplyScale(const FVector& DeltaScale, const FVector* PivotLocation, bool bAltDown, bool bShiftDown, bool bCtrlDown)
+{
+	Super::EditorApplyScale(FVector(DeltaScale.X, DeltaScale.Y, 0), PivotLocation, bAltDown, bShiftDown, bCtrlDown);
+
+	UPlanarReflectionComponent* ReflectionComponent = Cast<UPlanarReflectionComponent>(GetPlanarReflectionComponent());
+	check(ReflectionComponent);
+	const FVector ModifiedScale = FVector(0, 0, DeltaScale.Z) * ( AActor::bUsePercentageBasedScaling ? 5000.0f : 50.0f );
+	FMath::ApplyScaleToFloat(ReflectionComponent->DistanceFromPlaneFadeStart, ModifiedScale);
+	FMath::ApplyScaleToFloat(ReflectionComponent->DistanceFromPlaneFadeEnd, ModifiedScale);
+	PostEditChange();
+}
+#endif
+
+// -----------------------------------------------
+
+UPlanarReflectionComponent::UPlanarReflectionComponent(const FObjectInitializer& ObjectInitializer)
+	: Super(ObjectInitializer)
+{
+	bCaptureEveryFrame = true;
+	PrimaryComponentTick.bCanEverTick = true;
+	PrimaryComponentTick.TickGroup = TG_DuringPhysics;
+	// Tick in the editor so that bCaptureEveryFrame preview works
+	bTickInEditor = true;
+	RenderTarget = NULL;
+	NormalDistortionStrength = 500;
+	DistanceFromPlaneFadeStart = 400;
+	DistanceFromPlaneFadeEnd = 600;
+	AngleFromPlaneFadeStart = 20;
+	AngleFromPlaneFadeEnd = 30;
+	ProjectionWithExtraFOV = FMatrix::Identity;
+}
+
+void UPlanarReflectionComponent::CreateRenderState_Concurrent()
+{
+	UpdatePreviewShape();
+
+	Super::CreateRenderState_Concurrent();
+
+	if (ShouldComponentAddToScene() && ShouldRender())
+	{
+		SceneProxy = new FPlanarReflectionSceneProxy(this, RenderTarget);
+		World->Scene->AddPlanarReflection(this);
+	}
+}
+
+void UPlanarReflectionComponent::SendRenderTransform_Concurrent()
+{	
+	UpdatePreviewShape();
+
+	if (SceneProxy)
+	{
+		World->Scene->UpdatePlanarReflectionTransform(this);
+	}
+
+	Super::SendRenderTransform_Concurrent();
+}
+
+void UPlanarReflectionComponent::DestroyRenderState_Concurrent()
+{
+	Super::DestroyRenderState_Concurrent();
+
+	if (SceneProxy)
+	{
+		World->Scene->RemovePlanarReflection(this);
+
+		ENQUEUE_UNIQUE_RENDER_COMMAND_ONEPARAMETER(
+			FDestroyPlanarReflectionCommand,
+			FPlanarReflectionSceneProxy*,SceneProxy,SceneProxy,
+		{
+			delete SceneProxy;
+		});
+
+		SceneProxy = NULL;
+	}
+}
+
+void UPlanarReflectionComponent::BeginDestroy()
+{
+	if (RenderTarget)
+	{
+		BeginReleaseResource(RenderTarget);
+	}
+	
+	// Begin a fence to track the progress of the BeginReleaseResource being processed by the RT
+	ReleaseResourcesFence.BeginFence();
+
+	Super::BeginDestroy();
+}
+
+bool UPlanarReflectionComponent::IsReadyForFinishDestroy()
+{
+	// Wait until the fence is complete before allowing destruction
+	return Super::IsReadyForFinishDestroy() && ReleaseResourcesFence.IsFenceComplete();
+}
+
+void UPlanarReflectionComponent::FinishDestroy()
+{
+	Super::FinishDestroy();
+
+	delete RenderTarget;
+	RenderTarget = NULL;
+}
+
+void UPlanarReflectionComponent::UpdatePreviewShape()
+{
+	if (PreviewBox)
+	{
+		PreviewBox->InitBoxExtent(FVector(500 * 4, 500 * 4, DistanceFromPlaneFadeEnd));
+	}
+}
+
+// -----------------------------------------------
+
 
 USceneCaptureComponentCube::USceneCaptureComponentCube(const FObjectInitializer& ObjectInitializer)
 	: Super(ObjectInitializer)
diff --git a/Engine/Source/Runtime/Engine/Private/Components/SkinnedMeshComponent.cpp b/Engine/Source/Runtime/Engine/Private/Components/SkinnedMeshComponent.cpp
--- a/Engine/Source/Runtime/Engine/Private/Components/SkinnedMeshComponent.cpp
+++ b/Engine/Source/Runtime/Engine/Private/Components/SkinnedMeshComponent.cpp
@@ -586,7 +586,7 @@ bool USkinnedMeshComponent::ShouldCPUSkin()
 	return false;
 }
 
-void USkinnedMeshComponent::GetStreamingTextureInfo(TArray<FStreamingTexturePrimitiveInfo>& OutStreamingTextures) const
+void USkinnedMeshComponent::GetStreamingTextureInfo(FStreamingTextureLevelContext& LevelContext, TArray<FStreamingTexturePrimitiveInfo>& OutStreamingTextures) const
 {
 	if( SkeletalMesh )
 	{
@@ -2001,6 +2001,14 @@ void USkinnedMeshComponent::SetSpaceBaseDoubleBuffering(bool bInDoubleBufferedBl
 	}
 }
 
+void USkinnedMeshComponent::UpdateRecomputeTangent(int32 MaterialIndex)
+{
+	if (ensure(SkeletalMesh) && MeshObject)
+	{
+		MeshObject->UpdateRecomputeTangent(MaterialIndex, SkeletalMesh->Materials[MaterialIndex].bRecomputeTangent);
+	}
+}
+
 void FAnimUpdateRateParameters::SetTrailMode(float DeltaTime, uint8 UpdateRateShift, int32 NewUpdateRate, int32 NewEvaluationRate, bool bNewInterpSkippedFrames)
 {
 	OptimizeMode = TrailMode;
diff --git a/Engine/Source/Runtime/Engine/Private/Components/SplineMeshComponent.cpp b/Engine/Source/Runtime/Engine/Private/Components/SplineMeshComponent.cpp
--- a/Engine/Source/Runtime/Engine/Private/Components/SplineMeshComponent.cpp
+++ b/Engine/Source/Runtime/Engine/Private/Components/SplineMeshComponent.cpp
@@ -1161,9 +1161,9 @@ FStaticMeshStaticLightingMesh* USplineMeshComponent::AllocateStaticLightingMesh(
 }
 
 
-bool USplineMeshComponent::GetStreamingTextureFactors(float& OutWorldTexelFactor, float& OutWorldLightmapFactor, FBoxSphereBounds& OutBounds, int32 LODIndex, int32 ElementIndex) const
+bool USplineMeshComponent::GetStreamingTextureFactors(float& OutWorldTexelFactor, float& OutWorldLightmapFactor) const
 {
-	if (UStaticMeshComponent::GetStreamingTextureFactors(OutWorldTexelFactor, OutWorldLightmapFactor, OutBounds, LODIndex, ElementIndex))
+	if (UStaticMeshComponent::GetStreamingTextureFactors(OutWorldTexelFactor, OutWorldLightmapFactor))
 	{
 		// We need to come up with a compensation factor for spline deformed meshes
 
diff --git a/Engine/Source/Runtime/Engine/Private/Components/StaticMeshComponent.cpp b/Engine/Source/Runtime/Engine/Private/Components/StaticMeshComponent.cpp
--- a/Engine/Source/Runtime/Engine/Private/Components/StaticMeshComponent.cpp
+++ b/Engine/Source/Runtime/Engine/Private/Components/StaticMeshComponent.cpp
@@ -495,80 +495,238 @@ void UStaticMeshComponent::OnUnregister()
 	Super::OnUnregister();
 }
 
-void UStaticMeshComponent::UpdateStreamingTextureInfos(bool bForce)
-{
-	// Only rebuild the data in editor 
 #if WITH_EDITORONLY_DATA
-	StreamingTextureInfos.Empty();
 
-	if ((bForce || (!bIgnoreInstanceForTextureStreaming && Mobility == EComponentMobility::Static)) && StaticMesh && StaticMesh->RenderData)
+/** Return the total number of LOD sections in the LOD resources */
+static int32 GetNumberOfElements(const TIndirectArray<FStaticMeshLODResources>& LODResources)
+{
+	int32 Count = 0;
+	for (int32 LODIndex = 0; LODIndex < LODResources.Num(); ++LODIndex)
 	{
-		for (int32 LODIndex = 0; LODIndex < StaticMesh->RenderData->LODResources.Num(); ++LODIndex)
+		Count += LODResources[LODIndex].Sections.Num();
+	}
+	return Count;
+}
+
+/**
+ *	Set this struct to match the unpacked params.
+ *
+ *	@param	LevelTextures			[in,out]	The list of textures referred by all component of a level. The array index maps to UTexture2D::LevelIndex.
+ *	@param	UnpackedData			[in,out]	The unpacked data, emptied after the function executes.
+ *	@param	StreamingTextureData	[out]		The resulting packed data.
+ *	@param	RefBounds				[in]		The reference bounds used to packed the relative bounds.
+ */
+static void PackStreamingTextureData(TArray<UTexture2D*>& LevelTextures, TArray<FStreamingTexturePrimitiveInfo>& UnpackedData, TArray<FStreamingTextureBuildInfo>& StreamingTextureData, const FBoxSphereBounds& RefBounds)
+{
+	StreamingTextureData.Empty();
+
+	while (UnpackedData.Num())
+	{
+		FStreamingTexturePrimitiveInfo Info = UnpackedData[0];
+		UnpackedData.RemoveAtSwap(0);
+
+		// Merge with any other lod section using the same texture.
+		for (int32 Index = 0; Index < UnpackedData.Num(); ++Index)
 		{
-			const FStaticMeshLODResources& LOD = StaticMesh->RenderData->LODResources[LODIndex];
+			const FStreamingTexturePrimitiveInfo& CurrInfo = UnpackedData[Index];
 
-			for( int32 ElementIndex = 0; ElementIndex < LOD.Sections.Num(); ElementIndex++ )
+			if (CurrInfo.Texture == Info.Texture)
 			{
-				FStreamingTexturePrimitiveInfo Info;
+				Info.Bounds = Union(Info.Bounds, CurrInfo.Bounds);
+				// Take the max scale since it relates to higher texture resolution.
+				Info.TexelFactor = FMath::Max<float>(Info.TexelFactor, CurrInfo.TexelFactor);
 
-				if (StaticMesh->GetStreamingTextureFactor(Info.TexelFactor, Info.Bounds, 0, LODIndex, ElementIndex, ComponentToWorld))
-				{
-					Info.TexelFactor *= FMath::Max(0.0f, StreamingDistanceMultiplier);
-					StreamingTextureInfos.Push(Info);
-				}
-				else
+				UnpackedData.RemoveAtSwap(Index);
+				--Index;
+			}
+		}
+
+		FStreamingTextureBuildInfo PackedInfo;
+		PackedInfo.PackFrom(LevelTextures, RefBounds, Info);
+		StreamingTextureData.Push(PackedInfo);
+	}
+}
+
+#endif
+
+void UStaticMeshComponent::UpdateStreamingTextureData(TArray<UTexture2D*>& LevelTextures, const FTexCoordScaleMap& TexCoordScales, EMaterialQualityLevel::Type QualityLevel, ERHIFeatureLevel::Type FeatureLevel)
+{
+#if WITH_EDITORONLY_DATA // Only rebuild the data in editor 
+	const bool bUseMetrics = CVarStreamingUseNewMetrics.GetValueOnGameThread() != 0;
+	const bool bNeedsPrecomputedData = !bIgnoreInstanceForTextureStreaming && Mobility == EComponentMobility::Static && StaticMesh && StaticMesh->RenderData;
+
+	if (bUseMetrics && bNeedsPrecomputedData)
+	{
+		UpdateStreamingSectionData(TexCoordScales);
+
+		TArray<FStreamingTexturePrimitiveInfo> UnpackedData;
+
+		TArray<UTexture*> Textures;
+		TArray< TArray<int32> > Indices;
+
+		for (const FStreamingSectionBuildInfo& SectionData : *StreamingSectionData)
+		{
+			if (SectionData.MaterialIndex == INDEX_NONE) continue; // Skip the fallback data.
+
+			UMaterialInterface* Material = GetMaterial(SectionData.MaterialIndex);
+
+			Textures.Empty();
+			Indices.Empty();
+
+			// Get the texture register indices used by each textures. 
+			// Must be using the same quality and feature level as the one used to compute scales
+			// in order to remap to the correct texture register index.
+			Material->GetUsedTexturesAndIndices(Textures, Indices, QualityLevel, FeatureLevel);
+
+			for (int32 TextureIndex = 0; TextureIndex < Textures.Num(); ++TextureIndex)
+			{
+				UTexture2D* Texture = Cast<UTexture2D>(Textures[TextureIndex]);
+				if (!Texture) continue;
+
+				bool bMaxValid = false;
+				float MaxTexeFactor = 0;
+
+				if (Indices.IsValidIndex(TextureIndex))
 				{
-					StreamingTextureInfos.Empty();
-					return;
+					for (int32 TextureRegisterIndex : Indices[TextureIndex])
+					{
+						if (SectionData.TexCoordData.IsValidIndex(TextureRegisterIndex))
+						{
+							const FMaterialTexCoordBuildInfo& TexCoordInfo = SectionData.TexCoordData[TextureRegisterIndex];
+							 // 0 would indicate that this register index is irrelevant (not used). Also could cause divide by 0
+							if (TexCoordInfo.Scale > 0)
+							{
+								float TexelFactor = SectionData.TexelFactors[0];
+								if (TexCoordInfo.Index >= 0 && TexCoordInfo.Index < FMaterialTexCoordBuildInfo::MAX_NUM_TEX_COORD && SectionData.TexelFactors[TexCoordInfo.Index] > 0)
+								{
+									TexelFactor = SectionData.TexelFactors[TexCoordInfo.Index];
+								}
+								MaxTexeFactor = FMath::Max<float>(MaxTexeFactor, TexelFactor / TexCoordInfo.Scale);
+								bMaxValid = true;
+							}
+						}
+					}
 				}
+
+				FStreamingTexturePrimitiveInfo& Info = *new(UnpackedData) FStreamingTexturePrimitiveInfo();
+				Info.Texture = Texture;
+				Info.TexelFactor = bMaxValid ? MaxTexeFactor : SectionData.TexelFactors[0];
+				Info.Bounds.Origin = SectionData.BoxOrigin;
+				Info.Bounds.BoxExtent = SectionData.BoxExtent;
+				Info.Bounds.SphereRadius = SectionData.BoxExtent.Size();
 			}
 		}
+
+		PackStreamingTextureData(LevelTextures, UnpackedData, StreamingTextureData, Bounds);
+	}
+	else if (!bUseMetrics && bNeedsPrecomputedData)
+	{
+		StreamingTextureData.Empty();
+		StreamingSectionData.Reset();
+
+		StreamingSectionData = TSharedPtr<TArray<FStreamingSectionBuildInfo>, ESPMode::NotThreadSafe>(new TArray<FStreamingSectionBuildInfo>());
+
+		FStreamingSectionBuildInfo& SectionData = *new(*StreamingSectionData) FStreamingSectionBuildInfo();
+		SectionData.BoxOrigin = Bounds.Origin;
+		SectionData.BoxExtent = Bounds.BoxExtent;
+		SectionData.TexelFactors[0] = StaticMesh ? StaticMesh->GetStreamingTextureFactor(0) : 0.f;
+	}
+	else
+	{
+		StreamingTextureData.Empty();
+		StreamingSectionData.Reset();
 	}
 #endif
 }
 
-extern ENGINE_API TAutoConsoleVariable<int32> CVarStreamingUseNewMetrics;
-
-bool UStaticMeshComponent::GetStreamingTextureFactors(float& OutWorldTexelFactor, float& OutWorldLightmapFactor, FBoxSphereBounds& OutBounds, int32 LODIndex, int32 ElementIndex) const
+void UStaticMeshComponent::UpdateStreamingSectionData(const FTexCoordScaleMap& TexCoordScales)
 {
-	bool bUsePreciseStreamingBounds = false;
+#if WITH_EDITORONLY_DATA
+	StreamingSectionData.Reset();
 
-	if (StaticMesh && StaticMesh->RenderData && LODIndex < StaticMesh->RenderData->LODResources.Num())
+	if (StaticMesh && StaticMesh->RenderData)
 	{
-		TIndirectArray<FStaticMeshLODResources>& LODResources = StaticMesh->RenderData->LODResources;
+		StreamingSectionData = TSharedPtr<TArray<FStreamingSectionBuildInfo>, ESPMode::NotThreadSafe>(new TArray<FStreamingSectionBuildInfo>());
 
-		const ERHIFeatureLevel::Type FeatureLevel = GetWorld() ? GetWorld()->FeatureLevel : GMaxRHIFeatureLevel;
+		const TIndirectArray<FStaticMeshLODResources>& LODResources = StaticMesh->RenderData->LODResources;
+		StreamingSectionData->Reserve(GetNumberOfElements(LODResources) + 1);
 
-		const bool bHasValidLightmapCoordinates = StaticMesh->LightMapCoordinateIndex >= 0 &&
-			(uint32)StaticMesh->LightMapCoordinateIndex < LODResources[0].VertexBuffer.GetNumTexCoords();
+		// First section holds fallback data used to compare old streaming results with new one
+		{
+			FStreamingSectionBuildInfo& FallbackData = *new(*StreamingSectionData) FStreamingSectionBuildInfo();
+			FallbackData.BoxOrigin = Bounds.Origin;
+			FallbackData.BoxExtent = Bounds.BoxExtent;
+			FallbackData.TexelFactors[0] = StaticMesh ? StaticMesh->GetStreamingTextureFactor(0) : 0.f;
+		}
 
-		// Because there is a single array for each sections and lods, the index needs to be computed depending on the number of section per lod.
-		int32 InfoIndex = ElementIndex;
-		if (StreamingTextureInfos.Num() && LODResources.IsValidIndex(LODIndex) && LODResources[LODIndex].Sections.IsValidIndex(ElementIndex))
+		for (int32 LODIndex = 0; LODIndex < LODResources.Num(); ++LODIndex)
 		{
-			int32 DebugLODIndex = 0;
-			while (DebugLODIndex < LODIndex)
+			const FStaticMeshLODResources& LOD = LODResources[LODIndex];
+			for (int32 ElementIndex = 0; ElementIndex < LOD.Sections.Num(); ++ElementIndex)
 			{
-				InfoIndex += LODResources[DebugLODIndex].Sections.Num();
-				++DebugLODIndex;
+				// If the streaming factors are not valid, this section can be ignored. Could be related to 0 area size.
+				float LODElementTexelFactor;
+				FBoxSphereBounds LODElementBounds;
+				if (!StaticMesh->GetStreamingTextureFactor(LODElementTexelFactor, LODElementBounds, 0, LODIndex, ElementIndex, ComponentToWorld))
+					continue;
+
+				const FStaticMeshSection& Element = LOD.Sections[ElementIndex];
+				UMaterialInterface* Material = GetMaterial(Element.MaterialIndex);
+				if (!Material) continue;
+
+				FStreamingSectionBuildInfo& SectionData = *new(*StreamingSectionData) FStreamingSectionBuildInfo();
+				SectionData.LODIndex = LODIndex;
+				SectionData.ElementIndex = ElementIndex;
+				SectionData.MaterialIndex = Element.MaterialIndex;
+				SectionData.BoxOrigin = LODElementBounds.Origin;
+				SectionData.BoxExtent = LODElementBounds.BoxExtent;
+				SectionData.TexelFactors[0] = LODElementTexelFactor;
+
+				const TArray<FMaterialTexCoordBuildInfo>* TexCoordData = TexCoordScales.Find(Material);
+				int32 MaxTexCoordIndex = MaxTexCoordIndex = FMaterialTexCoordBuildInfo::MAX_NUM_TEX_COORD - 1;
+				if (TexCoordData)
+				{
+					SectionData.TexCoordData = *TexCoordData;
+					MaxTexCoordIndex = 0;
+					for (const FMaterialTexCoordBuildInfo& TexCoordInfo : *TexCoordData)
+					{
+						MaxTexCoordIndex = FMath::Max<int32>(TexCoordInfo.Index, MaxTexCoordIndex);
+					}
+					MaxTexCoordIndex = FMath::Min<int32>(MaxTexCoordIndex, FMaterialTexCoordBuildInfo::MAX_NUM_TEX_COORD - 1);
+				}
+
+				for (int32 TexCoordIndex = 1; TexCoordIndex <= MaxTexCoordIndex; ++TexCoordIndex)
+				{
+					if (StaticMesh->GetStreamingTextureFactor(LODElementTexelFactor, LODElementBounds, TexCoordIndex, LODIndex, ElementIndex, ComponentToWorld))
+					{
+						SectionData.TexelFactors[TexCoordIndex] = LODElementTexelFactor;
+					}
+				}
+
 			}
 		}
+	}
+#endif
+}
+
+bool UStaticMeshComponent::GetStreamingTextureFactors(float& OutWorldTexelFactor, float& OutWorldLightmapFactor) const
+{
+	if (StaticMesh && StaticMesh->RenderData && StaticMesh->RenderData->LODResources.Num() > 0)
+	{
+		OutWorldTexelFactor = OutWorldLightmapFactor = ComponentToWorld.GetMaximumAxisScale();
 
-		if (StreamingTextureInfos.IsValidIndex(InfoIndex) && CVarStreamingUseNewMetrics.GetValueOnAnyThread() > 0)
+		OutWorldTexelFactor *= StaticMesh->GetStreamingTextureFactor(0);
+
+		TIndirectArray<FStaticMeshLODResources>& LODResources = StaticMesh->RenderData->LODResources;
+		const bool bHasValidLightmapCoordinates = StaticMesh->LightMapCoordinateIndex >= 0 && (uint32)StaticMesh->LightMapCoordinateIndex < LODResources[0].VertexBuffer.GetNumTexCoords();
+		if (bHasValidLightmapCoordinates)
 		{
-			OutBounds = StreamingTextureInfos[InfoIndex].Bounds;
-			OutWorldTexelFactor = StreamingTextureInfos[InfoIndex].TexelFactor;
+			OutWorldLightmapFactor *= StaticMesh->GetStreamingTextureFactor(StaticMesh->LightMapCoordinateIndex);
 		}
 		else
 		{
-			float LocalTexelFactor = StaticMesh->GetStreamingTextureFactor(0) * FMath::Max(0.0f, StreamingDistanceMultiplier);
-			OutWorldTexelFactor	= LocalTexelFactor * ComponentToWorld.GetMaximumAxisScale();
-			OutBounds = Bounds;
+			OutWorldLightmapFactor = 0;
 		}
-
-		const float LocalLightmapFactor	= bHasValidLightmapCoordinates ? StaticMesh->GetStreamingTextureFactor(StaticMesh->LightMapCoordinateIndex) : 1.0f;
-		OutWorldLightmapFactor = LocalLightmapFactor * ComponentToWorld.GetMaximumAxisScale();
-
 		return true;
 	}
 	else
@@ -578,17 +736,17 @@ bool UStaticMeshComponent::GetStreamingTextureFactors(float& OutWorldTexelFactor
 }
 
 
-void UStaticMeshComponent::GetStreamingTextureInfo(TArray<FStreamingTexturePrimitiveInfo>& OutStreamingTextures) const
+void UStaticMeshComponent::GetStreamingTextureInfo(FStreamingTextureLevelContext& LevelContext, TArray<FStreamingTexturePrimitiveInfo>& OutStreamingTextures) const
 {
 	float WorldTexelFactor = 1.f;
 	float WorldLightmapFactor = 1.f;
 
-	if (!bIgnoreInstanceForTextureStreaming && StaticMesh && StaticMesh->RenderData)
+	if (!bIgnoreInstanceForTextureStreaming && GetStreamingTextureFactors(WorldTexelFactor, WorldLightmapFactor))
 	{
 		const ERHIFeatureLevel::Type FeatureLevel = GetWorld() ? GetWorld()->FeatureLevel : GMaxRHIFeatureLevel;
+		const bool bUseNewMetrics = CVarStreamingUseNewMetrics.GetValueOnGameThread() != 0;
 
-		const bool bHasValidLightmapCoordinates = ((StaticMesh->LightMapCoordinateIndex >= 0)
-			&& ((uint32)StaticMesh->LightMapCoordinateIndex < StaticMesh->RenderData->LODResources[0].VertexBuffer.GetNumTexCoords()));
+		LevelContext.BindComponent(bUseNewMetrics ? &StreamingTextureData : nullptr, Bounds, WorldTexelFactor, FMath::Max(0.0f, StreamingDistanceMultiplier));
 
 		for (int32 LODIndex = 0; LODIndex < StaticMesh->RenderData->LODResources.Num(); ++LODIndex)
 		{
@@ -604,28 +762,15 @@ void UStaticMeshComponent::GetStreamingTextureInfo(TArray<FStreamingTexturePrimi
 					Material = UMaterial::GetDefaultMaterial(MD_Surface);
 				}
 
-				FBoxSphereBounds SectionBounds;
-				if (!GetStreamingTextureFactors(WorldTexelFactor, WorldLightmapFactor, SectionBounds, LODIndex, ElementIndex))
-					continue;
-
 				// Enumerate the textures used by the material.
 				TArray<UTexture*> Textures;
-
 				Material->GetUsedTextures(Textures, EMaterialQualityLevel::Num, false, FeatureLevel, false);
 
-				// Add each texture to the output with the appropriate parameters.
-				// TODO: Take into account which UVIndex is being used.
-				for(int32 TextureIndex = 0;TextureIndex < Textures.Num();TextureIndex++)
-				{
-					FStreamingTexturePrimitiveInfo& StreamingTexture = *new(OutStreamingTextures) FStreamingTexturePrimitiveInfo;
-					StreamingTexture.Bounds = SectionBounds;
-					StreamingTexture.TexelFactor = WorldTexelFactor;
-					StreamingTexture.Texture = Textures[TextureIndex];
-				}
+				LevelContext.Process(Textures, OutStreamingTextures);
 			}
 
 			// Add in the lightmaps and shadowmaps.
-			if ( LODData.IsValidIndex(LODIndex) && bHasValidLightmapCoordinates )
+			if ( LODData.IsValidIndex(LODIndex) && WorldLightmapFactor != 0)
 			{
 				const FStaticMeshComponentLODInfo& LODInfo = LODData[LODIndex];
 				FLightMap2D* Lightmap = LODInfo.LightMap ? LODInfo.LightMap->GetLightMap2D() : NULL;
@@ -645,7 +790,6 @@ void UStaticMeshComponent::GetStreamingTextureInfo(TArray<FStreamingTexturePrimi
 				}
 
 				FShadowMap2D* Shadowmap = LODInfo.ShadowMap ? LODInfo.ShadowMap->GetShadowMap2D() : NULL;
-
 				if ( Shadowmap != NULL && Shadowmap->IsValid() )
 				{
 					const FVector2D& Scale = Shadowmap->GetCoordinateScale();
@@ -1683,7 +1827,7 @@ UMaterialInterface* UStaticMeshComponent::GetMaterial(int32 MaterialIndex) const
 	// Otherwise get from static mesh
 	else
 	{
-		return StaticMesh ? StaticMesh->GetMaterial(MaterialIndex) : NULL;
+		return StaticMesh ? StaticMesh->GetMaterial(MaterialIndex) : nullptr;
 	}
 }
 
