Copying //UE4/Dev-Rendering to Dev-Main (Source //UE4/Dev-Rendering@2932636)

diff --git a/Engine/Source/Runtime/Engine/Private/ContentStreaming.cpp b/Engine/Source/Runtime/Engine/Private/ContentStreaming.cpp
--- a/Engine/Source/Runtime/Engine/Private/ContentStreaming.cpp
+++ b/Engine/Source/Runtime/Engine/Private/ContentStreaming.cpp
@@ -6,73 +6,15 @@
 
 #include "EnginePrivate.h"
 #include "GenericPlatformMemoryPoolStats.h"
+#include "Streaming/AsyncTextureStreaming.h"
 #include "AudioStreaming.h"
-#include "TextureInstanceManager.h"
 #include "Engine/LightMapTexture2D.h"
 #include "Engine/ShadowMapTexture2D.h"
 
-DEFINE_LOG_CATEGORY_STATIC(LogContentStreaming, Log, All);
-
 /*-----------------------------------------------------------------------------
 	Stats.
 -----------------------------------------------------------------------------*/
 
-//@DEBUG:
-// Set to 1 to log all dynamic component notifications
-#define STREAMING_LOG_DYNAMIC		0
-// Set to 1 to log when we change a view
-#define STREAMING_LOG_VIEWCHANGES	0
-// Set to 1 to log when levels are added/removed
-#define STREAMING_LOG_LEVELS		0
-// Set to 1 to log textures that are canceled by CancelForcedTextures()
-#define STREAMING_LOG_CANCELFORCED	0
-
-/** Streaming stats */
-
-DEFINE_STAT(STAT_GameThreadUpdateTime);
-DEFINE_STAT(STAT_TexturePoolSize);
-DEFINE_STAT(STAT_TexturePoolAllocatedSize);
-DEFINE_STAT(STAT_OptimalTextureSize);
-DEFINE_STAT(STAT_EffectiveStreamingPoolSize);
-DEFINE_STAT(STAT_StreamingTexturesSize);
-DEFINE_STAT(STAT_NonStreamingTexturesSize);
-DEFINE_STAT(STAT_StreamingTexturesMaxSize);
-DEFINE_STAT(STAT_StreamingOverBudget);
-DEFINE_STAT(STAT_StreamingUnderBudget);
-DEFINE_STAT(STAT_NumWantingTextures);
-
-DEFINE_STAT(STAT_StreamingTextures);
-
-DEFINE_STAT(STAT_TotalStaticTextureHeuristicSize);
-DEFINE_STAT(STAT_TotalDynamicHeuristicSize);
-DEFINE_STAT(STAT_TotalLastRenderHeuristicSize);
-DEFINE_STAT(STAT_TotalForcedHeuristicSize);
-
-DEFINE_STAT(STAT_LightmapMemorySize);
-DEFINE_STAT(STAT_LightmapDiskSize);
-DEFINE_STAT(STAT_HLODTextureMemorySize);
-DEFINE_STAT(STAT_HLODTextureDiskSize);
-DEFINE_STAT(STAT_IntermediateTexturesSize);
-DEFINE_STAT(STAT_RequestSizeCurrentFrame);
-DEFINE_STAT(STAT_RequestSizeTotal);
-DEFINE_STAT(STAT_LightmapRequestSizeTotal);
-
-DEFINE_STAT(STAT_IntermediateTextures);
-DEFINE_STAT(STAT_RequestsInCancelationPhase);
-DEFINE_STAT(STAT_RequestsInUpdatePhase);
-DEFINE_STAT(STAT_RequestsInFinalizePhase);
-DEFINE_STAT(STAT_StreamingLatency);
-DEFINE_STAT(STAT_StreamingBandwidth);
-DEFINE_STAT(STAT_GrowingReallocations);
-DEFINE_STAT(STAT_ShrinkingReallocations);
-DEFINE_STAT(STAT_FullReallocations);
-DEFINE_STAT(STAT_FailedReallocations);
-DEFINE_STAT(STAT_PanicDefragmentations);
-DEFINE_STAT(STAT_DynamicStreamingTotal);
-DEFINE_STAT(STAT_NumVisibleTexturesWithLowResolutions);
-
-/** Accumulated total time spent on dynamic primitives, in seconds. */
-double GStreamingDynamicPrimitivesTime = 0.0;
 
 /*-----------------------------------------------------------------------------
 	Globals.
@@ -93,147 +35,6 @@ TArray<IStreamingManager::FSlaveLocation> IStreamingManager::SlaveLocations;
 /** Set when Tick() has been called. The first time a new view is added, it will clear out all old views. */
 bool IStreamingManager::bPendingRemoveViews = false;
 
-/** The lightmap used by the currently selected component (toggledebugcamera), if it's a static mesh component. */
-extern ENGINE_API class FLightMap2D* GDebugSelectedLightmap;
-
-/** Whether the texture pool size has been artificially lowered for testing. */
-bool GIsOperatingWithReducedTexturePool = false;
-
-/** Smaller value will stream out lightmaps more aggressively. */
-float GLightmapStreamingFactor = 1.0f;
-
-/** Smaller value will stream out shadowmaps more aggressively. */
-float GShadowmapStreamingFactor = 0.09f;
-
-/** For testing, finding useless textures or special demo purposes. If true, textures will never be streamed out (but they can be GC'd). 
-* Caution: this only applies to unlimited texture pools (i.e. not consoles)
-*/
-bool GNeverStreamOutTextures = false;
-
-// Used for scalability (GPU memory, streaming stalls)
-static TAutoConsoleVariable<float> CVarStreamingMipBias(
-	TEXT("r.Streaming.MipBias"),
-	0.0f,
-	TEXT("0..x reduce texture quality for streaming by a floating point number.\n")
-	TEXT("0: use full resolution (default)\n")
-	TEXT("1: drop one mip\n")
-	TEXT("2: drop two mips"),
-	ECVF_Scalability
-	);
-
-static TAutoConsoleVariable<float> CVarStreamingBoost(
-	TEXT("r.Streaming.Boost"),
-	1.0f,
-	TEXT("=1.0: normal\n")
-	TEXT("<1.0: decrease wanted mip levels\n")
-	TEXT(">1.0: increase wanted mip levels"),
-	ECVF_Default
-	);
-
-static TAutoConsoleVariable<float> CVarStreamingScreenSizeEffectiveMax(
-	TEXT("r.Streaming.MaxEffectiveScreenSize"),
-	0,
-	TEXT("0: Use current actual vertical screen size\n")	
-	TEXT("> 0: Clamp wanted mip size calculation to this value for the vertical screen size component."),
-	ECVF_Scalability
-	);
-
-#if PLATFORM_SUPPORTS_TEXTURE_STREAMING
-static TAutoConsoleVariable<int32> CVarSetTextureStreaming(
-	TEXT("r.TextureStreaming"),
-	1,
-	TEXT("Allows to define if texture streaming is enabled, can be changed at run time.\n")
-	TEXT("0: off\n")
-	TEXT("1: on (default)"),
-	ECVF_Default | ECVF_RenderThreadSafe);
-#endif
-
-static TAutoConsoleVariable<int32> CVarStreamingUseFixedPoolSize(
-	TEXT("r.Streaming.UseFixedPoolSize"),
-	0,
-	TEXT("If non-zero, do not allow the pool size to change at run time."),
-	ECVF_ReadOnly);
-
-static TAutoConsoleVariable<int32> CVarStreamingPoolSize(
-	TEXT("r.Streaming.PoolSize"),
-	-1,
-	TEXT("-1: Default texture pool size, otherwise the size in MB"),
-	ECVF_Scalability);
-
-static TAutoConsoleVariable<int32> CVarStreamingMaxTempMemoryAllowed(
-	TEXT("r.Streaming.MaxTempMemoryAllowed"),
-	50,
-	TEXT("Maximum temporary memory used when streaming in or out texture mips.\n")
-	TEXT("This memory contains mips used for the new updated texture.\n")
-	TEXT("The value must be high enough to not be a limiting streaming speed factor.\n"),
-	ECVF_Default);
-
-static TAutoConsoleVariable<int32> CVarStreamingShowWantedMips(
-	TEXT("r.Streaming.ShowWantedMips"),
-	0,
-	TEXT("If non-zero, will limit resolution to wanted mip."),
-	ECVF_Cheat);
-
-ENGINE_API TAutoConsoleVariable<int32> CVarStreamingUseNewMetrics(
-	TEXT("r.Streaming.UseNewMetrics"),
-	0,
-	TEXT("If non-zero, will use tight AABB bounds and improved texture factors."),
-	ECVF_Default);
-
-static TAutoConsoleVariable<int32> CVarStreamingHLODStrategy(
-	TEXT("r.Streaming.HLODStrategy"),
-	0,
-	TEXT("Define the HLOD streaming strategy.\n")
-	TEXT("0: stream\n")
-	TEXT("1: stream only mip 0\n")
-	TEXT("2: disable streaming"),
-	ECVF_Default);
-
-static TAutoConsoleVariable<float> CVarStreamingHiddenPrimitiveScale(
-	TEXT("r.Streaming.HiddenPrimitiveScale"),
-	0.5,
-	TEXT("Define the resolution scale to apply when not in range.\n")
-	TEXT(".5: drop one mip\n")
-	TEXT("1: ignore visiblity"),
-	ECVF_Default
-	);
-
-static TAutoConsoleVariable<int32> CVarStreamingSplitRequestSizeThreshold(
-	TEXT("r.Streaming.SplitRequestSizeThreshold"),
-	0, // Good value : 2 * 1024 * 1024,
-	TEXT("Define how many IO pending request the streamer can generate.\n"),
-	ECVF_Default
-	);
-
-
-/** Streaming priority: Linear distance factor from 0 to MAX_STREAMINGDISTANCE. */
-#define MAX_STREAMINGDISTANCE	10000.0f
-#define MAX_MIPDELTA			5.0f
-#define MAX_LASTRENDERTIME		90.0f
-
-/** For debugging purposes: Whether to consider the time factor when streaming. Turning it off is easier for debugging. */
-bool GStreamWithTimeFactor		= true;
-
-extern ENGINE_API UPrimitiveComponent* GDebugSelectedComponent;
-
-#if STATS
-/** Ringbuffer of bandwidth samples for streaming in mip-levels (MB/sec). */
-float FStreamingManagerTexture::BandwidthSamples[NUM_BANDWIDTHSAMPLES];
-/** Number of bandwidth samples that have been filled in. Will stop counting when it reaches NUM_BANDWIDTHSAMPLES. */
-int32 FStreamingManagerTexture::NumBandwidthSamples = 0;
-/** Current sample index in the ring buffer. */
-int32 FStreamingManagerTexture::BandwidthSampleIndex = 0;
-/** Average of all bandwidth samples in the ring buffer, in MB/sec. */
-float FStreamingManagerTexture::BandwidthAverage = 0.0f;
-/** Maximum bandwidth measured since the start of the game.  */
-float FStreamingManagerTexture::BandwidthMaximum = 0.0f;
-/** Minimum bandwidth measured since the start of the game.  */
-float FStreamingManagerTexture::BandwidthMinimum = 0.0f;
-#endif
-
-/** the float table {-1.0f,1.0f} **/
-float ENGINE_API GNegativeOneOneTable[2] = {-1.0f,1.0f};
-
 /**
  * Helper function to flush resource streaming from within Core project.
  */
@@ -243,388 +44,6 @@ void FlushResourceStreaming()
 	IStreamingManager::Get().BlockTillAllRequestsFinished();
 }
 
-/**
- * Helper function to clamp the mesh to camera distance
- */
-FORCEINLINE float ClampMeshToCameraDistanceSquared(float MeshToCameraDistanceSquared)
-{
-	// called from streaming thread, maybe even main thread
-	return FMath::Max<float>(MeshToCameraDistanceSquared, 0.0f);
-}
-
-
-/*-----------------------------------------------------------------------------
-	FStreamingTexture, the streaming system's version of UTexture2D.
------------------------------------------------------------------------------*/
-
-enum ETextureStreamingType
-{
-	StreamType_Static,
-	StreamType_Dynamic,
-	StreamType_Forced,
-	StreamType_LastRenderTime,
-	StreamType_Orphaned,
-	StreamType_Other,
-};
-
-static const TCHAR* GStreamTypeNames[] =
-{
-	TEXT("Static"),
-	TEXT("Dynamic"),
-	TEXT("Forced"),
-	TEXT("LastRenderTime"),
-	TEXT("Orphaned"),
-	TEXT("Other"),
-};
-
-struct FTexturePriority
-{
-	FTexturePriority( bool InCanDropMips, float InRetentionPriority, float InLoadOrderPriority, int32 InTextureIndex, const UTexture2D* InTexture)
-	:	bCanDropMips( InCanDropMips )
-	,	RetentionPriority( InRetentionPriority )
-	,	LoadOrderPriority( InLoadOrderPriority )
-	,	TextureIndex( InTextureIndex )
-	,	Texture( InTexture )
-	{
-	}
-	/** True if we allows this texture to drop mips to fit in budget. */
-	bool bCanDropMips;
-	/** Texture retention priority, higher value means it should be kept in memory. */
-	float RetentionPriority;
-	/** Texture load order priority, higher value means it should load first. */
-	float LoadOrderPriority;
-	/** Index into the FStreamingManagerTexture::StreamingTextures array. */
-	int32 TextureIndex;
-	/** The texture to stream. Only used for validation. */
-	const UTexture2D* Texture;
-};
-
-/** Self-contained structure to manage a streaming texture, possibly on a separate thread. */
-struct FStreamingTexture
-{
-	FStreamingTexture( UTexture2D* InTexture )
-	{
-		Texture = InTexture;
-		WantedMips = InTexture->ResidentMips;
-		MipCount = InTexture->GetNumMips();
-		PerfectWantedMips = InTexture->ResidentMips;
-		STAT( MostResidentMips = InTexture->ResidentMips );
-		LODGroup = (TextureGroup) InTexture->LODGroup;
-
-		//@TODO: Temporary: Done because there is currently no texture group for HLOD textures and adding one in a branch isn't ideal; this should be removed / replaced once there is a real HLOD texture group
-		bIsHLODTexture = !Texture->HasAnyFlags(RF_Public) && Texture->GetName().StartsWith(TEXT("T_LOD_"), ESearchCase::CaseSensitive);
-
-		// Landscape textures and HLOD textures should not be affected by MipBias
-		//@TODO: Long term this should probably be a property of the texture group instead
-		bCanBeAffectedByMipBias = true;
-		if ((LODGroup == TEXTUREGROUP_Terrain_Heightmap) || (LODGroup == TEXTUREGROUP_Terrain_Weightmap) || bIsHLODTexture)
-		{
-			bCanBeAffectedByMipBias = false;
-		}
-
-		NumNonStreamingMips = InTexture->GetNumNonStreamingMips();
-		ForceLoadRefCount = 0;
-		bIsStreamingLightmap = IsStreamingLightmap( Texture );
-		bIsLightmap = bIsStreamingLightmap || LODGroup == TEXTUREGROUP_Lightmap || LODGroup == TEXTUREGROUP_Shadowmap;
-		bUsesStaticHeuristics = false;
-		bUsesDynamicHeuristics = false;
-		bUsesLastRenderHeuristics = false;
-		bUsesForcedHeuristics = false;
-		bUsesOrphanedHeuristics = false;
-		bHasSplitRequest = false;
-		bIsLastSplitRequest = false;
-		BoostFactor = 1.0f;
-		InstanceRemovedTimestamp = -FLT_MAX;
-		LastRenderTimeRefCountTimestamp = -FLT_MAX;
-		LastRenderTimeRefCount = 0;
-
-		for ( int32 MipIndex=1; MipIndex <= MAX_TEXTURE_MIP_COUNT; ++MipIndex )
-		{
-			TextureSizes[MipIndex - 1] = Texture->CalcTextureMemorySize( FMath::Min(MipIndex, MipCount) );
-		}
-
-		UpdateCachedInfo();
-	}
-
-	/** Reinitializes the cached variables from UTexture2D. */
-	void UpdateCachedInfo( )
-	{
-		check(Texture && Texture->Resource);
-		ResidentMips = Texture->ResidentMips;
-		RequestedMips = Texture->RequestedMips;
-		MinAllowedMips = 1;			//FMath::Min(Texture->ResidentMips, Texture->RequestedMips);
-		MaxAllowedMips = MipCount;	//FMath::Max(Texture->ResidentMips, Texture->RequestedMips);
-		MaxAllowedOptimalMips = MaxAllowedMips;
-		STAT( MostResidentMips = FMath::Max(MostResidentMips, Texture->ResidentMips) );
-		float LastRenderTimeForTexture = Texture->GetLastRenderTimeForStreaming();
-		LastRenderTime = (FApp::GetCurrentTime() > LastRenderTimeForTexture) ? float( FApp::GetCurrentTime() - LastRenderTimeForTexture ) : 0.0f;
-		MinDistance = MAX_STREAMINGDISTANCE;
-		bForceFullyLoad = Texture->ShouldMipLevelsBeForcedResident() || (ForceLoadRefCount > 0);
-		TextureLODBias = Texture->GetCachedLODBias();
-		bInFlight = false;
-		bReadyForStreaming = IsStreamingTexture( Texture ) && IsReadyForStreaming( Texture );
-		NumCinematicMipLevels = Texture->bUseCinematicMipLevels ? Texture->NumCinematicMipLevels : 0;
-	}
-
-	ETextureStreamingType GetStreamingType() const
-	{
-		if ( bUsesForcedHeuristics || bForceFullyLoad )
-		{
-			return StreamType_Forced;
-		}
-		else if ( bUsesDynamicHeuristics )
-		{
-			return StreamType_Dynamic;
-		}
-		else if ( bUsesStaticHeuristics )
-		{
-			return StreamType_Static;
-		}
-		else if ( bUsesOrphanedHeuristics )
-		{
-			return StreamType_Orphaned;
-		}
-		else if ( bUsesLastRenderHeuristics )
-		{
-			return StreamType_LastRenderTime;
-		}
-		return StreamType_Other;
-	}
-
-	/**
-	 * Checks whether a UTexture2D is ready for streaming.
-	 *
-	 * @param Texture	Texture to check
-	 * @return			true if the UTexture2D is ready to be streamed in or out
-	 */
-	static bool IsReadyForStreaming( UTexture2D* Texture )
-	{
-		return Texture->IsReadyForStreaming();
-	}
-
-	/**
-	 * Checks whether a UTexture2D is a streaming lightmap or shadowmap.
-	 *
-	 * @param Texture	Texture to check
-	 * @return			true if the UTexture2D is a streaming lightmap or shadowmap
-	 */
-	static bool IsStreamingLightmap( UTexture2D* Texture )
-	{
-		ULightMapTexture2D* Lightmap = Cast<ULightMapTexture2D>(Texture);
-		UShadowMapTexture2D* Shadowmap = Cast<UShadowMapTexture2D>(Texture);
-
-		if ( Lightmap && Lightmap->LightmapFlags & LMF_Streamed )
-		{
-			return true;
-		}
-		else if ( Shadowmap && Shadowmap->ShadowmapFlags & SMF_Streamed )
-		{
-			return true;
-		}
-		return false;
-	}
-
-	/**
-	 * Returns the amount of memory used by the texture given a specified number of mip-maps, in bytes.
-	 *
-	 * @param MipCount	Number of mip-maps to account for
-	 * @return			Total amount of memory used for the specified mip-maps, in bytes
-	 */
-	int32 GetSize( int32 InMipCount ) const
-	{
-		check(InMipCount >= 0 && InMipCount <= MAX_TEXTURE_MIP_COUNT);
-		return TextureSizes[ InMipCount - 1 ];
-	}
-
-	/**
-	 * Returns whether this primitive has been used rencently. Conservative.
-	 *
-	 * @return			true if visible
-	 */
-	FORCEINLINE bool IsVisible() const
-	{
-		return LastRenderTime < 5.f && LastRenderTime != MAX_FLT;
-	}
-	
-	/**
-	 * Return whether we allow to not stream the mips to fit in memory.
-	 *
-	 * @return			true if it is ok to drop the mips
-	 */
-	FORCEINLINE bool CanDropMips() const
-	{
-		return (LODGroup != TEXTUREGROUP_Terrain_Heightmap) && WantedMips > MinAllowedMips && !bForceFullyLoad;
-	}
-
-
-	/**
-	 * Return the wanted mips to use for priority computates
-	 *
-	 * @return		the WantedMips
-	 */
-	FORCEINLINE int32 GetWantedMipsForPriority() const
-	{
-		if (bHasSplitRequest && !bIsLastSplitRequest)
-		{
-			return WantedMips + 1;
-		}
-		else
-		{
-			return WantedMips;
-		}
-	}
-
-	/**
-	 * Calculates a retention priority value for the textures. Higher value means more important.
-	 * Retention is used determine which texture to drop when out of texture pool memory.
-	 * The retention value must be stable, meaning it must not depend on the current state of the streamer.
-	 * Not doing so could make the streamer go into a loop where is never stops dropping and loading different textures when out of budget.
-	 * @return		Priority value
-	 */
-	float CalcRetentionPriority( )
-	{
-		const int32 WantedMipsForPriority = GetWantedMipsForPriority();
-
-		bool bMustKeep = LODGroup == TEXTUREGROUP_Terrain_Heightmap || bForceFullyLoad || WantedMipsForPriority <= MinAllowedMips;
-		bool bIsVisible = !GStreamWithTimeFactor || IsVisible();
-		bool bAlreadyLowRez = /*(PixelPerTexel > 1.5 && Distance > 1) ||*/ bIsHLODTexture || bIsLightmap;
-		
-		// Don't consider dropping mips that give nothing. This is because streaming time has a high cost per texture.
-		// Dropping many textures for not so much memory would affect the time it takes the streamer to recover.
-		// We also don't want to prioritize by size, since bigger wanted mips are expected to have move visual impact.
-		bool bGivesOffMemory = GetSize(WantedMipsForPriority) - GetSize(WantedMipsForPriority - 1) >= 256 * 1024; 
-
-		float Priority = 0;
-		if (bMustKeep)			Priority += 1024.f;
-		if (bAlreadyLowRez)		Priority += 512.f;
-		if (!bGivesOffMemory)	Priority += 256.f;
-		if (bIsVisible)			Priority += 128.f;
-
-		return Priority;
-	}
-
-	/**
-	 * Calculates a load order priority value for the texture. Higher value means more important.
-	 * Load Order can depend on the current state of resident mips, because it will not affect the streamer stability.
-	 * @return		Priority value
-	 */
-	float CalcLoadOrderPriority()
-	{
-		const int32 WantedMipsForPriority = GetWantedMipsForPriority();
-
-		bool bMustLoadFirst = LODGroup == TEXTUREGROUP_Terrain_Heightmap || bForceFullyLoad;
-		bool bIsVisible = !GStreamWithTimeFactor || IsVisible();
-		bool bFarFromTargetMips = WantedMipsForPriority - ResidentMips > 2;
-		bool bBigOnScreen = WantedMipsForPriority > 9.f || bIsHLODTexture;
-
-		float Priority = 0;
-		if (bMustLoadFirst)		Priority += 1024.f;
-		if (!bIsLastSplitRequest)	Priority += 512.f; 
-		if (bIsVisible)				Priority += 256.f;
-		if (bFarFromTargetMips)		Priority += 128.f;
-		if (bBigOnScreen)			Priority += 64.f;
-		return Priority;
-	}
-
-		
-	/**
-	 * Not thread-safe: Sets the streaming index in the corresponding UTexture2D.
-	 * @param ArrayIndex	Index into the FStreamingManagerTexture::StreamingTextures array
-	 */
-	void SetStreamingIndex( int32 ArrayIndex )
-	{
-		Texture->StreamingIndex = ArrayIndex;
-	}
-
-	/** Update texture streaming. Returns whether it did anything */
-	bool UpdateMipCount(  FStreamingManagerTexture& Manager, FStreamingContext& Context );
-
-	/** Texture to manage. */
-	UTexture2D*		Texture;
-
-	/** Cached number of mip-maps in the UTexture2D mip array (including the base mip) */
-	int32			MipCount;
-	/** Cached number of mip-maps in memory (including the base mip) */
-	int32			ResidentMips;
-	/** Cached number of mip-maps requested by the streaming system. */
-	int32			RequestedMips;
-	/** Cached number of mip-maps we would like to have in memory. */
-	int32			WantedMips;
-	/** Legacy: Same as WantedMips, but not clamped by fudge factor. */
-	int32			PerfectWantedMips;
-	/** Same as MaxAllowedMips, but not clamped by the -reducepoolsize command line option. */
-	int32			MaxAllowedOptimalMips;
-#if STATS
-	/** Most number of mip-levels this texture has ever had resident in memory. */
-	int32			MostResidentMips;
-#endif
-	/** Minimum number of mip-maps that this texture must keep in memory. */
-	int32			MinAllowedMips;
-	/** Maximum number of mip-maps that this texture can have in memory. */
-	int32			MaxAllowedMips;
-	/** Cached memory sizes for each possible mipcount. */
-	int32			TextureSizes[MAX_TEXTURE_MIP_COUNT + 1];
-	/** Ref-count for how many ULevels want this texture fully-loaded. */
-	int32			ForceLoadRefCount;
-
-	/** Cached texture group. */
-	TextureGroup	LODGroup;
-	/** Cached LOD bias. */
-	int32			TextureLODBias;
-	/** Cached number of mipmaps that are not allowed to stream. */
-	int32			NumNonStreamingMips;
-	/** Cached number of cinematic (high-resolution) mip-maps. Normally not streamed in, unless the texture is forcibly fully loaded. */
-	int32			NumCinematicMipLevels;
-
-	/** Last time this texture was rendered, 0-based from GStartTime. */
-	float			LastRenderTime;
-	/** Distance to the closest instance of this texture, in world space units. Defaults to a specific value for non-staticmesh textures. */
-	float			MinDistance;
-	/** If non-zero, the most recent time an instance location was removed for this texture. */
-	double			InstanceRemovedTimestamp;
-
-	/** Set to FApp::GetCurrentTime() every time LastRenderTimeRefCount is modified. */
-	double			LastRenderTimeRefCountTimestamp;
-	/** Current number of instances that need LRT heuristics for this texture. */
-	int32			LastRenderTimeRefCount;
-
-	/**
-	 * Temporary boost of the streaming distance factor.
-	 * This factor is automatically reset to 1.0 after it's been used for mip-calculations.
-	 */
-	float			BoostFactor;
-
-	/** Whether the texture should be forcibly fully loaded. */
-	uint32			bForceFullyLoad : 1;
-	/** Whether the texture is ready to be streamed in/out (cached from IsReadyForStreaming()). */
-	uint32			bReadyForStreaming : 1;
-	/** Whether the texture is currently being streamed in/out. */
-	uint32			bInFlight : 1;
-	/** Whether this is a streaming lightmap or shadowmap (cached from IsStreamingLightmap()). */
-	uint32			bIsStreamingLightmap : 1;
-	/** Whether this is a lightmap or shadowmap. */
-	uint32			bIsLightmap : 1;
-	/** Whether this texture uses StaticTexture heuristics. */
-	uint32			bUsesStaticHeuristics : 1;
-	/** Whether this texture uses DynamicTexture heuristics. */
-	uint32			bUsesDynamicHeuristics : 1;
-	/** Whether this texture uses LastRenderTime heuristics. */
-	uint32			bUsesLastRenderHeuristics : 1;
-	/** Whether this texture uses ForcedIntoMemory heuristics. */
-	uint32			bUsesForcedHeuristics : 1;
-	/** Whether this texture uses the OrphanedTexture heuristics. */
-	uint32			bUsesOrphanedHeuristics : 1;
-	/** Whether this texture has a split request strategy. */
-	uint32			bHasSplitRequest : 1;
-	/** Whether this is the second request to converge to the wanted mip. */
-	uint32			bIsLastSplitRequest : 1;
-	/** Can this texture be affected by mip bias? */
-	uint32			bCanBeAffectedByMipBias : 1;
-	/** Is this texture a HLOD texture? */
-	//@TODO: Temporary: Done because there is currently no texture group for HLOD textures and adding one in a branch isn't ideal; this should be removed / replaced once there is a real HLOD texture group
-	uint32			bIsHLODTexture : 1;
-};
-
 // ----------------------------------------------------------------------
 
 FFloatMipLevel::FFloatMipLevel() 
@@ -955,749 +374,142 @@ bool TrackTextureEvent( FStreamingTexture* StreamingTexture, UTexture2D* Texture
 #endif
 
 /*-----------------------------------------------------------------------------
-	Asynchronous texture streaming.
+	IStreamingManager implementation.
 -----------------------------------------------------------------------------*/
 
-#if STATS
-struct FTextureStreamingStats
+static FStreamingManagerCollection* StreamingManagerCollection = NULL;
+
+FStreamingManagerCollection& IStreamingManager::Get()
 {
-	FTextureStreamingStats( UTexture2D* InTexture2D, ETextureStreamingType InType, int32 InResidentMips, int32 InWantedMips, int32 InMostResidentMips, int32 InResidentSize, int32 InWantedSize, int32 InMaxSize, int32 InMostResidentSize, float InBoostFactor, float InPriority, int32 InTextureIndex )
-	:	TextureName( InTexture2D->GetFullName() )
-	,	SizeX( InTexture2D->GetSizeX() )
-	,	SizeY( InTexture2D->GetSizeY() )
-	,	NumMips( InTexture2D->GetNumMips() )
-	,	LODBias( InTexture2D->GetCachedLODBias() )
-	,	LastRenderTime( FMath::Clamp( InTexture2D->Resource ? (FApp::GetCurrentTime() - InTexture2D->Resource->LastRenderTime) : 1000000.0, 0.0, 1000000.0) )
-	,	StreamType( InType )
-	,	ResidentMips( InResidentMips )
-	,	WantedMips( InWantedMips )
-	,	MostResidentMips( InMostResidentMips )
-	,	ResidentSize( InResidentSize )
-	,	WantedSize( InWantedSize )
-	,	MaxSize( InMaxSize )
-	,	MostResidentSize( InMostResidentSize )
-	,	BoostFactor( InBoostFactor )
-	,	Priority( InPriority )
-	,	TextureIndex( InTextureIndex )
+	if (StreamingManagerCollection == NULL)
 	{
+		StreamingManagerCollection = new FStreamingManagerCollection();
 	}
-	/** Mirror of UTexture2D::GetName() */
-	FString		TextureName;
-	/** Mirror of UTexture2D::SizeX */
-	int32		SizeX;
-	/** Mirror of UTexture2D::SizeY */
-	int32		SizeY;
-	/** Mirror of UTexture2D::Mips.Num() */
-	int32		NumMips;
-	/** Mirror of UTexture2D::GetCachedLODBias() */
-	int32		LODBias;
-	/** How many seconds since it was last rendered: FApp::GetCurrentTime() - UTexture2D::Resource->LastRenderTime */
-	float		LastRenderTime;
-	/** What streaming heuristics were primarily used for this texture. */
-	ETextureStreamingType	StreamType;
-	/** Number of resident mip levels. */
-	int32		ResidentMips;
-	/** Number of wanted mip levels. */
-	int32		WantedMips;
-	/** Most number of mip-levels this texture has ever had resident in memory. */
-	int32		MostResidentMips;
-	/** Number of bytes currently in memory. */
-	int32		ResidentSize;
-	/** Number of bytes we want in memory. */
-	int32		WantedSize;
-	/** Number of bytes we could potentially stream in. */
-	int32		MaxSize;
-	/** Most number of bytes this texture has ever had resident in memory. */
-	int32		MostResidentSize;
-	/** Temporary boost of the streaming distance factor. */
-	float		BoostFactor;
-	/** Texture priority */
-	float		Priority;
-	/** Index into the FStreamingManagerTexture::StreamingTextures array. */
-	int32		TextureIndex;
-};
-#endif
+	return *StreamingManagerCollection;
+}
 
-/**
- * Helper struct for temporary information for one frame of processing texture streaming.
- */
-struct FStreamingContext
+void IStreamingManager::Shutdown()
 {
-	FStreamingContext( bool bProcessEverything, UTexture2D* IndividualStreamingTexture, bool bInCollectTextureStats )
+	if (StreamingManagerCollection != NULL)
 	{
-		Reset( bProcessEverything, IndividualStreamingTexture, bInCollectTextureStats );
+		delete StreamingManagerCollection;
+		StreamingManagerCollection = (FStreamingManagerCollection*)-1;//Force Error if manager used after shutdown
 	}
+}
 
-	/**
-	 * Initializes all variables for the one frame.
-	 * @param bProcessEverything			If true, process all resources with no throttling limits
-	 * @param IndividualStreamingTexture	A specific texture to be fully processed this frame, or NULL
-	 * @param bInCollectTextureStats		Whether to fill in the TextureStats array this frame
-	 */
-	void Reset( bool bProcessEverything, UTexture2D* IndividualStreamingTexture, bool bInCollectTextureStats )
+bool IStreamingManager::HasShutdown()
+{
+	return StreamingManagerCollection == (FStreamingManagerCollection*)-1;
+}
+
+/**
+ * Adds the passed in view information to the static array.
+ *
+ * @param ViewInfos				[in/out] Array to add the view to
+ * @param ViewOrigin			View origin
+ * @param ScreenSize			Screen size
+ * @param FOVScreenSize			Screen size taking FOV into account
+ * @param BoostFactor			A factor that affects all streaming distances for this location. 1.0f is default. Higher means higher-resolution textures and vice versa.
+ * @param bOverrideLocation		Whether this is an override location, which forces the streaming system to ignore all other regular locations
+ * @param Duration				How long the streaming system should keep checking this location (in seconds). 0 means just for the next Tick.
+ * @param InActorToBoost		Optional pointer to an actor who's textures should have their streaming priority boosted
+ */
+void IStreamingManager::AddViewInfoToArray( TArray<FStreamingViewInfo> &ViewInfos, const FVector& ViewOrigin, float ScreenSize, float FOVScreenSize, float BoostFactor, bool bOverrideLocation, float Duration, TWeakObjectPtr<AActor> InActorToBoost )
+{
+	// Check for duplicates and existing overrides.
+	bool bShouldAddView = true;
+	for ( int32 ViewIndex=0; ViewIndex < ViewInfos.Num(); ++ViewIndex )
 	{
-		bCollectTextureStats							= bInCollectTextureStats;
-		ThisFrameTotalRequestSize						= 0;
-		ThisFrameTotalLightmapRequestSize				= 0;
-		ThisFrameNumStreamingTextures					= 0;
-		ThisFrameNumRequestsInCancelationPhase			= 0;
-		ThisFrameNumRequestsInUpdatePhase				= 0;
-		ThisFrameNumRequestsInFinalizePhase				= 0;
-		ThisFrameTotalIntermediateTexturesSize			= 0;
-		ThisFrameNumIntermediateTextures				= 0;
-		ThisFrameTotalStreamingTexturesSize				= 0;
-		ThisFrameTotalStreamingTexturesMaxSize			= 0;
-		ThisFrameTotalLightmapMemorySize				= 0;
-		ThisFrameTotalLightmapDiskSize					= 0;
-		ThisFrameTotalHLODMemorySize					= 0;
-		ThisFrameTotalHLODDiskSize						= 0;
-		ThisFrameTotalMipCountIncreaseRequestsInFlight	= 0;
-		ThisFrameOptimalWantedSize						= 0;
-		ThisFrameTotalStaticTextureHeuristicSize		= 0;
-		ThisFrameTotalDynamicTextureHeuristicSize		= 0;
-		ThisFrameTotalLastRenderHeuristicSize			= 0;
-		ThisFrameTotalForcedHeuristicSize				= 0;
-
-		AvailableNow = 0;
-		AvailableLater = 0;
-		AvailableTempMemory = 0;
-
-		STAT( TextureStats.Empty() );
-
-		AllocatedMemorySize	= INDEX_NONE;
-		PendingMemoryAdjustment = INDEX_NONE;
-		
-		// Available texture memory, if supported by RHI. This stat is async as the renderer allocates the memory in its own thread so we
-		// only query once and roughly adjust the values as needed.
-		FTextureMemoryStats Stats;
-		RHIGetTextureMemoryStats(Stats);
-		
-		bRHISupportsMemoryStats = Stats.IsUsingLimitedPoolSize();
-
-		// Update stats if supported.
-		if( bRHISupportsMemoryStats )
-		{
-			AllocatedMemorySize	= Stats.AllocatedMemorySize;
-			PendingMemoryAdjustment = Stats.PendingMemoryAdjustment;
-	
-			// set total size for the pool (used to available)
-			SET_MEMORY_STAT(STAT_TexturePoolAllocatedSize, AllocatedMemorySize);
-			SET_MEMORY_STAT(STAT_TexturePoolSize, Stats.TexturePoolSize);
-			SET_MEMORY_STAT(MCR_TexturePool, Stats.TexturePoolSize);
-		}
-		else
+		FStreamingViewInfo& ViewInfo = ViewInfos[ ViewIndex ];
+		if ( ViewOrigin.Equals( ViewInfo.ViewOrigin, 0.5f ) &&
+			FMath::IsNearlyEqual( ScreenSize, ViewInfo.ScreenSize ) &&
+			FMath::IsNearlyEqual( FOVScreenSize, ViewInfo.FOVScreenSize ) &&
+			ViewInfo.bOverrideLocation == bOverrideLocation )
 		{
-			SET_MEMORY_STAT(STAT_TexturePoolAllocatedSize,0);
-			SET_MEMORY_STAT(STAT_TexturePoolSize, 0);
-			SET_MEMORY_STAT(MCR_TexturePool, 0);
+			// Update duration
+			ViewInfo.Duration = Duration;
+			// Overwrite BoostFactor if it isn't default 1.0
+			ViewInfo.BoostFactor = FMath::IsNearlyEqual(BoostFactor, 1.0f) ? ViewInfo.BoostFactor : BoostFactor;
+			bShouldAddView = false;
 		}
 	}
 
-	/**
-	 * Adds in the stats from another context.
-	 *
-	 * @param Other		Context to add stats from
-	 */
-	void AddStats( const FStreamingContext& Other )
+	if ( bShouldAddView )
 	{
-		ThisFrameTotalRequestSize						+= Other.ThisFrameTotalRequestSize;						
-		ThisFrameTotalLightmapRequestSize				+= Other.ThisFrameTotalLightmapRequestSize;
-		ThisFrameNumStreamingTextures					+= Other.ThisFrameNumStreamingTextures;
-		ThisFrameNumRequestsInCancelationPhase			+= Other.ThisFrameNumRequestsInCancelationPhase;
-		ThisFrameNumRequestsInUpdatePhase				+= Other.ThisFrameNumRequestsInUpdatePhase;
-		ThisFrameNumRequestsInFinalizePhase				+= Other.ThisFrameNumRequestsInFinalizePhase;
-		ThisFrameTotalIntermediateTexturesSize			+= Other.ThisFrameTotalIntermediateTexturesSize;
-		ThisFrameNumIntermediateTextures				+= Other.ThisFrameNumIntermediateTextures;
-		ThisFrameTotalStreamingTexturesSize				+= Other.ThisFrameTotalStreamingTexturesSize;
-		ThisFrameTotalStreamingTexturesMaxSize			+= Other.ThisFrameTotalStreamingTexturesMaxSize;
-		ThisFrameTotalLightmapMemorySize				+= Other.ThisFrameTotalLightmapMemorySize;
-		ThisFrameTotalLightmapDiskSize					+= Other.ThisFrameTotalLightmapDiskSize;
-		ThisFrameTotalHLODMemorySize					+= Other.ThisFrameTotalHLODMemorySize;
-		ThisFrameTotalHLODDiskSize						+= Other.ThisFrameTotalHLODDiskSize;
-		ThisFrameTotalMipCountIncreaseRequestsInFlight	+= Other.ThisFrameTotalMipCountIncreaseRequestsInFlight;
-		ThisFrameOptimalWantedSize						+= Other.ThisFrameOptimalWantedSize;
-		ThisFrameTotalStaticTextureHeuristicSize		+= Other.ThisFrameTotalStaticTextureHeuristicSize;
-		ThisFrameTotalDynamicTextureHeuristicSize		+= Other.ThisFrameTotalDynamicTextureHeuristicSize;
-		ThisFrameTotalLastRenderHeuristicSize			+= Other.ThisFrameTotalLastRenderHeuristicSize;
-		ThisFrameTotalForcedHeuristicSize				+= Other.ThisFrameTotalForcedHeuristicSize;
-		bCollectTextureStats							= bCollectTextureStats || Other.bCollectTextureStats;
-		STAT( TextureStats.Append( Other.TextureStats ) );
+		new (ViewInfos) FStreamingViewInfo( ViewOrigin, ScreenSize, FOVScreenSize, BoostFactor, bOverrideLocation, Duration, InActorToBoost );
 	}
+}
 
-	/** Whether the platform RHI supports texture memory stats. */
-	bool bRHISupportsMemoryStats;
-	/** Currently allocated number of bytes in the texture pool, or INDEX_NONE if bRHISupportsMemoryStats is false. */
-	int64 AllocatedMemorySize;
-	/** Pending Adjustments to allocated texture memory, due to async reallocations, etc. */
-	int64 PendingMemoryAdjustment;
-	/** Whether to fill in TextureStats this frame. */
-	bool bCollectTextureStats;
-#if STATS
-	/** Stats for all textures. */
-	TArray<FTextureStreamingStats>	TextureStats;
-#endif
-
-	// Stats for this frame.
-	uint64 ThisFrameTotalRequestSize;
-	uint64 ThisFrameTotalLightmapRequestSize;
-	uint64 ThisFrameNumStreamingTextures;
-	uint64 ThisFrameNumRequestsInCancelationPhase;
-	uint64 ThisFrameNumRequestsInUpdatePhase;
-	uint64 ThisFrameNumRequestsInFinalizePhase;
-	uint64 ThisFrameTotalIntermediateTexturesSize;
-	uint64 ThisFrameNumIntermediateTextures;
-	uint64 ThisFrameTotalStreamingTexturesSize;
-	uint64 ThisFrameTotalStreamingTexturesMaxSize;
-	uint64 ThisFrameTotalLightmapMemorySize;
-	uint64 ThisFrameTotalLightmapDiskSize;
-	uint64 ThisFrameTotalHLODMemorySize;
-	uint64 ThisFrameTotalHLODDiskSize;
-	uint64 ThisFrameTotalMipCountIncreaseRequestsInFlight;
-	uint64 ThisFrameOptimalWantedSize;
-	/** Number of bytes using StaticTexture heuristics this frame, currently in memory. */
-	uint64 ThisFrameTotalStaticTextureHeuristicSize;
-	/** Number of bytes using DynmicTexture heuristics this frame, currently in memory. */
-	uint64 ThisFrameTotalDynamicTextureHeuristicSize;
-	/** Number of bytes using LastRenderTime heuristics this frame, currently in memory. */
-	uint64 ThisFrameTotalLastRenderHeuristicSize;
-	/** Number of bytes using ForcedIntoMemory heuristics this frame, currently in memory. */
-	uint64 ThisFrameTotalForcedHeuristicSize;
-
-	/** Conservative amount of memory currently available or available after some pending requests are completed. */
-	int64 AvailableNow;
-	/** Amount of memory expected to be available once all streaming is completed. Depends on WantedMips */
-	int64 AvailableLater;
-	/** Amount of memory immediatly available temporary streaming data (pending textures). Depends on the delta between RequestedMips and ResidentMips */
-	int64 AvailableTempMemory;
-};
-
-bool FStreamingTexture::UpdateMipCount(  FStreamingManagerTexture& Manager, FStreamingContext& Context )
+/**
+ * Remove view infos with the same location from the given array.
+ *
+ * @param ViewInfos				[in/out] Array to remove the view from
+ * @param ViewOrigin			View origin
+ */
+void IStreamingManager::RemoveViewInfoFromArray( TArray<FStreamingViewInfo> &ViewInfos, const FVector& ViewOrigin )
 {
-	if ( !Texture || !bReadyForStreaming )
-	{
-		return false;
-	}
-
-	// If there is a pending load that attempts to load unrequired data, 
-	// or if there is a pending unload that attempts to unload required data, try to cancel it.
-	if ( bInFlight && 
-		(( RequestedMips > ResidentMips && RequestedMips > WantedMips ) || 
-		 ( RequestedMips < ResidentMips && RequestedMips < WantedMips )) )
+	for ( int32 ViewIndex=0; ViewIndex < ViewInfos.Num(); ++ViewIndex )
 	{
-		// Try to cancel, but don't try anything else before next update.
-		bReadyForStreaming = false;
-
-		if ( Texture->CancelPendingMipChangeRequest() )
+		FStreamingViewInfo& ViewInfo = ViewInfos[ ViewIndex ];
+		if ( ViewOrigin.Equals( ViewInfo.ViewOrigin, 0.5f ) )
 		{
-			bInFlight = false;
-			RequestedMips = ResidentMips;
-			return true;
+			ViewInfos.RemoveAtSwap( ViewIndex-- );
 		}
 	}
-	// If we should load/unload some mips.
-	else if ( !bInFlight && ResidentMips != WantedMips )
-	{
-		check(RequestedMips == ResidentMips);
-
-		// Try to unload, but don't try anything else before next update.
-		bReadyForStreaming = false;
-
-		if ( Texture->PendingMipChangeRequestStatus.GetValue() == TexState_ReadyFor_Requests )
-		{
-			check(!Texture->bHasCancelationPending);
-
-			bool bHasEnoughTempMemory = GetSize(WantedMips) <= Context.AvailableTempMemory;
-			bool bHasEnoughAvailableNow = WantedMips <= ResidentMips || GetSize(WantedMips) <= Context.AvailableNow;
-
-			// Check if there is enough space for temp memory.
-			if ( bHasEnoughTempMemory && bHasEnoughAvailableNow )
-			{
-				FTexture2DResource* Texture2DResource = (FTexture2DResource*)Texture->Resource;
-
-				if (WantedMips > ResidentMips) // If it is a load request.
-				{
-					// Manually update size as allocations are deferred/ happening in rendering thread.
-					int64 LoadRequestSize = GetSize(WantedMips) - GetSize(ResidentMips);
-					Context.ThisFrameTotalRequestSize			+= LoadRequestSize;
-					Context.ThisFrameTotalLightmapRequestSize	+= bIsStreamingLightmap ? LoadRequestSize : 0;
-					Context.AvailableNow						-= LoadRequestSize;
-				}
+}
 
-				Context.AvailableTempMemory -= GetSize(WantedMips);
+#if STREAMING_LOG_VIEWCHANGES
+TArray<FStreamingViewInfo> GPrevViewLocations;
+#endif
 
-				bInFlight = true;
-				RequestedMips = WantedMips;
-				Texture->RequestedMips = RequestedMips;
+/**
+ * Sets up the CurrentViewInfos array based on PendingViewInfos, LastingViewInfos and SlaveLocations.
+ * Removes out-dated LastingViewInfos.
+ *
+ * @param DeltaTime		Time since last call in seconds
+ */
+void IStreamingManager::SetupViewInfos( float DeltaTime )
+{
+	// Reset CurrentViewInfos
+	CurrentViewInfos.Empty( PendingViewInfos.Num() + LastingViewInfos.Num() + SlaveLocations.Num() );
 
-				Texture2DResource->BeginUpdateMipCount( bForceFullyLoad );
-				TrackTextureEvent( this, Texture, bForceFullyLoad, &Manager );
-				return true;
-			}
-		}
-		else
-		{
-			// Did UpdateStreamingTextures() miss a texture? Should never happen!
-			UE_LOG(LogContentStreaming, Warning, TEXT("BeginUpdateMipCount failure! PendingMipChangeRequestStatus == %d, Resident=%d, Requested=%d, Wanted=%d"), Texture->PendingMipChangeRequestStatus.GetValue(), Texture->ResidentMips, Texture->RequestedMips, WantedMips );
-		}
-	}
-	return false;
-}
+	bool bHaveMultiplePlayerViews = (PendingViewInfos.Num() > 1) ? true : false;
 
-/** Async work class for calculating priorities for all textures. */
-// this could implement a better abandon, but give how it is used, it does that anyway via the abort mechanism
-class FAsyncTextureStreaming : public FNonAbandonableTask
-{
-public:
-	FAsyncTextureStreaming( FStreamingManagerTexture* InStreamingManager )
-	:	StreamingManager( *InStreamingManager )
-	,	ThreadContext( false, NULL, false )
-	,	bAbort( false )
+	// Add the slave locations.
+	float ScreenSize = 1280.0f;
+	float FOVScreenSize = ScreenSize / FMath::Tan( 80.0f * float(PI) / 360.0f );
+	if ( PendingViewInfos.Num() > 0 )
 	{
-		Reset(false);
+		ScreenSize = PendingViewInfos[0].ScreenSize;
+		FOVScreenSize = PendingViewInfos[0].FOVScreenSize;
 	}
-
-	/** Resets the state to start a new async job. */
-	void Reset( bool bCollectTextureStats )
+	else if ( LastingViewInfos.Num() > 0 )
 	{
-		bAbort = false;
-		ThreadContext.Reset( false, NULL, bCollectTextureStats );
-		ThreadStats.Reset();
-		PrioritizedTextures.Empty( StreamingManager.StreamingTextures.Num() );
-		PrioTexIndicesSortedByLoadOrder.Empty( StreamingManager.StreamingTextures.Num() );
+		ScreenSize = LastingViewInfos[0].ScreenSize;
+		FOVScreenSize = LastingViewInfos[0].FOVScreenSize;
 	}
-
-	/** Notifies the async work that it should abort the thread ASAP. */
-	void Abort()
+	// Add them to the appropriate array (pending views or lasting views).
+	for ( int32 SlaveLocationIndex=0; SlaveLocationIndex < SlaveLocations.Num(); SlaveLocationIndex++ )
 	{
-		bAbort = true;
+		const FSlaveLocation& SlaveLocation = SlaveLocations[ SlaveLocationIndex ];
+		AddViewInformation( SlaveLocation.Location, ScreenSize, FOVScreenSize, SlaveLocation.BoostFactor, SlaveLocation.bOverrideLocation, SlaveLocation.Duration );
 	}
 
-	/** Whether the async work is being aborted. Can be used in conjunction with IsDone() to see if it has finished. */
-	bool IsAborted() const
+	// Apply a split-screen factor if we have multiple players on the same machine, and they currently have individual views.
+	float SplitScreenFactor = 1.0f;
+	
+	if ( bHaveMultiplePlayerViews && GEngine->IsSplitScreen(NULL) )
 	{
-		return bAbort;
+		SplitScreenFactor = 0.75f;
 	}
 
-	/** Statistics for the async work. */
-	struct FAsyncStats
+	// Should we use override views this frame? (If we have both a fullscreen movie and an override view.)
+	bool bUseOverrideViews = false;
+	bool bIsMoviePlaying = false;//GetMoviePlayer()->IsMovieCurrentlyPlaying();
+	if ( bIsMoviePlaying )
 	{
-		FAsyncStats()
-		{
-			Reset();
-		}
-		/** Resets the statistics to zero. */
-		void Reset()
-		{
-			TotalResidentSize = 0;
-			TotalPossibleResidentSize = 0;
-			TotalWantedMipsSize = 0;
-			TotalTempMemorySize = 0;
-			TotalRequiredSize = 0;
-			PendingStreamInSize = 0;
-			PendingStreamOutSize = 0;
-			WantedInSize = 0;
-			WantedOutSize = 0;
-			NumWantingTextures = 0;
-			NumVisibleTexturesWithLowResolutions = 0;
-		}
-		/** Total number of bytes currently in memory */
-		int64 TotalResidentSize;
-		/** Total number of bytes that could possibly be used according to the worst case between resident size (current state) and requested size (pending requests) */
-		int64 TotalPossibleResidentSize;
-		/** Total number of bytes required to satisfy wanted mips.  Not taking into account split requests. */
-		int64 TotalWantedMipsSize;
-		/** Total number of bytes used for updating textures (conservative value since it is platform specific */
-		int64 TotalTempMemorySize;
-		/** Total number of bytes required, using PerfectWantedMips */
-		int64 TotalRequiredSize;
-		/** Currently being streamed in, in bytes. */
-		SIZE_T PendingStreamInSize;
-		/** Currently being streamed out, in bytes. */
-		SIZE_T PendingStreamOutSize;
-		/** How much we want to stream in, in bytes. */
-		int64 WantedInSize;
-		/** How much we want to stream out, in bytes.  */
-		int64 WantedOutSize;
-		/** Number of textures that want more mips. */
-		int64 NumWantingTextures;
-		/** Number of textures that are visibles and still low rez. */
-		int64 NumVisibleTexturesWithLowResolutions;
-	};
-
-	void ClearRemovedTextureReferences()
-	{
-		// CleanUp Prioritized Textures, account for removed textures.
-		for (int32 TexPrioIndex = 0; TexPrioIndex < PrioritizedTextures.Num(); ++TexPrioIndex)
-		{
-			FTexturePriority& TexturePriority = PrioritizedTextures[ TexPrioIndex ];
-
-			// Because deleted textures can shrink StreamingTextures, we need to check that the index is in range.
-			// Even if it is, because of the RemoveSwap logic, it could refer to the wrong texture now.
-			if (!StreamingManager.StreamingTextures.IsValidIndex(TexturePriority.TextureIndex) || TexturePriority.Texture != StreamingManager.StreamingTextures[TexturePriority.TextureIndex].Texture)
-			{
-				TexturePriority.TextureIndex = INDEX_NONE;
-				TexturePriority.Texture = nullptr;
-			}
-		}
-
-	}
-
-
-	/** Returns the resulting priorities, matching the FStreamingManagerTexture::StreamingTextures array. */
-	const TArray<FTexturePriority>& GetPrioritizedTextures() const
-	{
-		return PrioritizedTextures;
-	}
-
-	const TArray<int32>& GetPrioTexIndicesSortedByLoadOrder() const
-	{
-		return PrioTexIndicesSortedByLoadOrder;
-	}
-
-	/** Returns the context (temporary info) used for the async work. */
-	const FStreamingContext& GetContext() const
-	{
-		return ThreadContext;
-	}
-
-	/** Returns the thread statistics. */
-	const FAsyncStats& GetStats() const
-	{
-		return ThreadStats;
-	}
-
-private:
-	friend class FAsyncTask<FAsyncTextureStreaming>;
-	/** Performs the async work. */
-	void DoWork()
-	{
-		DECLARE_SCOPE_CYCLE_COUNTER(TEXT("FAsyncTextureStreaming::DoWork"), STAT_AsyncTextureStreaming_DoWork, STATGROUP_StreamingDetails);
-		PrioritizedTextures.Empty( StreamingManager.StreamingTextures.Num() );
-		PrioTexIndicesSortedByLoadOrder.Empty( StreamingManager.StreamingTextures.Num() );
-
-		const int32 SplitRequestSizeThreshold = CVarStreamingSplitRequestSizeThreshold.GetValueOnAnyThread();
-
-		const bool bUseNewMetrics = CVarStreamingUseNewMetrics.GetValueOnAnyThread() > 0;
-		const float MaxEffectiveScreenSize = CVarStreamingScreenSizeEffectiveMax.GetValueOnAnyThread();
-		StreamingManager.ThreadSettings.TextureBoundsVisibility->ComputeBoundsViewInfos(StreamingManager.ThreadSettings.ThreadViewInfos, bUseNewMetrics, MaxEffectiveScreenSize);
-
-		// Number of textures that want more mips.
-		ThreadStats.NumWantingTextures = 0;
-		ThreadStats.NumVisibleTexturesWithLowResolutions = 0;
-
-		for ( int32 Index=0; Index < StreamingManager.StreamingTextures.Num() && !IsAborted(); ++Index )
-		{
-			FStreamingTexture& StreamingTexture = StreamingManager.StreamingTextures[ Index ];
-
-			StreamingTexture.bUsesStaticHeuristics = false;
-			StreamingTexture.bUsesDynamicHeuristics = false;
-			StreamingTexture.bUsesLastRenderHeuristics = false;
-			StreamingTexture.bUsesForcedHeuristics = false;
-			StreamingTexture.bUsesOrphanedHeuristics = false;
-			StreamingTexture.bHasSplitRequest  = false;
-			StreamingTexture.bIsLastSplitRequest = false;
-
-				// Figure out max number of miplevels allowed by LOD code.
-				StreamingManager.CalcMinMaxMips( StreamingTexture );
-
-				// Determine how many mips this texture should have in memory.
-				StreamingManager.CalcWantedMips( StreamingTexture );
-
-			ThreadStats.TotalResidentSize += StreamingTexture.GetSize( StreamingTexture.ResidentMips );
-			ThreadStats.TotalPossibleResidentSize += StreamingTexture.GetSize( FMath::Max<int32>(StreamingTexture.ResidentMips, StreamingTexture.RequestedMips) );
-			ThreadStats.TotalWantedMipsSize += StreamingTexture.GetSize( StreamingTexture.WantedMips );
-			if (StreamingTexture.RequestedMips != StreamingTexture.ResidentMips)
-			{
-				// UpdateMips load texture mip changes in temporary memory.
-				ThreadStats.TotalTempMemorySize += StreamingTexture.GetSize( StreamingTexture.RequestedMips );
-			}
-
-
-			// Split the request in two if the last mip to load is too big. Helps putprioriHelps converge the image quality.
-			if (SplitRequestSizeThreshold != 0 && 
-				StreamingTexture.LODGroup != TEXTUREGROUP_Terrain_Heightmap &&
-				StreamingTexture.WantedMips > StreamingTexture.ResidentMips && 
-				StreamingTexture.GetSize(StreamingTexture.WantedMips) - StreamingTexture.GetSize(StreamingTexture.WantedMips - 1) >= SplitRequestSizeThreshold) 
-			{
-				// The texture fits for a 2 request stream strategy, though it is not garantied that there will be 2 requests (since it could be only one mip aways)
-				StreamingTexture.bHasSplitRequest = true;
-				if (StreamingTexture.WantedMips == StreamingTexture.ResidentMips + 1)
-				{
-					StreamingTexture.bIsLastSplitRequest = true;
-				}
-				else
-				{
-					// Don't stream the last mips now.
-					--StreamingTexture.WantedMips;
-				}
-			}
-			if ( StreamingTexture.WantedMips > StreamingTexture.ResidentMips )
-			{
-				ThreadStats.NumWantingTextures++;
-			}
-
-			bool bTrackedTexture = TrackTextureEvent( &StreamingTexture, StreamingTexture.Texture, StreamingTexture.bForceFullyLoad, &StreamingManager );
-
-			// Add to sort list, if it wants to stream in or could potentially stream out.
-			if ( StreamingTexture.MaxAllowedMips > StreamingTexture.NumNonStreamingMips)
-			{
-				const UTexture2D* Texture = StreamingTexture.Texture; // This needs to go on the stack in case the main thread updates it at the same time (for the condition consistency).
-				if (Texture)
-				{
-					new (PrioritizedTextures) FTexturePriority( StreamingTexture.CanDropMips(), StreamingTexture.CalcRetentionPriority(), StreamingTexture.CalcLoadOrderPriority(), Index, Texture );
-				}
-			}
-
-			if (StreamingTexture.IsVisible() && StreamingTexture.GetWantedMipsForPriority() - StreamingTexture.ResidentMips > 1)
-			{
-				++ThreadStats.NumVisibleTexturesWithLowResolutions;
-			}
-
-				// Accumulate streaming numbers.
-			int64 ResidentTextureSize = StreamingTexture.GetSize( StreamingTexture.ResidentMips );
-			int64 WantedTextureSize = StreamingTexture.GetSize( StreamingTexture.WantedMips );
-			if ( StreamingTexture.bInFlight )
-			{
-				int64 RequestedTextureSize = StreamingTexture.GetSize( StreamingTexture.RequestedMips );
-				if ( StreamingTexture.RequestedMips > StreamingTexture.ResidentMips )
-				{
-					ThreadStats.PendingStreamInSize += FMath::Abs(RequestedTextureSize - ResidentTextureSize);
-				}
-				else
-				{
-					ThreadStats.PendingStreamOutSize += FMath::Abs(RequestedTextureSize - ResidentTextureSize);
-				}
-			}
-			else
-			{
-				if ( StreamingTexture.WantedMips > StreamingTexture.ResidentMips )
-				{
-					ThreadStats.WantedInSize += FMath::Abs(WantedTextureSize - ResidentTextureSize);
-				}
-				else
-				{
-					// Counting on shrinking reallocation.
-					ThreadStats.WantedOutSize += FMath::Abs(WantedTextureSize - ResidentTextureSize);
-				}
-			}
-
-			const int32 PerfectWantedTextureSize = StreamingTexture.GetSize( StreamingTexture.PerfectWantedMips );
-			ThreadStats.TotalRequiredSize += PerfectWantedTextureSize;
-			StreamingManager.UpdateFrameStats( ThreadContext, StreamingTexture, Index );
-
-			// Reset the boost factor
-			StreamingTexture.BoostFactor = 1.0f;
-		}
-
-		// Sort the candidates.
-		struct FCompareTexturePriority // Smaller retention priority first because we drop them first.
-		{
-			FORCEINLINE bool operator()( const FTexturePriority& A, const FTexturePriority& B ) const
-			{
-				if ( A.RetentionPriority < B.RetentionPriority )
-				{
-					return true;
-				}
-				else if ( A.RetentionPriority == B.RetentionPriority )
-				{
-					return ( A.TextureIndex < B.TextureIndex );
-				}
-				return false;
-			}
-		};
-		// Texture are sorted by retention priority since the ASYNC loading will resort request by load priority
-		PrioritizedTextures.Sort( FCompareTexturePriority() );
-
-
-		// Sort the candidates.
-		struct FCompareTexturePriorityByLoadOrder // Bigger load order priority first because we load them first.
-		{
-			FCompareTexturePriorityByLoadOrder(const TArray<FTexturePriority>& InPrioritizedTextures) : PrioritizedTextures(InPrioritizedTextures) {}
-			const TArray<FTexturePriority>& PrioritizedTextures;
-
-			FORCEINLINE bool operator()( int32 IndexA, int32 IndexB ) const
-			{
-				const FTexturePriority& A = PrioritizedTextures[IndexA];
-				const FTexturePriority& B = PrioritizedTextures[IndexB];
-
-				if ( A.LoadOrderPriority > B.LoadOrderPriority )
-				{
-					return true;
-				}
-				else if ( A.LoadOrderPriority == B.LoadOrderPriority )
-				{
-					return ( A.TextureIndex < B.TextureIndex );
-				}
-				return false;
-			}
-		};
-
-		PrioTexIndicesSortedByLoadOrder.AddUninitialized(PrioritizedTextures.Num());
-		for (int32 Index = 0; Index < PrioTexIndicesSortedByLoadOrder.Num(); ++Index)
-		{
-			PrioTexIndicesSortedByLoadOrder[Index] = Index;
-		}
-		PrioTexIndicesSortedByLoadOrder.Sort(FCompareTexturePriorityByLoadOrder(PrioritizedTextures));
-	}
-
-	FORCEINLINE TStatId GetStatId() const
-	{
-		RETURN_QUICK_DECLARE_CYCLE_STAT(FAsyncTextureStreaming, STATGROUP_ThreadPoolAsyncTasks);
-	}
-
-	/** Reference to the owning streaming manager, for accessing the thread-safe data. */
-	FStreamingManagerTexture&	StreamingManager;
-	/** Resulting priorities for keeping textures in memory and controlling load order, matching the FStreamingManagerTexture::StreamingTextures array. */
-	TArray<FTexturePriority>	PrioritizedTextures;
-	/** Indices for PrioritizedTextures, sorted by load order. */
-	TArray<int32>	PrioTexIndicesSortedByLoadOrder;
-
-	/** Context (temporary info) used for the async work. */
-	FStreamingContext			ThreadContext;
-	/** Thread statistics. */
-	FAsyncStats				ThreadStats;
-	/** Whether the async work should abort its processing. */
-	volatile bool				bAbort;
-};
-
-
-/*-----------------------------------------------------------------------------
-	IStreamingManager implementation.
------------------------------------------------------------------------------*/
-
-static FStreamingManagerCollection* StreamingManagerCollection = NULL;
-
-FStreamingManagerCollection& IStreamingManager::Get()
-{
-	if (StreamingManagerCollection == NULL)
-	{
-		StreamingManagerCollection = new FStreamingManagerCollection();
-	}
-	return *StreamingManagerCollection;
-}
-
-void IStreamingManager::Shutdown()
-{
-	if (StreamingManagerCollection != NULL)
-	{
-		delete StreamingManagerCollection;
-		StreamingManagerCollection = (FStreamingManagerCollection*)-1;//Force Error if manager used after shutdown
-	}
-}
-
-bool IStreamingManager::HasShutdown()
-{
-	return StreamingManagerCollection == (FStreamingManagerCollection*)-1;
-}
-
-/**
- * Adds the passed in view information to the static array.
- *
- * @param ViewInfos				[in/out] Array to add the view to
- * @param ViewOrigin			View origin
- * @param ScreenSize			Screen size
- * @param FOVScreenSize			Screen size taking FOV into account
- * @param BoostFactor			A factor that affects all streaming distances for this location. 1.0f is default. Higher means higher-resolution textures and vice versa.
- * @param bOverrideLocation		Whether this is an override location, which forces the streaming system to ignore all other regular locations
- * @param Duration				How long the streaming system should keep checking this location (in seconds). 0 means just for the next Tick.
- * @param InActorToBoost		Optional pointer to an actor who's textures should have their streaming priority boosted
- */
-void IStreamingManager::AddViewInfoToArray( TArray<FStreamingViewInfo> &ViewInfos, const FVector& ViewOrigin, float ScreenSize, float FOVScreenSize, float BoostFactor, bool bOverrideLocation, float Duration, TWeakObjectPtr<AActor> InActorToBoost )
-{
-	// Check for duplicates and existing overrides.
-	bool bShouldAddView = true;
-	for ( int32 ViewIndex=0; ViewIndex < ViewInfos.Num(); ++ViewIndex )
-	{
-		FStreamingViewInfo& ViewInfo = ViewInfos[ ViewIndex ];
-		if ( ViewOrigin.Equals( ViewInfo.ViewOrigin, 0.5f ) &&
-			FMath::IsNearlyEqual( ScreenSize, ViewInfo.ScreenSize ) &&
-			FMath::IsNearlyEqual( FOVScreenSize, ViewInfo.FOVScreenSize ) &&
-			ViewInfo.bOverrideLocation == bOverrideLocation )
-		{
-			// Update duration
-			ViewInfo.Duration = Duration;
-			// Overwrite BoostFactor if it isn't default 1.0
-			ViewInfo.BoostFactor = FMath::IsNearlyEqual(BoostFactor, 1.0f) ? ViewInfo.BoostFactor : BoostFactor;
-			bShouldAddView = false;
-		}
-	}
-
-	if ( bShouldAddView )
-	{
-		new (ViewInfos) FStreamingViewInfo( ViewOrigin, ScreenSize, FOVScreenSize, BoostFactor, bOverrideLocation, Duration, InActorToBoost );
-	}
-}
-
-/**
- * Remove view infos with the same location from the given array.
- *
- * @param ViewInfos				[in/out] Array to remove the view from
- * @param ViewOrigin			View origin
- */
-void IStreamingManager::RemoveViewInfoFromArray( TArray<FStreamingViewInfo> &ViewInfos, const FVector& ViewOrigin )
-{
-	for ( int32 ViewIndex=0; ViewIndex < ViewInfos.Num(); ++ViewIndex )
-	{
-		FStreamingViewInfo& ViewInfo = ViewInfos[ ViewIndex ];
-		if ( ViewOrigin.Equals( ViewInfo.ViewOrigin, 0.5f ) )
-		{
-			ViewInfos.RemoveAtSwap( ViewIndex-- );
-		}
-	}
-}
-
-#if STREAMING_LOG_VIEWCHANGES
-TArray<FStreamingViewInfo> GPrevViewLocations;
-#endif
-
-/**
- * Sets up the CurrentViewInfos array based on PendingViewInfos, LastingViewInfos and SlaveLocations.
- * Removes out-dated LastingViewInfos.
- *
- * @param DeltaTime		Time since last call in seconds
- */
-void IStreamingManager::SetupViewInfos( float DeltaTime )
-{
-	// Reset CurrentViewInfos
-	CurrentViewInfos.Empty( PendingViewInfos.Num() + LastingViewInfos.Num() + SlaveLocations.Num() );
-
-	bool bHaveMultiplePlayerViews = (PendingViewInfos.Num() > 1) ? true : false;
-
-	// Add the slave locations.
-	float ScreenSize = 1280.0f;
-	float FOVScreenSize = ScreenSize / FMath::Tan( 80.0f * float(PI) / 360.0f );
-	if ( PendingViewInfos.Num() > 0 )
-	{
-		ScreenSize = PendingViewInfos[0].ScreenSize;
-		FOVScreenSize = PendingViewInfos[0].FOVScreenSize;
-	}
-	else if ( LastingViewInfos.Num() > 0 )
-	{
-		ScreenSize = LastingViewInfos[0].ScreenSize;
-		FOVScreenSize = LastingViewInfos[0].FOVScreenSize;
-	}
-	// Add them to the appropriate array (pending views or lasting views).
-	for ( int32 SlaveLocationIndex=0; SlaveLocationIndex < SlaveLocations.Num(); SlaveLocationIndex++ )
-	{
-		const FSlaveLocation& SlaveLocation = SlaveLocations[ SlaveLocationIndex ];
-		AddViewInformation( SlaveLocation.Location, ScreenSize, FOVScreenSize, SlaveLocation.BoostFactor, SlaveLocation.bOverrideLocation, SlaveLocation.Duration );
-	}
-
-	// Apply a split-screen factor if we have multiple players on the same machine, and they currently have individual views.
-	float SplitScreenFactor = 1.0f;
-	
-	if ( bHaveMultiplePlayerViews && GEngine->IsSplitScreen(NULL) )
-	{
-		SplitScreenFactor = 0.75f;
-	}
-
-	// Should we use override views this frame? (If we have both a fullscreen movie and an override view.)
-	bool bUseOverrideViews = false;
-	bool bIsMoviePlaying = false;//GetMoviePlayer()->IsMovieCurrentlyPlaying();
-	if ( bIsMoviePlaying )
-	{
-		// Check if we have any override views.
-		for ( int32 ViewIndex=0; !bUseOverrideViews && ViewIndex < LastingViewInfos.Num(); ++ViewIndex )
+		// Check if we have any override views.
+		for ( int32 ViewIndex=0; !bUseOverrideViews && ViewIndex < LastingViewInfos.Num(); ++ViewIndex )
 		{
 			const FStreamingViewInfo& ViewInfo = LastingViewInfos[ ViewIndex ];
 			if ( ViewInfo.bOverrideLocation )
@@ -2392,262 +1204,24 @@ void FStreamingManagerCollection::AddOrRemoveTextureStreamingManagerIfNeeded(boo
 	}
 }
 
-/*-----------------------------------------------------------------------------
-	FStreamingManagerTexture implementation.
------------------------------------------------------------------------------*/
-
-/** Constructor, initializing all members and  */
-FStreamingManagerTexture::FStreamingManagerTexture()
-:	CurrentUpdateStreamingTextureIndex(0)
-,	bTriggerDumpTextureGroupStats( false )
-,	bDetailedDumpTextureGroupStats( false )
-,	bTriggerInvestigateTexture( false )
-,	AsyncWork( nullptr )
-,   DynamicComponentTextureManager( nullptr )
-,	ProcessingStage( 0 )
-,	NumTextureProcessingStages(5)
-,	bUseDynamicStreaming( false )
-,	BoostPlayerTextures( 3.0f )
-,	RangePrefetchDistance(1000.f)
-,	MemoryMargin(0)
-,	MinEvictSize(0)
-,	EffectiveStreamingPoolSize(0)
-,	IndividualStreamingTexture( nullptr )
-,	NumStreamingTextures(0)
-,	NumRequestsInCancelationPhase(0)
-,	NumRequestsInUpdatePhase(0)
-,	NumRequestsInFinalizePhase(0)
-,	TotalIntermediateTexturesSize(0)
-,	NumIntermediateTextures(0)
-,	TotalStreamingTexturesSize(0)
-,	TotalStreamingTexturesMaxSize(0)
-,	TotalLightmapMemorySize(0)
-,	TotalLightmapDiskSize(0)
-,	TotalHLODMemorySize(0)
-,	TotalHLODDiskSize(0)
-,	TotalMipCountIncreaseRequestsInFlight(0)
-,	TotalOptimalWantedSize(0)
-,	TotalStaticTextureHeuristicSize(0)
-,	TotalDynamicTextureHeuristicSize(0)
-,	TotalLastRenderHeuristicSize(0)
-,	TotalForcedHeuristicSize(0)
-,	MemoryOverBudget(0)
-,	MaxEverRequired(0)
-,	OriginalTexturePoolSize(0)
-,	PreviousPoolSizeTimestamp(0.0)
-,	PreviousPoolSizeSetting(-1)
-,	bCollectTextureStats(false)
-,	bReportTextureStats(false)
-,	bPauseTextureStreaming(false)
+/**
+ *	Helper struct to compare two FTextureSortElement objects.
+ **/
+struct FTextureStreamingCompare
 {
-	// Read settings from ini file.
-	int32 TempInt;
-	verify( GConfig->GetInt( TEXT("TextureStreaming"), TEXT("MemoryMargin"),				TempInt,						GEngineIni ) );
-	MemoryMargin = TempInt;
-	verify( GConfig->GetInt( TEXT("TextureStreaming"), TEXT("MinEvictSize"),				TempInt,						GEngineIni ) );
-	MinEvictSize = TempInt;
-
-	verify( GConfig->GetFloat( TEXT("TextureStreaming"), TEXT("LightmapStreamingFactor"),			GLightmapStreamingFactor,		GEngineIni ) );
-	verify( GConfig->GetFloat( TEXT("TextureStreaming"), TEXT("ShadowmapStreamingFactor"),			GShadowmapStreamingFactor,		GEngineIni ) );
-
-	int32 PoolSizeIniSetting = 0;
-	GConfig->GetInt(TEXT("TextureStreaming"), TEXT("PoolSize"), PoolSizeIniSetting, GEngineIni);
-	GConfig->GetBool(TEXT("TextureStreaming"), TEXT("UseDynamicStreaming"), bUseDynamicStreaming, GEngineIni);
-	GConfig->GetFloat( TEXT("TextureStreaming"), TEXT("BoostPlayerTextures"), BoostPlayerTextures, GEngineIni );
-	GConfig->GetBool(TEXT("TextureStreaming"), TEXT("NeverStreamOutTextures"), GNeverStreamOutTextures, GEngineIni);
-	GConfig->GetFloat( TEXT("TextureStreaming"), TEXT("RangePrefetchDistance"), RangePrefetchDistance, GEngineIni );
-
-	// Read pool size from the CVar
-	static const auto CVarStreamingTexturePoolSize = IConsoleManager::Get().FindTConsoleVariableDataInt(TEXT("r.Streaming.PoolSize"));
-	int32 PoolSizeCVar = CVarStreamingTexturePoolSize ? CVarStreamingTexturePoolSize->GetValueOnGameThread() : -1;
-
-	if( PoolSizeCVar != -1 )
-	{
-		OriginalTexturePoolSize =  int64(PoolSizeCVar) * 1024ll * 1024ll;
-		GTexturePoolSize = OriginalTexturePoolSize;
-	}	
-	else if ( GPoolSizeVRAMPercentage )
-	{
-		// If GPoolSizeVRAMPercentage is set, the pool size has already been calculated and we're not reading it from the .ini
-		OriginalTexturePoolSize = GTexturePoolSize;
-	}
-	else
+	/** 
+	 *	Called by Sort() to compare two elements.
+	 *	@param Texture1		- First object to compare
+	 *	@param Texture2		- Second object to compare
+	 *	@return				- Negative value if Texture1 should be sorted earlier in the array than Texture2, zero if arbitrary order, positive if opposite order.
+	 */
+	bool operator()( const FTextureSortElement& Texture1, const FTextureSortElement& Texture2 ) const
 	{
-		// Don't use a texture pool if it's not read from .ini or calculated at startup
-		OriginalTexturePoolSize = 0;
-		GTexturePoolSize = 0;
-	}
-
-	// -NeverStreamOutTextures
-	if (FParse::Param(FCommandLine::Get(), TEXT("NeverStreamOutTextures")))
-	{
-		GNeverStreamOutTextures = true;
-	}
-	if (GIsEditor)
-	{
-		// this would not be good or useful in the editor
-		GNeverStreamOutTextures = false;
-
-		// Use unlimited texture streaming in the editor. Quality is more important than stutter.
-		GTexturePoolSize = 0;
-	}
-	if (GNeverStreamOutTextures)
-	{
-		UE_LOG(LogContentStreaming, Log, TEXT("Textures will NEVER stream out!"));
-	}
-
-	UE_LOG(LogContentStreaming,Log,TEXT("Texture pool size is %.2f MB"),GTexturePoolSize/1024.f/1024.f);
-
-	// Convert from MByte to byte.
-	MinEvictSize *= 1024 * 1024;
-	MemoryMargin *= 1024 * 1024;
-
-#if STATS
-	NumLatencySamples = 0;
-	LatencySampleIndex = 0;
-	LatencyAverage = 0.0f;
-	LatencyMaximum = 0.0f;
-	for ( int32 LatencyIndex=0; LatencyIndex < NUM_LATENCYSAMPLES; ++LatencyIndex )
-	{
-		LatencySamples[ LatencyIndex ] = 0.0f;
-	}
-#endif
-
-#if STATS_FAST
-	MaxStreamingTexturesSize = 0;
-	MaxOptimalTextureSize = 0;
-	MaxStreamingOverBudget = MIN_int64;
-	MaxTexturePoolAllocatedSize = 0;
-	MinLargestHoleSize = OriginalTexturePoolSize;
-	MaxNumWantingTextures = 0;
-#endif
-
-	for ( int32 LODGroup=0; LODGroup < TEXTUREGROUP_MAX; ++LODGroup )
-	{
-		const FTextureLODGroup& TexGroup = UDeviceProfileManager::Get().GetActiveProfile()->GetTextureLODSettings()->GetTextureLODGroup(TextureGroup(LODGroup));
-		ThreadSettings.NumStreamedMips[LODGroup] = TexGroup.NumStreamedMips;
-	}
-
-	// setup the streaming resource flush function pointer
-	GFlushStreamingFunc = &FlushResourceStreaming;
-
-	ProcessingStage = 0;
-	AsyncWork = new FAsyncTask<FAsyncTextureStreaming>(this);
-
-	// We use a dynamic allocation here to reduce header dependencies.
-	DynamicComponentTextureManager = new FDynamicComponentTextureManager();
-	ThreadSettings.TextureBoundsVisibility = new FTextureBoundsVisibility();
-}
-
-FStreamingManagerTexture::~FStreamingManagerTexture()
-{
-	AsyncWork->EnsureCompletion();
-	delete AsyncWork;
-	delete DynamicComponentTextureManager;
-	delete ThreadSettings.TextureBoundsVisibility;
-}
-
-/**
- * Cancels the timed Forced resources (i.e used the Kismet action "Stream In Textures").
- */
-void FStreamingManagerTexture::CancelForcedResources()
-{
-	float CurrentTime = float(FPlatformTime::Seconds() - GStartTime);
-
-	// Update textures that are Forced on a timer.
-	for ( int32 TextureIndex=0; TextureIndex < StreamingTextures.Num(); ++TextureIndex )
-	{
-		FStreamingTexture& StreamingTexture = StreamingTextures[ TextureIndex ];
-
-		// Make sure this streaming texture hasn't been marked for removal.
-		if ( StreamingTexture.Texture )
-		{
-			// Remove any prestream requests from textures
-			float TimeLeft = StreamingTexture.Texture->ForceMipLevelsToBeResidentTimestamp - CurrentTime;
-			if ( TimeLeft > 0.0f )
-			{
-				StreamingTexture.Texture->SetForceMipLevelsToBeResident( -1.0f );
-				StreamingTexture.InstanceRemovedTimestamp = -FLT_MAX;
-				if ( StreamingTexture.Texture->Resource )
-				{
-					StreamingTexture.Texture->InvalidateLastRenderTimeForStreaming();
-				}
-#if STREAMING_LOG_CANCELFORCED
-				UE_LOG(LogContentStreaming, Log, TEXT("Canceling forced texture: %s (had %.1f seconds left)"), *StreamingTexture.Texture->GetFullName(), TimeLeft );
-#endif
-			}
-		}
-	}
-
-	// Reset the streaming system, so it picks up any changes to UTexture2D right away.
-	ProcessingStage = 0;
-}
-
-/**
- * Notifies manager of "level" change so it can prioritize character textures for a few frames.
- */
-void FStreamingManagerTexture::NotifyLevelChange()
-{
-}
-
-/** Don't stream world resources for the next NumFrames. */
-void FStreamingManagerTexture::SetDisregardWorldResourcesForFrames( int32 NumFrames )
-{
-	//@TODO: We could perhaps increase the priority factor for character textures...
-}
-
-/**
- *	Helper struct that represents a texture and the parameters used for sorting and streaming out high-res mip-levels.
- **/
-struct FTextureSortElement
-{
-	/**
-	 *	Constructor.
-	 *
-	 *	@param InTexture					- The texture to represent
-	 *	@param InSize						- Size of the whole texture and all current mip-levels, in bytes
-	 *	@param bInIsCharacterTexture		- 1 if this is a character texture, otherwise 0
-	 *	@param InTextureDataAddress			- Starting address of the texture data
-	 *	@param InNumRequiredResidentMips	- Minimum number of mip-levels required to stay in memory
-	 */
-	FTextureSortElement( UTexture2D* InTexture, int32 InSize, int32 bInIsCharacterTexture, uint32 InTextureDataAddress, int32 InNumRequiredResidentMips )
-	:	Texture( InTexture )
-	,	Size( InSize )
-	,	bIsCharacterTexture( bInIsCharacterTexture )
-	,	TextureDataAddress( InTextureDataAddress )
-	,	NumRequiredResidentMips( InNumRequiredResidentMips )
-	{
-	}
-	/** The texture that this element represents */
-	UTexture2D*	Texture;
-	/** Size of the whole texture and all current mip-levels, in bytes. */
-	int32			Size;
-	/** 1 if this is a character texture, otherwise 0 */
-	int32			bIsCharacterTexture;
-	/** Starting address of the texture data. */
-	uint32		TextureDataAddress;			
-	/** Minimum number of mip-levels required to stay in memory. */
-	int32			NumRequiredResidentMips;
-};
-
-/**
- *	Helper struct to compare two FTextureSortElement objects.
- **/
-struct FTextureStreamingCompare
-{
-	/** 
-	 *	Called by Sort() to compare two elements.
-	 *	@param Texture1		- First object to compare
-	 *	@param Texture2		- Second object to compare
-	 *	@return				- Negative value if Texture1 should be sorted earlier in the array than Texture2, zero if arbitrary order, positive if opposite order.
-	 */
-	bool operator()( const FTextureSortElement& Texture1, const FTextureSortElement& Texture2 ) const
-	{
-		// Character textures get lower priority (sorted later in the array).
-		if ( Texture1.bIsCharacterTexture != Texture2.bIsCharacterTexture )
-		{
-			return Texture1.bIsCharacterTexture < Texture2.bIsCharacterTexture;
-		}
+		// Character textures get lower priority (sorted later in the array).
+		if ( Texture1.bIsCharacterTexture != Texture2.bIsCharacterTexture )
+		{
+			return Texture1.bIsCharacterTexture < Texture2.bIsCharacterTexture;
+		}
 
 		// Larger textures get higher priority (sorted earlier in the array).
 		if ( Texture2.Size - Texture1.Size )
@@ -2733,2913 +1307,6 @@ void Renderthread_StreamOutTextureData(FRHICommandListImmediate& RHICmdList, TAr
 	*bSucceeded = SavedMemory >= RequiredMemorySize;
 }
 
-/**
- *	Try to stream out texture mip-levels to free up more memory.
- *	@param RequiredMemorySize	- Additional texture memory required
- *	@return						- Whether it succeeded or not
- **/
-bool FStreamingManagerTexture::StreamOutTextureData( int64 RequiredMemorySize )
-{
-	const int64 MaxTempMemoryAllowed = CVarStreamingMaxTempMemoryAllowed.GetValueOnGameThread() * 1024 * 1024;
-	RequiredMemorySize = FMath::Max<int64>(RequiredMemorySize, MinEvictSize);
-
-	// Array of candidates for reducing mip-levels.
-	TArray<FTextureSortElement> CandidateTextures;
-	CandidateTextures.Reserve( StreamingTextures.Num() );
-
-	// Don't stream out character textures (to begin with)
-	float CurrentTime = float(FPlatformTime::Seconds() - GStartTime);
-
-	volatile bool bSucceeded = false;
-	
-	//resizing textures actually creates a temp copy so we can only resize so many at a time without running out of memory during the eject itself.
-	int64 MemoryCostToResize = 0;	
-
-	// Collect all textures will be considered for streaming out.
-	for (int32 StreamingIndex = 0; StreamingIndex < StreamingTextures.Num(); ++StreamingIndex)
-	{
-		FStreamingTexture& StreamingTexture = StreamingTextures[StreamingIndex];
-		if (StreamingTexture.Texture)
-		{
-			UTexture2D* Texture = StreamingTexture.Texture;
-
-			// Skyboxes should not stream out.
-			if ( Texture->LODGroup == TEXTUREGROUP_Skybox )
-			{
-				continue;
-			}
-
-			// Number of mip-levels that must be resident due to mip-tails and UTexture2D::GetMinTextureResidentMipCount().
-			int32 NumMips = Texture->GetNumMips();
-			int32 MipTailBaseIndex = Texture->GetMipTailBaseIndex();
-			int32 NumRequiredResidentMips = (MipTailBaseIndex >= 0) ? FMath::Max<int32>(NumMips - MipTailBaseIndex, 0) : 0;
-			NumRequiredResidentMips = FMath::Max<int32>(NumRequiredResidentMips, UTexture2D::GetMinTextureResidentMipCount());
-
-			// Only consider streamable textures that have enough miplevels, and who are currently ready for streaming.
-			if ( IsStreamingTexture(Texture) && FStreamingTexture::IsReadyForStreaming(Texture) && Texture->ResidentMips > NumRequiredResidentMips  )
-			{
-				// We can't stream out mip-tails.
-				int32 CurrentBaseMip = NumMips - Texture->ResidentMips;
-				if ( MipTailBaseIndex < 0 || CurrentBaseMip < MipTailBaseIndex )
-				{
-					// Figure out whether texture should be forced resident based on bools and forced resident time.
-					bool bForceMipLevelsToBeResident = (Texture->ShouldMipLevelsBeForcedResident() || Texture->ForceMipLevelsToBeResidentTimestamp >= CurrentTime);
-					if ( bForceMipLevelsToBeResident == false && Texture->Resource )
-					{
-						// Don't try to stream out if the texture is currently being busy being streamed in/out.
-						bool bSafeToStream = (Texture->UpdateStreamingStatus() == false);
-						if ( bSafeToStream )
-						{
-							bool bIsCharacterTexture =	Texture->LODGroup == TEXTUREGROUP_Character ||
-														Texture->LODGroup == TEXTUREGROUP_CharacterSpecular ||
-														Texture->LODGroup == TEXTUREGROUP_CharacterNormalMap;
-							uint32 TextureDataAddress = 0;
-							MemoryCostToResize += Texture->CalcTextureMemorySize(FMath::Max(0, Texture->ResidentMips - 1));
-							CandidateTextures.Add( FTextureSortElement(Texture, Texture->CalcTextureMemorySize(Texture->ResidentMips), bIsCharacterTexture ? 1 : 0, TextureDataAddress, NumRequiredResidentMips) );
-						}
-					}
-				}
-			}
-		}
-
-		if (MemoryCostToResize >= MaxTempMemoryAllowed)
-		{			
-			// Queue up the process on the render thread.
-			ENQUEUE_UNIQUE_RENDER_COMMAND_THREEPARAMETER(
-				StreamOutTextureDataCommand,
-				TArray<FTextureSortElement>*, CandidateTextures, &CandidateTextures,
-				int64, RequiredMemorySize, RequiredMemorySize,
-				volatile bool*, bSucceeded, &bSucceeded,
-				{				
-				FRHICommandListExecutor::GetImmediateCommandList().ImmediateFlush(EImmediateFlushType::FlushRHIThreadFlushResources);				
-				RHIFlushResources();
-
-				Renderthread_StreamOutTextureData(RHICmdList, CandidateTextures, RequiredMemorySize, bSucceeded);
-
-				FRHICommandListExecutor::GetImmediateCommandList().ImmediateFlush(EImmediateFlushType::FlushRHIThreadFlushResources);
-				RHIFlushResources();
-			});
-
-			// Block until the command has finished executing.
-			FlushRenderingCommands();
-			MemoryCostToResize = 0;			
-			CandidateTextures.Reset();
-		}
-	}	
-
-	if (CandidateTextures.Num() > 0)
-	{
-		// Queue up the process on the render thread.
-		ENQUEUE_UNIQUE_RENDER_COMMAND_THREEPARAMETER(
-			StreamOutTextureDataCommand,
-			TArray<FTextureSortElement>*, CandidateTextures, &CandidateTextures,
-			int64, RequiredMemorySize, RequiredMemorySize,
-			volatile bool*, bSucceeded, &bSucceeded,
-			{
-			Renderthread_StreamOutTextureData(RHICmdList, CandidateTextures, RequiredMemorySize, bSucceeded);
-
-			FRHICommandListExecutor::GetImmediateCommandList().ImmediateFlush(EImmediateFlushType::FlushRHIThreadFlushResources);
-			RHIFlushResources();
-		});
-	}
-
-	// Block until the command has finished executing.
-	FlushRenderingCommands();
-
-	// Reset the streaming system, so it picks up any changes to UTexture2D ResidentMips and RequestedMips.
-	ProcessingStage = 0;
-
-	return bSucceeded;
-}
-
-/**
- * Adds new textures and level data on the gamethread (while the worker thread isn't active).
- */
-void FStreamingManagerTexture::UpdateThreadData()
-{
-	// Add new textures.
-	StreamingTextures.Reserve( StreamingTextures.Num() + PendingStreamingTextures.Num() );
-	for ( int32 TextureIndex=0; TextureIndex < PendingStreamingTextures.Num(); ++TextureIndex )
-	{
-		UTexture2D* Texture = PendingStreamingTextures[ TextureIndex ];
-		FStreamingTexture* StreamingTexture = new (StreamingTextures) FStreamingTexture( Texture );
-		StreamingTexture->SetStreamingIndex( StreamingTextures.Num() - 1 );
-	}
-	PendingStreamingTextures.Empty();
-
-	// Remove old levels. Note: Don't try to access the ULevel object, it may have been deleted already!
-	for ( int32 LevelIndex=0; LevelIndex < ThreadSettings.LevelData.Num(); ++LevelIndex )
-	{
-		FLevelData& LevelData = ThreadSettings.LevelData[ LevelIndex ];
-		if ( LevelData.Value.bRemove )
-		{
-			ThreadSettings.LevelData.RemoveAtSwap( LevelIndex-- );
-		}
-	}
-
-	{
-		STAT( double StartTime = FPlatformTime::Seconds() );
-
-		FRemovedTextureArray RemovedTextures;
-		DynamicComponentTextureManager->IncrementalTick(RemovedTextures, 1.f); // Perform a full update at this point.
-		SetTexturesRemovedTimestamp(RemovedTextures);
-
-		// Duplicate the data to be used on the worker thread
-		ThreadSettings.TextureBoundsVisibility->QuickSet(DynamicComponentTextureManager->GetTextureInstances());
-
-		STAT( GStreamingDynamicPrimitivesTime += FPlatformTime::Seconds() - StartTime );
-		}
-
-	// Add new levels.
-	for ( int32 LevelIndex=0; LevelIndex < PendingLevels.Num(); ++LevelIndex )
-	{
-		ULevel* Level = PendingLevels[ LevelIndex ];
-		FLevelData* LevelData = new (ThreadSettings.LevelData) FLevelData( Level );
-		FThreadLevelData& ThreadLevelData = LevelData->Value;
-
-		// Increase the the force-fully-load refcount.
-		for ( TMap<UTexture2D*,bool>::TIterator It(Level->ForceStreamTextures); It; ++It )
-		{
-			UTexture2D* Texture2D = It.Key();
-			if ( Texture2D && IsStreamingTexture( Texture2D ) )
-			{
-				FStreamingTexture& StreamingTexture = GetStreamingTexture( Texture2D );
-				check( StreamingTexture.ForceLoadRefCount >= 0 );
-				StreamingTexture.ForceLoadRefCount++;
-			}
-		}
-
-		// Copy all texture instances into ThreadLevelData.
-		//Note: ULevel could save it in this form, except it's harder to debug.
-		ThreadLevelData.ThreadTextureInstances.Empty( Level->TextureToInstancesMap.Num() );
-		for ( TMap<UTexture2D*,TArray<FStreamableTextureInstance> >::TIterator It(Level->TextureToInstancesMap); It; ++It )
-		{
-			// Convert to SIMD-friendly form
-			const UTexture2D* Texture2D = It.Key();
-			if ( Texture2D )
-			{
-				TArray<FStreamableTextureInstance>& TextureInstances = It.Value();
-				TArray<FStreamableTextureInstance4>& TextureInstances4 = ThreadLevelData.ThreadTextureInstances.Add( Texture2D, TArray<FStreamableTextureInstance4>() );
-				TextureInstances4.Reserve( (TextureInstances.Num() + 3) / 4 );
-
-				for ( int32 InstanceIndex=0; InstanceIndex < TextureInstances.Num(); ++InstanceIndex )
-				{
-					if ( (InstanceIndex & 3) == 0 )
-					{
-						TextureInstances4.Add(FStreamableTextureInstance4());
-					}
-					const FStreamableTextureInstance& Instance = TextureInstances[ InstanceIndex ];
-					FStreamableTextureInstance4& Instance4 = TextureInstances4[ InstanceIndex/4 ];
-					Instance4.BoundsOriginX[ InstanceIndex & 3 ] = Instance.Bounds.Origin.X;
-					Instance4.BoundsOriginY[ InstanceIndex & 3 ] = Instance.Bounds.Origin.Y;
-					Instance4.BoundsOriginZ[ InstanceIndex & 3 ] = Instance.Bounds.Origin.Z;
-					Instance4.BoxExtentSizeX[ InstanceIndex & 3 ] = Instance.Bounds.BoxExtent.X;
-					Instance4.BoxExtentSizeY[ InstanceIndex & 3 ] = Instance.Bounds.BoxExtent.Y;
-					Instance4.BoxExtentSizeZ[ InstanceIndex & 3 ] = Instance.Bounds.BoxExtent.Z;
-					Instance4.BoundingSphereRadius[ InstanceIndex & 3 ] = Instance.Bounds.SphereRadius;
-					Instance4.MinDistanceSq[ InstanceIndex & 3 ] = FMath::Square(FMath::Max(0.f, Instance.MinDistance - RangePrefetchDistance));
-					Instance4.MaxDistanceSq[ InstanceIndex & 3 ] = Instance.MaxDistance == MAX_FLT ? MAX_FLT : FMath::Square(Instance.MaxDistance + RangePrefetchDistance);
-					Instance4.TexelFactor[ InstanceIndex & 3 ] = Instance.TexelFactor;
-					ensureMsgf(!FMath::IsNearlyZero(Instance.TexelFactor), TEXT("Texture instance %d has a texel factor of zero: %s"), TextureInstances.Num(), *Texture2D->GetPathName());
-				}
-
-				// Note: The old instance array must remain, in case a level is added/removed and then added again.
-			}
-		}
-	}
-	PendingLevels.Empty();
-
-	// If any views specified an actor to boost, do it now
-	for (int32 i = 0; i < CurrentViewInfos.Num(); ++i)
-	{
-		if (CurrentViewInfos[i].ActorToBoost.IsValid())
-		{
-			BoostTextures(CurrentViewInfos[i].ActorToBoost.Get(false), BoostPlayerTextures);
-		}
-	}
-
-	// Cache the view information.
-	ThreadSettings.ThreadViewInfos = CurrentViewInfos;
-	
-	ThreadSettings.MipBias = FMath::Max(CVarStreamingMipBias.GetValueOnGameThread(), 0.0f);
-}
-
-/**
- * Temporarily boosts the streaming distance factor by the specified number.
- * This factor is automatically reset to 1.0 after it's been used for mip-calculations.
- */
-void FStreamingManagerTexture::BoostTextures( AActor* Actor, float BoostFactor )
-{
-	if ( Actor )
-	{
-		TArray<UTexture*> Textures;
-		Textures.Empty( 32 );
-
-		TInlineComponentArray<UPrimitiveComponent*> Components;
-		Actor->GetComponents(Components);
-
-		for(int32 ComponentIndex = 0;ComponentIndex < Components.Num();ComponentIndex++)
-		{
-			UPrimitiveComponent* Primitive = Components[ComponentIndex];
-			if ( Primitive->IsRegistered() && Primitive->IsA(UMeshComponent::StaticClass()) )
-			{
-				Textures.Reset();
-				Primitive->GetUsedTextures( Textures, EMaterialQualityLevel::Num );
-				for ( int32 TextureIndex=0; TextureIndex < Textures.Num(); ++TextureIndex )
-				{
-					UTexture2D* Texture2D = Cast<UTexture2D>( Textures[ TextureIndex ] );
-					if ( Texture2D && IsManagedStreamingTexture(Texture2D) )
-					{
-						FStreamingTexture& StreamingTexture = GetStreamingTexture( Texture2D );
-						StreamingTexture.BoostFactor = FMath::Max( StreamingTexture.BoostFactor, BoostFactor );
-					}
-				}
-			}
-		}
-	}
-}
-
-/** Adds a ULevel to the streaming manager. */
-void FStreamingManagerTexture::AddPreparedLevel( ULevel* Level )
-{
-	PendingLevels.AddUnique( Level );
-
-	if ( bUseDynamicStreaming )
-	{
-		// Add all dynamic primitives from the ULevel.
-		for ( TMap<TWeakObjectPtr<UPrimitiveComponent>,TArray<FDynamicTextureInstance> >::TIterator It(Level->DynamicTextureInstances); It; ++It )
-		{
-			UPrimitiveComponent* Primitive = It.Key().Get();
-			TArray<FDynamicTextureInstance>& TextureInstances = It.Value();
-			NotifyPrimitiveAttached( Primitive, DPT_Level );
-
-			//@TODO: This variable is deprecated, remove from serialization.
-			TextureInstances.Empty();
-		}
-	}
-}
-
-/** Removes a ULevel from the streaming manager. */
-void FStreamingManagerTexture::RemoveLevel( ULevel* Level )
-{
-	// Remove from pending new levels, if it's been added very recently (this frame)...
-	PendingLevels.Remove( Level );
-
-	// Mark the level for removal (will take effect when we're syncing threads).
-	FLevelData* LevelData = ThreadSettings.LevelData.FindByKey( Level );
-	if ( LevelData && !LevelData->Value.bRemove )
-	{
-		FThreadLevelData& ThreadLevelData = LevelData->Value;
-		ThreadLevelData.bRemove = true;
-
-		// Mark all textures with a timestamp. They're about to lose their location-based heuristic and we don't want them to
-		// start using LastRenderTime heuristic for a few seconds until they are garbage collected!
-		for( TMap<UTexture2D*,TArray<FStreamableTextureInstance> >::TIterator It(Level->TextureToInstancesMap); It; ++It )
-		{
-			const UTexture2D* Texture2D = It.Key();
-			if ( Texture2D && IsManagedStreamingTexture( Texture2D ) )
-			{
-				FStreamingTexture& StreamingTexture = GetStreamingTexture( Texture2D );
-				StreamingTexture.InstanceRemovedTimestamp = FApp::GetCurrentTime();
-			}
-		}
-
-		// Decrease the the force-fully-load refcount. 
-		// Note: Only update ForceLoadRefCount if the level was in the ThreadSettings and therefore bumped ForceLoadRefCount earlier.
-		// Note: We do this immediately, since ULevel may be being deleted.
-		for ( TMap<UTexture2D*,bool>::TIterator It(Level->ForceStreamTextures); It; ++It )
-		{
-			UTexture2D* Texture2D = It.Key();
-			if ( Texture2D && IsManagedStreamingTexture( Texture2D ) )
-			{
-				FStreamingTexture& StreamingTexture = GetStreamingTexture( Texture2D );
-				// Note: This refcount should always be positive, but for some reason it may not be when BuildStreamingData() is called in the editor...
-				if ( StreamingTexture.ForceLoadRefCount > 0 )
-				{
-					StreamingTexture.ForceLoadRefCount--;
-				}
-			}
-		}
-	}
-
-	if ( bUseDynamicStreaming )
-	{
-		// Mark all dynamic primitives from the ULevel (they'll be removed completely next sync).
-		for ( TMap<TWeakObjectPtr<UPrimitiveComponent>,TArray<FDynamicTextureInstance> >::TIterator It(Level->DynamicTextureInstances); It; ++It )
-		{
-			UPrimitiveComponent* Primitive = It.Key().GetEvenIfUnreachable();
-			NotifyPrimitiveDetached( Primitive );
-		}
-	}
-}
-
-/**
- * Adds a new texture to the streaming manager.
- */
-void FStreamingManagerTexture::AddStreamingTexture( UTexture2D* Texture )
-{
-	// Adds the new texture to the Pending list, to avoid reallocation of the thread-safe StreamingTextures array.
-	check( Texture->StreamingIndex == -1 );
-	int32 Index = PendingStreamingTextures.Add( Texture );
-	Texture->StreamingIndex = Index;
-}
-
-/**
- * Removes a texture from the streaming manager.
- */
-void FStreamingManagerTexture::RemoveStreamingTexture( UTexture2D* Texture )
-{
-	// Removes the texture from the Pending list or marks the StreamingTextures slot as unused.
-
-	// Remove it from the Pending list.
-	int32 Index = Texture->StreamingIndex;
-	if ( Index >= 0 && Index < PendingStreamingTextures.Num() )
-	{
-		UTexture2D* ExistingPendingTexture = PendingStreamingTextures[ Index ];
-		if ( ExistingPendingTexture == Texture )
-		{
-			PendingStreamingTextures.RemoveAtSwap( Index );
-			if ( Index != PendingStreamingTextures.Num() )
-			{
-				UTexture2D* SwappedPendingTexture = PendingStreamingTextures[ Index ];
-				SwappedPendingTexture->StreamingIndex = Index;
-			}
-			Texture->StreamingIndex = -1;
-		}
-	}
-
-	Index = Texture->StreamingIndex;
-	if ( Index >= 0 && Index < StreamingTextures.Num() )
-	{
-		FStreamingTexture& ExistingStreamingTexture = StreamingTextures[ Index ];
-		if ( ExistingStreamingTexture.Texture == Texture )
-		{
-			// If using the new priority system, mark StreamingTextures slot as unused.
-			ExistingStreamingTexture.Texture = NULL;
-			Texture->StreamingIndex = -1;
-		}
-	}
-
-//	checkSlow( Texture->StreamingIndex == -1 );	// The texture should have been in one of the two arrays!
-	Texture->StreamingIndex = -1;
-}
-
-/** Called when an actor is spawned. */
-void FStreamingManagerTexture::NotifyActorSpawned( AActor* Actor )
-{
-	if ( bUseDynamicStreaming )
-	{
-		TInlineComponentArray<UPrimitiveComponent*> Components;
-		Actor->GetComponents(Components);
-
-		for(int32 ComponentIndex = 0;ComponentIndex < Components.Num();ComponentIndex++)
-		{
-			UPrimitiveComponent* Primitive = Components[ComponentIndex];
-			if (Primitive->IsRegistered() && Primitive->IsA(UMeshComponent::StaticClass()))
-			{
-				NotifyPrimitiveAttached( Primitive, DPT_Spawned );
-			}
-		}
-	}
-}
-
-/** Called when a spawned primitive is deleted. */
-void FStreamingManagerTexture::NotifyActorDestroyed( AActor* Actor )
-{
-	TInlineComponentArray<UPrimitiveComponent*> Components;
-	Actor->GetComponents(Components);
-
-	for(int32 ComponentIndex = 0;ComponentIndex < Components.Num();ComponentIndex++)
-	{
-		UPrimitiveComponent* Primitive = Components[ComponentIndex];
-		NotifyPrimitiveDetached( Primitive );
-	}
-}
-
-/**
- * Called when a primitive is attached to an actor or another component.
- * Replaces previous info, if the primitive was already attached.
- *
- * @param InPrimitive	Newly attached dynamic/spawned primitive
- */
-void FStreamingManagerTexture::NotifyPrimitiveAttached( const UPrimitiveComponent* Primitive, EDynamicPrimitiveType DynamicType )
-{
-	// Only add it if it's a UMeshComponent, since we only track those in UMeshComponent::BeginDestroy().
-	if ( bUseDynamicStreaming && Primitive && Primitive->IsA(UMeshComponent::StaticClass()) )
-	{
-#if STREAMING_LOG_DYNAMIC
-		UE_LOG(LogContentStreaming, Log, TEXT("NotifyPrimitiveAttached(0x%08x \"%s\"), IsRegistered=%d"), SIZE_T(Primitive), *Primitive->GetReadableName(), Primitive->IsRegistered());
-#endif
-		FRemovedTextureArray RemovedTextures;
-		DynamicComponentTextureManager->Add(Primitive, DynamicType, RemovedTextures);
-		SetTexturesRemovedTimestamp(RemovedTextures);
-	}
-}
-
-/**
- * Called when a primitive is detached from an actor or another component.
- * Note: We should not be accessing the primitive or the UTexture2D after this call!
- */
-void FStreamingManagerTexture::NotifyPrimitiveDetached( const UPrimitiveComponent* Primitive )
-{
-	if ( bUseDynamicStreaming && Primitive )
-	{
-#if STREAMING_LOG_DYNAMIC
-		UE_LOG(LogContentStreaming, Log, TEXT("NotifyPrimitiveDetached(0x%08x \"%s\"), IsRegistered=%d"), SIZE_T(Primitive), *Primitive->GetReadableName(), Primitive->IsRegistered());
-#endif
-		FRemovedTextureArray RemovedTextures;
-		DynamicComponentTextureManager->Remove(Primitive, RemovedTextures);
-		SetTexturesRemovedTimestamp(RemovedTextures);
-	}
-}
-
-/**
-* Mark the textures with a timestamp. They're about to lose their location-based heuristic and we don't want them to
-* start using LastRenderTime heuristic for a few seconds until they are garbage collected!
-*
-* @param RemovedTextures	List of removed textures.
-*/
-void FStreamingManagerTexture::SetTexturesRemovedTimestamp(const FRemovedTextureArray& RemovedTextures)
-{
-	double CurrentTime = FApp::GetCurrentTime();
-	for ( int32 TextureIndex=0; TextureIndex < RemovedTextures.Num(); ++TextureIndex )
-	{
-		const UTexture2D* Texture2D = RemovedTextures[TextureIndex];
-		if ( Texture2D && IsManagedStreamingTexture(Texture2D) )
-		{
-			FStreamingTexture& StreamingTexture = GetStreamingTexture( Texture2D );
-			StreamingTexture.InstanceRemovedTimestamp = CurrentTime;
-		}
-	}
-}
-
-/**
- * Called when a primitive has had its textured changed.
- * Only affects primitives that were already attached.
- * Replaces previous info.
- */
-void FStreamingManagerTexture::NotifyPrimitiveUpdated( const UPrimitiveComponent* Primitive )
-{
-	FRemovedTextureArray RemovedTextures;
-	DynamicComponentTextureManager->Update(Primitive, RemovedTextures);
-	SetTexturesRemovedTimestamp(RemovedTextures);
-}
-
-/**
- * Called when a LastRenderTime primitive is attached to an actor or another component.
- * Modifies the LastRenderTimeRefCount for the textures used, so that those textures can
- * use both distance-based and LastRenderTime heuristics.
- *
- * @param Primitive		Newly attached dynamic/spawned primitive
- */
-void FStreamingManagerTexture::NotifyTimedPrimitiveAttached( const UPrimitiveComponent* Primitive, EDynamicPrimitiveType DynamicType )
-{
-	if ( Primitive  && Primitive->IsRegistered())
-	{
-		TArray<FStreamingTexturePrimitiveInfo> TextureInstanceInfos;
-		Primitive->GetStreamingTextureInfoWithNULLRemoval( TextureInstanceInfos );
-
-		TArray<FSpawnedTextureInstance>* TextureInstances = NULL;
-		for ( int32 InstanceIndex=0; InstanceIndex < TextureInstanceInfos.Num(); ++InstanceIndex )
-		{
-			FStreamingTexturePrimitiveInfo& Info = TextureInstanceInfos[InstanceIndex];
-			UTexture2D* Texture2D = Cast<UTexture2D>(Info.Texture);
-			if ( Texture2D && IsManagedStreamingTexture(Texture2D) )
-			{
-				// Note: Doesn't have to be cycle-perfect for thread safety.
-				FStreamingTexture& StreamingTexture = GetStreamingTexture( Texture2D );
-				StreamingTexture.LastRenderTimeRefCount++;
-				StreamingTexture.LastRenderTimeRefCountTimestamp = FApp::GetCurrentTime();
-			}
-		}
-	}
-}
-
-/**
- * Called when a LastRenderTime primitive is detached from an actor or another component.
- * Modifies the LastRenderTimeRefCount for the textures used, so that those textures can
- * use both distance-based and LastRenderTime heuristics.
- *
- * @param Primitive		Newly detached dynamic/spawned primitive
- */
-void FStreamingManagerTexture::NotifyTimedPrimitiveDetached( const UPrimitiveComponent* Primitive )
-{
-	if ( Primitive )
-	{
-		TArray<FStreamingTexturePrimitiveInfo> TextureInstanceInfos;
-		Primitive->GetStreamingTextureInfoWithNULLRemoval( TextureInstanceInfos );
-
-		TArray<FSpawnedTextureInstance>* TextureInstances = NULL;
-		for ( int32 InstanceIndex=0; InstanceIndex < TextureInstanceInfos.Num(); ++InstanceIndex )
-		{
-			FStreamingTexturePrimitiveInfo& Info = TextureInstanceInfos[InstanceIndex];
-			UTexture2D* Texture2D = Cast<UTexture2D>(Info.Texture);
-			if ( Texture2D && IsManagedStreamingTexture(Texture2D) )
-			{
-				// Note: Doesn't have to be cycle-perfect for thread safety.
-				FStreamingTexture& StreamingTexture = GetStreamingTexture( Texture2D );
-				if ( StreamingTexture.LastRenderTimeRefCount > 0 )
-				{
-					StreamingTexture.LastRenderTimeRefCount--;
-					StreamingTexture.LastRenderTimeRefCountTimestamp = FApp::GetCurrentTime();
-				}
-			}
-		}
-	}
-}
-
-/**
- * Returns the corresponding FStreamingTexture for a UTexture2D.
- */
-FStreamingTexture& FStreamingManagerTexture::GetStreamingTexture( const UTexture2D* Texture2D )
-{
-	return StreamingTextures[ Texture2D->StreamingIndex ];
-}
-
-/** Returns true if this is a streaming texture that is managed by the streaming manager. */
-bool FStreamingManagerTexture::IsManagedStreamingTexture( const UTexture2D* Texture2D )
-{
-	return Texture2D->StreamingIndex >= 0 && Texture2D->StreamingIndex < StreamingTextures.Num() && IsStreamingTexture( Texture2D );
-}
-
-/**
- * Updates streaming for an individual texture, taking into account all view infos.
- *
- * @param Texture	Texture to update
- */
-void FStreamingManagerTexture::UpdateIndividualTexture( UTexture2D* Texture )
-{
-	if (IStreamingManager::Get().IsStreamingEnabled())
-	{
-		IndividualStreamingTexture = Texture;
-		UpdateResourceStreaming( 0.0f );
-		IndividualStreamingTexture = NULL;
-	}
-}
-
-/**
- * Resets the streaming statistics to zero.
- */
-void FStreamingManagerTexture::ResetStreamingStats()
-{
-	NumStreamingTextures								= 0;
-	NumRequestsInCancelationPhase						= 0;
-	NumRequestsInUpdatePhase							= 0;
-	NumRequestsInFinalizePhase							= 0;
-	TotalIntermediateTexturesSize						= 0;
-	NumIntermediateTextures								= 0;
-	TotalStreamingTexturesSize							= 0;
-	TotalStreamingTexturesMaxSize						= 0;
-	TotalLightmapMemorySize								= 0;
-	TotalLightmapDiskSize								= 0;
-	TotalHLODMemorySize									= 0;
-	TotalHLODDiskSize									= 0;
-	TotalMipCountIncreaseRequestsInFlight				= 0;
-	TotalOptimalWantedSize								= 0;
-	TotalStaticTextureHeuristicSize						= 0;
-	TotalDynamicTextureHeuristicSize					= 0;
-	TotalLastRenderHeuristicSize						= 0;
-	TotalForcedHeuristicSize							= 0;
-}
-
-/**
- * Updates the streaming statistics with current frame's worth of stats.
- *
- * @param Context					Context for the current frame
- * @param bAllTexturesProcessed		Whether all processing is complete
- */
-void FStreamingManagerTexture::UpdateStreamingStats( const FStreamingContext& Context, bool bAllTexturesProcessed )
-{
-	NumStreamingTextures					+= Context.ThisFrameNumStreamingTextures;
-	NumRequestsInCancelationPhase			+= Context.ThisFrameNumRequestsInCancelationPhase;
-	NumRequestsInUpdatePhase				+= Context.ThisFrameNumRequestsInUpdatePhase;
-	NumRequestsInFinalizePhase				+= Context.ThisFrameNumRequestsInFinalizePhase;
-	NumIntermediateTextures					+= Context.ThisFrameNumIntermediateTextures;
-	TotalStreamingTexturesSize				+= Context.ThisFrameTotalStreamingTexturesSize;
-	TotalStreamingTexturesMaxSize			+= Context.ThisFrameTotalStreamingTexturesMaxSize;
-	TotalLightmapMemorySize					+= Context.ThisFrameTotalLightmapMemorySize;
-	TotalLightmapDiskSize					+= Context.ThisFrameTotalLightmapDiskSize;
-	TotalHLODMemorySize						+= Context.ThisFrameTotalHLODMemorySize;
-	TotalHLODDiskSize						+= Context.ThisFrameTotalHLODDiskSize;
-	TotalMipCountIncreaseRequestsInFlight	+= Context.ThisFrameTotalMipCountIncreaseRequestsInFlight;
-	TotalOptimalWantedSize					+= Context.ThisFrameOptimalWantedSize;
-	TotalStaticTextureHeuristicSize			+= Context.ThisFrameTotalStaticTextureHeuristicSize;
-	TotalDynamicTextureHeuristicSize		+= Context.ThisFrameTotalDynamicTextureHeuristicSize;
-	TotalLastRenderHeuristicSize			+= Context.ThisFrameTotalLastRenderHeuristicSize;
-	TotalForcedHeuristicSize				+= Context.ThisFrameTotalForcedHeuristicSize;
-
-
-	// Set the stats on wrap-around. Reset happens independently to correctly handle resetting in the middle of iteration.
-	if ( bAllTexturesProcessed )
-	{
-		FTextureMemoryStats Stats;
-		RHIGetTextureMemoryStats(Stats);
-
-		SET_DWORD_STAT( STAT_StreamingTextures,				NumStreamingTextures			);
-		SET_DWORD_STAT( STAT_RequestsInCancelationPhase,	NumRequestsInCancelationPhase	);
-		SET_DWORD_STAT( STAT_RequestsInUpdatePhase,			NumRequestsInUpdatePhase		);
-		SET_DWORD_STAT( STAT_RequestsInFinalizePhase,		NumRequestsInFinalizePhase		);
-		SET_DWORD_STAT( STAT_IntermediateTexturesSize,		TotalIntermediateTexturesSize	);
-		SET_DWORD_STAT( STAT_IntermediateTextures,			NumIntermediateTextures			);
-		SET_DWORD_STAT( STAT_StreamingTexturesMaxSize,		TotalStreamingTexturesMaxSize	);
-		SET_DWORD_STAT( STAT_LightmapMemorySize,			TotalLightmapMemorySize			);
-		SET_DWORD_STAT( STAT_LightmapDiskSize,				TotalLightmapDiskSize			);
-		SET_DWORD_STAT( STAT_HLODTextureMemorySize,			TotalHLODMemorySize				);
-		SET_DWORD_STAT( STAT_HLODTextureDiskSize,			TotalHLODDiskSize				);
-		
-		SET_FLOAT_STAT( STAT_StreamingBandwidth,			BandwidthAverage/1024.0f/1024.0f);
-		SET_FLOAT_STAT( STAT_StreamingLatency,				LatencyAverage					);
-		SET_MEMORY_STAT( STAT_StreamingTexturesSize,		TotalStreamingTexturesSize		);
-		SET_MEMORY_STAT( STAT_OptimalTextureSize,			TotalOptimalWantedSize			);
-		SET_MEMORY_STAT( STAT_TotalStaticTextureHeuristicSize,	TotalStaticTextureHeuristicSize	);
-		SET_MEMORY_STAT( STAT_TotalDynamicHeuristicSize,		TotalDynamicTextureHeuristicSize );
-		SET_MEMORY_STAT( STAT_TotalLastRenderHeuristicSize,		TotalLastRenderHeuristicSize	);
-		SET_MEMORY_STAT( STAT_TotalForcedHeuristicSize,			TotalForcedHeuristicSize		);
-#if 0 // this is a visualization thing, not a stats thing
-
-#if STATS_FAST
-		MaxStreamingTexturesSize = FMath::Max(MaxStreamingTexturesSize, GET_MEMORY_STAT(STAT_StreamingTexturesSize));
-		MaxOptimalTextureSize = FMath::Max(MaxOptimalTextureSize, GET_MEMORY_STAT(STAT_OptimalTextureSize));
-		MaxStreamingOverBudget = FMath::Max<int64>(MaxStreamingOverBudget, GET_MEMORY_STAT(STAT_StreamingOverBudget) - GET_MEMORY_STAT(STAT_StreamingUnderBudget));
-		MaxTexturePoolAllocatedSize = FMath::Max(MaxTexturePoolAllocatedSize, GET_MEMORY_STAT(STAT_TexturePoolAllocatedSize));
-#if STATS
-		MinLargestHoleSize = FMath::Min(MinLargestHoleSize, GET_MEMORY_STAT(STAT_TexturePool_LargestHole));
-#endif // STATS
-		MaxNumWantingTextures = FMath::Max(MaxNumWantingTextures, GET_MEMORY_STAT(STAT_NumWantingTextures));
-#endif
-#endif
-	}
-
-	SET_DWORD_STAT(		STAT_RequestSizeCurrentFrame,		Context.ThisFrameTotalRequestSize		);
-	INC_DWORD_STAT_BY(	STAT_RequestSizeTotal,				Context.ThisFrameTotalRequestSize		);
-	INC_DWORD_STAT_BY(	STAT_LightmapRequestSizeTotal,		Context.ThisFrameTotalLightmapRequestSize );
-}
-
-/**
- * Not thread-safe: Updates a portion (as indicated by 'StageIndex') of all streaming textures,
- * allowing their streaming state to progress.
- *
- * @param Context			Context for the current stage (frame)
- * @param StageIndex		Current stage index
- * @param NumUpdateStages	Number of texture update stages
- */
-void FStreamingManagerTexture::UpdateStreamingTextures( FStreamingContext& Context, int32 StageIndex, int32 NumUpdateStages )
-{
-	if ( StageIndex == 0 )
-	{
-		CurrentUpdateStreamingTextureIndex = 0;
-	}
-	int32 StartIndex = CurrentUpdateStreamingTextureIndex;
-	int32 EndIndex = StreamingTextures.Num() * (StageIndex + 1) / NumUpdateStages;
-	for ( int32 Index=StartIndex; Index < EndIndex; ++Index )
-	{
-		FStreamingTexture& StreamingTexture = StreamingTextures[ Index ];
-		FPlatformMisc::Prefetch( &StreamingTexture + 1 );
-
-		// Is this texture marked for removal?
-		if ( StreamingTexture.Texture == NULL )
-		{
-			StreamingTextures.RemoveAtSwap( Index );
-			if ( Index != StreamingTextures.Num() )
-			{
-				FStreamingTexture& SwappedTexture = StreamingTextures[ Index ];
-				// Note: The swapped texture may also be pending deletion.
-				if ( SwappedTexture.Texture )
-				{
-					SwappedTexture.Texture->StreamingIndex = Index;
-				}
-			}
-			--Index;
-			--EndIndex;
-			continue;
-		}
-
-		StreamingTexture.UpdateCachedInfo();
-
-		if ( StreamingTexture.bReadyForStreaming )
-		{
-			UpdateTextureStatus( StreamingTexture, Context );
-		}
-	}
-	CurrentUpdateStreamingTextureIndex = EndIndex;
-}
-
-
-
-/**
- * Stream textures in/out, based on the priorities calculated by the async work.
- * @param bProcessEverything	Whether we're processing all textures in one go
- */
-void FStreamingManagerTexture::StreamTextures( bool bProcessEverything )
-{
-	const int64 MaxTempMemoryAllowed = CVarStreamingMaxTempMemoryAllowed.GetValueOnAnyThread() * 1024 * 1024;
-
-	// Setup a context for this run.
-	FStreamingContext Context( bProcessEverything, IndividualStreamingTexture, bCollectTextureStats );
-
-	AsyncWork->GetTask().ClearRemovedTextureReferences();
-
-	const TArray<FTexturePriority>& PrioritizedTextures = AsyncWork->GetTask().GetPrioritizedTextures();
-	const TArray<int32>& PrioTexIndicesSortedByLoadOrder = AsyncWork->GetTask().GetPrioTexIndicesSortedByLoadOrder();
-
-	FAsyncTextureStreaming::FAsyncStats ThreadStats = AsyncWork->GetTask().GetStats();
-	Context.AddStats( AsyncWork->GetTask().GetContext() );
-
-	FIOSystem::Get().FlushLog();
-
-#if STATS
-	// Did we collect texture stats? Triggered by the ListStreamingTextures exec command.
-	if ( Context.TextureStats.Num() > 0 )
-	{
-		// Reinitialize each time
-		TextureStatsReport.Empty();
-
-		// Sort textures by cost.
-		struct FCompareFTextureStreamingStats
-		{
-			FORCEINLINE bool operator()( const FTextureStreamingStats& A, const FTextureStreamingStats& B ) const
-			{
-				if ( A.Priority > B.Priority )
-				{
-					return true;
-				}
-				else if ( A.Priority == B.Priority )
-				{
-					return ( A.TextureIndex < B.TextureIndex );
-				}
-				else
-				{
-					return false;
-				}
-			}
-		};
-		Context.TextureStats.Sort( FCompareFTextureStreamingStats() );
-		int64 TotalCurrentSize	= 0;
-		int64 TotalWantedSize	= 0;
-		int64 TotalMaxSize		= 0;
-		TextureStatsReport.Add( FString( TEXT(",Priority,Current,Wanted,Max,Largest Resident,Current Size (KB),Wanted Size (KB),Max Size (KB),Largest Resident Size (KB),Streaming Type,Last Rendered,BoostFactor,Name") ) );
-		for( int32 TextureIndex=0; TextureIndex<Context.TextureStats.Num(); TextureIndex++ )
-		{
-			const FTextureStreamingStats& TextureStat = Context.TextureStats[TextureIndex];
-			int32 LODBias			= TextureStat.LODBias;
-			int32 CurrDroppedMips	= TextureStat.NumMips - TextureStat.ResidentMips;
-			int32 WantedDroppedMips	= TextureStat.NumMips - TextureStat.WantedMips;
-			int32 MostDroppedMips	= TextureStat.NumMips - TextureStat.MostResidentMips;
-			TextureStatsReport.Add( FString::Printf( TEXT(",%.3f,%ix%i,%ix%i,%ix%i,%ix%i,%i,%i,%i,%i,%s,%3f sec,%.1f,%s"),
-				TextureStat.Priority,
-				FMath::Max(TextureStat.SizeX >> CurrDroppedMips, 1),
-				FMath::Max(TextureStat.SizeY >> CurrDroppedMips, 1),
-				FMath::Max(TextureStat.SizeX >> WantedDroppedMips, 1),
-				FMath::Max(TextureStat.SizeY >> WantedDroppedMips, 1),
-				FMath::Max(TextureStat.SizeX >> LODBias, 1),
-				FMath::Max(TextureStat.SizeY >> LODBias, 1),
-				FMath::Max(TextureStat.SizeX >> MostDroppedMips, 1),
-				FMath::Max(TextureStat.SizeY >> MostDroppedMips, 1),
-				TextureStat.ResidentSize/1024,
-				TextureStat.WantedSize/1024,
-				TextureStat.MaxSize/1024,
-				TextureStat.MostResidentSize/1024,
-				GStreamTypeNames[TextureStat.StreamType],
-				TextureStat.LastRenderTime,
-				TextureStat.BoostFactor,
-				*TextureStat.TextureName
-				) );
-			TotalCurrentSize	+= TextureStat.ResidentSize;
-			TotalWantedSize		+= TextureStat.WantedSize;
-			TotalMaxSize		+= TextureStat.MaxSize;
-		}
-		TextureStatsReport.Add( FString::Printf( TEXT("Total size: Current= %d  Wanted=%d  Max= %d"), TotalCurrentSize/1024, TotalWantedSize/1024, TotalMaxSize/1024 ) );
-		Context.TextureStats.Empty();
-
-		if( bReportTextureStats )
-		{
-			if( TextureStatsReport.Num() > 0 )
-			{
-				UE_LOG(LogContentStreaming, Log,  TEXT("Listing collected stats for all streaming textures") );
-				for( int32 ReportIndex = 0; ReportIndex < TextureStatsReport.Num(); ReportIndex++ )
-				{
-					UE_LOG(LogContentStreaming, Log, TEXT("%s"), *TextureStatsReport[ReportIndex]);
-				}
-				TextureStatsReport.Empty();
-			}
-			bReportTextureStats = false;
-		}
-		bCollectTextureStats = false;
-	}
-#endif
-
-	FTextureMemoryStats Stats;
-	RHIGetTextureMemoryStats(Stats);
-	ThreadStats.PendingStreamInSize += (Stats.PendingMemoryAdjustment > 0) ? Stats.PendingMemoryAdjustment : 0;
-	ThreadStats.PendingStreamOutSize += (Stats.PendingMemoryAdjustment < 0) ? -Stats.PendingMemoryAdjustment : 0;
-
-	int64 LargestAlloc = 0;
-	if ( Stats.IsUsingLimitedPoolSize() )
-	{
-		LargestAlloc = Stats.LargestContiguousAllocation;
-
-		// Update EffectiveStreamingPoolSize, trying to stabilize it independently of temp memory, allocator overhead and non-streaming resources normal variation.
-		// It's hard to know how much temp memory and allocator overhead is actually in AllocatedMemorySize as it is platform specific.
-		// We handle it by not using all memory available. If temp memory and memory margin values are effectively bigger than the actual used values, the pool will stabilize.
-		const int64 NonStreamingUsage = Stats.AllocatedMemorySize - ThreadStats.TotalResidentSize;
-		const int64 AvailableMemoryForStreaming =  Stats.TexturePoolSize - NonStreamingUsage;
-
-		if (AvailableMemoryForStreaming < EffectiveStreamingPoolSize)
-		{
-			// Reduce size immediately to avoid taking more memory.
-			EffectiveStreamingPoolSize = AvailableMemoryForStreaming;
-		}
-		else if (AvailableMemoryForStreaming - EffectiveStreamingPoolSize > MaxTempMemoryAllowed + MemoryMargin)
-		{
-			// Increase size considering that the variation does not come from temp memory or allocator overhead (or other recurring cause).
-			// It's unclear how much temp memory is actually in there, but the value will decrease if temp memory increases.
-			EffectiveStreamingPoolSize = AvailableMemoryForStreaming;
-		}
-
-		const int64 LocalMemoryOverBudget = FMath::Max(ThreadStats.TotalRequiredSize - EffectiveStreamingPoolSize, 0LL);
-		MemoryOverBudget = FMath::Max(LocalMemoryOverBudget, 0LL);
-		MaxEverRequired = FMath::Max(ThreadStats.TotalRequiredSize, MaxEverRequired);
-
-		// Texture stats must come from assets memory and should not include allocator overhead, as it can change dynamically.
-		// MemoryMargin is used to take into account the overheads.
-
-		Context.AvailableNow = EffectiveStreamingPoolSize - ThreadStats.TotalPossibleResidentSize; // Account for worst case based on the current state plus pending requests.
-		Context.AvailableLater = EffectiveStreamingPoolSize - ThreadStats.TotalWantedMipsSize;
-		Context.AvailableTempMemory = MaxTempMemoryAllowed - ThreadStats.TotalTempMemorySize;
-		
-		SET_MEMORY_STAT( STAT_EffectiveStreamingPoolSize, EffectiveStreamingPoolSize );
-		SET_MEMORY_STAT( STAT_StreamingOverBudget,	MemoryOverBudget);
-		SET_MEMORY_STAT( STAT_StreamingUnderBudget,	FMath::Max(-LocalMemoryOverBudget,0ll) );
-		SET_MEMORY_STAT( STAT_NonStreamingTexturesSize,	FMath::Max<int64>(NonStreamingUsage, 0) );
-	}
-	else   
-	{
-		EffectiveStreamingPoolSize = 0;
-
-		LargestAlloc = MAX_int64;
-
-		Context.AvailableNow = MAX_int64;
-		Context.AvailableLater = MAX_int64;
-		Context.AvailableTempMemory = MaxTempMemoryAllowed;
-
-		MemoryOverBudget = 0;
-		SET_MEMORY_STAT( STAT_EffectiveStreamingPoolSize, 0 );
-		SET_MEMORY_STAT( STAT_StreamingOverBudget, 0 );
-		SET_MEMORY_STAT( STAT_StreamingUnderBudget, 0 );
-	}
-	NumWantingResources = ThreadStats.NumWantingTextures;
-	NumWantingResourcesCounter++;
-	SET_DWORD_STAT( STAT_NumWantingTextures, ThreadStats.NumWantingTextures );
-	SET_DWORD_STAT( STAT_NumVisibleTexturesWithLowResolutions, ThreadStats.NumVisibleTexturesWithLowResolutions);
-
-	if ( !bPauseTextureStreaming )
-	{
-		if ((Stats.IsUsingLimitedPoolSize() && Context.AvailableLater > 0 && CVarStreamingShowWantedMips.GetValueOnAnyThread() == 0) ||
-			(!Stats.IsUsingLimitedPoolSize() && GNeverStreamOutTextures))
-		{
-			KeepUnwantedMips(Context);
-		}
-		else if (Stats.IsUsingLimitedPoolSize() && Context.AvailableLater < 0)
-		{
-			DropWantedMips(Context);
-			DropForcedMips(Context);
-		}
-
-		{
-			FMemMark Mark(FMemStack::Get());
-
-			for (int32 TexPrioIndex : PrioTexIndicesSortedByLoadOrder)
-			{
-				const int32 TextureIndex = PrioritizedTextures[ TexPrioIndex ].TextureIndex;
-				if (TextureIndex == INDEX_NONE) continue;
-				FStreamingTexture& StreamingTexture = StreamingTextures[TextureIndex];
-
-				StreamingTexture.UpdateMipCount(*this, Context);
-			}
-
-			Mark.Pop();
-		}
-	}
-
-	UpdateStreamingStats( Context, true );
-
-	STAT( StreamingTimes[ProcessingStage] += FPlatformTime::Seconds() );
-}
-
-void FStreamingManagerTexture::KeepUnwantedMips( FStreamingContext& Context )
-{
-	const TArray<FTexturePriority>& PrioritizedTextures = AsyncWork->GetTask().GetPrioritizedTextures();
-
-	bool bMoreMemoryCouldBeTaken = true;
-	while (bMoreMemoryCouldBeTaken && Context.AvailableLater > 0)
-	{
-		bMoreMemoryCouldBeTaken = false;
-
-		for (int32 TexPrioIndex = PrioritizedTextures.Num() - 1; TexPrioIndex >= 0 && Context.AvailableLater > 0; --TexPrioIndex)
-		{
-			const FTexturePriority& TexturePriority = PrioritizedTextures[ TexPrioIndex ];
-			if (TexturePriority.TextureIndex == INDEX_NONE) continue;
-		
-			FStreamingTexture& StreamingTexture = StreamingTextures[TexturePriority.TextureIndex];
-			if (StreamingTexture.WantedMips >= StreamingTexture.RequestedMips) continue;
-
-			if (StreamingTexture.bInFlight)
-			{
-				// In that case we can only cancel the request or accept the request.
-				int32 DeltaMipSize = StreamingTexture.GetSize( StreamingTexture.RequestedMips ) - StreamingTexture.GetSize( StreamingTexture.WantedMips );
-				if (DeltaMipSize > Context.AvailableLater) continue;
-
-				StreamingTexture.WantedMips = StreamingTexture.RequestedMips;
-				Context.AvailableLater -= DeltaMipSize;
-			}
-			else
-			{
-				// Otherwise we can stream out any number of mips.
-				int32 NextMipSize = StreamingTexture.GetSize( StreamingTexture.WantedMips + 1 ) - StreamingTexture.GetSize( StreamingTexture.WantedMips );
-				if (NextMipSize > Context.AvailableLater) continue;
-
-				++StreamingTexture.WantedMips;
-				Context.AvailableLater -= NextMipSize;
-
-				if (StreamingTexture.WantedMips < StreamingTexture.RequestedMips)
-				{
-					bMoreMemoryCouldBeTaken = true;
-				}
-			}
-		}
-	}
-}
-
-void FStreamingManagerTexture::DropWantedMips( FStreamingContext& Context )
-{
-	const TArray<FTexturePriority>& PrioritizedTextures = AsyncWork->GetTask().GetPrioritizedTextures();
-
-	// If this is not enough, drop one mip on everything possible, then repeat until it fits.
-	bool bMoreMemoryCouldBeFreed = true;
-	bool bIsDroppingFirstMip = true;
-	while (bMoreMemoryCouldBeFreed && Context.AvailableLater < 0)
-	{
-		bMoreMemoryCouldBeFreed = false;
-
-		for (int32 TexPrioIndex = 0; TexPrioIndex < PrioritizedTextures.Num() && Context.AvailableLater < 0; ++TexPrioIndex)
-		{
-			const FTexturePriority& TexturePriority = PrioritizedTextures[ TexPrioIndex ];
-			const int32 TextureIndex = TexturePriority.TextureIndex;
-
-			// This will also account for invalid entries.
-			if (TexturePriority.TextureIndex == INDEX_NONE || !TexturePriority.bCanDropMips) continue;
-			FStreamingTexture& StreamingTexture = StreamingTextures[TextureIndex];
-
-			// Allow to drop all mips until min mips
-			if (StreamingTexture.WantedMips <= StreamingTexture.MinAllowedMips) continue;
-
-			// When dropping the first mip, don't apply to split request unless it is the last mip request.
-			if (StreamingTexture.bHasSplitRequest && !StreamingTexture.bIsLastSplitRequest && bIsDroppingFirstMip) continue;
-
-			Context.AvailableLater += StreamingTexture.GetSize( StreamingTexture.WantedMips ) - StreamingTexture.GetSize( StreamingTexture.WantedMips - 1);
-
-			--StreamingTexture.WantedMips;
-				
-			if (StreamingTexture.WantedMips > StreamingTexture.MinAllowedMips)
-			{
-				bMoreMemoryCouldBeFreed = true;
-			}
-		}
-	bIsDroppingFirstMip = false;
-	}
-}
-
-void FStreamingManagerTexture::DropForcedMips( FStreamingContext& Context )
-{
-	const TArray<FTexturePriority>& PrioritizedTextures = AsyncWork->GetTask().GetPrioritizedTextures();
-
-	// If this is not enough, drop mip for the textures not normally allowing mips to be dropped.
-	bool bMoreMemoryCouldBeFreed = true;
-	while (bMoreMemoryCouldBeFreed && Context.AvailableLater < 0)
-	{
-		bMoreMemoryCouldBeFreed = false;
-
-		for (int32 TexPrioIndex = 0; TexPrioIndex < PrioritizedTextures.Num() && Context.AvailableLater < 0; ++TexPrioIndex)
-		{
-			const FTexturePriority& TexturePriority = PrioritizedTextures[ TexPrioIndex ];
-			const int32 TextureIndex = TexturePriority.TextureIndex;
-
-			// This will also account for invalid entries.
-			if (TexturePriority.TextureIndex == INDEX_NONE || TexturePriority.bCanDropMips) continue;
-			FStreamingTexture& StreamingTexture = StreamingTextures[TextureIndex];
-
-			// Allow to drop all mips until non resident mips.
-			if (StreamingTexture.WantedMips <= StreamingTexture.NumNonStreamingMips) continue;
-
-			Context.AvailableLater += StreamingTexture.GetSize( StreamingTexture.WantedMips ) - StreamingTexture.GetSize( StreamingTexture.WantedMips - 1);
-
-			--StreamingTexture.WantedMips;
-				
-			if (StreamingTexture.WantedMips > StreamingTexture.NumNonStreamingMips)
-			{
-				bMoreMemoryCouldBeFreed = true;
-			}
-		}
-	}
-}
-
-void FStreamingManagerTexture::CheckUserSettings()
-{	
-	int32 PoolSizeSetting = CVarStreamingPoolSize.GetValueOnGameThread();
-	int32 FixedPoolSizeSetting = CVarStreamingUseFixedPoolSize.GetValueOnGameThread();
-
-	if (FixedPoolSizeSetting == 0)
-	{
-		int64 TexturePoolSize = 0;
-		if (PoolSizeSetting == -1)
-		{
-			FTextureMemoryStats Stats;
-			RHIGetTextureMemoryStats(Stats);
-			if ( GPoolSizeVRAMPercentage > 0 && Stats.TotalGraphicsMemory > 0 )
-			{
-				int64 TotalGraphicsMemory = Stats.TotalGraphicsMemory;
-				if ( GCurrentRendertargetMemorySize > 0 )
-				{
-					TotalGraphicsMemory -= int64(GCurrentRendertargetMemorySize) * 1024;
-				}
-				float PoolSize = float(GPoolSizeVRAMPercentage) * 0.01f * float(TotalGraphicsMemory);
-
-				// Truncate the pool size to MB (but still count in bytes)
-				TexturePoolSize = int64(FGenericPlatformMath::TruncToFloat(PoolSize / 1024.0f / 1024.0f)) * 1024 * 1024;
-			}
-			else
-			{
-				TexturePoolSize = OriginalTexturePoolSize;
-			}
-		}
-		else
-		{
-			int64 PoolSize = int64(PoolSizeSetting) * 1024ll * 1024ll;
-			TexturePoolSize = PoolSize;
-		}
-
-		// Only adjust the pool size once every 10 seconds, but immediately in some other cases.
-		if ( PoolSizeSetting != PreviousPoolSizeSetting ||
-			 TexturePoolSize > GTexturePoolSize ||
-			 (FApp::GetCurrentTime() - PreviousPoolSizeTimestamp) > 10.0 )
-		{
-			if ( TexturePoolSize != GTexturePoolSize )
-			{
-				UE_LOG(LogContentStreaming,Log,TEXT("Texture pool size now %d MB"), int(TexturePoolSize/1024/1024));
-			}
-			PreviousPoolSizeSetting = PoolSizeSetting;
-			PreviousPoolSizeTimestamp = FApp::GetCurrentTime();
-			GTexturePoolSize = TexturePoolSize;
-		}
-	}
-}
-
-/**
- * Main function for the texture streaming system, based on texture priorities and asynchronous processing.
- * Updates streaming, taking into account all view infos.
- *
- * @param DeltaTime				Time since last call in seconds
- * @param bProcessEverything	[opt] If true, process all resources with no throttling limits
- */
-static TAutoConsoleVariable<int32> CVarFramesForFullUpdate(
-	TEXT("r.Streaming.FramesForFullUpdate"),
-	5,
-	TEXT("Texture streaming is time sliced per frame. This values gives the number of frames to visit all textures."));
-
-void FStreamingManagerTexture::UpdateResourceStreaming( float DeltaTime, bool bProcessEverything/*=false*/ )
-{
-	SCOPE_CYCLE_COUNTER(STAT_GameThreadUpdateTime);
-
-#if STREAMING_LOG_VIEWCHANGES
-	static bool bWasLocationOveridden = false;
-	bool bIsLocationOverridden = false;
-	for ( int32 ViewIndex=0; ViewIndex < CurrentViewInfos.Num(); ++ViewIndex )
-	{
-		FStreamingViewInfo& ViewInfo = CurrentViewInfos( ViewIndex );
-		if ( ViewInfo.bOverrideLocation )
-		{
-			bIsLocationOverridden = true;
-			break;
-		}
-	}
-	if ( bIsLocationOverridden != bWasLocationOveridden )
-	{
-		UE_LOG(LogContentStreaming, Log, TEXT("Texture streaming view location is now %s."), bIsLocationOverridden ? TEXT("OVERRIDDEN") : TEXT("normal") );
-		bWasLocationOveridden = bIsLocationOverridden;
-	}
-#endif
-	int32 NewNumTextureProcessingStages = CVarFramesForFullUpdate.GetValueOnGameThread();
-	if (NewNumTextureProcessingStages > 0 && NewNumTextureProcessingStages != NumTextureProcessingStages)
-	{
-		QUICK_SCOPE_CYCLE_COUNTER(STAT_FStreamingManagerTexture_UpdateResourceStreaming_EnsureCompletion_E);
-		AsyncWork->EnsureCompletion();
-		ProcessingStage = 0;
-		NumTextureProcessingStages = NewNumTextureProcessingStages;
-	}
-	int32 OldNumTextureProcessingStages = NumTextureProcessingStages;
-	if ( bProcessEverything || IndividualStreamingTexture )
-	{
-		QUICK_SCOPE_CYCLE_COUNTER(STAT_FStreamingManagerTexture_UpdateResourceStreaming_EnsureCompletion_A);
-		AsyncWork->EnsureCompletion();
-		ProcessingStage = 0;
-		NumTextureProcessingStages = 1;
-	}
-
-#if STATS
-	if ( ProcessingStage == 0 )
-	{
-		STAT( StreamingTimes.Empty( NumTextureProcessingStages ) );
-		STAT( StreamingTimes.AddZeroed( NumTextureProcessingStages ) );
-	}
-	STAT( StreamingTimes[ProcessingStage] -= FPlatformTime::Seconds() );
-#endif
-
-	// Init.
-	if ( ProcessingStage == 0 )
-	{
-		// Is the AsyncWork is running for some reason? (E.g. we reset the system by simply setting ProcessingStage to 0.)
-		if ( AsyncWork->IsDone() == false )
-		{
-			QUICK_SCOPE_CYCLE_COUNTER(STAT_FStreamingManagerTexture_UpdateResourceStreaming_EnsureCompletion_B);
-			AsyncWork->EnsureCompletion();
-		}
-
-		{
-			QUICK_SCOPE_CYCLE_COUNTER(STAT_FStreamingManagerTexture_UpdateResourceStreaming_CheckUserSettings);
-			CheckUserSettings();
-		}
-
-		{
-			QUICK_SCOPE_CYCLE_COUNTER(STAT_FStreamingManagerTexture_UpdateResourceStreaming_ResetStreamingStats);
-			ResetStreamingStats();
-		}
-		{
-			QUICK_SCOPE_CYCLE_COUNTER(STAT_FStreamingManagerTexture_UpdateResourceStreaming_UpdateThreadData);
-			UpdateThreadData();
-		}
-
-		if ( bTriggerDumpTextureGroupStats )
-		{
-			DumpTextureGroupStats( bDetailedDumpTextureGroupStats );
-		}
-		if ( bTriggerInvestigateTexture )
-		{
-			InvestigateTexture( InvestigateTextureName );
-		}
-	}
-	else
-	{
-		FRemovedTextureArray RemovedTextures;
-		DynamicComponentTextureManager->IncrementalTick(RemovedTextures, 1.f / (float)FMath::Max(NumTextureProcessingStages - 1, 1)); // -1 since we don't want to do anything at stage 0.
-		SetTexturesRemovedTimestamp(RemovedTextures);
-	}
-
-	// Non-threaded data collection.
-	int32 NumDataCollectionStages = FMath::Max( NumTextureProcessingStages - 1, 1 );
-	if ( ProcessingStage < NumDataCollectionStages )
-	{
-		// Setup a context for this run. Note that we only (potentially) collect texture stats from the AsyncWork.
-		FStreamingContext Context( bProcessEverything, IndividualStreamingTexture, false );
-		{
-			QUICK_SCOPE_CYCLE_COUNTER(STAT_FStreamingManagerTexture_UpdateResourceStreaming_UpdateStreamingTextures);
-			UpdateStreamingTextures( Context, ProcessingStage, NumDataCollectionStages );
-		}
-		{
-			QUICK_SCOPE_CYCLE_COUNTER(STAT_FStreamingManagerTexture_UpdateResourceStreaming_UpdateStreamingStats);
-			UpdateStreamingStats( Context, false );
-		}
-	}
-
-	// Start async task after the last data collection stage (if we're not paused).
-	if ( ProcessingStage == NumDataCollectionStages - 1 && bPauseTextureStreaming == false )
-	{
-		// Is the AsyncWork is running for some reason? (E.g. we reset the system by simply setting ProcessingStage to 0.)
-		if ( AsyncWork->IsDone() == false )
-		{
-			QUICK_SCOPE_CYCLE_COUNTER(STAT_FStreamingManagerTexture_UpdateResourceStreaming_EnsureCompletion_C);
-			AsyncWork->EnsureCompletion();
-		}
-
-		AsyncWork->GetTask().Reset(bCollectTextureStats);
-		if ( NumTextureProcessingStages > 1 )
-		{
-			QUICK_SCOPE_CYCLE_COUNTER(STAT_FStreamingManagerTexture_UpdateResourceStreaming_StartBackgroundTask);
-			AsyncWork->StartBackgroundTask();
-		}
-		else
-		{
-			// Perform the work synchronously on this thread.
-			QUICK_SCOPE_CYCLE_COUNTER(STAT_FStreamingManagerTexture_UpdateResourceStreaming_StartSynchronousTask);
-			AsyncWork->StartSynchronousTask();
-		}
-	}
-
-	// Are we still in the non-threaded data collection stages?
-	if ( ProcessingStage < NumTextureProcessingStages - 1 )
-	{
-		STAT( StreamingTimes[ProcessingStage] += FPlatformTime::Seconds() );
-		ProcessingStage++;
-	}
-	// Are we waiting for the async work to finish?
-	else if ( AsyncWork->IsDone() == false )
-	{
-		STAT( StreamingTimes[ProcessingStage] += FPlatformTime::Seconds() );
-	}
-	else
-	{
-		// All priorities have been calculated, do all streaming.
-		QUICK_SCOPE_CYCLE_COUNTER(STAT_FStreamingManagerTexture_UpdateResourceStreaming_StreamTextures);
-		StreamTextures( bProcessEverything );
-		ProcessingStage = 0;
-	}
-
-	NumTextureProcessingStages = OldNumTextureProcessingStages;
-	SET_FLOAT_STAT( STAT_DynamicStreamingTotal, GStreamingDynamicPrimitivesTime);
-}
-
-/**
- * Blocks till all pending requests are fulfilled.
- *
- * @param TimeLimit		Optional time limit for processing, in seconds. Specifying 0 means infinite time limit.
- * @param bLogResults	Whether to dump the results to the log.
- * @return				Number of streaming requests still in flight, if the time limit was reached before they were finished.
- */
-int32 FStreamingManagerTexture::BlockTillAllRequestsFinished( float TimeLimit /*= 0.0f*/, bool bLogResults /*= false*/ )
-{
-	double StartTime = FPlatformTime::Seconds();
-	float ElapsedTime = 0.0f;
-
-	FMemMark Mark(FMemStack::Get());
-
-	int32 NumPendingUpdates = 0;
-	int32 MaxPendingUpdates = 0;
-
-	// Add all textures to the initial pending array.
-	TArray<int32, TMemStackAllocator<> > PendingTextures[2];
-	PendingTextures[0].Empty( StreamingTextures.Num() );
-	for ( int32 TextureIndex=0; TextureIndex < StreamingTextures.Num(); ++TextureIndex )
-	{
-		PendingTextures[0].Add( TextureIndex );
-	}
-
-	int32 CurrentArray = 0;
-	do 
-	{
-		// Flush rendering commands.
-		FlushRenderingCommands();
-
-		// Update the textures in the current pending array, and add the ones that are still pending to the other array.
-		PendingTextures[1 - CurrentArray].Empty( StreamingTextures.Num() );
-		for ( int32 Index=0; Index < PendingTextures[CurrentArray].Num(); ++Index )
-		{
-			int32 TextureIndex = PendingTextures[CurrentArray][ Index ];
-			FStreamingTexture& StreamingTexture = StreamingTextures[ TextureIndex ];
-
-			// Make sure this streaming texture hasn't been marked for removal.
-			if ( StreamingTexture.Texture )
-			{
-				if ( StreamingTexture.Texture->UpdateStreamingStatus() )
-				{
-					PendingTextures[ 1 - CurrentArray ].Add( TextureIndex );
-				}
-				TrackTextureEvent( &StreamingTexture, StreamingTexture.Texture, StreamingTexture.bForceFullyLoad, this );
-			}
-		}
-
-		// Swap arrays.
-		CurrentArray = 1 - CurrentArray;
-
-		NumPendingUpdates = PendingTextures[ CurrentArray ].Num();
-		MaxPendingUpdates = FMath::Max( MaxPendingUpdates, NumPendingUpdates );
-
-		// Check for time limit.
-		ElapsedTime = float(FPlatformTime::Seconds() - StartTime);
-		if ( TimeLimit > 0.0f && ElapsedTime > TimeLimit )
-		{
-			break;
-		}
-
-		if ( NumPendingUpdates )
-		{
-			FPlatformProcess::Sleep( 0.010f );
-		}
-	} while ( NumPendingUpdates );
-
-	Mark.Pop();
-
-	if ( bLogResults )
-	{
-		UE_LOG(LogContentStreaming, Log, TEXT("Blocking on texture streaming: %.1f ms (%d textures updated, %d still pending)"), ElapsedTime*1000.0f, MaxPendingUpdates, NumPendingUpdates );
-#if STREAMING_LOG_VIEWCHANGES
-		for ( int32 ViewIndex=0; ViewIndex < CurrentViewInfos.Num(); ++ViewIndex )
-		{
-			FStreamingViewInfo& ViewInfo = CurrentViewInfos( ViewIndex );
-			if ( ViewInfo.bOverrideLocation )
-			{
-				UE_LOG(LogContentStreaming, Log, TEXT("Texture streaming view: X=%1.f, Y=%.1f, Z=%.1f (Override=%d, Boost=%.1f)"), ViewInfo.ViewOrigin.X, ViewInfo.ViewOrigin.Y, ViewInfo.ViewOrigin.Z, ViewInfo.bOverrideLocation, ViewInfo.BoostFactor );
-				break;
-			}
-		}
-#endif
-	}
-	return NumPendingUpdates;
-}
-
-/**
- * Adds a textures streaming handler to the array of handlers used to determine which
- * miplevels need to be streamed in/ out.
- *
- * @param TextureStreamingHandler	Handler to add
- */
-void FStreamingManagerTexture::AddTextureStreamingHandler( FStreamingHandlerTextureBase* TextureStreamingHandler )
-{
-	AsyncWork->EnsureCompletion();
-	TextureStreamingHandlers.Add( TextureStreamingHandler );
-}
-
-/**
- * Removes a textures streaming handler from the array of handlers used to determine which
- * miplevels need to be streamed in/ out.
- *
- * @param TextureStreamingHandler	Handler to remove
- */
-void FStreamingManagerTexture::RemoveTextureStreamingHandler( FStreamingHandlerTextureBase* TextureStreamingHandler )
-{
-	AsyncWork->EnsureCompletion();
-	TextureStreamingHandlers.Remove( TextureStreamingHandler );
-}
-
-/**
- * Calculates the minimum and maximum number of mip-levels for a streaming texture.
- */
-void FStreamingManagerTexture::CalcMinMaxMips( FStreamingTexture& StreamingTexture )
-{
-	int32 TextureLODBias = StreamingTexture.TextureLODBias;
-
-	// Figure out whether texture should be forced resident based on bools and forced resident time.
-	if( StreamingTexture.bForceFullyLoad )
-	{
-		// If the texture has cinematic high-res mip-levels, allow them to be streamed in.
-		TextureLODBias = FMath::Max( TextureLODBias - StreamingTexture.NumCinematicMipLevels, 0 );
-	}
-
-	// We only stream in skybox textures as you are always within the bounding box and those textures
-	// tend to be huge. A streaming fudge factor fluctuating around 1 will cause them to be streamed in
-	// and out, making it likely for the allocator to fragment.
-	if( StreamingTexture.LODGroup == TEXTUREGROUP_Skybox )
-	{
-		StreamingTexture.bForceFullyLoad = true;
-	}
-
-	// Don't stream in all referenced textures but rather only those that have been rendered in the last 5 minutes if
-	// we only stream in textures. This means you still might see texture popping, but the option is designed to avoid
-	// hitching due to CPU overhead, which is still taken care off by the 5 minute rule.
-	static auto CVarOnlyStreamInTextures = IConsoleManager::Get().FindTConsoleVariableDataInt(TEXT("r.OnlyStreamInTextures"));
-	if( CVarOnlyStreamInTextures->GetValueOnAnyThread() != 0 )
-	{
-		float SecondsSinceLastRender = StreamingTexture.LastRenderTime;
-		if( SecondsSinceLastRender < 300 )
-		{
-			StreamingTexture.bForceFullyLoad = true;
-		}
-	}
-
-	const int32 HLODStategy = CVarStreamingHLODStrategy.GetValueOnAnyThread();
-	if (StreamingTexture.bIsHLODTexture && HLODStategy == 2) // disable streaming
-	{
-		StreamingTexture.bForceFullyLoad = true;
-	}
-
-	// Calculate the min/max allowed mips.
-	UTexture2D::CalcAllowedMips(
-		StreamingTexture.MipCount,
-		StreamingTexture.NumNonStreamingMips,
-		TextureLODBias,
-		StreamingTexture.MinAllowedMips,
-		StreamingTexture.MaxAllowedMips );
-
-	// Take the reduced texture pool in to account.
-	StreamingTexture.MaxAllowedOptimalMips = StreamingTexture.MaxAllowedMips;
-	if ( GIsOperatingWithReducedTexturePool )
-	{				
-		int32 MaxTextureMipCount = FMath::Max( GMaxTextureMipCount - 2, 0 );
-		StreamingTexture.MaxAllowedMips = FMath::Min( StreamingTexture.MaxAllowedMips, MaxTextureMipCount );
-	}
-
-	// Should this texture be fully streamed in?
-	if ( StreamingTexture.bForceFullyLoad )
-	{
-		StreamingTexture.MinAllowedMips = StreamingTexture.MaxAllowedMips;
-	}
-	// Check if the texture LOD group restricts the number of streaming mip-levels (in absolute terms).
-	else if ( ThreadSettings.NumStreamedMips[StreamingTexture.LODGroup] >= 0 )
-	{
-		StreamingTexture.MinAllowedMips = FMath::Clamp( StreamingTexture.MipCount - ThreadSettings.NumStreamedMips[StreamingTexture.LODGroup], StreamingTexture.MinAllowedMips, StreamingTexture.MaxAllowedMips );
-	}
-
-	if ( StreamingTexture.bIsHLODTexture && HLODStategy == 1 && StreamingTexture.MaxAllowedMips > StreamingTexture.MinAllowedMips ) // stream only mip 0
-	{
-		StreamingTexture.MinAllowedMips = StreamingTexture.MaxAllowedMips - 1;
-	}
-
-	check( StreamingTexture.MinAllowedMips > 0 && StreamingTexture.MinAllowedMips <= StreamingTexture.MipCount );
-	check( StreamingTexture.MaxAllowedMips >= StreamingTexture.MinAllowedMips && StreamingTexture.MaxAllowedMips <= StreamingTexture.MipCount );
-}
-
-/**
- * Updates this frame's STATs by one texture.
- */
-void FStreamingManagerTexture::UpdateFrameStats( FStreamingContext& Context, FStreamingTexture& StreamingTexture, int32 TextureIndex )
-{
-#if STATS_FAST || STATS
-	STAT(int64 ResidentSize = StreamingTexture.GetSize(StreamingTexture.ResidentMips));
-	STAT(Context.ThisFrameTotalStreamingTexturesSize += ResidentSize);
-	STAT(int32 PerfectWantedMips = FMath::Clamp(StreamingTexture.PerfectWantedMips, StreamingTexture.MinAllowedMips, StreamingTexture.MaxAllowedOptimalMips) );
-	STAT(int64 PerfectWantedSize = StreamingTexture.GetSize( PerfectWantedMips ));
-	STAT(int64 MostResidentSize = StreamingTexture.GetSize( StreamingTexture.MostResidentMips ) );
-	STAT(Context.ThisFrameOptimalWantedSize += PerfectWantedSize );
-#endif
-
-#if STATS
-	int64 PotentialSize = StreamingTexture.GetSize(StreamingTexture.MaxAllowedMips);
-	ETextureStreamingType StreamType = StreamingTexture.GetStreamingType();
-	switch ( StreamType )
-	{
-		case StreamType_Forced:
-			Context.ThisFrameTotalForcedHeuristicSize += PerfectWantedSize;
-			break;
-		case StreamType_LastRenderTime:
-			Context.ThisFrameTotalLastRenderHeuristicSize += PerfectWantedSize;
-			break;
-		case StreamType_Dynamic:
-			Context.ThisFrameTotalDynamicTextureHeuristicSize += PerfectWantedSize;
-			break;
-		case StreamType_Static:
-			Context.ThisFrameTotalStaticTextureHeuristicSize += PerfectWantedSize;
-			break;
-		case StreamType_Orphaned:
-		case StreamType_Other:
-			break;
-	}
-	if ( Context.bCollectTextureStats )
-	{
-		FString TextureName = StreamingTexture.Texture->GetFullName();
-		if ( CollectTextureStatsName.Len() == 0 || TextureName.Contains( CollectTextureStatsName) )
-		{
-			new (Context.TextureStats) FTextureStreamingStats( StreamingTexture.Texture, StreamType, StreamingTexture.ResidentMips, PerfectWantedMips, StreamingTexture.MostResidentMips, ResidentSize, PerfectWantedSize, PotentialSize, MostResidentSize, StreamingTexture.BoostFactor, StreamingTexture.CalcLoadOrderPriority(), TextureIndex );
-		}
-	}
-	Context.ThisFrameNumStreamingTextures++;
-	Context.ThisFrameTotalStreamingTexturesMaxSize += PotentialSize;
-	Context.ThisFrameTotalLightmapMemorySize += StreamingTexture.bIsLightmap ? ResidentSize : 0;
-	Context.ThisFrameTotalLightmapDiskSize += StreamingTexture.bIsLightmap ? PotentialSize : 0;
-	Context.ThisFrameTotalHLODMemorySize += StreamingTexture.bIsHLODTexture ? ResidentSize : 0;
-	Context.ThisFrameTotalHLODDiskSize += StreamingTexture.bIsHLODTexture ? PotentialSize : 0;
-#endif
-}
-
-/**
- * Calculates the number of mip-levels we would like to have in memory for a texture.
- */
-void FStreamingManagerTexture::CalcWantedMips( FStreamingTexture& StreamingTexture )
-{
-	FFloatMipLevel WantedMips;
-	float MinDistance = FLT_MAX;
-	FFloatMipLevel PerfectWantedMips = WantedMips;
-
-	// Figure out miplevels to request based on handlers.
-	if ( StreamingTexture.MinAllowedMips != StreamingTexture.MaxAllowedMips )
-	{
-		// Iterate over all handlers and figure out the maximum requested number of mips.
-		for( int32 HandlerIndex=0; HandlerIndex<TextureStreamingHandlers.Num(); HandlerIndex++ )
-		{
-			FStreamingHandlerTextureBase* TextureStreamingHandler = TextureStreamingHandlers[HandlerIndex];
-			float HandlerDistance = FLT_MAX;
-			FFloatMipLevel HandlerWantedMips = TextureStreamingHandler->GetWantedMips( *this, StreamingTexture, HandlerDistance );
-			WantedMips = FMath::Max( WantedMips, HandlerWantedMips );
-			PerfectWantedMips = FMath::Max( PerfectWantedMips, HandlerWantedMips );
-			MinDistance = FMath::Min( MinDistance, HandlerDistance );
-		}
-
-		bool bShouldAlsoUseLastRenderTime = false;
-		if ( StreamingTexture.LastRenderTimeRefCount > 0 || (FApp::GetCurrentTime() - StreamingTexture.LastRenderTimeRefCountTimestamp) < 91.0 )
-		{
-			bShouldAlsoUseLastRenderTime = true;
-
-#if UE_BUILD_DEBUG
-			//@DEBUG: To be able to set a break-point here.
-			if ( WantedMips.ComputeMip(&StreamingTexture, ThreadSettings.MipBias, false) != INDEX_NONE )
-			{
-				int32 Q=0;
-			}
-#endif
-		}
-
-		// Not handled by any handler, use fallback handlers.
-		if ( !WantedMips.IsHandled() || bShouldAlsoUseLastRenderTime )
-		{
-			// Try the handler for textures that has recently lost instance locations.
-			{
-				float HandlerDistance = FLT_MAX;
-				FFloatMipLevel HandlerWantedMips = GetWantedMipsForOrphanedTexture(StreamingTexture, HandlerDistance);
-				WantedMips = FMath::Max(WantedMips, HandlerWantedMips);
-				MinDistance = FMath::Min(MinDistance, HandlerDistance);
-				PerfectWantedMips = FMath::Max(PerfectWantedMips, HandlerWantedMips);
-			}
-
-			// Still wasn't handled? Use the LastRenderTime handler as last resort.
-			if (!WantedMips.IsHandled() || bShouldAlsoUseLastRenderTime )
-			{
-				// Fallback handler used if texture is not handled by any other handlers. Guaranteed to handle texture and not return INDEX_NONE.
-				FStreamingHandlerTextureLastRender FallbackStreamingHandler;
-				float HandlerDistance = FLT_MAX;
-				FFloatMipLevel HandlerWantedMips = FallbackStreamingHandler.GetWantedMips( *this, StreamingTexture, HandlerDistance );
-				WantedMips = FMath::Max( WantedMips, HandlerWantedMips );
-				MinDistance = FMath::Min( MinDistance, HandlerDistance );
-				PerfectWantedMips = FMath::Max( PerfectWantedMips, HandlerWantedMips );
-			}
-		}
-	}
-	// Stream in all allowed miplevels if requested.
-	else
-	{
-		PerfectWantedMips = FFloatMipLevel::FromMipLevel(StreamingTexture.MaxAllowedOptimalMips);
-		WantedMips = FFloatMipLevel::FromMipLevel(StreamingTexture.MaxAllowedMips);
-	}
-
-	StreamingTexture.WantedMips = WantedMips.ComputeMip(&StreamingTexture, ThreadSettings.MipBias, false);
-	StreamingTexture.MinDistance = MinDistance;
-	StreamingTexture.PerfectWantedMips = PerfectWantedMips.ComputeMip(&StreamingTexture, ThreadSettings.MipBias, true);
-}
-
-/**
- * Fallback handler to catch textures that have been orphaned recently.
- * This handler prevents massive spike in texture memory usage.
- * Orphaned textures were previously streamed based on distance but those instance locations have been removed -
- * for instance because a ULevel is being unloaded. These textures are still active and in memory until they are garbage collected,
- * so we must ensure that they do not start using the LastRenderTime fallback handler and suddenly stream in all their mips -
- * just to be destroyed shortly after by a garbage collect.
- */
-FFloatMipLevel FStreamingManagerTexture::GetWantedMipsForOrphanedTexture( FStreamingTexture& StreamingTexture, float& Distance )
-{
-	FFloatMipLevel WantedMips;
-
-	// Did we recently remove instance locations for this texture?
-	const float TimeSinceInstanceWasRemoved = float(FApp::GetCurrentTime() - StreamingTexture.InstanceRemovedTimestamp);
-
-	// Was it less than 91 seconds ago?
-	if ( TimeSinceInstanceWasRemoved < 91.0f )
-	{
-		const float TimeSinceTextureWasRendered = StreamingTexture.LastRenderTime;
-
-		// Check if it hasn't been rendered since the instances were removed (with five second margin).
-		if ( (TimeSinceTextureWasRendered - TimeSinceInstanceWasRemoved) > -5.0f )
-		{
-			// It may be garbage-collected soon. Stream out the highest mip, if currently loaded.
-			WantedMips = FFloatMipLevel::FromMipLevel(FMath::Min( StreamingTexture.ResidentMips, StreamingTexture.MaxAllowedMips - 1 ));
-			Distance = 1000.0f;
-			StreamingTexture.bUsesOrphanedHeuristics = true;
-		}
-	}
-
-	return WantedMips;
-}
-
-/**
- * Updates the I/O state of a texture (allowing it to progress to the next stage) and some stats.
- */
-void FStreamingManagerTexture::UpdateTextureStatus( FStreamingTexture& StreamingTexture, FStreamingContext& Context )
-{
-	UTexture2D* Texture = StreamingTexture.Texture;
-
-	// Update streaming status. A return value of false means that we're done streaming this texture
-	// so we can potentially request another change.
-	StreamingTexture.bInFlight		= Texture->UpdateStreamingStatus( true );
-	StreamingTexture.ResidentMips	= Texture->ResidentMips;
-	StreamingTexture.RequestedMips	= Texture->RequestedMips;
-	int32		RequestStatus			= Texture->PendingMipChangeRequestStatus.GetValue();
-	bool	bHasCancelationPending	= Texture->bHasCancelationPending;
-
-	if( bHasCancelationPending )
-	{
-		Context.ThisFrameNumRequestsInCancelationPhase++;
-	}
-	else if( RequestStatus >= TexState_ReadyFor_Finalization )
-	{
-		Context.ThisFrameNumRequestsInUpdatePhase++;
-	}
-	else if( RequestStatus == TexState_InProgress_Finalization )
-	{
-		Context.ThisFrameNumRequestsInFinalizePhase++;
-	}
-
-	// Request is in flight so there is an intermediate texture with RequestedMips miplevels.
-	if( RequestStatus > 0 )
-	{
-		Context.ThisFrameTotalIntermediateTexturesSize += StreamingTexture.GetSize(StreamingTexture.RequestedMips);
-		Context.ThisFrameNumIntermediateTextures++;
-		// Update texture increase request stats.
-		if( StreamingTexture.RequestedMips > StreamingTexture.ResidentMips )
-		{
-			Context.ThisFrameTotalMipCountIncreaseRequestsInFlight++;
-		}
-	}
-
-	STAT( UpdateLatencyStats( Texture, StreamingTexture.WantedMips, StreamingTexture.bInFlight ) );
-
-	if ( StreamingTexture.bInFlight == false )
-	{
-		check( RequestStatus == TexState_ReadyFor_Requests );
-	}
-}
-
-/**
- * Cancels the current streaming request for the specified texture.
- *
- * @param StreamingTexture		Texture to cancel streaming for
- * @return						true if a streaming request was canceled
- */
-bool FStreamingManagerTexture::CancelStreamingRequest( FStreamingTexture& StreamingTexture )
-{
-	// Mark as unavailable for further streaming action this frame.
-	StreamingTexture.bReadyForStreaming = false;
-
-	StreamingTexture.RequestedMips = StreamingTexture.ResidentMips;
-	return StreamingTexture.Texture->CancelPendingMipChangeRequest();
-}
-
-#if STATS_FAST
-bool FStreamingManagerTexture::HandleDumpTextureStreamingStatsCommand( const TCHAR* Cmd, FOutputDevice& Ar )
-{
-	Ar.Logf( TEXT("Current Texture Streaming Stats") );
-	Ar.Logf( TEXT("  Textures In Memory, Current (KB) = %f"), MaxStreamingTexturesSize / 1024.0f);
-	Ar.Logf( TEXT("  Textures In Memory, Target (KB) =  %f"), MaxOptimalTextureSize / 1024.0f );
-	Ar.Logf( TEXT("  Over Budget (KB) =                 %f"), MaxStreamingOverBudget / 1024.0f );
-	Ar.Logf( TEXT("  Pool Memory Used (KB) =            %f"), MaxTexturePoolAllocatedSize / 1024.0f );
-	Ar.Logf( TEXT("  Largest free memory hole (KB) =    %f"), MinLargestHoleSize / 1024.0f );
-	Ar.Logf( TEXT("  Num Wanting Textures =             %d"), MaxNumWantingTextures );
-	MaxStreamingTexturesSize = 0;
-	MaxOptimalTextureSize = 0;
-	MaxStreamingOverBudget = MIN_int64;
-	MaxTexturePoolAllocatedSize = 0;
-	MinLargestHoleSize = OriginalTexturePoolSize;
-	MaxNumWantingTextures = 0;
-	return true;
-}
-#endif // STATS_FAST
-
-#if STATS
-/**
- * Updates the streaming latency STAT for a texture.
- *
- * @param Texture		Texture to update for
- * @param WantedMips	Number of desired mip-levels for the texture
- * @param bInFlight		Whether the texture is currently being streamed
- */
-void FStreamingManagerTexture::UpdateLatencyStats( UTexture2D* Texture, int32 WantedMips, bool bInFlight )
-{
-	// Is the texture currently not updating?
-	if ( bInFlight == false )
-	{
-		// Did we measure the latency time?
-		if ( Texture->Timer > 0.0f )
-		{
-			double TotalLatency = LatencyAverage*NumLatencySamples - LatencySamples[LatencySampleIndex] + Texture->Timer;
-			LatencySamples[LatencySampleIndex] = Texture->Timer;
-			LatencySampleIndex = (LatencySampleIndex + 1) % NUM_LATENCYSAMPLES;
-			NumLatencySamples = ( NumLatencySamples == NUM_LATENCYSAMPLES ) ? NumLatencySamples : (NumLatencySamples+1);
-			LatencyAverage = TotalLatency / NumLatencySamples;
-			LatencyMaximum = FMath::Max(LatencyMaximum, Texture->Timer);
-		}
-
-		// Have we detected that the texture should stream in more mips?
-		if ( WantedMips > Texture->ResidentMips )
-		{
-			// Set the start time. Make it negative so we can differentiate it with a measured time.
-			Texture->Timer = -float(FPlatformTime::Seconds() - GStartTime );
-		}
-		else
-		{
-			Texture->Timer = 0.0f;
-		}
-	}
-}
-
-bool FStreamingManagerTexture::HandleListStreamingTexturesCommand( const TCHAR* Cmd, FOutputDevice& Ar )
-{
-	// Collect and report stats
-	CollectTextureStatsName = FParse::Token(Cmd, 0);
-	bCollectTextureStats = true;
-	bReportTextureStats = true;
-	return true;
-}
-
-bool FStreamingManagerTexture::HandleListStreamingTexturesCollectCommand( const TCHAR* Cmd, FOutputDevice& Ar )
-{
-	// Collect stats, but disable automatic reporting
-	CollectTextureStatsName = FParse::Token(Cmd, 0);
-	bCollectTextureStats = true;
-	bReportTextureStats = false;
-	return true;
-}
-
-bool FStreamingManagerTexture::HandleListStreamingTexturesReportReadyCommand( const TCHAR* Cmd, FOutputDevice& Ar )
-{
-	if( bCollectTextureStats )
-	{
-		if( TextureStatsReport.Num() > 0 )
-		{
-			return true;
-		}
-		return false;
-	}
-	return true;
-}
-bool FStreamingManagerTexture::HandleListStreamingTexturesReportCommand( const TCHAR* Cmd, FOutputDevice& Ar )
-{
-	// TextureStatsReport is assumed to have been populated already via ListStreamingTexturesCollect
-	if( TextureStatsReport.Num() > 0 )
-	{
-		Ar.Logf( TEXT("Listing collected stats for all streaming textures") );
-		for( int32 ReportIndex = 0; ReportIndex < TextureStatsReport.Num(); ReportIndex++ )
-		{
-			Ar.Logf(*(TextureStatsReport[ReportIndex]));
-		}
-		TextureStatsReport.Empty();
-	}
-	else
-	{
-		Ar.Logf( TEXT("No stats have been collected for streaming textures. Use ListStreamingTexturesCollect to do so.") );
-	}
-	return true;
-}
-
-#endif
-
-#if !UE_BUILD_SHIPPING
-
-bool FStreamingManagerTexture::HandleResetMaxEverRequiredTexturesCommand(const TCHAR* Cmd, FOutputDevice& Ar)
-{
-	Ar.Logf(TEXT("OldMax: %u MaxEverRequired Reset."), MaxEverRequired);
-	ResetMaxEverRequired();	
-	return true;
-}
-
-bool FStreamingManagerTexture::HandleLightmapStreamingFactorCommand( const TCHAR* Cmd, FOutputDevice& Ar )
-{
-	FString FactorString(FParse::Token(Cmd, 0));
-	float NewFactor = ( FactorString.Len() > 0 ) ? FCString::Atof(*FactorString) : GLightmapStreamingFactor;
-	if ( NewFactor >= 0.0f )
-	{
-		GLightmapStreamingFactor = NewFactor;
-	}
-	Ar.Logf( TEXT("Lightmap streaming factor: %.3f (lower values makes streaming more aggressive)."), GLightmapStreamingFactor );
-	return true;
-}
-
-bool FStreamingManagerTexture::HandleCancelTextureStreamingCommand( const TCHAR* Cmd, FOutputDevice& Ar )
-{
-	UTexture2D::CancelPendingTextureStreaming();
-	return true;
-}
-
-bool FStreamingManagerTexture::HandleShadowmapStreamingFactorCommand( const TCHAR* Cmd, FOutputDevice& Ar )
-{
-	FString FactorString(FParse::Token(Cmd, 0));
-	float NewFactor = ( FactorString.Len() > 0 ) ? FCString::Atof(*FactorString) : GShadowmapStreamingFactor;
-	if ( NewFactor >= 0.0f )
-	{
-		GShadowmapStreamingFactor = NewFactor;
-	}
-	Ar.Logf( TEXT("Shadowmap streaming factor: %.3f (lower values makes streaming more aggressive)."), GShadowmapStreamingFactor );
-	return true;
-}
-
-bool FStreamingManagerTexture::HandleNumStreamedMipsCommand( const TCHAR* Cmd, FOutputDevice& Ar )
-{
-	FString NumTextureString(FParse::Token(Cmd, 0));
-	FString NumMipsString(FParse::Token(Cmd, 0));
-	int32 LODGroup = ( NumTextureString.Len() > 0 ) ? FCString::Atoi(*NumTextureString) : MAX_int32;
-	int32 NumMips = ( NumMipsString.Len() > 0 ) ? FCString::Atoi(*NumMipsString) : MAX_int32;
-	if ( LODGroup >= 0 && LODGroup < TEXTUREGROUP_MAX )
-	{
-		FTextureLODGroup& TexGroup = UDeviceProfileManager::Get().GetActiveProfile()->GetTextureLODSettings()->GetTextureLODGroup(TextureGroup(LODGroup));
-		if ( NumMips >= -1 && NumMips <= MAX_TEXTURE_MIP_COUNT )
-		{
-			TexGroup.NumStreamedMips = NumMips;
-		}
-		Ar.Logf( TEXT("%s.NumStreamedMips = %d"), UTexture::GetTextureGroupString(TextureGroup(LODGroup)), TexGroup.NumStreamedMips );
-	}
-	else
-	{
-		Ar.Logf( TEXT("Usage: NumStreamedMips TextureGroupIndex <N>") );
-	}
-	return true;
-}
-
-bool FStreamingManagerTexture::HandleTrackTextureCommand( const TCHAR* Cmd, FOutputDevice& Ar )
-{
-	FString TextureName(FParse::Token(Cmd, 0));
-	if ( TrackTexture( TextureName ) )
-	{
-		Ar.Logf( TEXT("Textures containing \"%s\" are now tracked."), *TextureName );
-	}
-	return true;
-}
-
-bool FStreamingManagerTexture::HandleListTrackedTexturesCommand( const TCHAR* Cmd, FOutputDevice& Ar )
-{
-	FString NumTextureString(FParse::Token(Cmd, 0));
-	int32 NumTextures = ( NumTextureString.Len() > 0 ) ? FCString::Atoi(*NumTextureString) : -1;
-	ListTrackedTextures( Ar, NumTextures );
-	return true;
-}
-
-FORCEINLINE float SqrtKeepMax(float V)
-{
-	return V == FLT_MAX ? FLT_MAX : FMath::Sqrt(V);
-}
-
-bool FStreamingManagerTexture::HandleDebugTrackedTexturesCommand( const TCHAR* Cmd, FOutputDevice& Ar )
-{
-#if ENABLE_TEXTURE_TRACKING
-	int32 NumTrackedTextures = GTrackedTextureNames.Num();
-	if ( NumTrackedTextures )
-	{
-		for (int32 StreamingIndex = 0; StreamingIndex < StreamingTextures.Num(); ++StreamingIndex)
-		{
-			FStreamingTexture& StreamingTexture = StreamingTextures[StreamingIndex];
-			if (StreamingTexture.Texture)
-			{
-				// See if it matches any of the texture names that we're tracking.
-				FString TextureNameString = StreamingTexture.Texture->GetFullName();
-				const TCHAR* TextureName = *TextureNameString;
-				for ( int32 TrackedTextureIndex=0; TrackedTextureIndex < NumTrackedTextures; ++TrackedTextureIndex )
-				{
-					const FString& TrackedTextureName = GTrackedTextureNames[TrackedTextureIndex];
-					if ( FCString::Stristr(TextureName, *TrackedTextureName) != NULL )
-					{
-						FTrackedTextureEvent* LastEvent = NULL;
-						for ( int32 LastEventIndex=0; LastEventIndex < GTrackedTextures.Num(); ++LastEventIndex )
-						{
-							FTrackedTextureEvent* Event = &GTrackedTextures[LastEventIndex];
-							if ( FCString::Strcmp(TextureName, Event->TextureName) == 0 )
-							{
-								LastEvent = Event;
-								break;
-							}
-						}
-
-						if (LastEvent)
-						{
-							Ar.Logf(
-								TEXT("Texture: \"%s\", ResidentMips: %d/%d, RequestedMips: %d, WantedMips: %d, DynamicWantedMips: %d, StreamingStatus: %d, StreamType: %s, Boost: %.1f"),
-								TextureName,
-								LastEvent->NumResidentMips,
-								StreamingTexture.Texture->GetNumMips(),
-								LastEvent->NumRequestedMips,
-								LastEvent->WantedMips,
-								LastEvent->DynamicWantedMips.ComputeMip(&StreamingTexture, ThreadSettings.MipBias, false),
-								LastEvent->StreamingStatus,
-								GStreamTypeNames[StreamingTexture.GetStreamingType()],
-								StreamingTexture.BoostFactor
-								);
-						}
-						else
-						{
-							Ar.Logf(TEXT("Texture: \"%s\", StreamType: %s, Boost: %.1f"),
-								TextureName,
-								GStreamTypeNames[StreamingTexture.GetStreamingType()],
-								StreamingTexture.BoostFactor
-								);
-						}
-						for( int32 HandlerIndex=0; HandlerIndex<TextureStreamingHandlers.Num(); HandlerIndex++ )
-						{
-							FStreamingHandlerTextureBase* TextureStreamingHandler = TextureStreamingHandlers[HandlerIndex];
-							float HandlerDistance = FLT_MAX;
-							FFloatMipLevel HandlerWantedMips = TextureStreamingHandler->GetWantedMips(*this, StreamingTexture, HandlerDistance);
-							Ar.Logf(
-								TEXT("    Handler %s: WantedMips: %d, PerfectWantedMips: %d, Distance: %f"),
-								TextureStreamingHandler->HandlerName,
-								HandlerWantedMips.ComputeMip(&StreamingTexture, ThreadSettings.MipBias, false),
-								HandlerWantedMips.ComputeMip(&StreamingTexture, ThreadSettings.MipBias, true),
-								HandlerDistance
-								);
-						}
-
-						for ( int32 LevelIndex=0; LevelIndex < ThreadSettings.LevelData.Num(); ++LevelIndex )
-						{
-							FStreamingManagerTexture::FThreadLevelData& LevelData = ThreadSettings.LevelData[ LevelIndex ].Value;
-							TArray<FStreamableTextureInstance4>* TextureInstances = LevelData.ThreadTextureInstances.Find( StreamingTexture.Texture );
-							if ( TextureInstances )
-							{
-								for ( int32 InstanceIndex=0; InstanceIndex < TextureInstances->Num(); ++InstanceIndex )
-								{
-									const FStreamableTextureInstance4& TextureInstance = (*TextureInstances)[InstanceIndex];
-									for (int32 i = 0; i < 4; i++)
-									{
-										Ar.Logf(
-											TEXT("    Instance: %f,%f,%f Radius: %f Range: [%f, %f] TexelFactor: %f"),
-											TextureInstance.BoundsOriginX[i],
-											TextureInstance.BoundsOriginY[i],
-											TextureInstance.BoundsOriginZ[i],
-											TextureInstance.BoundingSphereRadius[i],
-											FMath::Sqrt(TextureInstance.MinDistanceSq[i]),
-											SqrtKeepMax(TextureInstance.MaxDistanceSq[i]),
-											TextureInstance.TexelFactor[i]
-										);
-									}
-								}
-							}
-						}
-					}
-				}
-			}
-		}
-	}
-#endif // ENABLE_TEXTURE_TRACKING
-
-	return true;
-}
-
-bool FStreamingManagerTexture::HandleUntrackTextureCommand( const TCHAR* Cmd, FOutputDevice& Ar )
-{
-	FString TextureName(FParse::Token(Cmd, 0));
-	if ( UntrackTexture( TextureName ) )
-	{
-		Ar.Logf( TEXT("Textures containing \"%s\" are no longer tracked."), *TextureName );
-	}
-	return true;
-}
-
-bool FStreamingManagerTexture::HandleStreamOutCommand( const TCHAR* Cmd, FOutputDevice& Ar )
-{
-	FString Parameter(FParse::Token(Cmd, 0));
-	int64 FreeMB = (Parameter.Len() > 0) ? FCString::Atoi(*Parameter) : 0;
-	if ( FreeMB > 0 )
-	{
-		bool bSucceeded = StreamOutTextureData( FreeMB * 1024 * 1024 );
-		Ar.Logf( TEXT("Tried to stream out %d MB of texture data: %s"), FreeMB, bSucceeded ? TEXT("Succeeded") : TEXT("Failed") );
-	}
-	else
-	{
-		Ar.Logf( TEXT("Usage: StreamOut <N> (in MB)") );
-	}
-	return true;
-}
-
-bool FStreamingManagerTexture::HandlePauseTextureStreamingCommand( const TCHAR* Cmd, FOutputDevice& Ar )
-{
-	bPauseTextureStreaming = !bPauseTextureStreaming;
-	Ar.Logf( TEXT("Texture streaming is now \"%s\"."), bPauseTextureStreaming ? TEXT("PAUSED") : TEXT("UNPAUSED") );
-	return true;
-}
-
-bool FStreamingManagerTexture::HandleStreamingManagerMemoryCommand( const TCHAR* Cmd, FOutputDevice& Ar, UWorld* InWorld )
-{
-	int32 OldStreamingTextures = 0;
-	int32 NewStreamingTextures = StreamingTextures.Num() * sizeof( FStreamingTexture );
-	for ( TObjectIterator<UTexture2D> It; It; ++It )
-	{
-		OldStreamingTextures += sizeof(TLinkedList<UTexture2D*>);
-		NewStreamingTextures += sizeof(int32);
-	}
-
-	int32 OldSystem = 0;
-	int32 OldSystemSlack = 0;
-	for( int32 LevelIndex=0; LevelIndex<InWorld->GetNumLevels(); LevelIndex++ )
-	{
-		// Find instances of the texture in this level.
-		ULevel* Level = InWorld->GetLevel(LevelIndex);
-		for ( TMap<UTexture2D*,TArray<FStreamableTextureInstance> >::TConstIterator It(Level->TextureToInstancesMap); It; ++It )
-		{
-			const TArray<FStreamableTextureInstance>& TextureInstances = It.Value();
-			OldSystem += TextureInstances.Num() * sizeof(FStreamableTextureInstance);
-			OldSystemSlack += TextureInstances.GetSlack() * sizeof(FStreamableTextureInstance);
-		}
-	}
-
-	int32 NewSystem = 0;
-	int32 NewSystemSlack = 0;
-	for ( int32 LevelIndex=0; LevelIndex < ThreadSettings.LevelData.Num(); ++LevelIndex )
-	{
-		const FLevelData& LevelData = ThreadSettings.LevelData[ LevelIndex ];
-		for ( TMap<const UTexture2D*,TArray<FStreamableTextureInstance4> >::TConstIterator It(LevelData.Value.ThreadTextureInstances); It; ++It )
-		{
-			const TArray<FStreamableTextureInstance4>& TextureInstances = It.Value();
-			NewSystem += TextureInstances.Num() * sizeof(FStreamableTextureInstance4);
-			NewSystemSlack += TextureInstances.GetSlack() * sizeof(FStreamableTextureInstance4);
-		}
-	}
-	Ar.Logf( TEXT("Old: %.2f KB used, %.2f KB slack, %.2f KB texture info"), OldSystem/1024.0f, OldSystemSlack/1024.0f, OldStreamingTextures/1024.0f );
-	Ar.Logf( TEXT("New: %.2f KB used, %.2f KB slack, %.2f KB texture info"), NewSystem/1024.0f, NewSystemSlack/1024.0f, NewStreamingTextures/1024.0f );
-
-	return true;
-}
-
-bool FStreamingManagerTexture::HandleTextureGroupsCommand( const TCHAR* Cmd, FOutputDevice& Ar )
-{
-	bDetailedDumpTextureGroupStats = FParse::Param(Cmd, TEXT("Detailed"));
-	bTriggerDumpTextureGroupStats = true;
-	return true;
-}
-
-bool FStreamingManagerTexture::HandleInvestigateTextureCommand( const TCHAR* Cmd, FOutputDevice& Ar )
-{
-	FString TextureName(FParse::Token(Cmd, 0));
-	if ( TextureName.Len() )
-	{
-		bTriggerInvestigateTexture = true;
-		InvestigateTextureName = TextureName;
-	}
-	else
-	{
-		Ar.Logf( TEXT("Usage: InvestigateTexture <name>") );
-	}
-	return true;
-}
-#endif // !UE_BUILD_SHIPPING
-/**
- * Allows the streaming manager to process exec commands.
- * @param InWorld	world contexxt
- * @param Cmd		Exec command
- * @param Ar		Output device for feedback
- * @return			true if the command was handled
- */
-bool FStreamingManagerTexture::Exec( UWorld* InWorld, const TCHAR* Cmd, FOutputDevice& Ar )
-{
-#if STATS_FAST
-	if (FParse::Command(&Cmd,TEXT("DumpTextureStreamingStats")))
-	{
-		return HandleDumpTextureStreamingStatsCommand( Cmd, Ar );
-	}
-#endif
-#if STATS
-	if (FParse::Command(&Cmd,TEXT("ListStreamingTextures")))
-	{
-		return HandleListStreamingTexturesCommand( Cmd, Ar );
-	}
-	else if (FParse::Command(&Cmd,TEXT("ListStreamingTexturesCollect")))
-	{
-		return HandleListStreamingTexturesCollectCommand( Cmd, Ar );
-	}
-	else if (FParse::Command(&Cmd,TEXT("ListStreamingTexturesReportReady")))
-	{
-		return HandleListStreamingTexturesReportReadyCommand( Cmd, Ar );
-	}
-	else if (FParse::Command(&Cmd,TEXT("ListStreamingTexturesReport")))
-	{
-		return HandleListStreamingTexturesReportCommand( Cmd, Ar );
-	}
-#endif
-#if !UE_BUILD_SHIPPING
-	if (FParse::Command(&Cmd, TEXT("ResetMaxEverRequiredTextures")))
-	{
-		return HandleResetMaxEverRequiredTexturesCommand(Cmd, Ar);
-	}
-	if (FParse::Command(&Cmd,TEXT("LightmapStreamingFactor")))
-	{
-		return HandleLightmapStreamingFactorCommand( Cmd, Ar );
-	}
-	else if (FParse::Command(&Cmd,TEXT("CancelTextureStreaming")))
-	{
-		return HandleCancelTextureStreamingCommand( Cmd, Ar );
-	}
-	else if (FParse::Command(&Cmd,TEXT("ShadowmapStreamingFactor")))
-	{
-		return HandleShadowmapStreamingFactorCommand( Cmd, Ar );
-	}
-	else if (FParse::Command(&Cmd,TEXT("NumStreamedMips")))
-	{
-		return HandleNumStreamedMipsCommand( Cmd, Ar );
-	}
-	else if (FParse::Command(&Cmd,TEXT("TrackTexture")))
-	{
-		return HandleTrackTextureCommand( Cmd, Ar );
-	}
-	else if (FParse::Command(&Cmd,TEXT("ListTrackedTextures")))
-	{
-		return HandleListTrackedTexturesCommand( Cmd, Ar );
-	}
-	else if (FParse::Command(&Cmd,TEXT("DebugTrackedTextures")))
-	{
-		return HandleDebugTrackedTexturesCommand( Cmd, Ar );
-	}
-	else if (FParse::Command(&Cmd,TEXT("UntrackTexture")))
-	{
-		return HandleUntrackTextureCommand( Cmd, Ar );
-	}
-	else if (FParse::Command(&Cmd,TEXT("StreamOut")))
-	{
-		return HandleStreamOutCommand( Cmd, Ar );
-	}
-	else if (FParse::Command(&Cmd,TEXT("PauseTextureStreaming")))
-	{
-		return HandlePauseTextureStreamingCommand( Cmd, Ar );
-	}
-	else if (FParse::Command(&Cmd,TEXT("StreamingManagerMemory")))
-	{
-		return HandleStreamingManagerMemoryCommand( Cmd, Ar, InWorld );
-	}
-	else if (FParse::Command(&Cmd,TEXT("TextureGroups")))
-	{
-		return HandleTextureGroupsCommand( Cmd, Ar );
-	}
-	else if (FParse::Command(&Cmd,TEXT("InvestigateTexture")))
-	{
-		return HandleInvestigateTextureCommand( Cmd, Ar );
-	}
-#endif // !UE_BUILD_SHIPPING
-
-	return false;
-}
-
-void FStreamingManagerTexture::DumpTextureGroupStats( bool bDetailedStats )
-{
-	bTriggerDumpTextureGroupStats = false;
-#if !UE_BUILD_SHIPPING
-	struct FTextureGroupStats
-	{
-		FTextureGroupStats()
-		{
-			FMemory::Memzero( this, sizeof(FTextureGroupStats) );
-		}
-		int32 NumTextures;
-		int32 NumNonStreamingTextures;
-		int64 CurrentTextureSize;
-		int64 WantedTextureSize;
-		int64 MaxTextureSize;
-		int64 NonStreamingSize;
-	};
-	FTextureGroupStats TextureGroupStats[TEXTUREGROUP_MAX];
-	FTextureGroupStats TextureGroupWaste[TEXTUREGROUP_MAX];
-	int64 NumNonStreamingTextures = 0;
-	int64 NonStreamingSize = 0;
-	int32 NumNonStreamingPoolTextures = 0;
-	int64 NonStreamingPoolSize = 0;
-	int64 TotalSavings = 0;
-//	int32 UITexels = 0;
-	int32 NumDXT[PF_MAX];
-	int32 NumNonSaved[PF_MAX];
-	int32 NumOneMip[PF_MAX];
-	int32 NumBadAspect[PF_MAX];
-	int32 NumTooSmall[PF_MAX];
-	int32 NumNonPow2[PF_MAX];
-	int32 NumNULLResource[PF_MAX];
-	FMemory::Memzero( &NumDXT, sizeof(NumDXT) );
-	FMemory::Memzero( &NumNonSaved, sizeof(NumNonSaved) );
-	FMemory::Memzero( &NumOneMip, sizeof(NumOneMip) );
-	FMemory::Memzero( &NumBadAspect, sizeof(NumBadAspect) );
-	FMemory::Memzero( &NumTooSmall, sizeof(NumTooSmall) );
-	FMemory::Memzero( &NumNonPow2, sizeof(NumNonPow2) );
-	FMemory::Memzero( &NumNULLResource, sizeof(NumNULLResource) );
-
-	// Gather stats.
-	for( TObjectIterator<UTexture> It; It; ++It )
-	{
-		UTexture* Texture = *It;
-		FTextureGroupStats& Stat = TextureGroupStats[Texture->LODGroup];
-		FTextureGroupStats& Waste = TextureGroupWaste[Texture->LODGroup];
-		UTexture2D* Texture2D = Cast<UTexture2D>(Texture);
-		uint32 TextureAlign = 0;
-		if ( Texture2D && Texture2D->StreamingIndex >= 0 )
-		{
-			FStreamingTexture& StreamingTexture = GetStreamingTexture( Texture2D );
-			Stat.NumTextures++;
-			Stat.CurrentTextureSize += StreamingTexture.GetSize( StreamingTexture.ResidentMips );
-			Stat.WantedTextureSize += StreamingTexture.GetSize( StreamingTexture.WantedMips );
-			Stat.MaxTextureSize += StreamingTexture.GetSize( StreamingTexture.MaxAllowedMips );
-			
-			int64 WasteCurrent = StreamingTexture.GetSize( StreamingTexture.ResidentMips ) - RHICalcTexture2DPlatformSize(Texture2D->GetSizeX(), Texture2D->GetSizeY(), Texture2D->GetPixelFormat(), StreamingTexture.ResidentMips, 1, 0, TextureAlign);			
-
-			int64 WasteWanted = StreamingTexture.GetSize( StreamingTexture.WantedMips ) - RHICalcTexture2DPlatformSize(Texture2D->GetSizeX(), Texture2D->GetSizeY(), Texture2D->GetPixelFormat(), StreamingTexture.WantedMips, 1, 0, TextureAlign);			
-
-			int64 WasteMaxSize = StreamingTexture.GetSize( StreamingTexture.MaxAllowedMips ) - RHICalcTexture2DPlatformSize(Texture2D->GetSizeX(), Texture2D->GetSizeY(), Texture2D->GetPixelFormat(), StreamingTexture.MaxAllowedMips, 1, 0, TextureAlign);			
-
-			Waste.NumTextures++;
-			Waste.CurrentTextureSize += FMath::Max<int64>(WasteCurrent,0);
-			Waste.WantedTextureSize += FMath::Max<int64>(WasteWanted,0);
-			Waste.MaxTextureSize += FMath::Max<int64>(WasteMaxSize,0);
-		}
-		else
-		{
-			bool bIsPooledTexture = Texture->Resource && IsValidRef(Texture->Resource->TextureRHI) && appIsPoolTexture( Texture->Resource->TextureRHI );
-			int64 TextureSize = Texture->CalcTextureMemorySizeEnum(TMC_ResidentMips);
-			Stat.NumNonStreamingTextures++;
-			Stat.NonStreamingSize += TextureSize;
-			if ( Texture2D && Texture2D->Resource )
-			{				
-				int64 WastedSize = TextureSize - RHICalcTexture2DPlatformSize(Texture2D->GetSizeX(), Texture2D->GetSizeY(), Texture2D->GetPixelFormat(), Texture2D->GetNumMips(), 1, 0, TextureAlign);				
-
-				Waste.NumNonStreamingTextures++;
-				Waste.NonStreamingSize += FMath::Max<int64>(WastedSize, 0);
-			}
-			if ( bIsPooledTexture )
-			{
-				NumNonStreamingPoolTextures++;
-				NonStreamingPoolSize += TextureSize;
-			}
-			else
-			{
-				NumNonStreamingTextures++;
-				NonStreamingSize += TextureSize;
-			}
-		}
-
-// 		if ( Texture2D && Texture2D->Resource && Texture2D->LODGroup == TEXTUREGROUP_UI )
-// 		{
-// 			UITexels += Texture2D->GetSizeX() * Texture2D->GetSizeY();
-// 		}
-// 
-		if ( Texture2D && (Texture2D->GetPixelFormat() == PF_DXT1 || Texture2D->GetPixelFormat() == PF_DXT5) )
-		{
-			NumDXT[Texture2D->GetPixelFormat()]++;
-			if ( Texture2D->Resource )
-			{
-				// Track the reasons we couldn't save any memory from the mip-tail.
-				NumNonSaved[Texture2D->GetPixelFormat()]++;
-				if ( Texture2D->GetNumMips() < 2 )
-				{
-					NumOneMip[Texture2D->GetPixelFormat()]++;
-				}
-				else if ( Texture2D->GetSizeX() > Texture2D->GetSizeY() * 2 || Texture2D->GetSizeY() > Texture2D->GetSizeX() * 2 )
-				{
-					NumBadAspect[Texture2D->GetPixelFormat()]++;
-				}
-				else if ( Texture2D->GetSizeX() < 16 || Texture2D->GetSizeY() < 16 || Texture2D->GetNumMips() < 5 )
-				{
-					NumTooSmall[Texture2D->GetPixelFormat()]++;
-				}
-				else if ( (Texture2D->GetSizeX() & (Texture2D->GetSizeX() - 1)) != 0 || (Texture2D->GetSizeY() & (Texture2D->GetSizeY() - 1)) != 0 )
-				{
-					NumNonPow2[Texture2D->GetPixelFormat()]++;
-				}
-				else
-				{
-					// Unknown reason
-					int32 Q=0;
-				}
-			}
-			else
-			{
-				NumNULLResource[Texture2D->GetPixelFormat()]++;
-			}
-		}
-	}
-
-	// Output stats.
-	{
-		UE_LOG(LogContentStreaming, Log, TEXT("Texture memory usage:"));
-		FTextureGroupStats TotalStats;
-		for ( int32 GroupIndex=0; GroupIndex < TEXTUREGROUP_MAX; ++GroupIndex )
-		{
-			FTextureGroupStats& Stat = TextureGroupStats[GroupIndex];
-			TotalStats.NumTextures				+= Stat.NumTextures;
-			TotalStats.NumNonStreamingTextures	+= Stat.NumNonStreamingTextures;
-			TotalStats.CurrentTextureSize		+= Stat.CurrentTextureSize;
-			TotalStats.WantedTextureSize		+= Stat.WantedTextureSize;
-			TotalStats.MaxTextureSize			+= Stat.MaxTextureSize;
-			TotalStats.NonStreamingSize			+= Stat.NonStreamingSize;
-			UE_LOG(LogContentStreaming, Log, TEXT("%34s: NumTextures=%4d, Current=%8.1f KB, Wanted=%8.1f KB, OnDisk=%8.1f KB, NumNonStreaming=%4d, NonStreaming=%8.1f KB"),
-				UTexture::GetTextureGroupString((TextureGroup)GroupIndex),
-				Stat.NumTextures,
-				Stat.CurrentTextureSize / 1024.0f,
-				Stat.WantedTextureSize / 1024.0f,
-				Stat.MaxTextureSize / 1024.0f,
-				Stat.NumNonStreamingTextures,
-				Stat.NonStreamingSize / 1024.0f );
-		}
-		UE_LOG(LogContentStreaming, Log, TEXT("%34s: NumTextures=%4d, Current=%8.1f KB, Wanted=%8.1f KB, OnDisk=%8.1f KB, NumNonStreaming=%4d, NonStreaming=%8.1f KB"),
-			TEXT("Total"),
-			TotalStats.NumTextures,
-			TotalStats.CurrentTextureSize / 1024.0f,
-			TotalStats.WantedTextureSize / 1024.0f,
-			TotalStats.MaxTextureSize / 1024.0f,
-			TotalStats.NumNonStreamingTextures,
-			TotalStats.NonStreamingSize / 1024.0f );
-	}
-	if ( bDetailedStats )
-	{
-		UE_LOG(LogContentStreaming, Log, TEXT("Wasted memory due to inefficient texture storage:"));
-		FTextureGroupStats TotalStats;
-		for ( int32 GroupIndex=0; GroupIndex < TEXTUREGROUP_MAX; ++GroupIndex )
-		{
-			FTextureGroupStats& Stat = TextureGroupWaste[GroupIndex];
-			TotalStats.NumTextures				+= Stat.NumTextures;
-			TotalStats.NumNonStreamingTextures	+= Stat.NumNonStreamingTextures;
-			TotalStats.CurrentTextureSize		+= Stat.CurrentTextureSize;
-			TotalStats.WantedTextureSize		+= Stat.WantedTextureSize;
-			TotalStats.MaxTextureSize			+= Stat.MaxTextureSize;
-			TotalStats.NonStreamingSize			+= Stat.NonStreamingSize;
-			UE_LOG(LogContentStreaming, Log, TEXT("%34s: NumTextures=%4d, Current=%8.1f KB, Wanted=%8.1f KB, OnDisk=%8.1f KB, NumNonStreaming=%4d, NonStreaming=%8.1f KB"),
-				UTexture::GetTextureGroupString((TextureGroup)GroupIndex),
-				Stat.NumTextures,
-				Stat.CurrentTextureSize / 1024.0f,
-				Stat.WantedTextureSize / 1024.0f,
-				Stat.MaxTextureSize / 1024.0f,
-				Stat.NumNonStreamingTextures,
-				Stat.NonStreamingSize / 1024.0f );
-		}
-		UE_LOG(LogContentStreaming, Log, TEXT("%34s: NumTextures=%4d, Current=%8.1f KB, Wanted=%8.1f KB, OnDisk=%8.1f KB, NumNonStreaming=%4d, NonStreaming=%8.1f KB"),
-			TEXT("Total Wasted"),
-			TotalStats.NumTextures,
-			TotalStats.CurrentTextureSize / 1024.0f,
-			TotalStats.WantedTextureSize / 1024.0f,
-			TotalStats.MaxTextureSize / 1024.0f,
-			TotalStats.NumNonStreamingTextures,
-			TotalStats.NonStreamingSize / 1024.0f );
-	}
-
-	//@TODO: Calculate memory usage for non-pool textures properly!
-//	UE_LOG(LogContentStreaming, Log,  TEXT("%34s: NumTextures=%4d, Current=%7.1f KB"), TEXT("Non-streaming pool textures"), NumNonStreamingPoolTextures, NonStreamingPoolSize/1024.0f );
-//	UE_LOG(LogContentStreaming, Log,  TEXT("%34s: NumTextures=%4d, Current=%7.1f KB"), TEXT("Non-streaming non-pool textures"), NumNonStreamingTextures, NonStreamingSize/1024.0f );
-#endif // !UE_BUILD_SHIPPING
-}
-
-/**
- * Prints out detailed information about streaming textures that has a name that contains the given string.
- * Triggered by the InvestigateTexture exec command.
- *
- * @param InvestigateTextureName	Partial name to match textures against
- */
-void FStreamingManagerTexture::InvestigateTexture( const FString& InInvestigateTextureName )
-{
-	bTriggerInvestigateTexture = false;
-#if !UE_BUILD_SHIPPING
-	for ( int32 TextureIndex=0; TextureIndex < StreamingTextures.Num(); ++TextureIndex )
-	{
-		FStreamingTexture& StreamingTexture = StreamingTextures[TextureIndex];
-		FString TextureName = StreamingTexture.Texture->GetFullName();
-		if (TextureName.Contains(InInvestigateTextureName))
-		{
-			UTexture2D* Texture2D = StreamingTexture.Texture;
-			ETextureStreamingType StreamType = StreamingTexture.GetStreamingType();
-			int32 CurrentMipIndex = FMath::Max(Texture2D->GetNumMips() - StreamingTexture.ResidentMips, 0);
-			int32 WantedMipIndex = FMath::Max(Texture2D->GetNumMips() - StreamingTexture.WantedMips, 0);
-			int32 MaxMipIndex = FMath::Max(Texture2D->GetNumMips() - StreamingTexture.MaxAllowedMips, 0);
-			UE_LOG(LogContentStreaming, Log,  TEXT("Texture: %s"), *TextureName );
-			UE_LOG(LogContentStreaming, Log,  TEXT("  Texture group:   %s"), UTexture::GetTextureGroupString(StreamingTexture.LODGroup) );
-			if ( StreamType != StreamType_LastRenderTime && StreamingTexture.bUsesLastRenderHeuristics )
-			{
-				UE_LOG(LogContentStreaming, Log,  TEXT("  Stream logic:    %s and %s (%d references)"), GStreamTypeNames[StreamType], GStreamTypeNames[StreamType_LastRenderTime], StreamingTexture.LastRenderTimeRefCount );
-			}
-			else
-			{
-				UE_LOG(LogContentStreaming, Log,  TEXT("  Stream logic:    %s"), GStreamTypeNames[StreamType] );
-			}
-			if ( StreamingTexture.LastRenderTime > 1000000.0f )
-			{
-				UE_LOG(LogContentStreaming, Log,  TEXT("  Last rendertime: Never") );
-			}
-			else
-			{
-				UE_LOG(LogContentStreaming, Log,  TEXT("  Last rendertime: %.3f seconds ago"), StreamingTexture.LastRenderTime );
-			}
-			if ( StreamingTexture.ForceLoadRefCount > 0 )
-			{
-				UE_LOG(LogContentStreaming, Log,  TEXT("  Force all mips:  %d level references"), StreamingTexture.ForceLoadRefCount );
-			}
-			else if ( Texture2D->bGlobalForceMipLevelsToBeResident )
-			{
-				UE_LOG(LogContentStreaming, Log,  TEXT("  Force all mips:  bGlobalForceMipLevelsToBeResident") );
-			}
-			else if ( Texture2D->bForceMiplevelsToBeResident )
-			{
-				UE_LOG(LogContentStreaming, Log,  TEXT("  Force all mips:  bForceMiplevelsToBeResident") );
-			}
-			else if ( Texture2D->ShouldMipLevelsBeForcedResident() )
-			{
-				float CurrentTime = float(FPlatformTime::Seconds() - GStartTime);
-				float TimeLeft = CurrentTime - Texture2D->ForceMipLevelsToBeResidentTimestamp;
-				UE_LOG(LogContentStreaming, Log,  TEXT("  Force all mips:  %.1f seconds left"), FMath::Max(TimeLeft,0.0f) );
-			}
-			else if ( StreamingTexture.MipCount == 1 )
-			{
-				UE_LOG(LogContentStreaming, Log,  TEXT("  Force all mips:  No mip-maps") );
-			}
-			UE_LOG(LogContentStreaming, Log,  TEXT("  Current size:    %dx%d"), Texture2D->PlatformData->Mips[CurrentMipIndex].SizeX, Texture2D->PlatformData->Mips[CurrentMipIndex].SizeY );
-			UE_LOG(LogContentStreaming, Log,  TEXT("  Wanted size:     %dx%d"), Texture2D->PlatformData->Mips[WantedMipIndex].SizeX, Texture2D->PlatformData->Mips[WantedMipIndex].SizeY );
-			UE_LOG(LogContentStreaming, Log,  TEXT("  Max size:        %dx%d"), Texture2D->PlatformData->Mips[MaxMipIndex].SizeX, Texture2D->PlatformData->Mips[MaxMipIndex].SizeY );
-			UE_LOG(LogContentStreaming, Log,  TEXT("  LoadOrder Priority: %.3f"), StreamingTexture.CalcLoadOrderPriority() );
-			UE_LOG(LogContentStreaming, Log,  TEXT("  Retention Priority: %.3f"), StreamingTexture.CalcRetentionPriority() );
-			UE_LOG(LogContentStreaming, Log,  TEXT("  Boost factor:    %.1f"), StreamingTexture.BoostFactor );
-			UE_LOG(LogContentStreaming, Log,  TEXT("  Allowed mips:    %d-%d"), StreamingTexture.MinAllowedMips, StreamingTexture.MaxAllowedMips );
-			UE_LOG(LogContentStreaming, Log,  TEXT("  Global mip bias: %.1f"), ThreadSettings.MipBias );
-
-			float ScreenSizeFactor;
-			switch ( StreamingTexture.LODGroup )
-			{
-				case TEXTUREGROUP_Lightmap:
-					ScreenSizeFactor = GLightmapStreamingFactor;
-					break;
-				case TEXTUREGROUP_Shadowmap:
-					ScreenSizeFactor = GShadowmapStreamingFactor;
-					break;
-				default:
-					ScreenSizeFactor = 1.0f;
-			}
-			ScreenSizeFactor *= StreamingTexture.BoostFactor;
-
-			for( int32 ViewIndex=0; ViewIndex < ThreadNumViews(); ViewIndex++ )
-			{
-				// Calculate distance of viewer to bounding sphere.
-				const FStreamingViewInfo& ViewInfo = ThreadGetView(ViewIndex);
-				UE_LOG(LogContentStreaming, Log,  TEXT("  View%d: Position=(%d,%d,%d) ScreenSize=%f Boost=%f"), ViewIndex,
-					int32(ViewInfo.ViewOrigin.X), int32(ViewInfo.ViewOrigin.Y), int32(ViewInfo.ViewOrigin.Z),
-					ViewInfo.ScreenSize, ViewInfo.BoostFactor);
-			}
-
-			// Iterate over all associated/ visible levels.
-			for ( int32 LevelIndex=0; LevelIndex < ThreadSettings.LevelData.Num(); ++LevelIndex )
-			{
-				// Find instances of the texture in this level.
-				FStreamingManagerTexture::FThreadLevelData& LevelData = ThreadSettings.LevelData[ LevelIndex ].Value;
-				TArray<FStreamableTextureInstance4>* TextureInstances = LevelData.ThreadTextureInstances.Find( StreamingTexture.Texture );
-				if ( TextureInstances )
-				{
-					for ( int32 InstanceIndex=0; InstanceIndex < TextureInstances->Num(); ++InstanceIndex )
-					{
-						const FStreamableTextureInstance4& Instance = (*TextureInstances)[InstanceIndex];
-						const FVector4& CenterX = Instance.BoundsOriginX;
-						const FVector4& CenterY = Instance.BoundsOriginY;
-						const FVector4& CenterZ = Instance.BoundsOriginZ;
-						for ( int32 PartialIndex=0; PartialIndex < 4; ++PartialIndex )
-						{
-							if ( CenterX[PartialIndex] < 3.402823466e+30F )
-							{
-								FVector Center( CenterX[PartialIndex], CenterY[PartialIndex], CenterZ[PartialIndex] );
-								float Radius = Instance.BoundingSphereRadius[PartialIndex];
-								float TexelFactor = Instance.TexelFactor[PartialIndex];
-								float MinRelevantDistance = FMath::Sqrt(Instance.MinDistanceSq[PartialIndex]);
-								float MaxRelevantDistance = SqrtKeepMax(Instance.MaxDistanceSq[PartialIndex]);
-
-								float MinDistance = MAX_FLT;
-								float OutOfRangeDistance = MAX_FLT;
-								FFloatMipLevel WantedMipCount;
-								bool bInside = false;
-								bool bIsInRange = false;
-
-								for( int32 ViewIndex=0; ViewIndex < ThreadNumViews(); ViewIndex++ )
-								{
-									// Calculate distance of viewer to bounding sphere.
-									const FStreamingViewInfo& ViewInfo = ThreadGetView(ViewIndex);
-									float Distance = (ViewInfo.ViewOrigin - Center).Size();
-									if (Distance >= MinRelevantDistance && Distance <= MaxRelevantDistance)
-									{
-										bIsInRange = true;
-										MinDistance = FMath::Min(MinDistance, Distance);
-										float DistSqMinusRadiusSq = ClampMeshToCameraDistanceSquared(FMath::Square(Distance) - FMath::Square(Radius));
-										if( DistSqMinusRadiusSq > 1.f )
-										{
-											// Outside the texture instance bounding sphere, calculate miplevel based on screen space size of bounding sphere.
-											// Calculate the maximum screen space dimension in pixels.
-											const float ScreenSize = ViewInfo.ScreenSize * ViewInfo.BoostFactor * ScreenSizeFactor;
-											const float	ScreenSizeInTexels = TexelFactor * FMath::InvSqrtEst( DistSqMinusRadiusSq ) * ScreenSize;
-											// WantedMipCount is the number of mips so we need to adjust with "+ 1".
-											WantedMipCount = FMath::Max( WantedMipCount, FFloatMipLevel::FromScreenSizeInTexels(ScreenSizeInTexels) );
-										}
-										else
-										{
-											// Request all miplevels to be loaded if we're inside the bounding sphere.
-											WantedMipCount = FFloatMipLevel::FromMipLevel(StreamingTexture.MaxAllowedMips);
-											bInside = true;
-											break;
-										}
-									}
-									else
-									{
-										// Here we store distance in another variable since the relevant distance (from another view) might actually be greater.
-										OutOfRangeDistance = FMath::Min(OutOfRangeDistance, Distance);
-									}
-								}
-
-								FString RangeString;
-								if (MinRelevantDistance != 0 || MaxRelevantDistance != FLT_MAX)
-								{
-									if (MaxRelevantDistance == FLT_MAX)
-									{
-										RangeString = FString::Printf(TEXT("Range: [%.0f, INF], "), MinRelevantDistance);
-									}
-									else
-									{
-										RangeString = FString::Printf(TEXT("Range: [%.0f, %.0f], "), MinRelevantDistance, MaxRelevantDistance);
-									}
-								}
-									
-
-								if (bIsInRange)
-								{
-									int32 IntWantedMipCount = WantedMipCount.ComputeMip(&StreamingTexture, ThreadSettings.MipBias, false);
-									int32 WantedMip = FMath::Max(Texture2D->GetNumMips() - IntWantedMipCount, 0);
-
-									UE_LOG(LogContentStreaming, Log, TEXT("Static: Wanted=%dx%d, Distance=%.1f, %sTexelFactor=%.2f, Radius=%5.1f, Position=(%d,%d,%d)%s"),
-										Texture2D->PlatformData->Mips[WantedMip].SizeX, Texture2D->PlatformData->Mips[WantedMip].SizeY,
-										MinDistance, *RangeString, TexelFactor, Radius, int32(Center.X), int32(Center.Y), int32(Center.Z),
-										bInside ? TEXT(" [all mips]") : TEXT(""));
-								}
-								else // Here there is no use for distance it will be FLT_MAX
-								{
-									UE_LOG(LogContentStreaming, Log, TEXT("Static: Distance(OutOfRange)=%.1f, %sTexelFactor=%.2f, Radius=%5.1f, Position=(%d,%d,%d)%s"),
-										   OutOfRangeDistance, *RangeString, TexelFactor, Radius, int32(Center.X), int32(Center.Y), int32(Center.Z),
-										   bInside ? TEXT(" [all mips]") : TEXT(""));
-								}
-							}
-						}
-					}
-				}
-			}
-		}
-	}
-#endif // !UE_BUILD_SHIPPING
-}
-
-void FStreamingManagerTexture::DumpTextureInstances( const UPrimitiveComponent* Primitive, UTexture2D* Texture2D )
-{
-#if !UE_BUILD_SHIPPING && 0
-	//@TODO
-			UE_LOG(LogContentStreaming, Log, TEXT("%s: Wanted=%dx%d, Distance=%.1f, TexelFactor=%.2f, CurrentRadius=%5.1f, OriginalRadius=%5.1f, Position=(%d,%d,%d), IsRegistered=%d, Component=\"%s\""),
-				TypeString, Texture2D->PlatformData->Mips[WantedMipIndex].SizeX, Texture2D->PlatformData->Mips[WantedMipIndex].SizeY,
-				MinDistance, Instance.TexelFactor, PrimitiveData.BoundingSphere.W, 1.0f / Instance.InvOriginalRadius,
-				int32(PrimitiveData.BoundingSphere.Center.X), int32(PrimitiveData.BoundingSphere.Center.Y), int32(PrimitiveData.BoundingSphere.Center.Z),
-				PrimitiveData.bAttached ? true : false,
-				*Primitive->GetReadableName() );
-#endif // !UE_BUILD_SHIPPING
-}
-
-
-/*-----------------------------------------------------------------------------
-	Streaming handler implementations.
------------------------------------------------------------------------------*/
-
-/**
- * Returns mip count wanted by this handler for the passed in texture.
- * New version for the priority system, using SIMD and no fudge factor.
- * 
- * @param	StreamingManager	Streaming manager
- * @param	Streaming Texture	Texture to determine wanted mip count for
- * @param	MinDistance			[out] Distance to the closest instance of this texture, in world space units.
- * @return	Number of miplevels that should be streamed in or INDEX_NONE if texture isn't handled by this handler.
- */
-FFloatMipLevel FStreamingHandlerTextureStatic::GetWantedMips( FStreamingManagerTexture& StreamingManager, FStreamingTexture& StreamingTexture, float& MinDistance )
-{
-	FFloatMipLevel WantedMipCount;
-	bool bShouldAbortLoop = false;
-	bool bEntryFound = false; // True an entry for this texture exists­.
-
-	const bool bUseNewMetrics = CVarStreamingUseNewMetrics.GetValueOnAnyThread() > 0;
-	const float HiddenScale = CVarStreamingHiddenPrimitiveScale.GetValueOnAnyThread();
-	const float MaxResolution = (float)(0x1 << (StreamingTexture.MaxAllowedMips - 1));
-
-	// Nothing do to if there are no views or instances.
-	if( StreamingManager.ThreadNumViews() /*&& StreamingTexture.Instances.Num() > 0*/ )
-	{
-		float ScreenSizeFactor;
-		switch ( StreamingTexture.LODGroup )
-		{
-			case TEXTUREGROUP_Lightmap:
-				ScreenSizeFactor = GLightmapStreamingFactor;
-				break;
-			case TEXTUREGROUP_Shadowmap:
-				ScreenSizeFactor = GShadowmapStreamingFactor;
-				break;
-			default:
-				ScreenSizeFactor = 1.0f;
-		}
-
-		ScreenSizeFactor *= StreamingTexture.BoostFactor;
-
-		const VectorRegister MaxFloat4 = VectorSet((float)FLT_MAX, (float)FLT_MAX, (float)FLT_MAX, (float)FLT_MAX);
-		VectorRegister MinDistanceSq4 = MaxFloat4;
-		VectorRegister MaxTexels = VectorSet(-(float)FLT_MAX, -(float)FLT_MAX, -(float)FLT_MAX, -(float)FLT_MAX);
-		for ( int32 LevelIndex=0; LevelIndex < StreamingManager.ThreadSettings.LevelData.Num(); ++LevelIndex )
-		{
-			FStreamingManagerTexture::FThreadLevelData& LevelData = StreamingManager.ThreadSettings.LevelData[ LevelIndex ].Value;
-			TArray<FStreamableTextureInstance4>* TextureInstances = LevelData.ThreadTextureInstances.Find( StreamingTexture.Texture );
-			if ( TextureInstances )
-			{
-				bEntryFound = true;
-
-				for ( int32 InstanceIndex=0; InstanceIndex < TextureInstances->Num() && !bShouldAbortLoop; ++InstanceIndex )
-				{
-					const FStreamableTextureInstance4& TextureInstance = (*TextureInstances)[InstanceIndex];
-
-					// Calculate distance of viewer to bounding sphere.
-					const VectorRegister CenterX = VectorLoadAligned( &TextureInstance.BoundsOriginX );
-					const VectorRegister CenterY = VectorLoadAligned( &TextureInstance.BoundsOriginY );
-					const VectorRegister CenterZ = VectorLoadAligned( &TextureInstance.BoundsOriginZ );
-
-					// Calculate distance of viewer to bounding sphere.
-					const VectorRegister ExtentX = VectorLoadAligned( &TextureInstance.BoxExtentSizeX );
-					const VectorRegister ExtentY = VectorLoadAligned( &TextureInstance.BoxExtentSizeY );
-					const VectorRegister ExtentZ = VectorLoadAligned( &TextureInstance.BoxExtentSizeZ );
-
-					// Iterate over all view infos.
-					for( int32 ViewIndex=0; ViewIndex < StreamingManager.ThreadNumViews() && !bShouldAbortLoop; ViewIndex++ )
-					{
-						const FStreamingViewInfo& ViewInfo = StreamingManager.ThreadGetView(ViewIndex);
-
-						const float MaxEffectiveScreenSize = CVarStreamingScreenSizeEffectiveMax.GetValueOnAnyThread();
-						const float EffectiveScreenSize = (MaxEffectiveScreenSize > 0.0f) ? FMath::Min(MaxEffectiveScreenSize, ViewInfo.ScreenSize) : ViewInfo.ScreenSize;
-						const float ScreenSizeFloat = EffectiveScreenSize * ViewInfo.BoostFactor * ScreenSizeFactor;
-						const VectorRegister ScreenSize = VectorLoadFloat1( &ScreenSizeFloat );
-						const VectorRegister ViewOriginX = VectorLoadFloat1( &ViewInfo.ViewOrigin.X );
-						const VectorRegister ViewOriginY = VectorLoadFloat1( &ViewInfo.ViewOrigin.Y );
-						const VectorRegister ViewOriginZ = VectorLoadFloat1( &ViewInfo.ViewOrigin.Z );
-
-						//const float DistSq = FVector::DistSquared( ViewInfo.ViewOrigin, TextureInstance.BoundingSphere.Center );
-						VectorRegister Temp = VectorSubtract( ViewOriginX, CenterX );
-						VectorRegister DistSq = VectorMultiply( Temp, Temp );
-						Temp = VectorSubtract( ViewOriginY, CenterY );
-						DistSq = VectorMultiplyAdd( Temp, Temp, DistSq );
-						Temp = VectorSubtract( ViewOriginZ, CenterZ );
-						DistSq = VectorMultiplyAdd( Temp, Temp, DistSq );
-
-						// The range is handle as the distance to the center of the bounds (and not to the bounds), 
-						// because it needs to match how HLOD visibility is handled.
-						VectorRegister ClampedDistSq = VectorMax( VectorLoadAligned( &TextureInstance.MinDistanceSq ), DistSq );
-						ClampedDistSq = VectorMin( VectorLoadAligned( &TextureInstance.MaxDistanceSq ), ClampedDistSq );
-						const VectorRegister InRangeMask = VectorCompareEQ(DistSq, ClampedDistSq);
-
-						VectorRegister DistSqMinusRadiusSq = VectorZero();
-						if (bUseNewMetrics)
-						{
-							// In this case DistSqMinusRadiusSq will contain the distance to the box^2
-							Temp = VectorSubtract( ViewOriginX, CenterX );
-							Temp = VectorAbs( Temp );
-							VectorRegister BoxRef = VectorMin( Temp, ExtentX );
-							Temp = VectorSubtract( Temp, BoxRef );
-							DistSqMinusRadiusSq = VectorMultiply( Temp, Temp );
-
-							Temp = VectorSubtract( ViewOriginY, CenterY );
-							Temp = VectorAbs( Temp );
-							BoxRef = VectorMin( Temp, ExtentY );
-							Temp = VectorSubtract( Temp, BoxRef );
-							DistSqMinusRadiusSq = VectorMultiplyAdd( Temp, Temp, DistSqMinusRadiusSq );
-
-							Temp = VectorSubtract( ViewOriginZ, CenterZ );
-							Temp = VectorAbs( Temp );
-							BoxRef = VectorMin( Temp, ExtentZ );
-							Temp = VectorSubtract( Temp, BoxRef );
-							DistSqMinusRadiusSq = VectorMultiplyAdd( Temp, Temp, DistSqMinusRadiusSq );
-						}
-						else
-						{
-							// const float DistSqMinusRadiusSq = DistSq - FMath::Square(TextureInstance.BoundingSphere.W);
-							DistSqMinusRadiusSq = VectorLoadAligned( &TextureInstance.BoundingSphereRadius );
-							DistSqMinusRadiusSq = VectorMultiply( DistSqMinusRadiusSq, DistSqMinusRadiusSq );
-							DistSqMinusRadiusSq = VectorSubtract( DistSq, DistSqMinusRadiusSq );
-						}
-						// Force at least distance of 2 for entry not in their range, in order to enter in next conditional block.
-						// Also required in order to not enter in the final block computing the WantedMipCount.
-						DistSqMinusRadiusSq = VectorSelect( InRangeMask, DistSqMinusRadiusSq, VectorMax(DistSqMinusRadiusSq, VectorSet(2.f, 2.f, 2.f, 2.f) ) );
-
-						MinDistanceSq4 = VectorMin( MinDistanceSq4, DistSqMinusRadiusSq );
-
-						if ( VectorAllGreaterThan( DistSqMinusRadiusSq, VectorOne() ) )
-						{
-							if (StreamingTexture.LODGroup != TEXTUREGROUP_Terrain_Heightmap)
-							{
-								// Calculate the maximum screen space dimension in pixels.
-								//const float ScreenSizeInTexels = TextureInstance.TexelFactor * FMath::InvSqrtEst( DistSqMinusRadiusSq ) * ScreenSize;
-								DistSqMinusRadiusSq = VectorMax( DistSqMinusRadiusSq, VectorOne() );
-								VectorRegister Texels = VectorMultiply( VectorLoadAligned( &TextureInstance.TexelFactor ), VectorReciprocalSqrt(DistSqMinusRadiusSq) );
-								Texels = VectorMultiply( Texels, ScreenSize );
-
-								// Clamp to max resolution before actually scaling, otherwise this could have not effect.
-								Texels = VectorMin(VectorSet(MaxResolution, MaxResolution, MaxResolution, MaxResolution), Texels);
-								VectorRegister VisibilityScale = VectorSelect( InRangeMask, VectorOne(), VectorSet(HiddenScale, HiddenScale, HiddenScale, HiddenScale) );
-								Texels = VectorMultiply(Texels, VisibilityScale);
-
-								MaxTexels = VectorMax( MaxTexels, Texels );
-							}
-							else
-							{
-								// To check Forced LOD...
-								VectorRegister MinForcedLODs = VectorLoadAligned( &TextureInstance.TexelFactor );
-								bool AllForcedLOD = !!VectorAllLesserThan(MinForcedLODs, VectorZero());
-								MinForcedLODs = VectorMin( MinForcedLODs, VectorSwizzle(MinForcedLODs, 2, 3, 0, 1) );
-								MinForcedLODs = VectorMin( MinForcedLODs, VectorReplicate(MinForcedLODs, 1) );
-								float MinLODValue;
-								VectorStoreFloat1( MinForcedLODs, &MinLODValue );
-
-								if (MinLODValue <= 0)
-								{
-									WantedMipCount = FMath::Max(WantedMipCount, FFloatMipLevel::FromMipLevel(StreamingTexture.MaxAllowedMips - 13 - FMath::FloorToInt(MinLODValue)));
-									if (WantedMipCount >= FFloatMipLevel::FromMipLevel(StreamingTexture.MaxAllowedMips))
-									{
-										MinDistance = 1.0f;
-										bShouldAbortLoop = true;
-									}
-								}
-
-								if (!AllForcedLOD)
-								{
-									// Calculate the maximum screen space dimension in pixels.
-									DistSqMinusRadiusSq = VectorMax( DistSqMinusRadiusSq, VectorOne() );
-									VectorRegister Texels = VectorMultiply( VectorLoadAligned( &TextureInstance.TexelFactor ), VectorReciprocalSqrt(DistSqMinusRadiusSq) );
-									Texels = VectorMultiply( Texels, ScreenSize );
-									MaxTexels = VectorMax( MaxTexels, Texels );
-								}
-							}
-						}
-						else
-						{
-							WantedMipCount = FFloatMipLevel::FromMipLevel(StreamingTexture.MaxAllowedMips);
-							MinDistance = 1.0f;
-							bShouldAbortLoop = true;
-						}
-						StreamingTexture.bUsesStaticHeuristics = true;
-					}
-				}
-
-				if (StreamingTexture.LODGroup == TEXTUREGROUP_Terrain_Heightmap && VectorAllLesserThan(MaxTexels, VectorZero())) // All ForcedLOD cases...
-				{
-					MinDistance = 1.0f;
-					bShouldAbortLoop = true;
-				}
-
-				if ( StreamingTexture.bUsesStaticHeuristics && !bShouldAbortLoop )
-				{
-					MinDistanceSq4 = VectorMin( MinDistanceSq4, VectorSwizzle(MinDistanceSq4, 2, 3, 0, 1) );
-					MinDistanceSq4 = VectorMin( MinDistanceSq4, VectorReplicate(MinDistanceSq4, 1) );
-					float MinDistanceSq;
-					VectorStoreFloat1( MinDistanceSq4, &MinDistanceSq );
-					MinDistanceSq = ClampMeshToCameraDistanceSquared(MinDistanceSq);
-					if ( MinDistanceSq > 1.0f )
-					{
-						MaxTexels = VectorMax( MaxTexels, VectorSwizzle(MaxTexels, 2, 3, 0, 1) );
-						MaxTexels = VectorMax( MaxTexels, VectorReplicate(MaxTexels, 1) );
-						float ScreenSizeInTexels;
-						VectorStoreFloat1( MaxTexels, &ScreenSizeInTexels );
-						// If every entry is out of range, we risk returning -1 from the handler, which would be the same as not handled.
-						// Not handled textures, will use fallback handlers, where here we rather want to prevent streaming the texture.
-						if (bEntryFound)
-						{
-							ScreenSizeInTexels = FMath::Max(1.f, ScreenSizeInTexels); // Ensure IsHandled()
-						}
-						// WantedMipCount is the number of mips so we need to adjust with "+ 1".
-						WantedMipCount = FMath::Max(WantedMipCount, FFloatMipLevel::FromScreenSizeInTexels(ScreenSizeInTexels));
-						MinDistance = FMath::Min( MinDistanceSq, FMath::Sqrt( MinDistanceSq ) );
-					}
-					else
-					{
-						WantedMipCount = FFloatMipLevel::FromMipLevel(StreamingTexture.MaxAllowedMips);
-						MinDistance = 1.0f;
-					}
-				}
-			}
-		}
-	}
-
-	return WantedMipCount;
-}
-
-/*-----------------------------------------------------------------------------
-	Streaming handler implementations.
------------------------------------------------------------------------------*/
-
-/**
- * Returns mip count wanted by this handler for the passed in texture.
- * New version for the priority system, using SIMD and no fudge factor.
- * 
- * @param	StreamingManager	Streaming manager
- * @param	Streaming Texture	Texture to determine wanted mip count for
- * @param	MinDistance			[out] Distance to the closest instance of this texture, in world space units.
- * @return	Number of miplevels that should be streamed in or INDEX_NONE if texture isn't handled by this handler.
- */
-FFloatMipLevel FStreamingHandlerTextureDynamic::GetWantedMips( FStreamingManagerTexture& StreamingManager, FStreamingTexture& StreamingTexture, float& MinDistance )
-{
-	float MaxSize_InRange, MinDistanceSq_InRange, MaxSize_AnyRange, MinDistanceSq_AnyRange;
-	if (StreamingManager.ThreadSettings.TextureBoundsVisibility->GetTexelSize(StreamingTexture.Texture, MaxSize_InRange, MinDistanceSq_InRange, MaxSize_AnyRange, MinDistanceSq_AnyRange))
-	{
-		//*//*//*//*//*//*//*//*//*//*//*//*//*//*//*//*
-		StreamingTexture.bUsesDynamicHeuristics = true;
-		//*//*//*//*//*//*//*//*//*//*//*//*//*//*//*//*
-
-		const float HiddenScale = CVarStreamingHiddenPrimitiveScale.GetValueOnAnyThread();
-		const float MaxResolution = (float)(0x1 << (StreamingTexture.MaxAllowedMips - 1));
-
-		float ExtraScale = StreamingTexture.BoostFactor;
-		switch ( StreamingTexture.LODGroup )
-		{
-			case TEXTUREGROUP_Lightmap:
-				ExtraScale *= GLightmapStreamingFactor;
-				break;
-			case TEXTUREGROUP_Shadowmap:
-				ExtraScale *= GShadowmapStreamingFactor;
-				break;
-			default:
-				break;
-		}
-
-		MaxSize_InRange *= ExtraScale;
-		MaxSize_AnyRange = FMath::Min(MaxResolution, MaxSize_AnyRange) * ExtraScale * HiddenScale; // Clamp by max resolution first otherwise it could do no effect.
-
-		check(MinDistanceSq_AnyRange >= 1.f);
-		MinDistance = FMath::Min(MinDistance, FMath::Sqrt(MinDistanceSq_AnyRange));
-
-		return FFloatMipLevel::FromScreenSizeInTexels(FMath::Max(MaxSize_InRange, MaxSize_AnyRange));
-	}
-	else
-	{
-		return FFloatMipLevel();
-	}
-}
-
-/**
- * Returns mip count wanted by this handler for the passed in texture. 
- * 
- * @param	StreamingManager	Streaming manager
- * @param	Streaming Texture	Texture to determine wanted mip count for
- * @param	MinDistance			[out] Distance to the closest instance of this texture, in world space units.
- * @return	Number of miplevels that should be streamed in or INDEX_NONE if texture isn't handled by this handler.
- */
-FFloatMipLevel FStreamingHandlerTextureLastRender::GetWantedMips( FStreamingManagerTexture& StreamingManager, FStreamingTexture& StreamingTexture, float& MinDistance )
-{
-	FFloatMipLevel Ret;
-
-	float SecondsSinceLastRender = StreamingTexture.LastRenderTime;	//(FApp::GetCurrentTime() - Texture->Resource->LastRenderTime);;
-	StreamingTexture.bUsesLastRenderHeuristics = true;
-
-	if( SecondsSinceLastRender < 45.0f && GStreamWithTimeFactor )
-	{
-		Ret = FFloatMipLevel::FromMipLevel(StreamingTexture.MaxAllowedMips);
-		MinDistance = 0;
-	}
-	else if( SecondsSinceLastRender < 90.0f && GStreamWithTimeFactor )
-	{
-		Ret = FFloatMipLevel::FromMipLevel(StreamingTexture.MaxAllowedMips - 1);
-		MinDistance = 1000.0f;
-	}
-	else
-	{
-		Ret = FFloatMipLevel::FromMipLevel(0);
-		MinDistance = 10000.0f;
-	}
-
-	return Ret;
-}
-
-/**
- * Returns mip count wanted by this handler for the passed in texture. 
- * New version for the priority system, using SIMD and no fudge factor.
- * 
- * @param	StreamingManager	Streaming manager
- * @param	Streaming Texture	Texture to determine wanted mip count for
- * @param	MinDistance			[out] Distance to the closest instance of this texture, in world space units.
- * @return	Number of miplevels that should be streamed in or INDEX_NONE if texture isn't handled by this handler.
- */
-FFloatMipLevel FStreamingHandlerTextureLevelForced::GetWantedMips( FStreamingManagerTexture& StreamingManager, FStreamingTexture& StreamingTexture, float& MinDistance )
-{
-	FFloatMipLevel Ret;
-	if ( StreamingTexture.ForceLoadRefCount > 0 )
-	{
-		StreamingTexture.bUsesForcedHeuristics = true;
-		Ret = FFloatMipLevel::FromMipLevel(StreamingTexture.MaxAllowedMips);
-	}
-	return Ret;
-}
-
 /*-----------------------------------------------------------------------------
 	Texture streaming helper structs.
 -----------------------------------------------------------------------------*/
