Copying //UE4/Dev-Rendering to Dev-Main (Source //UE4/Dev-Rendering@2932636)

diff --git a/Engine/Source/Runtime/Engine/Private/DebugViewModeMaterialProxy.cpp b/Engine/Source/Runtime/Engine/Private/DebugViewModeMaterialProxy.cpp
--- a/Engine/Source/Runtime/Engine/Private/DebugViewModeMaterialProxy.cpp
+++ b/Engine/Source/Runtime/Engine/Private/DebugViewModeMaterialProxy.cpp
@@ -28,22 +28,22 @@ ENGINE_API void ClearAllDebugViewMaterials()
 TMap<const FMaterial*, FDebugViewModeMaterialProxy*> FDebugViewModeMaterialProxy::DebugMaterialShaderMap;
 volatile bool FDebugViewModeMaterialProxy::bReentrantCall = false;
 
-void FDebugViewModeMaterialProxy::AddShader(UMaterialInterface* InMaterialInterface, EMaterialShaderMapUsage::Type InUsage)
+void FDebugViewModeMaterialProxy::AddShader(UMaterialInterface* InMaterialInterface, EMaterialQualityLevel::Type QualityLevel, ERHIFeatureLevel::Type FeatureLevel, EMaterialShaderMapUsage::Type InUsage)
 {
 	if (!InMaterialInterface) return;
 
-	const FMaterial* Material = InMaterialInterface->GetMaterialResource(GMaxRHIFeatureLevel);
+	const FMaterial* Material = InMaterialInterface->GetMaterialResource(FeatureLevel);
 	if (!Material) return;
 
 	if (!DebugMaterialShaderMap.Contains(Material))
 	{
-		DebugMaterialShaderMap.Add(Material, new FDebugViewModeMaterialProxy(InMaterialInterface, InUsage));
+		DebugMaterialShaderMap.Add(Material, new FDebugViewModeMaterialProxy(InMaterialInterface, QualityLevel, FeatureLevel, InUsage));
 	}
 }
 
 const FMaterial* FDebugViewModeMaterialProxy::GetShader(EDebugViewShaderMode DebugViewShaderMode, const FMaterial* Material)
 {
-	if (DebugViewShaderMode == DVSM_TexCoordScaleAnalysis)
+	if (DebugViewShaderMode == DVSM_TexCoordScaleAccuracy || DebugViewShaderMode == DVSM_TexCoordScaleAnalysis)
 	{
 		FDebugViewModeMaterialProxy** BoundMaterial = DebugMaterialShaderMap.Find(Material);
 		if (BoundMaterial)
@@ -56,7 +56,9 @@ const FMaterial* FDebugViewModeMaterialProxy::GetShader(EDebugViewShaderMode Deb
 
 void FDebugViewModeMaterialProxy::ClearAllShaders()
 {
-	if (bReentrantCall) return;
+	if (bReentrantCall || DebugMaterialShaderMap.Num() == 0) return;
+
+	FlushRenderingCommands();
 
 	bReentrantCall = true;
 	for (TMap<const FMaterial*, FDebugViewModeMaterialProxy*>::TIterator It(DebugMaterialShaderMap); It; ++It)
@@ -69,17 +71,17 @@ void FDebugViewModeMaterialProxy::ClearAllShaders()
 	bReentrantCall = false;
 }
 
-FDebugViewModeMaterialProxy::FDebugViewModeMaterialProxy(UMaterialInterface* InMaterialInterface, EMaterialShaderMapUsage::Type InUsage)
+FDebugViewModeMaterialProxy::FDebugViewModeMaterialProxy(UMaterialInterface* InMaterialInterface, EMaterialQualityLevel::Type QualityLevel, ERHIFeatureLevel::Type FeatureLevel, EMaterialShaderMapUsage::Type InUsage)
 	: FMaterial()
 	, MaterialInterface(InMaterialInterface)
 	, Material(nullptr)
 	, Usage(InUsage)
 {
-	SetQualityLevelProperties(EMaterialQualityLevel::High, false, GMaxRHIFeatureLevel);
+	SetQualityLevelProperties(QualityLevel, false, FeatureLevel);
 	Material = InMaterialInterface->GetMaterial();
 	Material->AppendReferencedTextures(ReferencedTextures);
 
-	FMaterialResource* Resource = InMaterialInterface->GetMaterialResource(GMaxRHIFeatureLevel);
+	FMaterialResource* Resource = InMaterialInterface->GetMaterialResource(FeatureLevel);
 
 	FMaterialShaderMapId ResourceId;
 	Resource->GetShaderMapId(GMaxRHIShaderPlatform, ResourceId);
diff --git a/Engine/Source/Runtime/Engine/Private/DebugViewModeMaterialProxy.h b/Engine/Source/Runtime/Engine/Private/DebugViewModeMaterialProxy.h
--- a/Engine/Source/Runtime/Engine/Private/DebugViewModeMaterialProxy.h
+++ b/Engine/Source/Runtime/Engine/Private/DebugViewModeMaterialProxy.h
@@ -22,7 +22,7 @@ public:
 		SetQualityLevelProperties(EMaterialQualityLevel::High, false, GMaxRHIFeatureLevel);
 	}
 
-	FDebugViewModeMaterialProxy(UMaterialInterface* InMaterialInterface, EMaterialShaderMapUsage::Type InUsage);
+	FDebugViewModeMaterialProxy(UMaterialInterface* InMaterialInterface, EMaterialQualityLevel::Type QualityLevel, ERHIFeatureLevel::Type FeatureLevel, EMaterialShaderMapUsage::Type InUsage);
 
 	virtual bool RequiresSynchronousCompilation() const override { return false; };
 
@@ -92,7 +92,7 @@ public:
 
 	////////////////
 
-	static void AddShader(UMaterialInterface* InMaterialInterface, EMaterialShaderMapUsage::Type InUsage);
+	static void AddShader(UMaterialInterface* InMaterialInterface, EMaterialQualityLevel::Type QualityLevel, ERHIFeatureLevel::Type FeatureLevel, EMaterialShaderMapUsage::Type InUsage);
 	static const FMaterial* GetShader(EDebugViewShaderMode DebugViewShaderMode, const FMaterial* Material);
 	static void ClearAllShaders();
 
diff --git a/Engine/Source/Runtime/Engine/Private/DistanceFieldAtlas.cpp b/Engine/Source/Runtime/Engine/Private/DistanceFieldAtlas.cpp
--- a/Engine/Source/Runtime/Engine/Private/DistanceFieldAtlas.cpp
+++ b/Engine/Source/Runtime/Engine/Private/DistanceFieldAtlas.cpp
@@ -13,18 +13,46 @@
 #include "MeshUtilities.h"
 #endif
 
+static TAutoConsoleVariable<int32> CVarDistField(
+	TEXT("r.GenerateMeshDistanceFields"),
+	0,	
+	TEXT("Whether to build distance fields of static meshes, needed for distance field AO, which is used to implement Movable SkyLight shadows.\n")
+	TEXT("Enabling will increase mesh build times and memory usage.  Changing this value will cause a rebuild of all static meshes."),
+	ECVF_ReadOnly);
+
+static TAutoConsoleVariable<int32> CVarDistFieldRes(
+	TEXT("r.DistanceFields.MaxPerMeshResolution"),
+	128,	
+	TEXT("Highest resolution (in one dimension) allowed for a single static mesh asset, used to cap the memory usage of meshes with a large scale.\n")
+	TEXT("Changing this will cause all distance fields to be rebuilt.  Large values such as 512 can consume memory very quickly! (128Mb for one asset at 512)"),
+	ECVF_ReadOnly);
+
+static TAutoConsoleVariable<float> CVarDistFieldResScale(
+	TEXT("r.DistanceFields.DefaultVoxelDensity"),
+	.1f,	
+	TEXT("Determines how the default scale of a mesh converts into distance field voxel dimensions.\n")
+	TEXT("Changing this will cause all distance fields to be rebuilt.  Large values can consume memory very quickly!"),
+	ECVF_ReadOnly);
+
 static TAutoConsoleVariable<int32> CVarDistFieldAtlasResXY(
 	TEXT("r.DistanceFields.AtlasSizeXY"),
-	512,
+	512,	
 	TEXT("Size of the global mesh distance field atlas volume texture in X and Y."),
 	ECVF_ReadOnly);
 
 static TAutoConsoleVariable<int32> CVarDistFieldAtlasResZ(
 	TEXT("r.DistanceFields.AtlasSizeZ"),
-	1024,
+	1024,	
 	TEXT("Size of the global mesh distance field atlas volume texture in Z."),
 	ECVF_ReadOnly);
 
+static TAutoConsoleVariable<int32> CVarLandscapeGI(
+	TEXT("r.GenerateLandscapeGIData"),
+	1,
+	TEXT("Whether to generate a low-resolution base color texture for landscapes for rendering real-time global illumination.\n")
+	TEXT("This feature requires GenerateMeshDistanceFields is also enabled, and will increase mesh build times and memory usage.\n"),
+	ECVF_Default);
+
 TGlobalResource<FDistanceFieldVolumeTextureAtlas> GDistanceFieldVolumeTextureAtlas = TGlobalResource<FDistanceFieldVolumeTextureAtlas>(PF_R16F);
 
 FDistanceFieldVolumeTextureAtlas::FDistanceFieldVolumeTextureAtlas(EPixelFormat InFormat) :
@@ -543,4 +571,4 @@ void FDistanceFieldAsyncQueue::Shutdown()
 		}
 		FPlatformProcess::Sleep(0.01f);
 	}
-}
\ No newline at end of file
+}
diff --git a/Engine/Source/Runtime/Engine/Private/GPUSkinCache.cpp b/Engine/Source/Runtime/Engine/Private/GPUSkinCache.cpp
--- a/Engine/Source/Runtime/Engine/Private/GPUSkinCache.cpp
+++ b/Engine/Source/Runtime/Engine/Private/GPUSkinCache.cpp
@@ -17,6 +17,7 @@
 DEFINE_STAT(STAT_GPUSkinCache_TotalNumChunks);
 DEFINE_STAT(STAT_GPUSkinCache_TotalNumVertices);
 DEFINE_STAT(STAT_GPUSkinCache_TotalMemUsed);
+DEFINE_STAT(STAT_GPUSkinCache_TotalMemWasted);
 DEFINE_STAT(STAT_GPUSkinCache_SkippedForMaxChunksPerLOD);
 DEFINE_STAT(STAT_GPUSkinCache_SkippedForChunksPerFrame);
 DEFINE_STAT(STAT_GPUSkinCache_SkippedForZeroInfluences);
@@ -38,6 +39,16 @@ static FAutoConsoleVariableRef CVarEnableGPUSkinCacheShaders(
 	ECVF_RenderThreadSafe | ECVF_ReadOnly
 	);
 
+int32 GSkinCacheSafety = 1;
+static FAutoConsoleVariableRef CVarGSkinCacheSafety(
+	TEXT("r.SkinCache.Safety"),
+	GSkinCacheSafety,
+	TEXT("To supress ensure() and check() when running on stage\n")
+	TEXT(" 0: disable ensure() and check()\n")
+	TEXT(" 1: enable ensure() and check() (default)"),
+	ECVF_RenderThreadSafe
+	);
+
 // 0/1
 int32 GEnableGPUSkinCache = 0;
 static TAutoConsoleVariable<int32> CVarEnableGPUSkinCache(
@@ -51,13 +62,14 @@ static TAutoConsoleVariable<int32> CVarEnableGPUSkinCache(
 	ECVF_RenderThreadSafe
 	);
 
-static TAutoConsoleVariable<int32> CVarGPUSkinCacheRecomputeTangents(
+TAutoConsoleVariable<int32> CVarGPUSkinCacheRecomputeTangents(
 	TEXT("r.SkinCache.RecomputeTangents"),
-	0,
+	2,
 	TEXT("If r.SkinCaching is enabled this option recomputes the vertex tangents on the GPU\n")
 	TEXT("Can be changed at runtime, Requires r.SkinCacheShaders=1 and r.SkinCache=1\n")
-	TEXT(" 0: off (default)\n")
-	TEXT(" 1: on"),
+	TEXT(" 0: off\n")
+	TEXT(" 1: on\n")
+	TEXT(" 2: on, SkinCache only objects that would require the RecomputTangents feature (default)"),
 	ECVF_RenderThreadSafe
 	);
 
@@ -95,7 +107,7 @@ ENGINE_API bool DoSkeletalMeshIndexBuffersNeedSRV()
 	return GMaxRHIShaderPlatform == SP_PCD3D_SM5 && GEnableGPUSkinCacheShaders != 0;
 }
 
-ENGINE_API bool DoRecomputeSkinTangentsOnGPU()
+ENGINE_API bool DoRecomputeSkinTangentsOnGPU_RT()
 {
 	// currently only implemented and tested on Window SM5 (needs Compute, Atomics, SRV for index buffers, UAV for VertexBuffers)
 	return GMaxRHIShaderPlatform == SP_PCD3D_SM5 && GEnableGPUSkinCacheShaders && GEnableGPUSkinCache && CVarGPUSkinCacheRecomputeTangents.GetValueOnRenderThread() != 0;
@@ -119,22 +131,22 @@ public:
 
 		//DebugParameter.Bind(Initializer.ParameterMap, TEXT("DebugParameter"));
 
-		SkinInputStreamStride.Bind(Initializer.ParameterMap, TEXT("SkinCacheVertexStride"));
-		SkinInputStreamVertexCount.Bind(Initializer.ParameterMap, TEXT("SkinCacheVertexCount"));
-		SkinOutputBufferFloatOffset.Bind(Initializer.ParameterMap, TEXT("SkinCacheOutputBufferFloatOffset"));
+		InputStreamStride.Bind(Initializer.ParameterMap, TEXT("InputStreamStride"));
+		NumVertices.Bind(Initializer.ParameterMap, TEXT("NumVertices"));
+		SkinCacheStart.Bind(Initializer.ParameterMap, TEXT("SkinCacheStart"));
 		BoneMatrices.Bind(Initializer.ParameterMap, TEXT("BoneMatrices"));
 		SkinInputStream.Bind(Initializer.ParameterMap, TEXT("SkinStreamInputBuffer"));
-		SkinInputStreamFloatOffset.Bind(Initializer.ParameterMap, TEXT("SkinCacheInputStreamFloatOffset"));
+		InputStreamStart.Bind(Initializer.ParameterMap, TEXT("InputStreamStart"));
 		
-		SkinCacheBuffer.Bind(Initializer.ParameterMap, TEXT("SkinCacheBuffer"));
+		SkinCacheBufferUAV.Bind(Initializer.ParameterMap, TEXT("SkinCacheBufferUAV"));
 
 		MorphBuffer.Bind(Initializer.ParameterMap, TEXT("MorphBuffer"));
 		MorphBufferOffset.Bind(Initializer.ParameterMap, TEXT("MorphBufferOffset"));
 		SkinCacheDebug.Bind(Initializer.ParameterMap, TEXT("SkinCacheDebug"));
 	}
 
-	void SetParameters(const FVertexBufferAndSRV& BoneBuffer, FUniformBufferRHIRef UniformBuffer,
-		FRWBuffer& SkinBuffer, const FVector& MeshOrigin, const FVector& MeshExtension,
+		void SetParameters(const FVertexBufferAndSRV& BoneBuffer, FUniformBufferRHIRef UniformBuffer,
+		const FVector& MeshOrigin, const FVector& MeshExtension,
 		const FGPUSkinCache::FDispatchData& DispatchData)
 	{
 		FComputeShaderRHIParamRef ShaderRHI = GetComputeShader();
@@ -144,9 +156,9 @@ public:
 		SetShaderValue(RHICmdList, ShaderRHI, SkinMeshOriginParameter, MeshOrigin);
 		SetShaderValue(RHICmdList, ShaderRHI, SkinMeshExtensionParameter, MeshExtension);
 
-		SetShaderValue(RHICmdList, ShaderRHI, SkinInputStreamStride, DispatchData.InputStreamStride);
-		SetShaderValue(RHICmdList, ShaderRHI, SkinInputStreamVertexCount, DispatchData.NumVertices);
-		SetShaderValue(RHICmdList, ShaderRHI, SkinInputStreamFloatOffset, DispatchData.InputStreamFloatOffset);
+		SetShaderValue(RHICmdList, ShaderRHI, InputStreamStride, DispatchData.InputStreamStride);
+		SetShaderValue(RHICmdList, ShaderRHI, NumVertices, DispatchData.NumVertices);
+		SetShaderValue(RHICmdList, ShaderRHI, InputStreamStart, DispatchData.InputStreamStart);
 		//SetShaderValue(RHICmdList, ShaderRHI, DebugParameter, FUintVector4((uint32)(uint64)VertexBufferSRV.GetReference(), 0, 0, 1));
 
 		if (UniformBuffer)
@@ -155,14 +167,15 @@ public:
 		}
 		else
 		{
+			check(!GSkinCacheSafety || BoneBuffer.VertexBufferSRV);
 			SetSRVParameter(RHICmdList, ShaderRHI, BoneMatrices, BoneBuffer.VertexBufferSRV);
 		}
 
 		SetSRVParameter(RHICmdList, ShaderRHI, SkinInputStream, DispatchData.InputVertexBufferSRV);
 
 		// output UAV
-		SetUAVParameter(RHICmdList, ShaderRHI, SkinCacheBuffer, SkinBuffer.UAV);
-		SetShaderValue(RHICmdList, ShaderRHI, SkinOutputBufferFloatOffset, DispatchData.OutputBufferFloatOffset);
+		SetUAVParameter(RHICmdList, ShaderRHI, SkinCacheBufferUAV, DispatchData.SkinCacheBuffer->UAV);
+		SetShaderValue(RHICmdList, ShaderRHI, SkinCacheStart, DispatchData.SkinCacheStart);
 	
 		const bool bMorph = DispatchData.SkinType == 1;
 		if(bMorph)
@@ -174,19 +187,20 @@ public:
 		SetShaderValue(RHICmdList, ShaderRHI, SkinCacheDebug, CVarGPUSkinCacheDebug.GetValueOnRenderThread());
 	}
 
+
 	void UnsetParameters(FRHICommandList& RHICmdList)
 	{
 		FComputeShaderRHIParamRef ShaderRHI = GetComputeShader();
-
-		SetUAVParameter(RHICmdList, ShaderRHI, SkinCacheBuffer, 0);
+		
+		SetUAVParameter(RHICmdList, ShaderRHI, SkinCacheBufferUAV, 0);
 	}
 
 	virtual bool Serialize(FArchive& Ar) override
 	{
 		bool bShaderHasOutdatedParameters = FGlobalShader::Serialize(Ar);
-		Ar  << SkinMeshOriginParameter << SkinMeshExtensionParameter << SkinInputStreamStride
-			<< SkinInputStreamVertexCount << SkinInputStreamFloatOffset << SkinOutputBufferFloatOffset
-			<< SkinInputStream << SkinCacheBuffer << BoneMatrices << MorphBuffer << MorphBufferOffset
+		Ar  << SkinMeshOriginParameter << SkinMeshExtensionParameter << InputStreamStride
+			<< NumVertices << InputStreamStart << SkinCacheStart
+			<< SkinInputStream << SkinCacheBufferUAV << BoneMatrices << MorphBuffer << MorphBufferOffset
 			<< SkinCacheDebug;
 		//Ar << DebugParameter;
 		return bShaderHasOutdatedParameters;
@@ -196,12 +210,12 @@ private:
 
 	FShaderParameter SkinMeshOriginParameter;
 	FShaderParameter SkinMeshExtensionParameter;
-
-	FShaderParameter SkinInputStreamStride;
-	FShaderParameter SkinInputStreamVertexCount;
+	
+	FShaderParameter InputStreamStride;
+	FShaderParameter NumVertices;
 	FShaderParameter SkinCacheDebug;
-	FShaderParameter SkinInputStreamFloatOffset;
-	FShaderParameter SkinOutputBufferFloatOffset;
+	FShaderParameter InputStreamStart;
+	FShaderParameter SkinCacheStart;
 
 	//FShaderParameter DebugParameter;
 
@@ -209,7 +223,7 @@ private:
 
 	FShaderResourceParameter BoneMatrices;
 	FShaderResourceParameter SkinInputStream;
-	FShaderResourceParameter SkinCacheBuffer;
+	FShaderResourceParameter SkinCacheBufferUAV;
 
 	FShaderResourceParameter MorphBuffer;
 	FShaderParameter MorphBufferOffset;
@@ -233,8 +247,7 @@ public:
 		FGlobalShader::ModifyCompilationEnvironment(Platform,OutEnvironment);
 		const uint32 UseExtraBoneInfluences = bUseExtraBoneInfluencesT;
 		OutEnvironment.SetDefine(TEXT("GPUSKIN_USE_EXTRA_INFLUENCES"), UseExtraBoneInfluences);
-		OutEnvironment.SetDefine(TEXT("GPUSKIN_MORPH_BLEND"), (uint32)(SkinType == 1));
-
+		OutEnvironment.SetDefine(TEXT("GPUSKIN_MORPH_BLEND"), SkinType == 1);
 		OutEnvironment.SetDefine(TEXT("GPUSKIN_RWBUFFER_OFFSET_POSITION"), FGPUSkinCache::RWPositionOffsetInFloats);
 		OutEnvironment.SetDefine(TEXT("GPUSKIN_RWBUFFER_OFFSET_TANGENT_X"), FGPUSkinCache::RWTangentXOffsetInFloats);
 		OutEnvironment.SetDefine(TEXT("GPUSKIN_RWBUFFER_OFFSET_TANGENT_Z"), FGPUSkinCache::RWTangentZOffsetInFloats);
@@ -271,7 +284,7 @@ public:
 
 FGPUSkinCache::FGPUSkinCache()
 	: bInitialized(false)
-	, FrameCounter(0)
+	, SkinCacheFrameNumber(0)
 	, CacheMaxVectorCount(0)
 	, CacheCurrentFloatOffset(0)
 	, CachedChunksThisFrameCount(0)
@@ -309,7 +322,7 @@ void FGPUSkinCache::Initialize(FRHICommandListImmediate& RHICmdList)
 	}
 
 	UE_LOG(LogSkinCache, Log, TEXT("  Frames:            %d"), GPUSKINCACHE_FRAMES);
-	UE_LOG(LogSkinCache, Log, TEXT("  VertexBuffers:     %.1f MB"), FMath::DivideAndRoundUp(GGPUSkinCacheBufferSize * GPUSKINCACHE_FRAMES, 1024) / 1024.0f);
+	UE_LOG(LogSkinCache, Log, TEXT("  VertexBuffers:     %.1f MB"), FMath::DivideAndRoundUp((uint32)(NumFloatsPerBuffer * GPUSKINCACHE_FRAMES * sizeof(float)), (uint32)1024) / 1024.0f);
 
 	// we always allocate so we can change r.SkinCache.RecomputeTangents at runtime
 	{
@@ -349,11 +362,11 @@ void FGPUSkinCache::Cleanup()
 }
 
 // Kick off compute shader skinning for incoming chunk, return key for fast lookup in draw passes
-int32 FGPUSkinCache::StartCacheMesh(FRHICommandListImmediate& RHICmdList, int32 Key, FGPUBaseSkinVertexFactory* VertexFactory,
+int32 FGPUSkinCache::StartCacheMesh(FRHICommandListImmediate& RHICmdList, uint32 FrameNumber, int32 Key, FGPUBaseSkinVertexFactory* VertexFactory,
 	FGPUSkinPassthroughVertexFactory* TargetVertexFactory, const FSkelMeshChunk& BatchElement, FSkeletalMeshObjectGPUSkin* Skin,
 	const FMorphVertexBuffer* MorphVertexBuffer)
 {
-	if (CachedChunksThisFrameCount >= GMaxGPUSkinCacheElementsPerFrame && GFrameNumberRenderThread <= FrameCounter)
+	if (CachedChunksThisFrameCount >= GMaxGPUSkinCacheElementsPerFrame && FrameNumber <= SkinCacheFrameNumber)
 	{
 		INC_DWORD_STAT(STAT_GPUSkinCache_SkippedForChunksPerFrame);
 		return -1;
@@ -364,11 +377,11 @@ int32 FGPUSkinCache::StartCacheMesh(FRHICommandListImmediate& RHICmdList, int32
 		Initialize(RHICmdList);
 	}
 
-	if (GFrameNumberRenderThread > FrameCounter)	// Reset cache output if on a new frame
+	if (FrameNumber > SkinCacheFrameNumber)	// Reset cache output if on a new frame
 	{
 		CacheCurrentFloatOffset = 0;
 		CachedChunksThisFrameCount = 0;
-		FrameCounter = GFrameNumberRenderThread;
+		SkinCacheFrameNumber = FrameNumber;
 	}
 
 	FElementCacheStatusInfo::FSearchInfo SInfo;
@@ -381,7 +394,17 @@ int32 FGPUSkinCache::StartCacheMesh(FRHICommandListImmediate& RHICmdList, int32
 	{
 		ECSInfo = &(CachedElements.GetData()[Key]);
 
-		if (ECSInfo->Skin != Skin || &BatchElement != ECSInfo->BatchElement)
+		if (ECSInfo->Skin == Skin && &BatchElement == ECSInfo->BatchElement)
+		{
+			// that however should not happen:
+			ensure(!GSkinCacheSafety || ECSInfo->NumVertices == BatchElement.GetNumVertices());
+			// that also should not happen:
+			ensureMsgf(!GSkinCacheSafety || ECSInfo->BaseVertexIndex == BatchElement.BaseVertexIndex, TEXT("SkinCache internal error (Debug: %p %d %d %d %d)"),
+				&BatchElement,
+				ECSInfo->NumVertices, BatchElement.GetNumVertices(),
+				ECSInfo->BaseVertexIndex, BatchElement.BaseVertexIndex);
+		}
+		else
 		{
 			ECSInfo = CachedElements.FindByKey(SInfo);
 		}
@@ -391,22 +414,29 @@ int32 FGPUSkinCache::StartCacheMesh(FRHICommandListImmediate& RHICmdList, int32
 		ECSInfo = CachedElements.FindByKey(SInfo);
 	}
 
+	int32 UAVIndex = SkinCacheFrameNumber % GPUSKINCACHE_FRAMES;
 	bool AlignStreamOffsetsFrames = false;
-
 	const bool bExtraBoneInfluences = VertexFactory->UsesExtraBoneInfluences();
 	bool bHasPreviousOffset = false;
-	if (ECSInfo != nullptr)
+	if (ECSInfo)
 	{
-		if (ECSInfo->FrameUpdated == GFrameNumberRenderThread)	// Already cached this element this frame, return 0 key
+		if (ECSInfo->FrameUpdated == FrameNumber)
 		{
+			// Already cached this element this frame
 			INC_DWORD_STAT(STAT_GPUSkinCache_SkippedForAlreadyCached);
-			return -1;
+			// This line is needed when dragging in a specific skeletal mesh which uses APEX cloth (the PassThroughVertexFactory get recreated but not the data
+			// we store in it, also FrameNumber seems to not be increased, maybe we should fix that).
+			// //UE4/Private-GDC16-Character/Samples/NotForLicensees/GDC2016/CharacterDemo/Content/Developers/matthewstoneham/APEX/
+			TargetVertexFactory->UpdateVertexDeclaration(VertexFactory, &SkinCacheBuffer[UAVIndex]);
+			return ECSInfo->Key;
 		}
 
-		AlignStreamOffsetsFrames = (GFrameNumberRenderThread - ECSInfo->FrameUpdated) > 2;
+		// Seems needed because we bump the frame number in each rendering but then 2 is not always enough
+		// TODO: Fix by using per scene/world frame number or store previsou frame number in scene/world
+		AlignStreamOffsetsFrames = (FrameNumber - ECSInfo->FrameUpdated) > 2;
 
-		ECSInfo->FrameUpdated = GFrameNumberRenderThread;
 		ECSInfo->PreviousFrameStreamOffset = ECSInfo->StreamOffset;
+		ECSInfo->PreviousFrameUpdated = ECSInfo->FrameUpdated;
 		bHasPreviousOffset = true;
 	}
 	else	// New element, init and add
@@ -414,53 +444,42 @@ int32 FGPUSkinCache::StartCacheMesh(FRHICommandListImmediate& RHICmdList, int32
 		if (CachedElements.Num() < MaxCachedElements)
 		{
 			FElementCacheStatusInfo	Info;
-			Info.BatchElement = &BatchElement;
-			Info.VertexFactory = VertexFactory;
-			Info.TargetVertexFactory = TargetVertexFactory;
-			Info.Skin = Skin;
 			Info.Key = CachedElements.Num();
-			Info.FrameUpdated = GFrameNumberRenderThread;
-			Info.bExtraBoneInfluences = bExtraBoneInfluences;
 			CachedElements.Add(Info);
 
 			ECSInfo = &(CachedElements.GetData()[Info.Key]);
 		}
 		else
 		{
-			FElementCacheStatusInfo*	Info = FindEvictableCacheStatusInfo();
+			ECSInfo = FindEvictableCacheStatusInfo(FrameNumber);
 
-			if (Info == nullptr)
+			if (ECSInfo == nullptr)
 			{
 				INC_DWORD_STAT(STAT_GPUSkinCache_SkippedForOutOfECS);
 				return -1;
 			}
-
-			Info->BatchElement = &BatchElement;
-			Info->VertexFactory = VertexFactory;
-			Info->TargetVertexFactory = TargetVertexFactory;
-			Info->Skin = Skin;
-			Info->FrameUpdated = GFrameNumberRenderThread;
-			Info->bExtraBoneInfluences = bExtraBoneInfluences;
-
-			ECSInfo = Info;
 		}
+		ECSInfo->BatchElement = &BatchElement;
+		ECSInfo->Skin = Skin;
 
 		AlignStreamOffsetsFrames = true;
 	}
 
 	uint32	StreamStrides[MaxVertexElementCount];
 	uint32	StreamStrideCount = VertexFactory->GetStreamStrides(StreamStrides);
-
 	uint32	NumVertices = BatchElement.GetNumVertices();
-
-	uint32	InputStreamFloatOffset = (StreamStrides[0] * BatchElement.BaseVertexIndex) / sizeof(float);
-
+	uint32 InputStreamStart = (StreamStrides[0] * BatchElement.BaseVertexIndex) / sizeof(float);
 	uint32	NumRWFloats = (RWStrideInFloats * NumVertices);
 
-	int32 CompensatedOffset = ((int32)CacheCurrentFloatOffset * sizeof(float)) - (BatchElement.BaseVertexIndex * StreamStrides[0]);
+	// This is needed because the input index buffer are not 0 based for each chunk but 0 based for the whole mesh (better for combing draw calls).
+	// We would have to specify a negative BaseVertexIndex on the draw call and that is undefined (might work on some platforms)
+	// Doing this wastes memory and less elements might fit into the cache. By ordering the elements differently we could reduce the waste (content side or when inserting into the cache).
+	// in bytes
+	int32 CompensatedOffset = ((int32)CacheCurrentFloatOffset * sizeof(float)) - (BatchElement.BaseVertexIndex * RWStrideInFloats * sizeof(float));
 	if (CompensatedOffset < 0)
 	{
 		CacheCurrentFloatOffset -= CompensatedOffset / sizeof(float);
+		INC_DWORD_STAT_BY(STAT_GPUSkinCache_TotalMemWasted, -CompensatedOffset);
 	}
 
 	if (CacheCurrentFloatOffset + NumRWFloats > (uint32)GGPUSkinCacheBufferSize)
@@ -472,9 +491,26 @@ int32 FGPUSkinCache::StartCacheMesh(FRHICommandListImmediate& RHICmdList, int32
 
 	INC_DWORD_STAT_BY(STAT_GPUSkinCache_TotalMemUsed, NumRWFloats * sizeof(float));
 
-	ECSInfo->InputVBStride = StreamStrides[0];
+	// value section -------------
+	
+	ECSInfo->FrameUpdated = FrameNumber;
 	ECSInfo->StreamOffset = CacheCurrentFloatOffset;
 
+	if (AlignStreamOffsetsFrames || !bHasPreviousOffset)
+	{
+		ECSInfo->PreviousFrameStreamOffset = ECSInfo->StreamOffset;
+		ECSInfo->PreviousFrameUpdated = ECSInfo->FrameUpdated;
+	}
+
+	// debugging section -------------
+	
+	ECSInfo->NumVertices = NumVertices;
+	ECSInfo->BaseVertexIndex = BatchElement.BaseVertexIndex;
+	ECSInfo->InputVBStride = StreamStrides[0];
+	ECSInfo->bExtraBoneInfluences = bExtraBoneInfluences;
+	ECSInfo->VertexFactory = VertexFactory;
+	ECSInfo->TargetVertexFactory = TargetVertexFactory;
+
 	if (AlignStreamOffsetsFrames || !bHasPreviousOffset)
 	{
 		ECSInfo->PreviousFrameStreamOffset = ECSInfo->StreamOffset;
@@ -483,18 +519,27 @@ int32 FGPUSkinCache::StartCacheMesh(FRHICommandListImmediate& RHICmdList, int32
 	const FSkeletalMeshVertexBuffer* SkeletalMeshVertexBuffer = VertexFactory->GetSkinVertexBuffer();
 
 	FDispatchData DispatchData(RHICmdList, BatchElement, Skin->GetFeatureLevel(), Skin);
+
+	// could be removed for SHIPPING
+	{
+		FSkeletalMeshResource& SkeletalMeshResource = Skin->GetSkeletalMeshResource();
+		int32 LODIndex = Skin->GetLOD();
+		FStaticLODModel& LodModel = SkeletalMeshResource.LODModels[LODIndex];
+		DispatchData.ChunkIdx = LodModel.FindChunkIndex(BatchElement);
+	}
+
 	DispatchData.NumVertices = NumVertices;
 
 	if(MorphVertexBuffer)
 	{
 		DispatchData.MorphBuffer = MorphVertexBuffer->GetSRV();
-		check(DispatchData.MorphBuffer);
+		check(!GSkinCacheSafety || DispatchData.MorphBuffer);
 
 		// in bytes
 		const uint32 MorphStride = sizeof(FMorphGPUSkinVertex);
 
 		// see GPU code "check(MorphStride == sizeof(float) * 6);"
-		check(MorphStride == sizeof(float) * 6);
+		check(!GSkinCacheSafety || MorphStride == sizeof(float) * 6);
 
 		DispatchData.MorphBufferOffset = (MorphStride * BatchElement.BaseVertexIndex) / sizeof(float);
 	}
@@ -506,37 +551,42 @@ int32 FGPUSkinCache::StartCacheMesh(FRHICommandListImmediate& RHICmdList, int32
 	
 	// SkinType 0:normal, 1:with morph target, 2:with APEX cloth (not yet implemented)
 	DispatchData.SkinType = MorphVertexBuffer ? 1 : 0;
-	DispatchData.InputStreamFloatOffset = InputStreamFloatOffset;
-	DispatchData.OutputBufferFloatOffset = ECSInfo->StreamOffset;
+	DispatchData.InputStreamStart = InputStreamStart;
+	DispatchData.SkinCacheStart = ECSInfo->StreamOffset;
 
-	int32 UAVIndex = FrameCounter % GPUSKINCACHE_FRAMES;
 	DispatchData.SkinCacheBuffer = &SkinCacheBuffer[UAVIndex];
 	DispatchData.InputStreamStride = StreamStrides[0];
 	DispatchData.InputVertexBufferSRV = SkeletalMeshVertexBuffer->GetSRV();
-	check(DispatchData.InputVertexBufferSRV);
+	DispatchData.bExtraBoneInfluences = bExtraBoneInfluences;
+	check(!GSkinCacheSafety || DispatchData.InputVertexBufferSRV);
 
 	DispatchSkinCacheProcess(
-		ShaderData.GetBoneBufferForWriting(false, GFrameNumberRenderThread), ShaderData.GetUniformBuffer(),
-		ShaderData.MeshOrigin, ShaderData.MeshExtension, bExtraBoneInfluences,
+		ShaderData.GetBoneBufferForReading(false, FrameNumber), ShaderData.GetUniformBuffer(),
+		ShaderData.MeshOrigin, ShaderData.MeshExtension,
 		DispatchData);
-
-	TargetVertexFactory->UpdateVertexDeclaration(VertexFactory, &SkinCacheBuffer[FrameCounter % GPUSKINCACHE_FRAMES]);
+	
+	TargetVertexFactory->UpdateVertexDeclaration(VertexFactory, DispatchData.SkinCacheBuffer);
 
 	const bool bRecomputeTangents = MorphVertexBuffer && CVarGPUSkinCacheRecomputeTangents.GetValueOnRenderThread() != 0;
 
 	if(bRecomputeTangents)
 	{
 		FStaticLODModel* StaticLODModel = MorphVertexBuffer->GetStaticLODModel();
-		FRawStaticIndexBuffer16or32Interface* IndexBuffer = StaticLODModel->MultiSizeIndexContainer.GetIndexBuffer();
-
-		DispatchData.IndexBuffer = IndexBuffer->GetSRV();
-		check(DispatchData.IndexBuffer);
-
+	
 		// todo: this can be optimized, worse: with multiple chunks per section this is likely not to work
 		const FSkelMeshSection* Section = StaticLODModel->FindSectionForChunk(StaticLODModel->FindChunkIndex(BatchElement));
-		DispatchData.NumTriangles = Section->NumTriangles;
-		DispatchData.IndexBufferOffsetValue = Section->BaseIndex;
-		DispatchUpdateSkinTangents(DispatchData);
+		
+		check(Section);
+		if(Section->bRecomputeTangent)
+		{
+			FRawStaticIndexBuffer16or32Interface* IndexBuffer = StaticLODModel->MultiSizeIndexContainer.GetIndexBuffer();
+
+			DispatchData.IndexBuffer = IndexBuffer->GetSRV();
+			check(!GSkinCacheSafety || DispatchData.IndexBuffer);
+			DispatchData.NumTriangles = Section->NumTriangles;
+			DispatchData.IndexBufferOffsetValue = Section->BaseIndex;
+			DispatchUpdateSkinTangents(DispatchData);
+		}
 	}
 
 	CacheCurrentFloatOffset += NumRWFloats;
@@ -547,14 +597,14 @@ int32 FGPUSkinCache::StartCacheMesh(FRHICommandListImmediate& RHICmdList, int32
 
 void FGPUSkinCache::TransitionToReadable(FRHICommandList& RHICmdList)
 {
-	int32 BufferIndex = FrameCounter % GPUSKINCACHE_FRAMES;
+	int32 BufferIndex = SkinCacheFrameNumber % GPUSKINCACHE_FRAMES;
 	FUnorderedAccessViewRHIParamRef OutUAVs[] ={SkinCacheBuffer[BufferIndex].UAV};
 	RHICmdList.TransitionResources(EResourceTransitionAccess::EReadable, EResourceTransitionPipeline::EComputeToGfx, OutUAVs, ARRAY_COUNT(OutUAVs));
 }
 
 void FGPUSkinCache::TransitionToWriteable(FRHICommandList& RHICmdList)
 {
-	int32 BufferIndex = (FrameCounter + GPUSKINCACHE_FRAMES - 1) % GPUSKINCACHE_FRAMES;
+	int32 BufferIndex = (SkinCacheFrameNumber + GPUSKINCACHE_FRAMES - 1) % GPUSKINCACHE_FRAMES;
 	FUnorderedAccessViewRHIParamRef OutUAVs[] = {SkinCacheBuffer[BufferIndex].UAV};
 	RHICmdList.TransitionResources(EResourceTransitionAccess::ERWNoBarrier, EResourceTransitionPipeline::EGfxToCompute, OutUAVs, ARRAY_COUNT(OutUAVs));
 }
@@ -574,13 +624,16 @@ void FGPUSkinCache::TransitionAllToWriteable(FRHICommandList& RHICmdList)
 	}
 }
 
-bool FGPUSkinCache::InternalIsElementProcessed(int32 Key) const
+bool FGPUSkinCache::InternalIsElementProcessed(uint32 FrameNumber, int32 Key) const
 {
+	// Happens, is this a problem?
+//	ensure(FrameNumber == SkinCacheFrameNumber);
+
 	if (Key >= 0 && Key < CachedElements.Num())
 	{
 		const FElementCacheStatusInfo& CacheInfo = CachedElements.GetData()[Key];
 
-		if (CacheInfo.FrameUpdated == GFrameNumberRenderThread)
+		if (CacheInfo.FrameUpdated == FrameNumber)
 		{
 			return true;
 		}
@@ -589,47 +642,59 @@ bool FGPUSkinCache::InternalIsElementProcessed(int32 Key) const
 	return false;
 }
 
-bool FGPUSkinCache::InternalSetVertexStreamFromCache(FRHICommandList& RHICmdList, int32 Key, FShader* Shader, const FGPUSkinPassthroughVertexFactory* VertexFactory, uint32 BaseVertexIndex, FShaderParameter PreviousStreamFloatOffset, FShaderResourceParameter PreviousStreamBuffer)
+void FGPUSkinCache::InternalSetVertexStreamFromCache(FRHICommandList& RHICmdList, uint32 FrameNumber, int32 Key, FShader* Shader, const FGPUSkinPassthroughVertexFactory* VertexFactory,
+	uint32 BaseVertexIndex, FShaderParameter PreviousStreamFloatOffset, FShaderResourceParameter PreviousStreamBuffer)
 {
-	FElementCacheStatusInfo& CacheInfo = CachedElements.GetData()[Key];
+	const FElementCacheStatusInfo& CacheInfo = CachedElements.GetData()[Key];
+
+	// Verify that we found the correct chunk and that it was processed this frame
+
+	// happened, is that ok?: http://crashreporter/Buggs/Show/149552
+//	ensure(CacheInfo.BatchElement->BaseVertexIndex == BaseVertexIndex);
+	// this assumes during StartCacheMesh() we have the same BaseVertexIndex, see ensure there
+	ensureMsgf(!GSkinCacheSafety || CacheInfo.BaseVertexIndex == BaseVertexIndex, TEXT("SkinCache internal error (Debug: %p %d %d %d %d cloth:%d)"),
+		CacheInfo.BatchElement,
+		CacheInfo.NumVertices, CacheInfo.BatchElement->GetNumVertices(),
+		CacheInfo.BaseVertexIndex, BaseVertexIndex,
+		CacheInfo.BatchElement->HasApexClothData());
+
+	ensure(!GSkinCacheSafety || FrameNumber == SkinCacheFrameNumber);
+	ensure(!GSkinCacheSafety || CacheInfo.BatchElement->GetNumVertices() == CacheInfo.NumVertices);
 
-	// Verify that we found the correct chunk and that is was processed this frame
-	if (CacheInfo.BatchElement->BaseVertexIndex == BaseVertexIndex && CacheInfo.TargetVertexFactory == VertexFactory && CacheInfo.FrameUpdated == GFrameNumberRenderThread)
 	{
 		// Compensated offset since the StreamOffset is relative to the RWBuffer start, while the SetStreamSource one is relative to the BaseVertexIndex
 		int32 CompensatedOffset = ((int32)CacheInfo.StreamOffset * sizeof(float)) - (CacheInfo.BatchElement->BaseVertexIndex * RWStrideInFloats * sizeof(float));
 		if (CompensatedOffset < 0)
 		{
-			return false;
+			ensure(!GSkinCacheSafety);
+			// should never happen but better we render garbage than crash
+			return;
 		}
 
-		int32 UAVIndex = FrameCounter % GPUSKINCACHE_FRAMES;
+		int32 UAVIndex = SkinCacheFrameNumber % GPUSKINCACHE_FRAMES;
 		RHICmdList.SetStreamSource(VertexFactory->GetStreamIndex(UAVIndex), SkinCacheBuffer[UAVIndex].Buffer, RWStrideInFloats * sizeof(float), CompensatedOffset);
+	}
 
-		FVertexShaderRHIParamRef ShaderRHI = Shader->GetVertexShader();
-		if (ShaderRHI)
+	FVertexShaderRHIParamRef ShaderRHI = Shader->GetVertexShader();
+	if (ShaderRHI && PreviousStreamBuffer.IsBound())
+	{
+		// Compensated offset since in the shader Shader we use VertexId which is offset by BaseVertexIndex
+		int32 PreviousCompensatedOffset = ((int32)CacheInfo.PreviousFrameStreamOffset * sizeof(float)) - (CacheInfo.BatchElement->BaseVertexIndex * RWStrideInFloats * sizeof(float));
+		if (PreviousCompensatedOffset < 0)
 		{
-			int32 PreviousCompensatedOffset = ((int32)CacheInfo.PreviousFrameStreamOffset * sizeof(float)) - (CacheInfo.BatchElement->BaseVertexIndex * RWStrideInFloats * sizeof(float));
-			if (PreviousCompensatedOffset < 0)
-			{
-				return false;
-			}
-
-			if (PreviousStreamBuffer.IsBound())
-			{
-				int32 PrevUAVIndex = (FrameCounter + GPUSKINCACHE_FRAMES - 1) % GPUSKINCACHE_FRAMES;
-				SetShaderValue(RHICmdList, ShaderRHI, PreviousStreamFloatOffset, PreviousCompensatedOffset / sizeof(float));
-				RHICmdList.SetShaderResourceViewParameter(ShaderRHI, PreviousStreamBuffer.GetBaseIndex(), SkinCacheBuffer[PrevUAVIndex].SRV);
-			}
+			ensure(!GSkinCacheSafety);
+			// should never happen but better we render garbage than crash
+			return;
 		}
 
-		return true;
+		// Before this was assuming the previous frame is SkinCacheFrameNumber-1 but that is wrong when an object was just added (strong single frame motionblur artifact)
+		int32 PrevUAVIndex = (CacheInfo.PreviousFrameUpdated) % GPUSKINCACHE_FRAMES;
+		SetShaderValue(RHICmdList, ShaderRHI, PreviousStreamFloatOffset, PreviousCompensatedOffset / sizeof(float));
+		RHICmdList.SetShaderResourceViewParameter(ShaderRHI, PreviousStreamBuffer.GetBaseIndex(), SkinCacheBuffer[PrevUAVIndex].SRV);
 	}
-
-	return false;
 }
 
-FGPUSkinCache::FElementCacheStatusInfo* FGPUSkinCache::FindEvictableCacheStatusInfo()
+FGPUSkinCache::FElementCacheStatusInfo* FGPUSkinCache::FindEvictableCacheStatusInfo(uint32 FrameNumber)
 {
 	if (!CachedElements.Num())
 	{
@@ -639,10 +704,13 @@ FGPUSkinCache::FElementCacheStatusInfo* FGPUSkinCache::FindEvictableCacheStatusI
 	uint32	BestFrameCounter = 0x7fffffff;
 	FElementCacheStatusInfo* BestEntry = nullptr;
 
-	uint32	TargetFrameCounter = GFrameNumberRenderThread >= 10 ? GFrameNumberRenderThread - 10 : 0;
+	uint32 TargetFrameCounter = (FrameNumber >= 10) ? (FrameNumber - 10) : 0;
 
 	FElementCacheStatusInfo* CSInfo = CachedElements.GetData();
 
+	// TODO: This could be speed up a lot by using a free list. We should measure CPU cost before and after.
+
+	// O(n)
 	for (int32 i = 0; i < CachedElements.Num(); i++)
 	{
 		if (CSInfo->FrameUpdated < BestFrameCounter)
@@ -659,7 +727,7 @@ FGPUSkinCache::FElementCacheStatusInfo* FGPUSkinCache::FindEvictableCacheStatusI
 		CSInfo++;
 	}
 
-	if (BestEntry && BestFrameCounter < GFrameNumberRenderThread)
+	if (BestEntry && BestFrameCounter < FrameNumber)
 	{
 		return BestEntry;
 	}
@@ -671,14 +739,15 @@ FGPUSkinCache::FElementCacheStatusInfo* FGPUSkinCache::FindEvictableCacheStatusI
 
 void FGPUSkinCache::DispatchSkinCacheProcess(
 	const FVertexBufferAndSRV& BoneBuffer, FUniformBufferRHIRef UniformBuffer,
-	const FVector& MeshOrigin, const FVector& MeshExtension, bool bUseExtraBoneInfluences,
+	const FVector& MeshOrigin, const FVector& MeshExtension,
 	const FDispatchData& DispatchData)
 {
 	check(DispatchData.FeatureLevel >= ERHIFeatureLevel::SM5);
 
-	SCOPED_DRAW_EVENT(DispatchData.RHICmdList, SkinCacheDispatch);
-
-	int32 UAVIndex = FrameCounter % GPUSKINCACHE_FRAMES;
+	SCOPED_DRAW_EVENTF(DispatchData.RHICmdList, SkinCacheDispatch,
+		TEXT("SkinCacheDispatch Chunk=%d InputStreamStart=%d SkinCacheStart=%d Vert=%d Morph=%d/%d StrideInFloats(In/Out):%d/%d"),
+		DispatchData.ChunkIdx, DispatchData.InputStreamStart, DispatchData.SkinCacheStart, DispatchData.NumVertices, DispatchData.MorphBuffer != 0, DispatchData.MorphBufferOffset,
+		DispatchData.InputStreamStride / sizeof(float), FGPUSkinCache::RWStrideInFloats);
 
 	TShaderMapRef<TGPUSkinCacheCS<true,0> > SkinCacheCS10(GetGlobalShaderMap(DispatchData.FeatureLevel));
 	TShaderMapRef<TGPUSkinCacheCS<false,0> > SkinCacheCS00(GetGlobalShaderMap(DispatchData.FeatureLevel));
@@ -689,70 +758,58 @@ void FGPUSkinCache::DispatchSkinCacheProcess(
 	
 	switch(DispatchData.SkinType)
 	{
-		case 0: Shader = bUseExtraBoneInfluences ? (FBaseGPUSkinCacheCS*)*SkinCacheCS10 : (FBaseGPUSkinCacheCS*)*SkinCacheCS00; break;
-		case 1: Shader = bUseExtraBoneInfluences ? (FBaseGPUSkinCacheCS*)*SkinCacheCS11 : (FBaseGPUSkinCacheCS*)*SkinCacheCS01; break;
+		case 0: Shader = DispatchData.bExtraBoneInfluences ? (FBaseGPUSkinCacheCS*)*SkinCacheCS10 : (FBaseGPUSkinCacheCS*)*SkinCacheCS00; break;
+		case 1: Shader = DispatchData.bExtraBoneInfluences ? (FBaseGPUSkinCacheCS*)*SkinCacheCS11 : (FBaseGPUSkinCacheCS*)*SkinCacheCS01; break;
 		default:
 			check(0);
 	}
 	
 	check(Shader);
 	DispatchData.RHICmdList.SetComputeShader(Shader->GetComputeShader());
-	Shader->SetParameters(BoneBuffer, UniformBuffer,
-		SkinCacheBuffer[UAVIndex], MeshOrigin, MeshExtension,
-		DispatchData);
+	Shader->SetParameters(BoneBuffer, UniformBuffer,MeshOrigin, MeshExtension, DispatchData);
 
 	uint32 VertexCountAlign64 = FMath::DivideAndRoundUp(DispatchData.NumVertices, (uint32)64);
 	DispatchData.RHICmdList.DispatchComputeShader(VertexCountAlign64, 1, 1);
 	Shader->UnsetParameters(DispatchData.RHICmdList);
 }
 
-/** Encapsulates the RecomputeSkinTangents compute shader. */
-template <uint32 FullPrecisionUV>
-class FRecomputeTangentsPerTrianglePassCS : public FGlobalShader
+/** base of the FRecomputeTangentsPerTrianglePassCS class */
+class FBaseRecomputeTangents : public FGlobalShader
 {
-	DECLARE_SHADER_TYPE(FRecomputeTangentsPerTrianglePassCS, Global);
-
+public:
 	static bool ShouldCache(EShaderPlatform Platform)
 	{
 		// currently only implemented and tested on Window SM5 (needs Compute, Atomics, SRV for index buffers, UAV for VertexBuffers)
 		return Platform == SP_PCD3D_SM5;
 	}
 
-	static void ModifyCompilationEnvironment(EShaderPlatform Platform, FShaderCompilerEnvironment& OutEnvironment)
-	{
-		FGlobalShader::ModifyCompilationEnvironment(Platform, OutEnvironment);
-		OutEnvironment.SetDefine(TEXT("THREADGROUP_SIZEX"), ThreadGroupSizeX);
-		OutEnvironment.SetDefine(TEXT("INTERMEDIATE_ACCUM_BUFFER_NUM_INTS"), FGPUSkinCache::IntermediateAccumBufferNumInts);
-		OutEnvironment.SetDefine(TEXT("FULL_PRECISION_UV"), FullPrecisionUV);
-	}
-
 	static const uint32 ThreadGroupSizeX = 64;
 
-	FRecomputeTangentsPerTrianglePassCS() {}
-
-public:
 	FShaderResourceParameter IntermediateAccumBufferUAV;
-	FShaderParameter UnitCount;
+	FShaderParameter NumTriangles;
 	FShaderResourceParameter GPUSkinCacheBuffer;
-	FShaderParameter GPUSkinCacheFloatOffset;
+	FShaderParameter SkinCacheStart;
 	FShaderResourceParameter IndexBuffer;
 	FShaderParameter IndexBufferOffset;
-	FShaderParameter SkinInputStreamFloatOffset;
-	FShaderParameter SkinInputStreamStride;
+	FShaderParameter InputStreamStart;
+	FShaderParameter InputStreamStride;
 	FShaderResourceParameter SkinInputStream;
 
-	FRecomputeTangentsPerTrianglePassCS(const ShaderMetaType::CompiledShaderInitializerType& Initializer)
+	FBaseRecomputeTangents()
+	{}
+
+	FBaseRecomputeTangents(const ShaderMetaType::CompiledShaderInitializerType& Initializer)
 		: FGlobalShader(Initializer)
 	{
 		IntermediateAccumBufferUAV.Bind(Initializer.ParameterMap, TEXT("IntermediateAccumBufferUAV"));
-		UnitCount.Bind(Initializer.ParameterMap, TEXT("UnitCount"));
+		NumTriangles.Bind(Initializer.ParameterMap, TEXT("NumTriangles"));
 		GPUSkinCacheBuffer.Bind(Initializer.ParameterMap,TEXT("GPUSkinCacheBuffer"));
-		GPUSkinCacheFloatOffset.Bind(Initializer.ParameterMap,TEXT("GPUSkinCacheFloatOffset"));
+		SkinCacheStart.Bind(Initializer.ParameterMap,TEXT("SkinCacheStart"));
 		IndexBuffer.Bind(Initializer.ParameterMap,TEXT("IndexBuffer"));
 		IndexBufferOffset.Bind(Initializer.ParameterMap,TEXT("IndexBufferOffset"));
 
-		SkinInputStreamFloatOffset.Bind(Initializer.ParameterMap, TEXT("SkinCacheInputStreamFloatOffset"));
-		SkinInputStreamStride.Bind(Initializer.ParameterMap, TEXT("SkinCacheVertexStride"));
+		InputStreamStart.Bind(Initializer.ParameterMap, TEXT("InputStreamStart"));
+		InputStreamStride.Bind(Initializer.ParameterMap, TEXT("InputStreamStride"));
 		SkinInputStream.Bind(Initializer.ParameterMap, TEXT("SkinStreamInputBuffer"));
 	}
 
@@ -761,20 +818,20 @@ public:
 		const FComputeShaderRHIParamRef ShaderRHI = GetComputeShader();
 		FRHICommandListImmediate& RHICmdList = DispatchData.RHICmdList;
 
-		check(DispatchData.SkinCacheBuffer);
+		check(!GSkinCacheSafety || DispatchData.SkinCacheBuffer);
 
 //later		FGlobalShader::SetParameters(RHICmdList, ShaderRHI, View);
 
-		SetShaderValue(RHICmdList, ShaderRHI, UnitCount, DispatchData.NumTriangles);
+		SetShaderValue(RHICmdList, ShaderRHI, NumTriangles, DispatchData.NumTriangles);
 
 		SetSRVParameter(RHICmdList, ShaderRHI, GPUSkinCacheBuffer, DispatchData.SkinCacheBuffer->SRV);
-		SetShaderValue(RHICmdList, ShaderRHI, GPUSkinCacheFloatOffset, DispatchData.OutputBufferFloatOffset);
+		SetShaderValue(RHICmdList, ShaderRHI, SkinCacheStart, DispatchData.SkinCacheStart);
 
 		SetSRVParameter(RHICmdList, ShaderRHI, IndexBuffer, DispatchData.IndexBuffer);
 		SetShaderValue(RHICmdList, ShaderRHI, IndexBufferOffset, DispatchData.IndexBufferOffsetValue);
 		
-		SetShaderValue(RHICmdList, ShaderRHI, SkinInputStreamFloatOffset, DispatchData.InputStreamFloatOffset);
-		SetShaderValue(RHICmdList, ShaderRHI, SkinInputStreamStride, DispatchData.InputStreamStride);
+		SetShaderValue(RHICmdList, ShaderRHI, InputStreamStart, DispatchData.InputStreamStart);
+		SetShaderValue(RHICmdList, ShaderRHI, InputStreamStride, DispatchData.InputStreamStride);
  		SetSRVParameter(RHICmdList, ShaderRHI, SkinInputStream, DispatchData.InputVertexBufferSRV);
 
 		// UAV
@@ -791,8 +848,8 @@ public:
 	virtual bool Serialize(FArchive& Ar) override
 	{
 		bool bShaderHasOutdatedParameters = FGlobalShader::Serialize(Ar);
-		Ar << IntermediateAccumBufferUAV << UnitCount << GPUSkinCacheBuffer << GPUSkinCacheFloatOffset << IndexBuffer << IndexBufferOffset
-			<< SkinInputStreamFloatOffset << SkinInputStreamStride << SkinInputStream;
+		Ar << IntermediateAccumBufferUAV << NumTriangles << GPUSkinCacheBuffer << SkinCacheStart << IndexBuffer << IndexBufferOffset
+			<< InputStreamStart << InputStreamStride << SkinInputStream;
 		return bShaderHasOutdatedParameters;
 	}
 
@@ -807,11 +864,38 @@ public:
 	}
 };
 
+/** Encapsulates the RecomputeSkinTangents compute shader. */
+template <bool bUseExtraBoneInfluencesT, bool bFullPrecisionUV>
+class FRecomputeTangentsPerTrianglePassCS : public FBaseRecomputeTangents
+{
+	DECLARE_SHADER_TYPE(FRecomputeTangentsPerTrianglePassCS, Global);
+
+	static void ModifyCompilationEnvironment(EShaderPlatform Platform, FShaderCompilerEnvironment& OutEnvironment)
+	{
+		FGlobalShader::ModifyCompilationEnvironment(Platform, OutEnvironment);
+		const uint32 UseExtraBoneInfluences = bUseExtraBoneInfluencesT;
+		OutEnvironment.SetDefine(TEXT("GPUSKIN_USE_EXTRA_INFLUENCES"), UseExtraBoneInfluences);
+		OutEnvironment.SetDefine(TEXT("THREADGROUP_SIZEX"), ThreadGroupSizeX);
+		OutEnvironment.SetDefine(TEXT("INTERMEDIATE_ACCUM_BUFFER_NUM_INTS"), FGPUSkinCache::IntermediateAccumBufferNumInts);
+		OutEnvironment.SetDefine(TEXT("FULL_PRECISION_UV"), bFullPrecisionUV);
+	}
+
+	FRecomputeTangentsPerTrianglePassCS(const ShaderMetaType::CompiledShaderInitializerType& Initializer)
+		: FBaseRecomputeTangents(Initializer)
+	{
+	}
+
+	FRecomputeTangentsPerTrianglePassCS()
+	{}
+};
+
 // #define avoids a lot of code duplication
-#define VARIATION1(A) typedef FRecomputeTangentsPerTrianglePassCS<A> FRecomputeTangentsPerTrianglePassCS##A; \
-	IMPLEMENT_SHADER_TYPE2(FRecomputeTangentsPerTrianglePassCS##A, SF_Compute);
+#define VARIATION1(A)		VARIATION2(A,0)			VARIATION2(A,1)
+#define VARIATION2(A, B) typedef FRecomputeTangentsPerTrianglePassCS<A, B> FRecomputeTangentsPerTrianglePassCS##A##B; \
+	IMPLEMENT_SHADER_TYPE2(FRecomputeTangentsPerTrianglePassCS##A##B, SF_Compute);
 	VARIATION1(0) VARIATION1(1)
 #undef VARIATION1
+#undef VARIATION2
 
 /** Encapsulates the RecomputeSkinTangentsResolve compute shader. */
 class FRecomputeTangentsPerVertexPassCS : public FGlobalShader
@@ -827,6 +911,8 @@ class FRecomputeTangentsPerVertexPassCS : public FGlobalShader
 	static void ModifyCompilationEnvironment(EShaderPlatform Platform, FShaderCompilerEnvironment& OutEnvironment)
 	{
 		FGlobalShader::ModifyCompilationEnvironment(Platform, OutEnvironment);
+		// this pass cannot read the input as it doesn't have the permutation
+		OutEnvironment.SetDefine(TEXT("GPUSKIN_USE_EXTRA_INFLUENCES"), (uint32)0);
 		OutEnvironment.SetDefine(TEXT("THREADGROUP_SIZEX"), ThreadGroupSizeX);
 		OutEnvironment.SetDefine(TEXT("GPUSKIN_RWBUFFER_OFFSET_TANGENT_X"), FGPUSkinCache::RWTangentXOffsetInFloats);
 		OutEnvironment.SetDefine(TEXT("GPUSKIN_RWBUFFER_OFFSET_TANGENT_Z"), FGPUSkinCache::RWTangentZOffsetInFloats);
@@ -840,19 +926,21 @@ class FRecomputeTangentsPerVertexPassCS : public FGlobalShader
 
 public:
 	FShaderResourceParameter IntermediateAccumBufferUAV;
-	FShaderResourceParameter GPUSkinCacheBufferUAV;
-	FShaderParameter GPUSkinCacheFloatOffset;
-	FShaderParameter UnitCount;
-	FShaderParameter SkinInputStreamFloatOffset;
+	FShaderResourceParameter SkinCacheBufferUAV;
+	FShaderParameter SkinCacheStart;
+	FShaderParameter NumVertices;
+	FShaderParameter InputStreamStart;
+	FShaderParameter InputStreamStride;
 
 	FRecomputeTangentsPerVertexPassCS(const ShaderMetaType::CompiledShaderInitializerType& Initializer)
 		: FGlobalShader(Initializer)
 	{
 		IntermediateAccumBufferUAV.Bind(Initializer.ParameterMap, TEXT("IntermediateAccumBufferUAV"));
-		GPUSkinCacheBufferUAV.Bind(Initializer.ParameterMap,TEXT("GPUSkinCacheBufferUAV"));
-		GPUSkinCacheFloatOffset.Bind(Initializer.ParameterMap,TEXT("GPUSkinCacheFloatOffset"));
-		UnitCount.Bind(Initializer.ParameterMap, TEXT("UnitCount"));
-		SkinInputStreamFloatOffset.Bind(Initializer.ParameterMap, TEXT("SkinCacheInputStreamFloatOffset"));
+		SkinCacheBufferUAV.Bind(Initializer.ParameterMap,TEXT("SkinCacheBufferUAV"));
+		SkinCacheStart.Bind(Initializer.ParameterMap,TEXT("SkinCacheStart"));
+		NumVertices.Bind(Initializer.ParameterMap, TEXT("NumVertices"));
+		InputStreamStart.Bind(Initializer.ParameterMap, TEXT("InputStreamStart"));
+		InputStreamStride.Bind(Initializer.ParameterMap, TEXT("InputStreamStride"));
 	}
 
 	void SetParameters(const FGPUSkinCache::FDispatchData& DispatchData, FRWBuffer& IntermediateAccumBuffer)
@@ -860,31 +948,32 @@ public:
 		const FComputeShaderRHIParamRef ShaderRHI = GetComputeShader();
 		FRHICommandListImmediate& RHICmdList = DispatchData.RHICmdList;
 		
-		check(DispatchData.SkinCacheBuffer);
+		check(!GSkinCacheSafety || DispatchData.SkinCacheBuffer);
 
 //later		FGlobalShader::SetParameters(RHICmdList, ShaderRHI, View);
 
-		SetShaderValue(RHICmdList, ShaderRHI, GPUSkinCacheFloatOffset, DispatchData.OutputBufferFloatOffset);
-		SetShaderValue(RHICmdList, ShaderRHI, UnitCount, DispatchData.NumVertices);
-		SetShaderValue(RHICmdList, ShaderRHI, SkinInputStreamFloatOffset, DispatchData.InputStreamFloatOffset);
+		SetShaderValue(RHICmdList, ShaderRHI, SkinCacheStart, DispatchData.SkinCacheStart);
+		SetShaderValue(RHICmdList, ShaderRHI, NumVertices, DispatchData.NumVertices);
+		SetShaderValue(RHICmdList, ShaderRHI, InputStreamStart, DispatchData.InputStreamStart);
+		SetShaderValue(RHICmdList, ShaderRHI, InputStreamStride, DispatchData.InputStreamStride);
 
 		// UAVs
 		SetUAVParameter(RHICmdList, ShaderRHI, IntermediateAccumBufferUAV, IntermediateAccumBuffer.UAV);
-		SetUAVParameter(RHICmdList, ShaderRHI, GPUSkinCacheBufferUAV, DispatchData.SkinCacheBuffer->UAV);
+		SetUAVParameter(RHICmdList, ShaderRHI, SkinCacheBufferUAV, DispatchData.SkinCacheBuffer->UAV);
 	}
 
 	void UnsetParameters(FRHICommandList& RHICmdList)
 	{
 		FComputeShaderRHIParamRef ShaderRHI = GetComputeShader();
 
-		SetUAVParameter(RHICmdList, ShaderRHI, GPUSkinCacheBufferUAV, 0);
+		SetUAVParameter(RHICmdList, ShaderRHI, SkinCacheBufferUAV, 0);
 		SetUAVParameter(RHICmdList, ShaderRHI, IntermediateAccumBufferUAV, 0);
 	}
 
 	virtual bool Serialize(FArchive& Ar) override
 	{
 		bool bShaderHasOutdatedParameters = FGlobalShader::Serialize(Ar);
-		Ar << IntermediateAccumBufferUAV << GPUSkinCacheBufferUAV << GPUSkinCacheFloatOffset << UnitCount << SkinInputStreamFloatOffset;
+		Ar << IntermediateAccumBufferUAV << SkinCacheBufferUAV << SkinCacheStart << NumVertices << InputStreamStart << InputStreamStride;
 		return bShaderHasOutdatedParameters;
 	}
 };
@@ -916,47 +1005,49 @@ void FGPUSkinCache::DispatchUpdateSkinTangents(const FDispatchData& DispatchData
 		uint32 MaxVertexCount = ComputeRecomputeTangentMaxVertexCount();
 
 		// We could do automatically reallocate but that is not yet supported
-		ensure(RequiredVertexCount < MaxVertexCount);
+		ensure(!GSkinCacheSafety || RequiredVertexCount < MaxVertexCount);
 
 		// todo
 //		RHICmdList.TransitionResource(EResourceTransitionAccess::ERWNoBarrier, EResourceTransitionPipeline::EGfxToCompute, IntermediateAccumBuffer.UAV);
 
 		// This code can be optimized by batched up and doing it with less Dispatch calls (costs more memory)
 		{
-			TShaderMapRef<FRecomputeTangentsPerTrianglePassCS<0> > ComputeShader0(GetGlobalShaderMap(DispatchData.FeatureLevel));
-			TShaderMapRef<FRecomputeTangentsPerTrianglePassCS<1> > ComputeShader1(GetGlobalShaderMap(DispatchData.FeatureLevel));
+			TShaderMapRef<FRecomputeTangentsPerTrianglePassCS<0, 0> > ComputeShader00(GetGlobalShaderMap(DispatchData.FeatureLevel));
+			TShaderMapRef<FRecomputeTangentsPerTrianglePassCS<0, 1> > ComputeShader01(GetGlobalShaderMap(DispatchData.FeatureLevel));
+			TShaderMapRef<FRecomputeTangentsPerTrianglePassCS<1, 0> > ComputeShader10(GetGlobalShaderMap(DispatchData.FeatureLevel));
+			TShaderMapRef<FRecomputeTangentsPerTrianglePassCS<1, 1> > ComputeShader11(GetGlobalShaderMap(DispatchData.FeatureLevel));
 			
 			bool bFullPrecisionUV = LodModel.VertexBufferGPUSkin.GetUseFullPrecisionUVs();
-			uint32 NumTriangles = DispatchData.NumTriangles;
-			uint32 ThreadGroupCountValue = FMath::DivideAndRoundUp(NumTriangles, FRecomputeTangentsPerTrianglePassCS<0>::ThreadGroupSizeX);
 
-			SCOPED_DRAW_EVENTF(RHICmdList, SkinTangents_PerTrianglePass, TEXT("SkinTangents_PerTrianglePass IndexStart=%d Tri=%d UVPrecision=%d"),
-				DispatchData.IndexBufferOffsetValue, DispatchData.NumTriangles, bFullPrecisionUV);
+			FBaseRecomputeTangents* Shader = 0;
 
-			if(bFullPrecisionUV)
+			switch((uint32)bFullPrecisionUV)
 			{
-				const FComputeShaderRHIParamRef ShaderRHI = ComputeShader1->GetComputeShader();
-				RHICmdList.SetComputeShader(ShaderRHI);
-
-				ComputeShader1->SetParameters(DispatchData, SkinTangentIntermediate);
-				DispatchComputeShader(RHICmdList, *ComputeShader1, ThreadGroupCountValue, 1, 1);
-				ComputeShader1->UnsetParameters(RHICmdList);
+				case 0: Shader = DispatchData.bExtraBoneInfluences ? (FBaseRecomputeTangents*)*ComputeShader10 : (FBaseRecomputeTangents*)*ComputeShader00; break;
+				case 1: Shader = DispatchData.bExtraBoneInfluences ? (FBaseRecomputeTangents*)*ComputeShader11 : (FBaseRecomputeTangents*)*ComputeShader01; break;
+				default:
+					check(0);
 			}
-			else
-			{
-				FGlobalShader* ComputeShader = bFullPrecisionUV ? (FGlobalShader*)*ComputeShader1 : (FGlobalShader*)*ComputeShader0;
-				const FComputeShaderRHIParamRef ShaderRHI = ComputeShader0->GetComputeShader();
+	
+			check(Shader);
+
+			uint32 NumTriangles = DispatchData.NumTriangles;
+			uint32 ThreadGroupCountValue = FMath::DivideAndRoundUp(NumTriangles, FBaseRecomputeTangents::ThreadGroupSizeX);
+
+			SCOPED_DRAW_EVENTF(RHICmdList, SkinTangents_PerTrianglePass, TEXT("SkinTangents_PerTrianglePass IndexStart=%d Tri=%d ExtraBoneInfluences=%d UVPrecision=%d"),
+				DispatchData.IndexBufferOffsetValue, DispatchData.NumTriangles, DispatchData.bExtraBoneInfluences, bFullPrecisionUV);
+
+			const FComputeShaderRHIParamRef ShaderRHI = Shader->GetComputeShader();
 				RHICmdList.SetComputeShader(ShaderRHI);
 
-				ComputeShader0->SetParameters(DispatchData, SkinTangentIntermediate);
-				DispatchComputeShader(RHICmdList, *ComputeShader0, ThreadGroupCountValue, 1, 1);
-				ComputeShader0->UnsetParameters(RHICmdList);
-			}
+			Shader->SetParameters(DispatchData, SkinTangentIntermediate);
+			DispatchComputeShader(RHICmdList, Shader, ThreadGroupCountValue, 1, 1);
+			Shader->UnsetParameters(RHICmdList);
 		}
 
 		{
-			SCOPED_DRAW_EVENTF(RHICmdList, SkinTangents_PerVertexPass, TEXT("SkinTangents_PerVertexPass InputVertStart=%d, SkinCacheStart=%d, Vert=%d"),
-				DispatchData.InputStreamFloatOffset, DispatchData.OutputBufferFloatOffset, DispatchData.NumVertices);
+			SCOPED_DRAW_EVENTF(RHICmdList, SkinTangents_PerVertexPass, TEXT("SkinTangents_PerVertexPass InputStreamStart=%d, SkinCacheStart=%d, Vert=%d"),
+				DispatchData.InputStreamStart, DispatchData.SkinCacheStart, DispatchData.NumVertices);
 
 			TShaderMapRef<FRecomputeTangentsPerVertexPassCS> ComputeShader(GetGlobalShaderMap(DispatchData.FeatureLevel));
 			const FComputeShaderRHIParamRef ShaderRHI = ComputeShader->GetComputeShader();
@@ -966,9 +1057,7 @@ void FGPUSkinCache::DispatchUpdateSkinTangents(const FDispatchData& DispatchData
 			uint32 ThreadGroupCountValue = FMath::DivideAndRoundUp(VertexCount, FRecomputeTangentsPerVertexPassCS::ThreadGroupSizeX);
 
 			ComputeShader->SetParameters(DispatchData, SkinTangentIntermediate);
-	
 			DispatchComputeShader(RHICmdList, *ComputeShader, ThreadGroupCountValue, 1, 1);
-
 			ComputeShader->UnsetParameters(RHICmdList);
 		}
 // todo				RHICmdList.TransitionResource(EResourceTransitionAccess::EReadable, EResourceTransitionPipeline::EComputeToGfx, TangentsBlendBuffer.VertexBufferSRV);
diff --git a/Engine/Source/Runtime/Engine/Private/GPUSkinVertexFactory.cpp b/Engine/Source/Runtime/Engine/Private/GPUSkinVertexFactory.cpp
--- a/Engine/Source/Runtime/Engine/Private/GPUSkinVertexFactory.cpp
+++ b/Engine/Source/Runtime/Engine/Private/GPUSkinVertexFactory.cpp
@@ -597,15 +597,9 @@ public:
 	*/
 	virtual void SetMesh(FRHICommandList& RHICmdList, FShader* Shader,const FVertexFactory* VertexFactory,const FSceneView& View,const FMeshBatchElement& BatchElement,uint32 DataFlags) const override
 	{
-		if (GEnableGPUSkinCache)
-		{
-			bool bIsGPUCached = false;
-			check(VertexFactory->GetType() == &FGPUSkinPassthroughVertexFactory::StaticType);
-			if (GGPUSkinCache.SetVertexStreamFromCache(RHICmdList, BatchElement.UserIndex, Shader, (FGPUSkinPassthroughVertexFactory*)VertexFactory, BatchElement.MinVertexIndex, GPUSkinCachePreviousFloatOffset, GPUSkinCachePreviousBuffer))
-			{
-				bIsGPUCached = true;
-			}
-		}
+		check(VertexFactory->GetType() == &FGPUSkinPassthroughVertexFactory::StaticType);
+
+		GGPUSkinCache.SetVertexStreamFromCache(RHICmdList, View.Family->FrameNumber, BatchElement.UserIndex, Shader, (FGPUSkinPassthroughVertexFactory*)VertexFactory, BatchElement.MinVertexIndex, GPUSkinCachePreviousFloatOffset, GPUSkinCachePreviousBuffer);
 	}
 
 	virtual uint32 GetSize() const override { return sizeof(*this); }
@@ -858,7 +852,7 @@ struct FRHICommandUpdateClothBuffer : public FRHICommand<FRHICommandUpdateClothB
 };
 
 bool FGPUBaseSkinAPEXClothVertexFactory::ClothShaderType::UpdateClothSimulData(FRHICommandListImmediate& RHICmdList, const TArray<FVector4>& InSimulPositions,
-	const TArray<FVector4>& InSimulNormals, uint32 FrameNumber, ERHIFeatureLevel::Type FeatureLevel)
+	const TArray<FVector4>& InSimulNormals, uint32 FrameNumberToPrepare, ERHIFeatureLevel::Type FeatureLevel)
 {
 	QUICK_SCOPE_CYCLE_COUNTER(STAT_FGPUBaseSkinAPEXClothVertexFactory_UpdateClothSimulData);
 
@@ -869,9 +863,8 @@ bool FGPUBaseSkinAPEXClothVertexFactory::ClothShaderType::UpdateClothSimulData(F
 	if (FeatureLevel >= ERHIFeatureLevel::SM4)
 	{
 		check(IsInRenderingThread());
-
-		// + 1 as the FrameNumber is incremented later that in should be
-		CurrentClothBuffer = &GetClothBufferForWriting(FrameNumber + 1);
+		
+		CurrentClothBuffer = &GetClothBufferForWriting(FrameNumberToPrepare);
 
 		NumSimulVerts = FMath::Min(NumSimulVerts, (uint32)MAX_APEXCLOTH_VERTICES_FOR_VB);
 
diff --git a/Engine/Source/Runtime/Engine/Private/GameEngine.cpp b/Engine/Source/Runtime/Engine/Private/GameEngine.cpp
--- a/Engine/Source/Runtime/Engine/Private/GameEngine.cpp
+++ b/Engine/Source/Runtime/Engine/Private/GameEngine.cpp
@@ -188,7 +188,7 @@ void UGameEngine::ConditionallyOverrideSettings(int32& ResolutionX, int32& Resol
 	else if (FParse::Param(FCommandLine::Get(),TEXT("FullScreen")))
 	{
 		// -FullScreen
-		auto CVar = IConsoleManager::Get().FindTConsoleVariableDataInt(TEXT("r.FullScreenMode"));
+		static auto CVar = IConsoleManager::Get().FindTConsoleVariableDataInt(TEXT("r.FullScreenMode"));
 		check(CVar);
 		WindowMode = CVar->GetValueOnGameThread() == 0 ? EWindowMode::Fullscreen : EWindowMode::WindowedFullscreen;
 
diff --git a/Engine/Source/Runtime/Engine/Private/GlobalShader.cpp b/Engine/Source/Runtime/Engine/Private/GlobalShader.cpp
--- a/Engine/Source/Runtime/Engine/Private/GlobalShader.cpp
+++ b/Engine/Source/Runtime/Engine/Private/GlobalShader.cpp
@@ -393,7 +393,7 @@ void VerifyGlobalShaders(EShaderPlatform Platform, bool bLoadedFromCacheFile)
 
 	if (GlobalShaderJobs.Num() > 0)
 	{
-		GShaderCompilingManager->AddJobs(GlobalShaderJobs, true, true);
+		GShaderCompilingManager->AddJobs(GlobalShaderJobs, true, true, false);
 
 		const bool bAllowAsynchronousGlobalShaderCompiling =
 			// OpenGL requires that global shader maps are compiled before attaching
diff --git a/Engine/Source/Runtime/Engine/Private/Level.cpp b/Engine/Source/Runtime/Engine/Private/Level.cpp
--- a/Engine/Source/Runtime/Engine/Private/Level.cpp
+++ b/Engine/Source/Runtime/Engine/Private/Level.cpp
@@ -1549,6 +1549,8 @@ void ULevel::BuildStreamingData(UTexture2D* UpdateSpecificTextureOnly/*=NULL*/)
 	TArray<UObject *> ObjectsInOuter;
 	GetObjectsWithOuter(this, ObjectsInOuter);
 
+	FStreamingTextureLevelContext LevelContext(this);
+
 	for( int32 Index = 0; Index < ObjectsInOuter.Num(); Index++ )
 	{
 		UPrimitiveComponent* Primitive = Cast<UPrimitiveComponent>(ObjectsInOuter[Index]);
@@ -1563,7 +1565,7 @@ void ULevel::BuildStreamingData(UTexture2D* UpdateSpecificTextureOnly/*=NULL*/)
 				TArray<FStreamingTexturePrimitiveInfo> PrimitiveStreamingTextures;
 
 				// Ask the primitive to enumerate the streaming textures it uses.
-				Primitive->GetStreamingTextureInfoWithNULLRemoval(PrimitiveStreamingTextures);
+				Primitive->GetStreamingTextureInfoWithNULLRemoval(LevelContext, PrimitiveStreamingTextures);
 
 				for(int32 TextureIndex = 0;TextureIndex < PrimitiveStreamingTextures.Num();TextureIndex++)
 				{
diff --git a/Engine/Source/Runtime/Engine/Private/LevelTick.cpp b/Engine/Source/Runtime/Engine/Private/LevelTick.cpp
--- a/Engine/Source/Runtime/Engine/Private/LevelTick.cpp
+++ b/Engine/Source/Runtime/Engine/Private/LevelTick.cpp
@@ -29,6 +29,7 @@
 #include "Engine/CoreSettings.h"
 
 #include "InGamePerformanceTracker.h"
+#include "Streaming/TextureStreamingHelpers.h"
 
 // this will log out all of the objects that were ticked in the FDetailedTickStats struct so you can isolate what is expensive
 #define LOG_DETAILED_DUMPSTATS 0
diff --git a/Engine/Source/Runtime/Engine/Private/LevelUtils.cpp b/Engine/Source/Runtime/Engine/Private/LevelUtils.cpp
--- a/Engine/Source/Runtime/Engine/Private/LevelUtils.cpp
+++ b/Engine/Source/Runtime/Engine/Private/LevelUtils.cpp
@@ -388,6 +388,12 @@ void FLevelUtils::ApplyLevelTransform( ULevel* Level, const FTransform& LevelTra
 	bool bTransformActors =  !LevelTransform.Equals(FTransform::Identity);
 	if (bTransformActors)
 	{
+		if (!LevelTransform.GetRotation().IsIdentity())
+		{
+			// If there is a rotation applied, then the relative precomputed bounds become invalid.
+			Level->bTextureStreamingRotationChanged = true;
+		}
+
 		for (TMap< UTexture2D*, TArray<FStreamableTextureInstance> >::TIterator It(Level->TextureToInstancesMap); It; ++It)
 		{
 			TArray<FStreamableTextureInstance>& TextureInfo = It.Value();
diff --git a/Engine/Source/Runtime/Engine/Private/Materials/HLSLMaterialTranslator.h b/Engine/Source/Runtime/Engine/Private/Materials/HLSLMaterialTranslator.h
--- a/Engine/Source/Runtime/Engine/Private/Materials/HLSLMaterialTranslator.h
+++ b/Engine/Source/Runtime/Engine/Private/Materials/HLSLMaterialTranslator.h
@@ -3079,6 +3079,11 @@ protected:
 		return AddUniformExpression(new FMaterialUniformExpressionTextureParameter(ParameterName, TextureReferenceIndex, SamplerSource),ShaderType,TEXT(""));
 	}
 
+	virtual int32 GetTextureReferenceIndex(UTexture* TextureValue)
+	{
+		return Material->GetReferencedTextures().Find(TextureValue);
+	}
+
 	virtual int32 StaticBool(bool bValue) override
 	{
 		return AddInlinedCodeChunk(MCT_StaticBool,(bValue ? TEXT("true") : TEXT("false")));
diff --git a/Engine/Source/Runtime/Engine/Private/Materials/Material.cpp b/Engine/Source/Runtime/Engine/Private/Materials/Material.cpp
--- a/Engine/Source/Runtime/Engine/Private/Materials/Material.cpp
+++ b/Engine/Source/Runtime/Engine/Private/Materials/Material.cpp
@@ -704,6 +704,50 @@ void UMaterial::GetUsedTextures(TArray<UTexture*>& OutTextures, EMaterialQuality
 	}
 }
 
+void UMaterial::GetUsedTexturesAndIndices(TArray<UTexture*>& OutTextures, TArray< TArray<int32> >& OutIndices, EMaterialQualityLevel::Type QualityLevel, ERHIFeatureLevel::Type FeatureLevel) const
+{
+	OutTextures.Empty();
+	OutIndices.Empty();
+
+	if (!FPlatformProperties::IsServerOnly())
+	{
+		const FMaterialResource* CurrentResource = MaterialResources[QualityLevel][FeatureLevel];
+
+		if (CurrentResource)
+		{
+			const TArray<TRefCountPtr<FMaterialUniformExpressionTexture> >* ExpressionsByType[2] =
+			{
+				&CurrentResource->GetUniform2DTextureExpressions(),
+				&CurrentResource->GetUniformCubeTextureExpressions()
+			};
+
+			// Try to prevent resizing since this would be expensive.
+			OutIndices.Empty(ExpressionsByType[0]->Num() + ExpressionsByType[1]->Num());
+
+			for (int32 TypeIndex = 0; TypeIndex < ARRAY_COUNT(ExpressionsByType); TypeIndex++)
+			{
+				// Iterate over each of the material's texture expressions.
+				for (FMaterialUniformExpressionTexture* Expression : *ExpressionsByType[TypeIndex])
+				{
+					const bool bAllowOverride = false;
+					UTexture* Texture = NULL;
+					Expression->GetGameThreadTextureValue(this, *CurrentResource, Texture, bAllowOverride);
+
+					if (Texture)
+					{
+						int32 InsertIndex = OutTextures.AddUnique(Texture);
+						if (InsertIndex >= OutIndices.Num())
+						{
+							OutIndices.AddDefaulted(InsertIndex - OutIndices.Num() + 1);
+						}
+						OutIndices[InsertIndex].Add(Expression->GetTextureIndex());
+					}
+				}
+			}
+		}
+	}
+}
+
 #if !(UE_BUILD_SHIPPING || UE_BUILD_TEST)
 void UMaterial::LogMaterialsAndTextures(FOutputDevice& Ar, int32 Indent) const
 {
@@ -2775,6 +2819,10 @@ void UMaterial::PreEditChange(UProperty* PropertyThatChanged)
 
 void UMaterial::PostEditChangeProperty(FPropertyChangedEvent& PropertyChangedEvent)
 {
+	// If the material changes, then the debug view material must reset to prevent parameters mismatch
+	void ClearAllDebugViewMaterials();
+	ClearAllDebugViewMaterials();
+
 	Super::PostEditChangeProperty(PropertyChangedEvent);
 
 	UProperty* PropertyThatChanged = PropertyChangedEvent.Property;
diff --git a/Engine/Source/Runtime/Engine/Private/Materials/MaterialExpressions.cpp b/Engine/Source/Runtime/Engine/Private/Materials/MaterialExpressions.cpp
--- a/Engine/Source/Runtime/Engine/Private/Materials/MaterialExpressions.cpp
+++ b/Engine/Source/Runtime/Engine/Private/Materials/MaterialExpressions.cpp
@@ -1370,6 +1370,8 @@ int32 UMaterialExpressionTextureSample::Compile(class FMaterialCompiler* Compile
 				EffectiveTexture = TextureObjectParameter->Texture;
 				EffectiveSamplerType = TextureObjectParameter->SamplerType;
 			}
+
+			TextureReferenceIndex = Compiler->GetTextureReferenceIndex(EffectiveTexture);
 		}
 
 		if (EffectiveTexture && VerifySamplerType(Compiler, (Desc.Len() > 0 ? *Desc : TEXT("TextureSample")), EffectiveTexture, EffectiveSamplerType))
diff --git a/Engine/Source/Runtime/Engine/Private/Materials/MaterialInstance.cpp b/Engine/Source/Runtime/Engine/Private/Materials/MaterialInstance.cpp
--- a/Engine/Source/Runtime/Engine/Private/Materials/MaterialInstance.cpp
+++ b/Engine/Source/Runtime/Engine/Private/Materials/MaterialInstance.cpp
@@ -650,7 +650,7 @@ bool UMaterialInstance::GetRefractionSettings(float& OutBiasValue) const
 	}
 }
 
-void UMaterialInstance::GetTextureExpressionValues(const FMaterialResource* MaterialResource, TArray<UTexture*>& OutTextures) const
+void UMaterialInstance::GetTextureExpressionValues(const FMaterialResource* MaterialResource, TArray<UTexture*>& OutTextures, TArray< TArray<int32> >* OutIndices) const
 {
 	const TArray<TRefCountPtr<FMaterialUniformExpressionTexture> >* ExpressionsByType[2];
 
@@ -659,6 +659,11 @@ void UMaterialInstance::GetTextureExpressionValues(const FMaterialResource* Mate
 	ExpressionsByType[0] = &MaterialResource->GetUniform2DTextureExpressions();
 	ExpressionsByType[1] = &MaterialResource->GetUniformCubeTextureExpressions();
 
+	if (OutIndices) // Try to prevent resizing since this would be expensive.
+	{
+		OutIndices->Empty(ExpressionsByType[0]->Num() + ExpressionsByType[1]->Num());
+	}
+
 	for(int32 TypeIndex = 0;TypeIndex < ARRAY_COUNT(ExpressionsByType);TypeIndex++)
 	{
 		const TArray<TRefCountPtr<FMaterialUniformExpressionTexture> >& Expressions = *ExpressionsByType[TypeIndex];
@@ -671,7 +676,20 @@ void UMaterialInstance::GetTextureExpressionValues(const FMaterialResource* Mate
 			// Evaluate the expression in terms of this material instance.
 			UTexture* Texture = NULL;
 			Expression->GetGameThreadTextureValue(this,*MaterialResource,Texture, true);
-			OutTextures.AddUnique(Texture);
+			
+			if (Texture)
+			{
+				int32 InsertIndex = OutTextures.AddUnique(Texture);
+
+				if (OutIndices)
+				{
+					if (InsertIndex >= OutIndices->Num())
+					{
+						OutIndices->AddDefaulted(InsertIndex - OutIndices->Num() + 1);
+					}
+					(*OutIndices)[InsertIndex].Add(Expression->GetTextureIndex());
+				}
+			}
 		}
 	}
 }
@@ -825,6 +843,47 @@ void UMaterialInstance::LogMaterialsAndTextures(FOutputDevice& Ar, int32 Indent)
 }
 #endif
 
+void UMaterialInstance::GetUsedTexturesAndIndices(TArray<UTexture*>& OutTextures, TArray< TArray<int32> >& OutIndices, EMaterialQualityLevel::Type QualityLevel, ERHIFeatureLevel::Type FeatureLevel) const
+{
+	OutTextures.Empty();
+	OutIndices.Empty();
+
+	if (!FPlatformProperties::IsServerOnly())
+	{
+		const UMaterialInstance* MaterialInstanceToUse = this;
+		// Walk up the material instance chain to the first parent that has static parameters
+		while (MaterialInstanceToUse && !MaterialInstanceToUse->bHasStaticPermutationResource)
+		{
+			MaterialInstanceToUse = Cast<const UMaterialInstance>(MaterialInstanceToUse->Parent);
+		}
+
+		if (MaterialInstanceToUse && MaterialInstanceToUse->bHasStaticPermutationResource)
+		{
+			const FMaterialResource* CurrentResource = MaterialInstanceToUse->StaticPermutationMaterialResources[QualityLevel][FeatureLevel];
+			if (CurrentResource)
+			{
+				GetTextureExpressionValues(CurrentResource, OutTextures, &OutIndices);
+			}
+		}
+		else // Use the uniform expressions from the base material
+		{ 
+			const UMaterial* Material = GetMaterial();
+			if (Material)
+			{
+				const FMaterialResource* MaterialResource = Material->GetMaterialResource(FeatureLevel, QualityLevel);
+				if( MaterialResource )
+				{
+					GetTextureExpressionValues(MaterialResource, OutTextures, &OutIndices);
+				}
+			}
+			else // If the material instance has no material, use the default material.
+			{
+				UMaterial::GetDefaultMaterial(MD_Surface)->GetUsedTexturesAndIndices(OutTextures, OutIndices, QualityLevel, FeatureLevel);
+			}
+		}
+	}
+}
+
 void UMaterialInstance::OverrideTexture(const UTexture* InTextureToOverride, UTexture* OverrideTexture, ERHIFeatureLevel::Type InFeatureLevel)
 {
 #if WITH_EDITOR
diff --git a/Engine/Source/Runtime/Engine/Private/Materials/MaterialInterface.cpp b/Engine/Source/Runtime/Engine/Private/Materials/MaterialInterface.cpp
--- a/Engine/Source/Runtime/Engine/Private/Materials/MaterialInterface.cpp
+++ b/Engine/Source/Runtime/Engine/Private/Materials/MaterialInterface.cpp
@@ -64,6 +64,12 @@ void UMaterialInterface::PostLoad()
 	PostLoadDefaultMaterials();
 }
 
+void UMaterialInterface::GetUsedTexturesAndIndices(TArray<UTexture*>& OutTextures, TArray< TArray<int32> >& OutIndices, EMaterialQualityLevel::Type QualityLevel, ERHIFeatureLevel::Type FeatureLevel) const
+{
+	GetUsedTextures(OutTextures, QualityLevel, false, FeatureLevel, false);
+	OutIndices.AddDefaulted(OutTextures.Num());
+}
+
 FMaterialRelevance UMaterialInterface::GetRelevance_Internal(const UMaterial* Material, ERHIFeatureLevel::Type InFeatureLevel) const
 {
 	if(Material)
diff --git a/Engine/Source/Runtime/Engine/Private/Materials/MaterialShader.cpp b/Engine/Source/Runtime/Engine/Private/Materials/MaterialShader.cpp
--- a/Engine/Source/Runtime/Engine/Private/Materials/MaterialShader.cpp
+++ b/Engine/Source/Runtime/Engine/Private/Materials/MaterialShader.cpp
@@ -1355,8 +1355,11 @@ void FMaterialShaderMap::Compile(
 			// Mark as not having been compiled
 			bCompiledSuccessfully = false;
   
+			// Only cause a global component recreate state for non-preview materials
+			const bool bRecreateComponentRenderStateOnCompletion = Material->IsPersistent();
+
 			// Note: using Material->IsPersistent() to detect whether this is a preview material which should have higher priority over background compiling
-			GShaderCompilingManager->AddJobs(NewJobs, bApplyCompletedShaderMapForRendering && !bSynchronousCompile, bSynchronousCompile || !Material->IsPersistent());
+			GShaderCompilingManager->AddJobs(NewJobs, bApplyCompletedShaderMapForRendering && !bSynchronousCompile, bSynchronousCompile || !Material->IsPersistent(), bRecreateComponentRenderStateOnCompletion);
   
 			// Compile the shaders for this shader map now if the material is not deferring and deferred compiles are not enabled globally
 			if (bSynchronousCompile)
diff --git a/Engine/Source/Runtime/Engine/Private/Materials/MaterialShared.cpp b/Engine/Source/Runtime/Engine/Private/Materials/MaterialShared.cpp
--- a/Engine/Source/Runtime/Engine/Private/Materials/MaterialShared.cpp
+++ b/Engine/Source/Runtime/Engine/Private/Materials/MaterialShared.cpp
@@ -370,18 +370,25 @@ void FMaterialCompilationOutput::Serialize(FArchive& Ar)
 
 void FMaterial::GetShaderMapId(EShaderPlatform Platform, FMaterialShaderMapId& OutId) const
 { 
-	TArray<FShaderType*> ShaderTypes;
-	TArray<FVertexFactoryType*> VFTypes;
-	TArray<const FShaderPipelineType*> ShaderPipelineTypes;
+	if (bLoadedCookedShaderMapId)
+	{
+		OutId = CookedShaderMapId;
+	}
+	else
+	{
+		TArray<FShaderType*> ShaderTypes;
+		TArray<FVertexFactoryType*> VFTypes;
+		TArray<const FShaderPipelineType*> ShaderPipelineTypes;
 
-	GetDependentShaderAndVFTypes(Platform, ShaderTypes, ShaderPipelineTypes, VFTypes);
+		GetDependentShaderAndVFTypes(Platform, ShaderTypes, ShaderPipelineTypes, VFTypes);
 
-	OutId.Usage = GetShaderMapUsage();
-	OutId.BaseMaterialId = GetMaterialId();
-	OutId.QualityLevel = GetQualityLevelForShaderMapId();
-	OutId.FeatureLevel = GetFeatureLevel();
-	OutId.SetShaderDependencies(ShaderTypes, ShaderPipelineTypes, VFTypes);
-	GetReferencedTexturesHash(Platform, OutId.TextureReferencesHash);
+		OutId.Usage = GetShaderMapUsage();
+		OutId.BaseMaterialId = GetMaterialId();
+		OutId.QualityLevel = GetQualityLevelForShaderMapId();
+		OutId.FeatureLevel = GetFeatureLevel();
+		OutId.SetShaderDependencies(ShaderTypes, ShaderPipelineTypes, VFTypes);
+		GetReferencedTexturesHash(Platform, OutId.TextureReferencesHash);
+	}
 }
 
 EMaterialTessellationMode FMaterial::GetTessellationMode() const 
@@ -896,11 +903,6 @@ bool FMaterialResource::IsUsingHQForwardReflections() const
 	return Material->bUseHQForwardReflections;
 }
 
-bool FMaterialResource::IsUsingPlanarReflection() const
-{
-	return Material->bAcceptsPlanarReflection;
-}
-
 bool FMaterialResource::OutputsVelocityOnBasePass() const
 {
 	return Material->bOutputVelocityOnBasePass && !IsUIMaterial();
@@ -1371,7 +1373,6 @@ void FMaterial::SetupMaterialEnvironment(
 	OutEnvironment.SetDefine(TEXT("MATERIAL_USES_SCENE_COLOR_COPY"), RequiresSceneColorCopy_GameThread() ? TEXT("1") : TEXT("0"));
 	OutEnvironment.SetDefine(TEXT("MATERIAL_FULLY_ROUGH"), IsFullyRough() ? TEXT("1") : TEXT("0"));
 	OutEnvironment.SetDefine(TEXT("MATERIAL_HQ_FORWARD_REFLECTIONS"), IsUsingHQForwardReflections() ? TEXT("1") : TEXT("0"));
-	OutEnvironment.SetDefine(TEXT("MATERIAL_USE_PLANAR_REFLECTION"), IsUsingPlanarReflection() ? TEXT("1") : TEXT("0"));
 	OutEnvironment.SetDefine(TEXT("MATERIAL_NONMETAL"), IsNonmetal() ? TEXT("1") : TEXT("0"));
 	OutEnvironment.SetDefine(TEXT("MATERIAL_USE_LM_DIRECTIONALITY"), UseLmDirectionality() ? TEXT("1") : TEXT("0"));
 	OutEnvironment.SetDefine(TEXT("MATERIAL_INJECT_EMISSIVE_INTO_LPV"), ShouldInjectEmissiveIntoLPV() ? TEXT("1") : TEXT("0"));
diff --git a/Engine/Source/Runtime/Engine/Private/ModelRender.cpp b/Engine/Source/Runtime/Engine/Private/ModelRender.cpp
--- a/Engine/Source/Runtime/Engine/Private/ModelRender.cpp
+++ b/Engine/Source/Runtime/Engine/Private/ModelRender.cpp
@@ -27,6 +27,8 @@ namespace
 
 }	// anon namespace
 
+DEFINE_LOG_CATEGORY_STATIC(LogModelComponent, Log, All);
+
 /*-----------------------------------------------------------------------------
 FModelVertexBuffer
 -----------------------------------------------------------------------------*/
@@ -250,7 +252,7 @@ public:
 				{
 					ESceneDepthPriorityGroup DepthPriorityGroup = (ESceneDepthPriorityGroup)GetDepthPriorityGroup(View);
 
-					const FMaterialRenderProxy* MatProxyOverride = NULL;
+					const FMaterialRenderProxy* MatProxyOverride = nullptr;
 
 #if WITH_EDITOR
 					if(bInCollisionView && AllowDebugViewmodes())
@@ -350,7 +352,7 @@ public:
 												FMeshBatchElement& BatchElement = MeshElement.Elements[0];
 												BatchElement.IndexBuffer = IndexAllocation.IndexBuffer;
 												MeshElement.VertexFactory = &Component->GetModel()->VertexFactory;
-												MeshElement.MaterialRenderProxy = (MatProxyOverride != NULL) ? MatProxyOverride : ProxyElementInfo.GetMaterial()->GetRenderProxy(bOnlySelectedSurfaces, bOnlyHoveredSurfaces);
+												MeshElement.MaterialRenderProxy = (MatProxyOverride != nullptr) ? MatProxyOverride : ProxyElementInfo.GetMaterial()->GetRenderProxy(bOnlySelectedSurfaces, bOnlyHoveredSurfaces);
 												MeshElement.LCI = &ProxyElementInfo;
 												BatchElement.PrimitiveUniformBufferResource = &GetUniformBuffer();
 												BatchElement.FirstIndex = FirstIndex;
@@ -385,7 +387,7 @@ public:
 								FMeshBatchElement& BatchElement = MeshElement.Elements[0];
 								BatchElement.IndexBuffer = ModelElement.IndexBuffer;
 								MeshElement.VertexFactory = &Component->GetModel()->VertexFactory;
-								MeshElement.MaterialRenderProxy = (MatProxyOverride != NULL) ? MatProxyOverride : Elements[ElementIndex].GetMaterial()->GetRenderProxy(false);
+								MeshElement.MaterialRenderProxy = (MatProxyOverride != nullptr) ? MatProxyOverride : Elements[ElementIndex].GetMaterial()->GetRenderProxy(false);
 								MeshElement.LCI = &Elements[ElementIndex];
 								BatchElement.PrimitiveUniformBufferResource = &GetUniformBuffer();
 								BatchElement.FirstIndex = ModelElement.FirstIndex;
@@ -557,13 +559,19 @@ private:
 			, ModelElement(InModelElement)
 			, Bounds(ModelElement.BoundingBox)
 		{
-			const bool bHasStaticLighting = ModelElement.LightMap != NULL || ModelElement.ShadowMap != NULL;
+			const bool bHasStaticLighting = ModelElement.LightMap != nullptr || ModelElement.ShadowMap != nullptr;
 
 			// Determine the material applied to the model element.
 			Material = ModelElement.Material;
 
+			if (RequiresAdjacencyInformation(Material, InModelElement.Component->GetModel()->VertexFactory.GetType(), InModelElement.Component->GetScene()->GetFeatureLevel()))
+			{
+				UE_LOG(LogModelComponent, Warning, TEXT("Material %s requires adjacency information because of Crack Free Displacement or PN Triangle Tesselation, which is not supported with model components. Falling back to DefaultMaterial."), *Material->GetName());
+				Material = nullptr;
+			}
+
 			// If there isn't an applied material, or if we need static lighting and it doesn't support it, fall back to the default material.
-			if(!ModelElement.Material || (bHasStaticLighting && !ModelElement.Material->CheckMaterialUsage(MATUSAGE_StaticLighting)))
+			if(!Material || (bHasStaticLighting && !Material->CheckMaterialUsage(MATUSAGE_StaticLighting)))
 			{
 				Material = UMaterial::GetDefaultMaterial(MD_Surface);
 			}
@@ -647,7 +655,7 @@ public:
 	 *	@param	Index			The index of interest
 	 *
 	 *	@return	FElementInfo*	The element info at that index.
-	 *							NULL if out of range.
+	 *							nullptr if out of range.
 	 */
 	const FElementInfo* GetElement(int32 Index) const 
 	{
@@ -655,7 +663,7 @@ public:
 		{
 			return &(Elements[Index]);
 		}
-		return NULL;
+		return nullptr;
 	}
 
 	virtual void GetLCIs(FLCIArray& LCIs) override
diff --git a/Engine/Source/Runtime/Engine/Private/Particles/ParticleBeam2EmitterInstance.cpp b/Engine/Source/Runtime/Engine/Private/Particles/ParticleBeam2EmitterInstance.cpp
--- a/Engine/Source/Runtime/Engine/Private/Particles/ParticleBeam2EmitterInstance.cpp
+++ b/Engine/Source/Runtime/Engine/Private/Particles/ParticleBeam2EmitterInstance.cpp
@@ -1007,10 +1007,11 @@ void FParticleBeam2EmitterInstance::DetermineVertexAndTriangleCount()
  *	Retrieves the dynamic data for the emitter
  *	
  *	@param	bSelected					Whether the emitter is selected in the editor
+ *  @param InFeatureLevel				The relevant shader feature level.
  *
  *	@return	FDynamicEmitterDataBase*	The dynamic data, or NULL if it shouldn't be rendered
  */
-FDynamicEmitterDataBase* FParticleBeam2EmitterInstance::GetDynamicData(bool bSelected)
+FDynamicEmitterDataBase* FParticleBeam2EmitterInstance::GetDynamicData(bool bSelected, ERHIFeatureLevel::Type InFeatureLevel)
 {
 	QUICK_SCOPE_CYCLE_COUNTER(STAT_ParticleBeam2EmitterInstance_GetDynamicData);
 
@@ -1069,68 +1070,6 @@ FDynamicEmitterDataBase* FParticleBeam2EmitterInstance::GetDynamicData(bool bSel
 	return NewEmitterData;
 }
 
-/**
- *	Updates the dynamic data for the instance
- *
- *	@param	DynamicData		The dynamic data to fill in
- *	@param	bSelected		true if the particle system component is selected
- */
-bool FParticleBeam2EmitterInstance::UpdateDynamicData(FDynamicEmitterDataBase* DynamicData, bool bSelected)
-{
-	QUICK_SCOPE_CYCLE_COUNTER(STAT_ParticleBeam2EmitterInstance_UpdateDynamicData);
-
-	if (ActiveParticles <= 0 || !bEnabled)
-	{
-		return false;
-	}
-
-	UParticleLODLevel* LODLevel = SpriteTemplate->GetCurrentLODLevel(this);
-	if ((LODLevel == NULL) || (LODLevel->bEnabled == false))
-	{
-		return false;
-	}
-
-	//@SAS. This removes the need for the assertion in the actual render call...
-	if ((ActiveParticles > FDynamicBeam2EmitterData::MaxBeams) ||	// TTP #33330 - Max of 2048 beams from a single emitter
-		(ParticleStride >
-			((FDynamicBeam2EmitterData::MaxInterpolationPoints + 2) * (sizeof(FVector) + sizeof(float))) + 
-			(FDynamicBeam2EmitterData::MaxNoiseFrequency * (sizeof(FVector) + sizeof(FVector) + sizeof(float) + sizeof(float)))
-		)	// TTP #33330 - Max of 10k per beam (includes interpolation points, noise, etc.)
-		)
-	{
-#if !(UE_BUILD_SHIPPING || UE_BUILD_TEST)
-		if (Component && Component->GetWorld())
-		{
-			FString ErrorMessage = 
-				FString::Printf(TEXT("BeamEmitter with too much data: %s"),
-					Component ? 
-						Component->Template ? 
-							*(Component->Template->GetName()) :
-							TEXT("No template") :
-						TEXT("No component"));
-			FColor ErrorColor(255,0,0);
-			GEngine->AddOnScreenDebugMessage((uint64)((PTRINT)this), 5.0f, ErrorColor,ErrorMessage);
-			UE_LOG(LogParticles, Log, TEXT("%s"), *ErrorMessage);
-		}
-#endif	//#if !(UE_BUILD_SHIPPING || UE_BUILD_TEST)
-		return false;
-	}
-
-	checkf((DynamicData->GetSource().eEmitterType == DET_Beam2), TEXT("Beam2::UpdateDynamicData> Invalid DynamicData type!"));
-
-	// Now fill in the source data
-	FDynamicBeam2EmitterData* BeamDynamicData = (FDynamicBeam2EmitterData*)DynamicData;
-	if( !FillReplayData( BeamDynamicData->Source ) )
-	{
-		return false;
-	}
-
-	// Setup dynamic render data.  Only call this AFTER filling in source data for the emitter.
-	BeamDynamicData->Init( bSelected );
-
-	return true;
-}
-
 /**
  *	Retrieves replay data for the emitter
  *
diff --git a/Engine/Source/Runtime/Engine/Private/Particles/ParticleComponents.cpp b/Engine/Source/Runtime/Engine/Private/Particles/ParticleComponents.cpp
--- a/Engine/Source/Runtime/Engine/Private/Particles/ParticleComponents.cpp
+++ b/Engine/Source/Runtime/Engine/Private/Particles/ParticleComponents.cpp
@@ -3625,7 +3625,7 @@ void UParticleSystemComponent::DestroyRenderState_Concurrent()
 
 
 FDynamicEmitterDataBase* UParticleSystemComponent::CreateDynamicDataFromReplay( FParticleEmitterInstance* EmitterInstance, 
-	const FDynamicEmitterReplayDataBase* EmitterReplayData, bool bSelected )
+	const FDynamicEmitterReplayDataBase* EmitterReplayData, bool bSelected, ERHIFeatureLevel::Type InFeatureLevel )
 {
 	checkSlow(EmitterInstance && EmitterInstance->CurrentLODLevel);
 	check( EmitterReplayData != NULL );
@@ -3679,8 +3679,8 @@ FDynamicEmitterDataBase* UParticleSystemComponent::CreateDynamicDataFromReplay(
 					NewEmitterData->Init(
 						bSelected,
 						MeshEmitterInstance,
-						MeshEmitterInstance->MeshTypeData->Mesh
-						);
+						MeshEmitterInstance->MeshTypeData->Mesh,
+						InFeatureLevel);
 					EmitterData = NewEmitterData;
 				}
 			}
@@ -3748,7 +3748,7 @@ FDynamicEmitterDataBase* UParticleSystemComponent::CreateDynamicDataFromReplay(
 
 
 
-FParticleDynamicData* UParticleSystemComponent::CreateDynamicData()
+FParticleDynamicData* UParticleSystemComponent::CreateDynamicData(ERHIFeatureLevel::Type InFeatureLevel)
 {
 	//SCOPE_CYCLE_COUNTER(STAT_ParticleSystemComponent_CreateDynamicData);
 
@@ -3820,7 +3820,7 @@ FParticleDynamicData* UParticleSystemComponent::CreateDynamicData()
 						// Fill dynamic data from the replay frame data for this emitter so we can render it
 						// Grab the original emitter instance for that this replay was generated from
 						FDynamicEmitterDataBase* NewDynamicEmitterData =
-							CreateDynamicDataFromReplay( EmitterInstances[ CurEmitter.OriginalEmitterIndex ], CurEmitterReplay, IsOwnerSelected() );
+							CreateDynamicDataFromReplay( EmitterInstances[ CurEmitter.OriginalEmitterIndex ], CurEmitterReplay, IsOwnerSelected(), InFeatureLevel );
 
 
 						if( NewDynamicEmitterData != NULL )
@@ -3896,7 +3896,7 @@ FParticleDynamicData* UParticleSystemComponent::CreateDynamicData()
 							bIsOwnerSeleted = IsOwnerSelected();
 						}
 #endif
-						NewDynamicEmitterData = EmitterInst->GetDynamicData(bIsOwnerSeleted);
+						NewDynamicEmitterData = EmitterInst->GetDynamicData(bIsOwnerSeleted, InFeatureLevel);
 					}
 					if( NewDynamicEmitterData != NULL )
 					{
@@ -4003,7 +4003,7 @@ void UParticleSystemComponent::UpdateDynamicData(FParticleSystemSceneProxy* Prox
 	if (SceneProxy)
 	{
 		// Create the dynamic data for rendering this particle system
-		FParticleDynamicData* ParticleDynamicData = CreateDynamicData();
+		FParticleDynamicData* ParticleDynamicData = CreateDynamicData(Proxy->GetScene().GetFeatureLevel());
 		// Render the particles
 #if !(UE_BUILD_SHIPPING || UE_BUILD_TEST)
 		[REDACTED_EMAIL]. Remove thisline  - it is used for debugging purposes...
diff --git a/Engine/Source/Runtime/Engine/Private/Particles/ParticleEmitterInstances.cpp b/Engine/Source/Runtime/Engine/Private/Particles/ParticleEmitterInstances.cpp
--- a/Engine/Source/Runtime/Engine/Private/Particles/ParticleEmitterInstances.cpp
+++ b/Engine/Source/Runtime/Engine/Private/Particles/ParticleEmitterInstances.cpp
@@ -2581,6 +2581,7 @@ bool FParticleEmitterInstance::FillReplayData( FDynamicEmitterReplayDataBase& Ou
  * Gathers material relevance flags for this emitter instance.
  * @param OutMaterialRelevance - Pointer to where material relevance flags will be stored.
  * @param LODLevel - The LOD level for which to compute material relevance flags.
+ * @param InFeatureLevel - The relevant shader feature level.
  */
 void FParticleEmitterInstance::GatherMaterialRelevance(FMaterialRelevance* OutMaterialRelevance, const UParticleLODLevel* LODLevel, ERHIFeatureLevel::Type InFeatureLevel) const
 {
@@ -2692,10 +2693,11 @@ FParticleSpriteEmitterInstance::~FParticleSpriteEmitterInstance()
  *	Retrieves the dynamic data for the emitter
  *	
  *	@param	bSelected					Whether the emitter is selected in the editor
+ *  @param InFeatureLevel				The relevant shader feature level.
  *
  *	@return	FDynamicEmitterDataBase*	The dynamic data, or NULL if it shouldn't be rendered
  */
-FDynamicEmitterDataBase* FParticleSpriteEmitterInstance::GetDynamicData(bool bSelected)
+FDynamicEmitterDataBase* FParticleSpriteEmitterInstance::GetDynamicData(bool bSelected, ERHIFeatureLevel::Type InFeatureLevel)
 {
 	QUICK_SCOPE_CYCLE_COUNTER(STAT_ParticleSpriteEmitterInstance_GetDynamicData);
 
@@ -2728,41 +2730,6 @@ FDynamicEmitterDataBase* FParticleSpriteEmitterInstance::GetDynamicData(bool bSe
 	return NewEmitterData;
 }
 
-/**
- *	Updates the dynamic data for the instance
- *
- *	@param	DynamicData		The dynamic data to fill in
- *	@param	bSelected		true if the particle system component is selected
- */
-bool FParticleSpriteEmitterInstance::UpdateDynamicData(FDynamicEmitterDataBase* DynamicData, bool bSelected)
-{
-	QUICK_SCOPE_CYCLE_COUNTER(STAT_ParticleSpriteEmitterInstance_UpdateDynamicData);
-
-	checkf((DynamicData->GetSource().eEmitterType == DET_Sprite), TEXT("Sprite::UpdateDynamicData> Invalid DynamicData type!"));
-
-	if (ActiveParticles <= 0 || !bEnabled)
-	{
-		return false;
-	}
-
-	UParticleLODLevel* LODLevel = SpriteTemplate->GetCurrentLODLevel(this);
-	if ((LODLevel == NULL) || (LODLevel->bEnabled == false))
-	{
-		return false;
-	}
-	FDynamicSpriteEmitterData* SpriteDynamicData = (FDynamicSpriteEmitterData*)DynamicData;
-	// Now fill in the source data
-	if( !FillReplayData( SpriteDynamicData->Source ) )
-	{
-		return false;
-	}
-
-	// Setup dynamic render data.  Only call this AFTER filling in source data for the emitter.
-	SpriteDynamicData->Init( bSelected );
-
-	return true;
-}
-
 /**
  *	Retrieves replay data for the emitter
  *
@@ -3363,10 +3330,11 @@ bool FParticleMeshEmitterInstance::IsDynamicDataRequired(UParticleLODLevel* Curr
  *	Retrieves the dynamic data for the emitter
  *	
  *	@param	bSelected					Whether the emitter is selected in the editor
+ *  @param InFeatureLevel				The relevant shader feature level.
  *
  *	@return	FDynamicEmitterDataBase*	The dynamic data, or NULL if it shouldn't be rendered
  */
-FDynamicEmitterDataBase* FParticleMeshEmitterInstance::GetDynamicData(bool bSelected)
+FDynamicEmitterDataBase* FParticleMeshEmitterInstance::GetDynamicData(bool bSelected, ERHIFeatureLevel::Type InFeatureLevel)
 {
 	QUICK_SCOPE_CYCLE_COUNTER(STAT_ParticleMeshEmitterInstance_GetDynamicData);
 
@@ -3398,52 +3366,13 @@ FDynamicEmitterDataBase* FParticleMeshEmitterInstance::GetDynamicData(bool bSele
 	NewEmitterData->Init(
 		bSelected,
 		this,
-		MeshTypeData->Mesh
+		MeshTypeData->Mesh,
+		InFeatureLevel
 		);
 
 	return NewEmitterData;
 }
 
-/**
- *	Updates the dynamic data for the instance
- *
- *	@param	DynamicData		The dynamic data to fill in
- *	@param	bSelected		true if the particle system component is selected
- */
-bool FParticleMeshEmitterInstance::UpdateDynamicData(FDynamicEmitterDataBase* DynamicData, bool bSelected)
-{
-	QUICK_SCOPE_CYCLE_COUNTER(STAT_ParticleMeshEmitterInstance_UpdateDynamicData);
-
-	if (ActiveParticles <= 0 || !bEnabled)
-	{
-		return false;
-	}
-
-	UParticleLODLevel* LODLevel = SpriteTemplate->GetCurrentLODLevel(this);
-	if ((LODLevel == NULL) || (LODLevel->bEnabled == false))
-	{
-		return false;
-	}
-
-	checkf((DynamicData->GetSource().eEmitterType == DET_Mesh), TEXT("Mesh::UpdateDynamicData> Invalid DynamicData type!"));
-
-	FDynamicMeshEmitterData* MeshDynamicData = (FDynamicMeshEmitterData*)DynamicData;
-	// Now fill in the source data
-	if( !FillReplayData( MeshDynamicData->Source ) )
-	{
-		return false;
-	}
-
-	// Setup dynamic render data.  Only call this AFTER filling in source data for the emitter.
-	MeshDynamicData->Init(
-		bSelected,
-		this,
-		MeshTypeData->Mesh
-		);
-
-	return true;
-}
-
 /**
  *	Retrieves replay data for the emitter
  *
@@ -3520,11 +3449,12 @@ void FParticleMeshEmitterInstance::SetMeshMaterials( const TArray<UMaterialInter
  * Gathers material relevance flags for this emitter instance.
  * @param OutMaterialRelevance - Pointer to where material relevance flags will be stored.
  * @param LODLevel - The LOD level for which to compute material relevance flags.
+ * @param InFeatureLevel - The relevant shader feature level.
  */
 void FParticleMeshEmitterInstance::GatherMaterialRelevance( FMaterialRelevance* OutMaterialRelevance, const UParticleLODLevel* LODLevel, ERHIFeatureLevel::Type InFeatureLevel ) const
 {
 	TArray<UMaterialInterface*,TInlineAllocator<2> > Materials;
-	GetMeshMaterials(Materials, LODLevel);
+	GetMeshMaterials(Materials, LODLevel, InFeatureLevel, true); // Allow log issues since GatherMaterialRelevance is only called when the proxy is created.
 	for (int32 MaterialIndex = 0; MaterialIndex < Materials.Num(); ++MaterialIndex)
 	{
 		(*OutMaterialRelevance) |= Materials[MaterialIndex]->GetRelevance(InFeatureLevel);
@@ -3533,7 +3463,9 @@ void FParticleMeshEmitterInstance::GatherMaterialRelevance( FMaterialRelevance*
 
 void FParticleMeshEmitterInstance::GetMeshMaterials(
 	TArray<UMaterialInterface*,TInlineAllocator<2> >& OutMaterials,
-	const UParticleLODLevel* LODLevel
+	const UParticleLODLevel* LODLevel,
+	ERHIFeatureLevel::Type InFeatureLevel,
+	bool bLogWarnings
 	) const
 {
 	if (MeshTypeData && MeshTypeData->Mesh)
@@ -3579,6 +3511,16 @@ void FParticleMeshEmitterInstance::GetMeshMaterials(
 				Material = MeshTypeData->Mesh->GetMaterial(LODModel.Sections[SectionIndex].MaterialIndex);
 			}
 
+			// Check that adjacency data is not required since the implementation does not support it.
+			if (RequiresAdjacencyInformation(Material, LODModel.VertexFactory.GetType(), InFeatureLevel))
+			{
+				if (bLogWarnings)
+				{
+					UE_LOG(LogParticles, Warning, TEXT("Material %s requires adjacency information because of Crack Free Displacement or PN Triangle Tesselation, which is not supported with particles. Falling back to DefaultMaterial."), *Material->GetName());
+				}
+				Material = NULL;
+			}
+
 			// Use the default material...
 			if (Material == NULL)
 			{
diff --git a/Engine/Source/Runtime/Engine/Private/Particles/ParticleGpuSimulation.cpp b/Engine/Source/Runtime/Engine/Private/Particles/ParticleGpuSimulation.cpp
--- a/Engine/Source/Runtime/Engine/Private/Particles/ParticleGpuSimulation.cpp
+++ b/Engine/Source/Runtime/Engine/Private/Particles/ParticleGpuSimulation.cpp
@@ -896,14 +896,8 @@ public:
 	{
 		FGlobalShader::ModifyCompilationEnvironment( Platform, OutEnvironment );
 		OutEnvironment.SetDefine(TEXT("TILES_PER_INSTANCE"), TILES_PER_INSTANCE);
-		OutEnvironment.SetFloatDefine(
-			TEXT("TILE_SIZE_X"),
-			(float)GParticleSimulationTileSize / (float)GParticleSimulationTextureSizeX
-			);
-		OutEnvironment.SetFloatDefine(
-			TEXT("TILE_SIZE_Y"),
-			(float)GParticleSimulationTileSize / (float)GParticleSimulationTextureSizeY
-			);
+		OutEnvironment.SetDefine(TEXT("TILE_SIZE_X"), (float)GParticleSimulationTileSize / (float)GParticleSimulationTextureSizeX);
+		OutEnvironment.SetDefine(TEXT("TILE_SIZE_Y"), (float)GParticleSimulationTileSize / (float)GParticleSimulationTextureSizeY);
 
 		if (Platform == SP_OPENGL_ES2_ANDROID)
 		{
@@ -3127,7 +3121,7 @@ public:
 	/**
 	 *	Retrieves the dynamic data for the emitter
 	 */
-	virtual FDynamicEmitterDataBase* GetDynamicData(bool bSelected) override
+	virtual FDynamicEmitterDataBase* GetDynamicData(bool bSelected, ERHIFeatureLevel::Type InFeatureLevel) override
 	{
 		QUICK_SCOPE_CYCLE_COUNTER(STAT_FDynamicEmitterDataBase_GetDynamicData);
 		check(Component);
diff --git a/Engine/Source/Runtime/Engine/Private/Particles/ParticleSystemRender.cpp b/Engine/Source/Runtime/Engine/Private/Particles/ParticleSystemRender.cpp
--- a/Engine/Source/Runtime/Engine/Private/Particles/ParticleSystemRender.cpp
+++ b/Engine/Source/Runtime/Engine/Private/Particles/ParticleSystemRender.cpp
@@ -1335,7 +1335,8 @@ FDynamicMeshEmitterData::~FDynamicMeshEmitterData()
 /** Initialize this emitter's dynamic rendering data, called after source data has been filled in */
 void FDynamicMeshEmitterData::Init( bool bInSelected,
 									const FParticleMeshEmitterInstance* InEmitterInstance,
-									UStaticMesh* InStaticMesh )
+									UStaticMesh* InStaticMesh,
+									ERHIFeatureLevel::Type InFeatureLevel )
 {
 	bSelected = bInSelected;
 
@@ -1351,8 +1352,8 @@ void FDynamicMeshEmitterData::Init( bool bInSelected,
 
 	InEmitterInstance->GetMeshMaterials(
 		MeshMaterials,
-		InEmitterInstance->SpriteTemplate->LODLevels[InEmitterInstance->CurrentLODLevelIndex]
-		);
+		InEmitterInstance->SpriteTemplate->LODLevels[InEmitterInstance->CurrentLODLevelIndex],
+		InFeatureLevel);
 
 	for (int32 i = 0; i < MeshMaterials.Num(); ++i)
 	{
@@ -7032,7 +7033,7 @@ FPrimitiveSceneProxy* UParticleSystemComponent::CreateSceneProxy()
 		}
 
 		// Create the dynamic data for rendering this particle system.
-		FParticleDynamicData* ParticleDynamicData = CreateDynamicData();
+		FParticleDynamicData* ParticleDynamicData = CreateDynamicData(GetScene()->GetFeatureLevel());
 
 		if (CanBeOccluded())
 		{
diff --git a/Engine/Source/Runtime/Engine/Private/Particles/ParticleTrail2EmitterInstance.cpp b/Engine/Source/Runtime/Engine/Private/Particles/ParticleTrail2EmitterInstance.cpp
--- a/Engine/Source/Runtime/Engine/Private/Particles/ParticleTrail2EmitterInstance.cpp
+++ b/Engine/Source/Runtime/Engine/Private/Particles/ParticleTrail2EmitterInstance.cpp
@@ -2592,7 +2592,7 @@ bool FParticleRibbonEmitterInstance::IsDynamicDataRequired(UParticleLODLevel* Cu
 /**
  *	Retrieves the dynamic data for the emitter
  */
-FDynamicEmitterDataBase* FParticleRibbonEmitterInstance::GetDynamicData(bool bSelected)
+FDynamicEmitterDataBase* FParticleRibbonEmitterInstance::GetDynamicData(bool bSelected, ERHIFeatureLevel::Type InFeatureLevel)
 {
 	QUICK_SCOPE_CYCLE_COUNTER(STAT_ParticleRibbonEmitterInstance_GetDynamicData);
 
@@ -2642,47 +2642,6 @@ FDynamicEmitterDataBase* FParticleRibbonEmitterInstance::GetDynamicData(bool bSe
 	return NewEmitterData;
 }
 
-/**
- *	Updates the dynamic data for the instance
- *
- *	@param	DynamicData		The dynamic data to fill in
- *	@param	bSelected		true if the particle system component is selected
- */
-bool FParticleRibbonEmitterInstance::UpdateDynamicData(FDynamicEmitterDataBase* DynamicData, bool bSelected)
-{
-	QUICK_SCOPE_CYCLE_COUNTER(STAT_ParticleRibbonEmitterInstance_UpdateDynamicData);
-
-	if (ActiveParticles <= 0 || !bEnabled)
-	{
-		return false;
-	}
-
-	if (DynamicData->GetSource().eEmitterType != DET_Ribbon)
-	{
-		UE_LOG(LogParticles, Warning, TEXT("UpdateDynamicData> NOT A TRAILS EMITTER!"));
-		return false;
-	}
-
-	FDynamicRibbonEmitterData* TrailDynamicData = (FDynamicRibbonEmitterData*)DynamicData;
-	// Now fill in the source data
-	if (!FillReplayData(TrailDynamicData->Source))
-	{
-		return false;
-	}
-
-	TrailDynamicData->bRenderGeometry = TrailTypeData->bRenderGeometry;
-	TrailDynamicData->bRenderParticles = TrailTypeData->bRenderSpawnPoints;
-	TrailDynamicData->bRenderTangents = TrailTypeData->bRenderTangents;
-	TrailDynamicData->bRenderTessellation = TrailTypeData->bRenderTessellation;
-	TrailDynamicData->DistanceTessellationStepSize = TrailTypeData->DistanceTessellationStepSize;
-	TrailDynamicData->TangentTessellationScalar = TrailTypeData->TangentTessellationScalar;
-
-	// Setup dynamic render data.  Only call this AFTER filling in source data for the emitter.
-	TrailDynamicData->Init(bSelected);
-
-	return true;
-}
-
 /**
  *	Retrieves replay data for the emitter
  *
@@ -3926,7 +3885,7 @@ bool FParticleAnimTrailEmitterInstance::HasCompleted()
 /**
  *	Retrieves the dynamic data for the emitter
  */
-FDynamicEmitterDataBase* FParticleAnimTrailEmitterInstance::GetDynamicData(bool bSelected)
+FDynamicEmitterDataBase* FParticleAnimTrailEmitterInstance::GetDynamicData(bool bSelected, ERHIFeatureLevel::Type InFeatureLevel)
 {
 	QUICK_SCOPE_CYCLE_COUNTER(STAT_ParticleAnimTrailEmitterInstance_GetDynamicData);
 
@@ -3985,52 +3944,6 @@ FDynamicEmitterDataBase* FParticleAnimTrailEmitterInstance::GetDynamicData(bool
 	return NewEmitterData;
 }
 
-/**
- *	Updates the dynamic data for the instance
- *
- *	@param	DynamicData		The dynamic data to fill in
- *	@param	bSelected		true if the particle system component is selected
- */
-bool FParticleAnimTrailEmitterInstance::UpdateDynamicData(FDynamicEmitterDataBase* DynamicData, bool bSelected)
-{
-	QUICK_SCOPE_CYCLE_COUNTER(STAT_ParticleAnimTrailEmitterInstance_UpdateDynamicData);
-
-	if (ActiveParticles <= 0 || !bEnabled)
-	{
-		return false;
-	}
-
-	if (DynamicData->GetSource().eEmitterType != DET_AnimTrail)
-	{
-		UE_LOG(LogParticles, Warning, TEXT("UpdateDynamicData> NOT A TRAILS EMITTER!"));
-		return false;
-	}
-
-	FDynamicAnimTrailEmitterData* TrailDynamicData = (FDynamicAnimTrailEmitterData*)DynamicData;
-	// Now fill in the source data
-	if (!FillReplayData(TrailDynamicData->Source))
-	{
-		return false;
-	}
-
-#if WITH_EDITORONLY_DATA
-	TrailDynamicData->bRenderGeometry = bRenderGeometry;
-	TrailDynamicData->bRenderParticles = bRenderSpawnPoints;
-	TrailDynamicData->bRenderTangents = bRenderTangents;
-	TrailDynamicData->bRenderTessellation = bRenderTessellation;
-#else
-	TrailDynamicData->bRenderGeometry = true;
-	TrailDynamicData->bRenderParticles = false;
-	TrailDynamicData->bRenderTangents = false;
-	TrailDynamicData->bRenderTessellation = false;
-#endif
-
-	// Setup dynamic render data.  Only call this AFTER filling in source data for the emitter.
-	TrailDynamicData->Init(bSelected);
-
-	return true;
-}
-
 /**
  *	Retrieves replay data for the emitter
  *
