Copying //UE4/Dev-Rendering to Dev-Main (Source //UE4/Dev-Rendering@2932636)

diff --git a/Engine/Source/Runtime/Engine/Private/PlanarReflectionSceneProxy.cpp b/Engine/Source/Runtime/Engine/Private/PlanarReflectionSceneProxy.cpp
new file mode 100644
--- /dev/null
+++ b/Engine/Source/Runtime/Engine/Private/PlanarReflectionSceneProxy.cpp
@@ -0,0 +1,36 @@
+// Copyright 1998-2016 Epic Games, Inc. All Rights Reserved.
+
+/*=============================================================================
+	PlanarReflectionProxy.cpp
+=============================================================================*/
+
+#include "EnginePrivate.h"
+#include "Components/PlanarReflectionComponent.h"
+#include "PlanarReflectionSceneProxy.h"
+
+FPlanarReflectionSceneProxy::FPlanarReflectionSceneProxy(UPlanarReflectionComponent* Component, FPlanarReflectionRenderTarget* InRenderTarget)
+{
+	RenderTarget = InRenderTarget;
+	UpdateTransform(Component->ComponentToWorld.ToMatrixWithScale());
+
+	float ClampedFadeStart = FMath::Max(Component->DistanceFromPlaneFadeStart, 0.0f);
+	float ClampedFadeEnd = FMath::Max(Component->DistanceFromPlaneFadeEnd, 0.0f);
+	float DistanceFadeScale = 1.0f / FMath::Max(ClampedFadeEnd - ClampedFadeStart, DELTA);
+
+	PlanarReflectionParameters = FVector(
+		DistanceFadeScale,
+		-ClampedFadeStart * DistanceFadeScale,
+		Component->NormalDistortionStrength);
+
+	const float CosFadeStart = FMath::Cos(FMath::Clamp(Component->AngleFromPlaneFadeStart, 0.1f, 89.9f) * (float)PI / 180.0f);
+	const float CosFadeEnd = FMath::Cos(FMath::Clamp(Component->AngleFromPlaneFadeEnd, 0.1f, 89.9f) * (float)PI / 180.0f);
+	const float Range = 1.0f / FMath::Max(CosFadeStart - CosFadeEnd, DELTA);
+
+	PlanarReflectionParameters2 = FVector2D(
+		Range,
+		-CosFadeEnd * Range);
+
+	Component->GetProjectionWithExtraFOV(ProjectionWithExtraFOV);
+
+	OwnerName = Component->GetOwner() ? Component->GetOwner()->GetFName() : NAME_None;
+}
diff --git a/Engine/Source/Runtime/Engine/Private/PrimitiveSceneProxy.cpp b/Engine/Source/Runtime/Engine/Private/PrimitiveSceneProxy.cpp
--- a/Engine/Source/Runtime/Engine/Private/PrimitiveSceneProxy.cpp
+++ b/Engine/Source/Runtime/Engine/Private/PrimitiveSceneProxy.cpp
@@ -42,7 +42,6 @@ FPrimitiveSceneProxy::FPrimitiveSceneProxy(const UPrimitiveComponent* InComponen
 ,   bAffectDynamicIndirectLighting(InComponent->bAffectDynamicIndirectLighting)
 ,   bAffectDistanceFieldLighting(InComponent->bAffectDistanceFieldLighting)
 ,	bCastStaticShadow(InComponent->CastShadow && InComponent->bCastStaticShadow)
-,	bVisibleInPlanarReflection(InComponent->bVisibleInPlanarReflection)
 ,	bCastVolumetricTranslucentShadow(InComponent->bCastDynamicShadow && InComponent->CastShadow && InComponent->bCastVolumetricTranslucentShadow)
 ,	bCastCapsuleDirectShadow(false)
 ,	bCastCapsuleIndirectShadow(false)
diff --git a/Engine/Source/Runtime/Engine/Private/SceneManagement.cpp b/Engine/Source/Runtime/Engine/Private/SceneManagement.cpp
--- a/Engine/Source/Runtime/Engine/Private/SceneManagement.cpp
+++ b/Engine/Source/Runtime/Engine/Private/SceneManagement.cpp
@@ -373,27 +373,15 @@ FLODMask ComputeLODForMeshes( const TIndirectArray<class FStaticMesh>& StaticMes
 	FLODMask LODToRender;
 
 	// Handle forced LOD level first
-	if (ForcedLODLevel >= 0 || View.bIsPlanarReflectionCapture)
+	if (ForcedLODLevel >= 0)
 	{
-		// Note: starting at -1 which is the default LODIndex, for cases where LODIndex didn't get set
-		int8 MinLOD = 127, MaxLOD = -1;
-		for(int32 MeshIndex = 0 ; MeshIndex < StaticMeshes.Num() ; ++MeshIndex)
+		int8 MaxLOD = 0;
+		for (int32 MeshIndex = 0; MeshIndex < StaticMeshes.Num(); ++MeshIndex)
 		{
 			const FStaticMesh&  Mesh = StaticMeshes[MeshIndex];
-			MinLOD = FMath::Min(MinLOD, Mesh.LODIndex);
 			MaxLOD = FMath::Max(MaxLOD, Mesh.LODIndex);
 		}
-
-		//@todo Ronin
-		// Planar reflections always use the highest available LOD level
-		if (View.bIsPlanarReflectionCapture)
-		{
-			LODToRender.SetLOD(MaxLOD);
-		}
-		else
-		{
-			LODToRender.SetLOD(FMath::Clamp<int8>(ForcedLODLevel, MinLOD, MaxLOD));
-		}
+		LODToRender.SetLOD(FMath::Clamp<int8>(ForcedLODLevel, 0, MaxLOD));
 	}
 	else if (View.Family->EngineShowFlags.LOD)
 	{
diff --git a/Engine/Source/Runtime/Engine/Private/SceneView.cpp b/Engine/Source/Runtime/Engine/Private/SceneView.cpp
--- a/Engine/Source/Runtime/Engine/Private/SceneView.cpp
+++ b/Engine/Source/Runtime/Engine/Private/SceneView.cpp
@@ -134,6 +134,34 @@ static TAutoConsoleVariable<int32> CVarScreenPercentageEditor(
 	TEXT("1: allow upsample (blurry but faster) and downsample (cripser but slower)"),
 	ECVF_Default);
 
+static TAutoConsoleVariable<float> CVarDepthOfFieldDepthBlurAmount(
+	TEXT("r.DepthOfField.DepthBlur.Amount"),
+	1.0f,
+	TEXT("This scale multiplier only affects the CircleDOF DepthBlur feature (value defines in how many km the radius goes to 50%).\n")
+	TEXT(" x: Multiply the existing Depth Blur Amount with x\n")
+	TEXT("-x: Override the existing Depth Blur Amount with x (in km)\n")
+	TEXT(" 1: No adjustments (default)"),
+	ECVF_RenderThreadSafe | ECVF_Cheat);
+
+static TAutoConsoleVariable<float> CVarDepthOfFieldDepthBlurScale(
+	TEXT("r.DepthOfField.DepthBlur.Scale"),
+	1.0f,
+	TEXT("This scale multiplier only affects the CircleDOF DepthBlur feature. This is applied after r.DepthOfField.DepthBlur.ResolutionScale.\n")
+	TEXT(" 0: Disable Depth Blur\n")
+	TEXT(" x: Multiply the existing Depth Blur Radius with x\n")
+	TEXT("-x: Override the existing Depth Blur Radius with x\n")
+	TEXT(" 1: No adjustments (default)"),
+	ECVF_RenderThreadSafe | ECVF_Cheat);
+
+static TAutoConsoleVariable<float> CVarDepthOfFieldDepthBlurResolutionScale(
+	TEXT("r.DepthOfField.DepthBlur.ResolutionScale"),
+	1.0f,
+	TEXT("This scale multiplier only affects the CircleDOF DepthBlur feature. It's a temporary hack.\n")
+	TEXT("It lineary scale the DepthBlur by the resolution increase over 1920 (in width), does only affect resolution larger than that.\n")
+	TEXT("Actual math: float Factor = max(ViewWidth / 1920 - 1, 0); DepthBlurRadius *= 1 + Factor * (CVar - 1)\n")
+	TEXT(" 1: No adjustments (default)\n")
+	TEXT(" x: if the resolution is 1920 there is no change, if 2x larger than 1920 it scale the radius by x"),
+	ECVF_RenderThreadSafe | ECVF_Cheat);
 #endif
 
 static TAutoConsoleVariable<float> CVarSSAOFadeRadiusScale(
@@ -381,11 +409,11 @@ FSceneView::FSceneView(const FSceneViewInitOptions& InitOptions)
 	, bIsViewInfo(false)
 	, bIsSceneCapture(false)
 	, bIsReflectionCapture(false)
-	, bIsPlanarReflectionCapture(false)
-	, ReflectionPlane(0.0f)
+	, bIsPlanarReflection(false)
 	, bIsLocked(false)
 	, bStaticSceneOnly(false)
 	, bIsInstancedStereoEnabled(false)
+	, GlobalClippingPlane(FPlane(0, 0, 0, 0))
 #if WITH_EDITOR
 	, OverrideLODViewOrigin(InitOptions.OverrideLODViewOrigin)
 	, bAllowTranslucentPrimitivesInHitProxy( true )
@@ -1462,6 +1490,30 @@ void FSceneView::EndFinalPostprocessSettings(const FSceneViewInitOptions& ViewIn
 		float Value = CVarExposureOffset.GetValueOnGameThread();
 		FinalPostProcessSettings.AutoExposureBias += Value;
 	}
+
+	{
+		float& DepthBlurAmount = FinalPostProcessSettings.DepthOfFieldDepthBlurAmount;
+
+		float CVarAmount = CVarDepthOfFieldDepthBlurAmount.GetValueOnGameThread();
+
+		DepthBlurAmount = (CVarAmount > 0.0f) ? (DepthBlurAmount * CVarAmount) : -CVarAmount;
+	}
+
+	{
+		float& DepthBlurRadius = FinalPostProcessSettings.DepthOfFieldDepthBlurRadius;
+		{
+			float CVarResScale = FMath::Max(1.0f, CVarDepthOfFieldDepthBlurResolutionScale.GetValueOnGameThread());
+			
+			float Factor = FMath::Max(ViewRect.Width() / 1920.0f - 1.0f, 0.0f);
+
+			DepthBlurRadius *= 1.0f + Factor * (CVarResScale - 1.0f);
+		}
+		{
+			float CVarScale = CVarDepthOfFieldDepthBlurScale.GetValueOnGameThread();
+
+			DepthBlurRadius = (CVarScale > 0.0f) ? (DepthBlurRadius * CVarScale) : -CVarScale;
+		}
+	}
 #endif
 
 	if(FinalPostProcessSettings.DepthOfFieldMethod == DOFM_CircleDOF)
@@ -1808,6 +1860,10 @@ FSceneViewFamily::FSceneViewFamily( const ConstructionValues& CVS )
 	}
 
 	DebugViewShaderMode = ChooseDebugViewShaderMode();
+	if (!AllowDebugViewShaderMode(DebugViewShaderMode, GetFeatureLevel()))
+	{
+		DebugViewShaderMode = DVSM_None;
+	}
 #endif
 
 #if !WITH_EDITOR
@@ -1908,73 +1964,152 @@ ERHIFeatureLevel::Type FSceneViewFamily::GetFeatureLevel() const
 	}
 }
 
-#if !(UE_BUILD_SHIPPING || UE_BUILD_TEST)
-EDebugViewShaderMode FSceneViewFamily::ChooseDebugViewShaderMode() const
+const FSceneView& FSceneViewFamily::GetStereoEyeView(const EStereoscopicPass Eye) const
+{
+	const int32 EyeIndex = static_cast<int32>(Eye);
+	check(Views.Num() > 0 && Views.Num() >= EyeIndex);
+
+	// Mono or left eye
+	if (EyeIndex <= 1)
+	{
+		return *Views[0];
+	}
+
+	// Right eye
+	else
+	{
+		return *Views[1];
+	}
+}
+
+FSceneViewFamilyContext::~FSceneViewFamilyContext()
 {
-	const bool bDebugViewShaderExists = AllowDebugViewModeShader(GetFeatureLevelShaderPlatform(GetFeatureLevel()));
+	// Cleanup the views allocated for this view family.
+	for(int32 ViewIndex = 0;ViewIndex < Views.Num();ViewIndex++)
+	{
+		delete Views[ViewIndex];
+	}
+}
 
-	if (bDebugViewShaderExists)
+#if !(UE_BUILD_SHIPPING || UE_BUILD_TEST)
+
+namespace
+{
+	struct FDebugViewModeOptions
 	{
-		if (EngineShowFlags.ShaderComplexity)
+		FDebugViewModeOptions() : bEnableQuadOverdraw(false), bEnableStreamingAccuracy(false), bEnableTextureStreamingBuild(false)
 		{
-			if (EngineShowFlags.QuadOverdraw)
-			{
-				return DVSM_QuadComplexity;
-			}
-			else if (EngineShowFlags.ShaderComplexityWithQuadOverdraw)
+			// OpenGL shares the same settings as DX on Windows, but some shaders do not compile yet.
+			//@TODO : if (!IsOpenGLPlatform(GMaxRHIShaderPlatform))
+			if (GMaxRHIShaderPlatform == SP_PCD3D_SM5 || GMaxRHIShaderPlatform == SP_PCD3D_SM4)
 			{
-				return DVSM_ShaderComplexityContainedQuadOverhead;
-			}
-			else
-			{
-				return DVSM_ShaderComplexity;
-			}
+				GConfig->GetBool(TEXT("DebugViewModes"), TEXT("bEnableQuadOverdraw"), bEnableQuadOverdraw, GEngineIni);
+				GConfig->GetBool(TEXT("DebugViewModes"), TEXT("bEnableStreamingAccuracy"), bEnableStreamingAccuracy, GEngineIni);
+				GConfig->GetBool(TEXT("DebugViewModes"), TEXT("bEnableTextureStreamingBuild"), bEnableTextureStreamingBuild, GEngineIni);
+			}			
+
+			ERHIFeatureLevel::Type QuadOverdrawFeature = bEnableQuadOverdraw ? ERHIFeatureLevel::SM5 : ERHIFeatureLevel::Num;
+			ERHIFeatureLevel::Type StreamingAccuracyFeature = bEnableStreamingAccuracy ? ERHIFeatureLevel::SM4 : ERHIFeatureLevel::Num;
+			ERHIFeatureLevel::Type TexCoordScaleFeature = bEnableTextureStreamingBuild ? ERHIFeatureLevel::SM4 : ERHIFeatureLevel::Num;
+
+			static_assert(DVSM_MAX == 9, "Wrong DebugViewShaderMode Count");
+			MinFeatureLevels[DVSM_None] = ERHIFeatureLevel::Num;
+
+			MinFeatureLevels[DVSM_ShaderComplexity] = ERHIFeatureLevel::SM4;
+
+			MinFeatureLevels[DVSM_ShaderComplexityContainedQuadOverhead] = QuadOverdrawFeature;
+			MinFeatureLevels[DVSM_ShaderComplexityBleedingQuadOverhead] = QuadOverdrawFeature;			
+			MinFeatureLevels[DVSM_QuadComplexity] = QuadOverdrawFeature;
+
+			MinFeatureLevels[DVSM_WantedMipsAccuracy] = StreamingAccuracyFeature;
+			MinFeatureLevels[DVSM_TexelFactorAccuracy] = StreamingAccuracyFeature;
+
+			MinFeatureLevels[DVSM_TexCoordScaleAccuracy] = TexCoordScaleFeature;
+			MinFeatureLevels[DVSM_TexCoordScaleAnalysis] = TexCoordScaleFeature;
 		}
-		else if (EngineShowFlags.WantedMipsAccuracy)
+
+		bool bEnableQuadOverdraw;
+		bool bEnableStreamingAccuracy;
+		bool bEnableTextureStreamingBuild;
+
+		ERHIFeatureLevel::Type MinFeatureLevels[DVSM_MAX];
+	};
+}
+
+EDebugViewShaderMode FSceneViewFamily::ChooseDebugViewShaderMode() const
+{
+	if (EngineShowFlags.ShaderComplexity)
+	{
+		if (EngineShowFlags.QuadOverdraw)
 		{
-			return DVSM_WantedMipsAccuracy;
+			return DVSM_QuadComplexity;
 		}
-		else if (EngineShowFlags.TexelFactorAccuracy)
+		else if (EngineShowFlags.ShaderComplexityWithQuadOverdraw)
 		{
-			return DVSM_TexelFactorAccuracy;
+			return DVSM_ShaderComplexityContainedQuadOverhead;
 		}
-		else if (EngineShowFlags.TexCoordScaleAccuracy)
+		else
 		{
-			return DVSM_TexCoordScaleAnalysis;
+			return DVSM_ShaderComplexity;
 		}
 	}
-	else if (EngineShowFlags.ShaderComplexity && !EngineShowFlags.QuadOverdraw)
+	else if (EngineShowFlags.WantedMipsAccuracy)
+	{
+		return DVSM_WantedMipsAccuracy;
+	}
+	else if (EngineShowFlags.TexelFactorAccuracy)
 	{
-		return DVSM_ShaderComplexity;
+		return DVSM_TexelFactorAccuracy;
+	}
+	else if (EngineShowFlags.TexCoordAnalysis) // Test before accuracy is set since accuracy could also be set.
+	{
+		return DVSM_TexCoordScaleAnalysis;
+	}
+	else if (EngineShowFlags.TexCoordScaleAccuracy)
+	{
+		return DVSM_TexCoordScaleAccuracy;
 	}
 	return DVSM_None;
-};
-#endif
+}
 
-const FSceneView& FSceneViewFamily::GetStereoEyeView(const EStereoscopicPass Eye) const
+bool AllowDebugViewShaderMode(EDebugViewShaderMode ShaderMode, ERHIFeatureLevel::Type FeatureLevel)
 {
-	const int32 EyeIndex = static_cast<int32>(Eye);
-	check(Views.Num() > 0 && Views.Num() >= EyeIndex);
+#if WITH_EDITOR
+	static const FDebugViewModeOptions Options;
+	check(ShaderMode < DVSM_MAX && FeatureLevel < ERHIFeatureLevel::Num);
+	return FeatureLevel >= Options.MinFeatureLevels[ShaderMode]; 
+#else
+	return ShaderMode == DVSM_ShaderComplexity && FeatureLevel >= ERHIFeatureLevel::SM4;
+#endif
+}
 
-	// Mono or left eye
-	if (EyeIndex <= 1)
+bool AllowDebugViewShaderMode(EDebugViewShaderMode ShaderMode, EShaderPlatform Platform)
+{
+	if (ShaderMode == DVSM_None)
 	{
-		return *Views[0];
+		return false; // This must always return false as it used to know if there is a debug viewmode enabled.
 	}
-
-	// Right eye
-	else
+	//@TODO : else if (!IsPCPlatform(Platform))
+	else if (Platform != SP_PCD3D_SM5 && Platform != SP_PCD3D_SM4)
 	{
-		return *Views[1];
+		return ShaderMode == DVSM_ShaderComplexity && GetMaxSupportedFeatureLevel(Platform) >= ERHIFeatureLevel::SM4; // Only shader complexity is available on non shader platforms.
+	}
+	else // Otherwise, this will be valid only for the current platform.
+	{
+		return AllowDebugViewShaderMode(ShaderMode, GetMaxSupportedFeatureLevel(Platform));
 	}
 }
 
-FSceneViewFamilyContext::~FSceneViewFamilyContext()
+bool AllowDebugViewVSDSHS(ERHIFeatureLevel::Type FeatureLevel)
 {
-	// Cleanup the views allocated for this view family.
-	for(int32 ViewIndex = 0;ViewIndex < Views.Num();ViewIndex++)
-	{
-		delete Views[ViewIndex];
-	}
+	return AllowDebugViewShaderMode(DVSM_QuadComplexity, FeatureLevel) || AllowDebugViewShaderMode(DVSM_WantedMipsAccuracy, FeatureLevel) || AllowDebugViewShaderMode(DVSM_TexCoordScaleAnalysis, FeatureLevel);
 }
 
+bool AllowDebugViewVSDSHS(EShaderPlatform Platform)
+{
+	return AllowDebugViewShaderMode(DVSM_QuadComplexity, Platform) || AllowDebugViewShaderMode(DVSM_WantedMipsAccuracy, Platform) || AllowDebugViewShaderMode(DVSM_TexCoordScaleAnalysis, Platform);
+}
+
+#endif
+
+
diff --git a/Engine/Source/Runtime/Engine/Private/ShaderCompiler/ShaderCompiler.cpp b/Engine/Source/Runtime/Engine/Private/ShaderCompiler/ShaderCompiler.cpp
--- a/Engine/Source/Runtime/Engine/Private/ShaderCompiler/ShaderCompiler.cpp
+++ b/Engine/Source/Runtime/Engine/Private/ShaderCompiler/ShaderCompiler.cpp
@@ -13,6 +13,7 @@
 #include "EngineModule.h"
 #include "ShaderCompiler.h"
 #include "RendererInterface.h"
+#include "ComponentRecreateRenderStateContext.h"
 
 DEFINE_LOG_CATEGORY(LogShaderCompilers);
 
@@ -1111,6 +1112,7 @@ FShaderCompilingManager::FShaderCompilingManager() :
 {
 	WorkersBusyTime = 0;
 	bFallBackToDirectCompiles = false;
+	bRecreateComponentRenderStateOutstanding = false;
 
 	// Threads must use absolute paths on Windows in case the current directory is changed on another thread!
 	ShaderCompileWorkerName = FPaths::ConvertRelativePathToFull(ShaderCompileWorkerName);
@@ -1223,7 +1225,7 @@ FShaderCompilingManager::FShaderCompilingManager() :
 	Thread->StartThread();
 }
 
-void FShaderCompilingManager::AddJobs(TArray<FShaderCommonCompileJob*>& NewJobs, bool bApplyCompletedShaderMapForRendering, bool bOptimizeForLowLatency)
+void FShaderCompilingManager::AddJobs(TArray<FShaderCommonCompileJob*>& NewJobs, bool bApplyCompletedShaderMapForRendering, bool bOptimizeForLowLatency, bool bRecreateComponentRenderStateOnCompletion)
 {
 	check(!FPlatformProperties::RequiresCookedData());
 
@@ -1265,6 +1267,7 @@ void FShaderCompilingManager::AddJobs(TArray<FShaderCommonCompileJob*>& NewJobs,
 		NewJobs[JobIndex]->bOptimizeForLowLatency = bOptimizeForLowLatency;
 		FShaderMapCompileResults& ShaderMapInfo = ShaderMapJobs.FindOrAdd(NewJobs[JobIndex]->Id);
 		ShaderMapInfo.bApplyCompletedShaderMapForRendering = bApplyCompletedShaderMapForRendering;
+		ShaderMapInfo.bRecreateComponentRenderStateOnCompletion = bRecreateComponentRenderStateOnCompletion;
 		auto* PipelineJob = NewJobs[JobIndex]->GetShaderPipelineJob();
 		if (PipelineJob)
 		{
@@ -1455,7 +1458,8 @@ void FShaderCompilingManager::BlockOnAllShaderMapCompletion(TMap<int32, FShaderM
 
 void FShaderCompilingManager::ProcessCompiledShaderMaps(
 	TMap<int32, FShaderMapFinalizeResults>& CompiledShaderMaps, 
-	float TimeBudget)
+	float TimeBudget,
+	bool bRecreateComponentRenderState)
 {
 	// Keeps shader maps alive as they are passed from the shader compiler and applied to the owning FMaterial
 	TArray<TRefCountPtr<FMaterialShaderMap> > LocalShaderMapReferences;
@@ -1652,6 +1656,14 @@ void FShaderCompilingManager::ProcessCompiledShaderMaps(
 			Material->NotifyCompilationFinished();
 		}
 
+		if (bRecreateComponentRenderState)
+		{
+			// This is necessary because scene proxies cache material / shadermap state, like material view relevance
+			//@todo - find a way to only recreate the components referencing materials that have been updated
+			FGlobalComponentRecreateRenderStateContext Context;
+			bRecreateComponentRenderStateOutstanding = false;
+		}
+
 #if WITH_EDITOR
 		FEditorSupportDelegates::RedrawAllViewports.Broadcast();
 #endif // WITH_EDITOR
@@ -1850,7 +1862,7 @@ bool FShaderCompilingManager::HandlePotentialRetryOnError(TMap<int32, FShaderMap
 				}
 
 				// Send all the shaders from this shader map through the compiler again
-				AddJobs(Results.FinishedJobs, Results.bApplyCompletedShaderMapForRendering, true);
+				AddJobs(Results.FinishedJobs, Results.bApplyCompletedShaderMapForRendering, true, Results.bRecreateComponentRenderStateOnCompletion);
 			}
 		}
 
@@ -1958,7 +1970,7 @@ void FShaderCompilingManager::FinishCompilation(const TCHAR* MaterialName, const
 	} 
 	while (bRetry);
 
-	ProcessCompiledShaderMaps(CompiledShaderMaps, FLT_MAX);
+	ProcessCompiledShaderMaps(CompiledShaderMaps, FLT_MAX, false);
 	check(CompiledShaderMaps.Num() == 0);
 
 	const double EndTime = FPlatformTime::Seconds();
@@ -1984,7 +1996,7 @@ void FShaderCompilingManager::FinishAllCompilation()
 	} 
 	while (bRetry);
 
-	ProcessCompiledShaderMaps(CompiledShaderMaps, FLT_MAX);
+	ProcessCompiledShaderMaps(CompiledShaderMaps, FLT_MAX, false);
 	check(CompiledShaderMaps.Num() == 0);
 
 	const double EndTime = FPlatformTime::Seconds();
@@ -2012,6 +2024,7 @@ void FShaderCompilingManager::ProcessAsyncResults(bool bLimitExecutionTime, bool
 			}
 
 			int32 NumCompilingShaderMaps = 0;
+
 			{
 				// Lock CompileQueueSection so we can access the input and output queues
 				FScopeLock Lock(&CompileQueueSection);
@@ -2029,6 +2042,11 @@ void FShaderCompilingManager::ProcessAsyncResults(bool bLimitExecutionTime, bool
 
 					if (GetNumTotalJobs(Results.FinishedJobs) == Results.NumJobsQueued)
 					{
+						if (Results.bRecreateComponentRenderStateOnCompletion)
+						{
+							bRecreateComponentRenderStateOutstanding = true;
+						}
+
 						PendingFinalizeShaderMaps.Add(It.Key(), FShaderMapFinalizeResults(Results));
 						ShaderMapsToRemove.Add(It.Key());
 					}
@@ -2054,7 +2072,9 @@ void FShaderCompilingManager::ProcessAsyncResults(bool bLimitExecutionTime, bool
 				while (bRetry);
 
 				const float TimeBudget = bLimitExecutionTime ? ProcessGameThreadTargetTime : FLT_MAX;
-				ProcessCompiledShaderMaps(PendingFinalizeShaderMaps, TimeBudget);
+				// Only do the recreate when we are finished all compiling, to limit the number of hitchy global render state recreates
+				const bool bRecreateComponentRenderState = bRecreateComponentRenderStateOutstanding && NumCompilingShaderMaps == 0;
+				ProcessCompiledShaderMaps(PendingFinalizeShaderMaps, TimeBudget, bRecreateComponentRenderState);
 				check(bLimitExecutionTime || PendingFinalizeShaderMaps.Num() == 0);
 			}
 
@@ -2315,6 +2335,11 @@ void GlobalBeginCompileShader(
 		Input.Environment.SetDefine(TEXT("USE_DBUFFER"), CVar ? CVar->GetInt() : 0);
 	}
 
+	{
+		static IConsoleVariable* CVar = IConsoleManager::Get().FindConsoleVariable(TEXT("r.AllowGlobalClipPlane"));
+		Input.Environment.SetDefine(TEXT("PROJECT_ALLOW_GLOBAL_CLIP_PLANE"), CVar ? (CVar->GetInt() != 0) : 0);
+	}
+
 	NewJobs.Add(NewJob);
 }
 
diff --git a/Engine/Source/Runtime/Engine/Private/SkeletalMesh.cpp b/Engine/Source/Runtime/Engine/Private/SkeletalMesh.cpp
--- a/Engine/Source/Runtime/Engine/Private/SkeletalMesh.cpp
+++ b/Engine/Source/Runtime/Engine/Private/SkeletalMesh.cpp
@@ -1137,6 +1137,31 @@ FArchive& operator<<(FArchive& Ar,FSkelMeshChunk& C)
 	FSkelMeshSection
 -----------------------------------------------------------------------------*/
 
+// Custom serialization version for RecomputeTangent
+struct FRecomputeTangentCustomVersion
+{
+	enum Type
+	{
+		// Before any version changes were made in the plugin
+		BeforeCustomVersionWasAdded = 0,
+		// We serialize the RecomputeTangent Option
+		RuntimeRecomputeTangent = 1,
+		// -----<new versions can be added above this line>-------------------------------------------------
+		VersionPlusOne,
+		LatestVersion = VersionPlusOne - 1
+	};
+
+	// The GUID for this custom version number
+	const static FGuid GUID;
+
+private:
+	FRecomputeTangentCustomVersion() {}
+};
+
+const FGuid FRecomputeTangentCustomVersion::GUID(0x5579F886, 0x933A4C1F, 0x83BA087B, 0x6361B92F);
+// Register the custom version with core
+FCustomVersionRegistration GRegisterRecomputeTangentCustomVersion(FRecomputeTangentCustomVersion::GUID, FRecomputeTangentCustomVersion::LatestVersion, TEXT("RecomputeTangentCustomVer"));
+
 // Serialization.
 FArchive& operator<<(FArchive& Ar,FSkelMeshSection& S)
 {		
@@ -1171,6 +1196,12 @@ FArchive& operator<<(FArchive& Ar,FSkelMeshSection& S)
 		Ar << S.bEnableClothLOD_DEPRECATED;
 	}
 
+	Ar.UsingCustomVersion(FRecomputeTangentCustomVersion::GUID);
+	if (Ar.CustomVer(FRecomputeTangentCustomVersion::GUID) >= FRecomputeTangentCustomVersion::RuntimeRecomputeTangent)
+	{
+		Ar << S.bRecomputeTangent;
+	}
+
 	return Ar;
 }
 
@@ -1582,18 +1613,14 @@ void FStaticLODModel::ReleaseCPUResources()
 	{
 		if(MultiSizeIndexContainer.IsIndexBufferValid())
 		{
-			int32 MemorySize = MultiSizeIndexContainer.GetIndexBuffer()->Num() * MultiSizeIndexContainer.GetDataTypeSize();
-			DEC_DWORD_STAT_BY( STAT_SkeletalMeshIndexMemory, MemorySize);
 			MultiSizeIndexContainer.GetIndexBuffer()->Empty();
 		}
 		if(AdjacencyMultiSizeIndexContainer.IsIndexBufferValid())
 		{
-			DEC_DWORD_STAT_BY( STAT_SkeletalMeshIndexMemory, AdjacencyMultiSizeIndexContainer.GetIndexBuffer()->Num() * AdjacencyMultiSizeIndexContainer.GetDataTypeSize());
 			AdjacencyMultiSizeIndexContainer.GetIndexBuffer()->Empty();
 		}
 		if(VertexBufferGPUSkin.IsVertexDataValid())
 		{
-			DEC_DWORD_STAT_BY( STAT_SkeletalMeshVertexMemory, VertexBufferGPUSkin.GetVertexDataSize() );
 			VertexBufferGPUSkin.CleanUp();
 		}
 	}
@@ -3382,6 +3409,12 @@ FArchive& operator<<( FArchive& Ar, FSkeletalMaterial& Elem )
 		Ar << Elem.bEnableShadowCasting;
 	}
 
+	Ar.UsingCustomVersion(FRecomputeTangentCustomVersion::GUID);
+	if (Ar.CustomVer(FRecomputeTangentCustomVersion::GUID) >= FRecomputeTangentCustomVersion::RuntimeRecomputeTangent)
+	{
+		Ar << Elem.bRecomputeTangent;
+	}
+
 	return Ar;
 }
 
@@ -4769,14 +4802,28 @@ void FSkeletalMeshSceneProxy::GetDynamicElementsSection(const TArray<const FScen
 			Mesh.LCI = NULL;
 			Mesh.bWireframe |= bForceWireframe;
 			Mesh.Type = PT_TriangleList;
-			Mesh.VertexFactory = MeshObject->GetVertexFactory(LODIndex,Section.ChunkIndex);
+			Mesh.VertexFactory = MeshObject->GetSkinVertexFactory(View, LODIndex, Section.ChunkIndex);
+			
+			if(!Mesh.VertexFactory)
+			{
+				// hide this part
+				continue;
+			}
+
 			Mesh.bSelectable = bSelectable;
 			BatchElement.FirstIndex = Section.BaseIndex;
 
 			BatchElement.IndexBuffer = LODModel.MultiSizeIndexContainer.GetIndexBuffer();
 			BatchElement.MaxVertexIndex = LODModel.NumVertices - 1;
-
-			BatchElement.UserIndex = MeshObject->GPUSkinCacheKeys[Section.ChunkIndex];
+	
+			if(Section.ChunkIndex < FSkeletalMeshObject::MAX_GPUSKINCACHE_CHUNKS_PER_LOD)
+			{
+				BatchElement.UserIndex = MeshObject->GPUSkinCacheKeys[Section.ChunkIndex];
+			}
+			else
+			{
+				BatchElement.UserIndex = -1;
+			}
 
 			const bool bRequiresAdjacencyInformation = RequiresAdjacencyInformation( SectionElementInfo.Material, Mesh.VertexFactory->GetType(), ViewFamily.GetFeatureLevel() );
 			if ( bRequiresAdjacencyInformation )
diff --git a/Engine/Source/Runtime/Engine/Private/SkeletalRenderCPUSkin.cpp b/Engine/Source/Runtime/Engine/Private/SkeletalRenderCPUSkin.cpp
--- a/Engine/Source/Runtime/Engine/Private/SkeletalRenderCPUSkin.cpp
+++ b/Engine/Source/Runtime/Engine/Private/SkeletalRenderCPUSkin.cpp
@@ -164,21 +164,49 @@ void FSkeletalMeshObjectCPUSkin::EnableBlendWeightRendering(bool bEnabled, const
 	BonesOfInterest.Append(InBonesOfInterest);
 }
 
+void FSkeletalMeshObjectCPUSkin::UpdateRecomputeTangent(int32 MaterialIndex, bool bRecomputeTangent)
+{
+	// queue a call to update this data
+	ENQUEUE_UNIQUE_RENDER_COMMAND_THREEPARAMETER(
+		SkelMeshObjectUpdateMaterialDataCommand,
+		FSkeletalMeshObjectCPUSkin*, MeshObject, this,
+		int32, MaterialIndex, MaterialIndex,
+		bool, bRecomputeTangent, bRecomputeTangent,
+		{
+			// iterate through section and find the section that matches MaterialIndex, if so, set that flag
+			for (auto& LODModel : MeshObject->SkeletalMeshResource->LODModels)
+			{
+				for (int32 SectionIndex = 0; SectionIndex < LODModel.Sections.Num(); ++SectionIndex)
+				{
+					// @todo there can be more than one section that can use same material? If not, please break. 
+					if (LODModel.Sections[SectionIndex].MaterialIndex == MaterialIndex)
+					{
+						LODModel.Sections[SectionIndex].bRecomputeTangent = bRecomputeTangent;
+					}
+				}
+			}
+		}
+	);
+}
+
 void FSkeletalMeshObjectCPUSkin::Update(int32 LODIndex,USkinnedMeshComponent* InMeshComponent,const TArray<FActiveVertexAnim>& ActiveVertexAnims)
 {
 	// create the new dynamic data for use by the rendering thread
 	// this data is only deleted when another update is sent
 	FDynamicSkelMeshObjectDataCPUSkin* NewDynamicData = new FDynamicSkelMeshObjectDataCPUSkin(InMeshComponent,SkeletalMeshResource,LODIndex,ActiveVertexAnims);
 
+	// We prepare the next frame but still have the value from the last one
+	uint32 FrameNumberToPrepare = GFrameNumber + 1;
+
 	// queue a call to update this data
 	ENQUEUE_UNIQUE_RENDER_COMMAND_THREEPARAMETER(
 		SkelMeshObjectUpdateDataCommand,
 		FSkeletalMeshObjectCPUSkin*, MeshObject, this,
-		uint32, FrameNumber, GFrameNumber,
+		uint32, FrameNumberToPrepare, FrameNumberToPrepare,
 		FDynamicSkelMeshObjectDataCPUSkin*, NewDynamicData, NewDynamicData,
 	{
 		FScopeCycleCounter Context(MeshObject->GetStatId());
-		MeshObject->UpdateDynamicData_RenderThread(RHICmdList, NewDynamicData, FrameNumber);
+		MeshObject->UpdateDynamicData_RenderThread(RHICmdList, NewDynamicData, FrameNumberToPrepare);
 	}
 	);
 
@@ -190,7 +218,7 @@ void FSkeletalMeshObjectCPUSkin::Update(int32 LODIndex,USkinnedMeshComponent* In
 	}
 }
 
-void FSkeletalMeshObjectCPUSkin::UpdateDynamicData_RenderThread(FRHICommandListImmediate& RHICmdList, FDynamicSkelMeshObjectDataCPUSkin* InDynamicData, uint32 FrameNumber)
+void FSkeletalMeshObjectCPUSkin::UpdateDynamicData_RenderThread(FRHICommandListImmediate& RHICmdList, FDynamicSkelMeshObjectDataCPUSkin* InDynamicData, uint32 FrameNumberToPrepare)
 {
 	// we should be done with the old data at this point
 	delete DynamicData;
@@ -273,7 +301,7 @@ void FSkeletalMeshObjectCPUSkin::CacheVertices(int32 LODIndex, bool bForce) cons
 	}
 }
 
-const FVertexFactory* FSkeletalMeshObjectCPUSkin::GetVertexFactory(int32 LODIndex,int32 /*ChunkIdx*/) const
+const FVertexFactory* FSkeletalMeshObjectCPUSkin::GetSkinVertexFactory(const FSceneView* View, int32 LODIndex,int32 /*ChunkIdx*/) const
 {
 	check( LODs.IsValidIndex(LODIndex) );
 	return &LODs[LODIndex].VertexFactory;
diff --git a/Engine/Source/Runtime/Engine/Private/SkeletalRenderCPUSkin.h b/Engine/Source/Runtime/Engine/Private/SkeletalRenderCPUSkin.h
--- a/Engine/Source/Runtime/Engine/Private/SkeletalRenderCPUSkin.h
+++ b/Engine/Source/Runtime/Engine/Private/SkeletalRenderCPUSkin.h
@@ -146,11 +146,12 @@ public:
 	virtual void InitResources() override;
 	virtual void ReleaseResources() override;
 	virtual void Update(int32 LODIndex,USkinnedMeshComponent* InMeshComponent,const TArray<FActiveVertexAnim>& ActiveVertexAnims) override;
-	void UpdateDynamicData_RenderThread(FRHICommandListImmediate& RHICmdList, FDynamicSkelMeshObjectDataCPUSkin* InDynamicData, uint32 FrameNumber);
+	void UpdateDynamicData_RenderThread(FRHICommandListImmediate& RHICmdList, FDynamicSkelMeshObjectDataCPUSkin* InDynamicData, uint32 FrameNumberToPrepare);
+	virtual void UpdateRecomputeTangent(int32 MaterialIndex, bool bRecomputeTangent) override;
 	virtual void EnableBlendWeightRendering(bool bEnabled, const TArray<int32>& InBonesOfInterest) override;
 	virtual void CacheVertices(int32 LODIndex, bool bForce) const override;
 	virtual bool IsCPUSkinned() const override { return true; }
-	virtual const FVertexFactory* GetVertexFactory(int32 LODIndex,int32 ChunkIdx) const override;
+	virtual const FVertexFactory* GetSkinVertexFactory(const FSceneView* View, int32 LODIndex, int32 ChunkIdx) const override;
 	virtual TArray<FTransform>* GetSpaceBases() const override;
 	virtual const TArray<FMatrix>& GetReferenceToLocalMatrices() const override;
 	virtual int32 GetLOD() const override
diff --git a/Engine/Source/Runtime/Engine/Private/SkeletalRenderGPUSkin.cpp b/Engine/Source/Runtime/Engine/Private/SkeletalRenderGPUSkin.cpp
--- a/Engine/Source/Runtime/Engine/Private/SkeletalRenderGPUSkin.cpp
+++ b/Engine/Source/Runtime/Engine/Private/SkeletalRenderGPUSkin.cpp
@@ -217,15 +217,18 @@ void FSkeletalMeshObjectGPUSkin::Update(int32 LODIndex,USkinnedMeshComponent* In
 		}
 	}
 
+	// We prepare the next frame but still have the value from the last one
+	uint32 FrameNumberToPrepare = GFrameNumber + 1;
+
 	// queue a call to update this data
 	ENQUEUE_UNIQUE_RENDER_COMMAND_THREEPARAMETER(
 		SkelMeshObjectUpdateDataCommand,
 		FSkeletalMeshObjectGPUSkin*, MeshObject, this,
-		uint32, FrameNumber, GFrameNumber, 
+		uint32, FrameNumberToPrepare, FrameNumberToPrepare, 
 		FDynamicSkelMeshObjectDataGPUSkin*, NewDynamicData, NewDynamicData,
 	{
 		FScopeCycleCounter Context(MeshObject->GetStatId());
-		MeshObject->UpdateDynamicData_RenderThread(RHICmdList, NewDynamicData, FrameNumber);
+		MeshObject->UpdateDynamicData_RenderThread(RHICmdList, NewDynamicData, FrameNumberToPrepare);
 	}
 	);
 
@@ -237,12 +240,37 @@ void FSkeletalMeshObjectGPUSkin::Update(int32 LODIndex,USkinnedMeshComponent* In
 	}
 }
 
+void FSkeletalMeshObjectGPUSkin::UpdateRecomputeTangent(int32 MaterialIndex, bool bRecomputeTangent)
+{
+	// queue a call to update this data
+	ENQUEUE_UNIQUE_RENDER_COMMAND_THREEPARAMETER(
+		SkelMeshObjectUpdateMaterialDataCommand,
+		FSkeletalMeshObjectGPUSkin*, MeshObject, this,
+		int32, MaterialIndex, MaterialIndex, 
+		bool, bRecomputeTangent, bRecomputeTangent,
+		{
+			// iterate through section and find the section that matches MaterialIndex, if so, set that flag
+			for (auto& LODModel : MeshObject->SkeletalMeshResource->LODModels)
+			{
+				for (int32 SectionIndex = 0; SectionIndex < LODModel.Sections.Num(); ++SectionIndex)
+				{
+					// @todo there can be more than one section that can use same material? If not, please break. 
+					if (LODModel.Sections[SectionIndex].MaterialIndex == MaterialIndex)
+					{
+						LODModel.Sections[SectionIndex].bRecomputeTangent = bRecomputeTangent;
+					}
+				}
+			}
+		}
+	);
+}
+
 static TAutoConsoleVariable<int32> CVarDeferSkeletalDynamicDataUpdateUntilGDME(
 	TEXT("r.DeferSkeletalDynamicDataUpdateUntilGDME"),
 	0,
 	TEXT("If > 0, then do skeletal mesh dynamic data updates will be deferred until GDME. Experimental option."));
 
-void FSkeletalMeshObjectGPUSkin::UpdateDynamicData_RenderThread(FRHICommandListImmediate& RHICmdList, FDynamicSkelMeshObjectDataGPUSkin* InDynamicData, uint32 FrameNumber)
+void FSkeletalMeshObjectGPUSkin::UpdateDynamicData_RenderThread(FRHICommandListImmediate& RHICmdList, FDynamicSkelMeshObjectDataGPUSkin* InDynamicData, uint32 FrameNumberToPrepare)
 {
 	SCOPE_CYCLE_COUNTER(STAT_GPUSkinUpdateRTTime);
 	check(InDynamicData);
@@ -269,7 +297,7 @@ void FSkeletalMeshObjectGPUSkin::UpdateDynamicData_RenderThread(FRHICommandListI
 	}
 	else
 	{
-		ProcessUpdatedDynamicData(RHICmdList, FrameNumber, bMorphNeedsUpdate);
+		ProcessUpdatedDynamicData(RHICmdList, FrameNumberToPrepare, bMorphNeedsUpdate);
 	}
 }
 
@@ -291,14 +319,13 @@ void FSkeletalMeshObjectGPUSkin::WaitForRHIThreadFenceForDynamicData()
 	}
 }
 
-void FSkeletalMeshObjectGPUSkin::ProcessUpdatedDynamicData(FRHICommandListImmediate& RHICmdList, uint32 FrameNumber, bool bMorphNeedsUpdate)
+void FSkeletalMeshObjectGPUSkin::ProcessUpdatedDynamicData(FRHICommandListImmediate& RHICmdList, uint32 FrameNumberToPrepare, bool bMorphNeedsUpdate)
 {
 	QUICK_SCOPE_CYCLE_COUNTER(STAT_FSkeletalMeshObjectGPUSkin_ProcessUpdatedDynamicData);
 	bNeedsUpdateDeferred = false;
 	bMorphNeedsUpdateDeferred = false;
 
 	FSkeletalMeshObjectLOD& LOD = LODs[DynamicData->LODIndex];
-	const FSkelMeshObjectLODInfo& MeshLODInfo = LODInfo[DynamicData->LODIndex];
 
 	// if hasn't been updated, force update again
 	bMorphNeedsUpdate = LOD.MorphVertexBuffer.bHasBeenUpdated ? bMorphNeedsUpdate : true;
@@ -340,9 +367,26 @@ void FSkeletalMeshObjectGPUSkin::ProcessUpdatedDynamicData(FRHICommandListImmedi
 
 			bool bClothFactory = (FeatureLevel >= ERHIFeatureLevel::SM4) && (DynamicData->ClothSimulUpdateData.Num() > 0) && Chunk.HasApexClothData();
 
-			FGPUBaseSkinVertexFactory::FShaderDataType& ShaderData = 
-				bClothFactory ? VertexFactoryData.ClothVertexFactories[ChunkIdx]->GetVertexFactory()->GetShaderData() :
-				(DynamicData->NumWeightedActiveVertexAnims > 0 ? VertexFactoryData.MorphVertexFactories[ChunkIdx]->GetShaderData() : VertexFactoryData.VertexFactories[ChunkIdx]->GetShaderData());
+			FGPUBaseSkinVertexFactory* VertexFactory;
+			{
+				if(bClothFactory)
+				{
+					VertexFactory = VertexFactoryData.ClothVertexFactories[ChunkIdx]->GetVertexFactory();
+				}
+				else
+				{
+					if(DynamicData->NumWeightedActiveVertexAnims > 0)
+					{
+						VertexFactory = VertexFactoryData.MorphVertexFactories[ChunkIdx].Get();
+					}
+					else
+					{
+						VertexFactory = VertexFactoryData.VertexFactories[ChunkIdx].Get();
+					}
+				}
+			}
+
+			FGPUBaseSkinVertexFactory::FShaderDataType& ShaderData = VertexFactory->GetShaderData();
 
 			bool bUseSkinCache = bGPUSkinCacheEnabled;
 			if (bUseSkinCache)
@@ -362,19 +406,35 @@ void FSkeletalMeshObjectGPUSkin::ProcessUpdatedDynamicData(FRHICommandListImmedi
 					INC_DWORD_STAT(STAT_GPUSkinCache_SkippedForZeroInfluences);
 					bUseSkinCache = false;
 				}
+
+				{
+					extern TAutoConsoleVariable<int32> CVarGPUSkinCacheRecomputeTangents;
+					if(CVarGPUSkinCacheRecomputeTangents.GetValueOnRenderThread() == 2)
+					{
+						const FSkelMeshSection* Section = LODModel.FindSectionForChunk(ChunkIdx);
+
+						if(!Section->bRecomputeTangent)
+						{
+							// In this mode the SkinCache should only be used for RecomputeTangent meshes
+							bUseSkinCache = false;
+						}
+					}
+				}
 			}
 
 			// Create a uniform buffer from the bone transforms.
 			TArray<FMatrix>& ReferenceToLocalMatrices = DynamicData->ReferenceToLocal;
-			bool bNeedFence = ShaderData.UpdateBoneData(RHICmdList, ReferenceToLocalMatrices, Chunk.BoneMap, FrameNumber, FeatureLevel, bUseSkinCache);
+			bool bNeedFence = ShaderData.UpdateBoneData(RHICmdList, ReferenceToLocalMatrices, Chunk.BoneMap, FrameNumberToPrepare, FeatureLevel, bUseSkinCache);
 
 			// Try to use the GPU skinning cache if possible
 			if (bUseSkinCache)
 			{
-				int32 Key = GGPUSkinCache.StartCacheMesh(RHICmdList, GPUSkinCacheKeys[ChunkIdx], VertexFactoryData.VertexFactories[ChunkIdx].Get(),
+				int32 Key = GGPUSkinCache.StartCacheMesh(RHICmdList, FrameNumberToPrepare, GPUSkinCacheKeys[ChunkIdx], VertexFactory,
 					VertexFactoryData.PassthroughVertexFactories[ChunkIdx].Get(), Chunk, this, bMorph ? &LOD.MorphVertexBuffer : 0);
 
-				// if not failed(Key == -1)
+				// TODO: If failed we could remove the key (on the other hand it might a bu in cache for longer)
+
+				// -1 if failed
 				if(Key >= 0)
 				{
 					ensure(Key <= 0xffff);
@@ -391,7 +451,7 @@ void FSkeletalMeshObjectGPUSkin::ProcessUpdatedDynamicData(FRHICommandListImmedi
 				int16 ActorIdx = Chunk.CorrespondClothAssetIndex;
 				if( DynamicData->ClothSimulUpdateData.IsValidIndex(ActorIdx) )
 				{
-					bNeedFence = ClothShaderData.UpdateClothSimulData(RHICmdList, DynamicData->ClothSimulUpdateData[ActorIdx].ClothSimulPositions, DynamicData->ClothSimulUpdateData[ActorIdx].ClothSimulNormals, FrameNumber, FeatureLevel) || bNeedFence;
+					bNeedFence = ClothShaderData.UpdateClothSimulData(RHICmdList, DynamicData->ClothSimulUpdateData[ActorIdx].ClothSimulPositions, DynamicData->ClothSimulUpdateData[ActorIdx].ClothSimulNormals, FrameNumberToPrepare, FeatureLevel) || bNeedFence;
 				}
 			}
 #endif // WITH_APEX_CLOTHING
@@ -412,8 +472,8 @@ void FSkeletalMeshObjectGPUSkin::FSkeletalMeshObjectLOD::UpdateMorphVertexBuffer
 
 	if( IsValidRef(MorphVertexBuffer.VertexBufferRHI) )
 	{
-		extern ENGINE_API bool DoRecomputeSkinTangentsOnGPU();
-		bool bBlendTangentsOnCPU = !DoRecomputeSkinTangentsOnGPU();
+		extern ENGINE_API bool DoRecomputeSkinTangentsOnGPU_RT();
+		bool bBlendTangentsOnCPU = !DoRecomputeSkinTangentsOnGPU_RT();
 
 /*		// TODO: Need to finish this to avoid artifacts if the SkinCache is not handing all objects
 		if(!bBlendTangentsOnCPU)
@@ -538,7 +598,7 @@ void FSkeletalMeshObjectGPUSkin::FSkeletalMeshObjectLOD::UpdateMorphVertexBuffer
 	}
 }
 
-const FVertexFactory* FSkeletalMeshObjectGPUSkin::GetVertexFactory(int32 LODIndex,int32 ChunkIdx) const
+const FVertexFactory* FSkeletalMeshObjectGPUSkin::GetSkinVertexFactory(const FSceneView* View, int32 LODIndex,int32 ChunkIdx) const
 {
 	checkSlow( LODs.IsValidIndex(LODIndex) );
 	checkSlow( DynamicData );
@@ -554,20 +614,28 @@ const FVertexFactory* FSkeletalMeshObjectGPUSkin::GetVertexFactory(int32 LODInde
 		return LOD.GPUSkinVertexFactories.ClothVertexFactories[ChunkIdx]->GetVertexFactory();
 	}
 
-		// If the GPU skinning cache was used, return the passthrough vertex factory
-		if (GGPUSkinCache.IsElementProcessed(GPUSkinCacheKeys[ChunkIdx]))
+	// If the GPU skinning cache was used, return the passthrough vertex factory
+	if (GGPUSkinCache.IsElementProcessed(View->Family->FrameNumber, GPUSkinCacheKeys[ChunkIdx]))
+	{
+		if(View->Family->EngineShowFlags.SkinCache)
 		{
 			return LOD.GPUSkinVertexFactories.PassthroughVertexFactories[ChunkIdx].Get();
 		}
+		else
+		{
+			// hide this mesh
+			return 0;
+		}
+	}
 
 	// use the morph enabled vertex factory if any active morphs are set
 	if( DynamicData->NumWeightedActiveVertexAnims > 0 )
-		{
+	{
 		return LOD.GPUSkinVertexFactories.MorphVertexFactories[ChunkIdx].Get();
 	}
 
 	// use the default gpu skin vertex factory
-			return LOD.GPUSkinVertexFactories.VertexFactories[ChunkIdx].Get();
+	return LOD.GPUSkinVertexFactories.VertexFactories[ChunkIdx].Get();
 }
 
 /** 
diff --git a/Engine/Source/Runtime/Engine/Private/SkeletalRenderGPUSkin.h b/Engine/Source/Runtime/Engine/Private/SkeletalRenderGPUSkin.h
--- a/Engine/Source/Runtime/Engine/Private/SkeletalRenderGPUSkin.h
+++ b/Engine/Source/Runtime/Engine/Private/SkeletalRenderGPUSkin.h
@@ -204,9 +204,10 @@ public:
 	virtual void InitResources() override;
 	virtual void ReleaseResources() override;
 	virtual void Update(int32 LODIndex,USkinnedMeshComponent* InMeshComponent,const TArray<FActiveVertexAnim>& ActiveVertexAnims) override;
-	void UpdateDynamicData_RenderThread(FRHICommandListImmediate& RHICmdList, FDynamicSkelMeshObjectDataGPUSkin* InDynamicData, uint32 FrameNumber);
+	void UpdateDynamicData_RenderThread(FRHICommandListImmediate& RHICmdList, FDynamicSkelMeshObjectDataGPUSkin* InDynamicData, uint32 FrameNumberToPrepare);
+	virtual void UpdateRecomputeTangent(int32 MaterialIndex, bool bRecomputeTangent) override;
 	virtual void PreGDMECallback(uint32 FrameNumber) override;
-	virtual const FVertexFactory* GetVertexFactory(int32 LODIndex,int32 ChunkIdx) const override;
+	virtual const FVertexFactory* GetSkinVertexFactory(const FSceneView* View, int32 LODIndex,int32 ChunkIdx) const override;
 	virtual void CacheVertices(int32 LODIndex, bool bForce) const override {}
 	virtual bool IsCPUSkinned() const override { return false; }
 	virtual TArray<FTransform>* GetSpaceBases() const override;
@@ -441,8 +442,7 @@ private:
 	*/
 	void ReleaseMorphResources();
 
-	// @param FrameNumber from GFrameNumber
-	void ProcessUpdatedDynamicData(FRHICommandListImmediate& RHICmdList, uint32 FrameNumber, bool bMorphNeedsUpdate);
+	void ProcessUpdatedDynamicData(FRHICommandListImmediate& RHICmdList, uint32 FrameNumberToPrepare, bool bMorphNeedsUpdate);
 
 	void WaitForRHIThreadFenceForDynamicData();
 
diff --git a/Engine/Source/Runtime/Engine/Private/StaticMeshRender.cpp b/Engine/Source/Runtime/Engine/Private/StaticMeshRender.cpp
--- a/Engine/Source/Runtime/Engine/Private/StaticMeshRender.cpp
+++ b/Engine/Source/Runtime/Engine/Private/StaticMeshRender.cpp
@@ -97,11 +97,10 @@ FStaticMeshSceneProxy::FStaticMeshSceneProxy(UStaticMeshComponent* InComponent):
 	, MaterialRelevance(InComponent->GetMaterialRelevance(GetScene().GetFeatureLevel()))
 	, CollisionResponse(InComponent->GetCollisionResponseToChannels())
 #if WITH_EDITORONLY_DATA
+	, StreamingSectionData(InComponent->StreamingSectionData)
 	, SectionIndexPreview(InComponent->SectionIndexPreview)
 #endif
 #if !(UE_BUILD_SHIPPING || UE_BUILD_TEST)
-	, WorldTexelFactor(0)
-	, WorldLightmapFactor(0)
 	, LODForCollision(InComponent->StaticMesh->LODForCollision)
 	, bDrawMeshCollisionWireframe(InComponent->bDrawMeshCollisionWireframe)
 #endif
@@ -189,22 +188,6 @@ FStaticMeshSceneProxy::FStaticMeshSceneProxy(UStaticMeshComponent* InComponent):
 	}
 #endif
 
-#if !(UE_BUILD_SHIPPING || UE_BUILD_TEST)
-	if (AllowDebugViewmodes())
-	{
-		// Update texture streaming factors
-		StreamingTextureInfos = InComponent->StreamingTextureInfos;
-
-		FBoxSphereBounds DummyBounds;
-		InComponent->GetStreamingTextureFactors(WorldTexelFactor, WorldLightmapFactor, DummyBounds, INDEX_NONE, INDEX_NONE);
-	}
-	else
-	{
-		WorldTexelFactor = 0;
-		WorldLightmapFactor = 0;
-	}
-#endif
-
 	if (BodySetup)
 	{
 		CollisionTraceFlag = BodySetup->GetCollisionTraceFlag();
@@ -419,31 +402,20 @@ bool FStaticMeshSceneProxy::GetWireframeMeshElement(int32 LODIndex, int32 BatchI
 	return OutBatchElement.NumPrimitives > 0;
 }
 
-#if !(UE_BUILD_SHIPPING || UE_BUILD_TEST)
-bool FStaticMeshSceneProxy::GetStreamingTextureInfo(struct FStreamingTexturePrimitiveInfo& OutInfo, int32 LODIndex, int32 ElementIndex) const
+#if WITH_EDITORONLY_DATA
+const FStreamingSectionBuildInfo* FStaticMeshSceneProxy::GetStreamingSectionData(int32 LODIndex, int32 ElementIndex) const
 {
-	int32 DebugIndex = ElementIndex;
-	if (StreamingTextureInfos.Num() && RenderData->LODResources.IsValidIndex(LODIndex) && RenderData->LODResources[LODIndex].Sections.IsValidIndex(ElementIndex))
+	if (StreamingSectionData.IsValid())
 	{
-		int32 DebugLODIndex = 0;
-		while (DebugLODIndex < LODIndex)
+		for (const FStreamingSectionBuildInfo& SectionData : *StreamingSectionData)
 		{
-			DebugIndex += RenderData->LODResources[DebugLODIndex].Sections.Num();
-			++DebugLODIndex;
+			if (SectionData.LODIndex == LODIndex && SectionData.ElementIndex == ElementIndex)
+			{
+				return &SectionData;
+			}
 		}
 	}
-
-	if (StreamingTextureInfos.IsValidIndex(DebugIndex))
-	{
-		OutInfo = StreamingTextureInfos[DebugIndex];
-		return true;
-	}
-	else
-	{
-		OutInfo.Bounds = GetBounds();
-		OutInfo.TexelFactor = WorldTexelFactor;
-		return WorldTexelFactor != 0;
-	}
+	return nullptr;
 }
 #endif
 
@@ -1438,14 +1410,8 @@ FLODMask FStaticMeshSceneProxy::GetLODMask(const FSceneView* View) const
 	FLODMask Result;
 	int32 CVarForcedLODLevel = GetCVarForceLOD();
 
-	// Planar reflections always use the highest available LOD level
-	if (View->bIsPlanarReflectionCapture && RenderData->LODResources.Num() > 0)
-	{
-		Result.SetLOD(RenderData->LODResources.Num() - 1);
-	}
-
 	//If a LOD is being forced, use that one
-	else if (CVarForcedLODLevel >= 0)
+	if (CVarForcedLODLevel >= 0)
 	{
 		Result.SetLOD(FMath::Clamp<int32>(CVarForcedLODLevel, 0, RenderData->LODResources.Num() - 1));
 	}
diff --git a/Engine/Source/Runtime/Engine/Private/Streaming/AsyncTextureStreaming.cpp b/Engine/Source/Runtime/Engine/Private/Streaming/AsyncTextureStreaming.cpp
new file mode 100644
--- /dev/null
+++ b/Engine/Source/Runtime/Engine/Private/Streaming/AsyncTextureStreaming.cpp
@@ -0,0 +1,196 @@
+// Copyright 1998-2016 Epic Games, Inc. All Rights Reserved.
+
+/*=============================================================================
+AsyncTextureStreaming.cpp: Definitions of classes used for texture streaming async task.
+=============================================================================*/
+
+#include "EnginePrivate.h"
+#include "AsyncTextureStreaming.h"
+
+void FAsyncTextureStreaming::ClearRemovedTextureReferences()
+{
+	// CleanUp Prioritized Textures, account for removed textures.
+	for (int32 TexPrioIndex = 0; TexPrioIndex < PrioritizedTextures.Num(); ++TexPrioIndex)
+	{
+		FTexturePriority& TexturePriority = PrioritizedTextures[ TexPrioIndex ];
+
+		// Because deleted textures can shrink StreamingTextures, we need to check that the index is in range.
+		// Even if it is, because of the RemoveSwap logic, it could refer to the wrong texture now.
+		if (!StreamingManager.StreamingTextures.IsValidIndex(TexturePriority.TextureIndex) || TexturePriority.Texture != StreamingManager.StreamingTextures[TexturePriority.TextureIndex].Texture)
+		{
+			TexturePriority.TextureIndex = INDEX_NONE;
+			TexturePriority.Texture = nullptr;
+		}
+	}
+}
+
+void FAsyncTextureStreaming::DoWork()
+{
+	DECLARE_SCOPE_CYCLE_COUNTER(TEXT("FAsyncTextureStreaming::DoWork"), STAT_AsyncTextureStreaming_DoWork, STATGROUP_StreamingDetails);
+	PrioritizedTextures.Empty( StreamingManager.StreamingTextures.Num() );
+	PrioTexIndicesSortedByLoadOrder.Empty( StreamingManager.StreamingTextures.Num() );
+
+	const int32 SplitRequestSizeThreshold = CVarStreamingSplitRequestSizeThreshold.GetValueOnAnyThread();
+
+	const bool bUseNewMetrics = CVarStreamingUseNewMetrics.GetValueOnAnyThread() > 0;
+	const float MaxEffectiveScreenSize = CVarStreamingScreenSizeEffectiveMax.GetValueOnAnyThread();
+	StreamingManager.ThreadSettings.TextureBoundsVisibility->ComputeBoundsViewInfos(StreamingManager.ThreadSettings.ThreadViewInfos, bUseNewMetrics, MaxEffectiveScreenSize);
+
+	// Number of textures that want more mips.
+	ThreadStats.NumWantingTextures = 0;
+	ThreadStats.NumVisibleTexturesWithLowResolutions = 0;
+
+	for ( int32 Index=0; Index < StreamingManager.StreamingTextures.Num() && !IsAborted(); ++Index )
+	{
+		FStreamingTexture& StreamingTexture = StreamingManager.StreamingTextures[ Index ];
+
+		StreamingTexture.bUsesStaticHeuristics = false;
+		StreamingTexture.bUsesDynamicHeuristics = false;
+		StreamingTexture.bUsesLastRenderHeuristics = false;
+		StreamingTexture.bUsesForcedHeuristics = false;
+		StreamingTexture.bUsesOrphanedHeuristics = false;
+		StreamingTexture.bHasSplitRequest  = false;
+		StreamingTexture.bIsLastSplitRequest = false;
+
+			// Figure out max number of miplevels allowed by LOD code.
+			StreamingManager.CalcMinMaxMips( StreamingTexture );
+
+			// Determine how many mips this texture should have in memory.
+			StreamingManager.CalcWantedMips( StreamingTexture );
+
+		ThreadStats.TotalResidentSize += StreamingTexture.GetSize( StreamingTexture.ResidentMips );
+		ThreadStats.TotalPossibleResidentSize += StreamingTexture.GetSize( FMath::Max<int32>(StreamingTexture.ResidentMips, StreamingTexture.RequestedMips) );
+		ThreadStats.TotalWantedMipsSize += StreamingTexture.GetSize( StreamingTexture.WantedMips );
+		if (StreamingTexture.RequestedMips != StreamingTexture.ResidentMips)
+		{
+			// UpdateMips load texture mip changes in temporary memory.
+			ThreadStats.TotalTempMemorySize += StreamingTexture.GetSize( StreamingTexture.RequestedMips );
+		}
+
+
+		// Split the request in two if the last mip to load is too big. Helps putprioriHelps converge the image quality.
+		if (SplitRequestSizeThreshold != 0 && 
+			StreamingTexture.LODGroup != TEXTUREGROUP_Terrain_Heightmap &&
+			StreamingTexture.WantedMips > StreamingTexture.ResidentMips && 
+			StreamingTexture.GetSize(StreamingTexture.WantedMips) - StreamingTexture.GetSize(StreamingTexture.WantedMips - 1) >= SplitRequestSizeThreshold) 
+		{
+			// The texture fits for a 2 request stream strategy, though it is not garantied that there will be 2 requests (since it could be only one mip aways)
+			StreamingTexture.bHasSplitRequest = true;
+			if (StreamingTexture.WantedMips == StreamingTexture.ResidentMips + 1)
+			{
+				StreamingTexture.bIsLastSplitRequest = true;
+			}
+			else
+			{
+				// Don't stream the last mips now.
+				--StreamingTexture.WantedMips;
+			}
+		}
+		if ( StreamingTexture.WantedMips > StreamingTexture.ResidentMips )
+		{
+			ThreadStats.NumWantingTextures++;
+		}
+
+		bool bTrackedTexture = TrackTextureEvent( &StreamingTexture, StreamingTexture.Texture, StreamingTexture.bForceFullyLoad, &StreamingManager );
+
+		// Add to sort list, if it wants to stream in or could potentially stream out.
+		if ( StreamingTexture.MaxAllowedMips > StreamingTexture.NumNonStreamingMips)
+		{
+			const UTexture2D* Texture = StreamingTexture.Texture; // This needs to go on the stack in case the main thread updates it at the same time (for the condition consistency).
+			if (Texture)
+			{
+				new (PrioritizedTextures) FTexturePriority( StreamingTexture.CanDropMips(), StreamingTexture.CalcRetentionPriority(), StreamingTexture.CalcLoadOrderPriority(), Index, Texture );
+			}
+		}
+
+		if (StreamingTexture.IsVisible() && StreamingTexture.GetWantedMipsForPriority() - StreamingTexture.ResidentMips > 1)
+		{
+			++ThreadStats.NumVisibleTexturesWithLowResolutions;
+		}
+
+			// Accumulate streaming numbers.
+		int64 ResidentTextureSize = StreamingTexture.GetSize( StreamingTexture.ResidentMips );
+		int64 WantedTextureSize = StreamingTexture.GetSize( StreamingTexture.WantedMips );
+		if ( StreamingTexture.bInFlight )
+		{
+			int64 RequestedTextureSize = StreamingTexture.GetSize( StreamingTexture.RequestedMips );
+			if ( StreamingTexture.RequestedMips > StreamingTexture.ResidentMips )
+			{
+				ThreadStats.PendingStreamInSize += FMath::Abs(RequestedTextureSize - ResidentTextureSize);
+			}
+			else
+			{
+				ThreadStats.PendingStreamOutSize += FMath::Abs(RequestedTextureSize - ResidentTextureSize);
+			}
+		}
+		else
+		{
+			if ( StreamingTexture.WantedMips > StreamingTexture.ResidentMips )
+			{
+				ThreadStats.WantedInSize += FMath::Abs(WantedTextureSize - ResidentTextureSize);
+			}
+			else
+			{
+				// Counting on shrinking reallocation.
+				ThreadStats.WantedOutSize += FMath::Abs(WantedTextureSize - ResidentTextureSize);
+			}
+		}
+
+		const int32 PerfectWantedTextureSize = StreamingTexture.GetSize( StreamingTexture.PerfectWantedMips );
+		ThreadStats.TotalRequiredSize += PerfectWantedTextureSize;
+		StreamingManager.UpdateFrameStats( ThreadContext, StreamingTexture, Index );
+
+		// Reset the boost factor
+		StreamingTexture.BoostFactor = 1.0f;
+	}
+
+	// Sort the candidates.
+	struct FCompareTexturePriority // Smaller retention priority first because we drop them first.
+	{
+		FORCEINLINE bool operator()( const FTexturePriority& A, const FTexturePriority& B ) const
+		{
+			if ( A.RetentionPriority < B.RetentionPriority )
+			{
+				return true;
+			}
+			else if ( A.RetentionPriority == B.RetentionPriority )
+			{
+				return ( A.TextureIndex < B.TextureIndex );
+			}
+			return false;
+		}
+	};
+	// Texture are sorted by retention priority since the ASYNC loading will resort request by load priority
+	PrioritizedTextures.Sort( FCompareTexturePriority() );
+
+
+	// Sort the candidates.
+	struct FCompareTexturePriorityByLoadOrder // Bigger load order priority first because we load them first.
+	{
+		FCompareTexturePriorityByLoadOrder(const TArray<FTexturePriority>& InPrioritizedTextures) : PrioritizedTextures(InPrioritizedTextures) {}
+		const TArray<FTexturePriority>& PrioritizedTextures;
+
+		FORCEINLINE bool operator()( int32 IndexA, int32 IndexB ) const
+		{
+			const FTexturePriority& A = PrioritizedTextures[IndexA];
+			const FTexturePriority& B = PrioritizedTextures[IndexB];
+
+			if ( A.LoadOrderPriority > B.LoadOrderPriority )
+			{
+				return true;
+			}
+			else if ( A.LoadOrderPriority == B.LoadOrderPriority )
+			{
+				return ( A.TextureIndex < B.TextureIndex );
+			}
+			return false;
+		}
+	};
+
+	PrioTexIndicesSortedByLoadOrder.AddUninitialized(PrioritizedTextures.Num());
+	for (int32 Index = 0; Index < PrioTexIndicesSortedByLoadOrder.Num(); ++Index)
+	{
+		PrioTexIndicesSortedByLoadOrder[Index] = Index;
+	}
+	PrioTexIndicesSortedByLoadOrder.Sort(FCompareTexturePriorityByLoadOrder(PrioritizedTextures));
+}
\ No newline at end of file
diff --git a/Engine/Source/Runtime/Engine/Private/Streaming/AsyncTextureStreaming.h b/Engine/Source/Runtime/Engine/Private/Streaming/AsyncTextureStreaming.h
new file mode 100644
--- /dev/null
+++ b/Engine/Source/Runtime/Engine/Private/Streaming/AsyncTextureStreaming.h
@@ -0,0 +1,144 @@
+// Copyright 1998-2016 Epic Games, Inc. All Rights Reserved.
+
+/*=============================================================================
+AsyncTextureStreaming.h: Definitions of classes used for texture streaming async task.
+=============================================================================*/
+
+#pragma once
+
+#include "StreamingManagerTexture.h"
+
+/*-----------------------------------------------------------------------------
+	Asynchronous texture streaming.
+-----------------------------------------------------------------------------*/
+
+/** Async work class for calculating priorities for all textures. */
+// this could implement a better abandon, but give how it is used, it does that anyway via the abort mechanism
+class FAsyncTextureStreaming : public FNonAbandonableTask
+{
+public:
+	FAsyncTextureStreaming( FStreamingManagerTexture* InStreamingManager )
+	:	StreamingManager( *InStreamingManager )
+	,	ThreadContext( false, NULL, false )
+	,	bAbort( false )
+	{
+		Reset(false);
+	}
+
+	/** Resets the state to start a new async job. */
+	void Reset( bool bCollectTextureStats )
+	{
+		bAbort = false;
+		ThreadContext.Reset( false, NULL, bCollectTextureStats );
+		ThreadStats.Reset();
+		PrioritizedTextures.Empty( StreamingManager.StreamingTextures.Num() );
+		PrioTexIndicesSortedByLoadOrder.Empty( StreamingManager.StreamingTextures.Num() );
+	}
+
+	/** Notifies the async work that it should abort the thread ASAP. */
+	void Abort()
+	{
+		bAbort = true;
+	}
+
+	/** Whether the async work is being aborted. Can be used in conjunction with IsDone() to see if it has finished. */
+	bool IsAborted() const
+	{
+		return bAbort;
+	}
+
+	/** Statistics for the async work. */
+	struct FAsyncStats
+	{
+		FAsyncStats()
+		{
+			Reset();
+		}
+		/** Resets the statistics to zero. */
+		void Reset()
+		{
+			TotalResidentSize = 0;
+			TotalPossibleResidentSize = 0;
+			TotalWantedMipsSize = 0;
+			TotalTempMemorySize = 0;
+			TotalRequiredSize = 0;
+			PendingStreamInSize = 0;
+			PendingStreamOutSize = 0;
+			WantedInSize = 0;
+			WantedOutSize = 0;
+			NumWantingTextures = 0;
+			NumVisibleTexturesWithLowResolutions = 0;
+		}
+		/** Total number of bytes currently in memory */
+		int64 TotalResidentSize;
+		/** Total number of bytes that could possibly be used according to the worst case between resident size (current state) and requested size (pending requests) */
+		int64 TotalPossibleResidentSize;
+		/** Total number of bytes required to satisfy wanted mips.  Not taking into account split requests. */
+		int64 TotalWantedMipsSize;
+		/** Total number of bytes used for updating textures (conservative value since it is platform specific */
+		int64 TotalTempMemorySize;
+		/** Total number of bytes required, using PerfectWantedMips */
+		int64 TotalRequiredSize;
+		/** Currently being streamed in, in bytes. */
+		SIZE_T PendingStreamInSize;
+		/** Currently being streamed out, in bytes. */
+		SIZE_T PendingStreamOutSize;
+		/** How much we want to stream in, in bytes. */
+		int64 WantedInSize;
+		/** How much we want to stream out, in bytes.  */
+		int64 WantedOutSize;
+		/** Number of textures that want more mips. */
+		int64 NumWantingTextures;
+		/** Number of textures that are visibles and still low rez. */
+		int64 NumVisibleTexturesWithLowResolutions;
+	};
+
+	void ClearRemovedTextureReferences();
+
+	/** Returns the resulting priorities, matching the FStreamingManagerTexture::StreamingTextures array. */
+	const TArray<FTexturePriority>& GetPrioritizedTextures() const
+	{
+		return PrioritizedTextures;
+	}
+
+	const TArray<int32>& GetPrioTexIndicesSortedByLoadOrder() const
+	{
+		return PrioTexIndicesSortedByLoadOrder;
+	}
+
+	/** Returns the context (temporary info) used for the async work. */
+	const FStreamingContext& GetContext() const
+	{
+		return ThreadContext;
+	}
+
+	/** Returns the thread statistics. */
+	const FAsyncStats& GetStats() const
+	{
+		return ThreadStats;
+	}
+
+private:
+	friend class FAsyncTask<FAsyncTextureStreaming>;
+	/** Performs the async work. */
+	void DoWork();
+
+	FORCEINLINE TStatId GetStatId() const
+	{
+		RETURN_QUICK_DECLARE_CYCLE_STAT(FAsyncTextureStreaming, STATGROUP_ThreadPoolAsyncTasks);
+	}
+
+	/** Reference to the owning streaming manager, for accessing the thread-safe data. */
+	FStreamingManagerTexture&	StreamingManager;
+	/** Resulting priorities for keeping textures in memory and controlling load order, matching the FStreamingManagerTexture::StreamingTextures array. */
+	TArray<FTexturePriority>	PrioritizedTextures;
+	/** Indices for PrioritizedTextures, sorted by load order. */
+	TArray<int32>	PrioTexIndicesSortedByLoadOrder;
+
+	/** Context (temporary info) used for the async work. */
+	FStreamingContext			ThreadContext;
+	/** Thread statistics. */
+	FAsyncStats				ThreadStats;
+	/** Whether the async work should abort its processing. */
+	volatile bool				bAbort;
+};
