Copying //UE4/Dev-Rendering to Dev-Main (Source //UE4/Dev-Rendering@2932636)

diff --git a/Engine/Source/Runtime/Engine/Private/Streaming/StreamingManagerTexture.h b/Engine/Source/Runtime/Engine/Private/Streaming/StreamingManagerTexture.h
new file mode 100644
--- /dev/null
+++ b/Engine/Source/Runtime/Engine/Private/Streaming/StreamingManagerTexture.h
@@ -0,0 +1,761 @@
+// Copyright 1998-2016 Epic Games, Inc. All Rights Reserved.
+
+/*=============================================================================
+TextureStreamingManager.h: Definitions of classes used for texture streaming.
+=============================================================================*/
+
+#pragma once
+
+#include "StreamingTexture.h"
+#include "TextureInstanceManager.h"
+
+/*-----------------------------------------------------------------------------
+	Texture streaming.
+-----------------------------------------------------------------------------*/
+
+
+/**
+ * Structure containing all information needed for determining the screen space
+ * size of an object/ texture instance.
+ */
+struct FStreamableTextureInstance4
+{
+	FStreamableTextureInstance4()
+	:	BoundsOriginX( MAX_FLT, MAX_FLT, MAX_FLT, MAX_FLT )
+	,	BoundsOriginY( 0, 0, 0, 0 )
+	,	BoundsOriginZ( 0, 0, 0, 0 )
+	,	BoxExtentSizeX( 0, 0, 0, 0 )
+	,	BoxExtentSizeY( 0, 0, 0, 0 )
+	,	BoxExtentSizeZ( 0, 0, 0, 0 )
+	,	BoundingSphereRadius( 0, 0, 0, 0 )
+	,	MinDistanceSq( 0, 0, 0, 0 )
+	,	MaxDistanceSq( MAX_FLT, MAX_FLT, MAX_FLT, MAX_FLT )
+	,	TexelFactor( 0, 0, 0, 0 )
+	{
+	}
+
+	/** X coordinates for the bounds origin of 4 texture instances */
+	FVector4 BoundsOriginX;
+	/** Y coordinates for the bounds origin of 4 texture instances */
+	FVector4 BoundsOriginY;
+	/** Z coordinates for the bounds origin of 4 texture instances */
+	FVector4 BoundsOriginZ;
+
+	/** X size of the bounds box extent of 4 texture instances */
+	FVector4 BoxExtentSizeX;
+	/** Y size of the bounds box extent of 4 texture instances */
+	FVector4 BoxExtentSizeY;
+	/** Z size of the bounds box extent of 4 texture instances */
+	FVector4 BoxExtentSizeZ;
+
+	/** Sphere radii for the bounding sphere of 4 texture instances */
+	FVector4 BoundingSphereRadius;
+	/** Minimal distance ^2 (between the bounding sphere origin and the view origin) for which this entry is valid */
+	FVector4 MinDistanceSq;
+	/** Maximal distance ^2 (between the bounding sphere origin and the view origin) for which this entry is valid */
+	FVector4 MaxDistanceSq;
+
+	/** Texel scale factors for 4 texture instances */
+	FVector4 TexelFactor;
+};
+
+
+// To compute the maximum mip size required by any heuristic
+// float to avoid float->int conversion, unclamped until the end
+class FFloatMipLevel
+{
+public:
+	// constructor
+	FFloatMipLevel();
+
+	/**
+	 * Helper function to map screen texels to a mip level, includes the global bias
+	 * @param StreamingTexture can be 0 (if texture is getting destroyed)
+	 * @apram MipBias from cvar
+	 * @param bOptimal only needed for stats
+	 */
+	int32 ComputeMip(const FStreamingTexture* StreamingTexture, float MipBias, bool bOptimal) const;
+	//
+	bool operator>=(const FFloatMipLevel& rhs) const;
+	//
+	static FFloatMipLevel FromScreenSizeInTexels(float ScreenSizeInTexels);
+	
+	static FFloatMipLevel FromMipLevel(int32 InMipLevel);
+
+	void Invalidate();
+
+	// @return true: means this handle has a mip level that needs to be considered
+	bool IsHandled() const
+	{
+		return MipLevel > 0.0f;
+	}
+
+private:
+	// -1: not set yet
+	float MipLevel;
+};
+
+/**
+ * Streaming manager dealing with textures.
+ */
+struct FStreamingManagerTexture : public ITextureStreamingManager
+{
+	/** Constructor, initializing all members */
+	FStreamingManagerTexture();
+
+	virtual ~FStreamingManagerTexture();
+
+	/**
+	 * Updates streaming, taking into account all current view infos. Can be called multiple times per frame.
+	 *
+	 * @param DeltaTime				Time since last call in seconds
+	 * @param bProcessEverything	[opt] If true, process all resources with no throttling limits
+	 */
+	virtual void UpdateResourceStreaming( float DeltaTime, bool bProcessEverything=false ) override;
+
+	/**
+	 * Updates streaming for an individual texture, taking into account all view infos.
+	 *
+	 * @param Texture	Texture to update
+	 */
+	virtual void UpdateIndividualTexture( UTexture2D* Texture ) override;
+
+	/**
+	 * Blocks till all pending requests are fulfilled.
+	 *
+	 * @param TimeLimit		Optional time limit for processing, in seconds. Specifying 0 means infinite time limit.
+	 * @param bLogResults	Whether to dump the results to the log.
+	 * @return				Number of streaming requests still in flight, if the time limit was reached before they were finished.
+	 */
+	virtual int32 BlockTillAllRequestsFinished( float TimeLimit = 0.0f, bool bLogResults = false ) override;
+
+	/**
+	 * Cancels the timed Forced resources (i.e used the Kismet action "Stream In Textures").
+	 */
+	virtual void CancelForcedResources() override;
+
+	/**
+	 * Notifies manager of "level" change so it can prioritize character textures for a few frames.
+	 */
+	virtual void NotifyLevelChange() override;
+
+	/**
+	 * Adds a textures streaming handler to the array of handlers used to determine which
+	 * miplevels need to be streamed in/ out.
+	 *
+	 * @param TextureStreamingHandler	Handler to add
+	 */
+	void AddTextureStreamingHandler( FStreamingHandlerTextureBase* TextureStreamingHandler );
+
+	/**
+	 * Removes a textures streaming handler from the array of handlers used to determine which
+	 * miplevels need to be streamed in/ out.
+	 *
+	 * @param TextureStreamingHandler	Handler to remove
+	 */
+	void RemoveTextureStreamingHandler( FStreamingHandlerTextureBase* TextureStreamingHandler );
+
+	/** Don't stream world resources for the next NumFrames. */
+	virtual void SetDisregardWorldResourcesForFrames( int32 NumFrames ) override;
+
+	/**
+	 *	Try to stream out texture mip-levels to free up more memory.
+	 *	@param RequiredMemorySize	- Required minimum available texture memory
+	 *	@return						- Whether it succeeded or not
+	 **/
+	virtual bool StreamOutTextureData( int64 RequiredMemorySize ) override;
+
+	virtual int64 GetMemoryOverBudget() const override { return MemoryOverBudget; }
+
+	virtual int64 GetMaxEverRequired() const override { return MaxEverRequired; }
+
+	virtual void ResetMaxEverRequired() override { MaxEverRequired = 0; }
+
+	/**
+	 * Allows the streaming manager to process exec commands.
+	 *
+	 * @param InWorld World context
+	 * @param Cmd	Exec command
+	 * @param Ar	Output device for feedback
+	 * @return		true if the command was handled
+	 */
+	virtual bool Exec( UWorld* InWorld, const TCHAR* Cmd, FOutputDevice& Ar ) override;
+
+	/**
+	 * Exec command handlers
+	 */
+#if STATS_FAST
+	bool HandleDumpTextureStreamingStatsCommand( const TCHAR* Cmd, FOutputDevice& Ar );
+#endif // STATS_FAST
+#if STATS
+	bool HandleListStreamingTexturesCommand( const TCHAR* Cmd, FOutputDevice& Ar );
+	bool HandleListStreamingTexturesCollectCommand( const TCHAR* Cmd, FOutputDevice& Ar );
+	bool HandleListStreamingTexturesReportReadyCommand( const TCHAR* Cmd, FOutputDevice& Ar );
+	bool HandleListStreamingTexturesReportCommand( const TCHAR* Cmd, FOutputDevice& Ar );
+#endif // STATS
+#if !UE_BUILD_SHIPPING
+	bool HandleResetMaxEverRequiredTexturesCommand(const TCHAR* Cmd, FOutputDevice& Ar);
+	bool HandleLightmapStreamingFactorCommand( const TCHAR* Cmd, FOutputDevice& Ar );
+	bool HandleCancelTextureStreamingCommand( const TCHAR* Cmd, FOutputDevice& Ar );
+	bool HandleShadowmapStreamingFactorCommand( const TCHAR* Cmd, FOutputDevice& Ar );
+	bool HandleNumStreamedMipsCommand( const TCHAR* Cmd, FOutputDevice& Ar );
+	bool HandleTrackTextureCommand( const TCHAR* Cmd, FOutputDevice& Ar );
+	bool HandleListTrackedTexturesCommand( const TCHAR* Cmd, FOutputDevice& Ar );
+	bool HandleDebugTrackedTexturesCommand( const TCHAR* Cmd, FOutputDevice& Ar );
+	bool HandleUntrackTextureCommand( const TCHAR* Cmd, FOutputDevice& Ar );
+	bool HandleStreamOutCommand( const TCHAR* Cmd, FOutputDevice& Ar );
+	bool HandlePauseTextureStreamingCommand( const TCHAR* Cmd, FOutputDevice& Ar );
+	bool HandleStreamingManagerMemoryCommand( const TCHAR* Cmd, FOutputDevice& Ar, UWorld* InWorld );
+	bool HandleTextureGroupsCommand( const TCHAR* Cmd, FOutputDevice& Ar );
+	bool HandleInvestigateTextureCommand( const TCHAR* Cmd, FOutputDevice& Ar );
+#endif // !UE_BUILD_SHIPPING
+	/** Adds a new texture to the streaming manager. */
+	virtual void AddStreamingTexture( UTexture2D* Texture ) override;
+
+	/** Removes a texture from the streaming manager. */
+	virtual void RemoveStreamingTexture( UTexture2D* Texture ) override;
+
+	/** Adds a ULevel to the streaming manager. */
+	virtual void AddPreparedLevel( class ULevel* Level ) override;
+
+	/** Removes a ULevel from the streaming manager. */
+	virtual void RemoveLevel( class ULevel* Level ) override;
+
+	/** Called when an actor is spawned. */
+	virtual void NotifyActorSpawned( AActor* Actor ) override;
+
+	/** Called when a spawned actor is destroyed. */
+	virtual void NotifyActorDestroyed( AActor* Actor ) override;
+
+	/**
+	 * Called when a primitive is attached to an actor or another component.
+	 * Replaces previous info, if the primitive was already attached.
+	 *
+	 * @param InPrimitive	Newly attached dynamic/spawned primitive
+	 */
+	virtual void NotifyPrimitiveAttached( const UPrimitiveComponent* Primitive, EDynamicPrimitiveType DynamicType ) override;
+
+	/** Called when a primitive is detached from an actor or another component. */
+	virtual void NotifyPrimitiveDetached( const UPrimitiveComponent* Primitive ) override;
+
+	/**
+	 * Called when a LastRenderTime primitive is attached to an actor or another component.
+	 * Modifies the LastRenderTimeRefCount for the textures used, so that those textures can
+	 * use both distance-based and LastRenderTime heuristics.
+	 *
+	 * @param Primitive	Newly attached dynamic/spawned primitive
+	 */
+	virtual void NotifyTimedPrimitiveAttached( const UPrimitiveComponent* Primitive, EDynamicPrimitiveType DynamicType ) override;
+
+	/**
+	 * Called when a LastRenderTime primitive is detached from an actor or another component.
+	 * Modifies the LastRenderTimeRefCount for the textures used, so that those textures can
+	 * use both distance-based and LastRenderTime heuristics.
+	 *
+	 * @param Primitive	Newly detached dynamic/spawned primitive
+	 */
+	virtual void NotifyTimedPrimitiveDetached( const UPrimitiveComponent* Primitive ) override;
+
+	/**
+	 * Called when a primitive has had its textured changed.
+	 * Only affects primitives that were already attached.
+	 * Replaces previous info.
+	 */
+	virtual void NotifyPrimitiveUpdated( const UPrimitiveComponent* Primitive ) override;
+
+	/** Returns the corresponding FStreamingTexture for a UTexture2D. */
+	FStreamingTexture& GetStreamingTexture( const UTexture2D* Texture2D );
+
+	/** Returns true if this is a streaming texture that is managed by the streaming manager. */
+	virtual bool IsManagedStreamingTexture( const UTexture2D* Texture2D ) override;
+
+	/** Updates the I/O state of a texture (allowing it to progress to the next stage) and some stats. */
+	void UpdateTextureStatus( FStreamingTexture& StreamingTexture, FStreamingContext& Context );
+
+	/**
+	 * Cancels the current streaming request for the specified texture.
+	 *
+	 * @param StreamingTexture		Texture to cancel streaming for
+	 * @return						true if a streaming request was canceled
+	 */
+	bool CancelStreamingRequest( FStreamingTexture& StreamingTexture );
+
+	/** Resets the streaming statistics to zero. */
+	void ResetStreamingStats();
+
+	/**
+	 * Updates the streaming statistics with current frame's worth of stats.
+	 *
+	 * @param Context					Context for the current frame
+	 * @param bAllTexturesProcessed		Whether all processing is complete
+	 */
+	void UpdateStreamingStats( const FStreamingContext& Context, bool bAllTexturesProcessed );
+
+	/** Returns the number of cached view infos (thread-safe). */
+	int32 ThreadNumViews()
+	{
+		return ThreadSettings.ThreadViewInfos.Num();
+	}
+
+	/** Returns a cached view info for the specified index (thread-safe). */
+	const FStreamingViewInfo& ThreadGetView( int32 ViewIndex )
+	{
+		return ThreadSettings.ThreadViewInfos[ ViewIndex ];
+	}
+
+#if STATS
+	/** Ringbuffer of bandwidth samples for streaming in mip-levels (MB/sec). */
+	static float BandwidthSamples[NUM_BANDWIDTHSAMPLES];
+	/** Number of bandwidth samples that have been filled in. Will stop counting when it reaches NUM_BANDWIDTHSAMPLES. */
+	static int32 NumBandwidthSamples;
+	/** Current sample index in the ring buffer. */
+	static int32 BandwidthSampleIndex;
+	/** Average of all bandwidth samples in the ring buffer, in MB/sec. */
+	static float BandwidthAverage;
+	/** Maximum bandwidth measured since the start of the game.  */
+	static float BandwidthMaximum;
+	/** Minimum bandwidth measured since the start of the game.  */
+	static float BandwidthMinimum;
+#endif
+
+	/** Set current pause state for texture streaming */
+	virtual void PauseTextureStreaming(bool bInShouldPause) override
+	{
+		bPauseTextureStreaming = bInShouldPause;
+	}
+
+protected:
+//BEGIN: Thread-safe functions and data
+		friend class FAsyncTextureStreaming;
+		friend struct FStreamingHandlerTextureStatic;
+		friend struct FStreamingHandlerTextureDynamic;
+
+		/** Calculates the minimum and maximum number of mip-levels for a streaming texture. */
+		void CalcMinMaxMips( FStreamingTexture& StreamingTexture );
+
+		/** Calculates the number of mip-levels we would like to have in memory for a texture. */
+		void CalcWantedMips( FStreamingTexture& StreamingTexture );
+
+		/**
+		 * Fallback handler to catch textures that have been orphaned recently.
+		 * This handler prevents massive spike in texture memory usage.
+		 * Orphaned textures were previously streamed based on distance but those instance locations have been removed -
+		 * for instance because a ULevel is being unloaded. These textures are still active and in memory until they are garbage collected,
+		 * so we must ensure that they do not start using the LastRenderTime fallback handler and suddenly stream in all their mips -
+		 * just to be destroyed shortly after by a garbage collect.
+		 */
+		FFloatMipLevel GetWantedMipsForOrphanedTexture( FStreamingTexture& StreamingTexture, float& Distance );
+
+		/** Updates this frame's STATs by one texture. */
+		void UpdateFrameStats( FStreamingContext& Context, FStreamingTexture& StreamingTexture, int32 TextureIndex );
+
+		/**
+		 * Not thread-safe: Updates a portion (as indicated by 'StageIndex') of all streaming textures,
+		 * allowing their streaming state to progress.
+		 *
+		 * @param Context			Context for the current stage (frame)
+		 * @param StageIndex		Current stage index
+		 * @param NumUpdateStages	Number of texture update stages
+		 */
+		void UpdateStreamingTextures( FStreamingContext& Context, int32 StageIndex, int32 NumStages );
+
+		/** Adds new textures and level data on the gamethread (while the worker thread isn't active). */
+		void UpdateThreadData();
+
+		/** Checks for updates in the user settings (CVars, etc). */
+		void CheckUserSettings();
+
+		/**
+		 * Temporarily boosts the streaming distance factor by the specified number.
+		 * This factor is automatically reset to 1.0 after it's been used for mip-calculations.
+		 */
+		void BoostTextures( AActor* Actor, float BoostFactor ) override;
+
+		/**
+		 * Stream textures in/out, based on the priorities calculated by the async work.
+		 *
+		 * @param bProcessEverything	Whether we're processing all textures in one go
+		 */
+		void StreamTextures( bool bProcessEverything );
+
+		/**
+		 * Keep as many unnecessary mips as possible.
+		 *
+		 * @param Context				Context for the current stage
+		 */
+		void KeepUnwantedMips( FStreamingContext& Context );
+
+		/**
+		 * Drop as many mips required to fit in budget.
+		 *
+		 * @param Context				Context for the current stage
+		 */
+		void DropWantedMips( FStreamingContext& Context );
+
+		/**
+		 * Drop as many forced required to fit in budget.
+		 *
+		 * @param Context				Context for the current stage
+		 */
+		void DropForcedMips( FStreamingContext& Context );
+
+		/** Thread-safe helper struct for per-level information. */
+		struct FThreadLevelData
+		{
+			FThreadLevelData()
+			:	bRemove( false )
+			{
+			}
+
+			/** Whether this level has been removed. */
+			bool bRemove;
+
+			/** Texture instances used in this level, stored in SIMD-friendly format. */
+			TMap<const UTexture2D*,TArray<FStreamableTextureInstance4> > ThreadTextureInstances;
+		};
+
+		typedef TKeyValuePair< class ULevel*, FThreadLevelData >	FLevelData;
+
+		/** Thread-safe helper struct for streaming information. */
+		struct FThreadSettings
+		{
+			FThreadSettings() : TextureBoundsVisibility(nullptr) {}
+			
+			/** Cached from the system settings. */
+			int32 NumStreamedMips[TEXTUREGROUP_MAX];
+
+			/** Cached from each ULevel. */
+			TArray< FLevelData > LevelData;
+
+			/** Cached from FStreamingManagerBase. */
+			TArray<FStreamingViewInfo> ThreadViewInfos;
+
+			FTextureBoundsVisibility* TextureBoundsVisibility;
+
+			/** from cvar, >=0 */
+			float MipBias;
+		};
+
+
+		/** Thread-safe helper data for streaming information. */
+		FThreadSettings	ThreadSettings;
+
+		/** All streaming UTexture2D objects. */
+		TArray<FStreamingTexture> StreamingTextures;
+
+		/** Index of the StreamingTexture that will be updated next by UpdateStreamingTextures(). */
+		int32 CurrentUpdateStreamingTextureIndex;
+//END: Thread-safe functions and data
+
+	/**
+	 * Mark the textures with a timestamp. They're about to lose their location-based heuristic and we don't want them to
+	 * start using LastRenderTime heuristic for a few seconds until they are garbage collected!
+	 *
+	 * @param RemovedTextures	List of removed textures.
+	 */
+	void	SetTexturesRemovedTimestamp(const FRemovedTextureArray& RemovedTextures);
+
+	void	DumpTextureGroupStats( bool bDetailedStats );
+
+	/**
+	 * Prints out detailed information about streaming textures that has a name that contains the given string.
+	 * Triggered by the InvestigateTexture exec command.
+	 *
+	 * @param InvestigateTextureName	Partial name to match textures against
+	 */
+	void	InvestigateTexture( const FString& InvestigateTextureName );
+
+	void	DumpTextureInstances( const UPrimitiveComponent* Primitive, UTexture2D* Texture2D );
+
+	/** Next sync, dump texture group stats. */
+	bool	bTriggerDumpTextureGroupStats;
+
+	/** Whether to the dumped texture group stats should contain extra information. */
+	bool	bDetailedDumpTextureGroupStats;
+
+	/** Next sync, dump all information we have about a certain texture. */
+	bool	bTriggerInvestigateTexture;
+
+	/** Name of a texture to investigate. Can be partial name. */
+	FString	InvestigateTextureName;
+
+	/** Async work for calculating priorities for all textures. */
+	FAsyncTask<FAsyncTextureStreaming>*	AsyncWork;
+
+	/** Textures from dynamic primitives. */
+	FDynamicComponentTextureManager* DynamicComponentTextureManager;
+
+	/** New textures, before they've been added to the thread-safe container. */
+	TArray<UTexture2D*>		PendingStreamingTextures;
+
+	/** New levels, before they've been added to the thread-safe container. */
+	TArray<class ULevel*>	PendingLevels;
+
+	/** Stages [0,N-2] is non-threaded data collection, Stage N-1 is wait-for-AsyncWork-and-finalize. */
+	int32					ProcessingStage;
+
+	/** Total number of processing stages (N). */
+	int32					NumTextureProcessingStages;
+
+	/** Whether to support texture instance streaming for dynamic (movable/spawned) objects. */
+	bool					bUseDynamicStreaming;
+
+	float					BoostPlayerTextures;
+
+	/** Extra distance added to the range test, to start streaming the texture before they are actually used. */
+	float					RangePrefetchDistance;
+
+	/** Array of texture streaming objects to use during update. */
+	TArray<FStreamingHandlerTextureBase*> TextureStreamingHandlers;
+
+	/** Amount of memory to leave free in the texture pool. */
+	int64					MemoryMargin;
+
+	/** Minimum number of bytes to evict when we need to stream out textures because of a failed allocation. */
+	int64					MinEvictSize;
+
+	/** The actual memory pool size available to stream textures, excludes non-streaming texture, temp memory (for streaming mips), memory margin (allocator overhead). */
+	int64					EffectiveStreamingPoolSize;
+
+	/** If set, UpdateResourceStreaming() will only process this texture. */
+	UTexture2D*				IndividualStreamingTexture;
+
+	// Stats we need to keep across frames as we only iterate over a subset of textures.
+
+	/** Number of streaming textures */
+	uint32 NumStreamingTextures;
+	/** Number of requests in cancelation phase. */
+	uint32 NumRequestsInCancelationPhase;
+	/** Number of requests in mip update phase. */
+	uint32 NumRequestsInUpdatePhase;
+	/** Number of requests in mip finalization phase. */
+	uint32 NumRequestsInFinalizePhase;
+	/** Size ot all intermerdiate textures in flight. */
+	uint32 TotalIntermediateTexturesSize;
+	/** Number of intermediate textures in flight. */
+	uint32 NumIntermediateTextures;
+	/** Size of all streaming testures. */
+	uint64 TotalStreamingTexturesSize;
+	/** Maximum size of all streaming textures. */
+	uint64 TotalStreamingTexturesMaxSize;
+	/** Total number of bytes in memory, for all streaming lightmap textures. */
+	uint64 TotalLightmapMemorySize;
+	/** Total number of bytes on disk, for all streaming lightmap textures. */
+	uint64 TotalLightmapDiskSize;
+	/** Total number of bytes in memory, for all HLOD textures. */
+	uint64 TotalHLODMemorySize;
+	/** Total number of bytes on disk, for all HLOD textures. */
+	uint64 TotalHLODDiskSize;
+	/** Number of mip count increase requests in flight. */
+	uint32 TotalMipCountIncreaseRequestsInFlight;
+	/** Total number of bytes in memory, if all textures were streamed perfectly with a 1.0 fudge factor. */
+	uint64 TotalOptimalWantedSize;
+	/** Total number of bytes using StaticTexture heuristics, currently in memory. */
+	uint64 TotalStaticTextureHeuristicSize;
+	/** Total number of bytes using DynmicTexture heuristics, currently in memory. */
+	uint64 TotalDynamicTextureHeuristicSize;
+	/** Total number of bytes using LastRenderTime heuristics, currently in memory. */
+	uint64 TotalLastRenderHeuristicSize;
+	/** Total number of bytes using ForcedIntoMemory heuristics, currently in memory. */
+	uint64 TotalForcedHeuristicSize;
+	int64 MemoryOverBudget;
+	int64 MaxEverRequired;
+
+	/** Unmodified texture pool size, in bytes, as specified in the .ini file. */
+	int64 OriginalTexturePoolSize;
+	/** Timestamp when we last shrunk the pool size because of memory usage. */
+	double PreviousPoolSizeTimestamp;
+	/** PoolSize CVar setting the last time we adjusted the pool size. */
+	int32 PreviousPoolSizeSetting;
+
+	/** Whether to collect, and optionally report, texture stats for the next run. */
+	bool   bCollectTextureStats;
+	bool   bReportTextureStats;
+	TArray<FString> TextureStatsReport;
+
+	/** Optional string to match against the texture names when collecting stats. */
+	FString CollectTextureStatsName;
+
+	/** Whether texture streaming is paused or not. When paused, it won't stream any textures in or out. */
+	bool bPauseTextureStreaming;
+
+#if STATS
+	/**
+	 * Ring buffer containing all latency samples we keep track of, as measured in seconds from the time the streaming system
+	 * detects that a change is required until the data has been streamed in and the new texture is ready to be used.
+	 */
+	float LatencySamples[NUM_LATENCYSAMPLES];
+	/** Number of latency samples that have been filled in. Will stop counting when it reaches NUM_LATENCYSAMPLES. */
+	int32 NumLatencySamples;
+	/** Current sample index in the ring buffer. */
+	int32 LatencySampleIndex;
+	/** Average of all latency samples in the ring buffer, in seconds. */
+	float LatencyAverage;
+	/** Maximum latency measured since the start of the game.  */
+	float LatencyMaximum;
+
+	/**
+	 * Updates the streaming latency STAT for a texture.
+	 *
+	 * @param Texture		Texture to update for
+	 * @param WantedMips	Number of desired mip-levels for the texture
+	 * @param bInFlight		Whether the texture is currently being streamed
+	 */
+	void UpdateLatencyStats( UTexture2D* Texture, int32 WantedMips, bool bInFlight );
+
+	/** Total time taken for each processing stage. */
+	TArray<double> StreamingTimes;
+#endif
+
+#if STATS_FAST
+	uint64 MaxStreamingTexturesSize;
+	uint64 MaxOptimalTextureSize;
+	int64 MaxStreamingOverBudget;
+	uint64 MaxTexturePoolAllocatedSize;
+	uint64 MinLargestHoleSize;
+	uint32 MaxNumWantingTextures;
+#endif
+	
+	friend bool TrackTextureEvent( FStreamingTexture* StreamingTexture, UTexture2D* Texture, bool bForceMipLevelsToBeResident, const FStreamingManagerTexture* Manager);
+};
+
+
+
+
+/*-----------------------------------------------------------------------------
+	Texture streaming handler.
+-----------------------------------------------------------------------------*/
+
+/**
+ * Base of texture streaming handler functionality.
+ */
+struct FStreamingHandlerTextureBase
+{
+	/** Friendly name identifying the handler for debug purposes. */
+	const TCHAR* HandlerName;
+
+	/** Default constructor. */
+	FStreamingHandlerTextureBase()
+		: HandlerName(TEXT("Base"))
+	{
+	}
+
+	/**
+	 * Returns mip count wanted by this handler for the passed in texture. 
+	 * 
+	 * @param	StreamingManager	Streaming manager
+	 * @param	Streaming Texture	Texture to determine wanted mip count for
+	 * @param	MinDistance			[out] Distance to the closest instance of this texture, in world space units.
+	 * @return	Number of miplevels that should be streamed in or INDEX_NONE if texture isn't handled by this handler.
+	 */
+	virtual FFloatMipLevel GetWantedMips( FStreamingManagerTexture& StreamingManager, FStreamingTexture& StreamingTexture, float& MinDistance ) = 0;
+};
+
+/**
+ * Static texture streaming handler. Used to stream textures on static level geometry.
+ */
+struct FStreamingHandlerTextureStatic : public FStreamingHandlerTextureBase
+{
+	/** Default constructor. */
+	FStreamingHandlerTextureStatic()
+	{
+		HandlerName = TEXT("Static");
+	}
+
+	/**
+	 * Returns mip count wanted by this handler for the passed in texture. 
+	 * 
+	 * @param	StreamingManager	Streaming manager
+	 * @param	Streaming Texture	Texture to determine wanted mip count for
+	 * @param	MinDistance			[out] Distance to the closest instance of this texture, in world space units.
+	 * @return	Number of miplevels that should be streamed in or INDEX_NONE if texture isn't handled by this handler.
+	 */
+	virtual FFloatMipLevel GetWantedMips( FStreamingManagerTexture& StreamingManager, FStreamingTexture& StreamingTexture, float& MinDistance );
+};
+
+
+/**
+ * Dynamic texture streaming handler. Used to stream textures for dynamic primitives.
+ */
+struct FStreamingHandlerTextureDynamic : public FStreamingHandlerTextureBase
+{
+	/** Default constructor. */
+	FStreamingHandlerTextureDynamic()
+	{
+		HandlerName = TEXT("Dynamic");
+	}
+
+	/**
+	 * Returns mip count wanted by this handler for the passed in texture. 
+	 * 
+	 * @param	StreamingManager	Streaming manager
+	 * @param	Streaming Texture	Texture to determine wanted mip count for
+	 * @param	MinDistance			[out] Distance to the closest instance of this texture, in world space units.
+	 * @return	Number of miplevels that should be streamed in or INDEX_NONE if texture isn't handled by this handler.
+	 */
+	virtual FFloatMipLevel GetWantedMips( FStreamingManagerTexture& StreamingManager, FStreamingTexture& StreamingTexture, float& MinDistance );
+};
+
+/**
+ * Streaming handler that bases its decision on the last render time.
+ */
+struct FStreamingHandlerTextureLastRender : public FStreamingHandlerTextureBase
+{
+	/** Default constructor. */
+	FStreamingHandlerTextureLastRender()
+	{
+		HandlerName = TEXT("LastRender");
+	}
+
+	/**
+	 * Returns mip count wanted by this handler for the passed in texture. 
+	 * 
+	 * @param	StreamingManager	Streaming manager
+	 * @param	Streaming Texture	Texture to determine wanted mip count for
+	 * @param	MinDistance			[out] Distance to the closest instance of this texture, in world space units.
+	 * @return	Number of miplevels that should be streamed in or INDEX_NONE if texture isn't handled by this handler.
+	 */
+	virtual FFloatMipLevel GetWantedMips( FStreamingManagerTexture& StreamingManager, FStreamingTexture& StreamingTexture, float& MinDistance );
+};
+
+/**
+ *	Streaming handler that bases its decision on the bForceMipStreaming flag in PrimitiveComponent.
+ */
+struct FStreamingHandlerTextureLevelForced : public FStreamingHandlerTextureBase
+{
+	/** Default constructor. */
+	FStreamingHandlerTextureLevelForced()
+	{
+		HandlerName = TEXT("LevelForced");
+	}
+
+	/**
+	 * Returns mip count wanted by this handler for the passed in texture. 
+	 * 
+	 * @param	StreamingManager	Streaming manager
+	 * @param	Streaming Texture	Texture to determine wanted mip count for
+	 * @param	MinDistance			[out] Distance to the closest instance of this texture, in world space units.
+	 * @return	Number of miplevels that should be streamed in or INDEX_NONE if texture isn't handled by this handler.
+	 */
+	virtual FFloatMipLevel GetWantedMips( FStreamingManagerTexture& StreamingManager, FStreamingTexture& StreamingTexture, float& MinDistance );
+};
+
+/*-----------------------------------------------------------------------------
+	Texture streaming helper structs.
+-----------------------------------------------------------------------------*/
+
+struct FSpawnedTextureInstance
+{
+	FSpawnedTextureInstance( UTexture2D* InTexture2D, float InTexelFactor, float InOriginalRadius )
+	:	Texture2D( InTexture2D )
+	,	TexelFactor( InTexelFactor )
+	,	InvOriginalRadius( (InOriginalRadius > 0.0f) ? (1.0f/InOriginalRadius) : 1.0f )
+	{
+	}
+	/** Texture that is used by a dynamic UPrimitiveComponent. */
+	UTexture2D*		Texture2D;
+	/** Texture specific texel scale factor  */
+	float			TexelFactor;
+	/** 1.0f / OriginalBoundingSphereRadius, at the time the TexelFactor was calculated originally. */
+	float			InvOriginalRadius;
+};
\ No newline at end of file
diff --git a/Engine/Source/Runtime/Engine/Private/Streaming/StreamingTexture.cpp b/Engine/Source/Runtime/Engine/Private/Streaming/StreamingTexture.cpp
new file mode 100644
--- /dev/null
+++ b/Engine/Source/Runtime/Engine/Private/Streaming/StreamingTexture.cpp
@@ -0,0 +1,231 @@
+// Copyright 1998-2016 Epic Games, Inc. All Rights Reserved.
+
+/*=============================================================================
+StreamingTexture.cpp: Definitions of classes used for texture.
+=============================================================================*/
+
+#include "EnginePrivate.h"
+#include "StreamingTexture.h"
+#include "StreamingManagerTexture.h"
+#include "Engine/LightMapTexture2D.h"
+#include "Engine/ShadowMapTexture2D.h"
+
+FStreamingTexture::FStreamingTexture( UTexture2D* InTexture )
+{
+	Texture = InTexture;
+	WantedMips = InTexture->ResidentMips;
+	MipCount = InTexture->GetNumMips();
+	PerfectWantedMips = InTexture->ResidentMips;
+	STAT( MostResidentMips = InTexture->ResidentMips );
+	LODGroup = (TextureGroup) InTexture->LODGroup;
+
+	bIsHLODTexture = LODGroup == TEXTUREGROUP_HierarchicalLOD;
+
+	// Landscape textures and HLOD textures should not be affected by MipBias
+	//@TODO: Long term this should probably be a property of the texture group instead
+	bCanBeAffectedByMipBias = true;
+	if ((LODGroup == TEXTUREGROUP_Terrain_Heightmap) || (LODGroup == TEXTUREGROUP_Terrain_Weightmap) || bIsHLODTexture)
+	{
+		bCanBeAffectedByMipBias = false;
+	}
+
+	NumNonStreamingMips = InTexture->GetNumNonStreamingMips();
+	ForceLoadRefCount = 0;
+	bIsStreamingLightmap = IsStreamingLightmap( Texture );
+	bIsLightmap = bIsStreamingLightmap || LODGroup == TEXTUREGROUP_Lightmap || LODGroup == TEXTUREGROUP_Shadowmap;
+	bUsesStaticHeuristics = false;
+	bUsesDynamicHeuristics = false;
+	bUsesLastRenderHeuristics = false;
+	bUsesForcedHeuristics = false;
+	bUsesOrphanedHeuristics = false;
+	bHasSplitRequest = false;
+	bIsLastSplitRequest = false;
+	BoostFactor = 1.0f;
+	InstanceRemovedTimestamp = -FLT_MAX;
+	LastRenderTimeRefCountTimestamp = -FLT_MAX;
+	LastRenderTimeRefCount = 0;
+
+	for ( int32 MipIndex=1; MipIndex <= MAX_TEXTURE_MIP_COUNT; ++MipIndex )
+	{
+		TextureSizes[MipIndex - 1] = Texture->CalcTextureMemorySize( FMath::Min(MipIndex, MipCount) );
+	}
+
+	UpdateCachedInfo();
+}
+
+void FStreamingTexture::UpdateCachedInfo( )
+{
+	check(Texture && Texture->Resource);
+	ResidentMips = Texture->ResidentMips;
+	RequestedMips = Texture->RequestedMips;
+	MinAllowedMips = 1;			//FMath::Min(Texture->ResidentMips, Texture->RequestedMips);
+	MaxAllowedMips = MipCount;	//FMath::Max(Texture->ResidentMips, Texture->RequestedMips);
+	MaxAllowedOptimalMips = MaxAllowedMips;
+	STAT( MostResidentMips = FMath::Max(MostResidentMips, Texture->ResidentMips) );
+	float LastRenderTimeForTexture = Texture->GetLastRenderTimeForStreaming();
+	LastRenderTime = (FApp::GetCurrentTime() > LastRenderTimeForTexture) ? float( FApp::GetCurrentTime() - LastRenderTimeForTexture ) : 0.0f;
+	MinDistance = MAX_STREAMINGDISTANCE;
+	bForceFullyLoad = Texture->ShouldMipLevelsBeForcedResident() || (ForceLoadRefCount > 0);
+	TextureLODBias = Texture->GetCachedLODBias();
+	bInFlight = false;
+	bReadyForStreaming = IsStreamingTexture( Texture ) && IsReadyForStreaming( Texture );
+	NumCinematicMipLevels = Texture->bUseCinematicMipLevels ? Texture->NumCinematicMipLevels : 0;
+}
+
+ETextureStreamingType FStreamingTexture::GetStreamingType() const
+{
+	if ( bUsesForcedHeuristics || bForceFullyLoad )
+	{
+		return StreamType_Forced;
+	}
+	else if ( bUsesDynamicHeuristics )
+	{
+		return StreamType_Dynamic;
+	}
+	else if ( bUsesStaticHeuristics )
+	{
+		return StreamType_Static;
+	}
+	else if ( bUsesOrphanedHeuristics )
+	{
+		return StreamType_Orphaned;
+	}
+	else if ( bUsesLastRenderHeuristics )
+	{
+		return StreamType_LastRenderTime;
+	}
+	return StreamType_Other;
+}
+
+bool FStreamingTexture::UpdateMipCount(  FStreamingManagerTexture& Manager, FStreamingContext& Context )
+{
+	if ( !Texture || !bReadyForStreaming )
+	{
+		return false;
+	}
+
+	// If there is a pending load that attempts to load unrequired data, 
+	// or if there is a pending unload that attempts to unload required data, try to cancel it.
+	if ( bInFlight && 
+		(( RequestedMips > ResidentMips && RequestedMips > WantedMips ) || 
+		 ( RequestedMips < ResidentMips && RequestedMips < WantedMips )) )
+	{
+		// Try to cancel, but don't try anything else before next update.
+		bReadyForStreaming = false;
+
+		if ( Texture->CancelPendingMipChangeRequest() )
+		{
+			bInFlight = false;
+			RequestedMips = ResidentMips;
+			return true;
+		}
+	}
+	// If we should load/unload some mips.
+	else if ( !bInFlight && ResidentMips != WantedMips )
+	{
+		check(RequestedMips == ResidentMips);
+
+		// Try to unload, but don't try anything else before next update.
+		bReadyForStreaming = false;
+
+		if ( Texture->PendingMipChangeRequestStatus.GetValue() == TexState_ReadyFor_Requests )
+		{
+			check(!Texture->bHasCancelationPending);
+
+			bool bHasEnoughTempMemory = GetSize(WantedMips) <= Context.AvailableTempMemory;
+			bool bHasEnoughAvailableNow = WantedMips <= ResidentMips || GetSize(WantedMips) <= Context.AvailableNow;
+
+			// Check if there is enough space for temp memory.
+			if ( bHasEnoughTempMemory && bHasEnoughAvailableNow )
+			{
+				FTexture2DResource* Texture2DResource = (FTexture2DResource*)Texture->Resource;
+
+				if (WantedMips > ResidentMips) // If it is a load request.
+				{
+					// Manually update size as allocations are deferred/ happening in rendering thread.
+					int64 LoadRequestSize = GetSize(WantedMips) - GetSize(ResidentMips);
+					Context.ThisFrameTotalRequestSize			+= LoadRequestSize;
+					Context.ThisFrameTotalLightmapRequestSize	+= bIsStreamingLightmap ? LoadRequestSize : 0;
+					Context.AvailableNow						-= LoadRequestSize;
+				}
+
+				Context.AvailableTempMemory -= GetSize(WantedMips);
+
+				bInFlight = true;
+				RequestedMips = WantedMips;
+				Texture->RequestedMips = RequestedMips;
+
+				Texture2DResource->BeginUpdateMipCount( bForceFullyLoad );
+				TrackTextureEvent( this, Texture, bForceFullyLoad, &Manager );
+				return true;
+			}
+		}
+		else
+		{
+			// Did UpdateStreamingTextures() miss a texture? Should never happen!
+			UE_LOG(LogContentStreaming, Warning, TEXT("BeginUpdateMipCount failure! PendingMipChangeRequestStatus == %d, Resident=%d, Requested=%d, Wanted=%d"), Texture->PendingMipChangeRequestStatus.GetValue(), Texture->ResidentMips, Texture->RequestedMips, WantedMips );
+		}
+	}
+	return false;
+}
+
+float FStreamingTexture::CalcRetentionPriority( )
+{
+	const int32 WantedMipsForPriority = GetWantedMipsForPriority();
+
+	bool bMustKeep = LODGroup == TEXTUREGROUP_Terrain_Heightmap || bForceFullyLoad || WantedMipsForPriority <= MinAllowedMips;
+	bool bIsVisible = !GStreamWithTimeFactor || IsVisible();
+	bool bAlreadyLowRez = /*(PixelPerTexel > 1.5 && Distance > 1) ||*/ bIsHLODTexture || bIsLightmap;
+		
+	// Don't consider dropping mips that give nothing. This is because streaming time has a high cost per texture.
+	// Dropping many textures for not so much memory would affect the time it takes the streamer to recover.
+	// We also don't want to prioritize by size, since bigger wanted mips are expected to have move visual impact.
+	bool bGivesOffMemory = GetSize(WantedMipsForPriority) - GetSize(WantedMipsForPriority - 1) >= 256 * 1024; 
+
+	float Priority = 0;
+	if (bMustKeep)			Priority += 1024.f;
+	if (bAlreadyLowRez)		Priority += 512.f;
+	if (!bGivesOffMemory)	Priority += 256.f;
+	if (bIsVisible)			Priority += 128.f;
+
+	return Priority;
+}
+
+float FStreamingTexture::CalcLoadOrderPriority()
+{
+	const int32 WantedMipsForPriority = GetWantedMipsForPriority();
+
+	bool bMustLoadFirst = LODGroup == TEXTUREGROUP_Terrain_Heightmap || bForceFullyLoad;
+	bool bIsVisible = !GStreamWithTimeFactor || IsVisible();
+	bool bFarFromTargetMips = WantedMipsForPriority - ResidentMips > 2;
+	bool bBigOnScreen = WantedMipsForPriority > 9.f || bIsHLODTexture;
+
+	float Priority = 0;
+	if (bMustLoadFirst)		Priority += 1024.f;
+	if (!bIsLastSplitRequest)	Priority += 512.f; 
+	if (bIsVisible)				Priority += 256.f;
+	if (bFarFromTargetMips)		Priority += 128.f;
+	if (bBigOnScreen)			Priority += 64.f;
+	return Priority;
+}
+
+bool FStreamingTexture::IsReadyForStreaming( UTexture2D* Texture )
+{
+	return Texture->IsReadyForStreaming();
+}
+
+bool FStreamingTexture::IsStreamingLightmap( UTexture2D* Texture )
+{
+	ULightMapTexture2D* Lightmap = Cast<ULightMapTexture2D>(Texture);
+	UShadowMapTexture2D* Shadowmap = Cast<UShadowMapTexture2D>(Texture);
+
+	if ( Lightmap && Lightmap->LightmapFlags & LMF_Streamed )
+	{
+		return true;
+	}
+	else if ( Shadowmap && Shadowmap->ShadowmapFlags & SMF_Streamed )
+	{
+		return true;
+	}
+	return false;
+}
diff --git a/Engine/Source/Runtime/Engine/Private/Streaming/StreamingTexture.h b/Engine/Source/Runtime/Engine/Private/Streaming/StreamingTexture.h
new file mode 100644
--- /dev/null
+++ b/Engine/Source/Runtime/Engine/Private/Streaming/StreamingTexture.h
@@ -0,0 +1,204 @@
+// Copyright 1998-2016 Epic Games, Inc. All Rights Reserved.
+
+/*=============================================================================
+StreamingTexture.h: Definitions of classes used for texture streaming.
+=============================================================================*/
+
+#pragma once
+
+#include "TextureStreamingHelpers.h"
+
+/*-----------------------------------------------------------------------------
+	FStreamingTexture, the streaming system's version of UTexture2D.
+-----------------------------------------------------------------------------*/
+
+/** Self-contained structure to manage a streaming texture, possibly on a separate thread. */
+struct FStreamingTexture
+{
+	FStreamingTexture( UTexture2D* InTexture );
+
+	/** Reinitializes the cached variables from UTexture2D. */
+	void UpdateCachedInfo( );
+
+	ETextureStreamingType GetStreamingType() const;
+
+	/**
+	 * Checks whether a UTexture2D is ready for streaming.
+	 *
+	 * @param Texture	Texture to check
+	 * @return			true if the UTexture2D is ready to be streamed in or out
+	 */
+	static bool IsReadyForStreaming( UTexture2D* Texture );
+
+	/**
+	 * Checks whether a UTexture2D is a streaming lightmap or shadowmap.
+	 *
+	 * @param Texture	Texture to check
+	 * @return			true if the UTexture2D is a streaming lightmap or shadowmap
+	 */
+	static bool IsStreamingLightmap( UTexture2D* Texture );
+
+	/**
+	 * Returns the amount of memory used by the texture given a specified number of mip-maps, in bytes.
+	 *
+	 * @param MipCount	Number of mip-maps to account for
+	 * @return			Total amount of memory used for the specified mip-maps, in bytes
+	 */
+	int32 GetSize( int32 InMipCount ) const
+	{
+		check(InMipCount >= 0 && InMipCount <= MAX_TEXTURE_MIP_COUNT);
+		return TextureSizes[ InMipCount - 1 ];
+	}
+
+	/**
+	 * Returns whether this primitive has been used rencently. Conservative.
+	 *
+	 * @return			true if visible
+	 */
+	FORCEINLINE bool IsVisible() const
+	{
+		return LastRenderTime < 5.f && LastRenderTime != MAX_FLT;
+	}
+	
+	/**
+	 * Return whether we allow to not stream the mips to fit in memory.
+	 *
+	 * @return			true if it is ok to drop the mips
+	 */
+	FORCEINLINE bool CanDropMips() const
+	{
+		return (LODGroup != TEXTUREGROUP_Terrain_Heightmap) && WantedMips > MinAllowedMips && !bForceFullyLoad;
+	}
+
+
+	/**
+	 * Return the wanted mips to use for priority computates
+	 *
+	 * @return		the WantedMips
+	 */
+	FORCEINLINE int32 GetWantedMipsForPriority() const
+	{
+		if (bHasSplitRequest && !bIsLastSplitRequest)
+		{
+			return WantedMips + 1;
+		}
+		else
+		{
+			return WantedMips;
+		}
+	}
+
+	/**
+	 * Calculates a retention priority value for the textures. Higher value means more important.
+	 * Retention is used determine which texture to drop when out of texture pool memory.
+	 * The retention value must be stable, meaning it must not depend on the current state of the streamer.
+	 * Not doing so could make the streamer go into a loop where is never stops dropping and loading different textures when out of budget.
+	 * @return		Priority value
+	 */
+	float CalcRetentionPriority( );
+
+	/**
+	 * Calculates a load order priority value for the texture. Higher value means more important.
+	 * Load Order can depend on the current state of resident mips, because it will not affect the streamer stability.
+	 * @return		Priority value
+	 */
+	float CalcLoadOrderPriority();
+
+		
+	/**
+	 * Not thread-safe: Sets the streaming index in the corresponding UTexture2D.
+	 * @param ArrayIndex	Index into the FStreamingManagerTexture::StreamingTextures array
+	 */
+	void SetStreamingIndex( int32 ArrayIndex )
+	{
+		Texture->StreamingIndex = ArrayIndex;
+	}
+
+	/** Update texture streaming. Returns whether it did anything */
+	bool UpdateMipCount(  FStreamingManagerTexture& Manager, FStreamingContext& Context );
+
+	/** Texture to manage. */
+	UTexture2D*		Texture;
+
+	/** Cached number of mip-maps in the UTexture2D mip array (including the base mip) */
+	int32			MipCount;
+	/** Cached number of mip-maps in memory (including the base mip) */
+	int32			ResidentMips;
+	/** Cached number of mip-maps requested by the streaming system. */
+	int32			RequestedMips;
+	/** Cached number of mip-maps we would like to have in memory. */
+	int32			WantedMips;
+	/** Legacy: Same as WantedMips, but not clamped by fudge factor. */
+	int32			PerfectWantedMips;
+	/** Same as MaxAllowedMips, but not clamped by the -reducepoolsize command line option. */
+	int32			MaxAllowedOptimalMips;
+#if STATS
+	/** Most number of mip-levels this texture has ever had resident in memory. */
+	int32			MostResidentMips;
+#endif
+	/** Minimum number of mip-maps that this texture must keep in memory. */
+	int32			MinAllowedMips;
+	/** Maximum number of mip-maps that this texture can have in memory. */
+	int32			MaxAllowedMips;
+	/** Cached memory sizes for each possible mipcount. */
+	int32			TextureSizes[MAX_TEXTURE_MIP_COUNT + 1];
+	/** Ref-count for how many ULevels want this texture fully-loaded. */
+	int32			ForceLoadRefCount;
+
+	/** Cached texture group. */
+	TextureGroup	LODGroup;
+	/** Cached LOD bias. */
+	int32			TextureLODBias;
+	/** Cached number of mipmaps that are not allowed to stream. */
+	int32			NumNonStreamingMips;
+	/** Cached number of cinematic (high-resolution) mip-maps. Normally not streamed in, unless the texture is forcibly fully loaded. */
+	int32			NumCinematicMipLevels;
+
+	/** Last time this texture was rendered, 0-based from GStartTime. */
+	float			LastRenderTime;
+	/** Distance to the closest instance of this texture, in world space units. Defaults to a specific value for non-staticmesh textures. */
+	float			MinDistance;
+	/** If non-zero, the most recent time an instance location was removed for this texture. */
+	double			InstanceRemovedTimestamp;
+
+	/** Set to FApp::GetCurrentTime() every time LastRenderTimeRefCount is modified. */
+	double			LastRenderTimeRefCountTimestamp;
+	/** Current number of instances that need LRT heuristics for this texture. */
+	int32			LastRenderTimeRefCount;
+
+	/**
+	 * Temporary boost of the streaming distance factor.
+	 * This factor is automatically reset to 1.0 after it's been used for mip-calculations.
+	 */
+	float			BoostFactor;
+
+	/** Whether the texture should be forcibly fully loaded. */
+	uint32			bForceFullyLoad : 1;
+	/** Whether the texture is ready to be streamed in/out (cached from IsReadyForStreaming()). */
+	uint32			bReadyForStreaming : 1;
+	/** Whether the texture is currently being streamed in/out. */
+	uint32			bInFlight : 1;
+	/** Whether this is a streaming lightmap or shadowmap (cached from IsStreamingLightmap()). */
+	uint32			bIsStreamingLightmap : 1;
+	/** Whether this is a lightmap or shadowmap. */
+	uint32			bIsLightmap : 1;
+	/** Whether this texture uses StaticTexture heuristics. */
+	uint32			bUsesStaticHeuristics : 1;
+	/** Whether this texture uses DynamicTexture heuristics. */
+	uint32			bUsesDynamicHeuristics : 1;
+	/** Whether this texture uses LastRenderTime heuristics. */
+	uint32			bUsesLastRenderHeuristics : 1;
+	/** Whether this texture uses ForcedIntoMemory heuristics. */
+	uint32			bUsesForcedHeuristics : 1;
+	/** Whether this texture uses the OrphanedTexture heuristics. */
+	uint32			bUsesOrphanedHeuristics : 1;
+	/** Whether this texture has a split request strategy. */
+	uint32			bHasSplitRequest : 1;
+	/** Whether this is the second request to converge to the wanted mip. */
+	uint32			bIsLastSplitRequest : 1;
+	/** Can this texture be affected by mip bias? */
+	uint32			bCanBeAffectedByMipBias : 1;
+	/** Is this texture a HLOD texture? */
+	//@TODO: Temporary: Done because there is currently no texture group for HLOD textures and adding one in a branch isn't ideal; this should be removed / replaced once there is a real HLOD texture group
+	uint32			bIsHLODTexture : 1;
+};
diff --git a/Engine/Source/Runtime/Engine/Private/Streaming/TextureInstanceManager.cpp b/Engine/Source/Runtime/Engine/Private/Streaming/TextureInstanceManager.cpp
new file mode 100644
--- /dev/null
+++ b/Engine/Source/Runtime/Engine/Private/Streaming/TextureInstanceManager.cpp
@@ -0,0 +1,635 @@
+// Copyright 1998-2016 Epic Games, Inc. All Rights Reserved.
+
+/*=============================================================================
+	TextureInstanceManager.cpp: Implementation of content streaming classes.
+=============================================================================*/
+
+#include "EnginePrivate.h"
+#include "TextureInstanceManager.h"
+
+FTextureInstanceManager::FBounds4::FBounds4()
+:	OriginX( 0, 0, 0, 0 )
+,	OriginY( 0, 0, 0, 0 )
+,	OriginZ( 0, 0, 0, 0 )
+,	ExtentX( 0, 0, 0, 0 )
+,	ExtentY( 0, 0, 0, 0 )
+,	ExtentZ( 0, 0, 0, 0 )
+,	Radius( 0, 0, 0, 0 )
+,	MinDistanceSq( 0, 0, 0, 0 )
+,	MaxDistanceSq( MAX_FLT, MAX_FLT, MAX_FLT, MAX_FLT )
+,	LastRenderTime(0, 0, 0, 0)
+{
+}
+
+
+void FTextureInstanceManager::FBounds4::Clear(int32 Index)
+{
+	check(Index >= 0 && Index < 4);
+
+	OriginX.Component(Index) = 0;
+	OriginY.Component(Index) = 0;
+	OriginZ.Component(Index) = 0;
+	ExtentX.Component(Index) = 0;
+	ExtentY.Component(Index) = 0;
+	ExtentZ.Component(Index) = 0;
+	Radius.Component(Index) = 0;
+	MinDistanceSq.Component(Index) = 0;
+	MaxDistanceSq.Component(Index) = MAX_FLT;
+	LastRenderTime.Component(Index) = 0;
+}
+
+void FTextureInstanceManager::FBounds4::Set(int32 Index, const FBoxSphereBounds& Bounds, float InLastRenderTime, float InMinDistance, float InMaxDistance)
+{
+	check(Index >= 0 && Index < 4);
+
+	OriginX.Component(Index) = Bounds.Origin.X;
+	OriginY.Component(Index) = Bounds.Origin.Y;
+	OriginZ.Component(Index) = Bounds.Origin.Z;
+	ExtentX.Component(Index) = Bounds.BoxExtent.X;
+	ExtentY.Component(Index) = Bounds.BoxExtent.Y;
+	ExtentZ.Component(Index) = Bounds.BoxExtent.Z;
+	Radius.Component(Index) = Bounds.SphereRadius;
+	MinDistanceSq.Component(Index) = InMinDistance * InMinDistance;
+	MaxDistanceSq.Component(Index) = InMaxDistance != MAX_FLT ? (InMaxDistance * InMaxDistance) : MAX_FLT;
+	LastRenderTime.Component(Index) = InLastRenderTime;
+}
+
+FTextureInstanceManager::FElement::FElement()
+:	Component(nullptr)
+,	Texture(nullptr)
+,	BoundsIndex(INDEX_NONE)
+,	TexelFactor(0)
+,	PrevTextureLink(INDEX_NONE)
+,	NextTextureLink(INDEX_NONE)
+,	NextComponentLink(INDEX_NONE)
+{
+}
+
+int32 FTextureInstanceManager::AddBounds(const FBoxSphereBounds& Bounds, const UPrimitiveComponent* InComponent, float LastRenderTime, float MinDistance, float MaxDistance)
+{
+	check(bAllowNullComponents || InComponent);
+
+	int BoundsIndex = INDEX_NONE;
+
+	if (FreeBoundIndices.Num())
+	{
+		BoundsIndex = FreeBoundIndices.Pop();
+	}
+	else
+	{
+		BoundsIndex = Bounds4.Num() * 4;
+		Bounds4.Push(FBounds4());
+
+		Bounds4Components.Push(nullptr);
+		Bounds4Components.Push(nullptr);
+		Bounds4Components.Push(nullptr);
+		Bounds4Components.Push(nullptr);
+
+		// Since each element contains 4 entries, add the 3 unused ones
+		FreeBoundIndices.Push(BoundsIndex + 1);
+		FreeBoundIndices.Push(BoundsIndex + 2);
+		FreeBoundIndices.Push(BoundsIndex + 3);
+
+		// This is just to be compatible with : check(Instances.FirstDirtyComponentIndex == Instances.Bounds4Components.Num());
+		// Whenever new bounds are allocated, they are automatically up to date.
+		if (FirstDirtyComponentIndex == BoundsIndex)
+		{
+			FirstDirtyComponentIndex += 4;
+		}
+	}
+
+	Bounds4[BoundsIndex / 4].Set(BoundsIndex % 4, Bounds, LastRenderTime, MinDistance, MaxDistance);
+	Bounds4Components[BoundsIndex] = InComponent;
+
+	return BoundsIndex;
+}
+
+void FTextureInstanceManager::RemoveBounds(int32 BoundsIndex)
+{
+	check(BoundsIndex != INDEX_NONE);
+	checkSlow(!FreeBoundIndices.Contains(BoundsIndex));
+
+	// If note all indices were freed
+	if (1 + FreeBoundIndices.Num() != Bounds4.Num() * 4)
+	{
+		FreeBoundIndices.Push(BoundsIndex);
+		Bounds4[BoundsIndex / 4].Clear(BoundsIndex % 4);
+		Bounds4Components[BoundsIndex] = nullptr;
+	}
+	else
+	{
+		Bounds4.Empty();
+		Bounds4Components.Empty();
+		FreeBoundIndices.Empty();
+	}
+}
+
+void FTextureInstanceManager::AddElement(const UPrimitiveComponent* InComponent, const UTexture2D* InTexture, int InBoundsIndex, float InTexelFactor, int32*& ComponentLink)
+{
+	check(bAllowNullComponents || InComponent);
+
+	int32 ElementIndex = INDEX_NONE;
+	if (FreeElementIndices.Num())
+	{
+		ElementIndex = FreeElementIndices.Pop();
+	}
+	else
+	{
+		ElementIndex = Elements.Num();
+		Elements.Push(FElement());
+	}
+
+	FElement& Element = Elements[ElementIndex];
+
+	Element.Component = InComponent;
+	Element.Texture = InTexture;
+	Element.BoundsIndex = InBoundsIndex;
+	Element.TexelFactor = InTexelFactor;
+
+	if (InTexture)
+	{
+		int32* TextureLink = TextureMap.Find(InTexture);
+		if (TextureLink)
+		{
+			FElement& TextureLinkElement = Elements[*TextureLink];
+			if (TextureLinkElement.Component == InComponent && 
+				TextureLinkElement.BoundsIndex == InBoundsIndex)
+			{
+				// Abort inserting a new element, and merge the 2 entries together.
+				TextureLinkElement.TexelFactor = FMath::Max(TextureLinkElement.TexelFactor, Element.TexelFactor);
+
+				// Clear the element and insert it in the free list.
+				Element = FElement();
+				FreeElementIndices.Push(ElementIndex);
+				return;
+			}
+
+			// The new inserted element as the head element.
+			Element.NextTextureLink = *TextureLink;
+			TextureLinkElement.PrevTextureLink = ElementIndex;
+			*TextureLink = ElementIndex;
+		}
+		else
+		{
+			TextureMap.Add(InTexture, ElementIndex);
+		}
+	}
+
+	if (InComponent)
+	{
+		checkSlow(ComponentLink == ComponentMap.Find(InComponent));
+		
+		if (ComponentLink)
+		{
+			// The new inserted element as the head element.
+			Element.NextComponentLink = *ComponentLink;
+			*ComponentLink = ElementIndex;
+		}
+		else
+		{
+			ComponentLink = &ComponentMap.Add(InComponent, ElementIndex);
+		}
+	}
+}
+
+void FTextureInstanceManager::RemoveElement(int32 ElementIndex, int32& NextComponentLink, int32& BoundsIndex, const UTexture2D*& Texture)
+{
+	FElement& Element = Elements[ElementIndex];
+	NextComponentLink = Element.NextComponentLink; 
+	BoundsIndex = Element.BoundsIndex; 
+
+	// Unlink textures
+	if (Element.Texture)
+	{
+		if (Element.PrevTextureLink == INDEX_NONE) // If NONE, that means this is the head of the texture list.
+		{
+			if (Element.NextTextureLink != INDEX_NONE) // Check if there are other entries for this texture.
+			{
+				 // Replace the head
+				*TextureMap.Find(Element.Texture) = Element.NextTextureLink;
+				Elements[Element.NextTextureLink].PrevTextureLink = INDEX_NONE;
+			}
+			else // Otherwise, remove the texture entry
+			{
+				TextureMap.Remove(Element.Texture);
+				Texture = Element.Texture;
+			}
+		}
+		else // Otherwise, just relink entries.
+		{
+			Elements[Element.PrevTextureLink].NextTextureLink = Element.NextTextureLink;
+
+			if (Element.NextTextureLink != INDEX_NONE)
+			{
+				Elements[Element.NextTextureLink].PrevTextureLink = Element.PrevTextureLink;
+			}
+		}
+	}
+
+
+	// Clear the element and insert in free list.
+	if (1 + FreeElementIndices.Num() != Elements.Num())
+	{
+		FreeElementIndices.Push(ElementIndex);
+		Element = FElement();
+	}
+	else
+	{
+		Elements.Empty();
+		FreeElementIndices.Empty();
+	}
+}
+
+void FTextureInstanceManager::AddComponent(const UPrimitiveComponent* Component)
+{
+	FStreamingTextureLevelContext LevelContext;
+	TArray<FStreamingTexturePrimitiveInfo> TextureInstanceInfos;
+	Component->GetStreamingTextureInfoWithNULLRemoval(LevelContext, TextureInstanceInfos);
+
+	bool bHasEntries = false;
+
+	int32 BoundsIndex = AddBounds(Component->Bounds, Component, Component->LastRenderTime);
+	int32* ComponentLink = ComponentMap.Find(Component);
+	for (int32 TextureIndex = 0; TextureIndex  < TextureInstanceInfos.Num(); ++TextureIndex)
+	{
+		const FStreamingTexturePrimitiveInfo& Info = TextureInstanceInfos[TextureIndex];
+		if (Info.TexelFactor > 0.0f && Info.Bounds.SphereRadius > 0.0f && ensure(FMath::IsFinite(Info.TexelFactor)))
+		{
+			const UTexture2D* Texture2D = Cast<UTexture2D>(Info.Texture);
+			if (Texture2D && IsStreamingTexture(Texture2D))
+			{
+				AddElement(Component, Texture2D, BoundsIndex, Info.TexelFactor, ComponentLink);
+				bHasEntries = true;
+			}
+		}
+	}
+
+	if (!bHasEntries)
+	{
+		RemoveBounds(BoundsIndex);
+	}
+}
+
+void FTextureInstanceManager::RemoveComponent(const UPrimitiveComponent* Component, FRemovedTextureArray& RemovedTextures)
+{
+	TArray<int32, TInlineAllocator<12> > RemovedBoundsIndices;
+	int32 ElementIndex = INDEX_NONE;
+
+	ComponentMap.RemoveAndCopyValue(Component, ElementIndex);
+	while (ElementIndex != INDEX_NONE)
+	{
+		int32 BoundsIndex = INDEX_NONE;
+		const UTexture2D* Texture = nullptr;
+
+		RemoveElement(ElementIndex, ElementIndex, BoundsIndex, Texture);
+
+		if (BoundsIndex != INDEX_NONE)
+		{
+			RemovedBoundsIndices.AddUnique(BoundsIndex);
+		}
+
+		if (Texture)
+		{
+			RemovedTextures.AddUnique(Texture);
+		}
+
+	};
+
+	for(int32 I = 0; I < RemovedBoundsIndices.Num(); ++I)
+	{
+		RemoveBounds(RemovedBoundsIndices[I]);
+	}
+}
+
+void FTextureInstanceManager::UpdateBounds(const UPrimitiveComponent* Component)
+{
+	int32* ComponentLink = ComponentMap.Find(Component);
+	if (ComponentLink)
+	{
+		int32 BoundsIndex = Elements[*ComponentLink].BoundsIndex;
+
+		// For primitives that can update bounds, there is only one bound per component.
+		Bounds4[BoundsIndex / 4].Set(BoundsIndex % 4, Component->Bounds, Component->LastRenderTime);
+	}
+}
+
+void FTextureInstanceManager::IncrementalUpdateBounds(float Percentage)
+{
+	const int32 LastIndex = FMath::Min(Bounds4Components.Num(), FirstDirtyComponentIndex + FMath::CeilToInt((float)Bounds4Components.Num() * Percentage));
+	for (int32 Index = FirstDirtyComponentIndex; Index < LastIndex; ++Index)
+	{
+		const UPrimitiveComponent* Component = Bounds4Components[Index];
+		if (Component)
+		{
+			Bounds4[Index / 4].Set(Index % 4, Component->Bounds, Component->LastRenderTime);
+		}
+	}
+	FirstDirtyComponentIndex = LastIndex;
+}
+
+/*
+void FTextureInstanceManager::SanityCheck()
+{
+	// FreeBoundIndices does not contain duplicate entries
+
+}
+*/
+
+void FTextureBoundsVisibility::Set(const FTextureInstanceManager& Instances)
+{
+	QUICK_SCOPE_CYCLE_COUNTER(FTextureBoundsVisibility_Set);
+
+	// Data must be valid at this point.
+	check(Instances.FirstDirtyComponentIndex == Instances.Bounds4Components.Num());
+
+	Bounds4 = Instances.Bounds4;
+	Elements = Instances.Elements;
+	TextureMap = Instances.TextureMap;
+}
+
+void FTextureBoundsVisibility::QuickSet(FTextureInstanceManager& Instances)
+{
+	QUICK_SCOPE_CYCLE_COUNTER(FTextureBoundsVisibility_QuickSet);
+
+	// Data must be valid at this point.
+	check(Instances.FirstDirtyComponentIndex >= Instances.Bounds4Components.Num());
+	check(!Instances.bAllowNullComponents);
+
+	// Set the two arrays to the same size then swap the buffers.
+	const int32 DeltaNum = Instances.Bounds4.Num() - Bounds4.Num();
+	if (DeltaNum > 0) // Bounds4 is too small!
+	{
+		Bounds4.AddUninitialized(DeltaNum);
+	}
+	else if (DeltaNum < 0) // Bounds4 is too big
+	{
+		Bounds4.RemoveAt(Instances.Bounds4.Num(), -DeltaNum);
+	}
+	FMemory::Memswap(&Bounds4, &Instances.Bounds4, sizeof(Bounds4));
+	Instances.FirstDirtyComponentIndex = 0; // Expecting an incremental update.
+
+	Elements = Instances.Elements;
+
+	// Swap the map and keep the clean version
+	TextureMap = Instances.TextureMap;
+	FMemory::Memswap(&TextureMap , &Instances.TextureMap, sizeof(TextureMap));
+}
+
+void FTextureBoundsVisibility::ComputeBoundsViewInfos(const TArray<FStreamingViewInfo>& ViewInfos, bool bUseAABB, float MaxEffectiveScreenSize)
+{
+	const int32 NumViews = ViewInfos.Num();
+	const VectorRegister One4 = VectorSet(1.f, 1.f, 1.f, 1.f);
+
+	BoundsViewInfo4.Empty(NumViews);
+	for (int32 ViewIndex = 0; ViewIndex < NumViews; ++ViewIndex)
+	{
+		BoundsViewInfo4.AddDefaulted();
+		BoundsViewInfo4[ViewIndex].AddUninitialized(Bounds4.Num());
+	}
+
+	for (int32 Bounds4Index = 0; Bounds4Index < Bounds4.Num(); ++Bounds4Index)
+	{
+		const FBounds4& CurrentBounds4 = Bounds4[Bounds4Index];
+
+		// Calculate distance of viewer to bounding sphere.
+		const VectorRegister OriginX = VectorLoadAligned( &CurrentBounds4.OriginX );
+		const VectorRegister OriginY = VectorLoadAligned( &CurrentBounds4.OriginY );
+		const VectorRegister OriginZ = VectorLoadAligned( &CurrentBounds4.OriginZ );
+		const VectorRegister ExtentX = VectorLoadAligned( &CurrentBounds4.ExtentX );
+		const VectorRegister ExtentY = VectorLoadAligned( &CurrentBounds4.ExtentY );
+		const VectorRegister ExtentZ = VectorLoadAligned( &CurrentBounds4.ExtentZ );
+		const VectorRegister Radius = VectorLoadAligned( &CurrentBounds4.Radius );
+		const VectorRegister MinDistanceSq = VectorLoadAligned( &CurrentBounds4.MinDistanceSq );
+		const VectorRegister MaxDistanceSq = VectorLoadAligned( &CurrentBounds4.MaxDistanceSq );
+
+		for (int32 ViewIndex = 0; ViewIndex < NumViews; ++ViewIndex)
+		{
+			const FStreamingViewInfo& ViewInfo = ViewInfos[ViewIndex];
+
+			const float EffectiveScreenSize = (MaxEffectiveScreenSize > 0.0f) ? FMath::Min(MaxEffectiveScreenSize, ViewInfo.ScreenSize) : ViewInfo.ScreenSize;
+			const float ScreenSizeFloat = EffectiveScreenSize * ViewInfo.BoostFactor;
+
+			const VectorRegister ScreenSize = VectorLoadFloat1( &ScreenSizeFloat );
+			const VectorRegister ViewOriginX = VectorLoadFloat1( &ViewInfo.ViewOrigin.X );
+			const VectorRegister ViewOriginY = VectorLoadFloat1( &ViewInfo.ViewOrigin.Y );
+			const VectorRegister ViewOriginZ = VectorLoadFloat1( &ViewInfo.ViewOrigin.Z );
+
+			//const float DistSq = FVector::DistSquared( ViewInfo.ViewOrigin, TextureInstance.BoundingSphere.Center );
+			VectorRegister Temp = VectorSubtract( ViewOriginX, OriginX );
+			VectorRegister DistSq = VectorMultiply( Temp, Temp );
+			Temp = VectorSubtract( ViewOriginY, OriginY );
+			DistSq = VectorMultiplyAdd( Temp, Temp, DistSq );
+			Temp = VectorSubtract( ViewOriginZ, OriginZ );
+			DistSq = VectorMultiplyAdd( Temp, Temp, DistSq );
+
+			VectorRegister ClampedDistSq = VectorMax( MinDistanceSq, DistSq );
+			ClampedDistSq = VectorMin( MaxDistanceSq, ClampedDistSq );
+			const VectorRegister InRangeMask = VectorCompareEQ(DistSq, ClampedDistSq);
+
+			VectorRegister DistSqMinusRadiusSq = VectorZero();
+			if (bUseAABB)
+			{
+				// In this case DistSqMinusRadiusSq will contain the distance to the box^2
+				Temp = VectorSubtract( ViewOriginX, OriginX );
+				Temp = VectorAbs( Temp );
+				VectorRegister BoxRef = VectorMin( Temp, ExtentX );
+				Temp = VectorSubtract( Temp, BoxRef );
+				DistSqMinusRadiusSq = VectorMultiply( Temp, Temp );
+
+				Temp = VectorSubtract( ViewOriginY, OriginY );
+				Temp = VectorAbs( Temp );
+				BoxRef = VectorMin( Temp, ExtentY );
+				Temp = VectorSubtract( Temp, BoxRef );
+				DistSqMinusRadiusSq = VectorMultiplyAdd( Temp, Temp, DistSqMinusRadiusSq );
+
+				Temp = VectorSubtract( ViewOriginZ, OriginZ );
+				Temp = VectorAbs( Temp );
+				BoxRef = VectorMin( Temp, ExtentZ );
+				Temp = VectorSubtract( Temp, BoxRef );
+				DistSqMinusRadiusSq = VectorMultiplyAdd( Temp, Temp, DistSqMinusRadiusSq );
+			}
+			else
+			{
+				DistSqMinusRadiusSq = VectorMultiply( Radius, Radius );
+				DistSqMinusRadiusSq = VectorSubtract( DistSq, DistSqMinusRadiusSq );
+			}
+			DistSqMinusRadiusSq = VectorMax(DistSqMinusRadiusSq, One4); // Prevents / 0
+			VectorRegister ScreenSizeOverDistance = VectorReciprocalSqrt(DistSqMinusRadiusSq);
+			ScreenSizeOverDistance = VectorMultiply(ScreenSizeOverDistance, ScreenSize);
+
+			// Store results
+			FBoundsViewInfo4& BoundsVieWInfo = BoundsViewInfo4[ViewIndex][Bounds4Index];
+			VectorStoreAligned(ScreenSizeOverDistance, &BoundsVieWInfo.ScreenSizeOverDistance);
+			VectorStoreAligned(DistSqMinusRadiusSq, &BoundsVieWInfo.DistanceSq);
+			VectorStoreAligned(InRangeMask, &BoundsVieWInfo.InRangeMask);
+		}
+	}
+}
+
+bool FTextureBoundsVisibility::GetTexelSize(const UTexture2D* InTexture, float& MaxSize_InRange, float& MinDistanceSq_InRange, float& MaxSize_AnyRange, float& MinDistanceSq_AnyRange) const
+{
+	MaxSize_InRange = 0;
+	MaxSize_AnyRange = 0;
+	MinDistanceSq_InRange = MAX_FLT;
+	MinDistanceSq_AnyRange = MAX_FLT;
+
+	const int32* TextureLink = TextureMap.Find(InTexture);
+	if (TextureLink)
+	{
+		int32 ElementIndex = *TextureLink;
+		while (ElementIndex != INDEX_NONE)
+		{
+			const FElement& Element = Elements[ElementIndex];
+			const int32 BoundsIndex = Element.BoundsIndex;
+
+			for (int32 ViewIndex = 0; ViewIndex < BoundsViewInfo4.Num(); ++ViewIndex)
+			{
+				FBoundsViewInfo4& BoundsVieWInfo = const_cast<FBoundsViewInfo4&>(BoundsViewInfo4[ViewIndex][BoundsIndex / 4]);
+				const float ScreenSizeOverDistance = BoundsVieWInfo.ScreenSizeOverDistance.Component(BoundsIndex % 4);
+				const float DistanceSq = BoundsVieWInfo.DistanceSq.Component(BoundsIndex % 4);
+				const int32 InRangeMask = reinterpret_cast<const int32&>(BoundsVieWInfo.InRangeMask.Component(BoundsIndex % 4));
+
+#if 0
+				//@TODO (for static instances)
+				// Used with TEXTUREGROUP_Terrain_Heightmap
+				// To check Forced LOD...
+				if (Element.TexelFactor < 0)
+				{
+					VectorRegister MinForcedLODs = VectorLoadAligned( &TextureInstance.TexelFactor );
+					bool AllForcedLOD = !!VectorAllLesserThan(MinForcedLODs, VectorZero());
+					MinForcedLODs = VectorMin( MinForcedLODs, VectorSwizzle(MinForcedLODs, 2, 3, 0, 1) );
+					MinForcedLODs = VectorMin( MinForcedLODs, VectorReplicate(MinForcedLODs, 1) );
+					float MinLODValue;
+					VectorStoreFloat1( MinForcedLODs, &MinLODValue );
+
+					if (MinLODValue <= 0)
+					{
+						WantedMipCount = FMath::Max(WantedMipCount, FFloatMipLevel::FromMipLevel(StreamingTexture.MaxAllowedMips - 13 - FMath::FloorToInt(MinLODValue)));
+						if (WantedMipCount >= FFloatMipLevel::FromMipLevel(StreamingTexture.MaxAllowedMips))
+					{
+						MinDistance = 1.0f;
+						bShouldAbortLoop = true;
+					}
+				}
+#endif
+
+				float Size = Element.TexelFactor * ScreenSizeOverDistance;
+
+				MaxSize_AnyRange = FMath::Max(Size, MaxSize_AnyRange);
+				MinDistanceSq_AnyRange = FMath::Min(DistanceSq, MinDistanceSq_AnyRange);
+				if (InRangeMask)
+				{
+					MaxSize_InRange = FMath::Max(Size, MaxSize_InRange);
+					MinDistanceSq_InRange = FMath::Min(DistanceSq, MinDistanceSq_InRange);
+
+					if (MinDistanceSq_InRange <= 1)
+					{
+						return true; // Already at max
+					}
+				}
+			}
+
+			ElementIndex = Element.NextTextureLink;
+		}
+	}
+
+	return TextureLink != nullptr;
+}
+
+
+void FDynamicComponentTextureManager::Add(const UPrimitiveComponent* Component, EDynamicPrimitiveType DynamicType, FRemovedTextureArray& RemovedTextures)
+{
+	FComponentState* State = ComponentStates.Find(Component);
+
+	if (!State)
+	{
+		ComponentStates.Add(Component, FComponentState(DynamicType));
+	}
+	else
+	{
+		if (State->bHasTextures) // If we reinsert without having cleared the old data correctly, make sure to clear it now.
+		{
+			TextureInstances.RemoveComponent(Component, RemovedTextures);
+			State->bHasTextures = false;
+		}
+		State->bToDelete = false;
+	}
+	DirtyComponents.Add(Component);
+}
+
+// Remove the component immediately, removing all links. Only keep the dynamic type as reference.
+void FDynamicComponentTextureManager::Remove(const UPrimitiveComponent* Component, FRemovedTextureArray& RemovedTextures)
+{
+	FComponentState* State = ComponentStates.Find(Component);
+
+	if (State)
+	{
+		if (State->bHasTextures) // When a component becomes invalid, remove any texture and bound references.
+		{
+			TextureInstances.RemoveComponent(Component, RemovedTextures);
+			State->bHasTextures = false;
+		}
+		State->bToDelete = true;
+		DirtyComponents.Add(Component);
+	}
+
+}
+
+void FDynamicComponentTextureManager::Update(const UPrimitiveComponent* Component, FRemovedTextureArray& RemovedTextures)
+{
+	FComponentState* State = ComponentStates.Find(Component);
+
+	if (State && !State->bToDelete && State->bHasTextures)
+	{
+		// If we reinserted without having cleared the old data correctly, make sure to clear it now.
+		TextureInstances.RemoveComponent(Component, RemovedTextures);
+		State->bHasTextures = false;
+		DirtyComponents.Add(Component);
+	}
+}
+
+void FDynamicComponentTextureManager::IncrementalTick(FRemovedTextureArray& RemovedTextures, float Percentage)
+{
+	QUICK_SCOPE_CYCLE_COUNTER(FDynamicComponentTextureManager_IncrementalTick);
+
+	TextureInstances.IncrementalUpdateBounds(Percentage);
+
+	TSet<const UPrimitiveComponent*> NextDirtyComponents; // Components that have a dirty states (to accelerate the update).
+	for (TSet<const UPrimitiveComponent*>::TIterator It(DirtyComponents); It; ++It)
+	{
+		const UPrimitiveComponent* Component = *It;
+
+		FComponentState* State = ComponentStates.Find(Component);
+		if (!State) continue;
+
+		if (State->bToDelete)
+		{
+			check(!State->bHasTextures);
+			ComponentStates.Remove(Component);
+		}
+		else
+		{
+			bool bIsInScene = Component->IsRegistered();
+			if (!bIsInScene)
+			{
+				if (State->bHasTextures)
+				{
+					TextureInstances.RemoveComponent(Component, RemovedTextures);
+					State->bHasTextures = false;
+				}
+				NextDirtyComponents.Add(Component);
+			}
+			else if (bIsInScene && !State->bHasTextures)
+			{
+				TextureInstances.AddComponent(Component);
+				State->bHasTextures = true;
+			}
+			else
+			{
+				TextureInstances.UpdateBounds(Component);
+			}
+		}
+	}
+
+	FMemory::Memswap(&DirtyComponents, &NextDirtyComponents, sizeof(DirtyComponents));
+}
diff --git a/Engine/Source/Runtime/Engine/Private/Streaming/TextureInstanceManager.h b/Engine/Source/Runtime/Engine/Private/Streaming/TextureInstanceManager.h
new file mode 100644
--- /dev/null
+++ b/Engine/Source/Runtime/Engine/Private/Streaming/TextureInstanceManager.h
@@ -0,0 +1,205 @@
+// Copyright 1998-2016 Epic Games, Inc. All Rights Reserved.
+
+/*=============================================================================
+TextureInstanceManager.h: Definitions of classes used for texture streaming.
+=============================================================================*/
+
+#pragma once
+
+#include "TextureStreamingHelpers.h"
+
+// Main Thread Job Requirement : find all instance of a component and update it's bound.
+// Threaded Job Requirement : get the list of instance texture easily from the list of 
+
+// Can be used either for static primitives or dynamic primitives
+class FTextureInstanceManager
+{
+public:
+
+	// The bounds are into their own arrays to allow SIMD-friendly processing
+	struct FBounds4
+	{
+		FORCEINLINE FBounds4();
+
+		/** X coordinates for the bounds origin of 4 texture instances */
+		FVector4 OriginX;
+		/** Y coordinates for the bounds origin of 4 texture instances */
+		FVector4 OriginY;
+		/** Z coordinates for the bounds origin of 4 texture instances */
+		FVector4 OriginZ;
+
+		/** X size of the bounds box extent of 4 texture instances */
+		FVector4 ExtentX;
+		/** Y size of the bounds box extent of 4 texture instances */
+		FVector4 ExtentY;
+		/** Z size of the bounds box extent of 4 texture instances */
+		FVector4 ExtentZ;
+
+		/** Sphere radii for the bounding sphere of 4 texture instances */
+		FVector4 Radius;
+		/** Minimal distance (between the bounding sphere origin and the view origin) for which this entry is valid */
+		FVector4 MinDistanceSq;
+		/** Maximal distance (between the bounding sphere origin and the view origin) for which this entry is valid */
+		FVector4 MaxDistanceSq;
+
+		/** Last visibility time for this bound, used for priority */
+		FVector4 LastRenderTime;
+
+		FORCEINLINE_DEBUGGABLE void Set(int32 Index, const FBoxSphereBounds& Bounds, float LastRenderTime, float MinDistance = 0, float MaxDistance = MAX_FLT);
+
+		// Clears entry between 0 and 4
+		FORCEINLINE void Clear(int32 Index);
+	};
+
+	struct FElement
+	{
+		FORCEINLINE FElement();
+
+		const UPrimitiveComponent* Component; // Which component this relates too
+		const UTexture2D* Texture;	// Texture, never dereferenced.
+
+		int32 BoundsIndex;		// The Index associated to this component (static component can have several bounds).
+		float TexelFactor;		// The texture scale to be applied to this instance.
+
+		int32 PrevTextureLink;	// The previous element which uses the same texture as this Element. The first element referred by TextureMap will have INDEX_NONE.
+		int32 NextTextureLink;	// The next element which uses the same texture as this Element. Last element will have INDEX_NONE
+
+		// Components are always updated as a whole, so individual elements can not be removed.
+		int32 NextComponentLink;	// The next element that uses the same component as this Element. The first one is referred by ComponentMap and the last one will have INDEX_NONE.
+	};
+
+public:
+
+	FTextureInstanceManager(bool InAllowNullComponents = false) : bAllowNullComponents(InAllowNullComponents), FirstDirtyComponentIndex(0) {}
+
+	int32 AddBounds(const FBoxSphereBounds& Bounds, const UPrimitiveComponent* Component, float LastRenderTime, float MinDistance = 0 , float MaxDistance = MAX_FLT);
+
+	void AddElement(const UPrimitiveComponent* Component, const UTexture2D* Texture, int BoundsIndex, float TexelFactor, int32*& ComponentLink);
+
+	// Returns the next elements using the same component.
+	void RemoveElement(int32 ElementIndex, int32& NextComponentLink, int32& BoundsIndex, const UTexture2D*& Texture);
+
+	// Will also remove bounds
+	void AddComponent(const UPrimitiveComponent* Component);
+	void RemoveComponent(const UPrimitiveComponent* Component, FRemovedTextureArray& RemovedTextures);
+
+	// Update the component bounds, assumed that this would only be used with component using component bounds
+	void UpdateBounds(const UPrimitiveComponent* Component);
+
+	bool Contains(const UPrimitiveComponent* Component) const { return ComponentMap.Contains(Component); }
+
+	void IncrementalUpdateBounds(float Percentage);
+
+private:
+
+	void RemoveBounds(int32 Index);
+
+private:
+
+	TArray<FBounds4> Bounds4;
+	TArray<FElement> Elements;
+
+	TArray<int32> FreeBoundIndices;
+	TArray<int32> FreeElementIndices;
+
+	TMap<const UTexture2D*, int32> TextureMap;
+	TMap<const UPrimitiveComponent*, int32> ComponentMap;
+
+	bool bAllowNullComponents;
+
+	// Used to implement iterative bounds update and visibility update.
+	int32 FirstDirtyComponentIndex; // Ranges from 0 to Components.Num()
+	TArray<const UPrimitiveComponent*> Bounds4Components; // The components mapped to Bounds4
+
+	friend class FTextureBoundsVisibility;
+};
+
+// Data used to compute visibility
+class FTextureBoundsVisibility
+{
+public:
+
+	void Set(const FTextureInstanceManager& Instances);
+	void QuickSet(FTextureInstanceManager& Instances);
+
+	void ComputeBoundsViewInfos(const TArray<FStreamingViewInfo>& ViewInfos, bool bUseAABB, float MaxEffectiveScreenSize);
+
+	// MaxSize_InRange : Biggest instance size for in range views
+	// MaxSize_AnyRange : Biggest instance size not considering range.
+	bool GetTexelSize(const UTexture2D* InTexture, float& MaxSize_InRange, float& MinDistanceSq_InRange, float& MaxSize_AnyRange, float& MinDistanceSq_AnyRange) const;
+
+private:
+
+	typedef FTextureInstanceManager::FBounds4 FBounds4;
+	typedef FTextureInstanceManager::FElement FElement;
+
+	TArray<FBounds4> Bounds4;
+	TArray<FElement> Elements;
+	TMap<const UTexture2D*, int32> TextureMap;
+
+	struct FBoundsViewInfo4
+	{
+		FVector4 ScreenSizeOverDistance;
+		FVector4 DistanceSq;
+		FVector4 InRangeMask;
+	};
+
+	// Normalized Texel Factors for each bounds and view.
+	// @TODO : store data for different views continuously to improve reads.
+	TArray< TArray<FBoundsViewInfo4> > BoundsViewInfo4;
+};
+
+
+// Needs to be able to remove to defragment every time there are 10% free elements
+// Defragmentation is time sliced. So we only need update a small percentage each frame.
+// To reduce the number of reallocations, low indices must be used first and high last.
+
+// Tick is timesliced. This requires Component to be valid.
+class FDynamicComponentTextureManager
+{
+public:
+
+	// Add a component entry, without actually processing the component now.
+	void Add(const UPrimitiveComponent* Component, EDynamicPrimitiveType DynamicType, FRemovedTextureArray& RemovedTextures);
+
+	// Remove the component immediately, removing all links. Only keep the dynamic type as reference.
+	void Remove(const UPrimitiveComponent* Component, FRemovedTextureArray& RemovedTextures);
+
+	void Update(const UPrimitiveComponent* Component, FRemovedTextureArray& RemovedTextures);
+
+	void IncrementalTick(FRemovedTextureArray& RemovedTextures, float Percentage);
+
+	const FTextureInstanceManager& GetTextureInstances() const { return TextureInstances; }
+	FTextureInstanceManager& GetTextureInstances() { return TextureInstances; }
+
+	bool IsDeleted(const UPrimitiveComponent* Component) const 
+	{ 	
+		const FComponentState* State = ComponentStates.Find(Component);
+		return !State || State->bToDelete;
+	}
+
+private:
+
+	struct FComponentState
+	{
+		FComponentState(EDynamicPrimitiveType InDynamicType) : Data(0) 
+		{ 
+			DynamicType = InDynamicType; 
+		}
+
+		union 
+		{
+			uint32 Data;
+			struct 
+			{
+				uint32 DynamicType : 8;
+				uint32 bHasTextures : 1; // True if this is a new entry for which to update the texture lists.
+				uint32 bToDelete : 1;	// True if this entry needs to be deleted in the next update.
+			};
+		};
+	};
+
+	TMap<const UPrimitiveComponent*, FComponentState> ComponentStates;
+	TSet<const UPrimitiveComponent*> DirtyComponents; // Components that have a dirty states (to accelerate the incremental update).
+	FTextureInstanceManager TextureInstances;
+};
diff --git a/Engine/Source/Runtime/Engine/Private/Streaming/TextureStreamingBuild.cpp b/Engine/Source/Runtime/Engine/Private/Streaming/TextureStreamingBuild.cpp
new file mode 100644
--- /dev/null
+++ b/Engine/Source/Runtime/Engine/Private/Streaming/TextureStreamingBuild.cpp
@@ -0,0 +1,320 @@
+// Copyright 1998-2016 Epic Games, Inc. All Rights Reserved.
+
+/*=============================================================================
+TextureStreamingBuild.cpp : Contains definitions to build texture streaming data.
+=============================================================================*/
+
+#include "EnginePrivate.h"
+#include "DebugViewModeMaterialProxy.h"
+#include "ShaderCompiler.h"
+
+/**
+ * Build Shaders to compute tex coord scale per texture.
+ *
+ * @param InWorld			The world to build streaming for.
+ * @param QualityLevel		The quality level for the shaders.
+ * @param FeatureLevel		The feature level for the shaders.
+ * @param TexCoordScales	The map to store material interfaces used by the current primitive world.
+ */
+ENGINE_API void BuildTextureStreamingShaders(UWorld* InWorld, EMaterialQualityLevel::Type QualityLevel, ERHIFeatureLevel::Type FeatureLevel, OUT FTexCoordScaleMap& TexCoordScales)
+{
+#if WITH_EDITORONLY_DATA
+
+	if (!InWorld) return;
+
+	const bool bUseNewMetrics = CVarStreamingUseNewMetrics.GetValueOnGameThread() != 0;
+	const double StartTime = FPlatformTime::Seconds();
+
+	FlushRenderingCommands();
+
+	FDebugViewModeMaterialProxy::ClearAllShaders();
+
+	// Without the new metrics, the shaders will are not compiled, making the texcoord scale viewmode non functional.
+	if (bUseNewMetrics && AllowDebugViewShaderMode(DVSM_TexCoordScaleAnalysis, FeatureLevel))
+	{
+		for (int32 LevelIndex = 0; LevelIndex < InWorld->GetNumLevels(); LevelIndex++)
+		{
+			ULevel* Level = InWorld->GetLevel(LevelIndex);
+
+			TArray<UObject*> ObjectsInOuter;
+			GetObjectsWithOuter(Level, ObjectsInOuter);
+
+			for (int32 Index = 0; Index < ObjectsInOuter.Num(); Index++)
+			{
+				UPrimitiveComponent* Primitive = Cast<UPrimitiveComponent>(ObjectsInOuter[Index]);
+				if (!Primitive) continue;
+
+				TArray<UMaterialInterface*> Materials;
+				Primitive->GetUsedMaterials(Materials);
+
+				for (int32 MaterialIndex = 0; MaterialIndex < Materials.Num(); ++MaterialIndex)
+				{
+					UMaterialInterface* MaterialInterface = Materials[MaterialIndex];
+					if (!MaterialInterface) continue;
+
+					if (!TexCoordScales.Contains(MaterialInterface))
+					{
+						TexCoordScales.Add(MaterialInterface);
+					}
+
+					FDebugViewModeMaterialProxy::AddShader(MaterialInterface, QualityLevel, FeatureLevel, EMaterialShaderMapUsage::DebugViewModeTexCoordScale);
+				}
+			}
+		}
+
+		// This could be very slow if there were already other shaders compiling
+		GShaderCompilingManager->FinishAllCompilation();
+	}
+
+
+	UE_LOG(LogLevel, Display, TEXT("Build Texture Streaming Shaders took %.3f seconds."), FPlatformTime::Seconds() - StartTime);
+#else
+	UE_LOG(LogLevel, Fatal,TEXT("Build Texture Streaming Shaders should not be called on a console"));
+#endif
+}
+
+ENGINE_API void BuildTextureStreamingData(UWorld* InWorld, const FTexCoordScaleMap& InTexCoordScales, EMaterialQualityLevel::Type QualityLevel, ERHIFeatureLevel::Type FeatureLevel)
+{
+#if WITH_EDITORONLY_DATA
+	if (!InWorld) return;
+	const double StartTime = FPlatformTime::Seconds();
+
+	// ====================================================
+	// Build per primitive data
+	// ====================================================
+
+	for (int32 LevelIndex = 0; LevelIndex < InWorld->GetNumLevels(); LevelIndex++)
+	{
+		ULevel* Level = InWorld->GetLevel(LevelIndex);
+		if (!Level) continue;
+
+		TArray<UObject*> ObjectsInOuter;
+		GetObjectsWithOuter(Level, ObjectsInOuter);
+
+		TArray<UTexture2D*> LevelTextures;
+		for (int32 Index = 0; Index < ObjectsInOuter.Num(); Index++)
+		{
+			UStaticMeshComponent* Primitive = Cast<UStaticMeshComponent>(ObjectsInOuter[Index]);
+			if (!Primitive) continue;
+
+			Primitive->UpdateStreamingTextureData(LevelTextures, InTexCoordScales, QualityLevel, FeatureLevel);
+		}
+
+		Level->bTextureStreamingRotationChanged = false;
+		Level->StreamingTextureGuids.Empty(LevelTextures.Num());
+
+		// Reset LevelIndex to default for next use and build the level Guid array.
+		for (int32 TextureIndex = 0; TextureIndex < LevelTextures.Num(); ++TextureIndex)
+		{
+			UTexture2D* Texture2D = LevelTextures[TextureIndex];
+			Level->StreamingTextureGuids.Push(Texture2D->GetLightingGuid());
+			Texture2D->LevelIndex = INDEX_NONE;
+		}
+	}
+
+	// ====================================================
+	// Update texture streaming data
+	// ====================================================
+
+	ULevel::BuildStreamingData(InWorld);
+
+	// ====================================================
+	// Reregister everything for debug view modes to reflect changes
+	// ====================================================
+
+	for (int32 LevelIndex = 0; LevelIndex < InWorld->GetNumLevels(); LevelIndex++)
+	{
+		ULevel* Level = InWorld->GetLevel(LevelIndex);
+		if (!Level) continue;
+		
+		for (AActor* Actor : Level->Actors)
+		{
+			if (!Actor) continue;
+
+			Actor->MarkComponentsRenderStateDirty();
+		}
+	}
+	UE_LOG(LogLevel, Display, TEXT("Build Texture Streaming took %.3f seconds."), FPlatformTime::Seconds() - StartTime);
+#else
+	UE_LOG(LogLevel, Fatal,TEXT("Build Texture Streaming should not be called on a console"));
+#endif
+}
+
+
+static uint32 PackRelativeBox(const FVector& RefOrigin, const FVector& RefExtent, const FVector& Origin, const FVector& Extent)
+{
+	const FVector RefMin = RefOrigin - RefExtent;
+	// 15.5 and 31.5 have the / 2 scale included 
+	const FVector PackScale = FVector(15.5f, 15.5f, 31.5f) / RefExtent.ComponentMax(FVector(KINDA_SMALL_NUMBER, KINDA_SMALL_NUMBER, KINDA_SMALL_NUMBER));
+
+	const FVector Min = Origin - Extent;
+	const FVector Max = Origin + Extent;
+
+	const FVector RelMin = (Min - RefMin) * PackScale;
+	const FVector RelMax = (Max - RefMin) * PackScale;
+
+	const uint32 PackedMinX = (uint32)FMath::Clamp<int32>(FMath::FloorToInt(RelMin.X), 0, 31);
+	const uint32 PackedMinY = (uint32)FMath::Clamp<int32>(FMath::FloorToInt(RelMin.Y), 0, 31);
+	const uint32 PackedMinZ = (uint32)FMath::Clamp<int32>(FMath::FloorToInt(RelMin.Z), 0, 63);
+
+	const uint32 PackedMaxX = (uint32)FMath::Clamp<int32>(FMath::CeilToInt(RelMax.X), 0, 31);
+	const uint32 PackedMaxY = (uint32)FMath::Clamp<int32>(FMath::CeilToInt(RelMax.Y), 0, 31);
+	const uint32 PackedMaxZ = (uint32)FMath::Clamp<int32>(FMath::CeilToInt(RelMax.Z), 0, 63);
+
+	return PackedMinX | (PackedMinY << 5) | (PackedMinZ << 10) | (PackedMaxX << 16) | (PackedMaxY << 21) | (PackedMaxZ << 26);
+}
+
+static void UnpackRelativeBox(const FVector& InRefOrigin, const FVector& InRefExtent, uint32 InPackedRelBox, FVector& OutOrigin, FVector& OutExtent)
+{
+	const uint32 PackedMinX = InPackedRelBox & 31;
+	const uint32 PackedMinY = (InPackedRelBox >> 5) & 31;
+	const uint32 PackedMinZ = (InPackedRelBox >> 10) & 63;
+
+	const uint32 PackedMaxX = (InPackedRelBox >> 16) & 31;
+	const uint32 PackedMaxY = (InPackedRelBox >> 21) & 31;
+	const uint32 PackedMaxZ = (InPackedRelBox >> 26) & 63;
+
+	const FVector RefMin = InRefOrigin - InRefExtent;
+	// 15.5 and 31.5 have the / 2 scale included 
+	const FVector UnpackScale = InRefExtent.ComponentMax(FVector(KINDA_SMALL_NUMBER, KINDA_SMALL_NUMBER, KINDA_SMALL_NUMBER)) / FVector(15.5f, 15.5f, 31.5f);
+
+	const FVector Min = FVector((float)PackedMinX, (float)PackedMinY, (float)PackedMinZ) * UnpackScale + RefMin;
+	const FVector Max = FVector((float)PackedMaxX, (float)PackedMaxY, (float)PackedMaxZ) * UnpackScale + RefMin;
+
+	OutOrigin = .5 * (Min + Max);
+	OutExtent = .5 * (Max - Min);
+}
+
+void FStreamingTexturePrimitiveInfo::UnPackFrom(UTexture2D* InTexture, const FBoxSphereBounds& RefBounds, const FStreamingTextureBuildInfo& Info, bool bUseRelativeBox)
+{
+	check(InTexture->LevelIndex == Info.TextureLevelIndex);
+	Texture = InTexture;
+
+	if (bUseRelativeBox)
+	{
+		UnpackRelativeBox(RefBounds.Origin, RefBounds.BoxExtent, Info.PackedRelativeBox, Bounds.Origin, Bounds.BoxExtent);
+		Bounds.SphereRadius = Bounds.BoxExtent.Size();
+	}
+	else
+	{
+		Bounds = RefBounds;
+	}
+
+	TexelFactor = Info.TexelFactor;
+}
+
+void FStreamingTextureBuildInfo::PackFrom(TArray<UTexture2D*>& LevelTextures, const FBoxSphereBounds& RefBounds, const FStreamingTexturePrimitiveInfo& Info)
+{
+	PackedRelativeBox = PackRelativeBox(RefBounds.Origin, RefBounds.BoxExtent, Info.Bounds.Origin, Info.Bounds.BoxExtent);
+
+	UTexture2D* Texture2D = CastChecked<UTexture2D>(Info.Texture);
+	if (Texture2D->LevelIndex == INDEX_NONE)
+	{
+		// If this is the first time this texture gets processed in the packing process, encode it.
+		Texture2D->LevelIndex = LevelTextures.Add(Texture2D);
+	}
+	TextureLevelIndex = (uint16)Texture2D->LevelIndex;
+
+	TexelFactor = Info.TexelFactor;
+}
+
+
+FStreamingTextureLevelContext::FStreamingTextureLevelContext(const ULevel* InLevel) :
+	bUseRelativeBoxes(InLevel ? !InLevel->bTextureStreamingRotationChanged : false),
+	ComponentTimestamp(0),
+	ComponentBounds(ForceInitToZero),
+	ComponentMeshScale(1.f),
+	ComponentStreamingScale(1.f)
+{
+	if (InLevel)
+	{
+		for (int32 TextureIndex = 0; TextureIndex < InLevel->StreamingTextureGuids.Num(); ++TextureIndex)
+		{
+			TextureGuidToLevelIndex.Add(InLevel->StreamingTextureGuids[TextureIndex], TextureIndex);
+		}
+		BoundStates.AddZeroed(InLevel->StreamingTextureGuids.Num());
+	}
+}
+
+FStreamingTextureLevelContext::~FStreamingTextureLevelContext()
+{
+	// Reset the level indices for the next use.
+	for (UTexture2D* Texture : ProcessedTextures)
+	{
+		Texture->LevelIndex = INDEX_NONE;
+	}
+}
+
+void FStreamingTextureLevelContext::BindComponent(const TArray<FStreamingTextureBuildInfo>* BuildData, const FBoxSphereBounds& Bounds, float MeshScale, float StreamingScale)
+{
+	// First processed component must use a timestamp > 0  so that default value of 0 is invalid.
+	++ComponentTimestamp;
+
+	ComponentBounds = Bounds;
+	ComponentMeshScale = MeshScale;
+	ComponentStreamingScale = StreamingScale;
+	ComponentBuildData = BuildData;
+
+	if (BuildData)
+	{
+		for (int32 Index = 0; Index < BuildData->Num(); ++Index)
+		{
+			int32 TextureLevelIndex = (*BuildData)[Index].TextureLevelIndex;
+			if (BoundStates.IsValidIndex(TextureLevelIndex))
+			{
+				FTextureBoundState& BoundState = BoundStates[TextureLevelIndex];
+				BoundState.BuildDataIndex = Index;
+				BoundState.BuildDataTimestamp = ComponentTimestamp;
+			}
+		}
+	}
+}
+
+void FStreamingTextureLevelContext::Process(const TArray<UTexture*>& InTextures, TArray<FStreamingTexturePrimitiveInfo>& OutInfos)
+{
+	for (int32 TextureIndex = 0; TextureIndex < InTextures.Num(); ++TextureIndex)
+	{
+		UTexture2D* Texture2D = Cast<UTexture2D>(InTextures[TextureIndex]);
+		if (!Texture2D) continue; // Only interested in 2d texture for streaming
+
+		if (Texture2D->LevelIndex == INDEX_NONE)
+		{
+			int32* LevelIndex = TextureGuidToLevelIndex.Find(Texture2D->GetLightingGuid());
+			if (LevelIndex)
+			{
+				Texture2D->LevelIndex = *LevelIndex;
+			}
+			else // If it is not in build data, add an entry so that we can process it in a uniform way.
+			{
+				Texture2D->LevelIndex = BoundStates.AddZeroed();
+			}
+			ProcessedTextures.Push(Texture2D);
+		}
+
+		FTextureBoundState& BoundState = BoundStates[Texture2D->LevelIndex];
+
+		// Ignore this texture if it has already been handled for the current component.
+		if (BoundState.Timestamp != ComponentTimestamp)
+		{
+			FStreamingTexturePrimitiveInfo& StreamingTexture = *new(OutInfos) FStreamingTexturePrimitiveInfo;
+
+			// Check if there is any precomputed data
+			if (BoundState.BuildDataTimestamp == ComponentTimestamp && ComponentBuildData)
+			{
+				StreamingTexture.UnPackFrom(Texture2D, ComponentBounds, (*ComponentBuildData)[BoundState.BuildDataIndex], bUseRelativeBoxes);
+			}
+			else // Otherwise add default entry.
+			{
+				StreamingTexture.Texture = Texture2D;
+				StreamingTexture.Bounds = ComponentBounds;
+				StreamingTexture.TexelFactor = ComponentMeshScale;
+			}
+
+			// Component scale is not included in the build data, so that changing it does not invalidate the data.
+			StreamingTexture.TexelFactor *= ComponentStreamingScale;
+
+			// Mark this texture as processed for this component.
+			BoundState.Timestamp = ComponentTimestamp;
+		}
+	}
+}
