Copying //UE4/Dev-Rendering to Dev-Main (Source //UE4/Dev-Rendering@2932636)

diff --git a/Engine/Source/Runtime/Engine/Private/Streaming/TextureStreamingHelpers.cpp b/Engine/Source/Runtime/Engine/Private/Streaming/TextureStreamingHelpers.cpp
new file mode 100644
--- /dev/null
+++ b/Engine/Source/Runtime/Engine/Private/Streaming/TextureStreamingHelpers.cpp
@@ -0,0 +1,275 @@
+// Copyright 1998-2016 Epic Games, Inc. All Rights Reserved.
+
+/*=============================================================================
+TextureStreamingHelpers.cpp: Definitions of classes used for texture streaming.
+=============================================================================*/
+
+#include "EnginePrivate.h"
+#include "TextureStreamingHelpers.h"
+#include "Engine/Texture2D.h"
+#include "GenericPlatformMemoryPoolStats.h"
+
+/** Streaming stats */
+
+DEFINE_STAT(STAT_GameThreadUpdateTime);
+DEFINE_STAT(STAT_TexturePoolSize);
+DEFINE_STAT(STAT_TexturePoolAllocatedSize);
+DEFINE_STAT(STAT_OptimalTextureSize);
+DEFINE_STAT(STAT_EffectiveStreamingPoolSize);
+DEFINE_STAT(STAT_StreamingTexturesSize);
+DEFINE_STAT(STAT_NonStreamingTexturesSize);
+DEFINE_STAT(STAT_StreamingTexturesMaxSize);
+DEFINE_STAT(STAT_StreamingOverBudget);
+DEFINE_STAT(STAT_StreamingUnderBudget);
+DEFINE_STAT(STAT_NumWantingTextures);
+
+DEFINE_STAT(STAT_StreamingTextures);
+
+DEFINE_STAT(STAT_TotalStaticTextureHeuristicSize);
+DEFINE_STAT(STAT_TotalDynamicHeuristicSize);
+DEFINE_STAT(STAT_TotalLastRenderHeuristicSize);
+DEFINE_STAT(STAT_TotalForcedHeuristicSize);
+
+DEFINE_STAT(STAT_LightmapMemorySize);
+DEFINE_STAT(STAT_LightmapDiskSize);
+DEFINE_STAT(STAT_HLODTextureMemorySize);
+DEFINE_STAT(STAT_HLODTextureDiskSize);
+DEFINE_STAT(STAT_IntermediateTexturesSize);
+DEFINE_STAT(STAT_RequestSizeCurrentFrame);
+DEFINE_STAT(STAT_RequestSizeTotal);
+DEFINE_STAT(STAT_LightmapRequestSizeTotal);
+
+DEFINE_STAT(STAT_IntermediateTextures);
+DEFINE_STAT(STAT_RequestsInCancelationPhase);
+DEFINE_STAT(STAT_RequestsInUpdatePhase);
+DEFINE_STAT(STAT_RequestsInFinalizePhase);
+DEFINE_STAT(STAT_StreamingLatency);
+DEFINE_STAT(STAT_StreamingBandwidth);
+DEFINE_STAT(STAT_GrowingReallocations);
+DEFINE_STAT(STAT_ShrinkingReallocations);
+DEFINE_STAT(STAT_FullReallocations);
+DEFINE_STAT(STAT_FailedReallocations);
+DEFINE_STAT(STAT_PanicDefragmentations);
+DEFINE_STAT(STAT_DynamicStreamingTotal);
+DEFINE_STAT(STAT_NumVisibleTexturesWithLowResolutions);
+DEFINE_STAT(STAT_LevelsWithPostLoadRotations);
+
+DEFINE_LOG_CATEGORY(LogContentStreaming);
+
+ENGINE_API TAutoConsoleVariable<int32> CVarStreamingUseNewMetrics(
+	TEXT("r.Streaming.UseNewMetrics"),
+	0,
+	TEXT("If non-zero, will use tight AABB bounds and improved texture factors."),
+	ECVF_Default);
+
+
+/** Whether the texture pool size has been artificially lowered for testing. */
+bool GIsOperatingWithReducedTexturePool = false;
+
+TAutoConsoleVariable<float> CVarStreamingBoost(
+	TEXT("r.Streaming.Boost"),
+	1.0f,
+	TEXT("=1.0: normal\n")
+	TEXT("<1.0: decrease wanted mip levels\n")
+	TEXT(">1.0: increase wanted mip levels"),
+	ECVF_Default
+	);
+
+TAutoConsoleVariable<float> CVarStreamingScreenSizeEffectiveMax(
+	TEXT("r.Streaming.MaxEffectiveScreenSize"),
+	0,
+	TEXT("0: Use current actual vertical screen size\n")	
+	TEXT("> 0: Clamp wanted mip size calculation to this value for the vertical screen size component."),
+	ECVF_Scalability
+	);
+
+#if PLATFORM_SUPPORTS_TEXTURE_STREAMING
+TAutoConsoleVariable<int32> CVarSetTextureStreaming(
+	TEXT("r.TextureStreaming"),
+	1,
+	TEXT("Allows to define if texture streaming is enabled, can be changed at run time.\n")
+	TEXT("0: off\n")
+	TEXT("1: on (default)"),
+	ECVF_Default | ECVF_RenderThreadSafe);
+#endif
+
+TAutoConsoleVariable<int32> CVarStreamingUseFixedPoolSize(
+	TEXT("r.Streaming.UseFixedPoolSize"),
+	0,
+	TEXT("If non-zero, do not allow the pool size to change at run time."),
+	ECVF_ReadOnly);
+
+TAutoConsoleVariable<int32> CVarStreamingPoolSize(
+	TEXT("r.Streaming.PoolSize"),
+	-1,
+	TEXT("-1: Default texture pool size, otherwise the size in MB"),
+	ECVF_Scalability);
+
+TAutoConsoleVariable<int32> CVarStreamingMaxTempMemoryAllowed(
+	TEXT("r.Streaming.MaxTempMemoryAllowed"),
+	50,
+	TEXT("Maximum temporary memory used when streaming in or out texture mips.\n")
+	TEXT("This memory contains mips used for the new updated texture.\n")
+	TEXT("The value must be high enough to not be a limiting streaming speed factor.\n"),
+	ECVF_Default);
+
+TAutoConsoleVariable<int32> CVarStreamingShowWantedMips(
+	TEXT("r.Streaming.ShowWantedMips"),
+	0,
+	TEXT("If non-zero, will limit resolution to wanted mip."),
+	ECVF_Cheat);
+
+TAutoConsoleVariable<int32> CVarStreamingHLODStrategy(
+	TEXT("r.Streaming.HLODStrategy"),
+	0,
+	TEXT("Define the HLOD streaming strategy.\n")
+	TEXT("0: stream\n")
+	TEXT("1: stream only mip 0\n")
+	TEXT("2: disable streaming"),
+	ECVF_Default);
+
+TAutoConsoleVariable<float> CVarStreamingHiddenPrimitiveScale(
+	TEXT("r.Streaming.HiddenPrimitiveScale"),
+	0.5,
+	TEXT("Define the resolution scale to apply when not in range.\n")
+	TEXT(".5: drop one mip\n")
+	TEXT("1: ignore visiblity"),
+	ECVF_Default
+	);
+
+TAutoConsoleVariable<int32> CVarStreamingSplitRequestSizeThreshold(
+	TEXT("r.Streaming.SplitRequestSizeThreshold"),
+	0, // Good value : 2 * 1024 * 1024,
+	TEXT("Define how many IO pending request the streamer can generate.\n"),
+	ECVF_Default
+	);
+
+// Used for scalability (GPU memory, streaming stalls)
+TAutoConsoleVariable<float> CVarStreamingMipBias(
+	TEXT("r.Streaming.MipBias"),
+	0.0f,
+	TEXT("0..x reduce texture quality for streaming by a floating point number.\n")
+	TEXT("0: use full resolution (default)\n")
+	TEXT("1: drop one mip\n")
+	TEXT("2: drop two mips"),
+	ECVF_Scalability
+	);
+
+/** For debugging purposes: Whether to consider the time factor when streaming. Turning it off is easier for debugging. */
+bool GStreamWithTimeFactor		= true;
+
+extern ENGINE_API UPrimitiveComponent* GDebugSelectedComponent;
+
+/** the float table {-1.0f,1.0f} **/
+float ENGINE_API GNegativeOneOneTable[2] = {-1.0f,1.0f};
+
+
+/** Smaller value will stream out lightmaps more aggressively. */
+float GLightmapStreamingFactor = 1.0f;
+
+/** Smaller value will stream out shadowmaps more aggressively. */
+float GShadowmapStreamingFactor = 0.09f;
+
+/** For testing, finding useless textures or special demo purposes. If true, textures will never be streamed out (but they can be GC'd). 
+* Caution: this only applies to unlimited texture pools (i.e. not consoles)
+*/
+bool GNeverStreamOutTextures = false;
+
+/** Accumulated total time spent on dynamic primitives, in seconds. */
+double GStreamingDynamicPrimitivesTime = 0.0;
+
+
+/**
+ * Checks whether a UTexture2D is supposed to be streaming.
+ * @param Texture	Texture to check
+ * @return			true if the UTexture2D is supposed to be streaming
+ */
+bool IsStreamingTexture( const UTexture2D* Texture2D )
+{
+	return Texture2D && Texture2D->bIsStreamable && Texture2D->NeverStream == false && Texture2D->GetNumMips() > UTexture2D::GetMinTextureResidentMipCount();
+}
+
+void FStreamingContext::Reset( bool bProcessEverything, UTexture2D* IndividualStreamingTexture, bool bInCollectTextureStats )
+{
+	bCollectTextureStats							= bInCollectTextureStats;
+	ThisFrameTotalRequestSize						= 0;
+	ThisFrameTotalLightmapRequestSize				= 0;
+	ThisFrameNumStreamingTextures					= 0;
+	ThisFrameNumRequestsInCancelationPhase			= 0;
+	ThisFrameNumRequestsInUpdatePhase				= 0;
+	ThisFrameNumRequestsInFinalizePhase				= 0;
+	ThisFrameTotalIntermediateTexturesSize			= 0;
+	ThisFrameNumIntermediateTextures				= 0;
+	ThisFrameTotalStreamingTexturesSize				= 0;
+	ThisFrameTotalStreamingTexturesMaxSize			= 0;
+	ThisFrameTotalLightmapMemorySize				= 0;
+	ThisFrameTotalLightmapDiskSize					= 0;
+	ThisFrameTotalHLODMemorySize					= 0;
+	ThisFrameTotalHLODDiskSize						= 0;
+	ThisFrameTotalMipCountIncreaseRequestsInFlight	= 0;
+	ThisFrameOptimalWantedSize						= 0;
+	ThisFrameTotalStaticTextureHeuristicSize		= 0;
+	ThisFrameTotalDynamicTextureHeuristicSize		= 0;
+	ThisFrameTotalLastRenderHeuristicSize			= 0;
+	ThisFrameTotalForcedHeuristicSize				= 0;
+
+	AvailableNow = 0;
+	AvailableLater = 0;
+	AvailableTempMemory = 0;
+
+	STAT( TextureStats.Empty() );
+
+	AllocatedMemorySize	= INDEX_NONE;
+	PendingMemoryAdjustment = INDEX_NONE;
+		
+	// Available texture memory, if supported by RHI. This stat is async as the renderer allocates the memory in its own thread so we
+	// only query once and roughly adjust the values as needed.
+	FTextureMemoryStats Stats;
+	RHIGetTextureMemoryStats(Stats);
+		
+	bRHISupportsMemoryStats = Stats.IsUsingLimitedPoolSize();
+
+	// Update stats if supported.
+	if( bRHISupportsMemoryStats )
+	{
+		AllocatedMemorySize	= Stats.AllocatedMemorySize;
+		PendingMemoryAdjustment = Stats.PendingMemoryAdjustment;
+	
+		// set total size for the pool (used to available)
+		SET_MEMORY_STAT(STAT_TexturePoolAllocatedSize, AllocatedMemorySize);
+		SET_MEMORY_STAT(STAT_TexturePoolSize, Stats.TexturePoolSize);
+		SET_MEMORY_STAT(MCR_TexturePool, Stats.TexturePoolSize);
+	}
+	else
+	{
+		SET_MEMORY_STAT(STAT_TexturePoolAllocatedSize,0);
+		SET_MEMORY_STAT(STAT_TexturePoolSize, 0);
+		SET_MEMORY_STAT(MCR_TexturePool, 0);
+	}
+}
+
+void FStreamingContext::AddStats( const FStreamingContext& Other )
+{
+	ThisFrameTotalRequestSize						+= Other.ThisFrameTotalRequestSize;						
+	ThisFrameTotalLightmapRequestSize				+= Other.ThisFrameTotalLightmapRequestSize;
+	ThisFrameNumStreamingTextures					+= Other.ThisFrameNumStreamingTextures;
+	ThisFrameNumRequestsInCancelationPhase			+= Other.ThisFrameNumRequestsInCancelationPhase;
+	ThisFrameNumRequestsInUpdatePhase				+= Other.ThisFrameNumRequestsInUpdatePhase;
+	ThisFrameNumRequestsInFinalizePhase				+= Other.ThisFrameNumRequestsInFinalizePhase;
+	ThisFrameTotalIntermediateTexturesSize			+= Other.ThisFrameTotalIntermediateTexturesSize;
+	ThisFrameNumIntermediateTextures				+= Other.ThisFrameNumIntermediateTextures;
+	ThisFrameTotalStreamingTexturesSize				+= Other.ThisFrameTotalStreamingTexturesSize;
+	ThisFrameTotalStreamingTexturesMaxSize			+= Other.ThisFrameTotalStreamingTexturesMaxSize;
+	ThisFrameTotalLightmapMemorySize				+= Other.ThisFrameTotalLightmapMemorySize;
+	ThisFrameTotalLightmapDiskSize					+= Other.ThisFrameTotalLightmapDiskSize;
+	ThisFrameTotalHLODMemorySize					+= Other.ThisFrameTotalHLODMemorySize;
+	ThisFrameTotalHLODDiskSize						+= Other.ThisFrameTotalHLODDiskSize;
+	ThisFrameTotalMipCountIncreaseRequestsInFlight	+= Other.ThisFrameTotalMipCountIncreaseRequestsInFlight;
+	ThisFrameOptimalWantedSize						+= Other.ThisFrameOptimalWantedSize;
+	ThisFrameTotalStaticTextureHeuristicSize		+= Other.ThisFrameTotalStaticTextureHeuristicSize;
+	ThisFrameTotalDynamicTextureHeuristicSize		+= Other.ThisFrameTotalDynamicTextureHeuristicSize;
+	ThisFrameTotalLastRenderHeuristicSize			+= Other.ThisFrameTotalLastRenderHeuristicSize;
+	ThisFrameTotalForcedHeuristicSize				+= Other.ThisFrameTotalForcedHeuristicSize;
+	bCollectTextureStats							= bCollectTextureStats || Other.bCollectTextureStats;
+	STAT( TextureStats.Append( Other.TextureStats ) );
+}
diff --git a/Engine/Source/Runtime/Engine/Private/Streaming/TextureStreamingHelpers.h b/Engine/Source/Runtime/Engine/Private/Streaming/TextureStreamingHelpers.h
new file mode 100644
--- /dev/null
+++ b/Engine/Source/Runtime/Engine/Private/Streaming/TextureStreamingHelpers.h
@@ -0,0 +1,335 @@
+// Copyright 1998-2016 Epic Games, Inc. All Rights Reserved.
+
+/*=============================================================================
+TextureStreamingHelpers.h: Definitions of classes used for texture streaming.
+=============================================================================*/
+
+#pragma once
+
+/**
+ * Streaming stats
+ */
+DECLARE_DWORD_ACCUMULATOR_STAT_EXTERN(TEXT("Num Streaming Textures"),STAT_StreamingTextures,STATGROUP_Streaming, );
+DECLARE_MEMORY_STAT_POOL_EXTERN(TEXT("Pool Memory Size"), STAT_TexturePoolSize, STATGROUP_Streaming, FPlatformMemory::MCR_TexturePool, );
+DECLARE_MEMORY_STAT_POOL_EXTERN(TEXT("Pool Memory Used"), STAT_TexturePoolAllocatedSize, STATGROUP_Streaming, FPlatformMemory::MCR_TexturePool, );
+DECLARE_MEMORY_STAT_POOL_EXTERN(TEXT("Effective Streaming Pool Size"), STAT_EffectiveStreamingPoolSize, STATGROUP_Streaming, FPlatformMemory::MCR_TexturePool, );
+DECLARE_MEMORY_STAT_POOL_EXTERN(TEXT("Streaming Textures"),STAT_StreamingTexturesSize,STATGROUP_Streaming,FPlatformMemory::MCR_TexturePool, );
+DECLARE_MEMORY_STAT_POOL_EXTERN(TEXT("Non-streaming Textures"),STAT_NonStreamingTexturesSize,STATGROUP_Streaming,FPlatformMemory::MCR_TexturePool, );
+DECLARE_MEMORY_STAT_POOL_EXTERN(TEXT("Textures On Disk"),STAT_StreamingTexturesMaxSize,STATGROUP_StreamingDetails,FPlatformMemory::MCR_TexturePool, );
+DECLARE_MEMORY_STAT_POOL_EXTERN(TEXT("Lightmaps In Memory"),STAT_LightmapMemorySize,STATGROUP_StreamingDetails,FPlatformMemory::MCR_TexturePool, );
+DECLARE_MEMORY_STAT_POOL_EXTERN(TEXT("Lightmaps On Disk"),STAT_LightmapDiskSize,STATGROUP_StreamingDetails,FPlatformMemory::MCR_TexturePool, );
+DECLARE_MEMORY_STAT_POOL_EXTERN(TEXT("HLOD Textures In Memory"), STAT_HLODTextureMemorySize, STATGROUP_StreamingDetails, FPlatformMemory::MCR_TexturePool, );
+DECLARE_MEMORY_STAT_POOL_EXTERN(TEXT("HLOD Textures On Disk"), STAT_HLODTextureDiskSize, STATGROUP_StreamingDetails, FPlatformMemory::MCR_TexturePool, );
+DECLARE_DWORD_ACCUMULATOR_STAT_EXTERN(TEXT("Requests In Cancelation Phase"),STAT_RequestsInCancelationPhase,STATGROUP_StreamingDetails, );
+DECLARE_DWORD_ACCUMULATOR_STAT_EXTERN(TEXT("Requests In Update Phase"),STAT_RequestsInUpdatePhase,STATGROUP_StreamingDetails, );
+DECLARE_DWORD_ACCUMULATOR_STAT_EXTERN(TEXT("Requests in Finalize Phase"),STAT_RequestsInFinalizePhase,STATGROUP_StreamingDetails, );
+DECLARE_DWORD_ACCUMULATOR_STAT_EXTERN(TEXT("Intermediate Textures"),STAT_IntermediateTextures,STATGROUP_StreamingDetails, );
+DECLARE_MEMORY_STAT_POOL_EXTERN(TEXT("Intermediate Textures Size"),STAT_IntermediateTexturesSize,STATGROUP_StreamingDetails,FPlatformMemory::MCR_TexturePool, );
+DECLARE_MEMORY_STAT_POOL_EXTERN(TEXT("Textures Streamed In, Frame"),STAT_RequestSizeCurrentFrame,STATGROUP_StreamingDetails,FPlatformMemory::MCR_TexturePool, );
+DECLARE_MEMORY_STAT_POOL_EXTERN(TEXT("Textures Streamed In, Total"),STAT_RequestSizeTotal,STATGROUP_StreamingDetails,FPlatformMemory::MCR_TexturePool, );
+DECLARE_MEMORY_STAT_POOL_EXTERN(TEXT("Lightmaps Streamed In, Total"),STAT_LightmapRequestSizeTotal,STATGROUP_StreamingDetails,FPlatformMemory::MCR_TexturePool, );
+DECLARE_CYCLE_STAT_EXTERN(TEXT("Game Thread Update Time"),STAT_GameThreadUpdateTime,STATGROUP_Streaming, );
+DECLARE_FLOAT_ACCUMULATOR_STAT_EXTERN(TEXT("Streaming Latency, Average (sec)"),STAT_StreamingLatency,STATGROUP_StreamingDetails, );
+DECLARE_FLOAT_ACCUMULATOR_STAT_EXTERN(TEXT("Streaming Bandwidth, Average (MB/s)"),STAT_StreamingBandwidth,STATGROUP_StreamingDetails, );
+DECLARE_FLOAT_ACCUMULATOR_STAT_EXTERN(TEXT("Dynamic Streaming Total Time (sec)"),STAT_DynamicStreamingTotal,STATGROUP_StreamingDetails, );
+DECLARE_DWORD_ACCUMULATOR_STAT_EXTERN(TEXT("Num Visible Textures With Low Resolutions"),STAT_NumVisibleTexturesWithLowResolutions,STATGROUP_StreamingDetails, );
+DECLARE_CYCLE_STAT_EXTERN(TEXT("Volume Streaming Tick"),STAT_VolumeStreamingTickTime,STATGROUP_StreamingDetails, );
+DECLARE_DWORD_COUNTER_STAT_EXTERN(TEXT("Streaming Volumes"),STAT_VolumeStreamingChecks,STATGROUP_StreamingDetails, );
+DECLARE_DWORD_ACCUMULATOR_STAT_EXTERN(TEXT("Growing Reallocations"),STAT_GrowingReallocations,STATGROUP_StreamingDetails, );
+DECLARE_DWORD_ACCUMULATOR_STAT_EXTERN(TEXT("Shrinking Reallocations"),STAT_ShrinkingReallocations,STATGROUP_StreamingDetails, );
+DECLARE_DWORD_ACCUMULATOR_STAT_EXTERN(TEXT("Full Reallocations"),STAT_FullReallocations,STATGROUP_StreamingDetails, );
+DECLARE_DWORD_ACCUMULATOR_STAT_EXTERN(TEXT("Failed Reallocations"),STAT_FailedReallocations,STATGROUP_StreamingDetails, );
+DECLARE_DWORD_ACCUMULATOR_STAT_EXTERN(TEXT("Panic Defragmentations"),STAT_PanicDefragmentations,STATGROUP_StreamingDetails, );
+DECLARE_CYCLE_STAT_EXTERN(TEXT("AddToWorld Time"),STAT_AddToWorldTime,STATGROUP_StreamingDetails, );
+DECLARE_CYCLE_STAT_EXTERN(TEXT("RemoveFromWorld Time"),STAT_RemoveFromWorldTime,STATGROUP_StreamingDetails, );
+DECLARE_CYCLE_STAT_EXTERN(TEXT("UpdateLevelStreaming Time"),STAT_UpdateLevelStreamingTime,STATGROUP_StreamingDetails, );
+DECLARE_MEMORY_STAT_POOL_EXTERN(TEXT("Required Streaming Textures"),STAT_OptimalTextureSize,STATGROUP_Streaming,FPlatformMemory::MCR_TexturePool, );
+DECLARE_MEMORY_STAT_POOL_EXTERN(TEXT("Over Budget"),STAT_StreamingOverBudget,STATGROUP_Streaming,FPlatformMemory::MCR_TexturePool, );
+DECLARE_MEMORY_STAT_POOL_EXTERN(TEXT("Under Budget"),STAT_StreamingUnderBudget,STATGROUP_StreamingDetails,FPlatformMemory::MCR_TexturePool, );
+DECLARE_DWORD_ACCUMULATOR_STAT_EXTERN(TEXT("Num Pending Textures"),STAT_NumWantingTextures,STATGROUP_Streaming, );
+DECLARE_DWORD_ACCUMULATOR_STAT_EXTERN(TEXT("Levels With PostLoad Rotations"),STAT_LevelsWithPostLoadRotations,STATGROUP_Streaming, );
+DECLARE_MEMORY_STAT_POOL_EXTERN(TEXT("Static Textures In Memory"),STAT_TotalStaticTextureHeuristicSize,STATGROUP_StreamingDetails,FPlatformMemory::MCR_TexturePool, );
+DECLARE_MEMORY_STAT_POOL_EXTERN(TEXT("LastRenderTime Textures In Memory"),STAT_TotalLastRenderHeuristicSize,STATGROUP_StreamingDetails,FPlatformMemory::MCR_TexturePool, );
+DECLARE_MEMORY_STAT_POOL_EXTERN(TEXT("Dynamic Textures In Memory"),STAT_TotalDynamicHeuristicSize,STATGROUP_StreamingDetails,FPlatformMemory::MCR_TexturePool, );
+DECLARE_MEMORY_STAT_POOL_EXTERN(TEXT("Forced Textures In Memory"),STAT_TotalForcedHeuristicSize,STATGROUP_StreamingDetails,FPlatformMemory::MCR_TexturePool, );
+
+DECLARE_LOG_CATEGORY_EXTERN(LogContentStreaming, Log, All);
+
+
+extern float GLightmapStreamingFactor;
+extern float GShadowmapStreamingFactor;
+extern bool GNeverStreamOutTextures;
+extern double GStreamingDynamicPrimitivesTime;
+extern bool GIsOperatingWithReducedTexturePool;
+extern bool GStreamWithTimeFactor;
+
+//@DEBUG:
+// Set to 1 to log all dynamic component notifications
+#define STREAMING_LOG_DYNAMIC		0
+// Set to 1 to log when we change a view
+#define STREAMING_LOG_VIEWCHANGES	0
+// Set to 1 to log when levels are added/removed
+#define STREAMING_LOG_LEVELS		0
+// Set to 1 to log textures that are canceled by CancelForcedTextures()
+#define STREAMING_LOG_CANCELFORCED	0
+
+#if PLATFORM_SUPPORTS_TEXTURE_STREAMING
+extern TAutoConsoleVariable<int32> CVarSetTextureStreaming;
+#endif
+
+extern TAutoConsoleVariable<float> CVarStreamingBoost;
+extern TAutoConsoleVariable<float> CVarStreamingScreenSizeEffectiveMax;
+extern TAutoConsoleVariable<int32> CVarStreamingUseFixedPoolSize;
+extern TAutoConsoleVariable<int32> CVarStreamingPoolSize;
+extern TAutoConsoleVariable<int32> CVarStreamingMaxTempMemoryAllowed;
+extern TAutoConsoleVariable<int32> CVarStreamingShowWantedMips;
+extern TAutoConsoleVariable<int32> CVarStreamingHLODStrategy;
+extern TAutoConsoleVariable<float> CVarStreamingHiddenPrimitiveScale;
+extern TAutoConsoleVariable<int32> CVarStreamingSplitRequestSizeThreshold;
+extern TAutoConsoleVariable<float> CVarStreamingMipBias;
+
+typedef TArray<int32, TMemStackAllocator<> > FStreamingRequests;
+typedef TArray<const UTexture2D*, TInlineAllocator<12> > FRemovedTextureArray;
+
+#define NUM_BANDWIDTHSAMPLES 512
+#define NUM_LATENCYSAMPLES 512
+
+/** Streaming priority: Linear distance factor from 0 to MAX_STREAMINGDISTANCE. */
+#define MAX_STREAMINGDISTANCE	10000.0f
+#define MAX_MIPDELTA			5.0f
+#define MAX_LASTRENDERTIME		90.0f
+
+
+class UTexture2D;
+class UPrimitiveComponent;
+class FTextureBoundsVisibility;
+class FDynamicComponentTextureManager;
+template<typename T>
+class FAsyncTask;
+class FAsyncTextureStreaming;
+
+
+struct FStreamingTexture;
+struct FStreamingContext;
+struct FStreamingHandlerTextureBase;
+struct FTexturePriority;
+
+/**
+ * Checks whether a UTexture2D is supposed to be streaming.
+ * @param Texture	Texture to check
+ * @return			true if the UTexture2D is supposed to be streaming
+ */
+bool IsStreamingTexture( const UTexture2D* Texture2D );
+
+
+/**
+ *	Helper struct that represents a texture and the parameters used for sorting and streaming out high-res mip-levels.
+ **/
+struct FTextureSortElement
+{
+	/**
+	 *	Constructor.
+	 *
+	 *	@param InTexture					- The texture to represent
+	 *	@param InSize						- Size of the whole texture and all current mip-levels, in bytes
+	 *	@param bInIsCharacterTexture		- 1 if this is a character texture, otherwise 0
+	 *	@param InTextureDataAddress			- Starting address of the texture data
+	 *	@param InNumRequiredResidentMips	- Minimum number of mip-levels required to stay in memory
+	 */
+	FTextureSortElement( UTexture2D* InTexture, int32 InSize, int32 bInIsCharacterTexture, uint32 InTextureDataAddress, int32 InNumRequiredResidentMips )
+	:	Texture( InTexture )
+	,	Size( InSize )
+	,	bIsCharacterTexture( bInIsCharacterTexture )
+	,	TextureDataAddress( InTextureDataAddress )
+	,	NumRequiredResidentMips( InNumRequiredResidentMips )
+	{
+	}
+	/** The texture that this element represents */
+	UTexture2D*	Texture;
+	/** Size of the whole texture and all current mip-levels, in bytes. */
+	int32			Size;
+	/** 1 if this is a character texture, otherwise 0 */
+	int32			bIsCharacterTexture;
+	/** Starting address of the texture data. */
+	uint32		TextureDataAddress;			
+	/** Minimum number of mip-levels required to stay in memory. */
+	int32			NumRequiredResidentMips;
+};
+
+
+enum ETextureStreamingType
+{
+	StreamType_Static,
+	StreamType_Dynamic,
+	StreamType_Forced,
+	StreamType_LastRenderTime,
+	StreamType_Orphaned,
+	StreamType_Other,
+};
+
+static const TCHAR* GStreamTypeNames[] =
+{
+	TEXT("Static"),
+	TEXT("Dynamic"),
+	TEXT("Forced"),
+	TEXT("LastRenderTime"),
+	TEXT("Orphaned"),
+	TEXT("Other"),
+};
+
+struct FTexturePriority
+{
+	FTexturePriority( bool InCanDropMips, float InRetentionPriority, float InLoadOrderPriority, int32 InTextureIndex, const UTexture2D* InTexture)
+	:	bCanDropMips( InCanDropMips )
+	,	RetentionPriority( InRetentionPriority )
+	,	LoadOrderPriority( InLoadOrderPriority )
+	,	TextureIndex( InTextureIndex )
+	,	Texture( InTexture )
+	{
+	}
+	/** True if we allows this texture to drop mips to fit in budget. */
+	bool bCanDropMips;
+	/** Texture retention priority, higher value means it should be kept in memory. */
+	float RetentionPriority;
+	/** Texture load order priority, higher value means it should load first. */
+	float LoadOrderPriority;
+	/** Index into the FStreamingManagerTexture::StreamingTextures array. */
+	int32 TextureIndex;
+	/** The texture to stream. Only used for validation. */
+	const UTexture2D* Texture;
+};
+
+
+
+#if STATS
+struct FTextureStreamingStats
+{
+	FTextureStreamingStats( UTexture2D* InTexture2D, ETextureStreamingType InType, int32 InResidentMips, int32 InWantedMips, int32 InMostResidentMips, int32 InResidentSize, int32 InWantedSize, int32 InMaxSize, int32 InMostResidentSize, float InBoostFactor, float InPriority, int32 InTextureIndex )
+	:	TextureName( InTexture2D->GetFullName() )
+	,	SizeX( InTexture2D->GetSizeX() )
+	,	SizeY( InTexture2D->GetSizeY() )
+	,	NumMips( InTexture2D->GetNumMips() )
+	,	LODBias( InTexture2D->GetCachedLODBias() )
+	,	LastRenderTime( FMath::Clamp( InTexture2D->Resource ? (FApp::GetCurrentTime() - InTexture2D->Resource->LastRenderTime) : 1000000.0, 0.0, 1000000.0) )
+	,	StreamType( InType )
+	,	ResidentMips( InResidentMips )
+	,	WantedMips( InWantedMips )
+	,	MostResidentMips( InMostResidentMips )
+	,	ResidentSize( InResidentSize )
+	,	WantedSize( InWantedSize )
+	,	MaxSize( InMaxSize )
+	,	MostResidentSize( InMostResidentSize )
+	,	BoostFactor( InBoostFactor )
+	,	Priority( InPriority )
+	,	TextureIndex( InTextureIndex )
+	{
+	}
+	/** Mirror of UTexture2D::GetName() */
+	FString		TextureName;
+	/** Mirror of UTexture2D::SizeX */
+	int32		SizeX;
+	/** Mirror of UTexture2D::SizeY */
+	int32		SizeY;
+	/** Mirror of UTexture2D::Mips.Num() */
+	int32		NumMips;
+	/** Mirror of UTexture2D::GetCachedLODBias() */
+	int32		LODBias;
+	/** How many seconds since it was last rendered: FApp::GetCurrentTime() - UTexture2D::Resource->LastRenderTime */
+	float		LastRenderTime;
+	/** What streaming heuristics were primarily used for this texture. */
+	ETextureStreamingType	StreamType;
+	/** Number of resident mip levels. */
+	int32		ResidentMips;
+	/** Number of wanted mip levels. */
+	int32		WantedMips;
+	/** Most number of mip-levels this texture has ever had resident in memory. */
+	int32		MostResidentMips;
+	/** Number of bytes currently in memory. */
+	int32		ResidentSize;
+	/** Number of bytes we want in memory. */
+	int32		WantedSize;
+	/** Number of bytes we could potentially stream in. */
+	int32		MaxSize;
+	/** Most number of bytes this texture has ever had resident in memory. */
+	int32		MostResidentSize;
+	/** Temporary boost of the streaming distance factor. */
+	float		BoostFactor;
+	/** Texture priority */
+	float		Priority;
+	/** Index into the FStreamingManagerTexture::StreamingTextures array. */
+	int32		TextureIndex;
+};
+#endif
+
+/**
+ * Helper struct for temporary information for one frame of processing texture streaming.
+ */
+struct FStreamingContext
+{
+	FStreamingContext( bool bProcessEverything, UTexture2D* IndividualStreamingTexture, bool bInCollectTextureStats )
+	{
+		Reset( bProcessEverything, IndividualStreamingTexture, bInCollectTextureStats );
+	}
+
+	/**
+	 * Initializes all variables for the one frame.
+	 * @param bProcessEverything			If true, process all resources with no throttling limits
+	 * @param IndividualStreamingTexture	A specific texture to be fully processed this frame, or NULL
+	 * @param bInCollectTextureStats		Whether to fill in the TextureStats array this frame
+	 */
+	void Reset( bool bProcessEverything, UTexture2D* IndividualStreamingTexture, bool bInCollectTextureStats );
+
+	/**
+	 * Adds in the stats from another context.
+	 *
+	 * @param Other		Context to add stats from
+	 */
+	void AddStats( const FStreamingContext& Other );
+
+	/** Whether the platform RHI supports texture memory stats. */
+	bool bRHISupportsMemoryStats;
+	/** Currently allocated number of bytes in the texture pool, or INDEX_NONE if bRHISupportsMemoryStats is false. */
+	int64 AllocatedMemorySize;
+	/** Pending Adjustments to allocated texture memory, due to async reallocations, etc. */
+	int64 PendingMemoryAdjustment;
+	/** Whether to fill in TextureStats this frame. */
+	bool bCollectTextureStats;
+#if STATS
+	/** Stats for all textures. */
+	TArray<FTextureStreamingStats>	TextureStats;
+#endif
+
+	// Stats for this frame.
+	uint64 ThisFrameTotalRequestSize;
+	uint64 ThisFrameTotalLightmapRequestSize;
+	uint64 ThisFrameNumStreamingTextures;
+	uint64 ThisFrameNumRequestsInCancelationPhase;
+	uint64 ThisFrameNumRequestsInUpdatePhase;
+	uint64 ThisFrameNumRequestsInFinalizePhase;
+	uint64 ThisFrameTotalIntermediateTexturesSize;
+	uint64 ThisFrameNumIntermediateTextures;
+	uint64 ThisFrameTotalStreamingTexturesSize;
+	uint64 ThisFrameTotalStreamingTexturesMaxSize;
+	uint64 ThisFrameTotalLightmapMemorySize;
+	uint64 ThisFrameTotalLightmapDiskSize;
+	uint64 ThisFrameTotalHLODMemorySize;
+	uint64 ThisFrameTotalHLODDiskSize;
+	uint64 ThisFrameTotalMipCountIncreaseRequestsInFlight;
+	uint64 ThisFrameOptimalWantedSize;
+	/** Number of bytes using StaticTexture heuristics this frame, currently in memory. */
+	uint64 ThisFrameTotalStaticTextureHeuristicSize;
+	/** Number of bytes using DynmicTexture heuristics this frame, currently in memory. */
+	uint64 ThisFrameTotalDynamicTextureHeuristicSize;
+	/** Number of bytes using LastRenderTime heuristics this frame, currently in memory. */
+	uint64 ThisFrameTotalLastRenderHeuristicSize;
+	/** Number of bytes using ForcedIntoMemory heuristics this frame, currently in memory. */
+	uint64 ThisFrameTotalForcedHeuristicSize;
+
+	/** Conservative amount of memory currently available or available after some pending requests are completed. */
+	int64 AvailableNow;
+	/** Amount of memory expected to be available once all streaming is completed. Depends on WantedMips */
+	int64 AvailableLater;
+	/** Amount of memory immediatly available temporary streaming data (pending textures). Depends on the delta between RequestedMips and ResidentMips */
+	int64 AvailableTempMemory;
+};
+
diff --git a/Engine/Source/Runtime/Engine/Private/Texture2D.cpp b/Engine/Source/Runtime/Engine/Private/Texture2D.cpp
--- a/Engine/Source/Runtime/Engine/Private/Texture2D.cpp
+++ b/Engine/Source/Runtime/Engine/Private/Texture2D.cpp
@@ -17,12 +17,14 @@
 
 #include "TargetPlatform.h"
 #include "ContentStreaming.h"
+#include "Streaming/StreamingManagerTexture.h"
 
 UTexture2D::UTexture2D(const FObjectInitializer& ObjectInitializer)
 	: Super(ObjectInitializer)
 {
 	bHasCancelationPending = false;
-	StreamingIndex = -1;
+	StreamingIndex = INDEX_NONE;
+	LevelIndex = INDEX_NONE;
 	SRGB = true;
 }
 
diff --git a/Engine/Source/Runtime/Engine/Private/UnrealEngine.cpp b/Engine/Source/Runtime/Engine/Private/UnrealEngine.cpp
--- a/Engine/Source/Runtime/Engine/Private/UnrealEngine.cpp
+++ b/Engine/Source/Runtime/Engine/Private/UnrealEngine.cpp
@@ -758,7 +758,7 @@ void EngineMemoryWarningHandler(const FGenericMemoryWarningContext& GenericConte
 		Stats.AvailablePhysical / 1048576.0f);
 
 #if !UE_BUILD_SHIPPING && !UE_BUILD_TEST
-	const auto OOMMemReportVar = IConsoleManager::Get().FindTConsoleVariableDataInt(TEXT("Debug.OOMMemReport")); 
+	static const auto OOMMemReportVar = IConsoleManager::Get().FindTConsoleVariableDataInt(TEXT("Debug.OOMMemReport")); 
 	const int32 OOMMemReport = OOMMemReportVar ? OOMMemReportVar->GetValueOnAnyThread() : false;
 	if( OOMMemReport )
 	{
@@ -3574,8 +3574,9 @@ bool UEngine::HandleListTexturesCommand( const TCHAR* Cmd, FOutputDevice& Ar )
 		// Use the existing texture streaming functionality to gather referenced textures. Worth noting
 		// that GetStreamingTextureInfo doesn't check whether a texture is actually streamable or not
 		// and is also implemented for skeletal meshes and such.
+		FStreamingTextureLevelContext LevelContext;
 		TArray<FStreamingTexturePrimitiveInfo> StreamingTextures;
-		PrimitiveComponent->GetStreamingTextureInfoWithNULLRemoval( StreamingTextures );
+		PrimitiveComponent->GetStreamingTextureInfo( LevelContext, StreamingTextures );
 
 		// Increase usage count for all referenced textures
 		for( int32 TextureIndex=0; TextureIndex<StreamingTextures.Num(); TextureIndex++ )
@@ -4526,7 +4527,7 @@ bool UEngine::HandleDebugCommand( const TCHAR* Cmd, FOutputDevice& Ar )
 
 	// Handle "DEBUG CRASH" etc. 
 	return PerformError(Cmd, Ar);
-	}
+}
 
 
 bool UEngine::HandleMergeMeshCommand( const TCHAR* Cmd, FOutputDevice& Ar, UWorld* InWorld )
@@ -6338,7 +6339,7 @@ void UEngine::StartHardwareSurvey()
 #endif
 
 	if (bEnabled && FEngineAnalytics::IsAvailable())
-		{
+	{
 		IHardwareSurveyModule::Get().StartHardwareSurvey(FEngineAnalytics::GetProvider());
 		}
 	}	
@@ -11045,7 +11046,7 @@ void UEngine::CopyPropertiesForUnrelatedObjects(UObject* OldObject, UObject* New
 				if (!bContainedInsideNewInstance)
 				{
 					// A bad thing has happened and cannot be reasonably fixed at this point
-					UE_LOG(LogEngine, Log, TEXT("Warning: The CDO '%s' references a component that does not have the CDO in its outer chain!"), *NewObject->GetFullName(), *NewInstance->GetFullName()); 	
+					UE_LOG(LogEngine, Log, TEXT("Warning: The CDO '%s' references a component that does not have the CDO in its outer chain!"), *NewObject->GetFullName(), *NewInstance->GetFullName());
 				}
 			}
 		}
@@ -12117,22 +12118,22 @@ int32 UEngine::RenderStatSoundWaves(UWorld* World, FViewport* Viewport, FCanvas*
 			if (WaveInstance->GetActualVolume() >= 0.01f)
 			{
 				++ActiveInstances;
-			ActiveSounds.Add(WaveInstance->ActiveSound);
+				ActiveSounds.Add(WaveInstance->ActiveSound);
 
-			UAudioComponent* AudioComponent = WaveInstance->ActiveSound->GetAudioComponent();
-			AActor* SoundOwner = AudioComponent ? AudioComponent->GetOwner() : nullptr;
-			USoundClass* SoundClass = WaveInstance->SoundClass;
+				UAudioComponent* AudioComponent = WaveInstance->ActiveSound->GetAudioComponent();
+				AActor* SoundOwner = AudioComponent ? AudioComponent->GetOwner() : nullptr;
+				USoundClass* SoundClass = WaveInstance->SoundClass;
 
-			FString TheString = *FString::Printf(TEXT("%4i.    %6.2f  %s   Owner: %s   SoundClass: %s"),
-				InstanceIndex,
-				WaveInstance->GetActualVolume(),
-				*WaveInstance->WaveData->GetPathName(),
-				SoundOwner ? *SoundOwner->GetName() : TEXT("None"),
-				SoundClass ? *SoundClass->GetName() : TEXT("None"));
+				FString TheString = *FString::Printf(TEXT("%4i.    %6.2f  %s   Owner: %s   SoundClass: %s"),
+													 InstanceIndex,
+													 WaveInstance->GetActualVolume(),
+													 *WaveInstance->WaveData->GetPathName(),
+													 SoundOwner ? *SoundOwner->GetName() : TEXT("None"),
+													 SoundClass ? *SoundClass->GetName() : TEXT("None"));
 
-			Canvas->DrawShadowedString(X, Y, *TheString, GetSmallFont(), FColor::White);
-			Y += 12;
-		}
+				Canvas->DrawShadowedString(X, Y, *TheString, GetSmallFont(), FColor::White);
+				Y += 12;
+			}
 		}
 
 		int32 Max = AudioDevice->MaxChannels / 2;
@@ -12171,10 +12172,10 @@ int32 UEngine::RenderStatSoundCues(UWorld* World, FViewport* Viewport, FCanvas*
 			FWaveInstance* WaveInstance = WaveInstances[InstanceIndex];
 			if (WaveInstance->GetActualVolume() >= 0.01f)
 			{
-			ActiveSounds.Add(WaveInstance->ActiveSound);
+				ActiveSounds.Add(WaveInstance->ActiveSound);
+			}
 		}
 	}
-	}
 
 	Canvas->DrawShadowedString(X, Y, TEXT("Active Sound Cues:"), GetSmallFont(), FColor::Green);
 	Y += 12;
diff --git a/Engine/Source/Runtime/Engine/Private/World.cpp b/Engine/Source/Runtime/Engine/Private/World.cpp
--- a/Engine/Source/Runtime/Engine/Private/World.cpp
+++ b/Engine/Source/Runtime/Engine/Private/World.cpp
@@ -64,6 +64,7 @@
 
 #include "EngineModule.h"
 #include "ContentStreaming.h"
+#include "Streaming/TextureStreamingHelpers.h"
 #include "RendererInterface.h"
 #include "DataChannel.h"
 #include "ShaderCompiler.h"
@@ -2649,6 +2650,17 @@ void UWorld::TriggerStreamingDataRebuild()
 
 void UWorld::ConditionallyBuildStreamingData()
 {
+#if !(UE_BUILD_SHIPPING || UE_BUILD_TEST)
+	// This code is to trigger an update to the a streaming data update when the console variable is changed.
+	static int32 LastMetricsUsed = CVarStreamingUseNewMetrics.GetValueOnGameThread();
+	if (CVarStreamingUseNewMetrics.GetValueOnGameThread() != LastMetricsUsed)
+	{
+		LastMetricsUsed = CVarStreamingUseNewMetrics.GetValueOnGameThread();
+		bStreamingDataDirty = true;
+		BuildStreamingDataTimer = FPlatformTime::Seconds() - 1.0;
+	}
+#endif
+
 	if ( bStreamingDataDirty && FPlatformTime::Seconds() > BuildStreamingDataTimer )
 	{
 		bStreamingDataDirty = false;
diff --git a/Engine/Source/Runtime/Engine/Private/WorldSettings.cpp b/Engine/Source/Runtime/Engine/Private/WorldSettings.cpp
--- a/Engine/Source/Runtime/Engine/Private/WorldSettings.cpp
+++ b/Engine/Source/Runtime/Engine/Private/WorldSettings.cpp
@@ -63,7 +63,7 @@ AWorldSettings::AWorldSettings(const FObjectInitializer& ObjectInitializer)
 	bHidden = false;
 
 	DefaultColorScale = FVector(1.0f, 1.0f, 1.0f);
-
+	DefaultMaxDistanceFieldOcclusionDistance = 600;
 	bPlaceCellsOnlyAlongCameraTracks = false;
 	VisibilityCellSize = 200;
 	VisibilityAggressiveness = VIS_LeastAggressive;
@@ -320,6 +320,11 @@ void AWorldSettings::PostEditChangeProperty(FPropertyChangedEvent& PropertyChang
 		}
 	}
 
+	if (PropertyThatChanged != nullptr && GetWorld() != nullptr && GetWorld()->Scene)
+	{
+		GetWorld()->Scene->UpdateSceneSettings(this);
+	}
+
 	Super::PostEditChangeProperty(PropertyChangedEvent);
 }
 
diff --git a/Engine/Source/Runtime/Engine/Public/ContentStreaming.h b/Engine/Source/Runtime/Engine/Public/ContentStreaming.h
--- a/Engine/Source/Runtime/Engine/Public/ContentStreaming.h
+++ b/Engine/Source/Runtime/Engine/Public/ContentStreaming.h
@@ -12,67 +12,12 @@
 	Stats.
 -----------------------------------------------------------------------------*/
 
-/**
- * Streaming stats
- */
-DECLARE_DWORD_ACCUMULATOR_STAT_EXTERN(TEXT("Num Streaming Textures"),STAT_StreamingTextures,STATGROUP_Streaming, );
-DECLARE_MEMORY_STAT_POOL_EXTERN(TEXT("Pool Memory Size"), STAT_TexturePoolSize, STATGROUP_Streaming, FPlatformMemory::MCR_TexturePool, );
-DECLARE_MEMORY_STAT_POOL_EXTERN(TEXT("Pool Memory Used"), STAT_TexturePoolAllocatedSize, STATGROUP_Streaming, FPlatformMemory::MCR_TexturePool, );
-DECLARE_MEMORY_STAT_POOL_EXTERN(TEXT("Effective Streaming Pool Size"), STAT_EffectiveStreamingPoolSize, STATGROUP_Streaming, FPlatformMemory::MCR_TexturePool, );
-DECLARE_MEMORY_STAT_POOL_EXTERN(TEXT("Streaming Textures"),STAT_StreamingTexturesSize,STATGROUP_Streaming,FPlatformMemory::MCR_TexturePool, );
-DECLARE_MEMORY_STAT_POOL_EXTERN(TEXT("Non-streaming Textures"),STAT_NonStreamingTexturesSize,STATGROUP_Streaming,FPlatformMemory::MCR_TexturePool, );
-DECLARE_MEMORY_STAT_POOL_EXTERN(TEXT("Textures On Disk"),STAT_StreamingTexturesMaxSize,STATGROUP_StreamingDetails,FPlatformMemory::MCR_TexturePool, );
-DECLARE_MEMORY_STAT_POOL_EXTERN(TEXT("Lightmaps In Memory"),STAT_LightmapMemorySize,STATGROUP_StreamingDetails,FPlatformMemory::MCR_TexturePool, );
-DECLARE_MEMORY_STAT_POOL_EXTERN(TEXT("Lightmaps On Disk"),STAT_LightmapDiskSize,STATGROUP_StreamingDetails,FPlatformMemory::MCR_TexturePool, );
-DECLARE_MEMORY_STAT_POOL_EXTERN(TEXT("HLOD Textures In Memory"), STAT_HLODTextureMemorySize, STATGROUP_StreamingDetails, FPlatformMemory::MCR_TexturePool, );
-DECLARE_MEMORY_STAT_POOL_EXTERN(TEXT("HLOD Textures On Disk"), STAT_HLODTextureDiskSize, STATGROUP_StreamingDetails, FPlatformMemory::MCR_TexturePool, );
-DECLARE_DWORD_ACCUMULATOR_STAT_EXTERN(TEXT("Requests In Cancelation Phase"),STAT_RequestsInCancelationPhase,STATGROUP_StreamingDetails, );
-DECLARE_DWORD_ACCUMULATOR_STAT_EXTERN(TEXT("Requests In Update Phase"),STAT_RequestsInUpdatePhase,STATGROUP_StreamingDetails, );
-DECLARE_DWORD_ACCUMULATOR_STAT_EXTERN(TEXT("Requests in Finalize Phase"),STAT_RequestsInFinalizePhase,STATGROUP_StreamingDetails, );
-DECLARE_DWORD_ACCUMULATOR_STAT_EXTERN(TEXT("Intermediate Textures"),STAT_IntermediateTextures,STATGROUP_StreamingDetails, );
-DECLARE_MEMORY_STAT_POOL_EXTERN(TEXT("Intermediate Textures Size"),STAT_IntermediateTexturesSize,STATGROUP_StreamingDetails,FPlatformMemory::MCR_TexturePool, );
-DECLARE_MEMORY_STAT_POOL_EXTERN(TEXT("Textures Streamed In, Frame"),STAT_RequestSizeCurrentFrame,STATGROUP_StreamingDetails,FPlatformMemory::MCR_TexturePool, );
-DECLARE_MEMORY_STAT_POOL_EXTERN(TEXT("Textures Streamed In, Total"),STAT_RequestSizeTotal,STATGROUP_StreamingDetails,FPlatformMemory::MCR_TexturePool, );
-DECLARE_MEMORY_STAT_POOL_EXTERN(TEXT("Lightmaps Streamed In, Total"),STAT_LightmapRequestSizeTotal,STATGROUP_StreamingDetails,FPlatformMemory::MCR_TexturePool, );
-DECLARE_CYCLE_STAT_EXTERN(TEXT("Game Thread Update Time"),STAT_GameThreadUpdateTime,STATGROUP_Streaming, );
-DECLARE_FLOAT_ACCUMULATOR_STAT_EXTERN(TEXT("Streaming Latency, Average (sec)"),STAT_StreamingLatency,STATGROUP_StreamingDetails, );
-DECLARE_FLOAT_ACCUMULATOR_STAT_EXTERN(TEXT("Streaming Bandwidth, Average (MB/s)"),STAT_StreamingBandwidth,STATGROUP_StreamingDetails, );
-DECLARE_FLOAT_ACCUMULATOR_STAT_EXTERN(TEXT("Dynamic Streaming Total Time (sec)"),STAT_DynamicStreamingTotal,STATGROUP_StreamingDetails, );
-DECLARE_DWORD_ACCUMULATOR_STAT_EXTERN(TEXT("Num Visible Textures With Low Resolutions"),STAT_NumVisibleTexturesWithLowResolutions,STATGROUP_StreamingDetails, );
-DECLARE_CYCLE_STAT_EXTERN( TEXT( "Buffer Creation" ), STAT_AudioResourceCreationTime, STATGROUP_Audio , );
-DECLARE_CYCLE_STAT_EXTERN(TEXT("Volume Streaming Tick"),STAT_VolumeStreamingTickTime,STATGROUP_StreamingDetails, );
-DECLARE_DWORD_COUNTER_STAT_EXTERN(TEXT("Streaming Volumes"),STAT_VolumeStreamingChecks,STATGROUP_StreamingDetails, );
-DECLARE_DWORD_ACCUMULATOR_STAT_EXTERN(TEXT("Growing Reallocations"),STAT_GrowingReallocations,STATGROUP_StreamingDetails, );
-DECLARE_DWORD_ACCUMULATOR_STAT_EXTERN(TEXT("Shrinking Reallocations"),STAT_ShrinkingReallocations,STATGROUP_StreamingDetails, );
-DECLARE_DWORD_ACCUMULATOR_STAT_EXTERN(TEXT("Full Reallocations"),STAT_FullReallocations,STATGROUP_StreamingDetails, );
-DECLARE_DWORD_ACCUMULATOR_STAT_EXTERN(TEXT("Failed Reallocations"),STAT_FailedReallocations,STATGROUP_StreamingDetails, );
-DECLARE_DWORD_ACCUMULATOR_STAT_EXTERN(TEXT("Panic Defragmentations"),STAT_PanicDefragmentations,STATGROUP_StreamingDetails, );
-DECLARE_CYCLE_STAT_EXTERN(TEXT("AddToWorld Time"),STAT_AddToWorldTime,STATGROUP_StreamingDetails, );
-DECLARE_CYCLE_STAT_EXTERN(TEXT("RemoveFromWorld Time"),STAT_RemoveFromWorldTime,STATGROUP_StreamingDetails, );
-DECLARE_CYCLE_STAT_EXTERN(TEXT("UpdateLevelStreaming Time"),STAT_UpdateLevelStreamingTime,STATGROUP_StreamingDetails, );
-DECLARE_MEMORY_STAT_POOL_EXTERN(TEXT("Required Streaming Textures"),STAT_OptimalTextureSize,STATGROUP_Streaming,FPlatformMemory::MCR_TexturePool, );
-DECLARE_MEMORY_STAT_POOL_EXTERN(TEXT("Over Budget"),STAT_StreamingOverBudget,STATGROUP_Streaming,FPlatformMemory::MCR_TexturePool, );
-DECLARE_MEMORY_STAT_POOL_EXTERN(TEXT("Under Budget"),STAT_StreamingUnderBudget,STATGROUP_StreamingDetails,FPlatformMemory::MCR_TexturePool, );
-DECLARE_DWORD_ACCUMULATOR_STAT_EXTERN(TEXT("Num Pending Textures"),STAT_NumWantingTextures,STATGROUP_Streaming, );
-DECLARE_MEMORY_STAT_POOL_EXTERN(TEXT("Static Textures In Memory"),STAT_TotalStaticTextureHeuristicSize,STATGROUP_StreamingDetails,FPlatformMemory::MCR_TexturePool, );
-DECLARE_MEMORY_STAT_POOL_EXTERN(TEXT("LastRenderTime Textures In Memory"),STAT_TotalLastRenderHeuristicSize,STATGROUP_StreamingDetails,FPlatformMemory::MCR_TexturePool, );
-DECLARE_MEMORY_STAT_POOL_EXTERN(TEXT("Dynamic Textures In Memory"),STAT_TotalDynamicHeuristicSize,STATGROUP_StreamingDetails,FPlatformMemory::MCR_TexturePool, );
-DECLARE_MEMORY_STAT_POOL_EXTERN(TEXT("Forced Textures In Memory"),STAT_TotalForcedHeuristicSize,STATGROUP_StreamingDetails,FPlatformMemory::MCR_TexturePool, );
-
 // Forward declarations
-struct FStreamingTexture;
-struct FStreamingContext;
-class FAsyncTextureStreaming;
-class UActorComponent;
 class UPrimitiveComponent;
 class AActor;
 class UTexture2D;
 class FSoundSource;
-template<typename T>
-class FAsyncTask;
 struct FStreamingManagerTexture;
-class FDynamicComponentTextureManager;
-class FTextureBoundsVisibility;
 
 /** Helper function to flush resource streaming. */
 void FlushResourceStreaming();
@@ -96,43 +41,6 @@ enum ERemoveStreamingViews
 	RemoveStreamingViews_All
 };
 
-
-
-// To compute the maximum mip size required by any heuristic
-// float to avoid float->int conversion, unclamped until the end
-class FFloatMipLevel
-{
-public:
-	// constructor
-	FFloatMipLevel();
-
-	/**
-	 * Helper function to map screen texels to a mip level, includes the global bias
-	 * @param StreamingTexture can be 0 (if texture is getting destroyed)
-	 * @apram MipBias from cvar
-	 * @param bOptimal only needed for stats
-	 */
-	int32 ComputeMip(const FStreamingTexture* StreamingTexture, float MipBias, bool bOptimal) const;
-	//
-	bool operator>=(const FFloatMipLevel& rhs) const;
-	//
-	static FFloatMipLevel FromScreenSizeInTexels(float ScreenSizeInTexels);
-	
-	static FFloatMipLevel FromMipLevel(int32 InMipLevel);
-
-	void Invalidate();
-
-	// @return true: means this handle has a mip level that needs to be considered
-	bool IsHandled() const
-	{
-		return MipLevel > 0.0f;
-	}
-
-private:
-	// -1: not set yet
-	float MipLevel;
-};
-
 /**
  * Helper structure containing all relevant information for streaming.
  */
@@ -715,722 +623,3 @@ protected:
 	IAudioStreamingManager* AudioStreamingManager;
 };
 
-/*-----------------------------------------------------------------------------
-	Texture streaming helper structs.
------------------------------------------------------------------------------*/
-
-struct FSpawnedTextureInstance
-{
-	FSpawnedTextureInstance( UTexture2D* InTexture2D, float InTexelFactor, float InOriginalRadius )
-	:	Texture2D( InTexture2D )
-	,	TexelFactor( InTexelFactor )
-	,	InvOriginalRadius( (InOriginalRadius > 0.0f) ? (1.0f/InOriginalRadius) : 1.0f )
-	{
-	}
-	/** Texture that is used by a dynamic UPrimitiveComponent. */
-	UTexture2D*		Texture2D;
-	/** Texture specific texel scale factor  */
-	float			TexelFactor;
-	/** 1.0f / OriginalBoundingSphereRadius, at the time the TexelFactor was calculated originally. */
-	float			InvOriginalRadius;
-};
-
-/*-----------------------------------------------------------------------------
-	Texture streaming.
------------------------------------------------------------------------------*/
-
-struct FStreamingHandlerTextureBase;
-struct FTexturePriority;
-#define NUM_BANDWIDTHSAMPLES 512
-#define NUM_LATENCYSAMPLES 512
-
-typedef TArray<int32, TMemStackAllocator<> > FStreamingRequests;
-typedef TArray<const UTexture2D*, TInlineAllocator<12> > FRemovedTextureArray;
-
-/**
- * Structure containing all information needed for determining the screen space
- * size of an object/ texture instance.
- */
-struct FStreamableTextureInstance4
-{
-	FStreamableTextureInstance4()
-	:	BoundsOriginX( MAX_FLT, MAX_FLT, MAX_FLT, MAX_FLT )
-	,	BoundsOriginY( 0, 0, 0, 0 )
-	,	BoundsOriginZ( 0, 0, 0, 0 )
-	,	BoxExtentSizeX( 0, 0, 0, 0 )
-	,	BoxExtentSizeY( 0, 0, 0, 0 )
-	,	BoxExtentSizeZ( 0, 0, 0, 0 )
-	,	BoundingSphereRadius( 0, 0, 0, 0 )
-	,	MinDistanceSq( 0, 0, 0, 0 )
-	,	MaxDistanceSq( MAX_FLT, MAX_FLT, MAX_FLT, MAX_FLT )
-	,	TexelFactor( 0, 0, 0, 0 )
-	{
-	}
-
-	/** X coordinates for the bounds origin of 4 texture instances */
-	FVector4 BoundsOriginX;
-	/** Y coordinates for the bounds origin of 4 texture instances */
-	FVector4 BoundsOriginY;
-	/** Z coordinates for the bounds origin of 4 texture instances */
-	FVector4 BoundsOriginZ;
-
-	/** X size of the bounds box extent of 4 texture instances */
-	FVector4 BoxExtentSizeX;
-	/** Y size of the bounds box extent of 4 texture instances */
-	FVector4 BoxExtentSizeY;
-	/** Z size of the bounds box extent of 4 texture instances */
-	FVector4 BoxExtentSizeZ;
-
-	/** Sphere radii for the bounding sphere of 4 texture instances */
-	FVector4 BoundingSphereRadius;
-	/** Minimal distance ^2 (between the bounding sphere origin and the view origin) for which this entry is valid */
-	FVector4 MinDistanceSq;
-	/** Maximal distance ^2 (between the bounding sphere origin and the view origin) for which this entry is valid */
-	FVector4 MaxDistanceSq;
-
-	/** Texel scale factors for 4 texture instances */
-	FVector4 TexelFactor;
-};
-
-
-/**
- * Streaming manager dealing with textures.
- */
-struct FStreamingManagerTexture : public ITextureStreamingManager
-{
-	/** Constructor, initializing all members */
-	FStreamingManagerTexture();
-
-	virtual ~FStreamingManagerTexture();
-
-	/**
-	 * Updates streaming, taking into account all current view infos. Can be called multiple times per frame.
-	 *
-	 * @param DeltaTime				Time since last call in seconds
-	 * @param bProcessEverything	[opt] If true, process all resources with no throttling limits
-	 */
-	virtual void UpdateResourceStreaming( float DeltaTime, bool bProcessEverything=false ) override;
-
-	/**
-	 * Updates streaming for an individual texture, taking into account all view infos.
-	 *
-	 * @param Texture	Texture to update
-	 */
-	virtual void UpdateIndividualTexture( UTexture2D* Texture ) override;
-
-	/**
-	 * Blocks till all pending requests are fulfilled.
-	 *
-	 * @param TimeLimit		Optional time limit for processing, in seconds. Specifying 0 means infinite time limit.
-	 * @param bLogResults	Whether to dump the results to the log.
-	 * @return				Number of streaming requests still in flight, if the time limit was reached before they were finished.
-	 */
-	virtual int32 BlockTillAllRequestsFinished( float TimeLimit = 0.0f, bool bLogResults = false ) override;
-
-	/**
-	 * Cancels the timed Forced resources (i.e used the Kismet action "Stream In Textures").
-	 */
-	virtual void CancelForcedResources() override;
-
-	/**
-	 * Notifies manager of "level" change so it can prioritize character textures for a few frames.
-	 */
-	virtual void NotifyLevelChange() override;
-
-	/**
-	 * Adds a textures streaming handler to the array of handlers used to determine which
-	 * miplevels need to be streamed in/ out.
-	 *
-	 * @param TextureStreamingHandler	Handler to add
-	 */
-	void AddTextureStreamingHandler( FStreamingHandlerTextureBase* TextureStreamingHandler );
-
-	/**
-	 * Removes a textures streaming handler from the array of handlers used to determine which
-	 * miplevels need to be streamed in/ out.
-	 *
-	 * @param TextureStreamingHandler	Handler to remove
-	 */
-	void RemoveTextureStreamingHandler( FStreamingHandlerTextureBase* TextureStreamingHandler );
-
-	/** Don't stream world resources for the next NumFrames. */
-	virtual void SetDisregardWorldResourcesForFrames( int32 NumFrames ) override;
-
-	/**
-	 *	Try to stream out texture mip-levels to free up more memory.
-	 *	@param RequiredMemorySize	- Required minimum available texture memory
-	 *	@return						- Whether it succeeded or not
-	 **/
-	virtual bool StreamOutTextureData( int64 RequiredMemorySize ) override;
-
-	virtual int64 GetMemoryOverBudget() const override { return MemoryOverBudget; }
-
-	virtual int64 GetMaxEverRequired() const override { return MaxEverRequired; }
-
-	virtual void ResetMaxEverRequired() override { MaxEverRequired = 0; }
-
-	/**
-	 * Allows the streaming manager to process exec commands.
-	 *
-	 * @param InWorld World context
-	 * @param Cmd	Exec command
-	 * @param Ar	Output device for feedback
-	 * @return		true if the command was handled
-	 */
-	virtual bool Exec( UWorld* InWorld, const TCHAR* Cmd, FOutputDevice& Ar ) override;
-
-	/**
-	 * Exec command handlers
-	 */
-#if STATS_FAST
-	bool HandleDumpTextureStreamingStatsCommand( const TCHAR* Cmd, FOutputDevice& Ar );
-#endif // STATS_FAST
-#if STATS
-	bool HandleListStreamingTexturesCommand( const TCHAR* Cmd, FOutputDevice& Ar );
-	bool HandleListStreamingTexturesCollectCommand( const TCHAR* Cmd, FOutputDevice& Ar );
-	bool HandleListStreamingTexturesReportReadyCommand( const TCHAR* Cmd, FOutputDevice& Ar );
-	bool HandleListStreamingTexturesReportCommand( const TCHAR* Cmd, FOutputDevice& Ar );
-#endif // STATS
-#if !UE_BUILD_SHIPPING
-	bool HandleResetMaxEverRequiredTexturesCommand(const TCHAR* Cmd, FOutputDevice& Ar);
-	bool HandleLightmapStreamingFactorCommand( const TCHAR* Cmd, FOutputDevice& Ar );
-	bool HandleCancelTextureStreamingCommand( const TCHAR* Cmd, FOutputDevice& Ar );
-	bool HandleShadowmapStreamingFactorCommand( const TCHAR* Cmd, FOutputDevice& Ar );
-	bool HandleNumStreamedMipsCommand( const TCHAR* Cmd, FOutputDevice& Ar );
-	bool HandleTrackTextureCommand( const TCHAR* Cmd, FOutputDevice& Ar );
-	bool HandleListTrackedTexturesCommand( const TCHAR* Cmd, FOutputDevice& Ar );
-	bool HandleDebugTrackedTexturesCommand( const TCHAR* Cmd, FOutputDevice& Ar );
-	bool HandleUntrackTextureCommand( const TCHAR* Cmd, FOutputDevice& Ar );
-	bool HandleStreamOutCommand( const TCHAR* Cmd, FOutputDevice& Ar );
-	bool HandlePauseTextureStreamingCommand( const TCHAR* Cmd, FOutputDevice& Ar );
-	bool HandleStreamingManagerMemoryCommand( const TCHAR* Cmd, FOutputDevice& Ar, UWorld* InWorld );
-	bool HandleTextureGroupsCommand( const TCHAR* Cmd, FOutputDevice& Ar );
-	bool HandleInvestigateTextureCommand( const TCHAR* Cmd, FOutputDevice& Ar );
-#endif // !UE_BUILD_SHIPPING
-	/** Adds a new texture to the streaming manager. */
-	virtual void AddStreamingTexture( UTexture2D* Texture ) override;
-
-	/** Removes a texture from the streaming manager. */
-	virtual void RemoveStreamingTexture( UTexture2D* Texture ) override;
-
-	/** Adds a ULevel to the streaming manager. */
-	virtual void AddPreparedLevel( class ULevel* Level ) override;
-
-	/** Removes a ULevel from the streaming manager. */
-	virtual void RemoveLevel( class ULevel* Level ) override;
-
-	/** Called when an actor is spawned. */
-	virtual void NotifyActorSpawned( AActor* Actor ) override;
-
-	/** Called when a spawned actor is destroyed. */
-	virtual void NotifyActorDestroyed( AActor* Actor ) override;
-
-	/**
-	 * Called when a primitive is attached to an actor or another component.
-	 * Replaces previous info, if the primitive was already attached.
-	 *
-	 * @param InPrimitive	Newly attached dynamic/spawned primitive
-	 */
-	virtual void NotifyPrimitiveAttached( const UPrimitiveComponent* Primitive, EDynamicPrimitiveType DynamicType ) override;
-
-	/** Called when a primitive is detached from an actor or another component. */
-	virtual void NotifyPrimitiveDetached( const UPrimitiveComponent* Primitive ) override;
-
-	/**
-	 * Called when a LastRenderTime primitive is attached to an actor or another component.
-	 * Modifies the LastRenderTimeRefCount for the textures used, so that those textures can
-	 * use both distance-based and LastRenderTime heuristics.
-	 *
-	 * @param Primitive	Newly attached dynamic/spawned primitive
-	 */
-	virtual void NotifyTimedPrimitiveAttached( const UPrimitiveComponent* Primitive, EDynamicPrimitiveType DynamicType ) override;
-
-	/**
-	 * Called when a LastRenderTime primitive is detached from an actor or another component.
-	 * Modifies the LastRenderTimeRefCount for the textures used, so that those textures can
-	 * use both distance-based and LastRenderTime heuristics.
-	 *
-	 * @param Primitive	Newly detached dynamic/spawned primitive
-	 */
-	virtual void NotifyTimedPrimitiveDetached( const UPrimitiveComponent* Primitive ) override;
-
-	/**
-	 * Called when a primitive has had its textured changed.
-	 * Only affects primitives that were already attached.
-	 * Replaces previous info.
-	 */
-	virtual void NotifyPrimitiveUpdated( const UPrimitiveComponent* Primitive ) override;
-
-	/** Returns the corresponding FStreamingTexture for a UTexture2D. */
-	FStreamingTexture& GetStreamingTexture( const UTexture2D* Texture2D );
-
-	/** Returns true if this is a streaming texture that is managed by the streaming manager. */
-	virtual bool IsManagedStreamingTexture( const UTexture2D* Texture2D ) override;
-
-	/** Updates the I/O state of a texture (allowing it to progress to the next stage) and some stats. */
-	void UpdateTextureStatus( FStreamingTexture& StreamingTexture, FStreamingContext& Context );
-
-	/**
-	 * Cancels the current streaming request for the specified texture.
-	 *
-	 * @param StreamingTexture		Texture to cancel streaming for
-	 * @return						true if a streaming request was canceled
-	 */
-	bool CancelStreamingRequest( FStreamingTexture& StreamingTexture );
-
-	/** Resets the streaming statistics to zero. */
-	void ResetStreamingStats();
-
-	/**
-	 * Updates the streaming statistics with current frame's worth of stats.
-	 *
-	 * @param Context					Context for the current frame
-	 * @param bAllTexturesProcessed		Whether all processing is complete
-	 */
-	void UpdateStreamingStats( const FStreamingContext& Context, bool bAllTexturesProcessed );
-
-	/** Returns the number of cached view infos (thread-safe). */
-	int32 ThreadNumViews()
-	{
-		return ThreadSettings.ThreadViewInfos.Num();
-	}
-
-	/** Returns a cached view info for the specified index (thread-safe). */
-	const FStreamingViewInfo& ThreadGetView( int32 ViewIndex )
-	{
-		return ThreadSettings.ThreadViewInfos[ ViewIndex ];
-	}
-
-#if STATS
-	/** Ringbuffer of bandwidth samples for streaming in mip-levels (MB/sec). */
-	static float BandwidthSamples[NUM_BANDWIDTHSAMPLES];
-	/** Number of bandwidth samples that have been filled in. Will stop counting when it reaches NUM_BANDWIDTHSAMPLES. */
-	static int32 NumBandwidthSamples;
-	/** Current sample index in the ring buffer. */
-	static int32 BandwidthSampleIndex;
-	/** Average of all bandwidth samples in the ring buffer, in MB/sec. */
-	static float BandwidthAverage;
-	/** Maximum bandwidth measured since the start of the game.  */
-	static float BandwidthMaximum;
-	/** Minimum bandwidth measured since the start of the game.  */
-	static float BandwidthMinimum;
-#endif
-
-	/** Set current pause state for texture streaming */
-	virtual void PauseTextureStreaming(bool bInShouldPause) override
-	{
-		bPauseTextureStreaming = bInShouldPause;
-	}
-
-protected:
-//BEGIN: Thread-safe functions and data
-		friend class FAsyncTextureStreaming;
-		friend struct FStreamingHandlerTextureStatic;
-		friend struct FStreamingHandlerTextureDynamic;
-
-		/** Calculates the minimum and maximum number of mip-levels for a streaming texture. */
-		void CalcMinMaxMips( FStreamingTexture& StreamingTexture );
-
-		/** Calculates the number of mip-levels we would like to have in memory for a texture. */
-		void CalcWantedMips( FStreamingTexture& StreamingTexture );
-
-		/**
-		 * Fallback handler to catch textures that have been orphaned recently.
-		 * This handler prevents massive spike in texture memory usage.
-		 * Orphaned textures were previously streamed based on distance but those instance locations have been removed -
-		 * for instance because a ULevel is being unloaded. These textures are still active and in memory until they are garbage collected,
-		 * so we must ensure that they do not start using the LastRenderTime fallback handler and suddenly stream in all their mips -
-		 * just to be destroyed shortly after by a garbage collect.
-		 */
-		FFloatMipLevel GetWantedMipsForOrphanedTexture( FStreamingTexture& StreamingTexture, float& Distance );
-
-		/** Updates this frame's STATs by one texture. */
-		void UpdateFrameStats( FStreamingContext& Context, FStreamingTexture& StreamingTexture, int32 TextureIndex );
-
-		/**
-		 * Not thread-safe: Updates a portion (as indicated by 'StageIndex') of all streaming textures,
-		 * allowing their streaming state to progress.
-		 *
-		 * @param Context			Context for the current stage (frame)
-		 * @param StageIndex		Current stage index
-		 * @param NumUpdateStages	Number of texture update stages
-		 */
-		void UpdateStreamingTextures( FStreamingContext& Context, int32 StageIndex, int32 NumStages );
-
-		/** Adds new textures and level data on the gamethread (while the worker thread isn't active). */
-		void UpdateThreadData();
-
-		/** Checks for updates in the user settings (CVars, etc). */
-		void CheckUserSettings();
-
-		/**
-		 * Temporarily boosts the streaming distance factor by the specified number.
-		 * This factor is automatically reset to 1.0 after it's been used for mip-calculations.
-		 */
-		void BoostTextures( AActor* Actor, float BoostFactor ) override;
-
-		/**
-		 * Stream textures in/out, based on the priorities calculated by the async work.
-		 *
-		 * @param bProcessEverything	Whether we're processing all textures in one go
-		 */
-		void StreamTextures( bool bProcessEverything );
-
-		/**
-		 * Keep as many unnecessary mips as possible.
-		 *
-		 * @param Context				Context for the current stage
-		 */
-		void KeepUnwantedMips( FStreamingContext& Context );
-
-		/**
-		 * Drop as many mips required to fit in budget.
-		 *
-		 * @param Context				Context for the current stage
-		 */
-		void DropWantedMips( FStreamingContext& Context );
-
-		/**
-		 * Drop as many forced required to fit in budget.
-		 *
-		 * @param Context				Context for the current stage
-		 */
-		void DropForcedMips( FStreamingContext& Context );
-
-		/** Thread-safe helper struct for per-level information. */
-		struct FThreadLevelData
-		{
-			FThreadLevelData()
-			:	bRemove( false )
-			{
-			}
-
-			/** Whether this level has been removed. */
-			bool bRemove;
-
-			/** Texture instances used in this level, stored in SIMD-friendly format. */
-			TMap<const UTexture2D*,TArray<FStreamableTextureInstance4> > ThreadTextureInstances;
-		};
-
-		typedef TKeyValuePair< class ULevel*, FThreadLevelData >	FLevelData;
-
-		/** Thread-safe helper struct for streaming information. */
-		struct FThreadSettings
-		{
-			FThreadSettings() : TextureBoundsVisibility(nullptr) {}
-			
-			/** Cached from the system settings. */
-			int32 NumStreamedMips[TEXTUREGROUP_MAX];
-
-			/** Cached from each ULevel. */
-			TArray< FLevelData > LevelData;
-
-			/** Cached from FStreamingManagerBase. */
-			TArray<FStreamingViewInfo> ThreadViewInfos;
-
-			FTextureBoundsVisibility* TextureBoundsVisibility;
-
-			/** from cvar, >=0 */
-			float MipBias;
-		};
-
-
-		/** Thread-safe helper data for streaming information. */
-		FThreadSettings	ThreadSettings;
-
-		/** All streaming UTexture2D objects. */
-		TArray<FStreamingTexture> StreamingTextures;
-
-		/** Index of the StreamingTexture that will be updated next by UpdateStreamingTextures(). */
-		int32 CurrentUpdateStreamingTextureIndex;
-//END: Thread-safe functions and data
-
-	/**
-	 * Mark the textures with a timestamp. They're about to lose their location-based heuristic and we don't want them to
-	 * start using LastRenderTime heuristic for a few seconds until they are garbage collected!
-	 *
-	 * @param RemovedTextures	List of removed textures.
-	 */
-	void	SetTexturesRemovedTimestamp(const FRemovedTextureArray& RemovedTextures);
-
-	void	DumpTextureGroupStats( bool bDetailedStats );
-
-	/**
-	 * Prints out detailed information about streaming textures that has a name that contains the given string.
-	 * Triggered by the InvestigateTexture exec command.
-	 *
-	 * @param InvestigateTextureName	Partial name to match textures against
-	 */
-	void	InvestigateTexture( const FString& InvestigateTextureName );
-
-	void	DumpTextureInstances( const UPrimitiveComponent* Primitive, UTexture2D* Texture2D );
-
-	/** Next sync, dump texture group stats. */
-	bool	bTriggerDumpTextureGroupStats;
-
-	/** Whether to the dumped texture group stats should contain extra information. */
-	bool	bDetailedDumpTextureGroupStats;
-
-	/** Next sync, dump all information we have about a certain texture. */
-	bool	bTriggerInvestigateTexture;
-
-	/** Name of a texture to investigate. Can be partial name. */
-	FString	InvestigateTextureName;
-
-	/** Async work for calculating priorities for all textures. */
-	FAsyncTask<FAsyncTextureStreaming>*	AsyncWork;
-
-	/** Textures from dynamic primitives. */
-	FDynamicComponentTextureManager* DynamicComponentTextureManager;
-
-	/** New textures, before they've been added to the thread-safe container. */
-	TArray<UTexture2D*>		PendingStreamingTextures;
-
-	/** New levels, before they've been added to the thread-safe container. */
-	TArray<class ULevel*>	PendingLevels;
-
-	/** Stages [0,N-2] is non-threaded data collection, Stage N-1 is wait-for-AsyncWork-and-finalize. */
-	int32					ProcessingStage;
-
-	/** Total number of processing stages (N). */
-	int32					NumTextureProcessingStages;
-
-	/** Whether to support texture instance streaming for dynamic (movable/spawned) objects. */
-	bool					bUseDynamicStreaming;
-
-	float					BoostPlayerTextures;
-
-	/** Extra distance added to the range test, to start streaming the texture before they are actually used. */
-	float					RangePrefetchDistance;
-
-	/** Array of texture streaming objects to use during update. */
-	TArray<FStreamingHandlerTextureBase*> TextureStreamingHandlers;
-
-	/** Amount of memory to leave free in the texture pool. */
-	int64					MemoryMargin;
-
-	/** Minimum number of bytes to evict when we need to stream out textures because of a failed allocation. */
-	int64					MinEvictSize;
-
-	/** The actual memory pool size available to stream textures, excludes non-streaming texture, temp memory (for streaming mips), memory margin (allocator overhead). */
-	int64					EffectiveStreamingPoolSize;
-
-	/** If set, UpdateResourceStreaming() will only process this texture. */
-	UTexture2D*				IndividualStreamingTexture;
-
-	// Stats we need to keep across frames as we only iterate over a subset of textures.
-
-	/** Number of streaming textures */
-	uint32 NumStreamingTextures;
-	/** Number of requests in cancelation phase. */
-	uint32 NumRequestsInCancelationPhase;
-	/** Number of requests in mip update phase. */
-	uint32 NumRequestsInUpdatePhase;
-	/** Number of requests in mip finalization phase. */
-	uint32 NumRequestsInFinalizePhase;
-	/** Size ot all intermerdiate textures in flight. */
-	uint32 TotalIntermediateTexturesSize;
-	/** Number of intermediate textures in flight. */
-	uint32 NumIntermediateTextures;
-	/** Size of all streaming testures. */
-	uint64 TotalStreamingTexturesSize;
-	/** Maximum size of all streaming textures. */
-	uint64 TotalStreamingTexturesMaxSize;
-	/** Total number of bytes in memory, for all streaming lightmap textures. */
-	uint64 TotalLightmapMemorySize;
-	/** Total number of bytes on disk, for all streaming lightmap textures. */
-	uint64 TotalLightmapDiskSize;
-	/** Total number of bytes in memory, for all HLOD textures. */
-	uint64 TotalHLODMemorySize;
-	/** Total number of bytes on disk, for all HLOD textures. */
-	uint64 TotalHLODDiskSize;
-	/** Number of mip count increase requests in flight. */
-	uint32 TotalMipCountIncreaseRequestsInFlight;
-	/** Total number of bytes in memory, if all textures were streamed perfectly with a 1.0 fudge factor. */
-	uint64 TotalOptimalWantedSize;
-	/** Total number of bytes using StaticTexture heuristics, currently in memory. */
-	uint64 TotalStaticTextureHeuristicSize;
-	/** Total number of bytes using DynmicTexture heuristics, currently in memory. */
-	uint64 TotalDynamicTextureHeuristicSize;
-	/** Total number of bytes using LastRenderTime heuristics, currently in memory. */
-	uint64 TotalLastRenderHeuristicSize;
-	/** Total number of bytes using ForcedIntoMemory heuristics, currently in memory. */
-	uint64 TotalForcedHeuristicSize;
-	int64 MemoryOverBudget;
-	int64 MaxEverRequired;
-
-	/** Unmodified texture pool size, in bytes, as specified in the .ini file. */
-	int64 OriginalTexturePoolSize;
-	/** Timestamp when we last shrunk the pool size because of memory usage. */
-	double PreviousPoolSizeTimestamp;
-	/** PoolSize CVar setting the last time we adjusted the pool size. */
-	int32 PreviousPoolSizeSetting;
-
-	/** Whether to collect, and optionally report, texture stats for the next run. */
-	bool   bCollectTextureStats;
-	bool   bReportTextureStats;
-	TArray<FString> TextureStatsReport;
-
-	/** Optional string to match against the texture names when collecting stats. */
-	FString CollectTextureStatsName;
-
-	/** Whether texture streaming is paused or not. When paused, it won't stream any textures in or out. */
-	bool bPauseTextureStreaming;
-
-#if STATS
-	/**
-	 * Ring buffer containing all latency samples we keep track of, as measured in seconds from the time the streaming system
-	 * detects that a change is required until the data has been streamed in and the new texture is ready to be used.
-	 */
-	float LatencySamples[NUM_LATENCYSAMPLES];
-	/** Number of latency samples that have been filled in. Will stop counting when it reaches NUM_LATENCYSAMPLES. */
-	int32 NumLatencySamples;
-	/** Current sample index in the ring buffer. */
-	int32 LatencySampleIndex;
-	/** Average of all latency samples in the ring buffer, in seconds. */
-	float LatencyAverage;
-	/** Maximum latency measured since the start of the game.  */
-	float LatencyMaximum;
-
-	/**
-	 * Updates the streaming latency STAT for a texture.
-	 *
-	 * @param Texture		Texture to update for
-	 * @param WantedMips	Number of desired mip-levels for the texture
-	 * @param bInFlight		Whether the texture is currently being streamed
-	 */
-	void UpdateLatencyStats( UTexture2D* Texture, int32 WantedMips, bool bInFlight );
-
-	/** Total time taken for each processing stage. */
-	TArray<double> StreamingTimes;
-#endif
-
-#if STATS_FAST
-	uint64 MaxStreamingTexturesSize;
-	uint64 MaxOptimalTextureSize;
-	int64 MaxStreamingOverBudget;
-	uint64 MaxTexturePoolAllocatedSize;
-	uint64 MinLargestHoleSize;
-	uint32 MaxNumWantingTextures;
-#endif
-	
-	friend bool TrackTextureEvent( FStreamingTexture* StreamingTexture, UTexture2D* Texture, bool bForceMipLevelsToBeResident, const FStreamingManagerTexture* Manager);
-};
-
-/*-----------------------------------------------------------------------------
-	Texture streaming handler.
------------------------------------------------------------------------------*/
-
-/**
- * Base of texture streaming handler functionality.
- */
-struct FStreamingHandlerTextureBase
-{
-	/** Friendly name identifying the handler for debug purposes. */
-	const TCHAR* HandlerName;
-
-	/** Default constructor. */
-	FStreamingHandlerTextureBase()
-		: HandlerName(TEXT("Base"))
-	{
-	}
-
-	/**
-	 * Returns mip count wanted by this handler for the passed in texture. 
-	 * 
-	 * @param	StreamingManager	Streaming manager
-	 * @param	Streaming Texture	Texture to determine wanted mip count for
-	 * @param	MinDistance			[out] Distance to the closest instance of this texture, in world space units.
-	 * @return	Number of miplevels that should be streamed in or INDEX_NONE if texture isn't handled by this handler.
-	 */
-	virtual FFloatMipLevel GetWantedMips( FStreamingManagerTexture& StreamingManager, FStreamingTexture& StreamingTexture, float& MinDistance ) = 0;
-};
-
-/**
- * Static texture streaming handler. Used to stream textures on static level geometry.
- */
-struct FStreamingHandlerTextureStatic : public FStreamingHandlerTextureBase
-{
-	/** Default constructor. */
-	FStreamingHandlerTextureStatic()
-	{
-		HandlerName = TEXT("Static");
-	}
-
-	/**
-	 * Returns mip count wanted by this handler for the passed in texture. 
-	 * 
-	 * @param	StreamingManager	Streaming manager
-	 * @param	Streaming Texture	Texture to determine wanted mip count for
-	 * @param	MinDistance			[out] Distance to the closest instance of this texture, in world space units.
-	 * @return	Number of miplevels that should be streamed in or INDEX_NONE if texture isn't handled by this handler.
-	 */
-	virtual FFloatMipLevel GetWantedMips( FStreamingManagerTexture& StreamingManager, FStreamingTexture& StreamingTexture, float& MinDistance );
-};
-
-
-/**
- * Dynamic texture streaming handler. Used to stream textures for dynamic primitives.
- */
-struct FStreamingHandlerTextureDynamic : public FStreamingHandlerTextureBase
-{
-	/** Default constructor. */
-	FStreamingHandlerTextureDynamic()
-	{
-		HandlerName = TEXT("Dynamic");
-	}
-
-	/**
-	 * Returns mip count wanted by this handler for the passed in texture. 
-	 * 
-	 * @param	StreamingManager	Streaming manager
-	 * @param	Streaming Texture	Texture to determine wanted mip count for
-	 * @param	MinDistance			[out] Distance to the closest instance of this texture, in world space units.
-	 * @return	Number of miplevels that should be streamed in or INDEX_NONE if texture isn't handled by this handler.
-	 */
-	virtual FFloatMipLevel GetWantedMips( FStreamingManagerTexture& StreamingManager, FStreamingTexture& StreamingTexture, float& MinDistance );
-};
-
-/**
- * Streaming handler that bases its decision on the last render time.
- */
-struct FStreamingHandlerTextureLastRender : public FStreamingHandlerTextureBase
-{
-	/** Default constructor. */
-	FStreamingHandlerTextureLastRender()
-	{
-		HandlerName = TEXT("LastRender");
-	}
-
-	/**
-	 * Returns mip count wanted by this handler for the passed in texture. 
-	 * 
-	 * @param	StreamingManager	Streaming manager
-	 * @param	Streaming Texture	Texture to determine wanted mip count for
-	 * @param	MinDistance			[out] Distance to the closest instance of this texture, in world space units.
-	 * @return	Number of miplevels that should be streamed in or INDEX_NONE if texture isn't handled by this handler.
-	 */
-	virtual FFloatMipLevel GetWantedMips( FStreamingManagerTexture& StreamingManager, FStreamingTexture& StreamingTexture, float& MinDistance );
-};
-
-/**
- *	Streaming handler that bases its decision on the bForceMipStreaming flag in PrimitiveComponent.
- */
-struct FStreamingHandlerTextureLevelForced : public FStreamingHandlerTextureBase
-{
-	/** Default constructor. */
-	FStreamingHandlerTextureLevelForced()
-	{
-		HandlerName = TEXT("LevelForced");
-	}
-
-	/**
-	 * Returns mip count wanted by this handler for the passed in texture. 
-	 * 
-	 * @param	StreamingManager	Streaming manager
-	 * @param	Streaming Texture	Texture to determine wanted mip count for
-	 * @param	MinDistance			[out] Distance to the closest instance of this texture, in world space units.
-	 * @return	Number of miplevels that should be streamed in or INDEX_NONE if texture isn't handled by this handler.
-	 */
-	virtual FFloatMipLevel GetWantedMips( FStreamingManagerTexture& StreamingManager, FStreamingTexture& StreamingTexture, float& MinDistance );
-};
diff --git a/Engine/Source/Runtime/Engine/Public/DebugViewModeHelpers.h b/Engine/Source/Runtime/Engine/Public/DebugViewModeHelpers.h
new file mode 100644
--- /dev/null
+++ b/Engine/Source/Runtime/Engine/Public/DebugViewModeHelpers.h
@@ -0,0 +1,43 @@
+// Copyright 1998-2016 Epic Games, Inc. All Rights Reserved.
+
+/**
+ * Definition and helpers for debug view modes
+ */
+
+#pragma once
+
+/** 
+ * Enumeration for different Quad Overdraw visualization mode.
+ */
+enum EDebugViewShaderMode
+{
+	DVSM_None,						// No debug view.
+	DVSM_ShaderComplexity,			// Default shader complexity viewmode
+	DVSM_ShaderComplexityContainedQuadOverhead,	// Show shader complexity with quad overdraw scaling the PS instruction count.
+	DVSM_ShaderComplexityBleedingQuadOverhead,	// Show shader complexity with quad overdraw bleeding the PS instruction count over the quad.
+	DVSM_QuadComplexity,			// Show quad overdraw only.
+	DVSM_WantedMipsAccuracy,		// Accuraty of the wanted mips computed by the texture streamer.
+	DVSM_TexelFactorAccuracy,		// Accuraty of the texel factor computed on each mesh.
+	DVSM_TexCoordScaleAccuracy,		// Accuracy of the material texture coordinate scales computed in the texture streaming build.
+	DVSM_TexCoordScaleAnalysis,		// Outputs each texture lookup texcoord scale.
+	DVSM_MAX
+};
+
+#if !(UE_BUILD_SHIPPING || UE_BUILD_TEST)
+/** Returns true if the specified shadermode is available for the given shader platform. Called for shader compilation shader compilation. */
+extern ENGINE_API bool AllowDebugViewShaderMode(EDebugViewShaderMode ShaderMode, EShaderPlatform Platform);
+/** Returns true if the specified shadermode is currently available for the given feature level on the current platform. */
+extern ENGINE_API bool AllowDebugViewShaderMode(EDebugViewShaderMode ShaderMode, ERHIFeatureLevel::Type FeatureLevel);
+/** Returns true if the vertex shader (and potential hull and domain) should be compiled on the given platform. */
+extern ENGINE_API bool AllowDebugViewVSDSHS(EShaderPlatform Platform);
+/** Returns true if the vertex shader (and potential hull and domain) should be compiled for the given feature level on the current platform. */
+extern ENGINE_API bool AllowDebugViewVSDSHS(ERHIFeatureLevel::Type FeatureLevel);
+#else
+FORCEINLINE bool AllowDebugViewShaderMode(EDebugViewShaderMode ShaderMode, EShaderPlatform Platform) { return false; }
+FORCEINLINE bool AllowDebugViewShaderMode(EDebugViewShaderMode ShaderMode, ERHIFeatureLevel::Type FeatureLevel)  { return false; }
+FORCEINLINE bool AllowDebugViewVSDSHS(EShaderPlatform Platform) { return false; }
+FORCEINLINE bool AllowDebugViewVSDSHS(ERHIFeatureLevel::Type FeatureLevel) { return false; }
+#endif
+
+
+
diff --git a/Engine/Source/Runtime/Engine/Public/Engine.h b/Engine/Source/Runtime/Engine/Public/Engine.h
--- a/Engine/Source/Runtime/Engine/Public/Engine.h
+++ b/Engine/Source/Runtime/Engine/Public/Engine.h
@@ -97,6 +97,7 @@
 #include "Engine/TargetPoint.h"
 #include "VectorField/VectorFieldVolume.h"
 #include "Engine/DataAsset.h"
+#include "Engine/TextureStreamingTypes.h"
 #include "GameFramework/SpectatorPawnMovement.h"
 #include "Vehicles/VehicleWheel.h"
 #include "Vehicles/WheeledVehicleMovementComponent.h"
diff --git a/Engine/Source/Runtime/Engine/Public/GPUSkinCache.h b/Engine/Source/Runtime/Engine/Public/GPUSkinCache.h
--- a/Engine/Source/Runtime/Engine/Public/GPUSkinCache.h
+++ b/Engine/Source/Runtime/Engine/Public/GPUSkinCache.h
@@ -63,11 +63,13 @@ public:
 		RWTangentXOffsetInFloats = 3,	// Packed U8x4N
 		RWTangentZOffsetInFloats = 4,	// Packed U8x4N
 
+		// stride in float(4 bytes) in the SkinCache buffer
 		RWStrideInFloats = 5,
 		// 3 ints for normal, 3 ints for tangent, 1 for orientation = 7, rounded up to 8 as it should result in faster math and caching
 		IntermediateAccumBufferNumInts = 8,
 		// max vertex count we support once the r.SkinCache.RecomputeTangents feature is used (GPU memory in bytes = n * sizeof(int) * IntermediateAccumBufferNumInts)
-		IntermediateAccumBufferSizeInKB = 1024,
+		// 4*1024 is the conservative size needed for Senua in CharDemo with NinjaTheory (this should be a cvar and project specific or dynamic, with only using the feature on RecomputeTangents it can be much less)
+		IntermediateAccumBufferSizeInKB = 4 * 1024,
 	};
 
 	FGPUSkinCache();
@@ -79,14 +81,15 @@ public:
 	void Initialize(FRHICommandListImmediate& RHICmdList);
 	void Cleanup();
 
-	inline bool IsElementProcessed(int32 Key) const
+	// @param FrameNumber from GFrameNumber or better ViewFamily.FrameNumber
+	inline bool IsElementProcessed(uint32 FrameNumber, int32 Key) const
 	{
 		if (!GEnableGPUSkinCache)
 		{
 			return false;
 		}
 
-		return InternalIsElementProcessed(Key);
+		return InternalIsElementProcessed(FrameNumber, Key);
 	}
 
 	// For each SkeletalMeshObject:
@@ -94,40 +97,66 @@ public:
 	//	For each Chunk:
 	//		Call Add*
 	//	Call End*
+	// @param FrameNumber from GFrameNumber or better ViewFamily.FrameNumber
+	// @param Key index in CachedElements[] (redundant, could be computed)
 	// @param MorphVertexBuffer if no morph targets are used
 	// @returns -1 if failed, otherwise index into CachedElements[]
-	int32 StartCacheMesh(FRHICommandListImmediate& RHICmdList, int32 Key, class FGPUBaseSkinVertexFactory* VertexFactory,
+	int32 StartCacheMesh(FRHICommandListImmediate& RHICmdList, uint32 FrameNumber, int32 Key, class FGPUBaseSkinVertexFactory* VertexFactory,
 		class FGPUSkinPassthroughVertexFactory* TargetVertexFactory, const struct FSkelMeshChunk& BatchElement, FSkeletalMeshObjectGPUSkin* Skin,
 		const class FMorphVertexBuffer* MorphVertexBuffer);
 
-	inline bool SetVertexStreamFromCache(FRHICommandList& RHICmdList, int32 Key, FShader* Shader, const FGPUSkinPassthroughVertexFactory* VertexFactory, uint32 BaseVertexIndex, FShaderParameter PreviousStreamFloatOffset, FShaderResourceParameter PreviousStreamBuffer)
+	// @param FrameNumber from GFrameNumber or better ViewFamily.FrameNumber
+	inline void SetVertexStreamFromCache(FRHICommandList& RHICmdList, uint32 FrameNumber, int32 Key, FShader* Shader, const FGPUSkinPassthroughVertexFactory* VertexFactory,
+		uint32 BaseVertexIndex, FShaderParameter PreviousStreamFloatOffset, FShaderResourceParameter PreviousStreamBuffer)
 	{
 		if (Key >= 0 && Key < CachedElements.Num())
 		{
-			return InternalSetVertexStreamFromCache(RHICmdList, Key, Shader, VertexFactory, BaseVertexIndex, PreviousStreamFloatOffset, PreviousStreamBuffer);
+			InternalSetVertexStreamFromCache(RHICmdList, FrameNumber, Key, Shader, VertexFactory, BaseVertexIndex, PreviousStreamFloatOffset, PreviousStreamBuffer);
+			return;
 		}
 
-		return false;
+		ensure(0);
+		// should never happen but better we render garbage than crash
 	}
 
 	struct FElementCacheStatusInfo
 	{
+		// key section -------------
+
+		// TODO: using a pointer as key is dangerous
 		const FSkelMeshChunk* BatchElement;
-		const FVertexFactory* VertexFactory;
-		const FVertexFactory* TargetVertexFactory;
+		// TODO: using a pointer as key is dangerous
 		const FSkeletalMeshObjectGPUSkin* Skin;
 
+		// value section -------------
+
+		// index in CachedElements[] (redundant, could be computed)
 		int32	Key;
-	
+		// from GFrameNumber
 		uint32	FrameUpdated;
+		// start index in SkinCacheBuffer[UAVIndex]
+		uint32 StreamOffset;
+		// for velocity rendering (previous StreamOffset), start index in SkinCacheBuffer[UAVIndex]
+		uint32 PreviousFrameStreamOffset;
+		// from GFrameNumber
+		uint32 PreviousFrameUpdated;
+
+		// debugging section -------------
 
+		//
+		uint32 NumVertices;
+		//
+		uint32 BaseVertexIndex;
+		// input stride in bytes
 		uint32	InputVBStride;
-		// for non velocity rendering
-		uint32	StreamOffset;
-		// for velocity rendering (previous StreamOffset)
-		uint32	PreviousFrameStreamOffset;
+		//
+		bool bExtraBoneInfluences;
+		//
+		const FVertexFactory* VertexFactory;
+		//
+		const FVertexFactory* TargetVertexFactory;
 
-		bool	bExtraBoneInfluences;
+		// -----------------------------
 
 		struct FSearchInfo
 		{
@@ -156,6 +185,8 @@ public:
 	{
 		FRHICommandListImmediate& RHICmdList;
 		const FSkelMeshChunk& Chunk;
+		// for debugging / draw events, -1 if not set
+		uint32 ChunkIdx;
 		ERHIFeatureLevel::Type FeatureLevel;
 		
 		// must not be 0
@@ -163,14 +194,19 @@ public:
 
 		// 0:normal, 1:with morph target, 2:with APEX cloth (not yet implemented)
 		uint32 SkinType;
+		//
+		bool bExtraBoneInfluences;
 
 		// SkinCache output input for RecomputeSkinTagents
 		// 0 if no morph
 		FRWBuffer* SkinCacheBuffer;
-		uint32 OutputBufferFloatOffset;
+		// in floats (4 bytes)
+		uint32 SkinCacheStart;
 		uint32 NumVertices;
 
-		uint32 InputStreamFloatOffset;
+		// if floats (4 bytes)
+		uint32 InputStreamStart;
+		// in bytes
 		uint32 InputStreamStride;
 		FShaderResourceViewRHIRef InputVertexBufferSRV;
 
@@ -186,13 +222,15 @@ public:
 		FDispatchData(FRHICommandListImmediate& InRHICmdList, const FSkelMeshChunk& InChunk, ERHIFeatureLevel::Type InFeatureLevel, FSkeletalMeshObjectGPUSkin* InGPUSkin)
 			: RHICmdList(InRHICmdList)
 			, Chunk(InChunk)
+			, ChunkIdx(-1)
 			, FeatureLevel(InFeatureLevel)
 			, GPUSkin(InGPUSkin)
 			, SkinType(0)
+			, bExtraBoneInfluences(false)
 			, SkinCacheBuffer(0)
-			, OutputBufferFloatOffset(0)
+			, SkinCacheStart(0)
 			, NumVertices(0)
-			, InputStreamFloatOffset(0)
+			, InputStreamStart(0)
 			, InputStreamStride(0)
 			, MorphBuffer(0)
 			, MorphBufferOffset(0)
@@ -209,19 +247,25 @@ public:
 
 
 private:
-	FElementCacheStatusInfo* FindEvictableCacheStatusInfo();
+	// @param FrameNumber from GFrameNumber or better ViewFamily.FrameNumber
+	FElementCacheStatusInfo* FindEvictableCacheStatusInfo(uint32 FrameNumber);
 
 	void DispatchSkinCacheProcess(
 		const struct FVertexBufferAndSRV& BoneBuffer, FUniformBufferRHIRef UniformBuffer,
-		const FVector& MeshOrigin, const FVector& MeshExtension, bool bUseExtraBoneInfluences,
+		const FVector& MeshOrigin, const FVector& MeshExtension,
 		const FDispatchData& DispatchData);
-
+	
 	void DispatchUpdateSkinTangents(const FDispatchData& DispatchData);
 
-	bool InternalSetVertexStreamFromCache(FRHICommandList& RHICmdList, int32 Key, FShader* Shader, const class FGPUSkinPassthroughVertexFactory* VertexFactory, uint32 BaseVertexIndex, FShaderParameter PreviousStreamFloatOffset, FShaderResourceParameter PreviousStreamBuffer);
+	// @param FrameNumber from GFrameNumber or better ViewFamily.FrameNumber
+	void InternalSetVertexStreamFromCache(FRHICommandList& RHICmdList, uint32 FrameNumber, int32 Key, FShader* Shader, const class FGPUSkinPassthroughVertexFactory* VertexFactory,
+		uint32 BaseVertexIndex, FShaderParameter PreviousStreamFloatOffset, FShaderResourceParameter PreviousStreamBuffer);
 
 	void TransitionAllToWriteable(FRHICommandList& RHICmdList);
-	bool InternalIsElementProcessed(int32 Key) const;
+
+	// @param FrameNumber from GFrameNumber or better ViewFamily.FrameNumber
+	bool InternalIsElementProcessed(uint32 FrameNumber, int32 Key) const;
+
 
 	uint32 ComputeRecomputeTangentMaxVertexCount() const
 	{
@@ -229,9 +273,10 @@ private:
 		return IntermediateAccumBufferSizeInB / (sizeof(int32) * FGPUSkinCache::IntermediateAccumBufferNumInts);
 	}
 
+	// from GFrameNumber or better ViewFamily.FrameNumber
 	bool bInitialized : 1;
 
-	uint32	FrameCounter;
+	uint32	SkinCacheFrameNumber;
 
 	int		CacheMaxVectorCount;
 	int		CacheCurrentFloatOffset;
@@ -260,6 +305,7 @@ DECLARE_STATS_GROUP(TEXT("GPU Skin Cache"), STATGROUP_GPUSkinCache, STATCAT_Adva
 DECLARE_DWORD_COUNTER_STAT_EXTERN(TEXT("Total Num Chunks"), STAT_GPUSkinCache_TotalNumChunks, STATGROUP_GPUSkinCache,);
 DECLARE_DWORD_COUNTER_STAT_EXTERN(TEXT("Total Num Vertices"), STAT_GPUSkinCache_TotalNumVertices, STATGROUP_GPUSkinCache,);
 DECLARE_DWORD_COUNTER_STAT_EXTERN(TEXT("Total Memory Used"), STAT_GPUSkinCache_TotalMemUsed, STATGROUP_GPUSkinCache,);
+DECLARE_DWORD_COUNTER_STAT_EXTERN(TEXT("Total Memory Wasted"), STAT_GPUSkinCache_TotalMemWasted, STATGROUP_GPUSkinCache,);
 DECLARE_DWORD_COUNTER_STAT_EXTERN(TEXT("Skipped for Max Chunks Per LOD"), STAT_GPUSkinCache_SkippedForMaxChunksPerLOD, STATGROUP_GPUSkinCache,);
 DECLARE_DWORD_COUNTER_STAT_EXTERN(TEXT("Skipped for Num Chunks Per Frame"), STAT_GPUSkinCache_SkippedForChunksPerFrame, STATGROUP_GPUSkinCache,);
 DECLARE_DWORD_COUNTER_STAT_EXTERN(TEXT("Skipped for zero influences"), STAT_GPUSkinCache_SkippedForZeroInfluences, STATGROUP_GPUSkinCache,);
diff --git a/Engine/Source/Runtime/Engine/Public/GlobalShader.h b/Engine/Source/Runtime/Engine/Public/GlobalShader.h
--- a/Engine/Source/Runtime/Engine/Public/GlobalShader.h
+++ b/Engine/Source/Runtime/Engine/Public/GlobalShader.h
@@ -58,7 +58,7 @@ public:
 		ShouldCacheType InShouldCacheRef,
 		GetStreamOutElementsType InGetStreamOutElementsRef
 		):
-		FShaderType(InName,InSourceFilename,InFunctionName,InFrequency,InConstructSerializedRef,InGetStreamOutElementsRef),
+		FShaderType(EShaderTypeForDynamicCast::Global, InName, InSourceFilename, InFunctionName, InFrequency, InConstructSerializedRef, InGetStreamOutElementsRef),
 		ConstructCompiledRef(InConstructCompiledRef),
 		ShouldCacheRef(InShouldCacheRef),
 		ModifyCompilationEnvironmentRef(InModifyCompilationEnvironmentRef)
@@ -98,10 +98,6 @@ public:
 		(*ModifyCompilationEnvironmentRef)(Platform, Environment);
 	}
 
-	// Dynamic casting.
-	virtual FGlobalShaderType* GetGlobalShaderType() override { return this; }
-	virtual const FGlobalShaderType* GetGlobalShaderType() const override { return this; }
-
 private:
 	ConstructCompiledType ConstructCompiledRef;
 	ShouldCacheType ShouldCacheRef;
diff --git a/Engine/Source/Runtime/Engine/Public/MaterialCompiler.h b/Engine/Source/Runtime/Engine/Public/MaterialCompiler.h
--- a/Engine/Source/Runtime/Engine/Public/MaterialCompiler.h
+++ b/Engine/Source/Runtime/Engine/Public/MaterialCompiler.h
@@ -113,10 +113,12 @@ public:
 	virtual int32 Texture(UTexture* Texture,int32& TextureReferenceIndex,ESamplerSourceMode SamplerSource=SSM_FromTextureAsset) = 0;
 	virtual int32 TextureParameter(FName ParameterName,UTexture* DefaultTexture,int32& TextureReferenceIndex,ESamplerSourceMode SamplerSource=SSM_FromTextureAsset) = 0;
 
+	virtual int32 GetTextureReferenceIndex(UTexture* Texture) { return INDEX_NONE; }
+
 	int32 Texture(UTexture* InTexture,ESamplerSourceMode SamplerSource=SSM_FromTextureAsset)
 	{
 		int32 TextureReferenceIndex = INDEX_NONE;
-		return  Texture(InTexture, TextureReferenceIndex, SamplerSource);
+		return Texture(InTexture, TextureReferenceIndex, SamplerSource);
 	}
 
 	int32 TextureParameter(FName ParameterName,UTexture* DefaultTexture,ESamplerSourceMode SamplerSource=SSM_FromTextureAsset)
diff --git a/Engine/Source/Runtime/Engine/Public/MaterialShaderType.h b/Engine/Source/Runtime/Engine/Public/MaterialShaderType.h
--- a/Engine/Source/Runtime/Engine/Public/MaterialShaderType.h
+++ b/Engine/Source/Runtime/Engine/Public/MaterialShaderType.h
@@ -76,7 +76,7 @@ public:
 		ShouldCacheType InShouldCacheRef,
 		GetStreamOutElementsType InGetStreamOutElementsRef
 		):
-		FShaderType(InName,InSourceFilename,InFunctionName,InFrequency,InConstructSerializedRef,InGetStreamOutElementsRef),
+		FShaderType(EShaderTypeForDynamicCast::Material, InName, InSourceFilename, InFunctionName, InFrequency, InConstructSerializedRef, InGetStreamOutElementsRef),
 		ConstructCompiledRef(InConstructCompiledRef),
 		ShouldCacheRef(InShouldCacheRef),
 		ModifyCompilationEnvironmentRef(InModifyCompilationEnvironmentRef)
@@ -128,10 +128,6 @@ public:
 		return (*ShouldCacheRef)(Platform,Material);
 	}
 
-	// Dynamic casting.
-	virtual FMaterialShaderType* GetMaterialShaderType() override { return this; }
-	virtual const FMaterialShaderType* GetMaterialShaderType() const override { return this; }
-
 protected:
 
 	/**
diff --git a/Engine/Source/Runtime/Engine/Public/MaterialShared.h b/Engine/Source/Runtime/Engine/Public/MaterialShared.h
--- a/Engine/Source/Runtime/Engine/Public/MaterialShared.h
+++ b/Engine/Source/Runtime/Engine/Public/MaterialShared.h
@@ -224,6 +224,8 @@ public:
 
 	friend FArchive& operator<<(FArchive& Ar,class FMaterialUniformExpressionTexture*& Ref);
 
+	int32 GetTextureIndex() const { return TextureIndex; }
+
 protected:
 	/** Index into FMaterial::GetReferencedTextures */
 	int32 TextureIndex;
@@ -910,7 +912,8 @@ public:
 		QualityLevel(EMaterialQualityLevel::High),
 		bHasQualityLevelUsage(false),
 		FeatureLevel(ERHIFeatureLevel::SM4),
-		bContainsInlineShaders(false)
+		bContainsInlineShaders(false),
+		bLoadedCookedShaderMapId(false)
 	{}
 
 	/**
@@ -986,7 +989,6 @@ public:
 	virtual bool IsAdaptiveTessellationEnabled() const { return false; }
 	virtual bool IsFullyRough() const { return false; }
 	virtual bool IsUsingHQForwardReflections() const { return false; }
-	virtual bool IsUsingPlanarReflection() const { return false; }
 	virtual bool OutputsVelocityOnBasePass() const { return true; }
 	virtual bool IsNonmetal() const { return false; }
 	virtual bool UseLmDirectionality() const { return true; }
@@ -1108,6 +1110,9 @@ public:
 		checkSlow(IsInGameThread() || IsInAsyncLoadingThread());
 		GameThreadShaderMap = InMaterialShaderMap;
 		bContainsInlineShaders = true;
+		bLoadedCookedShaderMapId = true;
+		CookedShaderMapId = InMaterialShaderMap->GetShaderMapId();
+
 	}
 
 	ENGINE_API class FMaterialShaderMap* GetRenderingThreadShaderMap() const;
@@ -1274,6 +1279,9 @@ private:
 	 * If true, GameThreadShaderMap will contain a reference to the inlined shader map between Serialize and PostLoad.
 	 */
 	uint32 bContainsInlineShaders : 1;
+	uint32 bLoadedCookedShaderMapId : 1;
+
+	FMaterialShaderMapId CookedShaderMapId;
 
 	/**
 	* Compiles this material for Platform, storing the result in OutShaderMap if the compile was synchronous
@@ -1560,7 +1568,6 @@ public:
 	ENGINE_API virtual bool IsAdaptiveTessellationEnabled() const override;
 	ENGINE_API virtual bool IsFullyRough() const override;
 	ENGINE_API virtual bool IsUsingHQForwardReflections() const override;
-	ENGINE_API virtual bool IsUsingPlanarReflection() const override;
 	ENGINE_API virtual bool OutputsVelocityOnBasePass() const override;
 	ENGINE_API virtual bool IsNonmetal() const override;
 	ENGINE_API virtual bool UseLmDirectionality() const override;
diff --git a/Engine/Source/Runtime/Engine/Public/MeshBatch.h b/Engine/Source/Runtime/Engine/Public/MeshBatch.h
--- a/Engine/Source/Runtime/Engine/Public/MeshBatch.h
+++ b/Engine/Source/Runtime/Engine/Public/MeshBatch.h
@@ -39,6 +39,7 @@ struct FMeshBatchElement
 	uint32 NumInstances;
 	uint32 MinVertexIndex;
 	uint32 MaxVertexIndex;
+	// Meaning depends on the vertex factory, e.g. FGPUSkinPassthroughVertexFactory: element index in FGPUSkinCache::CachedElements
 	int32 UserIndex;
 	float MinScreenSize;
 	float MaxScreenSize;
diff --git a/Engine/Source/Runtime/Engine/Public/MeshMaterialShaderType.h b/Engine/Source/Runtime/Engine/Public/MeshMaterialShaderType.h
--- a/Engine/Source/Runtime/Engine/Public/MeshMaterialShaderType.h
+++ b/Engine/Source/Runtime/Engine/Public/MeshMaterialShaderType.h
@@ -47,7 +47,7 @@ public:
 		ShouldCacheType InShouldCacheRef,
 		GetStreamOutElementsType InGetStreamOutElementsRef
 		):
-		FShaderType(InName,InSourceFilename,InFunctionName,InFrequency,InConstructSerializedRef,InGetStreamOutElementsRef),
+		FShaderType(EShaderTypeForDynamicCast::MeshMaterial, InName,InSourceFilename,InFunctionName,InFrequency,InConstructSerializedRef,InGetStreamOutElementsRef),
 		ConstructCompiledRef(InConstructCompiledRef),
 		ShouldCacheRef(InShouldCacheRef),
 		ModifyCompilationEnvironmentRef(InModifyCompilationEnvironmentRef)
@@ -105,10 +105,6 @@ public:
 		return (*ShouldCacheRef)(Platform,Material,VertexFactoryType);
 	}
 
-	// Dynamic casting.
-	virtual FMeshMaterialShaderType* GetMeshMaterialShaderType() override { return this; }
-	virtual const FMeshMaterialShaderType* GetMeshMaterialShaderType() const override { return this; }
-
 protected:
 
 	/**
diff --git a/Engine/Source/Runtime/Engine/Public/ParticleEmitterInstances.h b/Engine/Source/Runtime/Engine/Public/ParticleEmitterInstances.h
--- a/Engine/Source/Runtime/Engine/Public/ParticleEmitterInstances.h
+++ b/Engine/Source/Runtime/Engine/Public/ParticleEmitterInstances.h
@@ -549,23 +549,11 @@ public:
 	/**
 	 *	Retrieves the dynamic data for the emitter
 	 */
-	virtual FDynamicEmitterDataBase* GetDynamicData(bool bSelected)
+	virtual FDynamicEmitterDataBase* GetDynamicData(bool bSelected, ERHIFeatureLevel::Type InFeatureLevel)
 	{
 		return NULL;
 	}
 
-	/**
-	 *	Updates the dynamic data for the instance
-	 *
-	 *	@param	DynamicData		The dynamic data to fill in
-	 *	@param	bSelected		true if the particle system component is selected
-	 */
-	virtual bool UpdateDynamicData(FDynamicEmitterDataBase* DynamicData, bool bSelected)
-	{
-		// Base class does nothing...
-		return false;
-	}
-
 	/**
 	 *	Retrieves replay data for the emitter
 	 *
@@ -795,15 +783,7 @@ struct FParticleSpriteEmitterInstance : public FParticleEmitterInstance
 	/**
 	 *	Retrieves the dynamic data for the emitter
 	 */
-	virtual FDynamicEmitterDataBase* GetDynamicData(bool bSelected) override;
-
-	/**
-	 *	Updates the dynamic data for the instance
-	 *
-	 *	@param	DynamicData		The dynamic data to fill in
-	 *	@param	bSelected		true if the particle system component is selected
-	 */
-	virtual bool UpdateDynamicData(FDynamicEmitterDataBase* DynamicData, bool bSelected) override;
+	virtual FDynamicEmitterDataBase* GetDynamicData(bool bSelected, ERHIFeatureLevel::Type InFeatureLevel) override;
 
 	/**
 	 *	Retrieves replay data for the emitter
@@ -864,17 +844,10 @@ struct ENGINE_API FParticleMeshEmitterInstance : public FParticleEmitterInstance
 	virtual void UpdateBoundingBox(float DeltaTime) override;
 	virtual uint32 RequiredBytes() override;
 	virtual void PostSpawn(FBaseParticle* Particle, float InterpolationPercentage, float SpawnTime) override;
-	virtual FDynamicEmitterDataBase* GetDynamicData(bool bSelected) override;
+	virtual FDynamicEmitterDataBase* GetDynamicData(bool bSelected, ERHIFeatureLevel::Type InFeatureLevel) override;
 	virtual bool IsDynamicDataRequired(UParticleLODLevel* CurrentLODLevel) override;
 
 	virtual bool Tick_MaterialOverrides() override;
-	/**
-	 *	Updates the dynamic data for the instance
-	 *
-	 *	@param	DynamicData		The dynamic data to fill in
-	 *	@param	bSelected		true if the particle system component is selected
-	 */
-	virtual bool UpdateDynamicData(FDynamicEmitterDataBase* DynamicData, bool bSelected) override;
 
 	/**
 	 *	Retrieves replay data for the emitter
@@ -931,7 +904,7 @@ struct ENGINE_API FParticleMeshEmitterInstance : public FParticleEmitterInstance
 	/**
 	 * Gets the materials applied to each section of a mesh.
 	 */
-	void GetMeshMaterials(TArray<UMaterialInterface*,TInlineAllocator<2> >& OutMaterials, const UParticleLODLevel* LODLevel) const;
+	void GetMeshMaterials(TArray<UMaterialInterface*,TInlineAllocator<2> >& OutMaterials, const UParticleLODLevel* LODLevel, ERHIFeatureLevel::Type InFeatureLevel, bool bLogWarnings = false) const;
 
 protected:
 
@@ -1065,15 +1038,7 @@ struct FParticleBeam2EmitterInstance : public FParticleEmitterInstance
 	/**
 	 *	Retrieves the dynamic data for the emitter
 	 */
-	virtual FDynamicEmitterDataBase* GetDynamicData(bool bSelected) override;
-
-	/**
-	 *	Updates the dynamic data for the instance
-	 *
-	 *	@param	DynamicData		The dynamic data to fill in
-	 *	@param	bSelected		true if the particle system component is selected
-	 */
-	virtual bool UpdateDynamicData(FDynamicEmitterDataBase* DynamicData, bool bSelected) override;
+	virtual FDynamicEmitterDataBase* GetDynamicData(bool bSelected, ERHIFeatureLevel::Type InFeatureLevel) override;
 
 	/**
 	 *	Retrieves replay data for the emitter
@@ -1585,15 +1550,7 @@ struct FParticleRibbonEmitterInstance : public FParticleTrailsEmitterInstance_Ba
 	/**
 	 *	Retrieves the dynamic data for the emitter
 	 */
-	virtual FDynamicEmitterDataBase* GetDynamicData(bool bSelected) override;
-
-	/**
-	 *	Updates the dynamic data for the instance
-	 *
-	 *	@param	DynamicData		The dynamic data to fill in
-	 *	@param	bSelected		true if the particle system component is selected
-	 */
-	virtual bool UpdateDynamicData(FDynamicEmitterDataBase* DynamicData, bool bSelected) override;
+	virtual FDynamicEmitterDataBase* GetDynamicData(bool bSelected, ERHIFeatureLevel::Type InFeatureLevel) override;
 
 	/**
 	 *	Retrieves replay data for the emitter
@@ -1756,15 +1713,7 @@ struct FParticleAnimTrailEmitterInstance : public FParticleTrailsEmitterInstance
 	/**
 	 *	Retrieves the dynamic data for the emitter
 	 */
-	virtual FDynamicEmitterDataBase* GetDynamicData(bool bSelected) override;
-
-	/**
-	 *	Updates the dynamic data for the instance
-	 *
-	 *	@param	DynamicData		The dynamic data to fill in
-	 *	@param	bSelected		true if the particle system component is selected
-	 */
-	virtual bool UpdateDynamicData(FDynamicEmitterDataBase* DynamicData, bool bSelected) override;
+	virtual FDynamicEmitterDataBase* GetDynamicData(bool bSelected, ERHIFeatureLevel::Type InFeatureLevel) override;
 
 	/**
 	 *	Retrieves replay data for the emitter
diff --git a/Engine/Source/Runtime/Engine/Public/ParticleHelper.h b/Engine/Source/Runtime/Engine/Public/ParticleHelper.h
--- a/Engine/Source/Runtime/Engine/Public/ParticleHelper.h
+++ b/Engine/Source/Runtime/Engine/Public/ParticleHelper.h
@@ -1792,7 +1792,7 @@ struct FDynamicMeshEmitterData : public FDynamicSpriteEmitterDataBase
 	FParticleVertexFactoryBase *CreateVertexFactory() override;
 
 	/** Initialize this emitter's dynamic rendering data, called after source data has been filled in */
-	void Init(bool bInSelected,const FParticleMeshEmitterInstance* InEmitterInstance,UStaticMesh* InStaticMesh);
+	void Init(bool bInSelected,const FParticleMeshEmitterInstance* InEmitterInstance,UStaticMesh* InStaticMesh, ERHIFeatureLevel::Type InFeatureLevel );
 
 	/**
 	 *	Create the render thread resources for this emitter data
