Fix for missing define

diff --git a/Engine/Source/Runtime/Engine/Private/UnrealNetwork.cpp b/Engine/Source/Runtime/Engine/Private/UnrealNetwork.cpp
--- a/Engine/Source/Runtime/Engine/Private/UnrealNetwork.cpp
+++ b/Engine/Source/Runtime/Engine/Private/UnrealNetwork.cpp
@@ -79,6 +79,10 @@ bool FNetworkVersion::IsNetworkCompatible( const uint32 LocalNetworkVersion, con
 
 FNetworkReplayVersion FNetworkVersion::GetReplayVersion()
 {
+	// Temp until Version.h is unlocked :O)
+#ifndef ENGINE_REPLAY_VERSION
+	#define ENGINE_REPLAY_VERSION BUILT_FROM_CHANGELIST
+#endif
 	return FNetworkReplayVersion(FApp::GetGameName(), 0, ENGINE_REPLAY_VERSION);
 }
 
