Copying //UE4/Dev-Blueprints to Dev-Main (//UE4/Dev-Main)

diff --git a/Engine/Source/Runtime/Engine/Classes/Kismet/KismetSystemLibrary.h b/Engine/Source/Runtime/Engine/Classes/Kismet/KismetSystemLibrary.h
--- a/Engine/Source/Runtime/Engine/Classes/Kismet/KismetSystemLibrary.h
+++ b/Engine/Source/Runtime/Engine/Classes/Kismet/KismetSystemLibrary.h
@@ -393,9 +393,16 @@ class ENGINE_API UKismetSystemLibrary : public UBlueprintFunctionLibrary
 	 * Clears a set timer.
 	 * @param Handle		The handle of the timer to clear.
 	 */
-	UFUNCTION(BlueprintCallable, meta=(DisplayName = "Clear Timer by Handle", WorldContext="WorldContextObject"), Category="Utilities|Time")
+	UFUNCTION(BlueprintCallable, meta=(DisplayName = "Clear Timer by Handle", WorldContext="WorldContextObject", DeprecatedFunction, DeprecationMessage = "Use Clear and Invalidate Timer by Handle. Note: you no longer need to reset your handle yourself after switching to the new function."), Category="Utilities|Time")
 	static void K2_ClearTimerHandle(UObject* WorldContextObject, FTimerHandle Handle);
 
+	/**
+	 * Clears a set timer.
+	 * @param Handle		The handle of the timer to clear.
+	 */
+	UFUNCTION(BlueprintCallable, meta=(DisplayName = "Clear and Invalidate Timer by Handle", WorldContext="WorldContextObject"), Category="Utilities|Time")
+	static void K2_ClearAndInvalidateTimerHandle(UObject* WorldContextObject, UPARAM(ref) FTimerHandle& Handle);
+
 	/**
 	 * Pauses a set timer at its current elapsed time.
 	 * @param Handle		The handle of the timer to pause.
diff --git a/Engine/Source/Runtime/Engine/Private/ActorConstruction.cpp b/Engine/Source/Runtime/Engine/Private/ActorConstruction.cpp
--- a/Engine/Source/Runtime/Engine/Private/ActorConstruction.cpp
+++ b/Engine/Source/Runtime/Engine/Private/ActorConstruction.cpp
@@ -451,11 +451,18 @@ void AActor::RerunConstructionScripts()
 	}
 }
 
+void AActor::FixComponentsFromDynamicClass()
+{
+
+}
+
 void AActor::ExecuteConstruction(const FTransform& Transform, const FComponentInstanceDataCache* InstanceDataCache, bool bIsDefaultTransform)
 {
 	check(!IsPendingKill());
 	check(!HasAnyFlags(RF_BeginDestroyed|RF_FinishDestroyed));
 
+	FixComponentsFromDynamicClass();
+
 	// ensure that any existing native root component gets this new transform
 	// we can skip this in the default case as the given transform will be the root component's transform
 	if (RootComponent && !bIsDefaultTransform)
@@ -684,6 +691,9 @@ UActorComponent* AActor::CreateComponentFromTemplateData(const FBlueprintCookedC
 			ArCustomPropertyList = InPropertyList;
 			ArUseCustomPropertyList = true;
 			ArWantBinaryPropertySerialization = true;
+
+			// Set this flag to emulate things that would happen in the SDO case when this flag is set (e.g. - not setting 'bHasBeenCreated').
+			ArPortFlags |= PPF_Duplicate;
 		}
 	};
 
diff --git a/Engine/Source/Runtime/Engine/Private/BlueprintGeneratedClass.cpp b/Engine/Source/Runtime/Engine/Private/BlueprintGeneratedClass.cpp
--- a/Engine/Source/Runtime/Engine/Private/BlueprintGeneratedClass.cpp
+++ b/Engine/Source/Runtime/Engine/Private/BlueprintGeneratedClass.cpp
@@ -266,7 +266,7 @@ void UBlueprintGeneratedClass::SerializeDefaultObject(UObject* Object, FArchive&
 {
 	Super::SerializeDefaultObject(Object, Ar);
 
-	if (Ar.IsLoading() && !Ar.IsObjectReferenceCollector() && ClassDefaultObject)
+	if (Ar.IsLoading() && !Ar.IsObjectReferenceCollector() && Object == ClassDefaultObject)
 	{
 		// On load, build the custom property list used in post-construct initialization logic. Note that in the editor, this will be refreshed during compile-on-load.
 		// @TODO - Potentially make this serializable (or cooked data) to eliminate the slight load time cost we'll incur below to generate this list in a cooked build. For now, it's not serialized since the raw UProperty references cannot be saved out.
@@ -274,6 +274,17 @@ void UBlueprintGeneratedClass::SerializeDefaultObject(UObject* Object, FArchive&
 	}
 }
 
+void UBlueprintGeneratedClass::PostLoadDefaultObject(UObject* Object)
+{
+	Super::PostLoadDefaultObject(Object);
+
+	if (Object == ClassDefaultObject)
+	{
+		// Rebuild the custom property list used in post-construct initialization logic. Note that PostLoad() may have altered some serialized properties.
+		UpdateCustomPropertyListForPostConstruction();
+	}
+}
+
 void UBlueprintGeneratedClass::BuildCustomPropertyListForPostConstruction(FCustomPropertyListNode*& InPropertyList, UStruct* InStruct, const uint8* DataPtr, const uint8* DefaultDataPtr)
 {
 	const UClass* OwnerClass = Cast<UClass>(InStruct);
@@ -316,6 +327,29 @@ void UBlueprintGeneratedClass::BuildCustomPropertyListForPostConstruction(FCusto
 						*CurrentNodePtr = nullptr;
 					}
 				}
+				else if (UArrayProperty* ArrayProperty = Cast<UArrayProperty>(Property))
+				{
+					// Create a new node for the array property.
+					*CurrentNodePtr = new(CustomPropertyListForPostConstruction) FCustomPropertyListNode(Property, Idx);
+
+					// Recursively gather up all array item indices that differ and assign to the current node's sub property list.
+					BuildCustomArrayPropertyListForPostConstruction(ArrayProperty, (*CurrentNodePtr)->SubPropertyList, PropertyValue, DefaultPropertyValue);
+
+					// This will be non-NULL if the above found at least one array item index that differs from the native CDO.
+					if ((*CurrentNodePtr)->SubPropertyList)
+					{
+						// Advance to the next node in the list.
+						CurrentNodePtr = &(*CurrentNodePtr)->PropertyListNext;
+					}
+					else
+					{
+						// Remove the node for the array property since it does not differ from the native CDO.
+						CustomPropertyListForPostConstruction.RemoveAt(CustomPropertyListForPostConstruction.Num() - 1);
+
+						// Clear the current node ptr since the array will have freed up the memory it referenced.
+						*CurrentNodePtr = nullptr;
+					}
+				}
 				else if (!Property->Identical(PropertyValue, DefaultPropertyValue))
 				{
 					// Create a new node, link it into the chain and add it into the array.
@@ -329,6 +363,86 @@ void UBlueprintGeneratedClass::BuildCustomPropertyListForPostConstruction(FCusto
 	}
 }
 
+void UBlueprintGeneratedClass::BuildCustomArrayPropertyListForPostConstruction(UArrayProperty* ArrayProperty, FCustomPropertyListNode*& InPropertyList, const uint8* DataPtr, const uint8* DefaultDataPtr)
+{
+	FCustomPropertyListNode** CurrentArrayNodePtr = &InPropertyList;
+
+	FScriptArrayHelper ArrayValueHelper(ArrayProperty, DataPtr);
+	FScriptArrayHelper DefaultArrayValueHelper(ArrayProperty, DefaultDataPtr);
+
+	for (int32 ArrayValueIndex = 0; ArrayValueIndex < ArrayValueHelper.Num(); ++ArrayValueIndex)
+	{
+		if (ArrayValueIndex < DefaultArrayValueHelper.Num())
+		{
+			const uint8* ArrayPropertyValue = ArrayValueHelper.GetRawPtr(ArrayValueIndex);
+			const uint8* DefaultArrayPropertyValue = DefaultArrayValueHelper.GetRawPtr(ArrayValueIndex);
+
+			if (UStructProperty* InnerStructProperty = Cast<UStructProperty>(ArrayProperty->Inner))
+			{
+				// Create a new node for the item value at this index.
+				*CurrentArrayNodePtr = new(CustomPropertyListForPostConstruction) FCustomPropertyListNode(ArrayProperty, ArrayValueIndex);
+
+				// Recursively gather up all struct fields that differ and assign to the array item value node's sub property list.
+				BuildCustomPropertyListForPostConstruction((*CurrentArrayNodePtr)->SubPropertyList, InnerStructProperty->Struct, ArrayPropertyValue, DefaultArrayPropertyValue);
+
+				// This will be non-NULL if the above found at least one struct field that differs from the native CDO.
+				if ((*CurrentArrayNodePtr)->SubPropertyList)
+				{
+					// Advance to the next node in the list.
+					CurrentArrayNodePtr = &(*CurrentArrayNodePtr)->PropertyListNext;
+				}
+				else
+				{
+					// Remove the node for the struct property since it does not differ from the native CDO.
+					CustomPropertyListForPostConstruction.RemoveAt(CustomPropertyListForPostConstruction.Num() - 1);
+
+					// Clear the current array item node ptr
+					*CurrentArrayNodePtr = nullptr;
+				}
+			}
+			else if (UArrayProperty* InnerArrayProperty = Cast<UArrayProperty>(ArrayProperty->Inner))
+			{
+				// Create a new node for the item value at this index.
+				*CurrentArrayNodePtr = new(CustomPropertyListForPostConstruction) FCustomPropertyListNode(ArrayProperty, ArrayValueIndex);
+
+				// Recursively gather up all array item indices that differ and assign to the array item value node's sub property list.
+				BuildCustomArrayPropertyListForPostConstruction(InnerArrayProperty, (*CurrentArrayNodePtr)->SubPropertyList, ArrayPropertyValue, DefaultArrayPropertyValue);
+
+				// This will be non-NULL if the above found at least one array item index that differs from the native CDO.
+				if ((*CurrentArrayNodePtr)->SubPropertyList)
+				{
+					// Advance to the next node in the list.
+					CurrentArrayNodePtr = &(*CurrentArrayNodePtr)->PropertyListNext;
+				}
+				else
+				{
+					// Remove the node for the array property since it does not differ from the native CDO.
+					CustomPropertyListForPostConstruction.RemoveAt(CustomPropertyListForPostConstruction.Num() - 1);
+
+					// Clear the current array item node ptr
+					*CurrentArrayNodePtr = nullptr;
+				}
+			}
+			else if (!ArrayProperty->Inner->Identical(ArrayPropertyValue, DefaultArrayPropertyValue))
+			{
+				// Create a new node, link it into the chain and add it into the array.
+				*CurrentArrayNodePtr = new(CustomPropertyListForPostConstruction) FCustomPropertyListNode(ArrayProperty, ArrayValueIndex);
+
+				// Advance to the next array item node ptr.
+				CurrentArrayNodePtr = &(*CurrentArrayNodePtr)->PropertyListNext;
+			}
+		}
+		else
+		{
+			// Signals the "end" of differences with the default value (signals that remaining values should be copied in full)
+			*CurrentArrayNodePtr = new(CustomPropertyListForPostConstruction) FCustomPropertyListNode(nullptr, ArrayValueIndex);
+
+			// Don't need to record anything else.
+			break;
+		}
+	}
+}
+
 void UBlueprintGeneratedClass::UpdateCustomPropertyListForPostConstruction()
 {
 	// Empty the current list.
@@ -1033,12 +1147,44 @@ void FBlueprintCookedComponentInstancingData::BuildCachedPropertyList(FCustomPro
 		{
 			BuildCachedPropertyList(&NewNode->SubPropertyList, StructProperty->Struct, CurrentSourceIdx);
 		}
+		else if (const UArrayProperty* ArrayProperty = Cast<UArrayProperty>(Property))
+		{
+			// If this is an array property, recursively build a sub-property list.
+			BuildCachedArrayPropertyList(ArrayProperty, &NewNode->SubPropertyList, CurrentSourceIdx);
+		}
 
 		// Advance current location to the next linked node.
 		CurrentNode = &NewNode->PropertyListNext;
 	}
 }
 
+void FBlueprintCookedComponentInstancingData::BuildCachedArrayPropertyList(const UArrayProperty* ArrayProperty, FCustomPropertyListNode** ArraySubPropertyNode, int32* CurrentSourceIdx) const
+{
+	// Build the array property's sub-property list. An empty name field signals the end of the changed array property list.
+	while (*CurrentSourceIdx < ChangedPropertyList.Num() &&
+		(ChangedPropertyList[*CurrentSourceIdx].PropertyName == NAME_None
+			|| ChangedPropertyList[*CurrentSourceIdx].PropertyName == ArrayProperty->GetFName()))
+	{
+		const FBlueprintComponentChangedPropertyInfo& ChangedArrayPropertyInfo = ChangedPropertyList[(*CurrentSourceIdx)++];
+		UProperty* InnerProperty = ChangedArrayPropertyInfo.PropertyName != NAME_None ? ArrayProperty->Inner : nullptr;
+
+		*ArraySubPropertyNode = new(CachedPropertyListForSerialization) FCustomPropertyListNode(InnerProperty, ChangedArrayPropertyInfo.ArrayIndex);
+
+		// If this is a UStruct property, recursively build a sub-property list.
+		if (const UStructProperty* InnerStructProperty = Cast<UStructProperty>(InnerProperty))
+		{
+			BuildCachedPropertyList(&(*ArraySubPropertyNode)->SubPropertyList, InnerStructProperty->Struct, CurrentSourceIdx);
+		}
+		else if (const UArrayProperty* InnerArrayProperty = Cast<UArrayProperty>(InnerProperty))
+		{
+			// If this is an array property, recursively build a sub-property list.
+			BuildCachedArrayPropertyList(InnerArrayProperty, &(*ArraySubPropertyNode)->SubPropertyList, CurrentSourceIdx);
+		}
+
+		ArraySubPropertyNode = &(*ArraySubPropertyNode)->PropertyListNext;
+	}
+}
+
 const FCustomPropertyListNode* FBlueprintCookedComponentInstancingData::GetCachedPropertyListForSerialization() const
 {
 	FCustomPropertyListNode* PropertyListRootNode = nullptr;
diff --git a/Engine/Source/Runtime/Engine/Private/Collision/CollisionProfile.cpp b/Engine/Source/Runtime/Engine/Private/Collision/CollisionProfile.cpp
--- a/Engine/Source/Runtime/Engine/Private/Collision/CollisionProfile.cpp
+++ b/Engine/Source/Runtime/Engine/Private/Collision/CollisionProfile.cpp
@@ -357,9 +357,20 @@ void UCollisionProfile::LoadProfileConfig(bool bForceInit)
 	// also initialize DefaultResposneContainer with default response for each channel
 	FCollisionResponseContainer::DefaultResponseContainer.SetAllChannels(ECR_Block);
 
-	for (auto Iter=DefaultChannelResponses.CreateConstIterator(); Iter; ++Iter)
+	// we can't guarantee that DefaultChannelResponses was loaded ordered (from
+	// config) - in the following loop, we fill out TraceTypeMapping and 
+	// ObjectTypeMapping from DefaultChannelResponses, and those mappings expect 
+	// to be ordered (since we reinterpret the index to be an enum value itself 
+	// in functions like ConvertToCollisionChannel(), etc.)
+	DefaultChannelResponses.Sort([](const FCustomChannelSetup& Rhs, const FCustomChannelSetup& Lhs)->bool 
+		{ 
+			return (Rhs.Channel < Lhs.Channel);
+		}
+	);
+
+	for (int32 ChannelResponseIndex = 0; ChannelResponseIndex < DefaultChannelResponses.Num(); ++ChannelResponseIndex)
 	{
-		const FCustomChannelSetup& CustomChannel = *Iter;
+		const FCustomChannelSetup& CustomChannel = DefaultChannelResponses[ChannelResponseIndex];
 		int32 EnumIndex = CustomChannel.Channel;
 		// make sure it is the range of channels we allow to change
 		if ( IS_VALID_COLLISIONCHANNEL(EnumIndex) )
@@ -371,6 +382,15 @@ void UCollisionProfile::LoadProfileConfig(bool bForceInit)
 				FString VariableName = ChannelDisplayNames[EnumIndex].ToString();
 				FString DisplayValue = CustomChannel.Name.ToString();
 
+				if (TraceTypeMapping.Contains(EnumIndex) || ObjectTypeMapping.Contains(EnumIndex))
+				{
+					UE_LOG(LogCollisionProfile, Warning, TEXT("Cannot map multiple responses to the same collision channel (%d); ignoring '%s' "), EnumIndex, *DisplayValue);
+					DefaultChannelResponses.RemoveAt(ChannelResponseIndex);
+					// decrement iterator so this index gets processed again (since we just removed an entry)
+					--ChannelResponseIndex;
+					continue;
+				}
+
 				// also has to set this for internal use
 				ChannelDisplayNames[EnumIndex] = FName(*DisplayValue);
 
@@ -433,8 +453,8 @@ void UCollisionProfile::LoadProfileConfig(bool bForceInit)
 #if WITH_EDITOR
 	// now propagate all changes to the channels to EObjectTypeQuery and ETraceTypeQuery for blueprint accessibility
 	// this code is to show in editor more friendly, so it doesn't have to be set if it's not for editor
-	int32 ObjectTypeEnumIndex = 0;
-	int32 TraceTypeEnumIndex = 0;
+	int32 ObjectTypeMappingCount = 0;
+	int32 TraceTypeMappingCount  = 0;
 
 	// first go through fill up ObjectType Enum
 	for ( int32 EnumIndex=0; EnumIndex<NumEnum; ++EnumIndex )
@@ -449,23 +469,31 @@ void UCollisionProfile::LoadProfileConfig(bool bForceInit)
 				// find out trace type or object type
 				if (Enum->GetMetaData(*TraceType, EnumIndex) == TraceValue)
 				{
-					TraceTypeEnum->RemoveMetaData(*HiddenMeta, TraceTypeEnumIndex);
-					TraceTypeEnum->SetMetaData(*KeyName, *DisplayName, TraceTypeEnumIndex);
-					++TraceTypeEnumIndex;
+					int32 TraceTypeIndex = TraceTypeMapping.Find((ECollisionChannel)EnumIndex);
+					if (ensure(TraceTypeIndex != INDEX_NONE))
+					{
+						TraceTypeEnum->RemoveMetaData(*HiddenMeta, TraceTypeIndex);
+						TraceTypeEnum->SetMetaData(*KeyName, *DisplayName, TraceTypeIndex);
+						++TraceTypeMappingCount;
+					}
 				}
 				else
 				{
-					ObjectTypeEnum->RemoveMetaData(*HiddenMeta, ObjectTypeEnumIndex);
-					ObjectTypeEnum->SetMetaData(*KeyName, *DisplayName, ObjectTypeEnumIndex);
-					++ObjectTypeEnumIndex;
+					int32 ObjectTypeIndex = ObjectTypeMapping.Find((ECollisionChannel)EnumIndex);
+					if (ensure(ObjectTypeIndex != INDEX_NONE))
+					{
+						ObjectTypeEnum->RemoveMetaData(*HiddenMeta, ObjectTypeIndex);
+						ObjectTypeEnum->SetMetaData(*KeyName, *DisplayName, ObjectTypeIndex);
+						++ObjectTypeMappingCount;
+					}
 				}
 			}
 		}
 	}
 
 	// make sure TraceTypeEnumIndex matches TraceTypeMapping
-	check (TraceTypeMapping.Num() == TraceTypeEnumIndex);
-	check (ObjectTypeMapping.Num() == ObjectTypeEnumIndex);
+	check (TraceTypeMapping.Num()  == TraceTypeMappingCount);
+	check (ObjectTypeMapping.Num() == ObjectTypeMappingCount);
 #endif
 
 	// collision redirector has to be loaded before profile
diff --git a/Engine/Source/Runtime/Engine/Private/Components/ActorComponent.cpp b/Engine/Source/Runtime/Engine/Private/Components/ActorComponent.cpp
--- a/Engine/Source/Runtime/Engine/Private/Components/ActorComponent.cpp
+++ b/Engine/Source/Runtime/Engine/Private/Components/ActorComponent.cpp
@@ -1633,7 +1633,7 @@ void UActorComponent::Serialize(FArchive& Ar)
 {
 	Super::Serialize(Ar);
 
-	if (Ar.IsLoading() && (Ar.HasAnyPortFlags(PPF_DuplicateForPIE)||!Ar.HasAnyPortFlags(PPF_Duplicate)))
+	if (Ar.IsLoading() && (Ar.HasAnyPortFlags(PPF_DuplicateForPIE)||!Ar.HasAnyPortFlags(PPF_Duplicate)) && !IsTemplate())
 	{
 		bHasBeenCreated = true;
 	}
diff --git a/Engine/Source/Runtime/Engine/Private/KismetMathLibrary.cpp b/Engine/Source/Runtime/Engine/Private/KismetMathLibrary.cpp
--- a/Engine/Source/Runtime/Engine/Private/KismetMathLibrary.cpp
+++ b/Engine/Source/Runtime/Engine/Private/KismetMathLibrary.cpp
@@ -236,6 +236,11 @@ int32 UKismetMathLibrary::Or_IntInt(int32 A, int32 B)
 	return A | B;
 }
 
+int32 UKismetMathLibrary::Not_Int(int32 A)
+{
+	return ~A;
+}
+
 int32 UKismetMathLibrary::SignOfInteger(int32 A)
 {
 	return FMath::Sign<int32>(A);
diff --git a/Engine/Source/Runtime/Engine/Private/KismetSystemLibrary.cpp b/Engine/Source/Runtime/Engine/Private/KismetSystemLibrary.cpp
--- a/Engine/Source/Runtime/Engine/Private/KismetSystemLibrary.cpp
+++ b/Engine/Source/Runtime/Engine/Private/KismetSystemLibrary.cpp
@@ -390,6 +390,18 @@ void UKismetSystemLibrary::K2_ClearTimerHandle(UObject* WorldContextObject, FTim
 	}
 }
 
+void UKismetSystemLibrary::K2_ClearAndInvalidateTimerHandle(UObject* WorldContextObject, FTimerHandle& Handle)
+{
+	if (Handle.IsValid())
+	{
+		UWorld* World = GEngine->GetWorldFromContextObject(WorldContextObject);
+		if (World)
+		{
+			World->GetTimerManager().ClearTimer(Handle);
+		}
+	}
+}
+
 void UKismetSystemLibrary::K2_PauseTimer(UObject* Object, FString FunctionName)
 {
 	FTimerDynamicDelegate Delegate;
diff --git a/Engine/Source/Runtime/Launch/Private/LaunchEngineLoop.cpp b/Engine/Source/Runtime/Launch/Private/LaunchEngineLoop.cpp
--- a/Engine/Source/Runtime/Launch/Private/LaunchEngineLoop.cpp
+++ b/Engine/Source/Runtime/Launch/Private/LaunchEngineLoop.cpp
@@ -290,6 +290,11 @@ bool LaunchSetGameName(const TCHAR *InCmdLine, FString& OutGameProjectFilePathUn
 		{
 			// Try to use the executable name as the game name.
 			LocalGameName = FPlatformProcess::ExecutableName();
+			int32 FirstCharToRemove = INDEX_NONE;
+			if (LocalGameName.FindChar(TCHAR('-'), FirstCharToRemove))
+			{
+				LocalGameName = LocalGameName.Left(FirstCharToRemove);
+			}
 			FApp::SetGameName(*LocalGameName);
 
 			// Check it's not UE4Game, otherwise assume a uproject file relative to the game project directory
