Copying //UE4/Dev-Networking to Dev-Main (//UE4/Dev-Main)

==========================
MAJOR FEATURES + CHANGES
==========================

Change 2909132 on 2016/03/14 by [REDACTED_PERSON]

	Use special bit to signify rep layout in network stream (rather than hi-jacking property)

	This is to prepare for backwards compatibility, where we might not be able to rely on property in question (hi-jacked property) to be stable

	#tests Replication, Replays

Change 2913459 on 2016/03/17 by [REDACTED_PERSON]

	Replay backwards compatibilty support: First pass at cleaning up code to support skipping over blobs of data easily

	* Clean up + unify how we read/write property function fields
	* We no longer serialize properties for rejected RPC's (we now have the size of the blob in the stream)
	* Skip over missing RPC functions
	* Cleanup how we receive a bunch in DataReplication

Change 2913787 on 2016/03/17 by [REDACTED_PERSON]

	Refactor how we read/write content blocks in network streams so we can easily skip over them

Change 2913940 on 2016/03/17 by [REDACTED_PERSON]

	Fix issues with queued bunches and new method of writing content blocks

Change 2914995 on 2016/03/18 by [REDACTED_PERSON]

	Replay backwards compatibility support: Ability to skip over unknown properties

Change 2915900 on 2016/03/19 by [REDACTED_PERSON]

	Cache replay stream chunks/checkpoints at the HTTP request/response level

Change 2923206 on 2016/03/25 by [REDACTED_PERSON]

	No longer use package checksum, use network checksum instead, and make optional

Change 2929343 on 2016/03/31 by [REDACTED_PERSON]

	Fix issues where we were'nt proprerly checking Bunch.IsError in all paths

Change 2929344 on 2016/03/31 by [REDACTED_PERSON]

	Don't save net guid flags unless we need to, add ability to sync NetGUIDAckStatus with another package map, minor cleanup

Change 2929817 on 2016/03/31 by [REDACTED_PERSON]

	Ability to export rep layout (as compatible rep layout struct). This gives us the ability to know the names of missing/changed properties in older replays.

	* Export granularity is at rep layout cmd level, this saves lots of space, since we never really have to export the entire rep layout structure (rarely do all properties get replicated)
	* We're very close to supporting this for regular replication, just need to update CompatibleRepLayoutPathExported and CompatibleRepLayoutCmdExported to actually track ack status (rather than assume reliable connection)

Change 2930798 on 2016/04/01 by [REDACTED_PERSON]

	Explicitly mark properties as incompatible, and then use this to only warn once + cleanup + unification of array/property receiving code

Change 2930977 on 2016/04/01 by [REDACTED_PERSON]

	Fix CIS

Change 2934147 on 2016/04/05 by [REDACTED_PERSON]

	Separate reading/writing polling status for PS4 sockets, so we can properly wait the full timeout when we want to check read status

Change 2934292 on 2016/04/05 by [REDACTED_PERSON]

	Fix issue where we weren't generating proper checksum for each element in static arrays, which was causing them not to replicate properly in the backwards compatibility path

Change 2936626 on 2016/04/07 by [REDACTED_PERSON]

	Fix issue where we weren't properly exporting compatible replayout cmds when dealing with queued RPC bunches

diff --git a/Engine/Source/Runtime/CoreUObject/Public/UObject/CoreNet.h b/Engine/Source/Runtime/CoreUObject/Public/UObject/CoreNet.h
--- a/Engine/Source/Runtime/CoreUObject/Public/UObject/CoreNet.h
+++ b/Engine/Source/Runtime/CoreUObject/Public/UObject/CoreNet.h
@@ -7,6 +7,7 @@
 
 DECLARE_DELEGATE_RetVal_OneParam( bool, FNetObjectIsDynamic, const UObject*);
 
+class FOutBunch;
 
 //
 // Information about a field.
@@ -136,7 +137,7 @@ class COREUOBJECT_API UPackageMap : public UObject
 
 	virtual void		ReceivedNak( const int32 NakPacketId ) { }
 	virtual void		ReceivedAck( const int32 AckPacketId ) { }
-	virtual void		NotifyBunchCommit( const int32 OutPacketId, const TArray< FNetworkGUID > & ExportNetGUIDs ) { }
+	virtual void		NotifyBunchCommit( const int32 OutPacketId, const FOutBunch* OutBunch ) { }
 
 	virtual void		GetNetGUIDStats(int32& AckCount, int32& UnAckCount, int32& PendingCount) { }
 
diff --git a/Engine/Source/Runtime/Engine/Classes/Engine/ActorChannel.h b/Engine/Source/Runtime/Engine/Classes/Engine/ActorChannel.h
--- a/Engine/Source/Runtime/Engine/Classes/Engine/ActorChannel.h
+++ b/Engine/Source/Runtime/Engine/Classes/Engine/ActorChannel.h
@@ -143,17 +143,26 @@ public:
 	void CleanupReplicators( const bool bKeepReplicators = false );
 
 	/** Writes the header for a content block of properties / RPCs for the given object (either the actor a subobject of the actor) */
-	void BeginContentBlock( UObject* Obj, FOutBunch &Bunch );
+	void WriteContentBlockHeader( UObject* Obj, FOutBunch &Bunch, const bool bHasRepLayout );
 
 	/** Writes the header for a content block specifically for deleting sub-objects */
-	void BeginContentBlockForSubObjectDelete( FOutBunch & Bunch, FNetworkGUID & GuidToDelete );
+	void WriteContentBlockForSubObjectDelete( FOutBunch & Bunch, FNetworkGUID & GuidToDelete );
 
-	/** Writes the footer for a content block of properties / RPCs for the given object (either the actor a subobject of the actor) */
-	void EndContentBlock( UObject *Obj, FOutBunch &Bunch, const FClassNetCache* ClassCache = nullptr );
+	/** Writes header and payload of content block */
+	int32 WriteContentBlockPayload( UObject* Obj, FOutBunch &Bunch, const bool bHasRepLayout, FNetBitWriter& Payload );
 
 	/** Reads the header of the content block and instantiates the subobject if necessary */
-	UObject* ReadContentBlockHeader(FInBunch& Bunch, bool& bObjectDeleted);
+	UObject* ReadContentBlockHeader( FInBunch& Bunch, bool& bObjectDeleted, bool& bOutHasRepLayout );
 
+	/** Reads content block header and payload */
+	UObject* ReadContentBlockPayload( FInBunch &Bunch, FNetBitReader& OutPayload, bool& bOutHasRepLayout );
+
+	/** Writes property/function header and data blob to network stream */
+	int32 WriteFieldHeaderAndPayload( FNetBitWriter& Bunch, const FClassNetCache* ClassCache, const FFieldNetCache * FieldCache, FNetBitWriter& Payload );
+
+	/** Reads property/function header and data blob from network stream */
+	bool ReadFieldHeaderAndPayload( UObject* Object, const FClassNetCache* ClassCache, FNetBitReader& Bunch, const FFieldNetCache** OutField, FNetBitReader& OutPayload ) const;
+		
 	/** Returns the replicator for the actor associated with this channel. Guaranteed to exist. */
 	FObjectReplicator & GetActorReplicationData();
 
diff --git a/Engine/Source/Runtime/Engine/Classes/Engine/PackageMapClient.h b/Engine/Source/Runtime/Engine/Classes/Engine/PackageMapClient.h
--- a/Engine/Source/Runtime/Engine/Classes/Engine/PackageMapClient.h
+++ b/Engine/Source/Runtime/Engine/Classes/Engine/PackageMapClient.h
@@ -22,12 +22,94 @@
 #include "Net/DataBunch.h"
 #include "PackageMapClient.generated.h"
 
+class FRepLayout;
+class FOutBunch;
+
+class ENGINE_API FCompatibleRepLayoutCmd
+{
+public:
+	FCompatibleRepLayoutCmd() : bExported( false ), Handle( 0 ), CompatibleChecksum( 0 ), bIncompatible( false )
+	{
+	}
+
+	FCompatibleRepLayoutCmd( const uint32 InHandle, const uint32 InCompatibleChecksum, const FString InName, FString InType ) :
+		bExported( false ),
+		Handle( InHandle ),
+		CompatibleChecksum( InCompatibleChecksum ),
+		Name( InName ),
+		Type( InType ),
+		bIncompatible( false )
+	{
+	}
+
+	friend FArchive& operator<<( FArchive& Ar, FCompatibleRepLayoutCmd& C )
+	{
+		uint8 Flags = C.bExported ? 1 : 0;
+
+		Ar << Flags;
+
+		if ( Ar.IsLoading() )
+		{
+			C.bExported = Flags == 1 ? true : false;
+		}
+
+		if ( C.bExported )
+		{
+			Ar.SerializeIntPacked( C.Handle );
+			Ar << C.CompatibleChecksum << C.Name << C.Type;
+		}
+
+		return Ar;
+	}
+
+
+	bool		bExported;
+	uint32		Handle;
+	uint32		CompatibleChecksum;
+	FString		Name;
+	FString		Type;
+
+	// Transient properties
+	bool		bIncompatible;		// If true, we've already determined that this property isn't compatible. We use this to curb warning spam.
+};
+
+class ENGINE_API FCompatibleRepLayout
+{
+public:
+	FCompatibleRepLayout() : PathNameIndex( 0 ) { }
+
+	FString								PathName;
+	uint32								PathNameIndex;
+	TArray< FCompatibleRepLayoutCmd >	Cmds;
+
+	friend FArchive& operator<<( FArchive& Ar, FCompatibleRepLayout& C )
+	{
+		Ar << C.PathName;
+
+		Ar.SerializeIntPacked( C.PathNameIndex );
+
+		uint32 NumCmds = C.Cmds.Num();
+		Ar.SerializeIntPacked( NumCmds );
+
+		if ( Ar.IsLoading() )
+		{
+			C.Cmds.AddDefaulted( ( int32 )NumCmds );
+		}
+
+		for ( int32 i = 0; i < C.Cmds.Num(); i++ )
+		{
+			Ar << C.Cmds[i];
+		}
+
+		return Ar;
+	}
+};
 
 /** Stores an object with path associated with FNetworkGUID */
 class FNetGuidCacheObject
 {
 public:
-	FNetGuidCacheObject() : NetworkChecksum( 0 ), PackageChecksum( 0 ), ReadOnlyTimestamp( 0 ), bNoLoad( 0 ), bIgnoreWhenMissing( 0 ), bIsPending( 0 ), bIsBroken( 0 )
+	FNetGuidCacheObject() : NetworkChecksum( 0 ), ReadOnlyTimestamp( 0 ), bNoLoad( 0 ), bIgnoreWhenMissing( 0 ), bIsPending( 0 ), bIsBroken( 0 )
 	{
 	}
 
@@ -37,7 +119,6 @@ public:
 	FNetworkGUID				OuterGUID;
 	FName						PathName;
 	uint32						NetworkChecksum;			// Network checksum saved, used to determine backwards compatible
-	uint32						PackageChecksum;			// If this is a package, this is the guid to expect
 
 	double						ReadOnlyTimestamp;			// Time in second when we should start timing out after going read only
 
@@ -52,6 +133,13 @@ class ENGINE_API FNetGUIDCache
 public:
 	FNetGUIDCache( UNetDriver * InDriver );
 
+	enum ENetworkChecksumMode
+	{
+		NETCHECKSUM_None			= 0,
+		NETCHECKSUM_SaveAndUse		= 1,
+		NETCHECKSUM_SaveButIgnore	= 2,
+	};
+
 	void			CleanReferences();
 	bool			SupportsObject( const UObject* Object ) const;
 	bool			IsDynamicObject( const UObject* Object );
@@ -62,7 +150,7 @@ public:
 	void			RegisterNetGUID_Internal( const FNetworkGUID& NetGUID, const FNetGuidCacheObject& CacheObject );
 	void			RegisterNetGUID_Server( const FNetworkGUID& NetGUID, const UObject* Object );
 	void			RegisterNetGUID_Client( const FNetworkGUID& NetGUID, const UObject* Object );
-	void			RegisterNetGUIDFromPath_Client( const FNetworkGUID& NetGUID, const FString& PathName, const FNetworkGUID& OuterGUID, const uint32 NetworkChecksum, const uint32 PackageChecksum, const bool bNoLoad, const bool bIgnoreWhenMissing );
+	void			RegisterNetGUIDFromPath_Client( const FNetworkGUID& NetGUID, const FString& PathName, const FNetworkGUID& OuterGUID, const uint32 NetworkChecksum, const bool bNoLoad, const bool bIgnoreWhenMissing );
 	UObject *		GetObjectFromNetGUID( const FNetworkGUID& NetGUID, const bool bIgnoreMustBeMapped );
 	bool			ShouldIgnoreWhenMissing( const FNetworkGUID& NetGUID ) const;
 	bool			IsGUIDRegistered( const FNetworkGUID& NetGUID ) const;
@@ -70,12 +158,10 @@ public:
 	bool			IsGUIDBroken( const FNetworkGUID& NetGUID, const bool bMustBeRegistered ) const;
 	FString			FullNetGUIDPath( const FNetworkGUID& NetGUID ) const;
 	void			GenerateFullNetGUIDPath_r( const FNetworkGUID& NetGUID, FString& FullPath ) const;
-	void			SetIgnorePackageMismatchOverride( const bool bInIgnorePackageMismatchOverride ) { bIgnorePackageMismatchOverride = bInIgnorePackageMismatchOverride; }
 	bool			ShouldIgnorePackageMismatch() const;
 	uint32			GetClassNetworkChecksum( const UClass* Class );
 	uint32			GetNetworkChecksum( const UObject* Obj );
-	void			SetShouldUseNetworkChecksum( const bool bInUseNetworkChecksum );
-	bool			ShouldUseNetworkChecksum() const;
+	void			SetNetworkChecksumMode( const ENetworkChecksumMode NewMode );
 
 	void			AsyncPackageCallback(const FName& PackageName, UPackage * Package, EAsyncLoadingResult::Type Result);
 	
@@ -89,8 +175,19 @@ public:
 
 	TMap< FName, FNetworkGUID >						PendingAsyncPackages;
 
-	bool											bUseNetworkChecksum;
-	bool											bIgnorePackageMismatchOverride;
+	ENetworkChecksumMode							NetworkChecksumMode;
+
+	/** Maps FRepLayout class owner to the respective FCompatibleRepLayout */
+	TMap < FString, TSharedPtr< FCompatibleRepLayout > >	CompatibleRepLayoutMap;
+
+	/** Maps rep layout path to assigned index */
+	TMap < FString, uint32 >								CompatibleRepLayoutPathToIndex;
+
+	/** Maps rep layout index to assigned path */
+	TMap < uint32, FString >								CompatibleRepLayoutIndexToPath;
+
+	/** Current index used when filling in CompatibleRepLayoutPathToIndex/CompatibleRepLayoutIndexToPath */
+	int32													UniqueCompatibleRepLayoutPathIndex;
 
 #if !(UE_BUILD_SHIPPING || UE_BUILD_TEST)
 	// History for debugging entries in the guid cache
@@ -135,11 +232,11 @@ public:
 
 	virtual void ReceivedNak( const int32 NakPacketId ) override;
 	virtual void ReceivedAck( const int32 AckPacketId ) override;
-	virtual void NotifyBunchCommit( const int32 OutPacketId, const TArray< FNetworkGUID > & ExportNetGUIDs ) override;
+	virtual void NotifyBunchCommit( const int32 OutPacketId, const FOutBunch* OutBunch ) override;
 	virtual void GetNetGUIDStats(int32 &AckCount, int32 &UnAckCount, int32 &PendingCount) override;
 
 	void ReceiveNetGUIDBunch( FInBunch &InBunch );
-	bool AppendExportBunches(TArray<FOutBunch *>& OutgoingBunches);
+	void AppendExportBunches(TArray<FOutBunch *>& OutgoingBunches);
 
 	TMap<FNetworkGUID, int32>	NetGUIDExportCountMap;	// How many times we've exported each NetGUID on this connection. Public for ListNetGUIDExports 
 
@@ -160,8 +257,20 @@ public:
 
 	class UNetConnection* GetConnection() { return Connection; }
 
+	void SyncPackageMapExportAckStatus( const UPackageMapClient* Source );
+
+	/** Functions to help with exporting/importing replayout property info */
+	TSharedPtr< FCompatibleRepLayout >	TrackCompatibleRepLayout( const FRepLayout* RepLayout );
+	void								TrackCompatibleRepLayoutCmd( const FRepLayout* RepLayout, FCompatibleRepLayout* CompatibleRepLayout, const int32 CmdIndex );
+	TSharedPtr< FCompatibleRepLayout >	GetCompatibleRepLayoutChecked( UObject* Owner ) const;
+	void								SerializeCompatibleReplayoutMap( FArchive& Ar );
+
 protected:
 
+	/** Functions to help with exporting/importing replayout property info */
+	void								AppendCompatibleRepLayoutExports( TArray<FOutBunch *>& OutgoingBunches );
+	void								ReceiveCompatibleRepLayoutExports( FInBunch &InBunch );
+
 	bool	ExportNetGUID( FNetworkGUID NetGUID, const UObject* Object, FString PathName, UObject* ObjOuter );
 	void	ExportNetGUIDHeader();
 
@@ -192,4 +301,10 @@ protected:
 	TSharedPtr< FNetGUIDCache >			GuidCache;
 
 	TArray< FNetworkGUID >				MustBeMappedGuidsInLastBunch;
+
+	/** List of rep layout cmd exports that need to go out on next bunch */
+	TSet< uint64 >						CompatibleRepLayoutCmdExports;
+
+	TMap < uint32, bool >									CompatibleRepLayoutPathExported;
+	TMap < uint64, bool >									CompatibleRepLayoutCmdExported;
 };
diff --git a/Engine/Source/Runtime/Engine/Private/DataBunch.cpp b/Engine/Source/Runtime/Engine/Private/DataBunch.cpp
--- a/Engine/Source/Runtime/Engine/Private/DataBunch.cpp
+++ b/Engine/Source/Runtime/Engine/Private/DataBunch.cpp
@@ -31,8 +31,8 @@ FInBunch::FInBunch( UNetConnection* InConnection, uint8* Src, int64 CountBits )
 ,	bReliable ( 0 )
 ,	bPartial ( 0 )
 ,	bPartialInitial ( 0 )
-,	bPartialFinal	( 0 )
-,	bHasGUIDs( 0 )
+,	bPartialFinal ( 0 )
+,	bHasPackageMapExports ( 0 )
 {
 	check(Connection);
 	// Match the byte swapping settings of the connection
@@ -55,7 +55,7 @@ FInBunch::FInBunch( UPackageMap *InPackageMap, uint8* Src, int64 CountBits )
 ,	bPartial	( 0 )
 ,	bPartialInitial ( 0 )
 ,	bPartialFinal	( 0 )
-,	bHasGUIDs ( 0 )
+,	bHasPackageMapExports ( 0 )
 {
 	ArMaxSerializeSize = MAX_STRING_SERIALIZE_SIZE;
 }
@@ -104,7 +104,7 @@ FOutBunch::FOutBunch( UChannel* InChannel, bool bInClose )
 ,	bPartial	( 0 )
 ,	bPartialInitial			( 0 )
 ,	bPartialFinal			( 0 )
-,	bHasGUIDs				( 0 )
+,	bHasPackageMapExports	( 0 )
 ,	bHasMustBeMappedGUIDs	( 0 )
 {
 	checkSlow(!Channel->Closing);
@@ -137,7 +137,7 @@ FOutBunch::FOutBunch( UPackageMap *InPackageMap, int64 MaxBits )
 ,	bPartial	( 0 )
 ,	bPartialInitial ( 0 )
 ,	bPartialFinal	( 0 )
-,	bHasGUIDs				( 0 )
+,	bHasPackageMapExports	( 0 )
 ,	bHasMustBeMappedGUIDs	( 0 )
 {
 }
diff --git a/Engine/Source/Runtime/Engine/Private/DataChannel.cpp b/Engine/Source/Runtime/Engine/Private/DataChannel.cpp
--- a/Engine/Source/Runtime/Engine/Private/DataChannel.cpp
+++ b/Engine/Source/Runtime/Engine/Private/DataChannel.cpp
@@ -296,7 +296,7 @@ bool UChannel::ReceivedSequencedBunch( FInBunch& Bunch )
 void UChannel::ReceivedRawBunch( FInBunch & Bunch, bool & bOutSkipAck )
 {
 	// Immediately consume the NetGUID portion of this bunch, regardless if it is partial or reliable.
-	if ( Bunch.bHasGUIDs )
+	if ( Bunch.bHasPackageMapExports )
 	{
 		Cast<UPackageMapClient>( Connection->PackageMap )->ReceiveNetGUIDBunch( Bunch );
 
@@ -450,7 +450,7 @@ bool UChannel::ReceivedNextBunch( FInBunch & Bunch, bool & bOutSkipAck )
 			}
 
 			InPartialBunch = new FInBunch(Bunch, false);
-			if ( !Bunch.bHasGUIDs && Bunch.GetBitsLeft() > 0 )
+			if ( !Bunch.bHasPackageMapExports && Bunch.GetBitsLeft() > 0 )
 			{
 				check( Bunch.GetBitsLeft() % 8 == 0); // Starting partial bunches should always be byte aligned.
 
@@ -487,14 +487,14 @@ bool UChannel::ReceivedNextBunch( FInBunch & Bunch, bool & bOutSkipAck )
 				// Merge.
 				UE_LOG(LogNetPartialBunch, Verbose, TEXT("Merging Partial Bunch: %d Bytes"), Bunch.GetBytesLeft() );
 
-				if ( !Bunch.bHasGUIDs && Bunch.GetBitsLeft() > 0 )
+				if ( !Bunch.bHasPackageMapExports && Bunch.GetBitsLeft() > 0 )
 				{
 					InPartialBunch->AppendDataFromChecked( Bunch.GetDataPosChecked(), Bunch.GetBitsLeft() );
 				}
 
 				// Only the final partial bunch should ever be non byte aligned. This is enforced during partial bunch creation
 				// This is to ensure fast copies/appending of partial bunches. The final partial bunch may be non byte aligned.
-				check( Bunch.bHasGUIDs || Bunch.bPartialFinal || Bunch.GetBitsLeft() % 8 == 0 );
+				check( Bunch.bHasPackageMapExports || Bunch.bPartialFinal || Bunch.GetBitsLeft() % 8 == 0 );
 
 				// Advance the sequence of the current partial bunch so we know what to expect next
 				InPartialBunch->ChSequence = Bunch.ChSequence;
@@ -506,7 +506,7 @@ bool UChannel::ReceivedNextBunch( FInBunch & Bunch, bool & bOutSkipAck )
 						UE_LOG(LogNetPartialBunch, Verbose, TEXT("Completed Partial Bunch: Channel: %d ChSequence: %d. Num: %d Rel: %d CRC 0x%X"), InPartialBunch->ChIndex, InPartialBunch->ChSequence, InPartialBunch->GetNumBits(), Bunch.bReliable, FCrc::MemCrc_DEPRECATED(InPartialBunch->GetData(), InPartialBunch->GetNumBytes()));
 					}
 
-					check( !Bunch.bHasGUIDs );		// Shouldn't have these, they only go in initial partial export bunches
+					check( !Bunch.bHasPackageMapExports );		// Shouldn't have these, they only go in initial partial export bunches
 
 					HandleBunch = InPartialBunch;
 
@@ -711,7 +711,7 @@ FPacketIdRange UChannel::SendBunch( FOutBunch* Bunch, bool Merge )
 	check(!Closing);
 	check(Connection->Channels[ChIndex]==this);
 	check(!Bunch->IsError());
-	check( !Bunch->bHasGUIDs );
+	check( !Bunch->bHasPackageMapExports );
 
 	// Set bunch flags.
 	if( OpenPacketId.First==INDEX_NONE && OpenedLocally )
@@ -850,7 +850,7 @@ FPacketIdRange UChannel::SendBunch( FOutBunch* Bunch, bool Merge )
 		NextBunch->ChIndex = Bunch->ChIndex;
 		NextBunch->ChType = Bunch->ChType;
 
-		if ( !NextBunch->bHasGUIDs )
+		if ( !NextBunch->bHasPackageMapExports )
 		{
 			NextBunch->bHasMustBeMappedGUIDs |= Bunch->bHasMustBeMappedGUIDs;
 		}
@@ -2021,54 +2021,62 @@ void UActorChannel::ProcessBunch( FInBunch & Bunch )
 	// ----------------------------------------------
 	while ( !Bunch.AtEnd() && Connection != NULL && Connection->State != USOCK_Closed )
 	{
-		bool bObjectDeleted = false;
-		UObject* RepObj = ReadContentBlockHeader( Bunch, bObjectDeleted );
+		FNetBitReader Reader( Bunch.PackageMap, 0 );
+
+		bool bHasRepLayout = false;
+
+		// Read the content block header and payload
+		UObject* RepObj = ReadContentBlockPayload( Bunch, Reader, bHasRepLayout );
 
 		if ( Bunch.IsError() )
 		{
 			if ( Connection->InternalAck )
 			{
-				UE_LOG( LogNet, Warning, TEXT( "UActorChannel::ReceivedBunch: ReadContentBlockHeader FAILED. Bunch.IsError() == TRUE. (InternalAck) Breaking actor. RepObj: %s, Channel: %i" ), RepObj ? *RepObj->GetFullName() : TEXT( "NULL" ), ChIndex );
+				UE_LOG( LogNet, Warning, TEXT( "UActorChannel::ReceivedBunch: ReadContentBlockPayload FAILED. Bunch.IsError() == TRUE. (InternalAck) Breaking actor. RepObj: %s, Channel: %i" ), RepObj ? *RepObj->GetFullName() : TEXT( "NULL" ), ChIndex );
 				Broken = 1;
 				break;
 			}
 
-			UE_LOG( LogNet, Error, TEXT( "UActorChannel::ReceivedBunch: ReadContentBlockHeader FAILED. Bunch.IsError() == TRUE. Closing connection. RepObj: %s, Channel: %i" ), RepObj ? *RepObj->GetFullName() : TEXT( "NULL" ), ChIndex );
+			UE_LOG( LogNet, Error, TEXT( "UActorChannel::ReceivedBunch: ReadContentBlockPayload FAILED. Bunch.IsError() == TRUE. Closing connection. RepObj: %s, Channel: %i" ), RepObj ? *RepObj->GetFullName() : TEXT( "NULL" ), ChIndex );
 			Connection->Close();
 			return;
 		}
 
-		if (bObjectDeleted)
+		if ( Reader.GetNumBits() == 0 )
 		{
-			// Nothing else in this block, continue on
+			// Nothing else in this block, continue on (should have been a delete or create block)
 			continue;
 		}
 
 		if ( !RepObj || RepObj->IsPendingKill() )
 		{
-			UE_LOG(LogNet, Warning, TEXT("UActorChannel::ProcessBunch: ReadContentBlockHeader failed to find/create object. RepObj: %s, Channel: %i"), RepObj ? *RepObj->GetFullName() : TEXT( "NULL" ), ChIndex);
-
 			if ( !Actor || Actor->IsPendingKill() )
 			{
+				// If we couldn't find the actor, that's pretty bad, we need to stop processing on this channel
+				UE_LOG( LogNet, Warning, TEXT( "UActorChannel::ProcessBunch: ReadContentBlockPayload failed to find/create ACTOR. RepObj: %s, Channel: %i" ), RepObj ? *RepObj->GetFullName() : TEXT( "NULL" ), ChIndex );
 				Broken = 1;
 			}
+			else
+			{
+				UE_LOG( LogNet, Warning, TEXT( "UActorChannel::ProcessBunch: ReadContentBlockPayload failed to find/create object. RepObj: %s, Channel: %i" ), RepObj ? *RepObj->GetFullName() : TEXT( "NULL" ), ChIndex );
+			}
 
-			break;	// We can't continue reading since this will throw off the de-serialization
+			continue;	// Since content blocks separate the payload from the main stream, we can skip to the next one
 		}
 
 		TSharedRef< FObjectReplicator > & Replicator = FindOrCreateReplicator( RepObj );
 
 		bool bHasUnmapped = false;
 
-		if ( !Replicator->ReceivedBunch( Bunch, RepFlags, bHasUnmapped ) )
+		if ( !Replicator->ReceivedBunch( Reader, RepFlags, bHasRepLayout, bHasUnmapped ) )
 		{
 			if ( Connection->InternalAck )
 			{
 				UE_LOG( LogNet, Warning, TEXT( "UActorChannel::ProcessBunch: Replicator.ReceivedBunch failed (InternalAck) Breaking actor.  Closing connection. RepObj: %s, Channel: %i" ), RepObj ? *RepObj->GetFullName() : TEXT( "NULL" ), ChIndex );
-				Broken = 1;
-				break;
+				continue;		// Don't consider this catastrophic in replays
 			}
 
+			// For now, with regular connections, consider this catastrophic, but someday we could consider supporting backwards compatibility here too
 			UE_LOG( LogNet, Error, TEXT( "UActorChannel::ProcessBunch: Replicator.ReceivedBunch failed.  Closing connection. RepObj: %s, Channel: %i" ), RepObj ? *RepObj->GetFullName() : TEXT( "NULL" ), ChIndex );
 			Connection->Close();
 			return;
@@ -2244,7 +2252,7 @@ bool UActorChannel::ReplicateActor()
 		if (!RepComp.Key().IsValid())
 		{
 			// Write a deletion content header:
-			BeginContentBlockForSubObjectDelete( Bunch, RepComp.Value()->ObjectNetGUID );
+			WriteContentBlockForSubObjectDelete( Bunch, RepComp.Value()->ObjectNetGUID );
 
 			WroteSomethingImportant = true;
 			Bunch.bReliable = true;
@@ -2410,7 +2418,7 @@ void UActorChannel::StartBecomingDormant()
 	Connection->StartTickingChannel(this);
 }
 
-void UActorChannel::BeginContentBlock( UObject* Obj, FOutBunch &Bunch )
+void UActorChannel::WriteContentBlockHeader( UObject* Obj, FOutBunch &Bunch, const bool bHasRepLayout )
 {
 	const int NumStartingBits = Bunch.GetNumBits();
 
@@ -2421,6 +2429,8 @@ void UActorChannel::BeginContentBlock( UObject* Obj, FOutBunch &Bunch )
 	}
 #endif
 
+	Bunch.WriteBit( bHasRepLayout ? 1 : 0 );
+
 	// If we are referring to the actor on the channel, we don't need to send anything (except a bit signifying this)
 	const bool IsActor = Obj == Actor;
 
@@ -2428,7 +2438,7 @@ void UActorChannel::BeginContentBlock( UObject* Obj, FOutBunch &Bunch )
 
 	if ( IsActor )
 	{
-		NETWORK_PROFILER(GNetworkProfiler.TrackBeginContentBlock(Obj, Bunch.GetNumBits() - NumStartingBits, Connection));
+		NETWORK_PROFILER( GNetworkProfiler.TrackBeginContentBlock( Obj, Bunch.GetNumBits() - NumStartingBits, Connection ) );
 		return;
 	}
 
@@ -2461,12 +2471,15 @@ void UActorChannel::BeginContentBlock( UObject* Obj, FOutBunch &Bunch )
 	NETWORK_PROFILER(GNetworkProfiler.TrackBeginContentBlock(Obj, Bunch.GetNumBits() - NumStartingBits, Connection));
 }
 
-void UActorChannel::BeginContentBlockForSubObjectDelete( FOutBunch & Bunch, FNetworkGUID & GuidToDelete )
+void UActorChannel::WriteContentBlockForSubObjectDelete( FOutBunch & Bunch, FNetworkGUID & GuidToDelete )
 {
 	check( Connection->Driver->IsServer() );
 
 	const int NumStartingBits = Bunch.GetNumBits();
 
+	// No replayout here
+	Bunch.WriteBit( 0 );
+
 	// Send a 0 bit to signify we are dealing with sub-objects
 	Bunch.WriteBit( 0 );
 
@@ -2488,38 +2501,45 @@ void UActorChannel::BeginContentBlockForSubObjectDelete( FOutBunch & Bunch, FNet
 	NETWORK_PROFILER(GNetworkProfiler.TrackBeginContentBlock(nullptr, Bunch.GetNumBits() - NumStartingBits, Connection));
 }
 
-void UActorChannel::EndContentBlock( UObject *Obj, FOutBunch &Bunch, const FClassNetCache* ClassCache )
+int32 UActorChannel::WriteContentBlockPayload( UObject* Obj, FOutBunch &Bunch, const bool bHasRepLayout, FNetBitWriter& Payload )
 {
-	check(Obj);
+	const int32 StartHeaderBits = Bunch.GetNumBits();
 
-	const int NumStartingBits = Bunch.GetNumBits();
+	WriteContentBlockHeader( Obj, Bunch, bHasRepLayout );
 
-	if (!ClassCache)
-	{
-		ClassCache = Connection->Driver->NetCache->GetClassNetCache( Obj->GetClass() );
-	}
+	uint32 NumPayloadBits = Payload.GetNumBits();
 
-	if ( Connection->InternalAck )
-	{
-		// Write out 0 checksum to signify done
-		uint32 Checksum = 0;
-		Bunch << Checksum;
-	}
-	else
-	{
-		// Write max int to signify done
-		Bunch.WriteIntWrapped(ClassCache->GetMaxIndex(), ClassCache->GetMaxIndex()+1);
-	}
+	Bunch.SerializeIntPacked( NumPayloadBits );
+
+	const int32 HeaderNumBits = Bunch.GetNumBits() - StartHeaderBits;
 
-	NETWORK_PROFILER(GNetworkProfiler.TrackEndContentBlock(Obj, Bunch.GetNumBits() - NumStartingBits, Connection));
+	Bunch.SerializeBits( Payload.GetData(), Payload.GetNumBits() );
+
+	return HeaderNumBits;
 }
 
-UObject* UActorChannel::ReadContentBlockHeader(FInBunch & Bunch, bool& bObjectDeleted)
+UObject* UActorChannel::ReadContentBlockHeader( FInBunch & Bunch, bool& bObjectDeleted, bool& bOutHasRepLayout )
 {
 	const bool IsServer = Connection->Driver->IsServer();
 	bObjectDeleted = false;
 
-	if ( Bunch.ReadBit() )
+	bOutHasRepLayout = Bunch.ReadBit() != 0 ? true : false;
+
+	if ( Bunch.IsError() )
+	{
+		UE_LOG( LogNetTraffic, Error, TEXT( "UActorChannel::ReadContentBlockHeader: Bunch.IsError() == true after bOutHasRepLayout. Actor: %s" ), *Actor->GetName() );
+		return NULL;
+	}
+
+	const bool bIsActor = Bunch.ReadBit() != 0 ? true : false;
+
+	if ( Bunch.IsError() )
+	{
+		UE_LOG( LogNetTraffic, Error, TEXT( "UActorChannel::ReadContentBlockHeader: Bunch.IsError() == true after reading actor bit. Actor: %s" ), *Actor->GetName() );
+		return NULL;
+	}
+
+	if ( bIsActor )
 	{
 		// If this is for the actor on the channel, we don't need to read anything else
 		return Actor;
@@ -2589,7 +2609,15 @@ UObject* UActorChannel::ReadContentBlockHeader(FInBunch & Bunch, bool& bObjectDe
 		return SubObj;
 	}
 
-	if ( Bunch.ReadBit() )
+	const bool bStablyNamed = Bunch.ReadBit() != 0 ? true : false;
+
+	if ( Bunch.IsError() )
+	{
+		UE_LOG( LogNetTraffic, Error, TEXT( "UActorChannel::ReadContentBlockHeader: Bunch.IsError() == true after reading stably named bit. Actor: %s" ), *Actor->GetName() );
+		return NULL;
+	}
+
+	if ( bStablyNamed )
 	{
 		// If this is a stably named sub-object, we shouldn't need to create it
 		if ( SubObj == NULL )
@@ -2698,6 +2726,139 @@ UObject* UActorChannel::ReadContentBlockHeader(FInBunch & Bunch, bool& bObjectDe
 	return SubObj;
 }
 
+UObject* UActorChannel::ReadContentBlockPayload( FInBunch &Bunch, FNetBitReader& OutPayload, bool& bOutHasRepLayout )
+{
+	bool bObjectDeleted = false;
+	UObject* RepObj = ReadContentBlockHeader( Bunch, bObjectDeleted, bOutHasRepLayout );
+
+	if ( Bunch.IsError() )
+	{
+		UE_LOG( LogNet, Error, TEXT( "UActorChannel::ReadContentBlockPayload: ReadContentBlockHeader FAILED. Bunch.IsError() == TRUE. Closing connection. RepObj: %s, Channel: %i" ), RepObj ? *RepObj->GetFullName() : TEXT( "NULL" ), ChIndex );
+		return nullptr;
+	}
+
+	if ( bObjectDeleted )
+	{
+		OutPayload.SetData( Bunch, 0 );
+
+		// Nothing else in this block, continue on
+		return nullptr;
+	}
+
+	uint32 NumPayloadBits = 0;
+	Bunch.SerializeIntPacked( NumPayloadBits );
+
+	if ( Bunch.IsError() )
+	{
+		UE_LOG( LogNet, Error, TEXT( "UActorChannel::ReceivedBunch: Read NumPayloadBits FAILED. Bunch.IsError() == TRUE. Closing connection. RepObj: %s, Channel: %i" ), RepObj ? *RepObj->GetFullName() : TEXT( "NULL" ), ChIndex );
+		return nullptr;
+	}
+
+	OutPayload.SetData( Bunch, NumPayloadBits );
+
+	return RepObj;
+}
+
+int32 UActorChannel::WriteFieldHeaderAndPayload( FNetBitWriter& Bunch, const FClassNetCache* ClassCache, const FFieldNetCache * FieldCache, FNetBitWriter& Payload )
+{
+	const int32 NumOriginalBits = Bunch.GetNumBits();
+
+	NET_CHECKSUM( Bunch );
+
+	if ( Connection->InternalAck )
+	{
+		uint32 Checksum = FieldCache->FieldChecksum;
+		Bunch << Checksum;
+	}
+	else
+	{
+		Bunch.WriteIntWrapped( FieldCache->FieldNetIndex, ClassCache->GetMaxIndex() + 1 );
+	}
+
+	uint32 NumPayloadBits = Payload.GetNumBits();
+
+	Bunch.SerializeIntPacked( NumPayloadBits );
+	Bunch.SerializeBits( Payload.GetData(), NumPayloadBits );
+
+	return Bunch.GetNumBits() - NumOriginalBits;
+}
+
+bool UActorChannel::ReadFieldHeaderAndPayload( UObject* Object, const FClassNetCache* ClassCache, FNetBitReader& Bunch, const FFieldNetCache** OutField, FNetBitReader& OutPayload ) const
+{
+	*OutField = nullptr;
+
+	if ( Bunch.GetBitsLeft() == 0 )
+	{
+		return false;	// We're done
+	}
+	
+	NET_CHECKSUM( Reader );
+
+	if ( Connection->InternalAck )
+	{
+		// Replays use checksum rather than index
+		uint32 Checksum = 0;
+		Bunch << Checksum;
+
+		if ( Bunch.IsError() )
+		{
+			UE_LOG( LogNet, Error, TEXT( "ReadFieldHeaderAndPayload: Error reading checksum: %s" ), *Object->GetFullName() );
+			return false;
+		}
+
+		*OutField = ClassCache->GetFromChecksum( Checksum );
+
+		if ( *OutField == NULL )
+		{
+			UE_LOG( LogNet, Warning, TEXT( "ReadFieldHeaderAndPayload: GetFromChecksum failed: %s" ), *Object->GetFullName() );
+		}
+	}
+	else
+	{
+		const int32 RepIndex = Bunch.ReadInt( ClassCache->GetMaxIndex() + 1 );
+
+		if ( Bunch.IsError() )
+		{
+			UE_LOG( LogRep, Error, TEXT( "ReadFieldHeaderAndPayload: Error reading RepIndex: %s" ), *Object->GetFullName() );
+			return false;
+		}
+
+		if ( RepIndex > ClassCache->GetMaxIndex() )
+		{
+			// We shouldn't be receiving this bunch of this object has no properties or RPC functions to process
+			UE_LOG( LogRep, Error, TEXT( "ReadFieldHeaderAndPayload: RepIndex too large: %s" ), *Object->GetFullName() );
+			Bunch.SetError();
+			return false;
+		}
+
+		*OutField = ClassCache->GetFromIndex( RepIndex );
+
+		if ( *OutField == NULL )
+		{
+			UE_LOG( LogNet, Warning, TEXT( "ReadFieldHeaderAndPayload: GetFromIndex failed: %s" ), *Object->GetFullName() );
+		}
+	}
+
+	uint32 NumPayloadBits = 0;
+	Bunch.SerializeIntPacked( NumPayloadBits );
+
+	if ( Bunch.IsError() )
+	{
+		UE_LOG( LogNet, Error, TEXT( "ReadFieldHeaderAndPayload: Error reading numbits: %s" ), *Object->GetFullName() );
+		return false;
+	}
+
+	OutPayload.SetData( Bunch, NumPayloadBits );
+
+	if ( Bunch.IsError() )
+	{
+		UE_LOG( LogNet, Error, TEXT( "ReadFieldHeaderAndPayload: Error reading payload: %s" ), *Object->GetFullName() );
+		return false;
+	}
+
+	return true;		// More to read
+}
+
 FObjectReplicator & UActorChannel::GetActorReplicationData()
 {
 	return ReplicationMap.FindChecked(Actor).Get();
@@ -2782,9 +2943,10 @@ bool UActorChannel::ReplicateSubobject(UObject *Obj, FOutBunch &Bunch, const FRe
 	bool WroteSomething = FindOrCreateReplicator(Obj).Get().ReplicateProperties(Bunch, RepFlags);
 	if (NewSubobject && !WroteSomething)
 	{
-		BeginContentBlock( Obj, Bunch );
+		// Write empty payload to force object creation
+		FNetBitWriter EmptyPayload;
+		WriteContentBlockPayload( Obj, Bunch, false, EmptyPayload );
 		WroteSomething= true;
-		EndContentBlock( Obj, Bunch );
 	}
 
 	return WroteSomething;
diff --git a/Engine/Source/Runtime/Engine/Private/DataReplication.cpp b/Engine/Source/Runtime/Engine/Private/DataReplication.cpp
--- a/Engine/Source/Runtime/Engine/Private/DataReplication.cpp
+++ b/Engine/Source/Runtime/Engine/Private/DataReplication.cpp
@@ -61,7 +61,7 @@ bool IsCustomDeltaProperty( UProperty * Property )
 	return false;
 }
 
-bool FObjectReplicator::SerializeCustomDeltaProperty( UNetConnection * Connection, void* Src, UProperty * Property, int32 ArrayDim, FNetBitWriter & OutBunch, TSharedPtr<INetDeltaBaseState> &NewFullState, TSharedPtr<INetDeltaBaseState> & OldState )
+bool FObjectReplicator::SerializeCustomDeltaProperty( UNetConnection * Connection, void* Src, UProperty * Property, uint32 ArrayIndex, FNetBitWriter & OutBunch, TSharedPtr<INetDeltaBaseState> &NewFullState, TSharedPtr<INetDeltaBaseState> & OldState )
 {
 	check( NewFullState.IsValid() == false ); // NewState is passed in as NULL and instantiated within this function if necessary
 
@@ -94,7 +94,12 @@ bool FObjectReplicator::SerializeCustomDeltaProperty( UNetConnection * Connectio
 
 	Parms.Struct = StructProperty->Struct;
 
-	return CppStructOps->NetDeltaSerialize( Parms, Property->ContainerPtrToValuePtr<void>( Src, ArrayDim ) );
+	if ( Property->ArrayDim != 1 )
+	{
+		OutBunch.SerializeIntPacked( ArrayIndex );
+	}
+
+	return CppStructOps->NetDeltaSerialize( Parms, Property->ContainerPtrToValuePtr<void>( Src, ArrayIndex ) );
 }
 
 /** 
@@ -390,71 +395,7 @@ void FObjectReplicator::ReceivedNak( int32 NakPacketId )
 	}
 }
 
-const FFieldNetCache* FObjectReplicator::ReadField( const FClassNetCache* ClassCache, FInBunch& Bunch ) const
-{
-	if ( Connection->InternalAck )
-	{
-		// Replays use checksum rather than index
-		uint32 Checksum = 0;
-		Bunch << Checksum;
-
-		if ( Bunch.IsError() )
-		{
-			UE_LOG(LogNet, Error, TEXT("ReadField: Error reading checksum: %s"), *GetObject()->GetFullName());
-			return NULL;
-		}
-
-		if ( Checksum == 0 )
-		{
-			return NULL;	// We're done
-		}
-
-		const FFieldNetCache * FieldNetCache = ClassCache->GetFromChecksum( Checksum );
-
-		if ( FieldNetCache == NULL )
-		{
-			UE_LOG(LogNet, Error, TEXT("ReadField: GetFromChecksum failed: %s"), *GetObject()->GetFullName());
-			Bunch.SetError();
-			return NULL;
-		}
-
-		return FieldNetCache;
-	}
-
-	const int32 RepIndex = Bunch.ReadInt( ClassCache->GetMaxIndex() + 1 );
-
-	if ( Bunch.IsError() )
-	{
-		UE_LOG(LogRep, Error, TEXT("ReadField: Error reading RepIndex: %s"), *GetObject()->GetFullName());
-		return NULL;
-	}
-
-	if ( RepIndex == ClassCache->GetMaxIndex() )
-	{
-		return NULL;	// We're done
-	}
-
-	if ( RepIndex > ClassCache->GetMaxIndex() )
-	{
-		// We shouldn't be receiving this bunch of this object has no properties or RPC functions to process
-		UE_LOG(LogRep, Error, TEXT("ReadField: RepIndex too large: %s"), *GetObject()->GetFullName());
-		Bunch.SetError();
-		return NULL;
-	}
-	
-	const FFieldNetCache * FieldNetCache = ClassCache->GetFromIndex( RepIndex );
-
-	if ( FieldNetCache == NULL )
-	{
-		UE_LOG(LogNet, Error, TEXT("ReadField: GetFromIndex failed: %s"), *GetObject()->GetFullName());
-		Bunch.SetError();
-		return NULL;
-	}
-
-	return FieldNetCache;
-}
-
-bool FObjectReplicator::ReceivedBunch( FInBunch& Bunch, const FReplicationFlags& RepFlags, bool& bOutHasUnmapped )
+bool FObjectReplicator::ReceivedBunch( FNetBitReader& Bunch, const FReplicationFlags& RepFlags, const bool bHasRepLayout, bool& bOutHasUnmapped )
 {
 	UObject* Object = GetObject();
 
@@ -466,7 +407,7 @@ bool FObjectReplicator::ReceivedBunch( FInBunch& Bunch, const FReplicationFlags&
 
 	UPackageMap * PackageMap = OwningChannel->Connection->PackageMap;
 
-	const bool bIsServer = ( OwningChannel->Connection->Driver->ServerConnection == NULL );
+	const bool bIsServer = OwningChannel->Connection->Driver->IsServer();
 
 	const FClassNetCache * ClassCache = OwningChannel->Connection->Driver->NetCache->GetClassNetCache( ObjectClass );
 
@@ -476,33 +417,63 @@ bool FObjectReplicator::ReceivedBunch( FInBunch& Bunch, const FReplicationFlags&
 		return false;
 	}
 
-	bool bThisBunchReplicatedProperties = false;
+	// Handle replayout properties
+	if ( bHasRepLayout )
+	{
+		// Server shouldn't receive properties.
+		if ( bIsServer )
+		{
+			UE_LOG( LogNet, Error, TEXT( "Server received RepLayout properties: %s" ), *Object->GetFullName() );
+			return false;
+		}
+
+		if ( !bHasReplicatedProperties )
+		{
+			bHasReplicatedProperties = true;		// Persistent, not reset until PostNetReceive is called
+			PreNetReceive();
+		}
 
-	// Read first field
-	const FFieldNetCache * FieldCache = ReadField( ClassCache, Bunch );
+		const bool bShouldReceiveRepNotifies = Connection->Driver->ShouldReceiveRepNotifiesForObject( Object );
 
-	if ( Bunch.IsError() )
-	{
-		UE_LOG(LogNet, Error, TEXT("ReceivedBunch: Error reading field 1: %s"), *Object->GetFullName());
-		return false;
-	}
+		bool bLocalHasUnmapped = false;
 
-	if ( FieldCache == NULL )
-	{
-		// There are no actual replicated properties or functions in this bunch. That is ok - we may have gotten this
-		// actor/sub-object because we want the client to spawn one (but we aren't actually replicating properties on it)
-		return true;
+		if ( !RepLayout->ReceiveProperties( OwningChannel, ObjectClass, RepState, ( void* )Object, Bunch, bLocalHasUnmapped, bShouldReceiveRepNotifies ) )
+		{
+			UE_LOG( LogRep, Error, TEXT( "RepLayout->ReceiveProperties FAILED: %s" ), *Object->GetFullName() );
+			return false;
+		}
+
+		if ( bLocalHasUnmapped )
+		{
+			bOutHasUnmapped = true;
+		}
 	}
 
-	while ( FieldCache )
+	FNetBitReader Reader( Bunch.PackageMap );
+
+	// Read fields from stream
+	const FFieldNetCache * FieldCache = nullptr;
+	
+	// Read each property/function blob into Reader (so we've safely jumped over this data in the Bunch/stream at this point)
+	while ( OwningChannel->ReadFieldHeaderAndPayload( Object, ClassCache, Bunch, &FieldCache, Reader ) )
 	{
-		// Receive properties from the net.
-		UProperty* ReplicatedProp = NULL;
+		if ( Bunch.IsError() )
+		{
+			UE_LOG( LogNet, Error, TEXT( "ReceivedBunch: Error reading field: %s" ), *Object->GetFullName() );
+			return false;
+		}
 
-		while ( FieldCache && ( ReplicatedProp = Cast< UProperty >( FieldCache->Field ) ) != NULL )
+		if ( FieldCache == nullptr )
 		{
-			NET_CHECKSUM( Bunch );
+			UE_LOG( LogNet, Warning, TEXT( "ReceivedBunch: FieldCache == nullptr: %s" ), *Object->GetFullName() );
+			continue;
+		}
+
+		UProperty* ReplicatedProp = NULL;
 
+		// Handle property
+		if ( ( ReplicatedProp = Cast< UProperty >( FieldCache->Field ) ) != NULL )
+		{
 			// Server shouldn't receive properties.
 			if ( bIsServer )
 			{
@@ -510,227 +481,185 @@ bool FObjectReplicator::ReceivedBunch( FInBunch& Bunch, const FReplicationFlags&
 				return false;
 			}
 		
-			bThisBunchReplicatedProperties = true;
+			// We should only be receiving custom delta properties (since RepLayout handles the rest)
+			if ( !Retirement[ReplicatedProp->RepIndex].CustomDelta )
+			{
+				UE_LOG( LogNet, Error, TEXT( "Client received non custom delta property value %s in %s" ), *ReplicatedProp->GetName(), *Object->GetFullName() );
+				return false;
+			}
 
+			// Call PreNetReceive if we haven't yet
 			if ( !bHasReplicatedProperties )
 			{
 				bHasReplicatedProperties = true;		// Persistent, not reset until PostNetReceive is called
 				PreNetReceive();
 			}
 
-			bool DebugProperty = false;
 #if !(UE_BUILD_SHIPPING || UE_BUILD_TEST)
 			{
 				static IConsoleVariable* CVar = IConsoleManager::Get().FindConsoleVariable(TEXT("net.Replication.DebugProperty"));
 				if (CVar && !CVar->GetString().IsEmpty() && ReplicatedProp->GetName().Contains(CVar->GetString()) )
 				{
 					UE_LOG(LogRep, Log, TEXT("Replicating Property[%d] %s on %s"), ReplicatedProp->RepIndex, *ReplicatedProp->GetName(), *Object->GetName());
-					DebugProperty = true;
 				}
 			}
 #endif
-			if ( !Retirement[ ReplicatedProp->RepIndex ].CustomDelta )
-			{
-				bool bLocalHasUnmapped = false;
-				// We hijack a non custom delta property to signify we are using FRepLayout to read the entire property block
-				const bool bShouldReceiveRepNotifies = Connection->Driver->ShouldReceiveRepNotifiesForObject( Object );
-				if ( !RepLayout->ReceiveProperties( ObjectClass, RepState, (void*)Object, Bunch, bLocalHasUnmapped, bShouldReceiveRepNotifies ) )
-				{
-					UE_LOG(LogRep, Error, TEXT("ReceiveProperties FAILED %s in %s"), *ReplicatedProp->GetName(), *Object->GetFullName());
-					return false;
-				}
 
-				if ( bLocalHasUnmapped )
-				{
-					bOutHasUnmapped = true;
-				}
-			}
-			else
+			// Receive array index (static sized array, i.e. MemberVariable[4])
+			uint32 Element = 0;
+			if ( ReplicatedProp->ArrayDim != 1 )
 			{
-				// Receive array index.
-				uint32 Element = 0;
-				if ( ReplicatedProp->ArrayDim != 1 )
-				{
-					check( ReplicatedProp->ArrayDim >= 2 );
-
-					Bunch.SerializeIntPacked( Element );
+				check( ReplicatedProp->ArrayDim >= 2 );
 
-					if ( Element >= (uint32)ReplicatedProp->ArrayDim )
-					{
-						UE_LOG(LogRep, Error, TEXT("Element index too large %s in %s"), *ReplicatedProp->GetName(), *Object->GetFullName());
-						return false;
-					}
-				}
-
-				// Pointer to destination.
-				uint8* Data = ReplicatedProp->ContainerPtrToValuePtr<uint8>((uint8*)Object, Element);
-				TArray<uint8>	MetaData;
-				const PTRINT DataOffset = Data - (uint8*)Object;
-
-				// Receive custom delta property.
-				UStructProperty * StructProperty = Cast< UStructProperty >( ReplicatedProp );
+				Reader.SerializeIntPacked( Element );
 
-				if ( StructProperty == NULL )
+				if ( Element >= (uint32)ReplicatedProp->ArrayDim )
 				{
-					// This property isn't custom delta
-					UE_LOG(LogRepTraffic, Error, TEXT("Property isn't custom delta %s"), *ReplicatedProp->GetName());
+					UE_LOG(LogRep, Error, TEXT("Element index too large %s in %s"), *ReplicatedProp->GetName(), *Object->GetFullName());
 					return false;
 				}
+			}
 
-				UScriptStruct * InnerStruct = StructProperty->Struct;
+			// Pointer to destination.
+			uint8* Data = ReplicatedProp->ContainerPtrToValuePtr<uint8>((uint8*)Object, Element);
+			TArray<uint8>	MetaData;
+			const PTRINT DataOffset = Data - (uint8*)Object;
 
-				if ( !( InnerStruct->StructFlags & STRUCT_NetDeltaSerializeNative ) )
-				{
-					// This property isn't custom delta
-					UE_LOG(LogRepTraffic, Error, TEXT("Property isn't custom delta %s"), *ReplicatedProp->GetName());
-					return false;
-				}
+			// Receive custom delta property.
+			UStructProperty * StructProperty = Cast< UStructProperty >( ReplicatedProp );
 
-				UScriptStruct::ICppStructOps * CppStructOps = InnerStruct->GetCppStructOps();
+			if ( StructProperty == NULL )
+			{
+				// This property isn't custom delta
+				UE_LOG(LogRepTraffic, Error, TEXT("Property isn't custom delta %s"), *ReplicatedProp->GetName());
+				return false;
+			}
 
-				check( CppStructOps );
-				check( !InnerStruct->InheritedCppStructOps() );
+			UScriptStruct * InnerStruct = StructProperty->Struct;
 
-				FNetDeltaSerializeInfo Parms;
+			if ( !( InnerStruct->StructFlags & STRUCT_NetDeltaSerializeNative ) )
+			{
+				// This property isn't custom delta
+				UE_LOG(LogRepTraffic, Error, TEXT("Property isn't custom delta %s"), *ReplicatedProp->GetName());
+				return false;
+			}
 
-				FNetSerializeCB NetSerializeCB( OwningChannel->Connection->Driver );
+			UScriptStruct::ICppStructOps * CppStructOps = InnerStruct->GetCppStructOps();
 
-				Parms.DebugName			= StructProperty->GetName();
-				Parms.Struct			= InnerStruct;
-				Parms.Map				= PackageMap;
-				Parms.Reader			= &Bunch;
-				Parms.NetSerializeCB	= &NetSerializeCB;
+			check( CppStructOps );
+			check( !InnerStruct->InheritedCppStructOps() );
 
-				// Call the custom delta serialize function to handle it
-				CppStructOps->NetDeltaSerialize( Parms, Data );
+			FNetDeltaSerializeInfo Parms;
 
-				if ( Bunch.IsError() )
-				{
-					UE_LOG(LogNet, Error, TEXT("ReceivedBunch: NetDeltaSerialize - Bunch.IsError() == true: %s"), *Object->GetFullName());
-					return false;
-				}
+			FNetSerializeCB NetSerializeCB( OwningChannel->Connection->Driver );
 
-				if ( Parms.bOutHasMoreUnmapped )
-				{
-					UnmappedCustomProperties.Add( DataOffset, StructProperty );
-					bOutHasUnmapped = true;
-				}
+			Parms.DebugName			= StructProperty->GetName();
+			Parms.Struct			= InnerStruct;
+			Parms.Map				= PackageMap;
+			Parms.Reader			= &Reader;
+			Parms.NetSerializeCB	= &NetSerializeCB;
 
-				// Successfully received it.
-				UE_LOG(LogRepTraffic, Log, TEXT(" %s - %s"), *Object->GetName(), *ReplicatedProp->GetName());
+			// Call the custom delta serialize function to handle it
+			CppStructOps->NetDeltaSerialize( Parms, Data );
 
-				// Notify the Object if this var is RepNotify
-				QueuePropertyRepNotify( Object, ReplicatedProp, Element, MetaData );
-			}	
-			
-			// Read next field
-			FieldCache = ReadField( ClassCache, Bunch );
-
-			if ( Bunch.IsError() )
+			if ( Reader.IsError() )
 			{
-				UE_LOG(LogNet, Error, TEXT("ReceivedBunch: Error reading field 2: %s"), *Object->GetFullName());
+				UE_LOG(LogNet, Error, TEXT("ReceivedBunch: NetDeltaSerialize - Reader.IsError() == true: %s"), *Object->GetFullName());
 				return false;
 			}
-		}
 
-		// Handle function calls.
-		if ( FieldCache && Cast< UFunction >( FieldCache->Field ) )
+			if ( Parms.bOutHasMoreUnmapped )
+			{
+				UnmappedCustomProperties.Add( DataOffset, StructProperty );
+				bOutHasUnmapped = true;
+			}
+
+			// Successfully received it.
+			UE_LOG(LogRepTraffic, Log, TEXT(" %s - %s"), *Object->GetName(), *ReplicatedProp->GetName());
+
+			// Notify the Object if this var is RepNotify
+			QueuePropertyRepNotify( Object, ReplicatedProp, Element, MetaData );
+		}
+		// Handle function call
+		else if ( Cast< UFunction >( FieldCache->Field ) )
 		{
-			FName Message = FieldCache->Field->GetFName();
-			UFunction * Function = Object->FindFunction( Message );
+			FName FunctionName = FieldCache->Field->GetFName();
+			UFunction * Function = Object->FindFunction( FunctionName );
 
 			if ( Function == NULL )
 			{
-				UE_LOG(LogNet, Error, TEXT("ReceivedBunch: Function == NULL: %s"), *Object->GetFullName());
+				UE_LOG( LogNet, Error, TEXT( "ReceivedBunch: Function not found. Object: %s, Function: %s" ), *Object->GetFullName(), *FunctionName.ToString() );
 				return false;
 			}
 
 			if ( ( Function->FunctionFlags & FUNC_Net ) == 0 )
 			{
-				UE_LOG(LogRep, Error, TEXT("Rejected non RPC function %s in %s"), *Message.ToString(), *Object->GetFullName());
+				UE_LOG( LogRep, Error, TEXT( "Rejected non RPC function. Object: %s, Function: %s" ), *Object->GetFullName(), *FunctionName.ToString() );
 				return false;
 			}
 
 			if ( ( Function->FunctionFlags & ( bIsServer ? FUNC_NetServer : ( FUNC_NetClient | FUNC_NetMulticast ) ) ) == 0 )
 			{
-				UE_LOG(LogRep, Error, TEXT("Rejected RPC function due to access rights %s in %s"), *Message.ToString(), *Object->GetFullName());
+				UE_LOG( LogRep, Error, TEXT( "Rejected RPC function due to access rights. Object: %s, Function: %s" ), *Object->GetFullName(), *FunctionName.ToString() );
 				return false;
 			}
 
-			UE_LOG(LogRepTraffic, Log, TEXT("      Received RPC: %s"), *Message.ToString());
-
-			// Get the parameters.
-			FMemMark Mark(FMemStack::Get());
-			uint8* Parms = new(FMemStack::Get(),MEM_Zeroed,Function->ParmsSize)uint8;
-
-			// Use the replication layout to receive the rpc parameter values
-			TSharedPtr<FRepLayout> FuncRepLayout = OwningChannel->Connection->Driver->GetFunctionRepLayout( Function );
-
-			FuncRepLayout->ReceivePropertiesForRPC( Object, Function, OwningChannel, Bunch, Parms );
-
-			if ( Bunch.IsError() )
-			{
-				UE_LOG(LogRep, Error, TEXT("ReceivedBunch: ReceivePropertiesForRPC - Bunch.IsError() == true: Function: %s, Object: %s"), *Message.ToString(), *Object->GetFullName());
-				return false;
-			}
+			UE_LOG( LogRepTraffic, Log, TEXT( "      Received RPC: %s" ), *FunctionName.ToString() );
 
 			// validate that the function is callable here
 			const bool bCanExecute = ( !bIsServer || RepFlags.bNetOwner );		// we are client or net owner
 
 			if ( bCanExecute )
 			{
+				// Get the parameters.
+				FMemMark Mark( FMemStack::Get() );
+				uint8* Parms = new( FMemStack::Get(), MEM_Zeroed, Function->ParmsSize )uint8;
+
+				// Use the replication layout to receive the rpc parameter values
+				TSharedPtr<FRepLayout> FuncRepLayout = OwningChannel->Connection->Driver->GetFunctionRepLayout( Function );
+
+				FuncRepLayout->ReceivePropertiesForRPC( Object, Function, OwningChannel, Reader, Parms );
+
+				if ( Reader.IsError() )
+				{
+					UE_LOG( LogRep, Error, TEXT( "ReceivedBunch: ReceivePropertiesForRPC - Reader.IsError() == true: Function: %s, Object: %s" ), *FunctionName.ToString(), *Object->GetFullName() );
+					return false;
+				}
+
 				// Call the function.
 				RPC_ResetLastFailedReason();
 
 				Object->ProcessEvent( Function, Parms );
 
-				if ( RPC_GetLastFailedReason() != NULL )
+				// Destroy the parameters.
+				// warning: highly dependent on UObject::ProcessEvent freeing of parms!
+				for ( TFieldIterator<UProperty> It( Function ); It && ( It->PropertyFlags & ( CPF_Parm | CPF_ReturnParm ) ) == CPF_Parm; ++It )
 				{
-					UE_LOG(LogRep, Error, TEXT("ReceivedBunch: RPC_GetLastFailedReason: %s"), RPC_GetLastFailedReason());
-					return false;
+					It->DestroyValue_InContainer( Parms );
 				}
-			}
-			else
-			{
-				UE_LOG(LogRep, Verbose, TEXT("Rejected unwanted function %s in %s"), *Message.ToString(), *Object->GetFullName());
 
-#if 0
-				// FIXME: This may cause false positives, need to think of a better way to handle this, removing it for now
-				if ( !OwningChannel->Connection->TrackLogsPerSecond() )	// This will disconnect the client if we get here too often
+				Mark.Pop();
+
+				if ( RPC_GetLastFailedReason() != NULL )
 				{
-					UE_LOG(LogRep, Error, TEXT("Rejected too many unwanted functions %s in %s"), *Message.ToString(), *Object->GetFullName());
+					UE_LOG(LogRep, Error, TEXT("ReceivedBunch: RPC_GetLastFailedReason: %s"), RPC_GetLastFailedReason());
 					return false;
 				}
-#endif
 			}
-
-			// Destroy the parameters.
-			// warning: highly dependent on UObject::ProcessEvent freeing of parms!
-			for( TFieldIterator<UProperty> It(Function); It && (It->PropertyFlags & (CPF_Parm|CPF_ReturnParm))==CPF_Parm; ++It )
+			else
 			{
-				It->DestroyValue_InContainer(Parms);
+				UE_LOG( LogRep, Verbose, TEXT( "Rejected unwanted function %s in %s" ), *FunctionName.ToString(), *Object->GetFullName() );
 			}
 
-			Mark.Pop();
-
 			if ( Object == NULL || Object->IsPendingKill() )
 			{
 				// replicated function destroyed Object
-				return true;		// FIXME: Should this return false to kick connection?  Seems we'll cause a read misalignment here if we don't
-			}
-
-			// Next.
-			FieldCache = ReadField( ClassCache, Bunch );
-
-			if ( Bunch.IsError() )
-			{
-				UE_LOG(LogNet, Error, TEXT("ReceivedBunch: Error reading field 3: %s"), *Object->GetFullName());
-				return false;
+				return true;
 			}
 		}
-		else if ( FieldCache )
+		else
 		{
-			UE_LOG(LogRep, Error, TEXT("ReceivedBunch: Invalid replicated field %i in %s"), FieldCache->FieldNetIndex, *Object->GetFullName());
+			UE_LOG( LogRep, Error, TEXT( "ReceivedBunch: Invalid replicated field %i in %s" ), FieldCache->FieldNetIndex, *Object->GetFullName() );
 			return false;
 		}
 	}
@@ -798,7 +727,7 @@ static FORCEINLINE FPropertyRetirement ** UpdateAckedRetirements( FPropertyRetir
 	return Rec;
 }
 
-void FObjectReplicator::ReplicateCustomDeltaProperties( FOutBunch & Bunch, FReplicationFlags RepFlags, bool & bContentBlockWritten )
+void FObjectReplicator::ReplicateCustomDeltaProperties( FNetBitWriter & Bunch, FReplicationFlags RepFlags )
 {
 	if ( LifetimeCustomDeltaProperties.Num() == 0 )
 	{
@@ -895,15 +824,10 @@ void FObjectReplicator::ReplicateCustomDeltaProperties( FOutBunch & Bunch, FRepl
 		// Save NewState into the RecentCustomDeltaState array (old state is a reference into our RecentCustomDeltaState map)
 		OldState = NewState; 
 
-		// Write header, and data to send to the actual bunch
-		RepLayout->WritePropertyHeader( Object, ObjectClass, OwningChannel, It, Bunch, Index, bContentBlockWritten );
-
-		const int NumStartingBits = Bunch.GetNumBits();
+		// Write property header and payload to the bunch
+		WritePropertyHeaderAndPayload( Object, It, Bunch, TempBitWriter );
 
-		// Send property.
-		Bunch.SerializeBits( TempBitWriter.GetData(), TempBitWriter.GetNumBits() );
-
-		NETWORK_PROFILER(GNetworkProfiler.TrackReplicateProperty(It, Bunch.GetNumBits() - NumStartingBits, Connection));
+		NETWORK_PROFILER( GNetworkProfiler.TrackReplicateProperty( It, TempBitWriter.GetNumBits(), Connection ) );
 	}
 }
 
@@ -929,19 +853,17 @@ bool FObjectReplicator::ReplicateProperties( FOutBunch & Bunch, FReplicationFlag
 
 	UNetConnection* OwningChannelConnection = OwningChannel->Connection;
 
-	const int32 StartingBitNum = Bunch.GetNumBits();
+	FNetBitWriter Writer( Bunch.PackageMap, 0 );
 
-	bool bContentBlockWritten = false;
+	// Replicate properties in the layout
+	const bool bHasRepLayout = RepLayout->ReplicateProperties( RepState, ( uint8* )Object, ObjectClass, OwningChannel, Writer, RepFlags );
 
 	// Replicate all the custom delta properties (fast arrays, etc)
-	ReplicateCustomDeltaProperties( Bunch, RepFlags, bContentBlockWritten );
-
-	// Replicate properties in the layout
-	RepLayout->ReplicateProperties( RepState, (uint8*)Object, ObjectClass, OwningChannel, Bunch, RepFlags, bContentBlockWritten );
+	ReplicateCustomDeltaProperties( Writer, RepFlags );
 
 	// LastUpdateEmpty - this is done before dequeing the multicasted unreliable functions on purpose as they should not prevent
 	// an actor channel from going dormant.
-	bLastUpdateEmpty = ( Bunch.GetNumBits() == StartingBitNum );
+	bLastUpdateEmpty = Writer.GetNumBits() == 0;
 
 	// Replicate Queued (unreliable functions)
 	if ( RemoteFunctions != NULL && RemoteFunctions->GetNumBits() > 0 )
@@ -953,13 +875,7 @@ bool FObjectReplicator::ReplicateProperties( FOutBunch & Bunch, FReplicationFlag
 			UE_LOG( LogRepTraffic, Warning,	TEXT("      Sending queued RPCs: %s. Channel[%d] [%.1f bytes]"), *Object->GetName(), OwningChannel->ChIndex, RemoteFunctions->GetNumBits() / 8.f );
 		}
 
-		if ( !bContentBlockWritten )
-		{
-			OwningChannel->BeginContentBlock( Object, Bunch );
-			bContentBlockWritten = true;
-		}
-
-		Bunch.SerializeBits( RemoteFunctions->GetData(), RemoteFunctions->GetNumBits() );
+		Writer.SerializeBits( RemoteFunctions->GetData(), RemoteFunctions->GetNumBits() );
 		RemoteFunctions->Reset();
 		RemoteFuncInfo.Empty();
 
@@ -969,11 +885,11 @@ bool FObjectReplicator::ReplicateProperties( FOutBunch & Bunch, FReplicationFlag
 	// See if we wrote something important (anything but the 'end' int below).
 	// Note that queued unreliable functions are considered important (WroteImportantData) but not for bLastUpdateEmpty. LastUpdateEmpty
 	// is used for dormancy purposes. WroteImportantData is for determining if we should not include a component in replication.
-	const bool WroteImportantData = ( Bunch.GetNumBits() != StartingBitNum );
+	const bool WroteImportantData = Writer.GetNumBits() != 0;
 
 	if ( WroteImportantData )
 	{
-		OwningChannel->EndContentBlock( Object, Bunch, OwningChannelConnection->Driver->NetCache->GetClassNetCache( ObjectClass ) );
+		OwningChannel->WriteContentBlockPayload( Object, Bunch, bHasRepLayout, Writer );
 	}
 
 	return WroteImportantData;
@@ -1350,3 +1266,27 @@ void FObjectReplicator::QueuePropertyRepNotify( UObject* Object, UProperty * Pro
 		}
 	}
 }
+
+void FObjectReplicator::WritePropertyHeaderAndPayload(
+	UObject*			Object,
+	UProperty*			Property,
+	FNetBitWriter&		Bunch,
+	FNetBitWriter&		Payload ) const
+{
+	// Get class network info cache.
+	const FClassNetCache* ClassCache = Connection->Driver->NetCache->GetClassNetCache( ObjectClass );
+
+	check( ClassCache );
+
+	// Get the network friend property index to replicate
+	const FFieldNetCache * FieldCache = ClassCache->GetFromField( Property );
+
+	checkSlow( FieldCache );
+
+	// Send property name and optional array index.
+	check( FieldCache->FieldNetIndex <= ClassCache->GetMaxIndex() );
+
+	const int32 HeaderBits = OwningChannel->WriteFieldHeaderAndPayload( Bunch, ClassCache, FieldCache, Payload );
+
+	NETWORK_PROFILER( GNetworkProfiler.TrackWritePropertyHeader( Property, HeaderBits, nullptr ) );
+}
diff --git a/Engine/Source/Runtime/Engine/Private/DemoNetDriver.cpp b/Engine/Source/Runtime/Engine/Private/DemoNetDriver.cpp
--- a/Engine/Source/Runtime/Engine/Private/DemoNetDriver.cpp
+++ b/Engine/Source/Runtime/Engine/Private/DemoNetDriver.cpp
@@ -476,7 +476,7 @@ bool UDemoNetDriver::InitConnect( FNetworkNotify* InNotify, const FURL& ConnectU
 		return false;
 	}
 
-	GuidCache->SetIgnorePackageMismatchOverride( true );
+	GuidCache->SetNetworkChecksumMode( FNetGUIDCache::NETCHECKSUM_SaveButIgnore );
 
 	// Playback, local machine is a client, and the demo stream acts "as if" it's the server.
 	ServerConnection = NewObject<UNetConnection>(GetTransientPackage(), UDemoNetConnection::StaticClass());
@@ -693,8 +693,7 @@ bool UDemoNetDriver::InitListen( FNetworkNotify* InNotify, FURL& ListenURL, bool
 		return false;
 	}
 
-	GuidCache->SetIgnorePackageMismatchOverride( true );
-	GuidCache->SetShouldUseNetworkChecksum( true );
+	GuidCache->SetNetworkChecksumMode( FNetGUIDCache::NETCHECKSUM_SaveButIgnore );
 
 	check( World != NULL );
 
@@ -1191,7 +1190,6 @@ static void SerializeGuidCache( TSharedPtr< class FNetGUIDCache > GuidCache, FAr
 		*CheckpointArchive << It.Value().OuterGUID;
 		*CheckpointArchive << PathName;
 		*CheckpointArchive << It.Value().NetworkChecksum;
-		*CheckpointArchive << It.Value().PackageChecksum;
 
 		uint8 Flags = 0;
 		
@@ -1236,6 +1234,9 @@ void UDemoNetDriver::SaveCheckpoint()
 	// First, save the current guid cache
 	SerializeGuidCache( GuidCache, CheckpointArchive );
 
+	// Save the compatible rep layout map
+	( ( UPackageMapClient* )ClientConnections[0]->PackageMap )->SerializeCompatibleReplayoutMap( *CheckpointArchive );
+
 	const uint32 GuidCacheSize = CheckpointArchive->TotalSize();
 
 	FURL CheckpointURL;
@@ -1252,6 +1253,9 @@ void UDemoNetDriver::SaveCheckpoint()
 	CheckpointConnection->PlayerController->NetConnection	= CheckpointConnection;
 	//CheckpointConnection->OwningActor						= CheckpointConnection->PlayerController;
 
+	// Sync up ack status, so we don't send what we've already explicity saved from above
+	( ( UPackageMapClient* )CheckpointConnection->PackageMap )->SyncPackageMapExportAckStatus( ( UPackageMapClient* )ClientConnections[0]->PackageMap );
+
 	// Make sure we have the exact same actor channel indexes
 	for ( auto It = ClientConnections[0]->ActorChannels.CreateIterator(); It; ++It )
 	{
@@ -2308,6 +2312,10 @@ void UDemoNetDriver::LoadCheckpoint( FArchive* GotoCheckpointArchive, int64 Goto
 	GuidCache->ObjectLookup.Empty();
 	GuidCache->NetGUIDLookup.Empty();
 
+	GuidCache->CompatibleRepLayoutMap.Empty();
+	GuidCache->CompatibleRepLayoutPathToIndex.Empty();
+	GuidCache->CompatibleRepLayoutIndexToPath.Empty();
+
 	// Restore the spectator controller packagemap entry (so we find it when we process the checkpoint)
 	if ( SpectatorGUID.IsValid() )
 	{
@@ -2350,7 +2358,6 @@ void UDemoNetDriver::LoadCheckpoint( FArchive* GotoCheckpointArchive, int64 Goto
 		*GotoCheckpointArchive << CacheObject.OuterGUID;
 		*GotoCheckpointArchive << PathName;
 		*GotoCheckpointArchive << CacheObject.NetworkChecksum;
-		*GotoCheckpointArchive << CacheObject.PackageChecksum;
 
 		CacheObject.PathName = FName( *PathName );
 
@@ -2363,6 +2370,9 @@ void UDemoNetDriver::LoadCheckpoint( FArchive* GotoCheckpointArchive, int64 Goto
 		GuidCache->ObjectLookup.Add( Guid, CacheObject );
 	}
 
+	// Read in the compatible rep layouts in this checkpoint
+	( ( UPackageMapClient* )ServerConnection->PackageMap )->SerializeCompatibleReplayoutMap( *GotoCheckpointArchive );
+
 	ReadDemoFrameIntoPlaybackPackets( *GotoCheckpointArchive );
 
 	if ( PlaybackPackets.Num() > 0 )
diff --git a/Engine/Source/Runtime/Engine/Private/NetConnection.cpp b/Engine/Source/Runtime/Engine/Private/NetConnection.cpp
--- a/Engine/Source/Runtime/Engine/Private/NetConnection.cpp
+++ b/Engine/Source/Runtime/Engine/Private/NetConnection.cpp
@@ -914,7 +914,7 @@ void UNetConnection::ReceivedPacket( FBitReader& Reader )
 			Bunch.bDormant				= Bunch.bClose ? Reader.ReadBit() : 0;
 			Bunch.bReliable				= Reader.ReadBit();
 			Bunch.ChIndex				= Reader.ReadInt( MAX_CHANNELS );
-			Bunch.bHasGUIDs				= Reader.ReadBit();
+			Bunch.bHasPackageMapExports	= Reader.ReadBit();
 			Bunch.bHasMustBeMappedGUIDs	= Reader.ReadBit();
 			Bunch.bPartial				= Reader.ReadBit();
 
@@ -971,7 +971,7 @@ void UNetConnection::ReceivedPacket( FBitReader& Reader )
 				return;
 			}
 
-			if (Bunch.bHasGUIDs)
+			if ( Bunch.bHasPackageMapExports )
 			{
 				Driver->NetGUIDInBytes += (BunchDataBits + (HeaderPos - IncomingStartPos)) >> 3;
 			}
@@ -1305,7 +1305,7 @@ int32 UNetConnection::SendRawBunch( FOutBunch& Bunch, bool InAllowMerge )
 	}
 	Header.WriteBit( Bunch.bReliable );
 	Header.WriteIntWrapped(Bunch.ChIndex, MAX_CHANNELS);
-	Header.WriteBit( Bunch.bHasGUIDs );
+	Header.WriteBit( Bunch.bHasPackageMapExports );
 	Header.WriteBit( Bunch.bHasMustBeMappedGUIDs );
 	Header.WriteBit( Bunch.bPartial );
 
@@ -1349,12 +1349,12 @@ int32 UNetConnection::SendRawBunch( FOutBunch& Bunch, bool InAllowMerge )
 
 	UE_LOG(LogNetTraffic, Verbose, TEXT("UNetConnection::SendRawBunch. ChIndex: %d. Bits: %d. PacketId: %d"), Bunch.ChIndex, Bunch.GetNumBits(), Bunch.PacketId );
 
-	if ( PackageMap && Bunch.bHasGUIDs )
+	if ( PackageMap && Bunch.bHasPackageMapExports )
 	{
-		PackageMap->NotifyBunchCommit( Bunch.PacketId, Bunch.ExportNetGUIDs );
+		PackageMap->NotifyBunchCommit( Bunch.PacketId, &Bunch );
 	}
 
-	if (Bunch.bHasGUIDs)
+	if ( Bunch.bHasPackageMapExports )
 	{
 		Driver->NetGUIDOutBytes += (Header.GetNumBits() + Bunch.GetNumBits()) >> 3;
 	}
diff --git a/Engine/Source/Runtime/Engine/Private/NetworkDriver.cpp b/Engine/Source/Runtime/Engine/Private/NetworkDriver.cpp
--- a/Engine/Source/Runtime/Engine/Private/NetworkDriver.cpp
+++ b/Engine/Source/Runtime/Engine/Private/NetworkDriver.cpp
@@ -1225,27 +1225,29 @@ void UNetDriver::InternalProcessRemoteFunction
 		Bunch.bReliable = 1;
 	}
 
-	// Queue unreliable multicast 
-	bool QueueBunch = (!Bunch.bReliable && Function->FunctionFlags & FUNC_NetMulticast);
-
-	if (!QueueBunch)
-	{
-		Ch->BeginContentBlock(TargetObj, Bunch);
-	}
-	
-	//UE_LOG(LogScript, Log, TEXT("   Call %s"),Function->GetFullName());
-	if ( Connection->InternalAck )
-	{
-		uint32 Checksum = FieldCache->FieldChecksum;
-		Bunch << Checksum;
-	}
-	else
+	// verify we haven't overflowed unacked bunch buffer (Connection is not net ready)
+	//@warning: needs to be after parameter evaluation for script stack integrity
+	if ( Bunch.IsError() )
 	{
-		check( FieldCache->FieldNetIndex <= ClassCache->GetMaxIndex() );
-		Bunch.WriteIntWrapped(FieldCache->FieldNetIndex, ClassCache->GetMaxIndex()+1);
-	}
+		if ( !Bunch.bReliable )
+		{
+			// Not reliable, so not fatal. This can happen a lot in debug builds at startup if client is slow to get in game
+			UE_LOG( LogNet, Warning, TEXT( "Can't send function '%s' on '%s': Reliable buffer overflow. FieldCache->FieldNetIndex: %d Max %d. Ch MaxPacket: %d" ), *Function->GetName(), *Actor->GetName(), FieldCache->FieldNetIndex, ClassCache->GetMaxIndex(), Ch->Connection->MaxPacket );
+		}
+		else
+		{
+			// The connection has overflowed the reliable buffer. We cannot recover from this. Disconnect this user.
+			UE_LOG( LogNet, Warning, TEXT( "Closing connection. Can't send function '%s' on '%s': Reliable buffer overflow. FieldCache->FieldNetIndex: %d Max %d. Ch MaxPacket: %d." ), *Function->GetName(), *Actor->GetName(), FieldCache->FieldNetIndex, ClassCache->GetMaxIndex(), Ch->Connection->MaxPacket );
+
+			FString ErrorMsg = NSLOCTEXT( "NetworkErrors", "ClientReliableBufferOverflow", "Outgoing reliable buffer overflow" ).ToString();
+			FNetControlMessage<NMT_Failure>::Send( Connection, ErrorMsg );
+			Connection->FlushNet( true );
+			Connection->Close();
 
-	const int HeaderBits = Bunch.GetNumBits();
+			PerfCountersIncrement( TEXT( "ClosedConnectionsDueToReliableBufferOverflow" ) );
+		}
+		return;
+	}
 
 #if !(UE_BUILD_SHIPPING || UE_BUILD_TEST)
 	extern TAutoConsoleVariable< int32 > CVarNetReliableDebug;
@@ -1295,35 +1297,30 @@ void UNetDriver::InternalProcessRemoteFunction
 		}
 	}
 
-	// verify we haven't overflowed unacked bunch buffer (Connection is not net ready)
-	//@warning: needs to be after parameter evaluation for script stack integrity
-	if (Bunch.IsError())
-	{
-		if (!Bunch.bReliable)
-		{
-			// Not reliable, so not fatal. This can happen a lot in debug builds at startup if client is slow to get in game
-			UE_LOG(LogNet, Warning, TEXT("Can't send function '%s' on '%s': Reliable buffer overflow. FieldCache->FieldNetIndex: %d Max %d. Ch MaxPacket: %d"), *Function->GetName(), *Actor->GetName(), FieldCache->FieldNetIndex, ClassCache->GetMaxIndex(), Ch->Connection->MaxPacket);
-		}
-		else
-		{
-			// The connection has overflowed the reliable buffer. We cannot recover from this. Disconnect this user.
-			UE_LOG(LogNet, Warning, TEXT("Closing connection. Can't send function '%s' on '%s': Reliable buffer overflow. FieldCache->FieldNetIndex: %d Max %d. Ch MaxPacket: %d."), *Function->GetName(), *Actor->GetName(), FieldCache->FieldNetIndex, ClassCache->GetMaxIndex(), Ch->Connection->MaxPacket);
-
-			FString ErrorMsg = NSLOCTEXT("NetworkErrors", "ClientReliableBufferOverflow", "Outgoing reliable buffer overflow").ToString();
-			FNetControlMessage<NMT_Failure>::Send(Connection, ErrorMsg);
-			Connection->FlushNet(true);
-			Connection->Close();
-
-			PerfCountersIncrement(TEXT("ClosedConnectionsDueToReliableBufferOverflow"));
-		}
-		return;
-	}
+	FNetBitWriter TempWriter( Bunch.PackageMap, 0 );
 
 	// Use the replication layout to send the rpc parameter values
 	TSharedPtr<FRepLayout> RepLayout = GetFunctionRepLayout( Function );
-	RepLayout->SendPropertiesForRPC( Actor, Function, Ch, Bunch, Parms );
+	RepLayout->SendPropertiesForRPC( Actor, Function, Ch, TempWriter, Parms );
 
-	const int ParameterBits = Bunch.GetNumBits() - HeaderBits;
+	int32 HeaderBits	= 0;
+	int32 ParameterBits	= 0;
+
+	// Queue unreliable multicast 
+	bool QueueBunch = ( !Bunch.bReliable && Function->FunctionFlags & FUNC_NetMulticast );
+
+	if ( QueueBunch )
+	{
+		Ch->WriteFieldHeaderAndPayload( Bunch, ClassCache, FieldCache, TempWriter );
+		ParameterBits = Bunch.GetNumBits();
+	}
+	else
+	{
+		FNetBitWriter TempBlockWriter( Bunch.PackageMap, 0 );
+		Ch->WriteFieldHeaderAndPayload( TempBlockWriter, ClassCache, FieldCache, TempWriter );
+		ParameterBits = TempBlockWriter.GetNumBits();
+		HeaderBits = Ch->WriteContentBlockPayload( TargetObj, Bunch, false, TempBlockWriter );
+	}
 
 	// Destroy the memory used for the copied out parameters
 	for ( int32 i = 0; i < LocalOutParms.Num(); i++ )
@@ -1332,14 +1329,6 @@ void UNetDriver::InternalProcessRemoteFunction
 		LocalOutParms[i]->DestroyValue_InContainer( Parms );
 	}
 
-	// Write footer for content we just wrote
-	if ( !QueueBunch )
-	{
-		Ch->EndContentBlock( TargetObj, Bunch );
-	}
-
-	const int FooterBits = Bunch.GetNumBits() - HeaderBits - ParameterBits;
-
 	static const auto CVar = IConsoleManager::Get().FindTConsoleVariableDataInt(TEXT("net.RPC.Debug"));
 	bool LogAsWarning = (CVar && CVar->GetValueOnGameThread() == 1);
 
@@ -1356,7 +1345,7 @@ void UNetDriver::InternalProcessRemoteFunction
 	else
 	{
 		// Make sure we're tracking all the bits in the bunch
-		check(Bunch.GetNumBits() == HeaderBits + ParameterBits + FooterBits);
+		check(Bunch.GetNumBits() == HeaderBits + ParameterBits);
 
 		if (QueueBunch)
 		{
@@ -1370,7 +1359,7 @@ void UNetDriver::InternalProcessRemoteFunction
 				UE_LOG(LogNetTraffic, Log,		TEXT("      Queing unreliable multicast RPC: %s::%s [%.1f bytes]"), *Actor->GetName(), *Function->GetName(), Bunch.GetNumBits() / 8.f );
 			}
 
-			NETWORK_PROFILER(GNetworkProfiler.TrackQueuedRPC(Connection, TargetObj, Actor, Function, HeaderBits, ParameterBits, FooterBits));
+			NETWORK_PROFILER(GNetworkProfiler.TrackQueuedRPC(Connection, TargetObj, Actor, Function, HeaderBits, ParameterBits, 0));
 			Ch->QueueRemoteFunctionBunch(TargetObj, Function, Bunch);
 		}
 		else
@@ -1384,7 +1373,7 @@ void UNetDriver::InternalProcessRemoteFunction
 				UE_LOG(LogNetTraffic, Log,		TEXT("      Sent RPC: %s::%s [%.1f bytes]"), *Actor->GetName(), *Function->GetName(), Bunch.GetNumBits() / 8.f );
 			}
 
-			NETWORK_PROFILER(GNetworkProfiler.TrackSendRPC(Actor, Function, HeaderBits, ParameterBits, FooterBits, Connection));
+			NETWORK_PROFILER(GNetworkProfiler.TrackSendRPC(Actor, Function, HeaderBits, ParameterBits, 0, Connection));
 			Ch->SendBunch( &Bunch, 1 );
 		}
 	}
diff --git a/Engine/Source/Runtime/Engine/Private/PackageMapClient.cpp b/Engine/Source/Runtime/Engine/Private/PackageMapClient.cpp
--- a/Engine/Source/Runtime/Engine/Private/PackageMapClient.cpp
+++ b/Engine/Source/Runtime/Engine/Private/PackageMapClient.cpp
@@ -24,7 +24,7 @@ static const int GUID_PACKET_ACKED		= -1;
 static const int INTERNAL_LOAD_OBJECT_RECURSION_LIMIT = 16;
 
 static TAutoConsoleVariable<int32> CVarAllowAsyncLoading( TEXT( "net.AllowAsyncLoading" ), 0, TEXT( "Allow async loading" ) );
-static TAutoConsoleVariable<int32> CVarIgnorePackageMismatch( TEXT( "net.IgnorePackageMismatch" ), 0, TEXT( "Ignore when package versions are different" ) );
+static TAutoConsoleVariable<int32> CVarIgnoreNetworkCheckumMismatch( TEXT( "net.IgnoreNetworkCheckumMismatch" ), 0, TEXT( "" ) );
 
 /*-----------------------------------------------------------------------------
 	UPackageMapClient implementation.
@@ -35,13 +35,6 @@ UPackageMapClient::UPackageMapClient(const FObjectInitializer& ObjectInitializer
 {
 }
 
-static uint32 GetPackageChecksum( const UPackage* Package )
-{
-	const FGuid Guid = Package->GetGuid();
-
-	return FCrc::MemCrc32( &Guid, sizeof( FGuid ) );
-}
-
 /**
  *	This is the meat of the PackageMap class which serializes a reference to Object.
  */
@@ -477,13 +470,15 @@ void UPackageMapClient::InternalWriteObject( FArchive & Ar, FNetworkGUID NetGUID
 	//   note: Default NetGUID is implied to always send path
 	FExportFlags ExportFlags;
 
-	ExportFlags.bHasNetworkChecksum = GuidCache->ShouldUseNetworkChecksum() ? 1 : 0;
+	ExportFlags.bHasNetworkChecksum = ( GuidCache->NetworkChecksumMode != FNetGUIDCache::NETCHECKSUM_None ) ? 1 : 0;
 
 	if ( NetGUID.IsDefault() )
 	{
 		// Only the client sends default guids
 		check( !IsNetGUIDAuthority() );
 		ExportFlags.bHasPath = 1;
+
+		Ar << ExportFlags.Value;
 	}
 	else if ( GuidCache->IsExportingNetGUIDBunch )
 	{
@@ -536,7 +531,6 @@ void UPackageMapClient::InternalWriteObject( FArchive & Ar, FNetworkGUID NetGUID
 		Ar << ObjectPathName;
 
 		uint32 NetworkChecksum = 0;
-		uint32 PackageChecksum = 0;
 
 		if ( ExportFlags.bHasNetworkChecksum )
 		{
@@ -544,13 +538,6 @@ void UPackageMapClient::InternalWriteObject( FArchive & Ar, FNetworkGUID NetGUID
 			Ar << NetworkChecksum;
 		}
 
-		if ( bIsPackage )
-		{
-			PackageChecksum = GetPackageChecksum( CastChecked< const UPackage >( Object ) );
-			Ar << PackageChecksum;
-			UE_LOG( LogNetPackageMap, VeryVerbose, TEXT( "InternalWriteObject: Package: %s, GUID: %u" ), *ObjectPathName, PackageChecksum );
-		}
-
 		FNetGuidCacheObject* CacheObject = GuidCache->ObjectLookup.Find( NetGUID );
 
 		if ( CacheObject != NULL )
@@ -560,7 +547,6 @@ void UPackageMapClient::InternalWriteObject( FArchive & Ar, FNetworkGUID NetGUID
 			CacheObject->bNoLoad			= ExportFlags.bNoLoad;
 			CacheObject->bIgnoreWhenMissing = ExportFlags.bNoLoad;
 			CacheObject->NetworkChecksum	= NetworkChecksum;
-			CacheObject->PackageChecksum	= PackageChecksum;
 		}
 
 		if ( GuidCache->IsExportingNetGUIDBunch )
@@ -584,7 +570,6 @@ static void SanityCheckExport(
 	const UObject *			Object, 
 	const FNetworkGUID &	NetGUID, 
 	const FString &			ExpectedPathName, 
-	const uint32			ExpectedPackageChecksum,
 	const UObject *			ExpectedOuter, 
 	const FNetworkGUID &	ExpectedOuterGUID,
 	const FExportFlags &	ExportFlags )
@@ -601,11 +586,6 @@ static void SanityCheckExport(
 		{
 			UE_LOG( LogNetPackageMap, Warning, TEXT( "SanityCheckExport: CacheObject->OuterGUID != ExpectedOuterGUID. NetGUID: %s, Object: %s, Expected: %s" ), *NetGUID.ToString(), *Object->GetPathName(), *ExpectedPathName );
 		}
-
-		if ( CacheObject->PackageChecksum != ExpectedPackageChecksum )
-		{
-			UE_LOG( LogNetPackageMap, Warning, TEXT( "SanityCheckExport: CacheObject->PackageChecksum != ExpectedPackageChecksum. NetGUID: %s, Object: %s, Expected: %s" ), *NetGUID.ToString(), *Object->GetPathName(), *ExpectedPathName );
-		}
 	}
 	else
 	{
@@ -678,11 +658,7 @@ FNetworkGUID UPackageMapClient::InternalLoadObject( FArchive & Ar, UObject *& Ob
 	// ----------------	
 	FExportFlags ExportFlags;
 
-	if ( NetGUID.IsDefault() )
-	{
-		ExportFlags.bHasPath = 1;
-	}
-	else if ( GuidCache->IsExportingNetGUIDBunch )
+	if ( NetGUID.IsDefault() || GuidCache->IsExportingNetGUIDBunch )
 	{
 		Ar << ExportFlags.Value;
 
@@ -701,7 +677,6 @@ FNetworkGUID UPackageMapClient::InternalLoadObject( FArchive & Ar, UObject *& Ob
 
 		FString PathName;
 		uint32	NetworkChecksum = 0;
-		uint32	PackageChecksum = 0;
 
 		Ar << PathName;
 
@@ -712,12 +687,6 @@ FNetworkGUID UPackageMapClient::InternalLoadObject( FArchive & Ar, UObject *& Ob
 
 		const bool bIsPackage = NetGUID.IsStatic() && !OuterGUID.IsValid();
 
-		if ( bIsPackage )
-		{
-			Ar << PackageChecksum;
-			UE_LOG( LogNetPackageMap, VeryVerbose, TEXT( "InternalLoadObject: Package: %s, GUID: %u" ), *PathName, PackageChecksum );
-		}
-
 		if ( Ar.IsError() )
 		{
 			UE_LOG( LogNetPackageMap, Error, TEXT( "InternalLoadObject: Failed to load path name" ) );
@@ -731,7 +700,7 @@ FNetworkGUID UPackageMapClient::InternalLoadObject( FArchive & Ar, UObject *& Ob
 		if ( Object != NULL )
 		{
 			// If we already have the object, just do some sanity checking and return
-			SanityCheckExport( GuidCache.Get(), Object, NetGUID, PathName, PackageChecksum, ObjOuter, OuterGUID, ExportFlags );
+			SanityCheckExport( GuidCache.Get(), Object, NetGUID, PathName, ObjOuter, OuterGUID, ExportFlags );
 			return NetGUID;
 		}
 
@@ -763,20 +732,25 @@ FNetworkGUID UPackageMapClient::InternalLoadObject( FArchive & Ar, UObject *& Ob
 				return NetGUID;
 			}
 
-			if ( bIsPackage )
+			if ( NetworkChecksum != 0 && GuidCache->NetworkChecksumMode == FNetGUIDCache::NETCHECKSUM_SaveAndUse && !CVarIgnoreNetworkCheckumMismatch.GetValueOnGameThread() )
 			{
-				UPackage * Package = Cast< UPackage >( Object );
+				const uint32 CompareNetworkChecksum = GuidCache->GetNetworkChecksum( Object );
 
-				if ( Package == NULL )
+				if ( CompareNetworkChecksum != NetworkChecksum )
 				{
-					UE_LOG( LogNetPackageMap, Error, TEXT( "UPackageMapClient::InternalLoadObject: Default object not a package from client: PathName: %s, ObjOuter: %s " ), *PathName, ObjOuter != NULL ? *ObjOuter->GetPathName() : TEXT( "NULL" ) );
+					UE_LOG( LogNetPackageMap, Error, TEXT( "UPackageMapClient::InternalLoadObject: Default object package network checksum mismatch! PathName: %s, ObjOuter: %s, GUID1: %u, GUID2: %u " ), *PathName, ObjOuter != NULL ? *ObjOuter->GetPathName() : TEXT( "NULL" ), CompareNetworkChecksum, NetworkChecksum );
 					Object = NULL;
 					return NetGUID;
 				}
+			}
+
+			if ( bIsPackage )
+			{
+				UPackage * Package = Cast< UPackage >( Object );
 
-				if ( !GuidCache->ShouldIgnorePackageMismatch() && GetPackageChecksum( Package ) != PackageChecksum )
+				if ( Package == NULL )
 				{
-					UE_LOG( LogNetPackageMap, Error, TEXT( "UPackageMapClient::InternalLoadObject: Default object package guid mismatch! PathName: %s, ObjOuter: %s, GUID1: %u, GUID2: %u " ), *PathName, ObjOuter != NULL ? *ObjOuter->GetPathName() : TEXT( "NULL" ), GetPackageChecksum( Package ), PackageChecksum );
+					UE_LOG( LogNetPackageMap, Error, TEXT( "UPackageMapClient::InternalLoadObject: Default object not a package from client: PathName: %s, ObjOuter: %s " ), *PathName, ObjOuter != NULL ? *ObjOuter->GetPathName() : TEXT( "NULL" ) );
 					Object = NULL;
 					return NetGUID;
 				}
@@ -805,7 +779,7 @@ FNetworkGUID UPackageMapClient::InternalLoadObject( FArchive & Ar, UObject *& Ob
 		const bool bIgnoreWhenMissing = ExportFlags.bNoLoad;
 
 		// Register this path and outer guid combo with the net guid
-		GuidCache->RegisterNetGUIDFromPath_Client( NetGUID, PathName, OuterGUID, NetworkChecksum, PackageChecksum, ExportFlags.bNoLoad, bIgnoreWhenMissing );
+		GuidCache->RegisterNetGUIDFromPath_Client( NetGUID, PathName, OuterGUID, NetworkChecksum, ExportFlags.bNoLoad, bIgnoreWhenMissing );
 
 		// Try again now that we've registered the path
 		Object = GuidCache->GetObjectFromNetGUID( NetGUID, GuidCache->IsExportingNetGUIDBunch );
@@ -858,10 +832,11 @@ bool UPackageMapClient::ExportNetGUID( FNetworkGUID NetGUID, const UObject* Obje
 
 			CurrentExportBunch = new FOutBunch(this, Connection->GetMaxSingleBunchSizeBits());
 			CurrentExportBunch->SetAllowResize(false);
-			CurrentExportBunch->bHasGUIDs = true;
+			CurrentExportBunch->bHasPackageMapExports = true;
 #if !(UE_BUILD_SHIPPING || UE_BUILD_TEST)
 			CurrentExportBunch->DebugString = TEXT("NetGUIDs");
 #endif
+			CurrentExportBunch->WriteBit( 0 );		// To signify this is NOT a rep layout export
 
 			ExportNetGUIDCount = 0;
 			*CurrentExportBunch << ExportNetGUIDCount;
@@ -930,20 +905,25 @@ bool UPackageMapClient::ExportNetGUID( FNetworkGUID NetGUID, const UObject* Obje
 	return false;
 }
 
+static void PatchHeaderCount( FBitWriter& Writer, bool bHasRepLayoutExport, uint32 NewCount )
+{
+	FBitWriterMark Reset;
+	FBitWriterMark Restore( Writer );
+	Reset.PopWithoutClear( Writer );
+	Writer.WriteBit( bHasRepLayoutExport ? 1 : 0 );
+	Writer << NewCount;
+	Restore.PopWithoutClear( Writer );
+}
+
 /** Called when an export bunch is finished. It writes how many NetGUIDs are contained in the bunch and finalizes the book keeping so we know what NetGUIDs are in the bunch */
 void UPackageMapClient::ExportNetGUIDHeader()
 {
 	check(CurrentExportBunch);
-	FOutBunch& Out = *CurrentExportBunch;
 
 	UE_LOG(LogNetPackageMap, Log, TEXT("	UPackageMapClient::ExportNetGUID. Bytes: %d Bits: %d ExportNetGUIDCount: %d"), CurrentExportBunch->GetNumBytes(), CurrentExportBunch->GetNumBits(), ExportNetGUIDCount);
 
 	// Rewrite how many NetGUIDs were exported.
-	FBitWriterMark Reset;
-	FBitWriterMark Restore(Out);
-	Reset.PopWithoutClear(Out);
-	Out << ExportNetGUIDCount;
-	Restore.PopWithoutClear(Out);
+	PatchHeaderCount( *CurrentExportBunch, false, ExportNetGUIDCount );
 
 	// If we've written new NetGUIDs to the 'bunch' set (current+1)
 	if (UE_LOG_ACTIVE(LogNetPackageMap,Verbose))
@@ -971,8 +951,17 @@ void UPackageMapClient::ExportNetGUIDHeader()
 
 void UPackageMapClient::ReceiveNetGUIDBunch( FInBunch &InBunch )
 {
+	check( InBunch.bHasPackageMapExports );
+
+	const bool bHasRepLayoutExport = InBunch.ReadBit() == 1 ? true : false;
+
+	if ( bHasRepLayoutExport )
+	{
+		ReceiveCompatibleRepLayoutExports( InBunch );
+		return;
+	}
+
 	GuidCache->IsExportingNetGUIDBunch = true;
-	check(InBunch.bHasGUIDs);
 
 	int32 NumGUIDsInBunch = 0;
 	InBunch << NumGUIDsInBunch;
@@ -1010,8 +999,272 @@ void UPackageMapClient::ReceiveNetGUIDBunch( FInBunch &InBunch )
 	GuidCache->IsExportingNetGUIDBunch = false;
 }
 
-bool UPackageMapClient::AppendExportBunches(TArray<FOutBunch *>& OutgoingBunches)
+TSharedPtr< FCompatibleRepLayout > UPackageMapClient::TrackCompatibleRepLayout( const FRepLayout* RepLayout )
+{
+	TSharedPtr< FCompatibleRepLayout > CompatibleRepLayout = GuidCache->CompatibleRepLayoutMap.FindRef( RepLayout->Owner->GetPathName() );
+
+	if ( !CompatibleRepLayout.IsValid() )
+	{
+		CompatibleRepLayout = TSharedPtr< FCompatibleRepLayout >( new FCompatibleRepLayout() );
+		CompatibleRepLayout->PathName = RepLayout->Owner->GetPathName();
+
+		CompatibleRepLayout->PathNameIndex = ++GuidCache->UniqueCompatibleRepLayoutPathIndex;
+
+		check( !GuidCache->CompatibleRepLayoutPathToIndex.Contains( CompatibleRepLayout->PathName ) );
+		check( !GuidCache->CompatibleRepLayoutIndexToPath.Contains( CompatibleRepLayout->PathNameIndex ) );
+
+		GuidCache->CompatibleRepLayoutPathToIndex.Add( CompatibleRepLayout->PathName, CompatibleRepLayout->PathNameIndex );
+		GuidCache->CompatibleRepLayoutIndexToPath.Add( CompatibleRepLayout->PathNameIndex, CompatibleRepLayout->PathName );
+
+		CompatibleRepLayout->Cmds.SetNum( RepLayout->Cmds.Num() );
+
+		for ( int32 i = 0; i < RepLayout->Cmds.Num(); i++ )
+		{
+			FCompatibleRepLayoutCmd CompatibleCmd(
+				i,
+				RepLayout->Cmds[i].CompatibleChecksum,
+				RepLayout->Cmds[i].Property ? RepLayout->Cmds[i].Property->GetName() : TEXT( "" ),
+				RepLayout->Cmds[i].Property ? RepLayout->Cmds[i].Property->GetCPPType( nullptr, 0 ) : TEXT( "" ) );
+
+			CompatibleRepLayout->Cmds[i] = CompatibleCmd;
+		}
+
+		GuidCache->CompatibleRepLayoutMap.Add( CompatibleRepLayout->PathName, CompatibleRepLayout );
+	}
+
+	return CompatibleRepLayout;
+}
+
+void UPackageMapClient::TrackCompatibleRepLayoutCmd( const FRepLayout* RepLayout, FCompatibleRepLayout* CompatibleRepLayout, const int32 CmdIndex )
+{
+	CompatibleRepLayout->Cmds[CmdIndex].bExported = true;
+
+	const uint64 CmdHandle = ( ( uint64 )CompatibleRepLayout->PathNameIndex ) << 32 | ( uint64 )CmdIndex;
+
+	// If this cmd hasn't been confirmed as exported, we need to export it for this bunch
+	if ( !CompatibleRepLayoutCmdExported.Contains( CmdHandle ) )
+	{
+		CompatibleRepLayoutCmdExports.Add( CmdHandle );		// NOTE - This is a set, so it will only add once
+	}
+}
+
+TSharedPtr< FCompatibleRepLayout > UPackageMapClient::GetCompatibleRepLayoutChecked( UObject* Owner ) const
+{
+	return GuidCache->CompatibleRepLayoutMap.FindChecked( Owner->GetPathName() );
+}
+
+void UPackageMapClient::SerializeCompatibleReplayoutMap( FArchive& Ar )
+{
+	if ( Ar.IsSaving() )
+	{
+		// Save the number of layouts
+		uint32 NumLayouts = GuidCache->CompatibleRepLayoutMap.Num();
+		Ar << NumLayouts;
+
+		// Save each layout
+		for ( auto It = GuidCache->CompatibleRepLayoutMap.CreateIterator(); It; ++It )
+		{
+			// Save out the compatible replayout
+			Ar << *It.Value().Get();
+		}
+	}
+	else
+	{
+		// Clear all of our mappings, since we're starting over
+		GuidCache->CompatibleRepLayoutMap.Empty();
+		GuidCache->CompatibleRepLayoutPathToIndex.Empty();
+		GuidCache->CompatibleRepLayoutIndexToPath.Empty();
+
+		// Read the number of layouts
+		uint32 NumLayouts = 0;
+		Ar << NumLayouts;
+
+		// Read each compatible layout
+		for ( int32 i = 0; i < ( int32 )NumLayouts; i++ )
+		{
+			TSharedPtr< FCompatibleRepLayout > CompatibleRepLayout = TSharedPtr< FCompatibleRepLayout >( new FCompatibleRepLayout() );
+
+			// Read in the compatible layout
+			Ar << *CompatibleRepLayout.Get();
+
+			// Assign index to path name
+			GuidCache->CompatibleRepLayoutPathToIndex.Add( CompatibleRepLayout->PathName, CompatibleRepLayout->PathNameIndex );
+			GuidCache->CompatibleRepLayoutIndexToPath.Add( CompatibleRepLayout->PathNameIndex, CompatibleRepLayout->PathName );
+
+			// Add the compatible layout to the map
+			GuidCache->CompatibleRepLayoutMap.Add( CompatibleRepLayout->PathName, CompatibleRepLayout );
+		}
+	}
+}
+
+void UPackageMapClient::AppendCompatibleRepLayoutExports( TArray<FOutBunch *>& OutgoingBunches )
+{
+	if ( CompatibleRepLayoutCmdExports.Num() == 0 )
+	{
+		return;	// Nothing to do
+	}
+
+	FOutBunch* ExportBunch = nullptr;
+	TSet< uint32 > ExportedPathInThisBunchAlready;
+
+	uint32 CurrentRepLayoutCmdCount = 0;
+
+	// Go through each layout, and try to export to single bunch, using a new bunch each time we fragment (go past max bunch size)
+	for ( const uint64 CmdExport : CompatibleRepLayoutCmdExports )
+	{
+		// Parse the path name index and cmd index out of the uint64
+		uint32 PathNameIndex	= CmdExport >> 32;
+		uint32 CmdIndex			= CmdExport & ( ( (uint64)1 << 32 ) - 1 );
+
+		check( PathNameIndex != 0 );
+
+		FString PathName = GuidCache->CompatibleRepLayoutIndexToPath.FindChecked( PathNameIndex );
+		TSharedPtr< FCompatibleRepLayout > CompatibleRepLayout = GuidCache->CompatibleRepLayoutMap.FindChecked( PathName );
+
+		for ( int32 NumTries = 0; NumTries < 2; NumTries++ )
+		{
+			if ( ExportBunch == nullptr )
+			{
+				ExportBunch = new FOutBunch( this, Connection->GetMaxSingleBunchSizeBits() );
+				ExportBunch->SetAllowResize( false );
+				ExportBunch->bHasPackageMapExports = true;
+
+#if !(UE_BUILD_SHIPPING || UE_BUILD_TEST)
+				ExportBunch->DebugString = TEXT( "CompatibleRepLayoutCmdExports" );
+#endif
+
+				ExportBunch->WriteBit( 1 );		// To signify this is a rep layout export
+
+				// Write stub cmd amount, we'll replace it with the final number when this bunch fills up (or we're done)
+				uint32 NumFakeRepLayoutCmds = 0;
+				*ExportBunch << NumFakeRepLayoutCmds;
+			}
+
+			// Save our spot so we can undo if we don't have enough room
+			FBitWriterMark LastExportMark;
+			LastExportMark.Init( *ExportBunch );
+
+			// Write path index
+			ExportBunch->SerializeIntPacked( PathNameIndex );
+
+			// Export the path if we need to
+			if ( !CompatibleRepLayoutPathExported.Contains( PathNameIndex ) && !ExportedPathInThisBunchAlready.Contains( PathNameIndex ) )
+			{
+				ExportBunch->WriteBit( 1 );
+				*ExportBunch << PathName;
+			}
+			else
+			{
+				ExportBunch->WriteBit( 0 );
+			}
+
+			check( CmdIndex == CompatibleRepLayout->Cmds[CmdIndex].Handle );
+
+			*ExportBunch << CompatibleRepLayout->Cmds[CmdIndex];
+
+			if ( !ExportBunch->IsError() )
+			{
+				// We had enough room, continue on to the next one
+				ExportBunch->ExportRepLayoutCmds.Add( CmdExport );		// Add this cmd to this bunch so we know to handle it during NotifyBunchCommit
+				ExportedPathInThisBunchAlready.Add( PathNameIndex );
+				CurrentRepLayoutCmdCount++;
+				break;
+			}
+
+			//
+			// If we get here, we overflowed, wrap up the currently pending bunch, and start a new one
+			//
+
+			if ( CurrentRepLayoutCmdCount == 0 || NumTries == 1 )
+			{
+				// This means we couldn't serialize a single compatible rep layout cmd into a single bunch. This should never happen unless a single cmd takes way too much space
+				UE_LOG( LogNetPackageMap, Fatal, TEXT( "AppendExportBunches: Failed to serialize CompatibleRepLayout into single bunch: %s, %i" ), *CompatibleRepLayout->PathName, CmdIndex );
+				return;
+			}
+
+			LastExportMark.Pop( *ExportBunch );
+
+			PatchHeaderCount( *ExportBunch, true, CurrentRepLayoutCmdCount );
+
+			OutgoingBunches.Add( ExportBunch );
+
+			// Reset bunch
+			ExportBunch					= nullptr;
+			CurrentRepLayoutCmdCount	= 0;
+
+			ExportedPathInThisBunchAlready.Empty();
+		}
+	}
+
+	// Wrap up the last bunch if needed
+	if ( CurrentRepLayoutCmdCount > 0 )
+	{
+		PatchHeaderCount( *ExportBunch, true, CurrentRepLayoutCmdCount );
+		OutgoingBunches.Add( ExportBunch );
+	}
+
+	CompatibleRepLayoutCmdExports.Empty();
+}
+
+void UPackageMapClient::ReceiveCompatibleRepLayoutExports( FInBunch &InBunch )
 {
+	// Read number of rep layout cmds
+	uint32 NumLayoutCmdExports = 0;
+	InBunch << NumLayoutCmdExports;
+
+	for ( int32 i = 0; i < ( int32 )NumLayoutCmdExports; i++ )
+	{
+		// Read the index that represents the name in the CompatibleRepLayoutIndexToPath map
+		uint32 PathNameIndex;
+		InBunch.SerializeIntPacked( PathNameIndex );
+
+		// See if the path name was exported (we'll expect it if we haven't seen this index before)
+		if ( InBunch.ReadBit() == 1 )
+		{
+			FString PathName;
+			InBunch << PathName;
+
+			GuidCache->CompatibleRepLayoutPathToIndex.Add( PathName, PathNameIndex );
+			GuidCache->CompatibleRepLayoutIndexToPath.Add( PathNameIndex, PathName );
+		}
+
+		// At this point, we expect to be able to find the entry in CompatibleRepLayoutIndexToPath
+		const FString PathName = GuidCache->CompatibleRepLayoutIndexToPath.FindChecked( PathNameIndex );
+
+		TSharedPtr< FCompatibleRepLayout > CompatibleRepLayout = GuidCache->CompatibleRepLayoutMap.FindRef( PathName );
+
+		if ( !CompatibleRepLayout.IsValid() )
+		{
+			CompatibleRepLayout = TSharedPtr< FCompatibleRepLayout >( new FCompatibleRepLayout() );
+			CompatibleRepLayout->PathName = PathName;
+			CompatibleRepLayout->PathNameIndex = PathNameIndex;
+			GuidCache->CompatibleRepLayoutMap.Add( CompatibleRepLayout->PathName, CompatibleRepLayout );
+		}
+
+		FCompatibleRepLayoutCmd Cmd;
+
+		// Read the cmd
+		InBunch << Cmd;
+
+		// Make room for the cmd
+		if ( ( int32 )Cmd.Handle >= CompatibleRepLayout->Cmds.Num() )
+		{
+			CompatibleRepLayout->Cmds.AddDefaulted( Cmd.Handle - CompatibleRepLayout->Cmds.Num() + 1 );
+		}
+
+		// Assign it to the correct slot (Cmd.Handle is just the index into the array)
+		CompatibleRepLayout->Cmds[Cmd.Handle] = Cmd;
+	}
+}
+
+void UPackageMapClient::AppendExportBunches(TArray<FOutBunch *>& OutgoingBunches)
+{
+	// If we have rep layouts to export, handle those now
+	if ( CompatibleRepLayoutCmdExports.Num() > 0 )
+	{
+		AppendCompatibleRepLayoutExports( OutgoingBunches );
+	}
+
 	// Finish current in progress bunch if necessary
 	if (ExportNetGUIDCount > 0)
 	{
@@ -1041,10 +1294,14 @@ bool UPackageMapClient::AppendExportBunches(TArray<FOutBunch *>& OutgoingBunches
 
 		OutgoingBunches.Append(ExportBunches);
 		ExportBunches.Empty();
-		return true;
 	}
+}
 
-	return false;
+void UPackageMapClient::SyncPackageMapExportAckStatus( const UPackageMapClient* Source )
+{
+	NetGUIDAckStatus				= Source->NetGUIDAckStatus;
+	CompatibleRepLayoutCmdExported	= Source->CompatibleRepLayoutCmdExported;
+	CompatibleRepLayoutPathExported	= Source->CompatibleRepLayoutPathExported;
 }
 
 //--------------------------------------------------------------------
@@ -1057,10 +1314,25 @@ bool UPackageMapClient::AppendExportBunches(TArray<FOutBunch *>& OutgoingBunches
  *	Called when a bunch is committed to the connection's Out buffer.
  *	ExportNetGUIDs is the list of GUIDs stored on the bunch that we use to update the expected sequence for those exported GUIDs
  */
-void UPackageMapClient::NotifyBunchCommit( const int32 OutPacketId, const TArray< FNetworkGUID > & ExportNetGUIDs )
+void UPackageMapClient::NotifyBunchCommit( const int32 OutPacketId, const FOutBunch* OutBunch )
 {
+	// Mark all of the rep layout cmds in this bunch as exported
+	// NOTE - This only currently works with reliable connections (i.e. InternalAck)
+	// For this to work with normal connections, we'll need to do real ack logic here
+	for ( int32 i = 0; i < OutBunch->ExportRepLayoutCmds.Num(); i++ )
+	{
+		CompatibleRepLayoutPathExported.Add( OutBunch->ExportRepLayoutCmds[i] >> 32, true );
+		CompatibleRepLayoutCmdExported.Add( OutBunch->ExportRepLayoutCmds[i], true );
+	}
+
+	const TArray< FNetworkGUID >& ExportNetGUIDs = OutBunch->ExportNetGUIDs;
+
+	if ( ExportNetGUIDs.Num() == 0 )
+	{
+		return;		// Nothing to do
+	}
+
 	check( OutPacketId > GUID_PACKET_ACKED );	// Assumptions break if this isn't true ( We assume ( OutPacketId > GUID_PACKET_ACKED ) == PENDING )
-	check( ExportNetGUIDs.Num() > 0 );			// We should have never created this bunch if there are not exports
 
 	for ( int32 i = 0; i < ExportNetGUIDs.Num(); i++ )
 	{
@@ -1393,9 +1665,10 @@ FNetworkGUID UPackageMapClient::GetNetGUIDFromObject(const UObject* InObject) co
 //	FNetGUIDCache
 //----------------------------------------------------------------------------------------
 
-FNetGUIDCache::FNetGUIDCache( UNetDriver * InDriver ) : IsExportingNetGUIDBunch( false ), Driver( InDriver ), bUseNetworkChecksum( false ), bIgnorePackageMismatchOverride( false )
+FNetGUIDCache::FNetGUIDCache( UNetDriver * InDriver ) : IsExportingNetGUIDBunch( false ), Driver( InDriver ), NetworkChecksumMode( NETCHECKSUM_SaveAndUse )
 {
 	UniqueNetIDs[0] = UniqueNetIDs[1] = 0;
+	UniqueCompatibleRepLayoutPathIndex = 0;
 }
 
 class FArchiveCountMemGUID : public FArchive
@@ -1715,7 +1988,7 @@ void FNetGUIDCache::RegisterNetGUID_Client( const FNetworkGUID& NetGUID, const U
  *	Associates a net guid with a path, that can be loaded or found later
  *  This function is only called on the client
  */
-void FNetGUIDCache::RegisterNetGUIDFromPath_Client( const FNetworkGUID& NetGUID, const FString& PathName, const FNetworkGUID& OuterGUID, const uint32 NetworkChecksum, const uint32 PackageChecksum, const bool bNoLoad, const bool bIgnoreWhenMissing )
+void FNetGUIDCache::RegisterNetGUIDFromPath_Client( const FNetworkGUID& NetGUID, const FString& PathName, const FNetworkGUID& OuterGUID, const uint32 NetworkChecksum, const bool bNoLoad, const bool bIgnoreWhenMissing )
 {
 	check( !IsNetGUIDAuthority() );		// Server never calls this locally
 	check( !NetGUID.IsDefault() );
@@ -1756,7 +2029,6 @@ void FNetGUIDCache::RegisterNetGUIDFromPath_Client( const FNetworkGUID& NetGUID,
 	CacheObject.PathName			= FName( *PathName );
 	CacheObject.OuterGUID			= OuterGUID;
 	CacheObject.NetworkChecksum		= NetworkChecksum;
-	CacheObject.PackageChecksum		= PackageChecksum;
 	CacheObject.bNoLoad				= bNoLoad;
 	CacheObject.bIgnoreWhenMissing	= bIgnoreWhenMissing;
 
@@ -1797,11 +2069,6 @@ void FNetGUIDCache::AsyncPackageCallback(const FName& PackageName, UPackage * Pa
 		CacheObject->bIsBroken = true;
 		UE_LOG( LogNetPackageMap, Error, TEXT( "AsyncPackageCallback: Package FAILED to load. Path: %s, NetGUID: %s" ), *PackageName.ToString(), *NetGUID.ToString() );
 	}
-	else if ( !ShouldIgnorePackageMismatch() && GetPackageChecksum( Package ) != CacheObject->PackageChecksum )
-	{
-		CacheObject->bIsBroken = true;
-		UE_LOG( LogNetPackageMap, Error, TEXT( "AsyncPackageCallback: Package GUID mismatch! Path: %s, NetGUID: %s, GUID1: %u, GUID2: %u" ), *PackageName.ToString(), *NetGUID.ToString(), GetPackageChecksum( Package ), CacheObject->PackageChecksum );
-	}
 }
 
 static bool ObjectLevelHasFinishedLoading( UObject* Object, UNetDriver* Driver )
@@ -2029,26 +2296,25 @@ UObject* FNetGUIDCache::GetObjectFromNetGUID( const FNetworkGUID& NetGUID, const
 				check(Package->IsFullyLoaded());
 			}
 		}
-
-		if ( !ShouldIgnorePackageMismatch() && GetPackageChecksum( Package ) != CacheObjectPtr->PackageChecksum )
-		{
-			// If the package guid doesn't match, don't allow it to load
-			CacheObjectPtr->bIsBroken = true;
-			UE_LOG( LogNetPackageMap, Error, TEXT( "GetObjectFromNetGUID: Package GUID mismatch! Path: %s, NetGUID: %s, GUID1: %u, GUID2: %u" ), *CacheObjectPtr->PathName.ToString(), *NetGUID.ToString(), GetPackageChecksum( Package ), CacheObjectPtr->PackageChecksum );
-			return NULL;
-		}
 	}
 
-	if ( CacheObjectPtr->NetworkChecksum != 0 )
+	if ( CacheObjectPtr->NetworkChecksum != 0 && !CVarIgnoreNetworkCheckumMismatch.GetValueOnGameThread() )
 	{
 		const uint32 NetworkChecksum = GetNetworkChecksum( Object );
 
 		if ( CacheObjectPtr->NetworkChecksum != NetworkChecksum )
 		{
-			UE_LOG( LogNetPackageMap, Warning, TEXT( "GetObjectFromNetGUID: Network checksum mismatch. FullNetGUIDPath: %s, %u, %u" ), *FullNetGUIDPath( NetGUID ), CacheObjectPtr->NetworkChecksum, NetworkChecksum );
+			if ( NetworkChecksumMode == NETCHECKSUM_SaveAndUse )
+			{
+				UE_LOG( LogNetPackageMap, Warning, TEXT( "GetObjectFromNetGUID: Network checksum mismatch. FullNetGUIDPath: %s, %u, %u" ), *FullNetGUIDPath( NetGUID ), CacheObjectPtr->NetworkChecksum, NetworkChecksum );
 
-			CacheObjectPtr->bIsBroken = true;
-			return NULL;
+				CacheObjectPtr->bIsBroken = true;
+				return NULL;
+			}
+			else
+			{
+				UE_LOG( LogNetPackageMap, Verbose, TEXT( "GetObjectFromNetGUID: Network checksum mismatch. FullNetGUIDPath: %s, %u, %u" ), *FullNetGUIDPath( NetGUID ), CacheObjectPtr->NetworkChecksum, NetworkChecksum );
+			}
 		}
 	}
 
@@ -2259,19 +2525,9 @@ uint32 FNetGUIDCache::GetNetworkChecksum( const UObject* Obj )
 	return ( Class != NULL ) ? GetClassNetworkChecksum( Class ) : GetClassNetworkChecksum( Obj->GetClass() );
 }
 
-bool FNetGUIDCache::ShouldIgnorePackageMismatch() const
-{
-	return bIgnorePackageMismatchOverride || CVarIgnorePackageMismatch.GetValueOnGameThread();
-};
-
-void FNetGUIDCache::SetShouldUseNetworkChecksum( const bool bInUseNetworkChecksum )
-{
-	bUseNetworkChecksum = bInUseNetworkChecksum;
-}
-
-bool FNetGUIDCache::ShouldUseNetworkChecksum() const
+void FNetGUIDCache::SetNetworkChecksumMode( const ENetworkChecksumMode NewMode )
 {
-	return bUseNetworkChecksum;
+	NetworkChecksumMode = NewMode;
 }
 
 //------------------------------------------------------
diff --git a/Engine/Source/Runtime/Engine/Private/RepLayout.cpp b/Engine/Source/Runtime/Engine/Private/RepLayout.cpp
--- a/Engine/Source/Runtime/Engine/Private/RepLayout.cpp
+++ b/Engine/Source/Runtime/Engine/Private/RepLayout.cpp
@@ -9,6 +9,7 @@
 #include "Net/DataReplication.h"
 #include "Net/NetworkProfiler.h"
 #include "Engine/ActorChannel.h"
+#include "Engine/PackageMapClient.h"
 
 static TAutoConsoleVariable<int32> CVarAllowPropertySkipping( TEXT( "net.AllowPropertySkipping" ), 1, TEXT( "Allow skipping of properties that haven't changed for other clients" ) );
 
@@ -597,9 +598,8 @@ bool FRepLayout::ReplicateProperties(
 	const uint8* RESTRICT		Data, 
 	UClass *					ObjectClass,
 	UActorChannel *				OwningChannel,
-	FOutBunch &					Writer, 
-	const FReplicationFlags &	RepFlags,
-	bool &						bContentBlockWritten ) const
+	FNetBitWriter&				Writer,
+	const FReplicationFlags &	RepFlags ) const
 {
 	SCOPE_CYCLE_COUNTER( STAT_NetReplicateDynamicPropTime );
 
@@ -757,11 +757,15 @@ bool FRepLayout::ReplicateProperties(
 		SanityCheckChangeList( Data, Changed );
 #endif
 
-		// For RepLayout properties, we hijack the first non custom property, and use that to identify these properties
-		WritePropertyHeader( (UObject*)Data, ObjectClass, OwningChannel, Parents[FirstNonCustomParent].Property, Writer, 0, bContentBlockWritten );
-
 		// Send the final merged change list
-		SendProperties( RepState, RepFlags, Data, ObjectClass, OwningChannel, Writer, Changed, bContentBlockWritten );
+		if ( OwningChannel->Connection->InternalAck )
+		{
+			SendProperties_BackwardsCompatible( RepState, Data, ObjectClass, OwningChannel, Writer, Changed );
+		}
+		else
+		{
+			SendProperties( RepState, RepFlags, Data, ObjectClass, OwningChannel, Writer, Changed );
+		}
 
 #ifdef ENABLE_SUPER_CHECKSUMS
 		Writer.WriteBit( bIsAllAcked ? 1 : 0 );
@@ -1100,83 +1104,225 @@ uint16 FRepLayout::SendProperties_r(
 	return Handle;
 }
 
-void FRepLayout::WritePropertyHeader( 
-	UObject *			Object,
-	UClass *			ObjectClass,
-	UActorChannel *		OwningChannel,
-	UProperty *			Property, 
-	FOutBunch &			Bunch, 
-	uint32				ArrayIndex, 
-	bool &				bContentBlockWritten ) const
+void FRepLayout::SendProperties( 
+	FRepState *	RESTRICT		RepState, 
+	const FReplicationFlags &	RepFlags,
+	const uint8* RESTRICT		Data, 
+	UClass *					ObjectClass,
+	UActorChannel	*			OwningChannel,
+	FNetBitWriter&				Writer, 
+	TArray< uint16 > &			Changed ) const
 {
-	UNetConnection * Connection = OwningChannel->Connection;
+#ifdef ENABLE_PROPERTY_CHECKSUMS
+	const bool bDoChecksum = CVarDoPropertyChecksum.GetValueOnGameThread() == 1;
+#else
+	const bool bDoChecksum = false;
+#endif
 
-	// Get class network info cache.
-	const FClassNetCache* ClassCache = Connection->Driver->NetCache->GetClassNetCache( ObjectClass );
+	FRepWriterState WriterState( Writer, Changed, bDoChecksum );
 
-	check( ClassCache );
+#ifdef ENABLE_PROPERTY_CHECKSUMS
+	Writer.WriteBit( bDoChecksum ? 1 : 0 );
+#endif
 
-	if ( !bContentBlockWritten )
-	{
-		OwningChannel->BeginContentBlock( Object, Bunch );
-		bContentBlockWritten = true;
-	}
+	SendProperties_r( RepState, RepFlags, WriterState, 0, Cmds.Num() - 1, RepState->StaticBuffer.GetData(), Data, 0 );
 
-	const int NumStartingBits = Bunch.GetNumBits();
+	WritePropertyHandle( Writer, 0, bDoChecksum );
+}
 
-	// Get the network friend property index to replicate
-	const FFieldNetCache * FieldCache = ClassCache->GetFromField( Property );
+static FORCEINLINE void WritePropertyHandle_BackwardsCompatible( FNetBitWriter & Writer, uint32 CompatibleCmdHandle, bool bDoChecksum )
+{
+	const int NumStartingBits = Writer.GetNumBits();
 
-	checkSlow( FieldCache );
+	Writer.SerializeIntPacked( CompatibleCmdHandle );
 
-	// Send property name and optional array index.
-	check( FieldCache->FieldNetIndex <= ClassCache->GetMaxIndex() );
+#ifdef ENABLE_PROPERTY_CHECKSUMS
+	if ( bDoChecksum )
+	{
+		SerializeGenericChecksum( Writer );
+	}
+#endif
 
-	if ( Connection->InternalAck )
+	NETWORK_PROFILER( GNetworkProfiler.TrackWritePropertyHandle( Writer.GetNumBits() - NumStartingBits, nullptr ) );
+}
+
+static bool ShouldSendProperty_BackwardsCompatible( FRepWriterState & WriterState, FNetBitWriter & Writer, const uint16 Handle, const uint32 CompatibleCmdHandle )
+{
+	if ( Handle == WriterState.Changed[WriterState.CurrentChanged] )
 	{
-		uint32 Checksum = FieldCache->FieldChecksum;
-		Bunch << Checksum;
+		// Write out the handle
+		WritePropertyHandle_BackwardsCompatible( Writer, CompatibleCmdHandle, WriterState.bDoChecksum );
+
+		// Advance to the next expected handle
+		WriterState.CurrentChanged++;
+
+		return true;
 	}
-	else
+
+	return false;
+}
+
+void FRepLayout::SendProperties_BackwardsCompatible_DynamicArray_r(
+	FRepState *	RESTRICT		RepState,
+	FRepWriterState &			WriterState,
+	FNetBitWriter &				Writer,
+	UPackageMapClient*			PackageMapClient,
+	FCompatibleRepLayout*		CompatibleRepLayout,
+	const int32					CmdIndex,
+	const uint8* RESTRICT		StoredData,
+	const uint8* RESTRICT		Data ) const
+{
+	const FRepLayoutCmd& Cmd = Cmds[CmdIndex];
+
+	FScriptArray * Array = ( FScriptArray * )Data;
+	FScriptArray * StoredArray = ( FScriptArray * )StoredData;
+
+	// Write array num
+	uint32 ArrayNum = Array->Num();
+	Writer.SerializeIntPacked( ArrayNum );
+
+	// Make the shadow state match the actual state at the time of send
+	FScriptArrayHelper StoredArrayHelper( ( UArrayProperty * )Cmd.Property, StoredData );
+	StoredArrayHelper.Resize( ArrayNum );
+
+	// Read the jump offset
+	// We won't need to actually jump over anything because we expect the change list to be pruned once we get here
+	// But we can use it to verify we read the correct amount.
+	const int32 ArrayChangedCount = WriterState.Changed[WriterState.CurrentChanged++];
+
+	const int32 OldChangedIndex = WriterState.CurrentChanged;
+
+	Data = ( uint8* )Array->GetData();
+	StoredData = ( uint8* )StoredArray->GetData();
+
+	uint16 LocalHandle = 0;
+
+	for ( int32 i = 0; i < Array->Num(); i++ )
 	{
-		Bunch.WriteIntWrapped( FieldCache->FieldNetIndex, ClassCache->GetMaxIndex() + 1 );
+		uint32 Index = i + 1;
+		Writer.SerializeIntPacked( Index );
+		const int32 ElementOffset = i * Cmd.ElementSize;
+		LocalHandle = SendProperties_BackwardsCompatible_r( RepState, WriterState, Writer, PackageMapClient, CompatibleRepLayout, CmdIndex + 1, Cmd.EndCmd - 1, StoredData + ElementOffset, Data + ElementOffset, LocalHandle );
 	}
 
-	NET_CHECKSUM( Bunch );
+	uint32 Index = 0;
+	Writer.SerializeIntPacked( Index );
+
+	check( WriterState.CurrentChanged - OldChangedIndex == ArrayChangedCount );	// Make sure we read correct amount
+	check( WriterState.Changed[WriterState.CurrentChanged] == 0 );				// Make sure we are at the end
+
+	WriterState.CurrentChanged++;
+}
 
-	if ( Property->ArrayDim != 1 )
+uint16 FRepLayout::SendProperties_BackwardsCompatible_r(
+	FRepState *	RESTRICT		RepState,
+	FRepWriterState &			WriterState,
+	FNetBitWriter &				Writer,
+	UPackageMapClient*			PackageMapClient,
+	FCompatibleRepLayout*		CompatibleRepLayout,
+	const int32					CmdStart,
+	const int32					CmdEnd,
+	const uint8* RESTRICT		StoredData,
+	const uint8* RESTRICT		Data,
+	uint16						Handle ) const
+{
+	for ( int32 CmdIndex = CmdStart; CmdIndex < CmdEnd; CmdIndex++ )
 	{
-		Bunch.SerializeIntPacked( ArrayIndex );
+		const FRepLayoutCmd& Cmd = Cmds[CmdIndex];
+
+		check( Cmd.Type != REPCMD_Return );
+
+		Handle++;
+
+		if ( Cmd.Type == REPCMD_DynamicArray )
+		{
+			if ( ShouldSendProperty_BackwardsCompatible( WriterState, Writer, Handle, CmdIndex + 1 ) )
+			{
+				PackageMapClient->TrackCompatibleRepLayoutCmd( this, CompatibleRepLayout, CmdIndex );
+
+				FNetBitWriter TempWriter( Writer.PackageMap, 0 );
+				SendProperties_BackwardsCompatible_DynamicArray_r( RepState, WriterState, TempWriter, PackageMapClient, CompatibleRepLayout, CmdIndex, StoredData + Cmd.Offset, Data + Cmd.Offset );
+
+				uint32 NumBits = TempWriter.GetNumBits();
+				Writer.SerializeIntPacked( NumBits );
+				Writer.SerializeBits( TempWriter.GetData(), NumBits );
+			}
+			CmdIndex = Cmd.EndCmd - 1;	// Jump past children of this array (-1 for the ++ in the for loop)
+			continue;
+		}
+
+		if ( ShouldSendProperty_BackwardsCompatible( WriterState, Writer, Handle, CmdIndex + 1 ) )
+		{
+#if !(UE_BUILD_SHIPPING || UE_BUILD_TEST)
+			if ( CVarDoReplicationContextString->GetInt() > 0 )
+			{
+				Writer.PackageMap->SetDebugContextString( FString::Printf( TEXT( "%s - %s" ), *Owner->GetPathName(), *Cmd.Property->GetPathName() ) );
+			}
+#endif
+			PackageMapClient->TrackCompatibleRepLayoutCmd( this, CompatibleRepLayout, CmdIndex );
+
+			const int32 NumStartBits = Writer.GetNumBits();
+
+			// This property changed, so send it
+			FNetBitWriter TempWriter( Writer.PackageMap, 0 );
+
+			Cmd.Property->NetSerializeItem( TempWriter, TempWriter.PackageMap, ( void* )( Data + Cmd.Offset ) );
+
+			uint32 NumBits = TempWriter.GetNumBits();
+			Writer.SerializeIntPacked( NumBits );
+			Writer.SerializeBits( TempWriter.GetData(), NumBits );
+
+			const int32 NumEndBits = Writer.GetNumBits();
+
+			const FRepParentCmd& ParentCmd = Parents[Cmd.ParentIndex];
+
+			NETWORK_PROFILER( GNetworkProfiler.TrackReplicateProperty( ParentCmd.Property, NumEndBits - NumStartBits, nullptr ) );
+
+			// Make the shadow state match the actual state at the time of send
+			StoreProperty( Cmd, ( void* )( StoredData + Cmd.Offset ), ( const void* )( Data + Cmd.Offset ) );
+
+#if !(UE_BUILD_SHIPPING || UE_BUILD_TEST)
+			if ( CVarDoReplicationContextString->GetInt() > 0 )
+			{
+				Writer.PackageMap->ClearDebugContextString();
+			}
+#endif
+
+#ifdef ENABLE_PROPERTY_CHECKSUMS
+			if ( WriterState.bDoChecksum )
+			{
+				SerializeReadWritePropertyChecksum( Cmd, CmdIndex, Data + Cmd.Offset, Writer );
+			}
+#endif
+		}
 	}
 
-	NETWORK_PROFILER(GNetworkProfiler.TrackWritePropertyHeader(Property, Bunch.GetNumBits() - NumStartingBits, nullptr ));
+	WritePropertyHandle_BackwardsCompatible( Writer, 0, WriterState.bDoChecksum );
+
+	return Handle;
 }
 
-void FRepLayout::SendProperties( 
-	FRepState *	RESTRICT		RepState, 
-	const FReplicationFlags &	RepFlags,
-	const uint8* RESTRICT		Data, 
+void FRepLayout::SendProperties_BackwardsCompatible(
+	FRepState *	RESTRICT		RepState,
+	const uint8* RESTRICT		Data,
 	UClass *					ObjectClass,
 	UActorChannel	*			OwningChannel,
-	FOutBunch &					Writer, 
-	TArray< uint16 > &			Changed, 
-	bool &						bContentBlockWritten ) const
+	FNetBitWriter&				Writer,
+	TArray< uint16 > &			Changed ) const
 {
 #ifdef ENABLE_PROPERTY_CHECKSUMS
 	const bool bDoChecksum = CVarDoPropertyChecksum.GetValueOnGameThread() == 1;
+	Writer.WriteBit( bDoChecksum ? 1 : 0 );
 #else
 	const bool bDoChecksum = false;
 #endif
 
-	FRepWriterState WriterState( Writer, Changed, bDoChecksum );
+	UPackageMapClient* PackageMapClient = ( UPackageMapClient* )OwningChannel->Connection->PackageMap;
+	FCompatibleRepLayout* CompatibleRepLayout = PackageMapClient->TrackCompatibleRepLayout( this ).Get();
 
-#ifdef ENABLE_PROPERTY_CHECKSUMS
-	Writer.WriteBit( bDoChecksum ? 1 : 0 );
-#endif
-
-	SendProperties_r( RepState, RepFlags, WriterState, 0, Cmds.Num() - 1, RepState->StaticBuffer.GetData(), Data, 0 );
+	FRepWriterState WriterState( Writer, Changed, bDoChecksum );
 
-	WritePropertyHandle( Writer, 0, bDoChecksum );
+	SendProperties_BackwardsCompatible_r( RepState, WriterState, Writer, PackageMapClient, CompatibleRepLayout, 0, Cmds.Num() - 1, RepState->StaticBuffer.GetData(), Data, 0 );
 }
 
 class FReceivedPropertiesStackState : public FCmdIteratorBaseStackState
@@ -1190,6 +1336,156 @@ public:
 	FUnmappedGuidMgr* UnmappedGuids;
 };
 
+static bool ReceivePropertyHelper( 
+	FNetBitReader&					Bunch, 
+	FUnmappedGuidMgr*				UnmappedGuids,
+	const int32						ElementOffset, 
+	uint8* RESTRICT					ShadowData,
+	uint8* RESTRICT					Data,
+	TArray< UProperty * >*			RepNotifies,
+	const TArray< FRepParentCmd >&	Parents,
+	const TArray< FRepLayoutCmd >&	Cmds,
+	const int32						CmdIndex,
+	const bool						bDoChecksum )
+{
+	const FRepLayoutCmd& Cmd	= Cmds[CmdIndex];
+	const FRepParentCmd& Parent = Parents[Cmd.ParentIndex];
+
+	// This swaps Role/RemoteRole as we write it
+	const FRepLayoutCmd& SwappedCmd = Parent.RoleSwapIndex != -1 ? Cmds[Parents[Parent.RoleSwapIndex].CmdStart] : Cmd;
+
+	// Let package map know we want to track and know about any guids that are unmapped during the serialize call
+	Bunch.PackageMap->ResetTrackedUnmappedGuids( true );
+
+	// Remember where we started reading from, so that if we have unmapped properties, we can re-deserialize from this data later
+	FBitReaderMark Mark( Bunch );
+
+	if ( RepNotifies != nullptr && Parent.Property->HasAnyPropertyFlags( CPF_RepNotify ) )
+	{
+		// Copy current value over so we can check to see if it changed
+		StoreProperty( Cmd, ShadowData + Cmd.Offset, Data + SwappedCmd.Offset );
+
+		// Read the property
+		Cmd.Property->NetSerializeItem( Bunch, Bunch.PackageMap, Data + SwappedCmd.Offset );
+
+		// Check to see if this property changed
+		if ( Parent.RepNotifyCondition == REPNOTIFY_Always || !PropertiesAreIdentical( Cmd, ShadowData + Cmd.Offset, Data + SwappedCmd.Offset ) )
+		{
+			(*RepNotifies).AddUnique( Parent.Property );
+		}
+		else
+		{
+			UE_CLOG( LogSkippedRepNotifies > 0, LogRep, Display, TEXT( "2 FReceivedPropertiesStackState Skipping RepNotify for propery %s because local value has not changed." ), *Cmd.Property->GetName() );
+		}
+	}
+	else
+	{
+		Cmd.Property->NetSerializeItem( Bunch, Bunch.PackageMap, Data + SwappedCmd.Offset );
+	}
+
+#ifdef ENABLE_PROPERTY_CHECKSUMS
+	if ( bDoChecksum )
+	{
+		SerializeReadWritePropertyChecksum( Cmd, CmdIndex, Data + SwappedCmd.Offset, Bunch );
+	}
+#endif
+
+	const int32 AbsOffset = ElementOffset + SwappedCmd.Offset;
+
+	const TArray< FNetworkGUID > & TrackedUnmappedGuids = Bunch.PackageMap->GetTrackedUnmappedGuids();
+
+	bool bHasUnmapped = false;
+
+	if ( TrackedUnmappedGuids.Num() > 0 )
+	{
+		// If we have an unmapped guid, we need to remember it, so we can fix up this object pointer when it finally arrives at a other time
+		// Note - If we already have an existing entry for this offset, we will replace it with this most recent data
+		UnmappedGuids->Map.Add( AbsOffset, FUnmappedGuidMgrElement( Bunch, Mark, TrackedUnmappedGuids, Cmd.ParentIndex, CmdIndex ) );
+
+		UE_LOG( LogRep, Verbose, TEXT( "ADDED unmapped property: Offset: %i, Name: %s" ), AbsOffset, *Cmd.Property->GetName() );
+
+		// List all the guids that were unmapped
+		for ( int32 i = 0; i < TrackedUnmappedGuids.Num(); i++ )
+		{
+			UE_LOG( LogRep, Verbose, TEXT( "  Guid: %s" ), *TrackedUnmappedGuids[i].ToString() );
+		}
+
+		bHasUnmapped = true;
+	}
+	else
+	{
+		// If we don't have any unmapped guids, then make sure to remove the entry so we don't serialize old data when we update unmapped objects
+		UnmappedGuids->Map.Remove( AbsOffset );
+	}
+
+	// Stop tracking unmapped objects
+	Bunch.PackageMap->ResetTrackedUnmappedGuids( false );
+
+	return bHasUnmapped;
+}
+
+static FUnmappedGuidMgr* PrepReceivedArray(
+	const int32				ArrayNum,
+	FScriptArray*			ShadowArray,
+	FScriptArray*			DataArray,
+	FUnmappedGuidMgr*		ParentUnmappedGuids,
+	const int32				AbsOffset,
+	const FRepParentCmd&	Parent, 
+	const FRepLayoutCmd&	Cmd, 
+	const int32				CmdIndex,
+	uint8* RESTRICT*		OutShadowBaseData,
+	uint8* RESTRICT*		OutBaseData,
+	TArray< UProperty * >*	RepNotifies )
+{
+	// Since we don't know yet if something under us could be unmapped, go ahead and allocate an array container now
+	FUnmappedGuidMgrElement * NewArrayElement = ParentUnmappedGuids->Map.Find( AbsOffset );
+
+	if ( NewArrayElement == NULL )
+	{
+		NewArrayElement = &ParentUnmappedGuids->Map.FindOrAdd( AbsOffset );
+
+		NewArrayElement->Array			= new FUnmappedGuidMgr;
+		NewArrayElement->ParentIndex	= Cmd.ParentIndex;
+		NewArrayElement->CmdIndex		= CmdIndex;
+	}
+
+	check( NewArrayElement != NULL );
+	check( NewArrayElement->ParentIndex == Cmd.ParentIndex );
+	check( NewArrayElement->CmdIndex == CmdIndex );
+
+	if ( RepNotifies != nullptr && ( DataArray->Num() != ArrayNum || Parent.RepNotifyCondition == REPNOTIFY_Always ) && Parent.Property->HasAnyPropertyFlags( CPF_RepNotify ) )
+	{
+		(*RepNotifies).AddUnique( Parent.Property );
+	}
+	else
+	{
+		UE_CLOG( LogSkippedRepNotifies > 0, LogRep, Display, TEXT( "1 FReceivedPropertiesStackState Skipping RepNotify for propery %s because local value has not changed." ), *Cmd.Property->GetName() );
+	}
+
+	check( CastChecked< UArrayProperty >( Cmd.Property ) != NULL );
+
+	// Resize arrays if needed
+	FScriptArrayHelper ArrayHelper( ( UArrayProperty * )Cmd.Property, DataArray );
+	ArrayHelper.Resize( ArrayNum );
+
+	// Re-compute the base data values since they could have changed after the resize above
+	*OutBaseData		= ( uint8* )DataArray->GetData();
+	*OutShadowBaseData	= nullptr;
+
+	// Only resize the shadow data array if we're actually tracking RepNotifies
+	if ( RepNotifies != nullptr )
+	{
+		check( ShadowArray != nullptr );
+
+		FScriptArrayHelper ShadowArrayHelper( ( UArrayProperty* )Cmd.Property, ShadowArray );
+		ShadowArrayHelper.Resize( ArrayNum );
+
+		*OutShadowBaseData = ( uint8* )ShadowArray->GetData();
+	}
+
+	return NewArrayElement->Array;
+}
+
 class FReceivePropertiesImpl : public FRepLayoutCmdIterator< FReceivePropertiesImpl, FReceivedPropertiesStackState >
 {
 public:
@@ -1245,51 +1541,20 @@ public:
 
 		const int32 AbsOffset = Data - PrevStackState.BaseData;
 
-		// Since we don't know yet if something under us could be unmapped, go ahead and allocate an array container now
-		FUnmappedGuidMgrElement * NewArrayElement = PrevStackState.UnmappedGuids->Map.Find( AbsOffset );
-
-		if ( NewArrayElement == NULL )
-		{
-			NewArrayElement = &PrevStackState.UnmappedGuids->Map.FindOrAdd( AbsOffset );
-			NewArrayElement->Array = new FUnmappedGuidMgr;
-			NewArrayElement->ParentIndex = Cmd.ParentIndex;
-			NewArrayElement->CmdIndex = CmdIndex;
-		}
-
-		check( NewArrayElement != NULL );
-		check( NewArrayElement->ParentIndex == Cmd.ParentIndex );
-		check( NewArrayElement->CmdIndex == CmdIndex );
-
-		StackState.UnmappedGuids = NewArrayElement->Array;
-
 		const FRepParentCmd& Parent = Parents[Cmd.ParentIndex];
 
-		if ( bDoRepNotify && ( StackState.DataArray->Num() != ArrayNum || Parent.RepNotifyCondition == REPNOTIFY_Always ) && Parent.Property->HasAnyPropertyFlags( CPF_RepNotify ) )
-		{
-			RepState->RepNotifies.AddUnique( Parent.Property );
-		}
-		else
-		{
-			UE_CLOG( LogSkippedRepNotifies > 0, LogRep, Display, TEXT( "1 FReceivedPropertiesStackState Skipping RepNotify for propery %s because local value has not changed."), *Cmd.Property->GetName() );
-		}
-
-		check( CastChecked< UArrayProperty >( Cmd.Property ) != NULL );
-
-		// Resize arrays if needed
-		FScriptArrayHelper ArrayHelper( (UArrayProperty *)Cmd.Property, Data );
-		ArrayHelper.Resize( ArrayNum );
-
-		// Re-compute the base data values since they could have changed after the resize above
-		StackState.BaseData			= (uint8*)StackState.DataArray->GetData();
-		
-		// Only resize the shadow data array if we're actually tracking RepNotifies
-		if (bDoRepNotify)
-		{
-			FScriptArrayHelper ShadowArrayHelper((UArrayProperty*)Cmd.Property, ShadowData);
-			ShadowArrayHelper.Resize(ArrayNum);
-
-			StackState.ShadowBaseData = (uint8*)StackState.ShadowArray->GetData();
-		}
+		StackState.UnmappedGuids = PrepReceivedArray( 
+			ArrayNum,
+			StackState.ShadowArray,
+			StackState.DataArray,
+			PrevStackState.UnmappedGuids,
+			AbsOffset,
+			Parent,
+			Cmd,
+			CmdIndex,
+			&StackState.ShadowBaseData,
+			&StackState.BaseData,
+			bDoRepNotify ? &RepState->RepNotifies : nullptr );
 
 		// Save the old handle so we can restore it when we pop out of the array
 		const uint16 OldHandle = CurrentHandle;
@@ -1312,80 +1577,12 @@ public:
 	{
 		check( StackState.UnmappedGuids != NULL );
 
-		const FRepParentCmd& Parent = Parents[Cmd.ParentIndex];
-
-		// This swaps Role/RemoteRole as we write it
-		const FRepLayoutCmd& SwappedCmd = Parent.RoleSwapIndex != -1 ? Cmds[Parents[Parent.RoleSwapIndex].CmdStart] : Cmd;
+		const int32 ElementOffset = ( Data - StackState.BaseData );
 
-		// Let package map know we want to track and know about any guids that are unmapped during the serialize call
-		Bunch.PackageMap->ResetTrackedUnmappedGuids( true );
-
-		// Remember where we started reading from, so that if we have unmapped properties, we can re-deserialize from this data later
-		FBitReaderMark Mark( Bunch );
-
-		if ( bDoRepNotify && Parent.Property->HasAnyPropertyFlags( CPF_RepNotify ) )
+		if ( ReceivePropertyHelper( Bunch, StackState.UnmappedGuids, ElementOffset, ShadowData, Data, bDoRepNotify ? &RepState->RepNotifies : nullptr, Parents, Cmds, CmdIndex, bDoChecksum ) )
 		{
-			// Copy current value over so we can check to see if it changed
-			StoreProperty( Cmd, ShadowData + Cmd.Offset, Data + SwappedCmd.Offset );
-
-			// Read the property
-			Cmd.Property->NetSerializeItem( Bunch, Bunch.PackageMap, Data + SwappedCmd.Offset );
-
-			// Check to see if this property changed
-			if ( Parent.RepNotifyCondition == REPNOTIFY_Always || !PropertiesAreIdentical( Cmd, ShadowData + Cmd.Offset, Data + SwappedCmd.Offset ) )
-			{
-				RepState->RepNotifies.AddUnique( Parent.Property );
-			}
-			else
-			{
-				UE_CLOG( LogSkippedRepNotifies > 0, LogRep, Display, TEXT( "2 FReceivedPropertiesStackState Skipping RepNotify for propery %s because local value has not changed."), *Cmd.Property->GetName() );
-			}
-		} 
-		else
-		{
-			Cmd.Property->NetSerializeItem( Bunch, Bunch.PackageMap, Data + SwappedCmd.Offset );
-		}
-
-		const TArray< FNetworkGUID > & TrackedUnmappedGuids = Bunch.PackageMap->GetTrackedUnmappedGuids();
-
-		const int32 AbsOffset = Data - StackState.BaseData;
-
-		const int32 AbsCmdOffset = AbsOffset + SwappedCmd.Offset;
-
-		if ( TrackedUnmappedGuids.Num() > 0 )
-		{
-			// If we have an unmapped guid, we need to remember it, so we can fix up this object pointer when it finally arrives at a other time
-			// Note - If we already have an existing entry for this offset, we will replace it with this most recent data
-			StackState.UnmappedGuids->Map.Add( AbsCmdOffset, FUnmappedGuidMgrElement( Bunch, Mark, TrackedUnmappedGuids, Cmd.ParentIndex, CmdIndex ) );
-
-			UE_LOG( LogRep, Verbose, TEXT( "ADDED unmapped property: Offset: %i, Name: %s"), AbsCmdOffset, *Cmd.Property->GetName() );
-
-			// List all the guids that were unmapped
-			for ( int32 i = 0; i < TrackedUnmappedGuids.Num(); i++ )
-			{
-				UE_LOG( LogRep, Verbose, TEXT( "  Guid: %s"), *TrackedUnmappedGuids[i].ToString() );
-			}
-
 			bHasUnmapped = true;
 		}
-		else
-		{
-			// If we don't have any unmapped guids, then make sure to remove the entry so we don't serialize old data when we update unmapped objects
-			StackState.UnmappedGuids->Map.Remove( AbsCmdOffset );
-		}
-
-		// Stop tracking unmapped objects
-		Bunch.PackageMap->ResetTrackedUnmappedGuids( false );
-
-#ifdef ENABLE_PROPERTY_CHECKSUMS
-		if ( bDoChecksum )
-		{
-			SerializeReadWritePropertyChecksum( Cmd, CmdIndex, Data + SwappedCmd.Offset, Bunch );
-		}
-#endif
-
-		// Store the property we just received FIXME: Can't do this because of rep notifies. This needs to be done for super checksums though, FIXME!!!
-		//StoreProperty( Cmd, StoredData + Cmd.Offset, Data + SwappedCmd.Offset );
 
 		// Read the next property handle
 		ReadNextHandle();
@@ -1400,7 +1597,7 @@ public:
 	bool					bDoRepNotify;
 };
 
-bool FRepLayout::ReceiveProperties( UClass * InObjectClass, FRepState * RESTRICT RepState, void* RESTRICT Data, FNetBitReader & InBunch, bool & bOutHasUnmapped, const bool bEnableRepNotifies ) const
+bool FRepLayout::ReceiveProperties( UActorChannel* OwningChannel, UClass * InObjectClass, FRepState * RESTRICT RepState, void* RESTRICT Data, FNetBitReader & InBunch, bool & bOutHasUnmapped, const bool bEnableRepNotifies ) const
 {
 	check( InObjectClass == Owner );
 
@@ -1412,6 +1609,13 @@ bool FRepLayout::ReceiveProperties( UClass * InObjectClass, FRepState * RESTRICT
 
 	bOutHasUnmapped = false;
 
+	if ( OwningChannel->Connection->InternalAck )
+	{
+		TSharedPtr< FCompatibleRepLayout > CompatibleRepLayout = ( ( UPackageMapClient* )OwningChannel->Connection->PackageMap )->GetCompatibleRepLayoutChecked( Owner );
+
+		return ReceiveProperties_BackwardsCompatible_r( RepState, CompatibleRepLayout.Get(), InBunch, 0, Cmds.Num() - 1, bEnableRepNotifies ? RepState->StaticBuffer.GetData() : nullptr, ( uint8* )Data, ( uint8* )Data, &RepState->UnmappedGuids, bOutHasUnmapped );
+	}
+
 	FReceivePropertiesImpl ReceivePropertiesImpl( InBunch, RepState, bDoChecksum, Parents, Cmds, bEnableRepNotifies );
 
 	// Read first handle
@@ -1439,6 +1643,206 @@ bool FRepLayout::ReceiveProperties( UClass * InObjectClass, FRepState * RESTRICT
 	return true;
 }
 
+int32 FRepLayout::FindCompatibleProperty( const int32 CmdStart, const int32 CmdEnd, const uint32 Checksum ) const
+{
+	for ( int32 CmdIndex = CmdStart; CmdIndex < CmdEnd; CmdIndex++ )
+	{
+		const FRepLayoutCmd& Cmd = Cmds[CmdIndex];
+
+		check( Cmd.Type != REPCMD_Return );
+
+		if ( Cmd.CompatibleChecksum == Checksum )
+		{
+			return CmdIndex;
+		}
+
+		// Jump over entire array and inner properties if checksum didn't match
+		if ( Cmd.Type == REPCMD_DynamicArray )
+		{
+			CmdIndex = Cmd.EndCmd - 1;
+		}
+	}
+
+	return -1;
+}
+
+bool FRepLayout::ReceiveProperties_BackwardsCompatible_r( 
+	FRepState * RESTRICT	RepState,
+	FCompatibleRepLayout*	CompatibleRepLayout,
+	FNetBitReader &			Reader,
+	const int32				CmdStart,
+	const int32				CmdEnd,
+	uint8* RESTRICT			ShadowData,
+	uint8* RESTRICT			OldData,
+	uint8* RESTRICT			Data,
+	FUnmappedGuidMgr*		UnmappedGuids,
+	bool &					bOutHasUnmapped ) const
+{
+	while ( true )
+	{
+		uint32 CompatibleCmdHandle = 0;
+		Reader.SerializeIntPacked( CompatibleCmdHandle );
+
+		if ( Reader.IsError() )
+		{
+			UE_LOG( LogRep, Warning, TEXT( "ReceiveProperties_BackwardsCompatible_r: Error reading handle. Owner: %s" ), *Owner->GetName() );
+			return false;
+		}
+
+		if ( CompatibleCmdHandle == 0 )
+		{
+			// We're done
+			break;
+		}
+
+		// We purposely add 1 on save, so we can reserve 0 for "done"
+		CompatibleCmdHandle--;
+
+		if ( !ensure( CompatibleCmdHandle < ( uint32 )CompatibleRepLayout->Cmds.Num() ) )
+		{
+			UE_LOG( LogRep, Warning, TEXT( "ReceiveProperties_BackwardsCompatible_r: CompatibleCmdHandle > CompatibleRepLayout->Cmds.Num(). Owner: %s, CompatibleCmdHandle: %u" ), *Owner->GetName(), CompatibleCmdHandle );
+			return false;
+		}
+
+		const uint32 Checksum = CompatibleRepLayout->Cmds[CompatibleCmdHandle].CompatibleChecksum;
+
+		if ( !ensure( Checksum != 0 ) )
+		{
+			UE_LOG( LogRep, Warning, TEXT( "ReceiveProperties_BackwardsCompatible_r: Checksum == 0. Owner: %s, Name: %s, Type: %s, CompatibleCmdHandle: %i" ), *Owner->GetName(), *CompatibleRepLayout->Cmds[CompatibleCmdHandle].Name, *CompatibleRepLayout->Cmds[CompatibleCmdHandle].Type, CompatibleCmdHandle );
+			return false;
+		}
+
+		uint32 NumBits = 0;
+		Reader.SerializeIntPacked( NumBits );
+
+		if ( Reader.IsError() )
+		{
+			UE_LOG( LogRep, Warning, TEXT( "ReceiveProperties_BackwardsCompatible_r: Error reading num bits. Owner: %s, Name: %s, Type: %s, CompatibleCmdHandle: %i, Checksum: %u" ), *Owner->GetName(), *CompatibleRepLayout->Cmds[CompatibleCmdHandle].Name, *CompatibleRepLayout->Cmds[CompatibleCmdHandle].Type, CompatibleCmdHandle, Checksum );
+			return false;
+		}
+
+		FNetBitReader TempReader;
+		
+		TempReader.PackageMap = Reader.PackageMap;
+		TempReader.SetData( Reader, NumBits );
+
+		if ( Reader.IsError() )
+		{
+			UE_LOG( LogRep, Warning, TEXT( "ReceiveProperties_BackwardsCompatible_r: Error reading payload. Owner: %s, Name: %s, Type: %s, CompatibleCmdHandle: %i, Checksum: %u" ), *Owner->GetName(), *CompatibleRepLayout->Cmds[CompatibleCmdHandle].Name, *CompatibleRepLayout->Cmds[CompatibleCmdHandle].Type, CompatibleCmdHandle, Checksum );
+			return false;
+		}
+
+		if ( CompatibleRepLayout->Cmds[CompatibleCmdHandle].bIncompatible )
+		{
+			continue;		// We've already warned that this property doesn't load anymore
+		}
+
+		// Find this property
+		const int32 CmdIndex = FindCompatibleProperty( CmdStart, CmdEnd, Checksum );
+
+		if ( CmdIndex == -1 )
+		{
+			UE_LOG( LogRep, Warning, TEXT( "ReceiveProperties_BackwardsCompatible_r: Property not found. Owner: %s, Name: %s, Type: %s, CompatibleCmdHandle: %i, Checksum: %u" ), *Owner->GetName(), *CompatibleRepLayout->Cmds[CompatibleCmdHandle].Name, *CompatibleRepLayout->Cmds[CompatibleCmdHandle].Type, CompatibleCmdHandle, Checksum );
+
+			// Mark this property as incompatible so we don't keep spamming this warning
+			CompatibleRepLayout->Cmds[CompatibleCmdHandle].bIncompatible = true;
+			continue;
+		}
+
+		const FRepLayoutCmd& Cmd = Cmds[CmdIndex];
+
+		if ( Cmd.Type == REPCMD_DynamicArray )
+		{
+			uint32 ArrayNum = 0;
+			TempReader.SerializeIntPacked( ArrayNum );
+
+			if ( TempReader.IsError() )
+			{
+				return false;
+			}
+
+			const int32 AbsOffset = ( Data - OldData ) + Cmd.Offset;
+
+			FScriptArray* DataArray		= ( FScriptArray* )( Data + Cmd.Offset );
+			FScriptArray* ShadowArray	= ShadowData ? ( FScriptArray* )( ShadowData + Cmd.Offset ) : nullptr;
+
+			uint8* LocalData			= Data;
+			uint8* LocalShadowData		= ShadowData;
+
+			FUnmappedGuidMgr* ArrayUnmappedGuids = PrepReceivedArray(
+				ArrayNum,
+				ShadowArray,
+				DataArray,
+				UnmappedGuids,
+				AbsOffset,
+				Parents[Cmd.ParentIndex],
+				Cmd,
+				CmdIndex,
+				&LocalShadowData,
+				&LocalData,
+				ShadowData != nullptr ? &RepState->RepNotifies : nullptr );
+
+			// Read until we read all array elements
+			while ( true )
+			{
+				uint32 Index = 0;
+				TempReader.SerializeIntPacked( Index );
+
+				if ( TempReader.IsError() )
+				{
+					UE_LOG( LogRep, Warning, TEXT( "ReceiveProperties_BackwardsCompatible_r: Error reading array index. Index: %i, Owner: %s, Name: %s, Type: %s, CompatibleCmdHandle: %i, Checksum: %u" ), Index, *Owner->GetName(), *CompatibleRepLayout->Cmds[CompatibleCmdHandle].Name, *CompatibleRepLayout->Cmds[CompatibleCmdHandle].Type, CompatibleCmdHandle, Checksum );
+					return false;
+				}
+
+				if ( Index == 0 )
+				{
+					// We're done
+					break;
+				}
+
+				const int32 ElementOffset = ( Index - 1 ) * Cmd.ElementSize;
+
+				uint8* ElementData			= LocalData + ElementOffset;
+				uint8* ElementShadowData	= LocalShadowData ? LocalShadowData + ElementOffset : nullptr;
+
+				if ( !ReceiveProperties_BackwardsCompatible_r( RepState, CompatibleRepLayout, TempReader, CmdIndex + 1, Cmd.EndCmd - 1, ElementShadowData, LocalData, ElementData, ArrayUnmappedGuids, bOutHasUnmapped ) )
+				{
+					return false;
+				}
+
+				if ( TempReader.IsError() )
+				{
+					UE_LOG( LogRep, Warning, TEXT( "ReceiveProperties_BackwardsCompatible_r: Error reading array index element payload. Index: %i, Owner: %s, Name: %s, Type: %s, CompatibleCmdHandle: %i, Checksum: %u" ), Index, *Owner->GetName(), *CompatibleRepLayout->Cmds[CompatibleCmdHandle].Name, *CompatibleRepLayout->Cmds[CompatibleCmdHandle].Type, CompatibleCmdHandle, Checksum );
+					return false;
+				}
+			}
+
+			if ( TempReader.GetBitsLeft() != 0 )
+			{
+				UE_LOG( LogRep, Warning, TEXT( "ReceiveProperties_BackwardsCompatible_r: Array didn't read propery number of bits. Owner: %s, Name: %s, Type: %s, CompatibleCmdHandle: %i, Checksum: %u" ), *Owner->GetName(), *CompatibleRepLayout->Cmds[CompatibleCmdHandle].Name, *CompatibleRepLayout->Cmds[CompatibleCmdHandle].Type, CompatibleCmdHandle, Checksum );
+				return false;
+			}
+		}
+		else
+		{
+			const int32 ElementOffset = ( Data - OldData );
+
+			if ( ReceivePropertyHelper( TempReader, UnmappedGuids, ElementOffset, ShadowData, Data, ShadowData != nullptr ? &RepState->RepNotifies : nullptr, Parents, Cmds, CmdIndex, false ) )
+			{
+				bOutHasUnmapped = true;
+			}
+
+			if ( TempReader.GetBitsLeft() != 0 )
+			{
+				UE_LOG( LogRep, Warning, TEXT( "ReceiveProperties_BackwardsCompatible_r: Property didn't read propery number of bits. Owner: %s, Name: %s, Type: %s, CompatibleCmdHandle: %i, Checksum: %u" ), *Owner->GetName(), *CompatibleRepLayout->Cmds[CompatibleCmdHandle].Name, *CompatibleRepLayout->Cmds[CompatibleCmdHandle].Type, CompatibleCmdHandle, Checksum );
+				return false;
+			}
+		}
+	}
+
+	return true;
+}
+
 FUnmappedGuidMgrElement::~FUnmappedGuidMgrElement()
 {
 	if ( Array != NULL )
@@ -2004,18 +2408,22 @@ bool FRepLayout::DiffProperties( TArray<UProperty*>* RepNotifies, void* RESTRICT
 	return DiffPropertiesImpl.bDifferent;
 }
 
-void FRepLayout::AddPropertyCmd( UProperty * Property, int32 Offset, int32 RelativeHandle, int32 ParentIndex )
+uint32 FRepLayout::AddPropertyCmd( UProperty * Property, int32 Offset, int32 RelativeHandle, int32 ParentIndex, uint32 ParentChecksum, int32 StaticArrayIndex )
 {
 	const int32 Index = Cmds.AddZeroed();
 
 	FRepLayoutCmd & Cmd = Cmds[Index];
 
-	Cmd.Property		= Property;
-	Cmd.Type			= REPCMD_Property;		// Initially set to generic type
-	Cmd.Offset			= Offset;
-	Cmd.ElementSize		= Property->ElementSize;
-	Cmd.RelativeHandle	= RelativeHandle;
-	Cmd.ParentIndex		= ParentIndex;
+	Cmd.Property			= Property;
+	Cmd.Type				= REPCMD_Property;		// Initially set to generic type
+	Cmd.Offset				= Offset;
+	Cmd.ElementSize			= Property->ElementSize;
+	Cmd.RelativeHandle		= RelativeHandle;
+	Cmd.ParentIndex			= ParentIndex;
+
+	Cmd.CompatibleChecksum	= FCrc::StrCrc32( *Property->GetName().ToLower(), ParentChecksum );								// Evolve checksum on name
+	Cmd.CompatibleChecksum	= FCrc::StrCrc32( *Property->GetCPPType( nullptr, 0 ), Cmd.CompatibleChecksum );				// Evolve by property type
+	Cmd.CompatibleChecksum	= FCrc::StrCrc32( *FString::Printf( TEXT( "%i" ), StaticArrayIndex ), Cmd.CompatibleChecksum );	// Evolve by StaticArrayIndex (to make all unrolled static array elements unique)
 
 	// Try to special case to custom types we know about
 	if ( Property->IsA( UStructProperty::StaticClass() ) )
@@ -2103,20 +2511,28 @@ void FRepLayout::AddPropertyCmd( UProperty * Property, int32 Offset, int32 Relat
 	{
 		UE_LOG( LogRep, VeryVerbose, TEXT( "AddPropertyCmd: Falling back to default type for property [%s]" ), *Cmd.Property->GetFullName() );
 	}
+
+	return Cmd.CompatibleChecksum;
 }
 
-void FRepLayout::AddArrayCmd( UArrayProperty * Property, int32 Offset, int32 RelativeHandle, int32 ParentIndex )
+uint32 FRepLayout::AddArrayCmd( UArrayProperty * Property, int32 Offset, int32 RelativeHandle, int32 ParentIndex, uint32 ParentChecksum, int32 StaticArrayIndex )
 {
 	const int32 Index = Cmds.AddZeroed();
 
 	FRepLayoutCmd & Cmd = Cmds[Index];
 
-	Cmd.Type			= REPCMD_DynamicArray;
-	Cmd.Property		= Property;
-	Cmd.Offset			= Offset;
-	Cmd.ElementSize		= Property->Inner->ElementSize;
-	Cmd.RelativeHandle	= RelativeHandle;
-	Cmd.ParentIndex		= ParentIndex;
+	Cmd.Type				= REPCMD_DynamicArray;
+	Cmd.Property			= Property;
+	Cmd.Offset				= Offset;
+	Cmd.ElementSize			= Property->Inner->ElementSize;
+	Cmd.RelativeHandle		= RelativeHandle;
+	Cmd.ParentIndex			= ParentIndex;
+	
+	Cmd.CompatibleChecksum	= FCrc::StrCrc32( *Property->GetName().ToLower(), ParentChecksum );					// Evolve checksum on name
+	Cmd.CompatibleChecksum	= FCrc::StrCrc32( *Property->GetCPPType( nullptr, 0 ), Cmd.CompatibleChecksum );	// Evolve by property type
+	Cmd.CompatibleChecksum	= FCrc::StrCrc32( *FString::Printf( TEXT( "%i" ), StaticArrayIndex ), Cmd.CompatibleChecksum );	// Evolve by StaticArrayIndex (to make all unrolled static array elements unique)
+
+	return Cmd.CompatibleChecksum;
 }
 
 void FRepLayout::AddReturnCmd()
@@ -2128,7 +2544,7 @@ void FRepLayout::AddReturnCmd()
 	Cmd.Type = REPCMD_Return;
 }
 
-int32 FRepLayout::InitFromProperty_r( UProperty * Property, int32 Offset, int32 RelativeHandle, int32 ParentIndex )
+int32 FRepLayout::InitFromProperty_r( UProperty * Property, int32 Offset, int32 RelativeHandle, int32 ParentIndex, uint32 ParentChecksum, int32 StaticArrayIndex )
 {
 	UArrayProperty * ArrayProp = Cast< UArrayProperty >( Property );
 
@@ -2138,9 +2554,9 @@ int32 FRepLayout::InitFromProperty_r( UProperty * Property, int32 Offset, int32
 
 		RelativeHandle++;
 
-		AddArrayCmd( ArrayProp, Offset + ArrayProp->GetOffset_ForGC(), RelativeHandle, ParentIndex );
+		const uint32 ArrayChecksum = AddArrayCmd( ArrayProp, Offset + ArrayProp->GetOffset_ForGC(), RelativeHandle, ParentIndex, ParentChecksum, StaticArrayIndex );
 
-		InitFromProperty_r( ArrayProp->Inner, 0, 0, ParentIndex );
+		InitFromProperty_r( ArrayProp->Inner, 0, 0, ParentIndex, ArrayChecksum, 0 );
 		
 		AddReturnCmd();
 
@@ -2164,7 +2580,7 @@ int32 FRepLayout::InitFromProperty_r( UProperty * Property, int32 Offset, int32
 		if ( Struct->StructFlags & STRUCT_NetSerializeNative )
 		{
 			RelativeHandle++;
-			AddPropertyCmd( Property, Offset + Property->GetOffset_ForGC(), RelativeHandle, ParentIndex );
+			AddPropertyCmd( Property, Offset + Property->GetOffset_ForGC(), RelativeHandle, ParentIndex, ParentChecksum, StaticArrayIndex );
 			return RelativeHandle;
 		}
 
@@ -2197,11 +2613,20 @@ int32 FRepLayout::InitFromProperty_r( UProperty * Property, int32 Offset, int32
 
 		Sort( NetProperties.GetData(), NetProperties.Num(), FCompareUFieldOffsets() );
 
+		// Evolve checksum on struct name
+		uint32 StructChecksum = FCrc::StrCrc32( *Property->GetName().ToLower(), ParentChecksum );
+
+		// Evolve by property type
+		StructChecksum = FCrc::StrCrc32( *Property->GetCPPType( nullptr, 0 ), StructChecksum );
+
+		// Evolve by StaticArrayIndex (to make all unrolled static array elements unique)
+		StructChecksum = FCrc::StrCrc32( *FString::Printf( TEXT( "%i" ), StaticArrayIndex ), StructChecksum );
+
 		for ( int32 i = 0; i < NetProperties.Num(); i++ )
 		{
 			for ( int32 j = 0; j < NetProperties[i]->ArrayDim; j++ )
 			{
-				RelativeHandle = InitFromProperty_r( NetProperties[i], Offset + StructProp->GetOffset_ForGC() + j * NetProperties[i]->ElementSize, RelativeHandle, ParentIndex );
+				RelativeHandle = InitFromProperty_r( NetProperties[i], Offset + StructProp->GetOffset_ForGC() + j * NetProperties[i]->ElementSize, RelativeHandle, ParentIndex, StructChecksum, j );
 			}
 		}
 		return RelativeHandle;
@@ -2210,7 +2635,7 @@ int32 FRepLayout::InitFromProperty_r( UProperty * Property, int32 Offset, int32
 	// Add actual property
 	RelativeHandle++;
 
-	AddPropertyCmd( Property, Offset + Property->GetOffset_ForGC(), RelativeHandle, ParentIndex );
+	AddPropertyCmd( Property, Offset + Property->GetOffset_ForGC(), RelativeHandle, ParentIndex, ParentChecksum, StaticArrayIndex );
 
 	return RelativeHandle;
 }
@@ -2246,7 +2671,7 @@ void FRepLayout::InitFromObjectClass( UClass * InObjectClass )
 		check( Parents[i].Property->RepIndex + Parents[i].ArrayIndex == i );
 
 		Parents[ParentHandle].CmdStart = Cmds.Num();
-		RelativeHandle = InitFromProperty_r( Property, Property->ElementSize * ArrayIdx, RelativeHandle, ParentHandle );
+		RelativeHandle = InitFromProperty_r( Property, Property->ElementSize * ArrayIdx, RelativeHandle, ParentHandle, 0, ArrayIdx );
 		Parents[ParentHandle].CmdEnd = Cmds.Num();
 		Parents[ParentHandle].Flags |= PARENT_IsConditional;
 
@@ -2364,7 +2789,7 @@ void FRepLayout::InitFromFunction( UFunction * InFunction )
 		{
 			const int32 ParentHandle = AddParentProperty( *It, ArrayIdx );
 			Parents[ParentHandle].CmdStart = Cmds.Num();
-			RelativeHandle = InitFromProperty_r( *It, It->ElementSize * ArrayIdx, RelativeHandle, ParentHandle );
+			RelativeHandle = InitFromProperty_r( *It, It->ElementSize * ArrayIdx, RelativeHandle, ParentHandle, 0, ArrayIdx );
 			Parents[ParentHandle].CmdEnd = Cmds.Num();
 		}
 	}
@@ -2387,7 +2812,7 @@ void FRepLayout::InitFromStruct( UStruct * InStruct )
 		{
 			const int32 ParentHandle = AddParentProperty( *It, ArrayIdx );
 			Parents[ParentHandle].CmdStart = Cmds.Num();
-			RelativeHandle = InitFromProperty_r( *It, It->ElementSize * ArrayIdx, RelativeHandle, ParentHandle );
+			RelativeHandle = InitFromProperty_r( *It, It->ElementSize * ArrayIdx, RelativeHandle, ParentHandle, 0, ArrayIdx );
 			Parents[ParentHandle].CmdEnd = Cmds.Num();
 		}
 	}
diff --git a/Engine/Source/Runtime/Engine/Public/Net/DataBunch.h b/Engine/Source/Runtime/Engine/Public/Net/DataBunch.h
--- a/Engine/Source/Runtime/Engine/Public/Net/DataBunch.h
+++ b/Engine/Source/Runtime/Engine/Public/Net/DataBunch.h
@@ -27,10 +27,11 @@ public:
 	uint8					bPartial;				// Not a complete bunch
 	uint8					bPartialInitial;		// The first bunch of a partial bunch
 	uint8					bPartialFinal;			// The final bunch of a partial bunch
-	uint8					bHasGUIDs;				// This bunch has networkGUID name/id pairs
+	uint8					bHasPackageMapExports;	// This bunch has networkGUID name/id pairs
 	uint8					bHasMustBeMappedGUIDs;	// This bunch has guids that must be mapped before we can process this bunch
 
 	TArray< FNetworkGUID >	ExportNetGUIDs;			// List of GUIDs that went out on this bunch
+	TArray< uint64 >		ExportRepLayoutCmds;
 
 #if !(UE_BUILD_SHIPPING || UE_BUILD_TEST)
 	FString			DebugString;
@@ -73,7 +74,7 @@ public:
 		Str += FString::Printf(TEXT("bDormant: %d "), bDormant);
 		Str += FString::Printf(TEXT("bReliable: %d "), bReliable);
 		Str += FString::Printf(TEXT("bPartial: %d//%d//%d "), bPartial, bPartialInitial, bPartialFinal);
-		Str += FString::Printf(TEXT("bHasGUIDs: %d "), bHasGUIDs);
+		Str += FString::Printf( TEXT( "bHasPackageMapExports: %d " ), bHasPackageMapExports );
 		Str += GetDebugString();
 #else
 		FString Str = FString::Printf(TEXT("Channel[%d]. Seq %d. PacketId: %d"), ChIndex, ChSequence, PacketId);
@@ -102,7 +103,7 @@ public:
 	uint8				bPartial;				// Not a complete bunch
 	uint8				bPartialInitial;		// The first bunch of a partial bunch
 	uint8				bPartialFinal;			// The final bunch of a partial bunch
-	uint8				bHasGUIDs;				// This bunch has networkGUID name/id pairs
+	uint8				bHasPackageMapExports;	// This bunch has networkGUID name/id pairs
 	uint8				bHasMustBeMappedGUIDs;	// This bunch has guids that must be mapped before we can process this bunch
 
 	FString	ToString()
@@ -119,7 +120,7 @@ public:
 		Str += FString::Printf(TEXT("bDormant: %d "), bDormant);
 		Str += FString::Printf(TEXT("bReliable: %d "), bReliable);
 		Str += FString::Printf(TEXT("bPartial: %d//%d//%d "), bPartial, bPartialInitial, bPartialFinal);
-		Str += FString::Printf(TEXT("bHasGUIDs: %d "), bHasGUIDs);
+		Str += FString::Printf( TEXT( "bHasPackageMapExports: %d " ), bHasPackageMapExports );
 #else
 		FString Str = FString::Printf(TEXT("Channel[%d]. Seq %d. PacketId: %d"), ChIndex, ChSequence, PacketId);
 #endif
diff --git a/Engine/Source/Runtime/Engine/Public/Net/DataReplication.h b/Engine/Source/Runtime/Engine/Public/Net/DataReplication.h
--- a/Engine/Source/Runtime/Engine/Public/Net/DataReplication.h
+++ b/Engine/Source/Runtime/Engine/Public/Net/DataReplication.h
@@ -103,7 +103,7 @@ public:
 	/** Takes Data, and compares against shadow state to log differences */
 	bool ValidateAgainstState( const UObject* ObjectState );
 
-	static bool SerializeCustomDeltaProperty( UNetConnection * Connection, void* Src, UProperty * Property, int32 ArrayDim, FNetBitWriter & OutBunch, TSharedPtr<INetDeltaBaseState> & NewFullState, TSharedPtr<INetDeltaBaseState> & OldState );
+	static bool SerializeCustomDeltaProperty( UNetConnection * Connection, void* Src, UProperty * Property, uint32 ArrayIndex, FNetBitWriter & OutBunch, TSharedPtr<INetDeltaBaseState> & NewFullState, TSharedPtr<INetDeltaBaseState> & OldState );
 
 	/** Packet was dropped */
 	void	ReceivedNak( int32 NakPacketId );
@@ -111,13 +111,11 @@ public:
 	void	Serialize(FArchive& Ar);
 
 	/** Writes dirty properties to bunch */
-	void	ReplicateCustomDeltaProperties( FOutBunch & Bunch, FReplicationFlags RepFlags, bool & bContentBlockWritten );
+	void	ReplicateCustomDeltaProperties( FNetBitWriter & Bunch, FReplicationFlags RepFlags );
 	bool	ReplicateProperties( FOutBunch & Bunch, FReplicationFlags RepFlags );
 	void	PostSendBunch(FPacketIdRange & PacketRange, uint8 bReliable);
 	
-	const FFieldNetCache* ReadField( const FClassNetCache* ClassCache, FInBunch& Bunch ) const;
-
-	bool	ReceivedBunch( FInBunch & Bunch, const FReplicationFlags& RepFlags, bool & bOutHasUnmapped );
+	bool	ReceivedBunch( FNetBitReader& Bunch, const FReplicationFlags& RepFlags, const bool bHasRepLayout, bool& bOutHasUnmapped );
 	void	PostReceivedBunch();
 
 	void	ForceRefreshUnreliableProperties();
@@ -156,4 +154,10 @@ public:
 	}
 
 	void QueuePropertyRepNotify( UObject* Object, UProperty * Property, const int32 ElementIndex, TArray< uint8 > & MetaData );
+
+	void WritePropertyHeaderAndPayload(
+		UObject*			Object,
+		UProperty*			Property,
+		FNetBitWriter&		Bunch,
+		FNetBitWriter&		Payload ) const;
 };
\ No newline at end of file
diff --git a/Engine/Source/Runtime/Engine/Public/Net/RepLayout.h b/Engine/Source/Runtime/Engine/Public/Net/RepLayout.h
--- a/Engine/Source/Runtime/Engine/Public/Net/RepLayout.h
+++ b/Engine/Source/Runtime/Engine/Public/Net/RepLayout.h
@@ -9,8 +9,8 @@
 #include "CoreNet.h"
 #include "Engine/EngineTypes.h"
 
-class FOutBunch;
-class FInBunch;
+class FCompatibleRepLayout;
+class UPackageMapClient;
 
 class FRepChangedParent
 {
@@ -233,13 +233,14 @@ public:
 class FRepLayoutCmd
 {
 public:
-	UProperty * Property;		// Pointer back to property, used for NetSerialize calls, etc.
+	UProperty * Property;			// Pointer back to property, used for NetSerialize calls, etc.
 	uint8		Type;
-	uint16		EndCmd;			// For arrays, this is the cmd index to jump to, to skip this arrays inner elements
-	uint16		ElementSize;	// For arrays, element size of data
-	int32		Offset;			// Absolute offset of property
-	uint16		RelativeHandle;	// Handle relative to start of array, or top list
-	uint16		ParentIndex;	// Index into Parents
+	uint16		EndCmd;				// For arrays, this is the cmd index to jump to, to skip this arrays inner elements
+	uint16		ElementSize;		// For arrays, element size of data
+	int32		Offset;				// Absolute offset of property
+	uint16		RelativeHandle;		// Handle relative to start of array, or top list
+	uint16		ParentIndex;		// Index into Parents
+	uint32		CompatibleChecksum;	// Used to determine if property is still compatible
 };
 
 class FRepWriterState
@@ -266,6 +267,7 @@ public:
 class FRepLayout
 {
 	friend class FRepState;
+	friend class UPackageMapClient;
 
 public:
 	FRepLayout() : FirstNonCustomParent( 0 ), RoleIndex( -1 ), RemoteRoleIndex( -1 ), Owner( NULL ) {}
@@ -280,23 +282,13 @@ public:
 
 	void InitChangedTracker( FRepChangedPropertyTracker * ChangedTracker ) const;
 
-	void WritePropertyHeader( 
-		UObject *			Object,
-		UClass *			ObjectClass,
-		UActorChannel *		OwningChannel,
-		UProperty *			Property, 
-		FOutBunch &			Bunch, 
-		uint32				ArrayIndex, 
-		bool &				bContentBlockWritten ) const;
-
 	bool ReplicateProperties( 
 		FRepState * RESTRICT		RepState, 
 		const uint8* RESTRICT		Data, 
 		UClass *					ObjectClass,
 		UActorChannel *				OwningChannel,
-		FOutBunch &					Writer, 
-		const FReplicationFlags &	RepFlags,
-		bool &						bContentBlockWritten ) const;
+		FNetBitWriter&				Writer, 
+		const FReplicationFlags &	RepFlags ) const;
 
 	void SendProperties( 
 		FRepState *	RESTRICT		RepState, 
@@ -304,13 +296,12 @@ public:
 		const uint8* RESTRICT		Data, 
 		UClass *					ObjectClass,
 		UActorChannel *				OwningChannel,
-		FOutBunch &					Writer, 
-		TArray< uint16 >	 &		Changed, 
-		bool &						bContentBlockWritten ) const;
+		FNetBitWriter&				Writer,
+		TArray< uint16 >&			Changed ) const;
 
 	ENGINE_API void InitFromObjectClass( UClass * InObjectClass );
 
-	bool ReceiveProperties( UClass * InObjectClass, FRepState * RESTRICT RepState, void* RESTRICT Data, FNetBitReader & InBunch, bool & bOutHasUnmapped, const bool bEnableRepNotifies ) const;
+	bool ReceiveProperties( UActorChannel* OwningChannel, UClass * InObjectClass, FRepState * RESTRICT RepState, void* RESTRICT Data, FNetBitReader & InBunch, bool & bOutHasUnmapped, const bool bEnableRepNotifies ) const;
 	void UpdateUnmappedObjects( FRepState *	RepState, UPackageMap * PackageMap, UObject* Object, bool & bOutSomeObjectsWereMapped, bool & bOutHasMoreUnmapped ) const;
 
 	void CallRepNotifies( FRepState * RepState, UObject* Object ) const;
@@ -340,9 +331,6 @@ public:
 	// Serializes all replicated properties of a UObject in or out of an archive (depending on what type of archive it is)
 	ENGINE_API void SerializeObjectReplicatedProperties(UObject* Object, FArchive & Ar) const;
 
-	void WriteNetworkChecksum( FOutBunch& Bunch );
-	bool ReadNetworkChecksum( FInBunch& Bunch );
-
 private:
 	void RebuildConditionalProperties( FRepState * RESTRICT	RepState, const FRepChangedPropertyTracker& ChangedTracker, const FReplicationFlags& RepFlags ) const;
 
@@ -406,7 +394,51 @@ private:
 		const uint8* RESTRICT		Data, 
 		uint16						Handle ) const;
 
-	void UpdateUnmappedObjects_r( 
+	void SendProperties_BackwardsCompatible_DynamicArray_r(
+		FRepState *	RESTRICT		RepState,
+		FRepWriterState &			WriterState,
+		FNetBitWriter &				Writer,
+		UPackageMapClient*			PackageMapClient,
+		FCompatibleRepLayout*		CompatibleRepLayout, 
+		const int32					CmdIndex,
+		const uint8* RESTRICT		StoredData,
+		const uint8* RESTRICT		Data ) const;
+
+	uint16 SendProperties_BackwardsCompatible_r(
+		FRepState *	RESTRICT		RepState,
+		FRepWriterState &			WriterState,
+		FNetBitWriter &				Writer,
+		UPackageMapClient*			PackageMapClient,
+		FCompatibleRepLayout*		CompatibleRepLayout,
+		const int32					CmdStart,
+		const int32					CmdEnd,
+		const uint8* RESTRICT		StoredData,
+		const uint8* RESTRICT		Data,
+		uint16						Handle ) const;
+		
+	void SendProperties_BackwardsCompatible(
+		FRepState *	RESTRICT		RepState,
+		const uint8* RESTRICT		Data,
+		UClass *					ObjectClass,
+		UActorChannel	*			OwningChannel,
+		FNetBitWriter&				Writer,
+		TArray< uint16 > &			Changed ) const;
+	
+	int32 FindCompatibleProperty( const int32 CmdStart, const int32 CmdEnd, const uint32 Checksum ) const;
+
+	bool ReceiveProperties_BackwardsCompatible_r(
+		FRepState * RESTRICT	RepState,
+		FCompatibleRepLayout*	CompatibleRepLayout,
+		FNetBitReader &			Reader,
+		const int32				CmdStart,
+		const int32				CmdEnd,
+		uint8* RESTRICT			ShadowData,
+		uint8* RESTRICT			OldData,
+		uint8* RESTRICT			Data,
+		FUnmappedGuidMgr*		UnmappedGuids,
+		bool &					bOutHasUnmapped ) const;
+
+	void UpdateUnmappedObjects_r(
 		FRepState *			RepState, 
 		FUnmappedGuidMgr *	UnmappedGuids, 
 		UObject *			OriginalObject,
@@ -438,10 +470,10 @@ private:
 
 	uint16 AddParentProperty( UProperty * Property, int32 ArrayIndex );
 
-	int32 InitFromProperty_r( UProperty * Property, int32 Offset, int32 RelativeHandle, int32 ParentIndex );
+	int32 InitFromProperty_r( UProperty * Property, int32 Offset, int32 RelativeHandle, int32 ParentIndex, uint32 ParentChecksum, int32 StaticArrayIndex );
 
-	void AddPropertyCmd( UProperty * Property, int32 Offset, int32 RelativeHandle, int32 ParentIndex );
-	void AddArrayCmd( UArrayProperty * Property, int32 Offset, int32 RelativeHandle, int32 ParentIndex );
+	uint32 AddPropertyCmd( UProperty * Property, int32 Offset, int32 RelativeHandle, int32 ParentIndex, uint32 ParentChecksum, int32 StaticArrayIndex );
+	uint32 AddArrayCmd( UArrayProperty * Property, int32 Offset, int32 RelativeHandle, int32 ParentIndex, uint32 ParentChecksum, int32 StaticArrayIndex );
 	void AddReturnCmd();
 
 	void SerializeProperties_DynamicArray_r( 
diff --git a/Engine/Source/Runtime/NetworkReplayStreaming/HttpNetworkReplayStreaming/Private/HttpNetworkReplayStreaming.cpp b/Engine/Source/Runtime/NetworkReplayStreaming/HttpNetworkReplayStreaming/Private/HttpNetworkReplayStreaming.cpp
--- a/Engine/Source/Runtime/NetworkReplayStreaming/HttpNetworkReplayStreaming/Private/HttpNetworkReplayStreaming.cpp
+++ b/Engine/Source/Runtime/NetworkReplayStreaming/HttpNetworkReplayStreaming/Private/HttpNetworkReplayStreaming.cpp
@@ -13,6 +13,7 @@ DECLARE_CYCLE_STAT( TEXT( "Http replay decompress time" ), STAT_HttpReplay_Decom
 
 static TAutoConsoleVariable<FString> CVarMetaFilterOverride( TEXT( "httpReplay.MetaFilterOverride" ), TEXT( "" ), TEXT( "" ) );
 static TAutoConsoleVariable<float> CVarChunkUploadDelayInSeconds( TEXT( "httpReplay.ChunkUploadDelayInSeconds" ), 10.0f, TEXT( "" ) );
+static TAutoConsoleVariable<int> CVarMaxCacheSize( TEXT( "httpReplay.MaxCacheSize" ), 1024 * 1024 * 10, TEXT( "" ) );		// 10 MB cache by default
 
 class FNetworkReplayListItem : public FOnlineJsonSerializable
 {
@@ -330,6 +331,59 @@ void FHttpNetworkReplayStreamer::AddCustomRequestToQueue( TSharedPtr< FQueuedHtt
 	QueuedHttpRequests.Add( Request );
 }
 
+void FHttpNetworkReplayStreamer::AddResponseToCache( FHttpRequestPtr HttpRequest, FHttpResponsePtr HttpResponse )
+{
+	if ( !HttpRequest.IsValid() )
+	{
+		return;
+	}
+
+	if ( !HttpResponse.IsValid() )
+	{
+		return;
+	}
+
+	// Add to cache (or freshen existing entry)
+	ResponseCache.Add( HttpRequest->GetURL(), FCachedResponse( HttpResponse, FPlatformTime::Seconds() ) );
+
+	// Anytime we add something to cache, make sure it's within budget
+	CleanupResponseCache();
+}
+
+void FHttpNetworkReplayStreamer::CleanupResponseCache()
+{
+	// Remove older entries until we're under the CVarMaxCacheSize threshold
+	while ( ResponseCache.Num() )
+	{
+		double OldestTime = 0.0;
+		FString OldestKey;
+		uint32 TotalSize = 0;
+
+		for ( auto It = ResponseCache.CreateIterator(); It; ++It )
+		{
+			if ( OldestKey.IsEmpty() || It.Value().LastAccessTime < OldestTime )
+			{
+				OldestTime = It.Value().LastAccessTime;
+				OldestKey = It.Key();
+			}
+
+			// Accumulate total cache size
+			TotalSize += It.Value().Response->GetContent().Num();
+		}
+
+		check( !OldestKey.IsEmpty() );
+
+		const uint32 MaxCacheSize = CVarMaxCacheSize.GetValueOnGameThread();
+
+		if ( TotalSize <= MaxCacheSize )
+		{
+			break;	// We're good
+		}
+
+		ResponseCache.Remove( OldestKey );
+	}
+}
+
 static FString BuildRequestErrorString( FHttpRequestPtr HttpRequest, FHttpResponsePtr HttpResponse )
 {
 	FString ExtraInfo;
@@ -1677,6 +1731,9 @@ void FHttpNetworkReplayStreamer::HttpDownloadFinished( FHttpRequestPtr HttpReque
 		{
 			if ( HttpResponse->GetContent().Num() > 0 )
 			{
+				// Add response to response cache
+				AddResponseToCache( HttpRequest, HttpResponse );
+
 				if ( SupportsCompression() )
 				{
 					SCOPE_CYCLE_COUNTER( STAT_HttpReplay_DecompressTime );
@@ -1760,6 +1817,9 @@ void FHttpNetworkReplayStreamer::HttpDownloadCheckpointFinished( FHttpRequestPtr
 			return;
 		}
 
+		// Add response to response cache
+		AddResponseToCache( HttpRequest, HttpResponse );
+
 		// Get the checkpoint data
 		if ( SupportsCompression() )
 		{
@@ -2056,6 +2116,26 @@ bool FHttpNetworkReplayStreamer::ProcessNextHttpRequest()
 
 		InFlightHttpRequest = QueuedRequest;
 
+		FCachedResponse* CachedResponse = ResponseCache.Find( InFlightHttpRequest->Request->GetURL() );
+
+		if ( CachedResponse != nullptr )
+		{
+			if ( InFlightHttpRequest->Request->OnProcessRequestComplete().IsBound() )
+			{
+				// If we have this response in the cache, process it now
+				CachedResponse->LastAccessTime = FPlatformTime::Seconds();
+				InFlightHttpRequest->Request->OnProcessRequestComplete().ExecuteIfBound( InFlightHttpRequest->Request, CachedResponse->Response, true );
+
+				// We can now immediately process next request
+				return ProcessNextHttpRequest();
+			}
+			else
+			{
+				// Should never happen, but just in case
+				ResponseCache.Remove( InFlightHttpRequest->Request->GetURL() );
+			}
+		}
+
 		ProcessRequestInternal( InFlightHttpRequest->Request );
 
 		// We can return now since we know a http request is now in flight
diff --git a/Engine/Source/Runtime/NetworkReplayStreaming/HttpNetworkReplayStreaming/Public/HttpNetworkReplayStreaming.h b/Engine/Source/Runtime/NetworkReplayStreaming/HttpNetworkReplayStreaming/Public/HttpNetworkReplayStreaming.h
--- a/Engine/Source/Runtime/NetworkReplayStreaming/HttpNetworkReplayStreaming/Public/HttpNetworkReplayStreaming.h
+++ b/Engine/Source/Runtime/NetworkReplayStreaming/HttpNetworkReplayStreaming/Public/HttpNetworkReplayStreaming.h
@@ -181,6 +181,17 @@ public:
 	virtual bool PreProcess( FHttpNetworkReplayStreamer* Streamer, const FString& ServerURL, const FString& SessionName ) override;
 };
 
+class FCachedResponse
+{
+public:
+	FCachedResponse( FHttpResponsePtr InResponse, const double InLastAccessTime ) : Response( InResponse ), LastAccessTime( InLastAccessTime )
+	{
+	}
+	
+	FHttpResponsePtr	Response;
+	double				LastAccessTime;
+};
+
 /**
  * Http network replay streaming manager
  */
@@ -237,6 +248,8 @@ public:
 	virtual void AddOrUpdateEvent( const FString& Name, const uint32 TimeInMS, const FString& Group, const FString& Meta, const TArray<uint8>& Data ) override;
 	void AddRequestToQueue( const EQueuedHttpRequestType::Type Type, TSharedPtr< class IHttpRequest > Request, const int32 InMaxRetries = 0, const float InRetryDelay = 0.0f );
 	void AddCustomRequestToQueue( TSharedPtr< FQueuedHttpRequest > Request );
+	void AddResponseToCache( FHttpRequestPtr HttpRequest, FHttpResponsePtr HttpResponse );
+	void CleanupResponseCache();
 	bool RetryRequest( TSharedPtr< FQueuedHttpRequest > Request, FHttpResponsePtr HttpResponse );
 	void EnumerateCheckpoints();
 	void ConditionallyEnumerateCheckpoints();
@@ -314,9 +327,11 @@ public:
 	TArray< TSharedPtr< FQueuedHttpRequest > >	QueuedHttpRequests;
 	TSharedPtr< FQueuedHttpRequest >			InFlightHttpRequest;
 
-	TSet< FString >					EventGroupSet;
+	TSet< FString >						EventGroupSet;
+
+	uint32								TotalUploadBytes;
 
-	uint32							TotalUploadBytes;
+	TMap< FString, FCachedResponse >	ResponseCache;
 };
 
 class HTTPNETWORKREPLAYSTREAMING_API FHttpNetworkReplayStreamingFactory : public INetworkReplayStreamingFactory, public FTickableGameObject
