Disabling debugviewmode shaders on SM4
Reenabling shader complexity on all platforms

diff --git a/Engine/Source/Runtime/Engine/Private/SceneView.cpp b/Engine/Source/Runtime/Engine/Private/SceneView.cpp
--- a/Engine/Source/Runtime/Engine/Private/SceneView.cpp
+++ b/Engine/Source/Runtime/Engine/Private/SceneView.cpp
@@ -2011,13 +2011,13 @@ namespace
 			}			
 
 			ERHIFeatureLevel::Type QuadOverdrawFeature = bEnableQuadOverdraw ? ERHIFeatureLevel::SM5 : ERHIFeatureLevel::Num;
-			ERHIFeatureLevel::Type StreamingAccuracyFeature = bEnableStreamingAccuracy ? ERHIFeatureLevel::SM4 : ERHIFeatureLevel::Num;
-			ERHIFeatureLevel::Type TexCoordScaleFeature = bEnableTextureStreamingBuild ? ERHIFeatureLevel::SM4 : ERHIFeatureLevel::Num;
+			ERHIFeatureLevel::Type StreamingAccuracyFeature = bEnableStreamingAccuracy ? ERHIFeatureLevel::SM5 : ERHIFeatureLevel::Num;
+			ERHIFeatureLevel::Type TexCoordScaleFeature = bEnableTextureStreamingBuild ? ERHIFeatureLevel::SM5 : ERHIFeatureLevel::Num;
 
 			static_assert(DVSM_MAX == 9, "Wrong DebugViewShaderMode Count");
 			MinFeatureLevels[DVSM_None] = ERHIFeatureLevel::Num;
 
-			MinFeatureLevels[DVSM_ShaderComplexity] = ERHIFeatureLevel::SM4;
+			MinFeatureLevels[DVSM_ShaderComplexity] = ERHIFeatureLevel::ES2;
 
 			MinFeatureLevels[DVSM_ShaderComplexityContainedQuadOverhead] = QuadOverdrawFeature;
 			MinFeatureLevels[DVSM_ShaderComplexityBleedingQuadOverhead] = QuadOverdrawFeature;			
@@ -2081,7 +2081,7 @@ bool AllowDebugViewShaderMode(EDebugViewShaderMode ShaderMode, ERHIFeatureLevel:
 	check(ShaderMode < DVSM_MAX && FeatureLevel < ERHIFeatureLevel::Num);
 	return FeatureLevel >= Options.MinFeatureLevels[ShaderMode]; 
 #else
-	return ShaderMode == DVSM_ShaderComplexity && FeatureLevel >= ERHIFeatureLevel::SM4;
+	return ShaderMode == DVSM_ShaderComplexity;
 #endif
 }
 
@@ -2094,7 +2094,7 @@ bool AllowDebugViewShaderMode(EDebugViewShaderMode ShaderMode, EShaderPlatform P
 	//@TODO : else if (!IsPCPlatform(Platform))
 	else if (Platform != SP_PCD3D_SM5 && Platform != SP_PCD3D_SM4)
 	{
-		return ShaderMode == DVSM_ShaderComplexity && GetMaxSupportedFeatureLevel(Platform) >= ERHIFeatureLevel::SM4; // Only shader complexity is available on non shader platforms.
+		return ShaderMode == DVSM_ShaderComplexity;
 	}
 	else // Otherwise, this will be valid only for the current platform.
 	{
