Copying //UE4/Dev-Landscape to Dev-Main (//UE4/Dev-Main)

diff --git a/Engine/Config/BaseScalability.ini b/Engine/Config/BaseScalability.ini
--- a/Engine/Config/BaseScalability.ini
+++ b/Engine/Config/BaseScalability.ini
@@ -28,6 +28,7 @@ PerfIndexThresholds_ShadowQuality="Min 18 40 90"
 PerfIndexThresholds_PostProcessQuality="GPU 18 40 95"
 PerfIndexThresholds_TextureQuality="GPU 18 40 95"
 PerfIndexThresholds_EffectsQuality="Min 18 40 90"
+PerfIndexThresholds_FoliageQuality="GPU 18 40 95"
 
 [AntiAliasingQuality@0]
 r.PostProcessAAQuality=0
@@ -257,4 +258,22 @@ r.MaterialQualityLevel=1
 r.SSS.Scale=1
 r.SSS.SampleSet=2
 r.SSS.Quality=1
-r.EmitterSpawnRateScale=1.0
\ No newline at end of file
+r.EmitterSpawnRateScale=1.0
+
+;-----------------------------------------------------------------------------------------------------------------
+
+[FoliageQuality@0]
+foliage.DensityScale=0
+grass.DensityScale=0
+
+[FoliageQuality@1]
+foliage.DensityScale=0.4
+grass.DensityScale=0.4
+
+[FoliageQuality@2]
+foliage.DensityScale=0.8
+grass.DensityScale=0.8
+
+[FoliageQuality@3]
+foliage.DensityScale=1.0
+grass.DensityScale=1.0
diff --git a/Engine/Source/Editor/BlueprintGraph/Private/BasicTokenParser.cpp b/Engine/Source/Editor/BlueprintGraph/Private/BasicTokenParser.cpp
--- a/Engine/Source/Editor/BlueprintGraph/Private/BasicTokenParser.cpp
+++ b/Engine/Source/Editor/BlueprintGraph/Private/BasicTokenParser.cpp
@@ -300,7 +300,7 @@ bool FBasicTokenParser::GetToken(FBasicToken& Token, bool bNoConsts/* = false*/)
 		Token.TokenType = FBasicToken::TOKEN_Identifier;
 
 		// Lookup the token's global name.
-		Token.TokenName = FName( Token.Identifier, FNAME_Find, true );
+		Token.TokenName = FName(Token.Identifier, FNAME_Find);
 
 		// If const values are allowed, determine whether the identifier represents a constant
 		if ( !bNoConsts )
@@ -468,7 +468,7 @@ bool FBasicTokenParser::GetToken(FBasicToken& Token, bool bNoConsts/* = false*/)
 		Token.TokenType = FBasicToken::TOKEN_Symbol;
 
 		// Lookup the token's global name.
-		Token.TokenName = FName( Token.Identifier, FNAME_Find, true );
+		Token.TokenName = FName(Token.Identifier, FNAME_Find);
 
 		return true;
 	}
diff --git a/Engine/Source/Editor/LandscapeEditor/Private/LandscapeEditorDetailCustomization_NewLandscape.cpp b/Engine/Source/Editor/LandscapeEditor/Private/LandscapeEditorDetailCustomization_NewLandscape.cpp
--- a/Engine/Source/Editor/LandscapeEditor/Private/LandscapeEditorDetailCustomization_NewLandscape.cpp
+++ b/Engine/Source/Editor/LandscapeEditor/Private/LandscapeEditorDetailCustomization_NewLandscape.cpp
@@ -1526,13 +1526,14 @@ FReply FLandscapeEditorStructCustomization_FLandscapeImportLayer::OnLayerFilenam
 
 bool FLandscapeEditorStructCustomization_FLandscapeImportLayer::ShouldFilterLayerInfo(const FAssetData& AssetData, FName LayerName)
 {
-	ULandscapeLayerInfoObject* LayerInfo = CastChecked<ULandscapeLayerInfoObject>(AssetData.GetAsset());
-	if (LayerInfo->LayerName != LayerName)
+	const FString* const LayerNameMetaData = AssetData.TagsAndValues.Find("LayerName");
+	if (LayerNameMetaData && !LayerNameMetaData->IsEmpty())
 	{
-		return true;
+		return FName(**LayerNameMetaData) != LayerName;
 	}
 
-	return false;
+	ULandscapeLayerInfoObject* LayerInfo = CastChecked<ULandscapeLayerInfoObject>(AssetData.GetAsset());
+	return LayerInfo->LayerName != LayerName;
 }
 
 EVisibility FLandscapeEditorStructCustomization_FLandscapeImportLayer::GetImportLayerCreateVisibility(TSharedRef<IPropertyHandle> PropertyHandle_LayerInfo)
diff --git a/Engine/Source/Editor/UnrealEd/Private/EditorObject.cpp b/Engine/Source/Editor/UnrealEd/Private/EditorObject.cpp
--- a/Engine/Source/Editor/UnrealEd/Private/EditorObject.cpp
+++ b/Engine/Source/Editor/UnrealEd/Private/EditorObject.cpp
@@ -300,7 +300,7 @@ static const TCHAR* ImportProperties(
 
 				// Create model.
 				UModelFactory* ModelFactory = NewObject<UModelFactory>();
-				ModelFactory->FactoryCreateText( UModel::StaticClass(), SubobjectRoot, FName(BrushName, FNAME_Add, true), RF_NoFlags, NULL, TEXT("t3d"), SourceText, SourceText+FCString::Strlen(SourceText), Warn );
+				ModelFactory->FactoryCreateText(UModel::StaticClass(), SubobjectRoot, FName(BrushName, FNAME_Add), RF_NoFlags, NULL, TEXT("t3d"), SourceText, SourceText+FCString::Strlen(SourceText), Warn);
 				ImportedBrush = 1;
 			}
 		}
diff --git a/Engine/Source/Editor/UnrealEd/Private/Factories/SkeletalMeshImport.cpp b/Engine/Source/Editor/UnrealEd/Private/Factories/SkeletalMeshImport.cpp
--- a/Engine/Source/Editor/UnrealEd/Private/Factories/SkeletalMeshImport.cpp
+++ b/Engine/Source/Editor/UnrealEd/Private/Factories/SkeletalMeshImport.cpp
@@ -209,7 +209,7 @@ bool ProcessImportMeshSkeleton(FReferenceSkeleton& RefSkeleton, int32& SkeletalD
 	{
 		const VBone & BinaryBone = RefBonesBinary[ b ];
 		const FString BoneName = FSkeletalMeshImportData::FixupBoneName( BinaryBone.Name );
-		const FMeshBoneInfo BoneInfo(FName(*BoneName, FNAME_Add, true), BinaryBone.Name, BinaryBone.ParentIndex);
+		const FMeshBoneInfo BoneInfo(FName(*BoneName, FNAME_Add), BinaryBone.Name, BinaryBone.ParentIndex);
 		const FTransform BoneTransform(BinaryBone.BonePos.Transform);
 
 		if(RefSkeleton.FindBoneIndex(BoneInfo.Name) != INDEX_NONE)
diff --git a/Engine/Source/Editor/UnrealEd/Private/PerformanceMonitor.cpp b/Engine/Source/Editor/UnrealEd/Private/PerformanceMonitor.cpp
--- a/Engine/Source/Editor/UnrealEd/Private/PerformanceMonitor.cpp
+++ b/Engine/Source/Editor/UnrealEd/Private/PerformanceMonitor.cpp
@@ -148,6 +148,7 @@ bool FPerformanceMonitor::WillAutoScalabilityHelp() const
 	IsAutoScaleLower |= NewLevels.PostProcessQuality < CurrentLevels.PostProcessQuality;
 	IsAutoScaleLower |= NewLevels.TextureQuality < CurrentLevels.TextureQuality;
 	IsAutoScaleLower |= NewLevels.EffectsQuality < CurrentLevels.EffectsQuality;
+	IsAutoScaleLower |= NewLevels.FoliageQuality < CurrentLevels.FoliageQuality;
 
 	// We don't check things like real-time, because the user may have enabled it temporarily.
 
@@ -167,6 +168,7 @@ Scalability::FQualityLevels FPerformanceMonitor::GetAutoScalabilityQualityLevels
 	NewLevels.PostProcessQuality	= FMath::Min(NewLevels.PostProcessQuality, ExistingLevels.PostProcessQuality);
 	NewLevels.TextureQuality		= FMath::Min(NewLevels.TextureQuality, ExistingLevels.TextureQuality);
 	NewLevels.EffectsQuality		= FMath::Min(NewLevels.EffectsQuality, ExistingLevels.EffectsQuality);
+	NewLevels.FoliageQuality		= FMath::Min(NewLevels.FoliageQuality, ExistingLevels.FoliageQuality);
 
 	return NewLevels;
 }
diff --git a/Engine/Source/Editor/UnrealEd/Private/SScalabilitySettings.cpp b/Engine/Source/Editor/UnrealEd/Private/SScalabilitySettings.cpp
--- a/Engine/Source/Editor/UnrealEd/Private/SScalabilitySettings.cpp
+++ b/Engine/Source/Editor/UnrealEd/Private/SScalabilitySettings.cpp
@@ -18,6 +18,7 @@ ECheckBoxState SScalabilitySettings::IsGroupQualityLevelSelected(const TCHAR* In
 	else if (FCString::Strcmp(InGroupName, TEXT("ShadowQuality")) == 0) QualityLevel = CachedQualityLevels.ShadowQuality;
 	else if (FCString::Strcmp(InGroupName, TEXT("TextureQuality")) == 0) QualityLevel = CachedQualityLevels.TextureQuality;
 	else if (FCString::Strcmp(InGroupName, TEXT("EffectsQuality")) == 0) QualityLevel = CachedQualityLevels.EffectsQuality;
+	else if (FCString::Strcmp(InGroupName, TEXT("FoliageQuality")) == 0) QualityLevel = CachedQualityLevels.FoliageQuality;
 
 	return (QualityLevel == InQualityLevel) ? ECheckBoxState::Checked : ECheckBoxState::Unchecked;
 }
@@ -31,6 +32,7 @@ void SScalabilitySettings::OnGroupQualityLevelChanged(ECheckBoxState NewState, c
 	else if (FCString::Strcmp(InGroupName, TEXT("ShadowQuality")) == 0) CachedQualityLevels.ShadowQuality = InQualityLevel;
 	else if (FCString::Strcmp(InGroupName, TEXT("TextureQuality")) == 0) CachedQualityLevels.TextureQuality = InQualityLevel;
 	else if (FCString::Strcmp(InGroupName, TEXT("EffectsQuality")) == 0) CachedQualityLevels.EffectsQuality = InQualityLevel;
+	else if (FCString::Strcmp(InGroupName, TEXT("FoliageQuality")) == 0) CachedQualityLevels.FoliageQuality = InQualityLevel;
 
 	Scalability::SetQualityLevels(CachedQualityLevels);
 	Scalability::SaveState(GEditorSettingsIni);
@@ -173,6 +175,7 @@ void SScalabilitySettings::Construct( const FArguments& InArgs )
 				+MakeGridSlot(0,5,5,1) [ SNew (SBorder).BorderImage(FEditorStyle::GetBrush("Scalability.RowBackground")) ]
 				+MakeGridSlot(0,6,5,1) [ SNew (SBorder).BorderImage(FEditorStyle::GetBrush("Scalability.RowBackground")) ]
 				+MakeGridSlot(0,7,5,1) [ SNew (SBorder).BorderImage(FEditorStyle::GetBrush("Scalability.RowBackground")) ]
+				+MakeGridSlot(0,8,5,1) [ SNew (SBorder).BorderImage(FEditorStyle::GetBrush("Scalability.RowBackground")) ]
 
 				+MakeGridSlot(0,0).VAlign(VAlign_Center) [ SNew(STextBlock).Text(LOCTEXT("QualityLabel", "Quality")).Font(TitleFont) ]
 				+MakeGridSlot(1,0) [ MakeHeaderButtonWidget(NamesLow, 0, LOCTEXT("QualityLow", "Set all groups to low quality")) ]
@@ -220,6 +223,12 @@ void SScalabilitySettings::Construct( const FArguments& InArgs )
 				+MakeGridSlot(2,7) [ MakeButtonWidget(NamesMedium, TEXT("EffectsQuality"), 1, LOCTEXT("EffectsQualityMed", "Set effects quality to medium")) ]
 				+MakeGridSlot(3,7) [ MakeButtonWidget(NamesHigh, TEXT("EffectsQuality"), 2, LOCTEXT("EffectsQualityHigh", "Set effects quality to high")) ]
 				+MakeGridSlot(4,7) [ MakeButtonWidget(NamesEpic, TEXT("EffectsQuality"), 3, LOCTEXT("EffectsQualityEpic", "Set effects quality to epic")) ]
+
+				+MakeGridSlot(0,8) [ SNew(STextBlock).Text(LOCTEXT("FoliageQualityLabel1", "Foliage")).Font(GroupFont) ]
+				+MakeGridSlot(1,8) [ MakeButtonWidget(NamesLow, TEXT("FoliageQuality"), 0, LOCTEXT("FoliageQualityLow", "Set foliage quality to low")) ]
+				+MakeGridSlot(2,8) [ MakeButtonWidget(NamesMedium, TEXT("FoliageQuality"), 1, LOCTEXT("FoliageQualityMed", "Set foliage quality to medium")) ]
+				+MakeGridSlot(3,8) [ MakeButtonWidget(NamesHigh, TEXT("FoliageQuality"), 2, LOCTEXT("FoliageQualityHigh", "Set foliage quality to high")) ]
+				+MakeGridSlot(4,8) [ MakeButtonWidget(NamesEpic, TEXT("FoliageQuality"), 3, LOCTEXT("FoliageQualityEpic", "Set foliage quality to epic")) ]
 			]
 
 			+ SVerticalBox::Slot()
diff --git a/Engine/Source/Editor/UnrealEd/Private/Settings/EditorSettings.cpp b/Engine/Source/Editor/UnrealEd/Private/Settings/EditorSettings.cpp
--- a/Engine/Source/Editor/UnrealEd/Private/Settings/EditorSettings.cpp
+++ b/Engine/Source/Editor/UnrealEd/Private/Settings/EditorSettings.cpp
@@ -53,6 +53,7 @@ void UEditorSettings::LoadScalabilityBenchmark()
 		GConfig->GetInt(Section, TEXT("PostProcessQuality"), Temporary.PostProcessQuality, GEditorSettingsIni);
 		GConfig->GetInt(Section, TEXT("TextureQuality"), Temporary.TextureQuality, GEditorSettingsIni);
 		GConfig->GetInt(Section, TEXT("EffectsQuality"), Temporary.EffectsQuality, GEditorSettingsIni);
+		GConfig->GetInt(Section, TEXT("FoliageQuality"), Temporary.FoliageQuality, GEditorSettingsIni);
 		EngineBenchmarkResult = Temporary;
 	}
 }
@@ -75,6 +76,7 @@ void UEditorSettings::AutoApplyScalabilityBenchmark()
 	GConfig->SetInt(Section, TEXT("PostProcessQuality"), Temporary.PostProcessQuality, GEditorSettingsIni);
 	GConfig->SetInt(Section, TEXT("TextureQuality"), Temporary.TextureQuality, GEditorSettingsIni);
 	GConfig->SetInt(Section, TEXT("EffectsQuality"), Temporary.EffectsQuality, GEditorSettingsIni);
+	GConfig->SetInt(Section, TEXT("FoliageQuality"), Temporary.FoliageQuality, GEditorSettingsIni);
 
 	Scalability::SetQualityLevels(Temporary);
 	Scalability::SaveState(GEditorSettingsIni);
diff --git a/Engine/Source/Programs/UnrealHeaderTool/Private/BaseParser.cpp b/Engine/Source/Programs/UnrealHeaderTool/Private/BaseParser.cpp
--- a/Engine/Source/Programs/UnrealHeaderTool/Private/BaseParser.cpp
+++ b/Engine/Source/Programs/UnrealHeaderTool/Private/BaseParser.cpp
@@ -387,7 +387,7 @@ bool FBaseParser::GetToken( FToken& Token, bool bNoConsts/*=false*/, ESymbolPars
 		Token.TokenType = TOKEN_Identifier;
 
 		// Lookup the token's global name.
-		Token.TokenName = FName( Token.Identifier, FNAME_Find, true );
+		Token.TokenName = FName(Token.Identifier, FNAME_Find);
 
 		// If const values are allowed, determine whether the identifier represents a constant
 		if ( !bNoConsts )
@@ -576,7 +576,7 @@ bool FBaseParser::GetToken( FToken& Token, bool bNoConsts/*=false*/, ESymbolPars
 		Token.TokenType = TOKEN_Symbol;
 
 		// Lookup the token's global name.
-		Token.TokenName = FName( Token.Identifier, FNAME_Find, true );
+		Token.TokenName = FName(Token.Identifier, FNAME_Find);
 
 		return true;
 	}
diff --git a/Engine/Source/Programs/UnrealHeaderTool/Private/HeaderParser.cpp b/Engine/Source/Programs/UnrealHeaderTool/Private/HeaderParser.cpp
--- a/Engine/Source/Programs/UnrealHeaderTool/Private/HeaderParser.cpp
+++ b/Engine/Source/Programs/UnrealHeaderTool/Private/HeaderParser.cpp
@@ -996,7 +996,7 @@ UField* FHeaderParser::FindField
 )
 {
 	check(InIdentifier);
-	FName InName( InIdentifier, FNAME_Find, true );
+	FName InName(InIdentifier, FNAME_Find);
 	if (InName != NAME_None)
 	{
 		for( ; Scope; Scope = Cast<UStruct>(Scope->GetOuter()) )
@@ -1293,13 +1293,13 @@ UEnum* FHeaderParser::CompileEnum()
 			case UEnum::ECppForm::Namespaced:
 			case UEnum::ECppForm::EnumClass:
 			{
-				NewTag = FName(*FString::Printf(TEXT("%s::%s"), EnumToken.Identifier, TagToken.Identifier), FNAME_Add, true);
+				NewTag = FName(*FString::Printf(TEXT("%s::%s"), EnumToken.Identifier, TagToken.Identifier), FNAME_Add);
 			}
 			break;
 
 			case UEnum::ECppForm::Regular:
 			{
-				NewTag = FName(TagToken.Identifier, FNAME_Add, true);
+				NewTag = FName(TagToken.Identifier, FNAME_Add);
 			}
 			break;
 		}
@@ -4227,7 +4227,7 @@ UProperty* FHeaderParser::GetVarNameAndDim
 	// If this is the first time seeing the property name, then flag it for replace instead of add
 	const EFindName FindFlag = VarProperty.PropertyFlags & CPF_Config ? GetFindFlagForPropertyName(VarProperty.Identifier) : FNAME_Add;
 	// create the FName for the property, splitting (ie Unnamed_3 -> Unnamed,3)
-	FName PropertyName(VarProperty.Identifier, FindFlag, true);
+	FName PropertyName(VarProperty.Identifier, FindFlag);
 
 	// Add property.
 	UProperty* NewProperty = NULL;
diff --git a/Engine/Source/Runtime/Core/Private/UObject/UnrealNames.cpp b/Engine/Source/Runtime/Core/Private/UObject/UnrealNames.cpp
--- a/Engine/Source/Runtime/Core/Private/UObject/UnrealNames.cpp
+++ b/Engine/Source/Runtime/Core/Private/UObject/UnrealNames.cpp
@@ -393,7 +393,7 @@ int32							FName::NumWideNames;
  * @param Name Value for the string portion of the name
  * @param FindType Action to take (see EFindName)
  */
-FName::FName( const WIDECHAR* Name, EFindName FindType, bool )
+FName::FName(const WIDECHAR* Name, EFindName FindType)
 {
 	if (Name)
 	{
@@ -405,7 +405,7 @@ FName::FName( const WIDECHAR* Name, EFindName FindType, bool )
 	}
 }
 
-FName::FName( const ANSICHAR* Name, EFindName FindType, bool )
+FName::FName(const ANSICHAR* Name, EFindName FindType)
 {
 	if (Name)
 	{
diff --git a/Engine/Source/Runtime/Core/Public/UObject/NameTypes.h b/Engine/Source/Runtime/Core/Public/UObject/NameTypes.h
--- a/Engine/Source/Runtime/Core/Public/UObject/NameTypes.h
+++ b/Engine/Source/Runtime/Core/Public/UObject/NameTypes.h
@@ -547,7 +547,7 @@ public:
 	 */
 	FORCEINLINE ANSICHAR const* GetPlainANSIString() const
 	{
-		return GetDisplayNameEntry()->GetAnsiName();		
+		return GetDisplayNameEntry()->GetAnsiName();
 	}
 
 	/**
@@ -681,7 +681,7 @@ public:
 	 * Takes an FName and checks to see that it follows the rules that Unreal requires for package or group names.
 	 *
 	 * @param	InReason		If the check fails, this string is filled in with the reason why.
-	 * @param	bIsGroupName	if true, check legallity for a group name, else check legality for a package name
+	 * @param	bIsGroupName	if true, check legality for a group name, else check legality for a package name
 	 *
 	 * @return	true if the name is valid
 	 */
@@ -715,7 +715,7 @@ public:
 	/**
 	 * Create an FName with a hardcoded string index.
 	 *
-	 * @param N The harcdcoded value the string portion of the name will have. The number portion will be NAME_NO_NUMBER
+	 * @param N The hardcoded value the string portion of the name will have. The number portion will be NAME_NO_NUMBER
 	 */
 	FORCEINLINE FName( EName N )
 		: ComparisonIndex( N )
@@ -730,7 +730,7 @@ public:
 	/**
 	 * Create an FName with a hardcoded string index and (instance).
 	 *
-	 * @param N The harcdcoded value the string portion of the name will have
+	 * @param N The hardcoded value the string portion of the name will have
 	 * @param InNumber The hardcoded value for the number portion of the name
 	 */
 	FORCEINLINE FName( EName N, int32 InNumber )
@@ -799,10 +799,13 @@ public:
 	 *
 	 * @param Name			Value for the string portion of the name
 	 * @param FindType		Action to take (see EFindName)
-	 * @param unused
 	 */
-	FName( const WIDECHAR*  Name, EFindName FindType=FNAME_Add, bool bUnused=true );
-	FName( const ANSICHAR* Name, EFindName FindType=FNAME_Add, bool bUnused=true );
+	FName(const WIDECHAR* Name, EFindName FindType=FNAME_Add);
+	FName(const ANSICHAR* Name, EFindName FindType=FNAME_Add);
+
+	// Deprecated bUnused
+	DEPRECATED(4.12, "Removed bUnused from FName") FName(const WIDECHAR* Name, EFindName FindType, bool bUnused) : FName(Name, FindType) { }
+	DEPRECATED(4.12, "Removed bUnused from FName") FName(const ANSICHAR* Name, EFindName FindType, bool bUnused) : FName(Name, FindType) { }
 
 	/**
 	 * Create an FName. If FindType is FNAME_Find, and the string part of the name 
@@ -825,7 +828,7 @@ public:
 	 * Create an FName with a hardcoded string index.
 	 *
 	 * @param HardcodedIndex	The hardcoded value the string portion of the name will have. 
-	 * @param Name				The hardcoded name to intialize
+	 * @param Name				The hardcoded name to initialize
 	 */
 	explicit FName( EName HardcodedIndex, const TCHAR* Name );
 
diff --git a/Engine/Source/Runtime/CoreUObject/Private/UObject/PropertyName.cpp b/Engine/Source/Runtime/CoreUObject/Private/UObject/PropertyName.cpp
--- a/Engine/Source/Runtime/CoreUObject/Private/UObject/PropertyName.cpp
+++ b/Engine/Source/Runtime/CoreUObject/Private/UObject/PropertyName.cpp
@@ -28,7 +28,7 @@ const TCHAR* UNameProperty::ImportText_Internal( const TCHAR* Buffer, void* Data
 {
 	if (!(PortFlags & PPF_Delimited))
 	{
-		*(FName*)Data = FName(Buffer, FNAME_Add, true);
+		*(FName*)Data = FName(Buffer, FNAME_Add);
 
 		// in order to indicate that the value was successfully imported, advance the buffer past the last character that was imported
 		Buffer += FCString::Strlen(Buffer);
@@ -40,7 +40,7 @@ const TCHAR* UNameProperty::ImportText_Internal( const TCHAR* Buffer, void* Data
 		if (!Buffer)
 			return NULL;
 
-		*(FName*)Data = FName(*Temp, FNAME_Add, true);
+		*(FName*)Data = FName(*Temp, FNAME_Add);
 	}
 	return Buffer;
 }
diff --git a/Engine/Source/Runtime/CoreUObject/Private/UObject/UObjectGlobals.cpp b/Engine/Source/Runtime/CoreUObject/Private/UObject/UObjectGlobals.cpp
--- a/Engine/Source/Runtime/CoreUObject/Private/UObject/UObjectGlobals.cpp
+++ b/Engine/Source/Runtime/CoreUObject/Private/UObject/UObjectGlobals.cpp
@@ -244,7 +244,7 @@ UObject* StaticFindObject( UClass* ObjectClass, UObject* InObjectPackage, const
 		return NULL;
 	}
 
-	FName ObjectName(*InName, FNAME_Add, true);
+	FName ObjectName(*InName, FNAME_Add);
 	return StaticFindObjectFast(ObjectClass, ObjectPackage, ObjectName, ExactClass, bAnyPackage);
 }
 
@@ -554,7 +554,7 @@ UPackage* CreatePackage( UObject* InOuter, const TCHAR* PackageName )
 		Result = FindObject<UPackage>( InOuter, *InName );
 		if( Result == NULL )
 		{
-			FName NewPackageName(*InName, FNAME_Add, true);
+			FName NewPackageName(*InName, FNAME_Add);
 			if (FPackageName::IsShortPackageName(NewPackageName))
 			{
 				UE_LOG(LogUObjectGlobals, Warning, TEXT("Attempted to create a package with a short package name: %s Outer: %s"), PackageName, InOuter ? *InOuter->GetFullName() : TEXT("NullOuter"));
diff --git a/Engine/Source/Runtime/Engine/Classes/Components/HierarchicalInstancedStaticMeshComponent.h b/Engine/Source/Runtime/Engine/Classes/Components/HierarchicalInstancedStaticMeshComponent.h
--- a/Engine/Source/Runtime/Engine/Classes/Components/HierarchicalInstancedStaticMeshComponent.h
+++ b/Engine/Source/Runtime/Engine/Classes/Components/HierarchicalInstancedStaticMeshComponent.h
@@ -96,6 +96,9 @@ class ENGINE_API UHierarchicalInstancedStaticMeshComponent : public UInstancedSt
 	UPROPERTY()
 	int32 NumBuiltInstances;
 
+	// Normally equal to NumBuiltInstances, but can be lower if density scaling is in effect
+	int32 NumBuiltRenderInstances;
+
 	// Bounding box of any built instances (cached from the ClusterTree)
 	UPROPERTY()
 	FBox BuiltInstanceBounds;
@@ -108,6 +111,14 @@ class ENGINE_API UHierarchicalInstancedStaticMeshComponent : public UInstancedSt
 	UPROPERTY()
 	TArray<FBox> UnbuiltInstanceBoundsList;
 
+	// Enable for detail meshes that don't really affect the game. Disable for anything important.
+	// Typically, this will be enabled for small meshes without collision (e.g. grass) and disabled for large meshes with collision (e.g. trees)
+	UPROPERTY()
+	uint32 bEnableDensityScaling:1;
+
+	// Which instances have been removed by foliage density scaling?
+	TBitArray<> ExcludedDueToDensityScaling;
+
 	// The number of nodes in the occlusion layer
 	UPROPERTY()
 	int32 OcclusionLayerNumNodes;
@@ -154,6 +165,8 @@ public:
 
 	virtual bool ShouldCreatePhysicsState() const override;
 
+	virtual int32 GetNumRenderInstances() const override { return NumBuiltRenderInstances + UnbuiltInstanceBoundsList.Num(); }
+
 	void BuildTree();
 	void BuildTreeAsync();
 	static void BuildTreeAnyThread(
diff --git a/Engine/Source/Runtime/Engine/Classes/Components/InstancedStaticMeshComponent.h b/Engine/Source/Runtime/Engine/Classes/Components/InstancedStaticMeshComponent.h
--- a/Engine/Source/Runtime/Engine/Classes/Components/InstancedStaticMeshComponent.h
+++ b/Engine/Source/Runtime/Engine/Classes/Components/InstancedStaticMeshComponent.h
@@ -223,6 +223,9 @@ public:
 	 */
 	void ReleasePerInstanceRenderData();
 
+	// Number of instances in the render-side instance buffer
+	virtual int32 GetNumRenderInstances() const { return PerInstanceSMData.Num() + RemovedInstances.Num(); }
+
 private:
 	/** Creates body instances for all instances owned by this component */
 	void CreateAllInstanceBodies();
@@ -232,7 +235,7 @@ private:
 
 	/** Sets up new instance data to sensible defaults, creates physics counterparts if possible */
 	void SetupNewInstanceData(FInstancedStaticMeshInstanceData& InOutNewInstanceData, int32 InInstanceIndex, const FTransform& InInstanceTransform);
-	
+
 protected:
 	/** Request to navigation system to update only part of navmesh occupied by specified instance */
 	virtual void PartialNavigationUpdate(int32 InstanceIdx);
diff --git a/Engine/Source/Runtime/Engine/Classes/Components/SplineMeshComponent.h b/Engine/Source/Runtime/Engine/Classes/Components/SplineMeshComponent.h
--- a/Engine/Source/Runtime/Engine/Classes/Components/SplineMeshComponent.h
+++ b/Engine/Source/Runtime/Engine/Classes/Components/SplineMeshComponent.h
@@ -306,7 +306,7 @@ class ENGINE_API USplineMeshComponent : public UStaticMeshComponent, public IInt
 
 	// Destroys the body setup, used to clear collision if the mesh goes missing
 	void DestroyBodySetup();
-#if WITH_EDITOR
+#if WITH_EDITOR || WITH_RUNTIME_PHYSICS_COOKING
 	// Builds collision for the spline mesh (if collision is enabled)
 	void RecreateCollision();
 #endif
diff --git a/Engine/Source/Runtime/Engine/Classes/GameFramework/GameUserSettings.h b/Engine/Source/Runtime/Engine/Classes/GameFramework/GameUserSettings.h
--- a/Engine/Source/Runtime/Engine/Classes/GameFramework/GameUserSettings.h
+++ b/Engine/Source/Runtime/Engine/Classes/GameFramework/GameUserSettings.h
@@ -202,6 +202,14 @@ public:
 	// Returns the post-processing quality (0..3, higher is better)
 	UFUNCTION(BlueprintCallable, Category=Settings)
 	int32 GetPostProcessingQuality() const;
+	
+	// Sets the post-processing quality (0..3, higher is better)
+	UFUNCTION(BlueprintCallable, Category=Settings)
+	void SetFoliageQuality(int32 Value);
+
+	// Returns the post-processing quality (0..3, higher is better)
+	UFUNCTION(BlueprintCallable, Category=Settings)
+	int32 GetFoliageQuality() const;
 
 	/** Checks if any user settings is different from current */
 	UFUNCTION(BlueprintPure, Category=Settings)
diff --git a/Engine/Source/Runtime/Engine/Private/ChartCreation.cpp b/Engine/Source/Runtime/Engine/Private/ChartCreation.cpp
--- a/Engine/Source/Runtime/Engine/Private/ChartCreation.cpp
+++ b/Engine/Source/Runtime/Engine/Private/ChartCreation.cpp
@@ -219,6 +219,7 @@ protected:
 		PrintToEndpoint(FString::Printf(TEXT("\tPost-Process Quality: %d"), ScalabilityQuality.PostProcessQuality));
 		PrintToEndpoint(FString::Printf(TEXT("\tTexture Quality: %d"), ScalabilityQuality.TextureQuality));
 		PrintToEndpoint(FString::Printf(TEXT("\tEffects Quality: %d"), ScalabilityQuality.EffectsQuality));
+		PrintToEndpoint(FString::Printf(TEXT("\tFoliage Quality: %d"), ScalabilityQuality.FoliageQuality));
 
 		PrintToEndpoint(FString::Printf(TEXT("%i frames collected over %4.2f seconds, disregarding %4.2f seconds for a %4.2f FPS average"),
 			NumFrames,
@@ -450,6 +451,7 @@ protected:
 		ParamArray.Add(FAnalyticsEventAttribute(TEXT("PostProcessQuality"), ScalabilityQuality.PostProcessQuality));
 		ParamArray.Add(FAnalyticsEventAttribute(TEXT("TextureQuality"), ScalabilityQuality.TextureQuality));
 		ParamArray.Add(FAnalyticsEventAttribute(TEXT("FXQuality"), ScalabilityQuality.EffectsQuality));
+		ParamArray.Add(FAnalyticsEventAttribute(TEXT("FoliageQuality"), ScalabilityQuality.FoliageQuality));
 
 		ParamArray.Add(FAnalyticsEventAttribute(TEXT("AvgFPS"), FString::Printf(TEXT("%4.2f"), AvgFPS)));
 		ParamArray.Add(FAnalyticsEventAttribute(TEXT("TimeDisregarded"), FString::Printf(TEXT("%4.2f"), TimeDisregarded)));
@@ -579,6 +581,7 @@ protected:
 		FPSChartRow = FPSChartRow.Replace(TEXT("TOKEN_SETTINGS_PP"), *FString::Printf(TEXT("%d"), ScalabilityQuality.PostProcessQuality), ESearchCase::CaseSensitive);
 		FPSChartRow = FPSChartRow.Replace(TEXT("TOKEN_SETTINGS_TEX"), *FString::Printf(TEXT("%d"), ScalabilityQuality.TextureQuality), ESearchCase::CaseSensitive);
 		FPSChartRow = FPSChartRow.Replace(TEXT("TOKEN_SETTINGS_FX"), *FString::Printf(TEXT("%d"), ScalabilityQuality.EffectsQuality), ESearchCase::CaseSensitive);
+		FPSChartRow = FPSChartRow.Replace(TEXT("TOKEN_SETTINGS_FLG"), *FString::Printf(TEXT("%d"), ScalabilityQuality.FoliageQuality), ESearchCase::CaseSensitive);
 
 		FPSChartRow = FPSChartRow.Replace(TEXT("TOKEN_AVG_FPS"), *FString::Printf(TEXT("%4.2f"), AvgFPS), ESearchCase::CaseSensitive);
 		FPSChartRow = FPSChartRow.Replace(TEXT("TOKEN_TIME_DISREGARDED"), *FString::Printf(TEXT("%4.2f"), TimeDisregarded), ESearchCase::CaseSensitive);
diff --git a/Engine/Source/Runtime/Engine/Private/Components/SplineMeshComponent.cpp b/Engine/Source/Runtime/Engine/Private/Components/SplineMeshComponent.cpp
--- a/Engine/Source/Runtime/Engine/Private/Components/SplineMeshComponent.cpp
+++ b/Engine/Source/Runtime/Engine/Private/Components/SplineMeshComponent.cpp
@@ -504,13 +504,11 @@ void USplineMeshComponent::UpdateRenderStateAndCollision()
 {
 	MarkRenderStateDirty();
 
-#if WITH_EDITOR
-	if (!GetWorld()->AreActorsInitialized())
-	{
-		DestroyBodySetup();
-		RecreatePhysicsState();
-	}
-#endif // WITH_EDITOR
+#if WITH_EDITOR || WITH_RUNTIME_PHYSICS_COOKING
+	DestroyBodySetup();
+	RecreateCollision();
+	RecreatePhysicsState();
+#endif // WITH_EDITOR || WITH_RUNTIME_PHYSICS_COOKING
 
 	bMeshDirty = false;
 }
@@ -886,7 +884,7 @@ void USplineMeshComponent::GetMeshId(FString& OutMeshId)
 
 void USplineMeshComponent::CreatePhysicsState()
 {
-#if WITH_EDITOR
+#if WITH_EDITOR || WITH_RUNTIME_PHYSICS_COOKING
 	// With editor code we can recreate the collision if the mesh changes
 	const FGuid MeshBodySetupGuid = (StaticMesh != NULL ? StaticMesh->BodySetup->BodySetupGuid : FGuid());
 	if (CachedMeshBodySetupGuid != MeshBodySetupGuid)
@@ -976,7 +974,7 @@ void USplineMeshComponent::DestroyBodySetup()
 }
 
 
-#if WITH_EDITOR
+#if WITH_EDITOR || WITH_RUNTIME_PHYSICS_COOKING
 void USplineMeshComponent::RecreateCollision()
 {
 	if (StaticMesh && IsCollisionEnabled())
diff --git a/Engine/Source/Runtime/Engine/Private/GameUserSettings.cpp b/Engine/Source/Runtime/Engine/Private/GameUserSettings.cpp
--- a/Engine/Source/Runtime/Engine/Private/GameUserSettings.cpp
+++ b/Engine/Source/Runtime/Engine/Private/GameUserSettings.cpp
@@ -649,6 +649,16 @@ int32 UGameUserSettings::GetPostProcessingQuality() const
 	return ScalabilityQuality.PostProcessQuality;
 }
 
+void UGameUserSettings::SetFoliageQuality(int32 Value)
+{
+	ScalabilityQuality.FoliageQuality = FMath::Clamp(Value, 0, 3);
+}
+
+int32 UGameUserSettings::GetFoliageQuality() const
+{
+	return ScalabilityQuality.FoliageQuality;
+}
+
 UGameUserSettings* UGameUserSettings::GetGameUserSettings()
 {
 	return GEngine->GetGameUserSettings();
diff --git a/Engine/Source/Runtime/Engine/Private/HierarchicalInstancedStaticMesh.cpp b/Engine/Source/Runtime/Engine/Private/HierarchicalInstancedStaticMesh.cpp
--- a/Engine/Source/Runtime/Engine/Private/HierarchicalInstancedStaticMesh.cpp
+++ b/Engine/Source/Runtime/Engine/Private/HierarchicalInstancedStaticMesh.cpp
@@ -34,7 +34,7 @@ static TAutoConsoleVariable<int32> CVarDisableCull(
 static TAutoConsoleVariable<int32> CVarCullAll(
 	TEXT("foliage.CullAll"),
 	0,
-	TEXT("If greater than zero, everything is conside3red culled."));
+	TEXT("If greater than zero, everything is considered culled."));
 
 static TAutoConsoleVariable<int32> CVarDitheredLOD(
 	TEXT("foliage.DitheredLOD"),
@@ -86,6 +86,11 @@ static TAutoConsoleVariable<int32> CVarMinInstancesPerOcclusionQuery(
 	256,
 	TEXT("Controls the granualrity of occlusion culling. 1024 to 65536 is a reasonable range. This is not exact, actual minimum might be off by a factor of two."));
 
+static TAutoConsoleVariable<float> CVarFoliageDensityScale(
+	TEXT("foliage.DensityScale"),
+	1.0,
+	TEXT("Controls the amount of foliage to render. Foliage must opt-in to density scaling through the foliage type."));
+
 DECLARE_CYCLE_STAT(TEXT("Traversal Time"),STAT_FoliageTraversalTime,STATGROUP_Foliage);
 DECLARE_CYCLE_STAT(TEXT("Build Time"), STAT_FoliageBuildTime, STATGROUP_Foliage);
 DECLARE_CYCLE_STAT(TEXT("Batch Time"),STAT_FoliageBatchTime,STATGROUP_Foliage);
@@ -97,6 +102,53 @@ DECLARE_DWORD_COUNTER_STAT(TEXT("Occlusion Culled Instances"), STAT_OcclusionCul
 DECLARE_DWORD_COUNTER_STAT(TEXT("Traversals"),STAT_FoliageTraversals,STATGROUP_Foliage);
 DECLARE_MEMORY_STAT(TEXT("Instance Buffers"),STAT_FoliageInstanceBuffers,STATGROUP_Foliage);
 
+static void FoliageCVarSinkFunction()
+{
+	static float CachedFoliageDensityScale = 1.0f;
+	float FoliageDensityScale = CVarFoliageDensityScale.GetValueOnGameThread();
+
+	if (FoliageDensityScale != CachedFoliageDensityScale)
+	{
+		CachedFoliageDensityScale = FoliageDensityScale;
+
+		for (auto* Component : TObjectRange<UHierarchicalInstancedStaticMeshComponent>(RF_ClassDefaultObject | RF_ArchetypeObject, true, EInternalObjectFlags::PendingKill))
+		{
+			if (Component->bEnableDensityScaling && Component->GetWorld() && Component->GetWorld()->IsGameWorld())
+			{
+				if (FoliageDensityScale == 0)
+				{
+					// exclude all instances
+					Component->ExcludedDueToDensityScaling.Init(true, Component->PerInstanceSMData.Num());
+				}
+				else if (FoliageDensityScale > 0.0f && FoliageDensityScale < 1.0f)
+				{
+					FRandomStream Rand(Component->InstancingRandomSeed);
+					if (Component->ExcludedDueToDensityScaling.Num() == 0)
+					{
+						// TBitArray<> doesn't have SetNumUninitialized...
+						Component->ExcludedDueToDensityScaling.Init(false, Component->PerInstanceSMData.Num());
+					}
+
+					for (int32 i = 0; i < Component->ExcludedDueToDensityScaling.Num(); ++i)
+					{
+						Component->ExcludedDueToDensityScaling[i] = (Rand.FRand() > FoliageDensityScale);
+					}
+				}
+				else
+				{
+					// Show all instances
+					Component->ExcludedDueToDensityScaling.Empty();
+				}
+
+				Component->BuildTree();
+				Component->MarkRenderStateDirty();
+			}
+		}
+	}
+}
+
+static FAutoConsoleVariableSink CVarFoliageSink(FConsoleCommandDelegate::CreateStatic(&FoliageCVarSinkFunction));
+
 struct FClusterTree
 {
 	TArray<FClusterNode> Nodes;
@@ -108,6 +160,7 @@ struct FClusterTree
 
 class FClusterBuilder
 {
+	int32 OriginalNum;
 	int32 Num;
 	FBox InstBox;
 	int32 BranchingFactor;
@@ -151,18 +204,18 @@ class FClusterBuilder
 
 	void Split(int32 InNum)
 	{
-		check(InNum);
+		checkSlow(InNum);
 		Clusters.Reset();
 		Split(0, InNum - 1);
 		Clusters.Sort();
-		check(Clusters.Num());
+		checkSlow(Clusters.Num() > 0);
 		int32 At = 0;
 		for (auto& Cluster : Clusters)
 		{
-			check(At == Cluster.Start);
+			checkSlow(At == Cluster.Start);
 			At += Cluster.Num;
 		}
-		check(At == InNum);
+		checkSlow(At == InNum);
 	}
 
 	void Split(int32 Start, int32 End)
@@ -178,7 +231,7 @@ class FClusterBuilder
 			Clusters.Add(FRunPair(Start, NumRange));
 			return;
 		}
-		check(NumRange >= 2);
+		checkSlow(NumRange >= 2);
 		SortPairs.Reset();
 		int32 BestAxis = -1;
 		float BestAxisValue = -1.0f;
@@ -221,22 +274,42 @@ class FClusterBuilder
 				StartRight--;
 			}
 		}
-		check(EndLeft + 1 == StartRight);
-		check(EndLeft >= Start);
-		check(End >= StartRight);
+		checkSlow(EndLeft + 1 == StartRight);
+		checkSlow(EndLeft >= Start);
+		checkSlow(End >= StartRight);
 
 		Split(Start, EndLeft);
 		Split(StartRight, End);
-
 	}
 public:
 	FClusterTree* Result;
 
-	FClusterBuilder(TArray<FMatrix>& InTransforms, const FBox& InInstBox, int32 InMaxInstancesPerLeaf)
-		: Num(InTransforms.Num())
+	FClusterBuilder(TArray<FMatrix> InTransforms, const FBox& InInstBox, int32 InMaxInstancesPerLeaf, TBitArray<> ExcludedDueToDensityScaling = TBitArray<>())
+		: OriginalNum(InTransforms.Num())
 		, InstBox(InInstBox)
+		, Transforms(MoveTemp(InTransforms))
 		, Result(nullptr)
 	{
+		SortPoints.AddUninitialized(OriginalNum);
+		for (int32 Index = 0; Index < OriginalNum; Index++)
+		{
+			SortPoints[Index] = Transforms[Index].GetOrigin();
+		}
+
+		for (int32 Index = 0; Index < ExcludedDueToDensityScaling.Num(); ++Index)
+		{
+			if (!ExcludedDueToDensityScaling[Index])
+			{
+				SortIndex.Add(Index);
+			}
+		}
+		for (int32 Index = ExcludedDueToDensityScaling.Num(); Index < OriginalNum; ++Index)
+		{
+			SortIndex.Add(Index);
+		}
+
+		Num = SortIndex.Num();
+
 		OcclusionLayerTarget = CVarMaxOcclusionQueriesPerComponent.GetValueOnAnyThread();
 		int32 MinInstancesPerOcclusionQuery = CVarMinInstancesPerOcclusionQuery.GetValueOnAnyThread();
 
@@ -250,15 +323,6 @@ public:
 		}
 		InternalNodeBranchingFactor = CVarFoliageSplitFactor.GetValueOnAnyThread();
 		MaxInstancesPerLeaf = InMaxInstancesPerLeaf;
-
-		Exchange(Transforms, InTransforms);
-		SortIndex.AddUninitialized(Num);
-		SortPoints.AddUninitialized(Num);
-		for (int32 Index = 0; Index < Num; Index++)
-		{
-			SortIndex[Index] = Index;
-			SortPoints[Index] = Transforms[Index].GetOrigin();
-		}
 	}
 
 	void BuildAsync(ENamedThreads::Type CurrentThread, const FGraphEventRef& MyCompletionGraphEvent)
@@ -270,6 +334,16 @@ public:
 	void Build()
 	{
 		Result = new FClusterTree;
+
+		if (Num == 0)
+		{
+			// Can happen if all instances are excluded due to scalability
+			// It doesn't only happen with a scalability factor of 0 - 
+			// even with a scalability factor of 0.99, if there's only one instance of this type you can end up with Num == 0 if you're unlucky
+			Result->InstanceReorderTable.Init(INDEX_NONE, OriginalNum);
+			return;
+		}
+
 		bool bIsOcclusionLayer = false;
 		BranchingFactor = MaxInstancesPerLeaf;
 		if (BranchingFactor > 2 && OcclusionLayerTarget && Num / BranchingFactor <= OcclusionLayerTarget)
@@ -384,7 +458,7 @@ public:
 					Node.LastInstance = InverseInstanceIndex[Node.LastInstance];
 				}
 				OldInstanceIndex.Reset();
-				Exchange(OldInstanceIndex, SortedInstances);
+				Swap(OldInstanceIndex, SortedInstances);
 				SortedInstances.AddUninitialized(Num);
 				for (int32 Index = 0; Index < Num; Index++)
 				{
@@ -435,7 +509,7 @@ public:
 						LevelIndex++;
 					}
 				}
-				check(LevelStarts[LevelStarts.Num() - 1] == NewNum);
+				checkSlow(LevelStarts[LevelStarts.Num() - 1] == NewNum);
 				InverseChildIndex.Reset();
 				// InverseChildIndex[old index] == new index
 				InverseChildIndex.AddUninitialized(NewNum);
@@ -453,7 +527,7 @@ public:
 					}
 				}
 				{
-					Exchange(OldNodes, Result->Nodes);
+					Swap(OldNodes, Result->Nodes);
 					Result->Nodes.Empty(NewNum);
 					for (int32 Index = 0; Index < Clusters.Num(); Index++)
 					{
@@ -474,10 +548,10 @@ public:
 					OldIndex += Clusters[Index].Num;
 					Node.LastChild = OldIndex - 1;
 					Node.FirstInstance = Result->Nodes[Node.FirstChild].FirstInstance;
-					check(Node.FirstInstance == InstanceTracker);
+					checkSlow(Node.FirstInstance == InstanceTracker);
 					Node.LastInstance = Result->Nodes[Node.LastChild].LastInstance;
 					InstanceTracker = Node.LastInstance + 1;
-					check(InstanceTracker <= Num);
+					checkSlow(InstanceTracker <= Num);
 					FBox NodeBox(ForceInit);
 					for (int32 ChildIndex = Node.FirstChild; ChildIndex <= Node.LastChild; ChildIndex++)
 					{
@@ -494,7 +568,7 @@ public:
 		}
 
 		// Save inverse map
-		Result->InstanceReorderTable.AddUninitialized(Num);
+		Result->InstanceReorderTable.Init(INDEX_NONE, OriginalNum);
 		for (int32 Index = 0; Index < Num; Index++)
 		{
 			Result->InstanceReorderTable[SortedInstances[Index]] = Index;
@@ -641,7 +715,7 @@ struct FFoliageCullInstanceParams;
 
 class FHierarchicalStaticMeshSceneProxy : public FInstancedStaticMeshSceneProxy
 {
-	TSharedPtr<TArray<FClusterNode>, ESPMode::ThreadSafe> ClusterTreePtr;
+	TSharedRef<TArray<FClusterNode>, ESPMode::ThreadSafe> ClusterTreePtr;
 	const TArray<FClusterNode>& ClusterTree;
 
 	TArray<FBox> UnbuiltBounds;
@@ -667,11 +741,11 @@ public:
 
 	FHierarchicalStaticMeshSceneProxy(bool bInIsGrass, UHierarchicalInstancedStaticMeshComponent* InComponent, ERHIFeatureLevel::Type InFeatureLevel)
 		: FInstancedStaticMeshSceneProxy(InComponent, InFeatureLevel)
-		, ClusterTreePtr(InComponent->ClusterTreePtr)
+		, ClusterTreePtr(InComponent->ClusterTreePtr.ToSharedRef())
 		, ClusterTree(*InComponent->ClusterTreePtr)
 		, UnbuiltBounds(InComponent->UnbuiltInstanceBoundsList)
-		, FirstUnbuiltIndex(InComponent->NumBuiltInstances)
-		, LastUnbuiltIndex(InComponent->PerInstanceSMData.Num()+InComponent->RemovedInstances.Num()-1)
+		, FirstUnbuiltIndex(InComponent->NumBuiltRenderInstances)
+		, LastUnbuiltIndex(InComponent->GetNumRenderInstances()-1)
 		, bIsGrass(bInIsGrass)
 		, SceneProxyCreatedFrameNumberRenderThread(UINT32_MAX)
 #if !(UE_BUILD_SHIPPING || UE_BUILD_TEST)
@@ -679,24 +753,22 @@ public:
 #endif
 	{
 		check(InComponent->InstanceReorderTable.Num() == InComponent->PerInstanceSMData.Num());
-		check(UnbuiltBounds.Num() == (LastUnbuiltIndex - FirstUnbuiltIndex + 1));
 		SetupOcclusion(InComponent);
 	}
 
 	FHierarchicalStaticMeshSceneProxy(bool bInIsGrass, UHierarchicalInstancedStaticMeshComponent* InComponent, ERHIFeatureLevel::Type InFeatureLevel, FStaticMeshInstanceData& Other)
 		: FInstancedStaticMeshSceneProxy(InComponent, InFeatureLevel, Other)
-		, ClusterTreePtr(InComponent->ClusterTreePtr)
+		, ClusterTreePtr(InComponent->ClusterTreePtr.ToSharedRef())
 		, ClusterTree(*InComponent->ClusterTreePtr)
 		, UnbuiltBounds(InComponent->UnbuiltInstanceBoundsList)
-		, FirstUnbuiltIndex(InComponent->NumBuiltInstances)
-		, LastUnbuiltIndex(InstancedRenderData.NumInstances+InComponent->RemovedInstances.Num()-1)
+		, FirstUnbuiltIndex(InComponent->NumBuiltRenderInstances)
+		, LastUnbuiltIndex(InComponent->GetNumRenderInstances()-1)
 		, bIsGrass(bInIsGrass)
 #if !(UE_BUILD_SHIPPING || UE_BUILD_TEST)
 		, CaptureTag(0)
 #endif
 	{
 		check(!bInIsGrass || (!InComponent->InstanceReorderTable.Num() && !InComponent->PerInstanceSMData.Num()));
-		check(UnbuiltBounds.Num() == (LastUnbuiltIndex - FirstUnbuiltIndex + 1));
 		SetupOcclusion(InComponent);
 	}
 
@@ -1654,7 +1726,9 @@ UHierarchicalInstancedStaticMeshComponent::UHierarchicalInstancedStaticMeshCompo
 	, ClusterTreePtr(MakeShareable(new TArray<FClusterNode>))
 	, WriteOncePrebuiltInstanceBuffer(/*NeedsCPUAccess*/ false, /*bSupportsVertexHalfFloat*/ GVertexElementTypeSupport.IsSupported(VET_Half2))
 	, NumBuiltInstances(0)
+	, NumBuiltRenderInstances(0)
 	, UnbuiltInstanceBounds(0)
+	, bEnableDensityScaling(false)
 	, bIsAsyncBuilding(false)
 	, bDiscardAsyncBuildResults(false)
 	, bConcurrentRemoval(false)
@@ -1719,7 +1793,11 @@ void UHierarchicalInstancedStaticMeshComponent::RemoveInstanceInternal(int32 Ins
 	PartialNavigationUpdate(InstanceIndex);
 
 	// Save the render index
-	RemovedInstances.Add(InstanceReorderTable[InstanceIndex]);
+	const int32 RemovedRenderIndex = InstanceReorderTable[InstanceIndex];
+	if (RemovedRenderIndex != INDEX_NONE)
+	{
+		RemovedInstances.Add(RemovedRenderIndex);
+	}
 
 	// Remove the instance
 	PerInstanceSMData.RemoveAtSwap(InstanceIndex);
@@ -1874,25 +1952,15 @@ bool UHierarchicalInstancedStaticMeshComponent::UpdateInstanceTransform(int32 In
 		bConcurrentRemoval = true;
 	}
 
-	const int32 OldReorderIndex = InstanceReorderTable[InstanceIndex];
+	int32 RenderIndex = InstanceReorderTable[InstanceIndex];
 	const FMatrix OldTransform = PerInstanceSMData[InstanceIndex].Transform;
 	const FTransform NewLocalTransform = bWorldSpace ? NewInstanceTransform.GetRelativeTransform(ComponentToWorld) : NewInstanceTransform;
 	const FVector NewLocalLocation = NewLocalTransform.GetTranslation();
 
 	// if we are only updating rotation/scale we update the instance directly in the cluster tree
-	const bool bIsBuiltInstance = OldReorderIndex < NumBuiltInstances;
-	const bool bDoInPlaceUpdate = !bIsBuiltInstance || NewLocalLocation.Equals(OldTransform.GetOrigin());
-
-	// If we're updating an instance in the tree and can't do it in-place, we have to remove it and re-add it
-	if (bIsBuiltInstance && !bDoInPlaceUpdate)
-	{
-		// Allocate a new instance render order ID, rendered last
-		const int32 NewReorderIndex = PerInstanceSMData.Num() + RemovedInstances.Num();
-		InstanceReorderTable[InstanceIndex] = NewReorderIndex;
-
-		// Treat the old instance render data like a removal.
-		RemovedInstances.Add(OldReorderIndex);
-	}
+	const bool bIsOmittedInstance = (RenderIndex == INDEX_NONE);
+	const bool bIsBuiltInstance = !bIsOmittedInstance && RenderIndex < NumBuiltRenderInstances;
+	const bool bDoInPlaceUpdate = !bIsOmittedInstance && (!bIsBuiltInstance || NewLocalLocation.Equals(OldTransform.GetOrigin()));
 
 	bool Result = Super::UpdateInstanceTransform(InstanceIndex, NewInstanceTransform, bWorldSpace, bMarkRenderStateDirty);
 
@@ -1909,25 +1977,35 @@ bool UHierarchicalInstancedStaticMeshComponent::UpdateInstanceTransform(int32 In
 				if (!OldInstanceBounds.IsInside(NewInstanceBounds))
 				{
 					BuiltInstanceBounds += NewInstanceBounds;
-					UpdateInstanceTreeBoundsInternal(OldReorderIndex, NewInstanceBounds);
+					UpdateInstanceTreeBoundsInternal(RenderIndex, NewInstanceBounds);
 				}
 			}
 			else
 			{
 				UnbuiltInstanceBounds += NewInstanceBounds;
-				UnbuiltInstanceBoundsList[InstanceIndex - NumBuiltInstances] = NewInstanceBounds;
+				UnbuiltInstanceBoundsList[RenderIndex - NumBuiltRenderInstances] = NewInstanceBounds;
 			}
 		}
 		else
 		{
+			// If we can't update an instance in-place, we have to remove it and re-add it
+			if (!bIsOmittedInstance)
+			{
+				RemovedInstances.Add(RenderIndex);
+			}
+
+			// Allocate a new instance render order ID, rendered last (it is now an unbuilt instance)
+			RenderIndex = GetNumRenderInstances();
+			InstanceReorderTable[InstanceIndex] = RenderIndex;
+
 			UnbuiltInstanceBounds += NewInstanceBounds;
 			UnbuiltInstanceBoundsList.Add(NewInstanceBounds);
 		}
-	}
 
-	if (!bDoInPlaceUpdate && !IsAsyncBuilding())
-	{
-		BuildTreeAsync();
+		if (!bDoInPlaceUpdate && !IsAsyncBuilding())
+		{
+			BuildTreeAsync();
+		}
 	}
 
 	return Result;
@@ -1937,9 +2015,6 @@ int32 UHierarchicalInstancedStaticMeshComponent::AddInstance(const FTransform& I
 {
 	int32 InstanceIndex = UInstancedStaticMeshComponent::AddInstance(InstanceTransform);
 
-	// Need to offset the newly added instance's RenderIndex by the amount that will be adjusted at the end of the frame
-	InstanceReorderTable.Add(InstanceIndex + RemovedInstances.Num());
-
 	if (PerInstanceSMData.Num() == 1)
 	{
 		BuildTree();
@@ -1953,6 +2028,9 @@ int32 UHierarchicalInstancedStaticMeshComponent::AddInstance(const FTransform& I
 
 		if (StaticMesh)
 		{
+			// Need to offset the newly added instance's RenderIndex by the amount that will be adjusted at the end of the frame
+			InstanceReorderTable.Add(GetNumRenderInstances());
+
 			const FBox NewInstanceBounds = StaticMesh->GetBounds().GetBox().TransformBy(InstanceTransform);
 			UnbuiltInstanceBounds += NewInstanceBounds;
 			UnbuiltInstanceBoundsList.Add(NewInstanceBounds);
@@ -1971,6 +2049,7 @@ void UHierarchicalInstancedStaticMeshComponent::ClearInstances()
 
 	ClusterTreePtr = MakeShareable(new TArray<FClusterNode>);
 	NumBuiltInstances = 0;
+	NumBuiltRenderInstances = 0;
 	SortedInstances.Empty();
 	UnbuiltInstanceBounds.Init();
 	UnbuiltInstanceBoundsList.Empty();
@@ -2060,13 +2139,13 @@ void UHierarchicalInstancedStaticMeshComponent::BuildTree()
 		// @todo: Level error or something to let LDs know this
 		1;//StaticMesh->LODModels(0).Elements.Num() == 1;
 
-	if(bMeshIsValid)
+	if (bMeshIsValid)
 	{
 		//double StartTime = FPlatformTime::Seconds();
 		// If we don't have a random seed for this instanced static mesh component yet, then go ahead and
 		// generate one now.  This will be saved with the static mesh component and used for future generation
 		// of random numbers for this component's instances. (Used by the PerInstanceRandom material expression)
-		while( InstancingRandomSeed == 0 )
+		while (InstancingRandomSeed == 0)
 		{
 			InstancingRandomSeed = FMath::Rand();
 		}
@@ -2078,20 +2157,20 @@ void UHierarchicalInstancedStaticMeshComponent::BuildTree()
 			InstanceTransforms[Index] = PerInstanceSMData[Index].Transform;
 		}
 
-		TUniquePtr<FClusterBuilder> Builder(new FClusterBuilder(InstanceTransforms, StaticMesh->GetBounds().GetBox(), DesiredInstancesPerLeaf()));
-		Builder->Build();
+		FClusterBuilder Builder(InstanceTransforms, StaticMesh->GetBounds().GetBox(), DesiredInstancesPerLeaf(), ExcludedDueToDensityScaling);
+		Builder.Build();
 
-		NumBuiltInstances = Builder->Result->InstanceReorderTable.Num();
-		OcclusionLayerNumNodes = Builder->Result->OutOcclusionLayerNum;
+		NumBuiltInstances = Builder.Result->InstanceReorderTable.Num();
+		NumBuiltRenderInstances = Builder.Result->SortedInstances.Num();
+		OcclusionLayerNumNodes = Builder.Result->OutOcclusionLayerNum;
 		UnbuiltInstanceBounds.Init();
 		RemovedInstances.Empty();
 		UnbuiltInstanceBoundsList.Empty();
-		BuiltInstanceBounds = (Builder->Result->Nodes.Num() > 0 ? FBox(Builder->Result->Nodes[0].BoundMin, Builder->Result->Nodes[0].BoundMax) : FBox(0));
+		BuiltInstanceBounds = (Builder.Result->Nodes.Num() > 0 ? FBox(Builder.Result->Nodes[0].BoundMin, Builder.Result->Nodes[0].BoundMax) : FBox(0));
 
-		ClusterTreePtr = MakeShareable(new TArray<FClusterNode>);
-		Exchange(*ClusterTreePtr, Builder->Result->Nodes);
-		Exchange(InstanceReorderTable, Builder->Result->InstanceReorderTable);
-		Exchange(SortedInstances, Builder->Result->SortedInstances);
+		ClusterTreePtr = MakeShareable(new TArray<FClusterNode>(MoveTemp(Builder.Result->Nodes)));
+		InstanceReorderTable = MoveTemp(Builder.Result->InstanceReorderTable);
+		SortedInstances = MoveTemp(Builder.Result->SortedInstances);
 
 		FlushAccumulatedNavigationUpdates();
 				
@@ -2103,6 +2182,8 @@ void UHierarchicalInstancedStaticMeshComponent::BuildTree()
 			bDiscardAsyncBuildResults = true;
 		}
 	}
+
+	ReleasePerInstanceRenderData();
 }
 
 void UHierarchicalInstancedStaticMeshComponent::BuildTreeAnyThread(
@@ -2117,14 +2198,13 @@ void UHierarchicalInstancedStaticMeshComponent::BuildTreeAnyThread(
 {
 	check(MaxInstancesPerLeaf > 0);
 
-	TUniquePtr<FClusterBuilder> Builder(new FClusterBuilder(InstanceTransforms, MeshBox, MaxInstancesPerLeaf));
-	Builder->Build();
-	OutOcclusionLayerNum = Builder->Result->OutOcclusionLayerNum;
-
-	Exchange(OutClusterTree, Builder->Result->Nodes);
-	Exchange(OutInstanceReorderTable, Builder->Result->InstanceReorderTable);
-	Exchange(OutSortedInstances, Builder->Result->SortedInstances);
+	FClusterBuilder Builder(InstanceTransforms, MeshBox, MaxInstancesPerLeaf);
+	Builder.Build();
+	OutOcclusionLayerNum = Builder.Result->OutOcclusionLayerNum;
 
+	OutClusterTree = MoveTemp(Builder.Result->Nodes);
+	OutInstanceReorderTable = MoveTemp(Builder.Result->InstanceReorderTable);
+	OutSortedInstances = MoveTemp(Builder.Result->SortedInstances);
 }
 
 void UHierarchicalInstancedStaticMeshComponent::AcceptPrebuiltTree(TArray<FClusterNode>& InClusterTree, int InOcclusionLayerNumNodes)
@@ -2132,8 +2212,9 @@ void UHierarchicalInstancedStaticMeshComponent::AcceptPrebuiltTree(TArray<FClust
 	QUICK_SCOPE_CYCLE_COUNTER(STAT_UHierarchicalInstancedStaticMeshComponent_AcceptPrebuiltTree);
 	// this is only for prebuild data, already in the correct order
 	check(!PerInstanceSMData.Num());
-	NumBuiltInstances = WriteOncePrebuiltInstanceBuffer.NumInstances();
-	check(NumBuiltInstances);
+	NumBuiltInstances = 0;
+	NumBuiltRenderInstances = WriteOncePrebuiltInstanceBuffer.NumInstances();
+	check(NumBuiltRenderInstances);
 	UnbuiltInstanceBounds.Init();
 	UnbuiltInstanceBoundsList.Empty();
 	RemovedInstances.Empty();
@@ -2146,7 +2227,7 @@ void UHierarchicalInstancedStaticMeshComponent::AcceptPrebuiltTree(TArray<FClust
 	// Verify that the mesh is valid before using it.
 	const bool bMeshIsValid = 
 		// make sure we have instances
-		NumBuiltInstances > 0 &&
+		NumBuiltRenderInstances > 0 &&
 		// make sure we have an actual staticmesh
 		StaticMesh &&
 		StaticMesh->HasValidRenderData() &&
@@ -2164,7 +2245,7 @@ void UHierarchicalInstancedStaticMeshComponent::AcceptPrebuiltTree(TArray<FClust
 		{
 			InstancingRandomSeed = FMath::Rand();
 		}
-		Exchange(*ClusterTreePtr, InClusterTree);
+		*ClusterTreePtr = MoveTemp(InClusterTree);
 		PostBuildStats();
 
 	}
@@ -2198,6 +2279,7 @@ void UHierarchicalInstancedStaticMeshComponent::BuildFlatTree(const TArray<int32
 		}
 
 		NumBuiltInstances = PerInstanceSMData.Num();
+		NumBuiltRenderInstances = PerInstanceSMData.Num();
 		UnbuiltInstanceBounds.Init();
 		UnbuiltInstanceBoundsList.Empty();
 		RemovedInstances.Empty();
@@ -2284,7 +2366,7 @@ void UHierarchicalInstancedStaticMeshComponent::ApplyBuildTreeAsync(ENamedThread
 		{
 			bConcurrentRemoval = false;
 
-			UE_LOG(LogStaticMesh, Verbose, TEXT("Discarded foliage hierarchy of %d elements build due to concurrent removal (%.1fs)"), Builder->Result->InstanceReorderTable.Num(), float(FPlatformTime::Seconds() - StartTime));
+			UE_LOG(LogStaticMesh, Verbose, TEXT("Discarded foliage hierarchy of %d elements build due to concurrent removal (%.1fs)"), Builder->Result->InstanceReorderTable.Num(), (float)(FPlatformTime::Seconds() - StartTime));
 
 			// There were removes or updates while we were building, it's too slow to fix up the result now, so build async again.
 			BuildTreeAsync();
@@ -2292,6 +2374,7 @@ void UHierarchicalInstancedStaticMeshComponent::ApplyBuildTreeAsync(ENamedThread
 		else
 		{
 			NumBuiltInstances = Builder->Result->InstanceReorderTable.Num();
+			NumBuiltRenderInstances = Builder->Result->SortedInstances.Num();
 
 			if (NumBuiltInstances < PerInstanceSMData.Num())
 			{
@@ -2303,16 +2386,15 @@ void UHierarchicalInstancedStaticMeshComponent::ApplyBuildTreeAsync(ENamedThread
 				}
 			}
 
-			ClusterTreePtr = MakeShareable(new TArray<FClusterNode>);
+			ClusterTreePtr = MakeShareable(new TArray<FClusterNode>(MoveTemp(Builder->Result->Nodes)));
 			TArray<FClusterNode>& ClusterTree = *ClusterTreePtr;
-			Exchange(ClusterTree, Builder->Result->Nodes);
-			Exchange(InstanceReorderTable, Builder->Result->InstanceReorderTable);
-			Exchange(SortedInstances, Builder->Result->SortedInstances);
+			InstanceReorderTable = MoveTemp(Builder->Result->InstanceReorderTable);
+			SortedInstances = MoveTemp(Builder->Result->SortedInstances);
 			RemovedInstances.Empty();
 			OcclusionLayerNumNodes = Builder->Result->OutOcclusionLayerNum;
 			BuiltInstanceBounds = (ClusterTree.Num() > 0 ? FBox(ClusterTree[0].BoundMin, ClusterTree[0].BoundMax) : FBox(0));
 
-			UE_LOG(LogStaticMesh, Verbose, TEXT("Built a foliage hierarchy with %d of %d elements in %.1fs."), NumBuiltInstances, PerInstanceSMData.Num(), float(FPlatformTime::Seconds() - StartTime));
+			UE_LOG(LogStaticMesh, Verbose, TEXT("Built a foliage hierarchy with %d of %d elements in %.1fs."), NumBuiltInstances, PerInstanceSMData.Num(), (float)(FPlatformTime::Seconds() - StartTime));
 
 			if (NumBuiltInstances < PerInstanceSMData.Num())
 			{
@@ -2410,7 +2492,7 @@ FPrimitiveSceneProxy* UHierarchicalInstancedStaticMeshComponent::CreateSceneProx
 	// Verify that the mesh is valid before using it.
 	const bool bMeshIsValid = 
 		// make sure we have instances
-		(PerInstanceSMData.Num() > 0 || WriteOncePrebuiltInstanceBuffer.NumInstances() > 0 || bPerInstanceRenderDataWasPrebuilt) &&
+		(GetNumRenderInstances() - RemovedInstances.Num() > 0 || WriteOncePrebuiltInstanceBuffer.NumInstances() > 0 || bPerInstanceRenderDataWasPrebuilt) &&
 		// make sure we have an actual staticmesh
 		StaticMesh &&
 		StaticMesh->HasValidRenderData() &&
@@ -2418,7 +2500,7 @@ FPrimitiveSceneProxy* UHierarchicalInstancedStaticMeshComponent::CreateSceneProx
 		// @todo: Level error or something to let LDs know this
 		1;//StaticMesh->LODModels(0).Elements.Num() == 1;
 
-	if(bMeshIsValid)
+	if (bMeshIsValid)
 	{
 		if (bHasPerInstanceHitProxies && GetWorld() && GetWorld()->IsGameWorld())
 		{
@@ -2428,7 +2510,7 @@ FPrimitiveSceneProxy* UHierarchicalInstancedStaticMeshComponent::CreateSceneProx
 		// If we don't have a random seed for this instanced static mesh component yet, then go ahead and
 		// generate one now.  This will be saved with the static mesh component and used for future generation
 		// of random numbers for this component's instances. (Used by the PerInstanceRandom material expression)
-		while( InstancingRandomSeed == 0 )
+		while (InstancingRandomSeed == 0)
 		{
 			InstancingRandomSeed = FMath::Rand();
 		}
@@ -2440,7 +2522,7 @@ FPrimitiveSceneProxy* UHierarchicalInstancedStaticMeshComponent::CreateSceneProx
 
 			bool bIsGrass = !PerInstanceSMData.Num();
 
-			check(bIsGrass || PerInstanceSMData.Num() == WriteOncePrebuiltInstanceBuffer.NumInstances());
+			check(bIsGrass || WriteOncePrebuiltInstanceBuffer.NumInstances() == GetNumRenderInstances());
 			return ::new FHierarchicalStaticMeshSceneProxy(bIsGrass, this, GetWorld()->FeatureLevel, WriteOncePrebuiltInstanceBuffer);
 		}
 
@@ -2449,23 +2531,27 @@ FPrimitiveSceneProxy* UHierarchicalInstancedStaticMeshComponent::CreateSceneProx
 		INC_DWORD_STAT_BY(STAT_FoliageInstanceBuffers, ProxySize);
 		return ::new FHierarchicalStaticMeshSceneProxy(false, this, GetWorld()->FeatureLevel);
 	}
-	return NULL;
+	return nullptr;
 }
 
 void FAsyncBuildInstanceBuffer::DoWork()
 {
 	QUICK_SCOPE_CYCLE_COUNTER(FAsyncBuildInstanceBuffer_DoWork);
 
-	int32 Num = Component->PerInstanceSMData.Num();
-	Component->WriteOncePrebuiltInstanceBuffer.AllocateInstances(Num);
+	const int32 NumInstances = Component->PerInstanceSMData.Num();
+	const int32 NumRenderInstances = Component->GetNumRenderInstances();
+	Component->WriteOncePrebuiltInstanceBuffer.AllocateInstances(NumRenderInstances);
 	FRandomStream RandomStream(Component->InstancingRandomSeed);
 	bool bUseRemapTable = Component->PerInstanceSMData.Num() == Component->InstanceReorderTable.Num();
 
-	for (int32 InstanceIndex = 0; InstanceIndex < Num; InstanceIndex++)
+	for (int32 InstanceIndex = 0; InstanceIndex < NumInstances; InstanceIndex++)
 	{
 		const FInstancedStaticMeshInstanceData& Instance = Component->PerInstanceSMData[InstanceIndex];
 		const int32 DestInstanceIndex = bUseRemapTable ? Component->InstanceReorderTable[InstanceIndex] : InstanceIndex;
-		Component->WriteOncePrebuiltInstanceBuffer.SetInstance(DestInstanceIndex, Instance.Transform, RandomStream.GetFraction(), Instance.LightmapUVBias, Instance.ShadowmapUVBias);
+		if (DestInstanceIndex != INDEX_NONE)
+		{
+			Component->WriteOncePrebuiltInstanceBuffer.SetInstance(DestInstanceIndex, Instance.Transform, RandomStream.GetFraction(), Instance.LightmapUVBias, Instance.ShadowmapUVBias);
+		}
 	}
 	if (Component->RemovedInstances.Num())
 	{
@@ -2477,7 +2563,6 @@ void FAsyncBuildInstanceBuffer::DoWork()
 		}
 	}
 
-
 	check(World);
 	check(World->AsyncPreRegisterLevelStreamingTasks.GetValue() > 0);
 	FPlatformMisc::MemoryBarrier();
@@ -2487,20 +2572,34 @@ void FAsyncBuildInstanceBuffer::DoWork()
 void UHierarchicalInstancedStaticMeshComponent::PostLoad()
 {
 	Super::PostLoad();
-	if (CVarASyncInstaneBufferConversion.GetValueOnGameThread() > 0)
+
+	NumBuiltRenderInstances = ClusterTreePtr.IsValid() && ClusterTreePtr->Num() > 0 ? (*ClusterTreePtr.Get())[0].LastInstance - (*ClusterTreePtr.Get())[0].FirstInstance + 1 : 0;
+
+	if (bEnableDensityScaling && GetWorld() && GetWorld()->IsGameWorld())
 	{
-		UWorld* World = GetWorld();
-		if (World && World->IsGameWorld() && PerInstanceSMData.Num())
+		const float ScalabilityDensity = FMath::Clamp(CVarFoliageDensityScale.GetValueOnGameThread(), 0.0f, 1.0f);
+		if (ScalabilityDensity == 0)
 		{
-			World->AsyncPreRegisterLevelStreamingTasks.Increment();
-			while( InstancingRandomSeed == 0 )
+			// exclude all instances
+			ExcludedDueToDensityScaling.Init(true, PerInstanceSMData.Num());
+			NumBuiltRenderInstances = 0;
+		}
+		else if (ScalabilityDensity > 0.0f && ScalabilityDensity < 1.0f)
+		{
+			FRandomStream Rand(InstancingRandomSeed);
+			ExcludedDueToDensityScaling.Init(false, PerInstanceSMData.Num());
+
+			for (int32 i = 0; i < ExcludedDueToDensityScaling.Num(); ++i)
 			{
-				InstancingRandomSeed = FMath::Rand();
+				ExcludedDueToDensityScaling[i] = (Rand.FRand() > ScalabilityDensity);
 			}
-			AsyncBuildInstanceBufferTask = new FAsyncTask<FAsyncBuildInstanceBuffer>(this, World);
-			AsyncBuildInstanceBufferTask->StartBackgroundTask();
+
+			// prune tree
+			StaticMesh->ConditionalPostLoad();
+			BuildTree();
 		}
 	}
+
 	// For some reason we don't have a tree, or it is out of date. Build one now!
 	if (StaticMesh && PerInstanceSMData.Num() > 0 && (!ClusterTreePtr.IsValid() || ClusterTreePtr->Num() == 0 || (NumBuiltInstances != PerInstanceSMData.Num()) || UnbuiltInstanceBoundsList.Num() > 0 || GetLinkerUE4Version() < VER_UE4_REBUILD_HIERARCHICAL_INSTANCE_TREES))
 	{
@@ -2509,6 +2608,21 @@ void UHierarchicalInstancedStaticMeshComponent::PostLoad()
 		StaticMesh->ConditionalPostLoad();
 		BuildTree();
 	}
+
+	if (CVarASyncInstaneBufferConversion.GetValueOnGameThread() > 0)
+	{
+		UWorld* World = GetWorld();
+		if (World && World->IsGameWorld() && NumBuiltRenderInstances > 0)
+		{
+			World->AsyncPreRegisterLevelStreamingTasks.Increment();
+			while (InstancingRandomSeed == 0)
+			{
+				InstancingRandomSeed = FMath::Rand();
+			}
+			AsyncBuildInstanceBufferTask = new FAsyncTask<FAsyncBuildInstanceBuffer>(this, World);
+			AsyncBuildInstanceBufferTask->StartBackgroundTask();
+		}
+	}
 }
 
 static void GatherInstanceTransformsInArea(const UHierarchicalInstancedStaticMeshComponent& Component, const FBox& AreaBox, int32 Child, TArray<FTransform>& InstanceData)
@@ -2797,10 +2911,9 @@ static void RebuildFoliageTrees(const TArray<FString>& Args)
 	for (TObjectIterator<UHierarchicalInstancedStaticMeshComponent> It; It; ++It)
 	{
 		UHierarchicalInstancedStaticMeshComponent* Comp = *It;
-		if (Comp && !Comp->IsTemplate() && !Comp->IsPendingKill() && !Comp->IsAsyncBuilding())
+		if (Comp && !Comp->IsTemplate() && !Comp->IsPendingKill())
 		{
 			Comp->BuildTree();
-			Comp->ReleasePerInstanceRenderData();
 			Comp->MarkRenderStateDirty();
 		}
 	}
diff --git a/Engine/Source/Runtime/Engine/Private/InstancedStaticMesh.cpp b/Engine/Source/Runtime/Engine/Private/InstancedStaticMesh.cpp
--- a/Engine/Source/Runtime/Engine/Private/InstancedStaticMesh.cpp
+++ b/Engine/Source/Runtime/Engine/Private/InstancedStaticMesh.cpp
@@ -72,14 +72,14 @@ void FStaticMeshInstanceBuffer::Init(UInstancedStaticMeshComponent* InComponent,
 	bool bUseRemapTable = InComponent->PerInstanceSMData.Num() == InComponent->InstanceReorderTable.Num();
 
 	int32 NumRealInstances = InComponent->PerInstanceSMData.Num();
-	int32 NumRemoved = InComponent->RemovedInstances.Num();
+	int32 NumRenderInstances = InComponent->GetNumRenderInstances();
 
 	// Allocate the vertex data storage type.
 	AllocateData();
 
 	SetupCPUAccess(InComponent);
 
-	NumInstances = NumRealInstances + NumRemoved;
+	NumInstances = NumRenderInstances;
 	InstanceData->AllocateInstances(NumInstances);
 
 	// Setup our random number generator such that random values are generated consistently for any
@@ -90,12 +90,16 @@ void FStaticMeshInstanceBuffer::Init(UInstancedStaticMeshComponent* InComponent,
 	{
 		const FInstancedStaticMeshInstanceData& Instance = InComponent->PerInstanceSMData[InstanceIndex];
 		const int32 DestInstanceIndex = bUseRemapTable ? InComponent->InstanceReorderTable[InstanceIndex] : InstanceIndex;
-		InstanceData->SetInstance(DestInstanceIndex, Instance.Transform, RandomStream.GetFraction(), Instance.LightmapUVBias, Instance.ShadowmapUVBias);
+		if (DestInstanceIndex != INDEX_NONE)
+		{
+			InstanceData->SetInstance(DestInstanceIndex, Instance.Transform, RandomStream.GetFraction(), Instance.LightmapUVBias, Instance.ShadowmapUVBias);
+		}
 	}
 
 	SetPerInstanceEditorData(InComponent, InHitProxies);
 
 	// Hide any removed instances
+	int32 NumRemoved = InComponent->RemovedInstances.Num();
 	if (NumRemoved)
 	{
 		check(bUseRemapTable);
@@ -141,7 +145,10 @@ void FStaticMeshInstanceBuffer::SetPerInstanceEditorData(UInstancedStaticMeshCom
 			// Record if the instance is selected
 			bool bSelected = InstanceIndex < InComponent->SelectedInstances.Num() && InComponent->SelectedInstances[InstanceIndex];
 			const int32 DestInstanceIndex = bUseRemapTable ? InComponent->InstanceReorderTable[InstanceIndex] : InstanceIndex;
-			InstanceData->SetInstanceEditorData(DestInstanceIndex, HitProxyColor, bSelected);
+			if (DestInstanceIndex != INDEX_NONE)
+			{
+				InstanceData->SetInstanceEditorData(DestInstanceIndex, HitProxyColor, bSelected);
+			}
 		}
 	}
 #endif
diff --git a/Engine/Source/Runtime/Engine/Private/PhysicsEngine/BodySetup.cpp b/Engine/Source/Runtime/Engine/Private/PhysicsEngine/BodySetup.cpp
--- a/Engine/Source/Runtime/Engine/Private/PhysicsEngine/BodySetup.cpp
+++ b/Engine/Source/Runtime/Engine/Private/PhysicsEngine/BodySetup.cpp
@@ -759,10 +759,12 @@ void UBodySetup::Serialize(FArchive& Ar)
 	bool bCooked = Ar.IsCooking();
 	Ar << bCooked;
 
+#if !WITH_RUNTIME_PHYSICS_COOKING
 	if (FPlatformProperties::RequiresCookedData() && !bCooked && Ar.IsLoading())
 	{
 		UE_LOG(LogPhysics, Fatal, TEXT("This platform requires cooked packages, and physX data was not cooked into %s."), *GetFullName());
 	}
+#endif //!WITH_RUNTIME_PHYSICS_COOKING
 
 	if (bCooked)
 	{
diff --git a/Engine/Source/Runtime/Engine/Private/Scalability.cpp b/Engine/Source/Runtime/Engine/Private/Scalability.cpp
--- a/Engine/Source/Runtime/Engine/Private/Scalability.cpp
+++ b/Engine/Source/Runtime/Engine/Private/Scalability.cpp
@@ -55,6 +55,13 @@ static TAutoConsoleVariable<int32> CVarEffectsQuality(
 	TEXT(" 0:low, 1:med, 2:high, 3:epic, default: 3"),
 	ECVF_ScalabilityGroup);
 
+static TAutoConsoleVariable<int32> CVarFoliageQuality(
+	TEXT("sg.FoliageQuality"),
+	3,
+	TEXT("Scalability quality state (internally used by scalability system, ini load/save or using SCALABILITY console command)\n")
+	TEXT(" 0:low, 1:med, 2:high, 3:epic, default: 3"),
+	ECVF_ScalabilityGroup);
+
 namespace Scalability
 {
 	// Select a the correct quality level for the given benchmark value and thresholds
@@ -217,29 +224,9 @@ void OnChangeResolutionQuality(IConsoleVariable* Var)
 {
 	SetResolutionQualityLevel(Var->GetFloat());
 }
-void OnChangeViewDistanceQuality(IConsoleVariable* Var)
-{
-	SetGroupQualityLevel(TEXT("ViewDistanceQuality"), Var->GetInt());
-}
-void OnChangeAntiAliasingQuality(IConsoleVariable* Var)
-{
-	SetGroupQualityLevel(TEXT("AntiAliasingQuality"), Var->GetInt());
-}
-void OnChangeShadowQuality(IConsoleVariable* Var)
-{
-	SetGroupQualityLevel(TEXT("ShadowQuality"), Var->GetInt());
-}
-void OnChangePostProcessQuality(IConsoleVariable* Var)
-{
-	SetGroupQualityLevel(TEXT("PostProcessQuality"), Var->GetInt());
-}
-void OnChangeTextureQuality(IConsoleVariable* Var)
-{
-	SetGroupQualityLevel(TEXT("TextureQuality"), Var->GetInt());
-}
-void OnChangeEffectsQuality(IConsoleVariable* Var)
+void OnChangeQuality(IConsoleVariable* Var, const TCHAR* GroupName)
 {
-	SetGroupQualityLevel(TEXT("EffectsQuality"), Var->GetInt());
+	SetGroupQualityLevel(GroupName, Var->GetInt());
 }
 
 void InitScalabilitySystem()
@@ -257,12 +244,13 @@ void InitScalabilitySystem()
 	}
 
 	CVarResolutionQuality.AsVariable()->SetOnChangedCallback(FConsoleVariableDelegate::CreateStatic(&OnChangeResolutionQuality));
-	CVarViewDistanceQuality.AsVariable()->SetOnChangedCallback(FConsoleVariableDelegate::CreateStatic(&OnChangeViewDistanceQuality));
-	CVarAntiAliasingQuality.AsVariable()->SetOnChangedCallback(FConsoleVariableDelegate::CreateStatic(&OnChangeAntiAliasingQuality));
-	CVarShadowQuality.AsVariable()->SetOnChangedCallback(FConsoleVariableDelegate::CreateStatic(&OnChangeShadowQuality));
-	CVarPostProcessQuality.AsVariable()->SetOnChangedCallback(FConsoleVariableDelegate::CreateStatic(&OnChangePostProcessQuality));
-	CVarTextureQuality.AsVariable()->SetOnChangedCallback(FConsoleVariableDelegate::CreateStatic(&OnChangeTextureQuality));
-	CVarEffectsQuality.AsVariable()->SetOnChangedCallback(FConsoleVariableDelegate::CreateStatic(&OnChangeEffectsQuality));
+	CVarViewDistanceQuality.AsVariable()->SetOnChangedCallback(FConsoleVariableDelegate::CreateStatic(&OnChangeQuality, TEXT("ViewDistanceQuality")));
+	CVarAntiAliasingQuality.AsVariable()->SetOnChangedCallback(FConsoleVariableDelegate::CreateStatic(&OnChangeQuality, TEXT("AntiAliasingQuality")));
+	CVarShadowQuality.AsVariable()->SetOnChangedCallback(FConsoleVariableDelegate::CreateStatic(&OnChangeQuality, TEXT("ShadowQuality")));
+	CVarPostProcessQuality.AsVariable()->SetOnChangedCallback(FConsoleVariableDelegate::CreateStatic(&OnChangeQuality, TEXT("PostProcessQuality")));
+	CVarTextureQuality.AsVariable()->SetOnChangedCallback(FConsoleVariableDelegate::CreateStatic(&OnChangeQuality, TEXT("TextureQuality")));
+	CVarEffectsQuality.AsVariable()->SetOnChangedCallback(FConsoleVariableDelegate::CreateStatic(&OnChangeQuality, TEXT("EffectsQuality")));
+	CVarFoliageQuality.AsVariable()->SetOnChangedCallback(FConsoleVariableDelegate::CreateStatic(&OnChangeQuality, TEXT("FoliageQuality")));
 }
 
 /** Get the percentage scale for a given quality level */
@@ -295,6 +283,7 @@ FQualityLevels BenchmarkQualityLevels(uint32 WorkScale, float CPUMultiplier, flo
 	Results.PostProcessQuality = ComputeOptionFromPerfIndex(TEXT("PostProcessQuality"), CPUPerfIndex, GPUPerfIndex);
 	Results.TextureQuality = ComputeOptionFromPerfIndex(TEXT("TextureQuality"), CPUPerfIndex, GPUPerfIndex);
 	Results.EffectsQuality = ComputeOptionFromPerfIndex(TEXT("EffectsQuality"), CPUPerfIndex, GPUPerfIndex);
+	Results.FoliageQuality = ComputeOptionFromPerfIndex(TEXT("FoliageQuality"), CPUPerfIndex, GPUPerfIndex);
 	Results.CPUBenchmarkResults = CPUPerfIndex;
 	Results.GPUBenchmarkResults = GPUPerfIndex;
 
@@ -390,6 +379,7 @@ void ProcessCommand(const TCHAR* Cmd, FOutputDevice& Ar)
 		PrintGroupInfo(TEXT("PostProcessQuality"), bInfoMode);
 		PrintGroupInfo(TEXT("TextureQuality"), bInfoMode);
 		PrintGroupInfo(TEXT("EffectsQuality"), bInfoMode);
+		PrintGroupInfo(TEXT("FoliageQuality"), bInfoMode);
 
 		if (CPUBenchmarkValue >= 0.0f)
 		{
@@ -411,6 +401,7 @@ void SetQualityLevels(const FQualityLevels& QualityLevels)
 	CVarPostProcessQuality.AsVariable()->Set(QualityLevels.PostProcessQuality, ECVF_SetByScalability);
 	CVarTextureQuality.AsVariable()->Set(QualityLevels.TextureQuality, ECVF_SetByScalability);
 	CVarEffectsQuality.AsVariable()->Set(QualityLevels.EffectsQuality, ECVF_SetByScalability);
+	CVarFoliageQuality.AsVariable()->Set(QualityLevels.FoliageQuality, ECVF_SetByScalability);
 }
 
 FQualityLevels GetQualityLevels()
@@ -425,6 +416,7 @@ FQualityLevels GetQualityLevels()
 	Ret.PostProcessQuality = CVarPostProcessQuality.GetValueOnGameThread();
 	Ret.TextureQuality = CVarTextureQuality.GetValueOnGameThread();
 	Ret.EffectsQuality = CVarEffectsQuality.GetValueOnGameThread();
+	Ret.FoliageQuality = CVarFoliageQuality.GetValueOnGameThread();
 
 	return Ret;
 }
@@ -452,6 +444,7 @@ void FQualityLevels::SetFromSingleQualityLevel(int32 Value)
 	PostProcessQuality = Value;
 	TextureQuality = Value;
 	EffectsQuality = Value;
+	FoliageQuality = Value;
 }
 
 // Returns the overall value if all settings are set to the same thing
@@ -461,7 +454,7 @@ int32 FQualityLevels::GetSingleQualityLevel() const
 	int32 Result = ViewDistanceQuality;
 
 	const int32 Target = ViewDistanceQuality;
-	if ((Target == AntiAliasingQuality) && (Target == ShadowQuality) && (Target == PostProcessQuality) && (Target == TextureQuality) && (Target == EffectsQuality))
+	if ((Target == AntiAliasingQuality) && (Target == ShadowQuality) && (Target == PostProcessQuality) && (Target == TextureQuality) && (Target == EffectsQuality) && (Target == FoliageQuality))
 	{
 		if (GetRenderScaleLevelFromQualityLevel(Target) == ResolutionQuality)
 		{
@@ -491,6 +484,7 @@ void LoadState(const FString& IniName)
 	GConfig->GetInt(Section, TEXT("sg.PostProcessQuality"), State.PostProcessQuality, IniName);
 	GConfig->GetInt(Section, TEXT("sg.TextureQuality"), State.TextureQuality, IniName);
 	GConfig->GetInt(Section, TEXT("sg.EffectsQuality"), State.EffectsQuality, IniName);
+	GConfig->GetInt(Section, TEXT("sg.FoliageQuality"), State.FoliageQuality, IniName);
 
 	SetQualityLevels(State);
 }
@@ -511,6 +505,7 @@ void SaveState(const FString& IniName)
 	GConfig->SetInt(Section, TEXT("sg.PostProcessQuality"), State.PostProcessQuality, IniName);
 	GConfig->SetInt(Section, TEXT("sg.TextureQuality"), State.TextureQuality, IniName);
 	GConfig->SetInt(Section, TEXT("sg.EffectsQuality"), State.EffectsQuality, IniName);
+	GConfig->SetInt(Section, TEXT("sg.FoliageQuality"), State.FoliageQuality, IniName);
 }
 
 void RecordQualityLevelsAnalytics(bool bAutoApplied)
@@ -528,6 +523,7 @@ void RecordQualityLevelsAnalytics(bool bAutoApplied)
 		Attributes.Add(FAnalyticsEventAttribute(TEXT("PostProcessQuality"), State.PostProcessQuality));
 		Attributes.Add(FAnalyticsEventAttribute(TEXT("TextureQuality"), State.TextureQuality));
 		Attributes.Add(FAnalyticsEventAttribute(TEXT("EffectsQuality"), State.EffectsQuality));
+		Attributes.Add(FAnalyticsEventAttribute(TEXT("FoliageQuality"), State.TextureQuality));
 		Attributes.Add(FAnalyticsEventAttribute(TEXT("AutoAppliedSettings"), bAutoApplied));
 
 		FEngineAnalytics::GetProvider().RecordEvent(TEXT("Editor.Performance.ScalabiltySettings"), Attributes);
diff --git a/Engine/Source/Runtime/Engine/Public/Scalability.h b/Engine/Source/Runtime/Engine/Public/Scalability.h
--- a/Engine/Source/Runtime/Engine/Public/Scalability.h
+++ b/Engine/Source/Runtime/Engine/Public/Scalability.h
@@ -21,6 +21,7 @@ namespace Scalability
 		int32 PostProcessQuality;
 		int32 TextureQuality;
 		int32 EffectsQuality;
+		int32 FoliageQuality;
 
 		float CPUBenchmarkResults;
 		float GPUBenchmarkResults;
diff --git a/Engine/Source/Runtime/Foliage/Private/InstancedFoliage.cpp b/Engine/Source/Runtime/Foliage/Private/InstancedFoliage.cpp
--- a/Engine/Source/Runtime/Foliage/Private/InstancedFoliage.cpp
+++ b/Engine/Source/Runtime/Foliage/Private/InstancedFoliage.cpp
@@ -288,7 +288,9 @@ UFoliageType::UFoliageType(const FObjectInitializer& ObjectInitializer)
 	CullDistance.Max = 0;
 	bEnableStaticLighting_DEPRECATED = true;
 	MinimumLayerWeight = 0.5f;
+#if WITH_EDITORONLY_DATA
 	IsSelected = false;
+#endif
 	DensityAdjustmentFactor = 1.0f;
 	CollisionWithWorld = false;
 	CollisionScale = FVector(0.9f, 0.9f, 0.9f);
@@ -330,7 +332,10 @@ UFoliageType::UFoliageType(const FObjectInitializer& ObjectInitializer)
 	Curve->AddKey(1.f, 1.f);
 
 	UpdateGuid = FGuid::NewGuid();
+#if WITH_EDITORONLY_DATA
 	HiddenEditorViews = 0;
+#endif
+	bEnableDensityScaling = false;
 
 #if WITH_EDITORONLY_DATA
 	// Deprecated since FFoliageCustomVersion::FoliageTypeCustomization
@@ -867,6 +872,12 @@ void FFoliageMeshInfo::UpdateComponentSettings(const UFoliageType* InSettings)
 			bNeedsMarkRenderStateDirty = true;
 		}
 
+		if (Component->bEnableDensityScaling != FoliageType->bEnableDensityScaling)
+		{
+			Component->bEnableDensityScaling = FoliageType->bEnableDensityScaling;
+			bNeedsMarkRenderStateDirty = true;
+		}
+
 		UFoliageInstancedStaticMeshComponent* FoliageComponent = Cast<UFoliageInstancedStaticMeshComponent>(Component);
 
 		if (FoliageComponent && FoliageComponent->FoliageHiddenEditorViews != InSettings->HiddenEditorViews)
@@ -2381,9 +2392,15 @@ void AInstancedFoliageActor::PostLoad()
 			FFoliageMeshInfo& MeshInfo = *MeshPair.Value;
 			UFoliageType* FoliageType = MeshPair.Key;
 
+			// Make sure the mesh has been PostLoaded as if not it can be considered invalid resulting in a bad tree
+			UStaticMesh* StaticMesh = FoliageType->GetStaticMesh();
+			if (StaticMesh)
+			{
+				StaticMesh->ConditionalPostLoad();
+			}
+
 			if (MeshInfo.Instances.Num() && MeshInfo.Component == nullptr)
 			{
-				const UStaticMesh* StaticMesh = FoliageType->GetStaticMesh();
 				FFormatNamedArguments Arguments;
 				if (StaticMesh)
 				{
@@ -2426,7 +2443,7 @@ void AInstancedFoliageActor::PostLoad()
 			}
 			
 			// Clean up case where embeded instances had their static mesh deleted
-			if (FoliageType->IsNotAssetOrBlueprint() && FoliageType->GetStaticMesh() == nullptr)
+			if (FoliageType->IsNotAssetOrBlueprint() && StaticMesh == nullptr)
 			{
 				OnFoliageTypeMeshChangedEvent.Broadcast(FoliageType);
 				RemoveFoliageType(&FoliageType, 1);
diff --git a/Engine/Source/Runtime/Foliage/Public/FoliageType.h b/Engine/Source/Runtime/Foliage/Public/FoliageType.h
--- a/Engine/Source/Runtime/Foliage/Public/FoliageType.h
+++ b/Engine/Source/Runtime/Foliage/Public/FoliageType.h
@@ -291,13 +291,14 @@ public:
 	UPROPERTY(EditAnywhere, Category=InstanceSettings, meta=(HideObjectType=true))
 	TEnumAsByte<EHasCustomNavigableGeometry::Type> CustomNavigableGeometry;
 
+#if WITH_EDITORONLY_DATA
 	/** Bitflag to represent in which editor views this foliage mesh is hidden. */
 	UPROPERTY(transient)
 	uint64 HiddenEditorViews;
 
 	UPROPERTY()
 	uint32 IsSelected:1;
-
+#endif
 public:
 	// PROCEDURAL
 
@@ -455,6 +456,14 @@ public:
 	UPROPERTY(EditDefaultsOnly, Category=Reapply)
 	uint32 ReapplyVertexColorMask:1;
 
+public:
+	// SCALABILITY
+
+	// Enable for detail meshes that don't really affect the game. Disable for anything important.
+	// Typically, this will be enabled for small meshes without collision (e.g. grass) and disabled for large meshes with collision (e.g. trees)
+	UPROPERTY(EditAnywhere, Category=Scalability)
+	uint32 bEnableDensityScaling:1;
+
 private:
 
 #if WITH_EDITORONLY_DATA
diff --git a/Engine/Source/Runtime/Landscape/Private/LandscapeGrass.cpp b/Engine/Source/Runtime/Landscape/Private/LandscapeGrass.cpp
--- a/Engine/Source/Runtime/Landscape/Private/LandscapeGrass.cpp
+++ b/Engine/Source/Runtime/Landscape/Private/LandscapeGrass.cpp
@@ -67,7 +67,7 @@ static TAutoConsoleVariable<int32> CVarUseHaltonDistribution(
 static TAutoConsoleVariable<float> CVarGrassDensityScale(
 	TEXT("grass.densityScale"),
 	1,
-	TEXT("Multiplier on all grass densities.  Do a grass.flushcache or grass.flushcachepie afterwards as appropriate."));
+	TEXT("Multiplier on all grass densities."));
 
 static TAutoConsoleVariable<int32> CVarGrassEnable(
 	TEXT("grass.Enable"),
@@ -100,6 +100,24 @@ DECLARE_CYCLE_STAT(TEXT("Grass End Comp"), STAT_FoliageGrassEndComp, STATGROUP_F
 DECLARE_CYCLE_STAT(TEXT("Grass Destroy Comps"), STAT_FoliageGrassDestoryComp, STATGROUP_Foliage);
 DECLARE_CYCLE_STAT(TEXT("Grass Update"), STAT_GrassUpdate, STATGROUP_Foliage);
 
+static void GrassCVarSinkFunction()
+{
+	static float CachedGrassDensityScale = 1.0f;
+	float GrassDensityScale = CVarGrassDensityScale.GetValueOnGameThread();
+
+	if (GrassDensityScale != CachedGrassDensityScale)
+	{
+		CachedGrassDensityScale = GrassDensityScale;
+
+		for (auto* Landscape : TObjectRange<ALandscapeProxy>(RF_ClassDefaultObject | RF_ArchetypeObject, true, EInternalObjectFlags::PendingKill))
+		{
+			Landscape->FlushGrassComponents(nullptr, false);
+		}
+	}
+}
+
+static FAutoConsoleVariableSink CVarGrassSink(FConsoleCommandDelegate::CreateStatic(&GrassCVarSinkFunction));
+
 //
 // Grass weightmap rendering
 //
