Fixed regression with collision trace channels serializing incorrectly (UE-29460)
- This could cause game-specific trace channels to be loaded "off by one"
- Switched VREditor to use a reserved engine trace channel instead of inserting a new one

diff --git a/Engine/Source/Editor/VREditor/Gizmo/VREditorGizmoHandle.cpp b/Engine/Source/Editor/VREditor/Gizmo/VREditorGizmoHandle.cpp
--- a/Engine/Source/Editor/VREditor/Gizmo/VREditorGizmoHandle.cpp
+++ b/Engine/Source/Editor/VREditor/Gizmo/VREditorGizmoHandle.cpp
@@ -220,8 +220,8 @@ class UStaticMeshComponent* UVREditorGizmoHandleGroup::CreateMeshHandle( class U
 
 	HandleComponent->SetCollisionEnabled( ECollisionEnabled::QueryOnly );
 	HandleComponent->SetCollisionResponseToAllChannels( ECR_Ignore );
-	HandleComponent->SetCollisionResponseToChannel( ECC_EditorGizmo, ECollisionResponse::ECR_Block );
-	HandleComponent->SetCollisionObjectType( ECC_EditorGizmo );
+	HandleComponent->SetCollisionResponseToChannel( COLLISION_GIZMO, ECollisionResponse::ECR_Block );
+	HandleComponent->SetCollisionObjectType( COLLISION_GIZMO );
 
 	HandleComponent->bGenerateOverlapEvents = false;
 	HandleComponent->SetCanEverAffectNavigation( false );
diff --git a/Engine/Source/Editor/VREditor/Gizmo/VREditorUniformScaleGizmoHandle.cpp b/Engine/Source/Editor/VREditor/Gizmo/VREditorUniformScaleGizmoHandle.cpp
--- a/Engine/Source/Editor/VREditor/Gizmo/VREditorUniformScaleGizmoHandle.cpp
+++ b/Engine/Source/Editor/VREditor/Gizmo/VREditorUniformScaleGizmoHandle.cpp
@@ -28,8 +28,8 @@ UVREditorUniformScaleGizmoHandleGroup::UVREditorUniformScaleGizmoHandleGroup()
 
 	UniformScaleHandle->SetCollisionEnabled( ECollisionEnabled::QueryOnly );
 	UniformScaleHandle->SetCollisionResponseToAllChannels( ECR_Ignore );
-	UniformScaleHandle->SetCollisionResponseToChannel( ECC_EditorGizmo, ECollisionResponse::ECR_Block );
-	UniformScaleHandle->SetCollisionObjectType( ECC_EditorGizmo );
+	UniformScaleHandle->SetCollisionResponseToChannel( COLLISION_GIZMO, ECollisionResponse::ECR_Block );
+	UniformScaleHandle->SetCollisionObjectType( COLLISION_GIZMO );
 
 	UniformScaleHandle->bGenerateOverlapEvents = false;
 	UniformScaleHandle->SetCanEverAffectNavigation( false );
diff --git a/Engine/Source/Editor/VREditor/VREditorDockableWindow.cpp b/Engine/Source/Editor/VREditor/VREditorDockableWindow.cpp
--- a/Engine/Source/Editor/VREditor/VREditorDockableWindow.cpp
+++ b/Engine/Source/Editor/VREditor/VREditorDockableWindow.cpp
@@ -37,8 +37,8 @@ AVREditorDockableWindow::AVREditorDockableWindow()
 
 		WindowMeshComponent->SetCollisionEnabled( ECollisionEnabled::QueryOnly );
 		WindowMeshComponent->SetCollisionResponseToAllChannels( ECollisionResponse::ECR_Ignore );
-		WindowMeshComponent->SetCollisionResponseToChannel( ECC_EditorGizmo, ECollisionResponse::ECR_Block );
-		WindowMeshComponent->SetCollisionObjectType( ECC_EditorGizmo );
+		WindowMeshComponent->SetCollisionResponseToChannel( COLLISION_GIZMO, ECollisionResponse::ECR_Block );
+		WindowMeshComponent->SetCollisionObjectType( COLLISION_GIZMO );
 
 		WindowMeshComponent->bGenerateOverlapEvents = false;
 		WindowMeshComponent->SetCanEverAffectNavigation( false );
@@ -62,8 +62,8 @@ AVREditorDockableWindow::AVREditorDockableWindow()
 
 		SelectionBarMeshComponent->SetCollisionEnabled( ECollisionEnabled::QueryOnly );
 		SelectionBarMeshComponent->SetCollisionResponseToAllChannels( ECollisionResponse::ECR_Ignore );
-		SelectionBarMeshComponent->SetCollisionResponseToChannel( ECC_EditorGizmo, ECollisionResponse::ECR_Block );
-		SelectionBarMeshComponent->SetCollisionObjectType( ECC_EditorGizmo );
+		SelectionBarMeshComponent->SetCollisionResponseToChannel( COLLISION_GIZMO, ECollisionResponse::ECR_Block );
+		SelectionBarMeshComponent->SetCollisionObjectType( COLLISION_GIZMO );
 
 		SelectionBarMeshComponent->bGenerateOverlapEvents = false;
 		SelectionBarMeshComponent->SetCanEverAffectNavigation( false );
@@ -107,8 +107,8 @@ AVREditorDockableWindow::AVREditorDockableWindow()
 
 		CloseButtonMeshComponent->SetCollisionEnabled( ECollisionEnabled::QueryOnly );
 		CloseButtonMeshComponent->SetCollisionResponseToAllChannels( ECollisionResponse::ECR_Ignore );
-		CloseButtonMeshComponent->SetCollisionResponseToChannel( ECC_EditorGizmo, ECollisionResponse::ECR_Block );
-		CloseButtonMeshComponent->SetCollisionObjectType( ECC_EditorGizmo );
+		CloseButtonMeshComponent->SetCollisionResponseToChannel( COLLISION_GIZMO, ECollisionResponse::ECR_Block );
+		CloseButtonMeshComponent->SetCollisionObjectType( COLLISION_GIZMO );
 
 		CloseButtonMeshComponent->bGenerateOverlapEvents = false;
 		CloseButtonMeshComponent->SetCanEverAffectNavigation( false );
diff --git a/Engine/Source/Editor/VREditor/VREditorMode.cpp b/Engine/Source/Editor/VREditor/VREditorMode.cpp
--- a/Engine/Source/Editor/VREditor/VREditorMode.cpp
+++ b/Engine/Source/Editor/VREditor/VREditorMode.cpp
@@ -1837,7 +1837,7 @@ FHitResult FVREditorMode::GetHitResultFromLaserPointer( int32 HandIndex, TArray<
 			
 			const FCollisionResponseParams& ResponseParam = FCollisionResponseParams::DefaultResponseParam;
 
-			const ECollisionChannel CollisionChannel = bOnlyEditorGizmos ? ECC_EditorGizmo : ECC_Visibility;
+			const ECollisionChannel CollisionChannel = bOnlyEditorGizmos ? COLLISION_GIZMO : ECC_Visibility;
 
 			// Don't trace against our own head/hand meshes
 			TraceParams.AddIgnoredActor( AvatarMeshActor );
@@ -1859,13 +1859,13 @@ FHitResult FVREditorMode::GetHitResultFromLaserPointer( int32 HandIndex, TArray<
 			else
 			{
 				FCollisionObjectQueryParams EverythingButGizmos( FCollisionObjectQueryParams::AllObjects );
-				EverythingButGizmos.RemoveObjectTypesToQuery( ECC_EditorGizmo );
+				EverythingButGizmos.RemoveObjectTypesToQuery( COLLISION_GIZMO );
 				bHit = GetWorld()->LineTraceSingleByObjectType( HitResult, LaserPointerStart, LaserPointerEnd, EverythingButGizmos, TraceParams );
 			}
 			if( bHit )
 			{
 				// Is this better than what we have already?  Always prefer transform gizmos even if they were not using
-				// ECC_EditorGizmo (some gizmo handles are opaque and use ECC_Visibility as their collision channel)
+				// COLLISION_GIZMO (some gizmo handles are opaque and use ECC_Visibility as their collision channel)
 				// NOTE: We're treating components of floating UI actors as "gizmos" for the purpose of hit testing as long as the component is not the actual UI widget component
 				const bool bHitResultIsUI = UISystem->IsWidgetAnEditorUIWidget( HitResult.GetComponent() );
 				const bool bHitResultIsGizmo = HitResult.GetActor() != nullptr && ( ( HitResult.GetActor() == WorldInteraction->GetTransformGizmoActor() ) || ( !bHitResultIsUI && HitResult.GetActor()->IsA( AVREditorFloatingUI::StaticClass() ) ) );
diff --git a/Engine/Source/Runtime/Engine/Classes/Engine/EngineTypes.h b/Engine/Source/Runtime/Engine/Classes/Engine/EngineTypes.h
--- a/Engine/Source/Runtime/Engine/Classes/Engine/EngineTypes.h
+++ b/Engine/Source/Runtime/Engine/Classes/Engine/EngineTypes.h
@@ -419,17 +419,17 @@ enum ECollisionChannel
 	ECC_PhysicsBody UMETA(DisplayName="PhysicsBody"),
 	ECC_Vehicle UMETA(DisplayName="Vehicle"),
 	ECC_Destructible UMETA(DisplayName="Destructible"),
-	ECC_EditorGizmo UMETA(Hidden, DisplayName="EditorGizmo", TraceQuery="1"),
 	// @NOTE : when you add more here for predefined engine channel
 	// please change the max in the CollisionProfile
-	// search ECC_EditorGizmo
+	// search ECC_Destructible
 
 	// Unspecified Engine Trace Channels
-	ECC_EngineTraceChannel1 UMETA(Hidden),
+	ECC_EngineTraceChannel1 UMETA(Hidden),		// IMPORTANT: This engine trace channel is reserved by the COLLISION_GIZMO definition
 	ECC_EngineTraceChannel2 UMETA(Hidden),
 	ECC_EngineTraceChannel3 UMETA(Hidden),
 	ECC_EngineTraceChannel4 UMETA(Hidden), 
 	ECC_EngineTraceChannel5 UMETA(Hidden),
+	ECC_EngineTraceChannel6 UMETA(Hidden),
 
 	// in order to use this custom channels
 	// we recommend to define in your local file
@@ -479,6 +479,9 @@ enum ECollisionChannel
 };
 
 
+#define COLLISION_GIZMO ECC_EngineTraceChannel1
+
+
 /** @note, if you change this, change GetCollisionChannelFromOverlapFilter() to match */
 UENUM(BlueprintType)
 enum EOverlapFilterOption
@@ -706,8 +709,6 @@ struct ENGINE_API FCollisionResponseContainer
 	UPROPERTY(EditAnywhere, BlueprintReadOnly, Category=CollisionResponseContainer, meta=(DisplayName="Destructible"))
 	TEnumAsByte<enum ECollisionResponse> Destructible;    // 7
 
-	UPROPERTY(EditAnywhere, BlueprintReadOnly, Category=CollisionResponseContainer, meta=(DisplayName="EditorGizmo"))
-	TEnumAsByte<enum ECollisionResponse> EditorGizmo;    // 7
 
 	///////////////////////////////////////
 	// Unspecified Engine Trace Channels
@@ -727,6 +728,9 @@ struct ENGINE_API FCollisionResponseContainer
 	UPROPERTY(EditAnywhere, BlueprintReadOnly, Category=CollisionResponseContainer)
 	TEnumAsByte<enum ECollisionResponse> EngineTraceChannel5;    // 12
 
+	UPROPERTY(EditAnywhere, BlueprintReadOnly, Category=CollisionResponseContainer)
+	TEnumAsByte<enum ECollisionResponse> EngineTraceChannel6;    // 13
+
 	///////////////////////////////////////
 	// in order to use this custom channels
 	// we recommend to define in your local file
@@ -811,14 +815,14 @@ struct ENGINE_API FCollisionResponseContainer
 			uint8 PhysicsBody;			// 5
 			uint8 Vehicle;				// 6
 			uint8 Destructible;			// 7
-			uint8 EditorGizmo;			// 8
 
 			// Unspecified Engine Trace Channels
-			uint8 EngineTraceChannel1;   // 9
-			uint8 EngineTraceChannel2;   // 10
-			uint8 EngineTraceChannel3;   // 11
-			uint8 EngineTraceChannel4;   // 12
-			uint8 EngineTraceChannel5;   // 13
+			uint8 EngineTraceChannel1;   // 8
+			uint8 EngineTraceChannel2;   // 9
+			uint8 EngineTraceChannel3;   // 10
+			uint8 EngineTraceChannel4;   // 11
+			uint8 EngineTraceChannel5;   // 12
+			uint8 EngineTraceChannel6;   // 13
 
 			// Unspecified Game Trace Channels
 			uint8 GameTraceChannel1;     // 14
@@ -3542,4 +3546,4 @@ struct FUserActivity
 	FUserActivity(const FString& InActionName)
 		: ActionName(InActionName)
 	{ }
-};
+};
\ No newline at end of file
diff --git a/Engine/Source/Runtime/Engine/Private/Collision/CollisionProfile.cpp b/Engine/Source/Runtime/Engine/Private/Collision/CollisionProfile.cpp
--- a/Engine/Source/Runtime/Engine/Private/Collision/CollisionProfile.cpp
+++ b/Engine/Source/Runtime/Engine/Private/Collision/CollisionProfile.cpp
@@ -6,7 +6,7 @@ DEFINE_LOG_CATEGORY_STATIC(LogCollisionProfile, Warning, All)
 
 #define MIN_CUSTOMIZABLE_COLLISIONCHANNEL	ECC_GameTraceChannel1
 #define MAX_CUSTOMIZABLE_COLLISIONCHANNEL	ECC_GameTraceChannel18
-#define IS_VALID_COLLISIONCHANNEL(x) ((x)> ECC_EditorGizmo && (x) < ECC_OverlapAll_Deprecated)
+#define IS_VALID_COLLISIONCHANNEL(x) ((x)> ECC_Destructible && (x) < ECC_OverlapAll_Deprecated)
 
 // do not chnage this name. This value is serialize to other object, if you change, it will mess up serialization, you'll need to fix up name by versioning
 FName UCollisionProfile::CustomCollisionProfileName = FName(TEXT("Custom"));
@@ -338,7 +338,7 @@ void UCollisionProfile::LoadProfileConfig(bool bForceInit)
 			ECollisionChannel CollisionChannel = (ECollisionChannel)EnumIndex;
 			// for any engine level collision profile, we hard coded here
 			// meta data doesn't work in cooked build, so we'll have to manually handle them
-			if ((CollisionChannel == ECC_Visibility) || (CollisionChannel == ECC_Camera) || (CollisionChannel == ECC_EditorGizmo))
+			if ((CollisionChannel == ECC_Visibility) || (CollisionChannel == ECC_Camera))
 			{
 				// remove from object query flags
 				FCollisionQueryFlag::Get().RemoveFromAllObjectsQueryFlag(CollisionChannel);
@@ -356,7 +356,7 @@ void UCollisionProfile::LoadProfileConfig(bool bForceInit)
 	// Now Load Channel setups, and set display names if customized
 	// also initialize DefaultResposneContainer with default response for each channel
 	FCollisionResponseContainer::DefaultResponseContainer.SetAllChannels(ECR_Block);
-	FCollisionResponseContainer::DefaultResponseContainer.SetResponse( ECC_EditorGizmo, ECR_Ignore );
+	FCollisionResponseContainer::DefaultResponseContainer.SetResponse( COLLISION_GIZMO, ECR_Ignore );
 
 	// we can't guarantee that DefaultChannelResponses was loaded ordered (from
 	// config) - in the following loop, we fill out TraceTypeMapping and 
