UE4 - Fix ensures related to ECVF_Scalability in Orion. Merge fallout.

diff --git a/Engine/Source/Runtime/Engine/Private/Texture2D.cpp b/Engine/Source/Runtime/Engine/Private/Texture2D.cpp
--- a/Engine/Source/Runtime/Engine/Private/Texture2D.cpp
+++ b/Engine/Source/Runtime/Engine/Private/Texture2D.cpp
@@ -37,7 +37,7 @@ static TAutoConsoleVariable<float> CVarSetMipMapLODBias(
 	TEXT("r.MipMapLODBias"),
 	0.0f,
 	TEXT("Apply additional mip map bias for all 2D textures, range of -15.0 to 15.0"),
-	ECVF_RenderThreadSafe);
+	ECVF_RenderThreadSafe | ECVF_Scalability);
 
 static TAutoConsoleVariable<int32> CVarVirtualTextureEnabled(
 	TEXT("r.VirtualTexture"),
