- Fixed compile error on Linux that came from Dev-Rendering

diff --git a/Engine/Source/Runtime/Engine/Public/Scalability.h b/Engine/Source/Runtime/Engine/Public/Scalability.h
--- a/Engine/Source/Runtime/Engine/Public/Scalability.h
+++ b/Engine/Source/Runtime/Engine/Public/Scalability.h
@@ -93,31 +93,31 @@ namespace Scalability
 
 
 	// A set of accessor helper objects that allow fast and thread safe access to the scalability CVars values without giving direct access.
-	ENGINE_API struct FResolutionQualityCVarAccessor : public FCVarAccessorFloat
+	struct ENGINE_API FResolutionQualityCVarAccessor : public FCVarAccessorFloat
 	{
 		FResolutionQualityCVarAccessor();
 	};
-	ENGINE_API struct FViewDistanceQualityCVarAccessor : public FCVarAccessorInt
+	struct ENGINE_API FViewDistanceQualityCVarAccessor : public FCVarAccessorInt
 	{
 		FViewDistanceQualityCVarAccessor();
 	};
-	ENGINE_API struct FAntiAliasingQualityCVarAccessor : public FCVarAccessorInt
+	struct ENGINE_API FAntiAliasingQualityCVarAccessor : public FCVarAccessorInt
 	{
 		FAntiAliasingQualityCVarAccessor();
 	};
-	ENGINE_API struct FShadowQualityCVarAccessor : public FCVarAccessorInt
+	struct ENGINE_API FShadowQualityCVarAccessor : public FCVarAccessorInt
 	{
 		FShadowQualityCVarAccessor();
 	};
-	ENGINE_API struct FPostProcessQualityCVarAccessor : public FCVarAccessorInt
+	struct ENGINE_API FPostProcessQualityCVarAccessor : public FCVarAccessorInt
 	{
 		FPostProcessQualityCVarAccessor();
 	};
-	ENGINE_API struct FTextureQualityCVarAccessor : public FCVarAccessorInt
+	struct ENGINE_API FTextureQualityCVarAccessor : public FCVarAccessorInt
 	{
 		FTextureQualityCVarAccessor();
 	};
-	ENGINE_API struct FEffectsQualityCVarAccessor : public FCVarAccessorInt
+	struct ENGINE_API FEffectsQualityCVarAccessor : public FCVarAccessorInt
 	{
 		FEffectsQualityCVarAccessor();
 	};
