fix for writing out DefaultEngine.ini into Engine/Config at startup.

diff --git a/Engine/Source/Runtime/Engine/Private/RenderingSettings.cpp b/Engine/Source/Runtime/Engine/Private/RenderingSettings.cpp
--- a/Engine/Source/Runtime/Engine/Private/RenderingSettings.cpp
+++ b/Engine/Source/Runtime/Engine/Private/RenderingSettings.cpp
@@ -40,9 +40,6 @@ void URendererSettings::PostInitProperties()
 
 	SanatizeReflectionCaptureResolution();
 
-	UUserInterfaceSettings* UISettings = GetMutableDefault<UUserInterfaceSettings>(UUserInterfaceSettings::StaticClass());
-	UISettings->UpdateDefaultConfigFile();
-
 #if WITH_EDITOR
 	if (IsTemplate())
 	{
