- Redoing a linux compile fix for doing some linux testing in main before we copy up

diff --git a/Engine/Source/Runtime/Engine/Private/Components/ReflectionCaptureComponent.cpp b/Engine/Source/Runtime/Engine/Private/Components/ReflectionCaptureComponent.cpp
--- a/Engine/Source/Runtime/Engine/Private/Components/ReflectionCaptureComponent.cpp
+++ b/Engine/Source/Runtime/Engine/Private/Components/ReflectionCaptureComponent.cpp
@@ -1179,7 +1179,7 @@ void UReflectionCaptureComponent::PostEditChangeProperty(FPropertyChangedEvent&
 
 	if (PropertyChangedEvent.Property
 		&& (PropertyChangedEvent.Property->GetFName() == GET_MEMBER_NAME_CHECKED(UReflectionCaptureComponent, bVisible)
-			|| PropertyChangedEvent.MemberProperty && PropertyChangedEvent.MemberProperty->GetFName() == GET_MEMBER_NAME_CHECKED(UReflectionCaptureComponent, CaptureOffset)))
+			|| (PropertyChangedEvent.MemberProperty && PropertyChangedEvent.MemberProperty->GetFName() == GET_MEMBER_NAME_CHECKED(UReflectionCaptureComponent, CaptureOffset))))
 	{
 		SetCaptureIsDirty();
 	}
