Copying //UE4/Dev-Sequencer to //UE4/Main (Source: //UE4/Dev-Sequencer @ 2945541)

diff --git a/Engine/Source/Runtime/Engine/Classes/Camera/CameraComponent.h b/Engine/Source/Runtime/Engine/Classes/Camera/CameraComponent.h
--- a/Engine/Source/Runtime/Engine/Classes/Camera/CameraComponent.h
+++ b/Engine/Source/Runtime/Engine/Classes/Camera/CameraComponent.h
@@ -15,57 +15,57 @@ class ENGINE_API UCameraComponent : public USceneComponent
 	GENERATED_UCLASS_BODY()
 
 	/** The horizontal field of view (in degrees) in perspective mode (ignored in Orthographic mode) */
-	UPROPERTY(Interp, EditAnywhere, BlueprintReadWrite, Category=CameraSettings, meta=(UIMin = "5.0", UIMax = "170", ClampMin = "0.001", ClampMax = "360.0", Units = deg))
+	UPROPERTY(Interp, EditAnywhere, BlueprintReadWrite, Category = CameraSettings, meta = (UIMin = "5.0", UIMax = "170", ClampMin = "0.001", ClampMax = "360.0", Units = deg))
 	float FieldOfView;
-	UFUNCTION(BlueprintCallable, Category=Camera)
+	UFUNCTION(BlueprintCallable, Category = Camera)
 	void SetFieldOfView(float InFieldOfView) { FieldOfView = InFieldOfView; }
 
 	/** The desired width (in world units) of the orthographic view (ignored in Perspective mode) */
-	UPROPERTY(Interp, EditAnywhere, BlueprintReadWrite, Category=CameraSettings)
+	UPROPERTY(Interp, EditAnywhere, BlueprintReadWrite, Category = CameraSettings)
 	float OrthoWidth;
-	UFUNCTION(BlueprintCallable, Category=Camera)
+	UFUNCTION(BlueprintCallable, Category = Camera)
 	void SetOrthoWidth(float InOrthoWidth) { OrthoWidth = InOrthoWidth; }
 
 	/** The near plane distance of the orthographic view (in world units) */
-	UPROPERTY(Interp, EditAnywhere, BlueprintReadWrite, Category=CameraSettings)
+	UPROPERTY(Interp, EditAnywhere, BlueprintReadWrite, Category = CameraSettings)
 	float OrthoNearClipPlane;
-	UFUNCTION(BlueprintCallable, Category=Camera)
+	UFUNCTION(BlueprintCallable, Category = Camera)
 	void SetOrthoNearClipPlane(float InOrthoNearClipPlane) { OrthoNearClipPlane = InOrthoNearClipPlane; }
 
 	/** The far plane distance of the orthographic view (in world units) */
-	UPROPERTY(Interp, EditAnywhere, BlueprintReadWrite, Category=CameraSettings)
+	UPROPERTY(Interp, EditAnywhere, BlueprintReadWrite, Category = CameraSettings)
 	float OrthoFarClipPlane;
-	UFUNCTION(BlueprintCallable, Category=Camera)
+	UFUNCTION(BlueprintCallable, Category = Camera)
 	void SetOrthoFarClipPlane(float InOrthoFarClipPlane) { OrthoFarClipPlane = InOrthoFarClipPlane; }
 
 	// Aspect Ratio (Width/Height)
-	UPROPERTY(Interp, EditAnywhere, BlueprintReadWrite, Category=CameraSettings, meta=(ClampMin = "0.1", ClampMax = "100.0", EditCondition="bConstrainAspectRatio"))
+	UPROPERTY(Interp, EditAnywhere, BlueprintReadWrite, Category = CameraSettings, meta = (ClampMin = "0.1", ClampMax = "100.0", EditCondition = "bConstrainAspectRatio"))
 	float AspectRatio;
-	UFUNCTION(BlueprintCallable, Category=Camera)
+	UFUNCTION(BlueprintCallable, Category = Camera)
 	void SetAspectRatio(float InAspectRatio) { AspectRatio = InAspectRatio; }
 
 	// If bConstrainAspectRatio is true, black bars will be added if the destination view has a different aspect ratio than this camera requested.
-	UPROPERTY(Interp, EditAnywhere, BlueprintReadWrite, Category=CameraSettings)
-	uint32 bConstrainAspectRatio:1;
-	UFUNCTION(BlueprintCallable, Category=Camera)
+	UPROPERTY(Interp, EditAnywhere, BlueprintReadWrite, Category = CameraSettings)
+	uint32 bConstrainAspectRatio : 1;
+	UFUNCTION(BlueprintCallable, Category = Camera)
 	void SetConstraintAspectRatio(bool bInConstrainAspectRatio) { bConstrainAspectRatio = bInConstrainAspectRatio; }
 
 	// If true, account for the field of view angle when computing which level of detail to use for meshes.
-	UPROPERTY(Interp, EditAnywhere, AdvancedDisplay, BlueprintReadWrite, Category=CameraSettings)
-	uint32 bUseFieldOfViewForLOD:1;
-	UFUNCTION(BlueprintCallable, Category=Camera)
+	UPROPERTY(Interp, EditAnywhere, AdvancedDisplay, BlueprintReadWrite, Category = CameraSettings)
+	uint32 bUseFieldOfViewForLOD : 1;
+	UFUNCTION(BlueprintCallable, Category = Camera)
 	void SetUseFieldOfViewForLOD(bool bInUseFieldOfViewForLOD) { bUseFieldOfViewForLOD = bInUseFieldOfViewForLOD; }
 
 	/** True if the camera's orientation and position should be locked to the HMD */
-	UPROPERTY(EditAnywhere, BlueprintReadWrite, Category=CameraSettings)
-	uint32 bLockToHmd:1;
+	UPROPERTY(EditAnywhere, BlueprintReadWrite, Category = CameraSettings)
+	uint32 bLockToHmd : 1;
 
 	/**
 	 * If this camera component is placed on a pawn, should it use the view/control rotation of the pawn where possible?
 	 * @see APawn::GetViewRotation()
 	 */
-	UPROPERTY(EditAnywhere, BlueprintReadWrite, Category=CameraSettings)
-	uint32 bUsePawnControlRotation:1;
+	UPROPERTY(EditAnywhere, BlueprintReadWrite, Category = CameraSettings)
+	uint32 bUsePawnControlRotation : 1;
 
 protected:
 	/** True to enable the additive view offset, for adjusting the view without moving the component. */
@@ -73,15 +73,15 @@ protected:
 
 public:
 	// The type of camera
-	UPROPERTY(Interp, EditAnywhere, BlueprintReadWrite, Category=CameraSettings)
+	UPROPERTY(Interp, EditAnywhere, BlueprintReadWrite, Category = CameraSettings)
 	TEnumAsByte<ECameraProjectionMode::Type> ProjectionMode;
-	UFUNCTION(BlueprintCallable, Category=Camera)
+	UFUNCTION(BlueprintCallable, Category = Camera)
 	void SetProjectionMode(ECameraProjectionMode::Type InProjectionMode) { ProjectionMode = InProjectionMode; }
 
 	/** Indicates if PostProcessSettings should be used when using this Camera to view through. */
-	UPROPERTY(Interp, EditAnywhere, BlueprintReadWrite, Category=PostProcess, meta=(UIMin = "0.0", UIMax = "1.0"))
+	UPROPERTY(Interp, EditAnywhere, BlueprintReadWrite, Category = PostProcess, meta = (UIMin = "0.0", UIMax = "1.0"))
 	float PostProcessBlendWeight;
-	UFUNCTION(BlueprintCallable, Category=Camera)
+	UFUNCTION(BlueprintCallable, Category = Camera)
 	void SetPostProcessBlendWeight(float InPostProcessBlendWeight) { PostProcessBlendWeight = InPostProcessBlendWeight; }
 
 	/** Post process settings to use for this camera. Don't forget to check the properties you want to override */
@@ -109,11 +109,11 @@ public:
 	 * Returns camera's Point of View.
 	 * Called by Camera class. Subclass and postprocess to add any effects.
 	 */
-	UFUNCTION(BlueprintCallable, Category=Camera)
+	UFUNCTION(BlueprintCallable, Category = Camera)
 	virtual void GetCameraView(float DeltaTime, FMinimalViewInfo& DesiredView);
 
 	/** Adds an Blendable (implements IBlendableInterface) to the array of Blendables (if it doesn't exist) and update the weight */
-	UFUNCTION(BlueprintCallable, Category="Rendering")
+	UFUNCTION(BlueprintCallable, Category = "Rendering")
 	void AddOrUpdateBlendable(TScriptInterface<IBlendableInterface> InBlendableObject, float InWeight = 1.0f) { PostProcessSettings.AddBlendable(InBlendableObject, InWeight); }
 
 protected:
@@ -123,7 +123,7 @@ protected:
 
 	UPROPERTY(transient)
 	class UStaticMesh* CameraMesh;
-	
+
 	// The camera mesh to show visually where the camera is placed
 	class UStaticMeshComponent* ProxyMeshComponent;
 
@@ -138,23 +138,23 @@ protected:
 	/** An optional extra FOV offset to adjust the final view without modifying the component */
 	float AdditiveFOVOffset;
 
+	/** Optional extra PostProcessing blends stored for this camera. They are not applied automatically in GetCameraView. */
+	TArray<FPostProcessSettings> ExtraPostProcessBlends;
+	TArray<float> ExtraPostProcessBlendWeights;
+
 public:
 
 	/** Applies the given additive offset, preserving any existing offset */
-	void AddAdditiveOffset(FTransform const& Transform, float FOV)
-	{
-		bUseAdditiveOffset = true;
-		AdditiveOffset = AdditiveOffset * Transform;
-		AdditiveFOVOffset += FOV;
-	}
-
+	void AddAdditiveOffset(FTransform const& Transform, float FOV);
 	/** Removes any additive offset. */
-	void ClearAdditiveOffset()
-	{
-		bUseAdditiveOffset = false;
-		AdditiveOffset = FTransform::Identity;
-		AdditiveFOVOffset = 0.f;
-	}
+	void ClearAdditiveOffset();
+
+	/** Stores a given PP and weight to be later applied when the final PP is computed. Used for things like in-editor camera animation preview. */
+	void AddExtraPostProcessBlend(FPostProcessSettings const& PPSettings, float PPBlendWeight);
+	/** Removes any extra PP blends. */
+	void ClearExtraPostProcessBlends();
+	/** Returns any extra PP blends that were stored. */
+	void GetExtraPostProcessBlends(TArray<FPostProcessSettings>& OutSettings, TArray<float>& OutWeights) const;
 
 	/** 
 	 * Can be called from external code to notify that this camera was cut to, so it can update 
diff --git a/Engine/Source/Runtime/Engine/Classes/Camera/CameraShake.h b/Engine/Source/Runtime/Engine/Classes/Camera/CameraShake.h
--- a/Engine/Source/Runtime/Engine/Classes/Camera/CameraShake.h
+++ b/Engine/Source/Runtime/Engine/Classes/Camera/CameraShake.h
@@ -223,6 +223,9 @@ protected:
 	/** Matrix defining the playspace, used when PlaySpace == CAPS_UserDefined */
 	FMatrix UserPlaySpaceMatrix;
 
+	/** Temp actor to use for playing camera anims. Used when playing a camera anim in non-gameplay context, e.g. in the editor */
+	AActor* TempCameraActorForCameraAnims;
+
 public:
 	/** Overall intensity scale for this shake instance. */
 	UPROPERTY(transient, BlueprintReadOnly, Category = CameraShake)
@@ -274,6 +277,8 @@ public:
 
 	/** Sets current playback time and applies the shake (both oscillation and cameraanim) to the given POV. */
 	void SetCurrentTimeAndApplyShake(float NewTime, FMinimalViewInfo& POV);
+
+	void SetTempCameraAnimActor(AActor* Actor) { TempCameraActorForCameraAnims = Actor; }
 };
 
 
diff --git a/Engine/Source/Runtime/Engine/Classes/Engine/LevelStreaming.h b/Engine/Source/Runtime/Engine/Classes/Engine/LevelStreaming.h
--- a/Engine/Source/Runtime/Engine/Classes/Engine/LevelStreaming.h
+++ b/Engine/Source/Runtime/Engine/Classes/Engine/LevelStreaming.h
@@ -123,7 +123,7 @@ public:
 	uint32 bShouldBeLoaded:1;
 
 	/** Whether the level should be visible if it is loaded																		*/
-	UPROPERTY(Category=LevelStreaming, BlueprintReadWrite)
+	UPROPERTY(EditAnywhere, Category=LevelStreaming, BlueprintReadWrite)
 	uint32 bShouldBeVisible:1;
 
 	/** Whether we want to force a blocking load																				*/
@@ -277,7 +277,7 @@ public:
 
 	/** Returns the Level Script Actor of the level if the level is loaded and valid */
 	UFUNCTION(BlueprintPure, meta = (BlueprintInternalUseOnly = "true"))
-	ALevelScriptActor* GetLevelScriptActor();
+	ENGINE_API ALevelScriptActor* GetLevelScriptActor();
 
 	//~==============================================================================================
 	// Delegates
diff --git a/Engine/Source/Runtime/Engine/Classes/Matinee/InterpData.h b/Engine/Source/Runtime/Engine/Classes/Matinee/InterpData.h
--- a/Engine/Source/Runtime/Engine/Classes/Matinee/InterpData.h
+++ b/Engine/Source/Runtime/Engine/Classes/Matinee/InterpData.h
@@ -72,10 +72,10 @@ class UInterpData : public UObject
 	//~ End UObject Interface
 
 	/** Search through all InterpGroups in this InterpData to find a group whose GroupName matches the given name. Returns INDEX_NONE if group not found. */
-	int32 FindGroupByName( FName GroupName );
+	ENGINE_API int32 FindGroupByName( FName GroupName );
 
 	/** Search through all InterpGroups in this InterpData to find a group whose GroupName matches the given name. Returns INDEX_NONE if not group found. */
-	int32 FindGroupByName( const FString& InGroupName );
+	ENGINE_API int32 FindGroupByName( const FString& InGroupName );
 
 	/** Search through all groups to find all tracks of the given class. */
 	ENGINE_API void FindTracksByClass(UClass* TrackClass, TArray<class UInterpTrack*>& OutputTracks);
diff --git a/Engine/Source/Runtime/Engine/Private/Actor.cpp b/Engine/Source/Runtime/Engine/Private/Actor.cpp
--- a/Engine/Source/Runtime/Engine/Private/Actor.cpp
+++ b/Engine/Source/Runtime/Engine/Private/Actor.cpp
@@ -267,7 +267,14 @@ UWorld* AActor::GetWorld() const
 {
 	// CDO objects do not belong to a world
 	// If the actors outer is destroyed or unreachable we are shutting down and the world should be NULL
-	return (!HasAnyFlags(RF_ClassDefaultObject) && !GetOuter()->HasAnyFlags(RF_BeginDestroyed) && !GetOuter()->IsUnreachable() ? GetLevel()->OwningWorld : NULL);
+	if (!HasAnyFlags(RF_ClassDefaultObject) && !GetOuter()->HasAnyFlags(RF_BeginDestroyed) && !GetOuter()->IsUnreachable())
+	{
+		if (ULevel* Level = GetLevel())
+		{
+			return Level->OwningWorld;
+		}
+	}
+	return nullptr;
 }
 
 FTimerManager& AActor::GetWorldTimerManager() const
@@ -1626,7 +1633,7 @@ bool AActor::IsInLevel(const ULevel *TestLevel) const
 
 ULevel* AActor::GetLevel() const
 {
-	return CastChecked<ULevel>( GetOuter() );
+	return Cast<ULevel>(GetOuter());
 }
 
 bool AActor::IsInPersistentLevel(bool bIncludeLevelStreamingPersistent) const
diff --git a/Engine/Source/Runtime/Engine/Private/Animation/AnimInstance.cpp b/Engine/Source/Runtime/Engine/Private/Animation/AnimInstance.cpp
--- a/Engine/Source/Runtime/Engine/Private/Animation/AnimInstance.cpp
+++ b/Engine/Source/Runtime/Engine/Private/Animation/AnimInstance.cpp
@@ -940,6 +940,11 @@ void UAnimInstance::DisplayDebug(class UCanvas* Canvas, const FDebugDisplayInfo&
 	}
 }
 
+void UAnimInstance::ResetDynamics()
+{
+	GetProxyOnGameThread<FAnimInstanceProxy>().ResetDynamics();
+}
+
 void UAnimInstance::RecalcRequiredBones()
 {
 	USkeletalMeshComponent* SkelMeshComp = GetSkelMeshComponent();
@@ -1710,7 +1715,13 @@ void UAnimInstance::StopSlotAnimation(float InBlendOutTime, FName SlotNodeName)
 	}
 }
 
-bool UAnimInstance::IsPlayingSlotAnimation(UAnimSequenceBase* Asset, FName SlotNodeName )
+bool UAnimInstance::IsPlayingSlotAnimation(UAnimSequenceBase* Asset, FName SlotNodeName)
+{
+	UAnimMontage* Montage = nullptr;
+	return IsPlayingSlotAnimation(Asset, SlotNodeName, Montage);
+}
+
+bool UAnimInstance::IsPlayingSlotAnimation(UAnimSequenceBase* Asset, FName SlotNodeName, UAnimMontage*& OutMontage)
 {
 	for (int32 InstanceIndex = 0; InstanceIndex < MontageInstances.Num(); InstanceIndex++)
 	{
@@ -1725,7 +1736,7 @@ bool UAnimInstance::IsPlayingSlotAnimation(UAnimSequenceBase* Asset, FName SlotN
 				const FAnimTrack* AnimTrack = CurMontage->GetAnimationData(SlotNodeName);
 				if (AnimTrack && AnimTrack->AnimSegments.Num() == 1)
 				{
-					// find if the 
+					OutMontage = CurMontage;
 					return (AnimTrack->AnimSegments[0].AnimReference == Asset);
 				}
 			}
diff --git a/Engine/Source/Runtime/Engine/Private/Animation/AnimInstanceProxy.cpp b/Engine/Source/Runtime/Engine/Private/Animation/AnimInstanceProxy.cpp
--- a/Engine/Source/Runtime/Engine/Private/Animation/AnimInstanceProxy.cpp
+++ b/Engine/Source/Runtime/Engine/Private/Animation/AnimInstanceProxy.cpp
@@ -116,6 +116,7 @@ void FAnimInstanceProxy::InitializeRootNode()
 	if (RootNode != nullptr)
 	{
 		GameThreadPreUpdateNodes.Reset();
+		DynamicResetNodes.Reset();
 
 		// cache any state machine descriptions we have
 		for(UStructProperty* Property : AnimClassInterface->GetAnimNodeProperties())
@@ -136,6 +137,11 @@ void FAnimInstanceProxy::InitializeRootNode()
 						GameThreadPreUpdateNodes.Add(AnimNode);
 					}
 
+					if(AnimNode->NeedsDynamicReset())
+					{
+						DynamicResetNodes.Add(AnimNode);
+					}
+
 					if(Property->Struct->IsChildOf(FAnimNode_StateMachine::StaticStruct()))
 					{
 						FAnimNode_StateMachine* StateMachine = static_cast<FAnimNode_StateMachine*>(AnimNode);
@@ -1475,4 +1481,12 @@ void FAnimInstanceProxy::RecordStateWeight(const int32& InMachineClassIndex, con
 	}
 }
 
-#undef LOCTEXT_NAMESPACE
\ No newline at end of file
+void FAnimInstanceProxy::ResetDynamics()
+{
+	for(FAnimNode_Base* Node : DynamicResetNodes)
+	{
+		Node->ResetDynamics();
+	}
+}
+
+#undef LOCTEXT_NAMESPACE
diff --git a/Engine/Source/Runtime/Engine/Private/Animation/AnimMontage.cpp b/Engine/Source/Runtime/Engine/Private/Animation/AnimMontage.cpp
--- a/Engine/Source/Runtime/Engine/Private/Animation/AnimMontage.cpp
+++ b/Engine/Source/Runtime/Engine/Private/Animation/AnimMontage.cpp
@@ -2031,7 +2031,8 @@ void FAnimMontageInstance::SetMatineeAnimPositionInner(FName SlotName, USkeletal
 		}
 		else
 		{
-			bool bShouldChange = AnimInst->IsPlayingSlotAnimation(InAnimSequence, SlotName) == false;
+			UAnimMontage* PlayingMontage = nullptr;
+			bool bShouldChange = AnimInst->IsPlayingSlotAnimation(InAnimSequence, SlotName, PlayingMontage) == false;
 			if(bShouldChange)
 			{
 				if(CurrentlyPlayingMontage.IsValid())
@@ -2049,8 +2050,11 @@ void FAnimMontageInstance::SetMatineeAnimPositionInner(FName SlotName, USkeletal
 				AnimInst->PlaySlotAnimationAsDynamicMontage(InAnimSequence, SlotName, 0.0f, 0.0f, 0.f, 1);
 				CurrentlyPlayingMontage = AnimInst->GetCurrentActiveMontage();
 			}
-
-			ensure(CurrentlyPlayingMontage.IsValid());
+			else
+			{
+				CurrentlyPlayingMontage = PlayingMontage;
+				ensure(CurrentlyPlayingMontage.IsValid());
+			}
 
 			struct FAnimMontageInstance* AnimMontageInst = AnimInst->GetActiveInstanceForMontage(*CurrentlyPlayingMontage);
 			if(AnimMontageInst)
@@ -2101,7 +2105,8 @@ void FAnimMontageInstance::PreviewMatineeSetAnimPositionInner(FName SlotName, US
 	}
 	else if (AnimInst)
 	{
-		bool bShouldChange = AnimInst->IsPlayingSlotAnimation(InAnimSequence, SlotName) == false;
+		UAnimMontage* PlayingMontage = nullptr;
+		bool bShouldChange = AnimInst->IsPlayingSlotAnimation(InAnimSequence, SlotName, PlayingMontage) == false;
 		if(bShouldChange)
 		{
 			if(CurrentlyPlayingMontage.IsValid())
@@ -2119,8 +2124,11 @@ void FAnimMontageInstance::PreviewMatineeSetAnimPositionInner(FName SlotName, US
 			AnimInst->PlaySlotAnimationAsDynamicMontage(InAnimSequence, SlotName, 0.0f, 0.0f, 0.f, 1);
 			CurrentlyPlayingMontage = AnimInst->GetCurrentActiveMontage();
 		}
-
-		ensure(CurrentlyPlayingMontage.IsValid());
+		else
+		{
+			CurrentlyPlayingMontage = PlayingMontage;
+			ensure(CurrentlyPlayingMontage.IsValid());
+		}
 
 		struct FAnimMontageInstance* AnimMontageInst = AnimInst->GetActiveInstanceForMontage(*CurrentlyPlayingMontage);
 		if(AnimMontageInst)
diff --git a/Engine/Source/Runtime/Engine/Private/BufferVisualizationData.cpp b/Engine/Source/Runtime/Engine/Private/BufferVisualizationData.cpp
--- a/Engine/Source/Runtime/Engine/Private/BufferVisualizationData.cpp
+++ b/Engine/Source/Runtime/Engine/Private/BufferVisualizationData.cpp
@@ -20,6 +20,12 @@ static TAutoConsoleVariable<int32> CVarBufferVisualizationDumpFrames(
 	TEXT("1:on"),
 	ECVF_RenderThreadSafe);
 
+static TAutoConsoleVariable<int32> CVarSaveUncompressedEXRFrames(
+	TEXT("r.SaveUncompressedEXRFrames"),
+	0,
+	TEXT("If enabled will save EXR screenshots with no compression, rather then the zip default"),
+	ECVF_RenderThreadSafe);
+
 void FBufferVisualizationData::Initialize()
 {
 	if (!bIsInitialized)
diff --git a/Engine/Source/Runtime/Engine/Private/Camera/CameraAnimInst.cpp b/Engine/Source/Runtime/Engine/Private/Camera/CameraAnimInst.cpp
--- a/Engine/Source/Runtime/Engine/Private/Camera/CameraAnimInst.cpp
+++ b/Engine/Source/Runtime/Engine/Private/Camera/CameraAnimInst.cpp
@@ -186,15 +186,6 @@ void UCameraAnimInst::SetScale(float NewScale)
 {
 	BasePlayScale = NewScale;
 }
-
-void UCameraAnimInst::SetCameraActor(class AActor* Actor)
-{
-	if (InterpGroupInst)
-	{
-		InterpGroupInst->SetGroupActor(Actor);
-	}
-}
-
 static const FName NAME_CameraComponentFieldOfViewPropertyName(TEXT("CameraComponent.FieldOfView"));
 
 void UCameraAnimInst::Play(UCameraAnim* Anim, class AActor* CamActor, float InRate, float InScale, float InBlendInTime, float InBlendOutTime, bool bInLooping, bool bRandomStartTime, float Duration)
diff --git a/Engine/Source/Runtime/Engine/Private/Camera/CameraComponent.cpp b/Engine/Source/Runtime/Engine/Private/Camera/CameraComponent.cpp
--- a/Engine/Source/Runtime/Engine/Private/Camera/CameraComponent.cpp
+++ b/Engine/Source/Runtime/Engine/Private/Camera/CameraComponent.cpp
@@ -326,6 +326,41 @@ void UCameraComponent::NotifyCameraCut()
 	}
 };
 
+void UCameraComponent::AddAdditiveOffset(FTransform const& Transform, float FOV)
+{
+	bUseAdditiveOffset = true;
+	AdditiveOffset = AdditiveOffset * Transform;
+	AdditiveFOVOffset += FOV;
+}
+
+/** Removes any additive offset. */
+void UCameraComponent::ClearAdditiveOffset()
+{
+	bUseAdditiveOffset = false;
+	AdditiveOffset = FTransform::Identity;
+	AdditiveFOVOffset = 0.f;
+}
+
+
+void UCameraComponent::AddExtraPostProcessBlend(FPostProcessSettings const& PPSettings, float PPBlendWeight)
+{
+	checkSlow(ExtraPostProcessBlends.Num() == ExtraPostProcessBlendWeights.Num());
+	ExtraPostProcessBlends.Add(PPSettings);
+	ExtraPostProcessBlendWeights.Add(PPBlendWeight);
+}
+
+void UCameraComponent::ClearExtraPostProcessBlends()
+{
+	ExtraPostProcessBlends.Empty();
+	ExtraPostProcessBlendWeights.Empty();
+}
+
+void UCameraComponent::GetExtraPostProcessBlends(TArray<FPostProcessSettings>& OutSettings, TArray<float>& OutWeights) const
+{
+	OutSettings = ExtraPostProcessBlends;
+	OutWeights = ExtraPostProcessBlendWeights;
+}
+
 
 #undef LOCTEXT_NAMESPACE
 
diff --git a/Engine/Source/Runtime/Engine/Private/Camera/CameraShake.cpp b/Engine/Source/Runtime/Engine/Private/Camera/CameraShake.cpp
--- a/Engine/Source/Runtime/Engine/Private/Camera/CameraShake.cpp
+++ b/Engine/Source/Runtime/Engine/Private/Camera/CameraShake.cpp
@@ -183,8 +183,8 @@ void UCameraShake::PlayShake(APlayerCameraManager* Camera, float Scale, ECameraA
 					if (AnimInst)
 					{
 						// note: we don't have a temp camera actor necessary for evaluating a camera anim.
-						// caller is responsible in this case for providing one by calling SetCameraACtor(??) on the CamAnimInst
-						AnimInst->Play(Anim, nullptr, AnimPlayRate, FinalAnimScale, AnimBlendInTime, AnimBlendOutTime, bLoop, bRandomStart, Duration);
+						// caller is responsible in this case for providing one by calling SetTempCameraAnimActor() on the shake instance before playing the shake
+						AnimInst->Play(Anim, TempCameraActorForCameraAnims, AnimPlayRate, FinalAnimScale, AnimBlendInTime, AnimBlendOutTime, bLoop, bRandomStart, Duration);
 						AnimInst->SetPlaySpace(InPlaySpace, UserPlaySpaceRot);
 					}
 				}
diff --git a/Engine/Source/Runtime/Engine/Private/CurveBase.cpp b/Engine/Source/Runtime/Engine/Private/CurveBase.cpp
--- a/Engine/Source/Runtime/Engine/Private/CurveBase.cpp
+++ b/Engine/Source/Runtime/Engine/Private/CurveBase.cpp
@@ -209,6 +209,8 @@ void FNameCurve::DeleteKey(FKeyHandle KeyHandle)
 	int32 Index = GetIndex(KeyHandle);
 	Keys.RemoveAt(Index);
 
+	KeyHandlesToIndices.Remove(KeyHandle);
+
 	// update key indices
 	for (auto It = KeyHandlesToIndices.CreateIterator(); It; ++It)
 	{
diff --git a/Engine/Source/Runtime/Engine/Private/HighResScreenshot.cpp b/Engine/Source/Runtime/Engine/Private/HighResScreenshot.cpp
--- a/Engine/Source/Runtime/Engine/Private/HighResScreenshot.cpp
+++ b/Engine/Source/Runtime/Engine/Private/HighResScreenshot.cpp
@@ -204,10 +204,19 @@ bool FHighResScreenshotConfig::SaveImage(const FString& File, const TArray<TPixe
 
 	if (ImageWriter && ImageWriter->ImageWrapper->SetRaw((void*)&Bitmap[0], sizeof(TPixelType)* x * y, x, y, Traits::SourceChannelLayout, BitsPerPixel))
 	{
+		// Optionally disable EXR Compression
+		static const auto CVarSaveUncompressedEXRFrames = IConsoleManager::Get().FindTConsoleVariableDataInt(TEXT("r.SaveUncompressedEXRFrames"));
+		const bool bSaveUncompressedEXR = (CVarSaveUncompressedEXRFrames->GetValueOnRenderThread() > 0);
+
+		const int32 CompressionQuality = (bIsWritingHDRImage && bSaveUncompressedEXR)
+			? ImageCompression::CompressionQuality::Uncompressed
+			: ImageCompression::CompressionQuality::Default;
+
+		// Compress and write image
 		FArchive* Ar = FileManager->CreateFileWriter(Filename.GetCharArray().GetData());
 		if (Ar != nullptr)
 		{
-			const TArray<uint8>& CompressedData = ImageWriter->ImageWrapper->GetCompressed();
+			const TArray<uint8>& CompressedData = ImageWriter->ImageWrapper->GetCompressed(CompressionQuality);
 			int32 CompressedSize = CompressedData.Num();
 			Ar->Serialize((void*)CompressedData.GetData(), CompressedSize);
 			delete Ar;
diff --git a/Engine/Source/Runtime/Engine/Public/Animation/AnimInstanceProxy.h b/Engine/Source/Runtime/Engine/Public/Animation/AnimInstanceProxy.h
--- a/Engine/Source/Runtime/Engine/Public/Animation/AnimInstanceProxy.h
+++ b/Engine/Source/Runtime/Engine/Public/Animation/AnimInstanceProxy.h
@@ -260,6 +260,8 @@ public:
 	void RecordStateWeight(const int32& InMachineClassIndex, const int32& InStateIndex, const float& InStateWeight);
 
 	bool IsSlotNodeRelevantForNotifies(FName SlotNodeName) const;
+	/** Reset any dynamics running simulation-style updates (e.g. on teleport, time skip etc.) */
+	void ResetDynamics();
 
 	/** Only restricted classes can access the protected interface */
 	friend class UAnimInstance;
@@ -562,6 +564,9 @@ private:
 	/** Delegate fired on the game thread before update occurs */
 	TArray<FAnimNode_Base*> GameThreadPreUpdateNodes;
 
+	/** All nodes that need to be reset on DynamicReset() */
+	TArray<FAnimNode_Base*> DynamicResetNodes;
+	
 	/** Native transition rules */
 	TArray<FNativeTransitionBinding> NativeTransitionBindings;
 
diff --git a/Engine/Source/Runtime/Engine/Public/EngineUtils.h b/Engine/Source/Runtime/Engine/Public/EngineUtils.h
--- a/Engine/Source/Runtime/Engine/Public/EngineUtils.h
+++ b/Engine/Source/Runtime/Engine/Public/EngineUtils.h
@@ -246,11 +246,12 @@ public:
 				LocalCurrentActor = LocalSpawnedActorArray[LocalIndex - LocalObjectArray.Num()];
 			}
 			State->ConsideredCount++;
-
-			if ( LocalCurrentActor != NULL
+			
+			ULevel* ActorLevel = LocalCurrentActor ? LocalCurrentActor->GetLevel() : nullptr;
+			if ( ActorLevel
 				&& static_cast<const Derived*>(this)->IsActorSuitable(LocalCurrentActor)
-				&& static_cast<const Derived*>(this)->CanIterateLevel(LocalCurrentActor->GetLevel())
-				&& LocalCurrentActor->GetWorld() == LocalCurrentWorld)
+				&& static_cast<const Derived*>(this)->CanIterateLevel(ActorLevel)
+				&& ActorLevel->GetWorld() == LocalCurrentWorld)
 			{
 				// ignore non-persistent world settings
 				if (LocalCurrentActor->GetLevel() == LocalCurrentWorld->PersistentLevel || !LocalCurrentActor->IsA(AWorldSettings::StaticClass()))
diff --git a/Engine/Source/Runtime/ImageWrapper/Private/ExrImageWrapper.cpp b/Engine/Source/Runtime/ImageWrapper/Private/ExrImageWrapper.cpp
--- a/Engine/Source/Runtime/ImageWrapper/Private/ExrImageWrapper.cpp
+++ b/Engine/Source/Runtime/ImageWrapper/Private/ExrImageWrapper.cpp
@@ -9,6 +9,7 @@
 
 FExrImageWrapper::FExrImageWrapper()
 	: FImageWrapperBase()
+	, bUseCompression(true)
 {
 }
 
@@ -309,15 +310,16 @@ void FExrImageWrapper::CompressRaw(const sourcetype* SrcData, bool bIgnoreAlpha)
 		NumWriteComponents = 3;
 	}
 
-	Imf::Header Header(Width, Height);
-
+	Imf::Compression Comp = bUseCompression ? Imf::Compression::ZIP_COMPRESSION : Imf::Compression::NO_COMPRESSION;
+	Imf::Header Header(Width, Height, 1, Imath::V2f(0, 0), 1, Imf::LineOrder::INCREASING_Y, Comp);
+	
 	for (uint32 Channel = 0; Channel < NumWriteComponents; Channel++)
 	{
 		Header.channels().insert(GetRawChannelName(Channel), Imf::Channel(OutputFormat));
 	}
 
 	FMemFileOut MemFile("");
-	const int32 OutputPixelSize = ((OutputFormat == Imf::HALF) ? 2 : 4);
+	const int32 OutputPixelSize = ((OutputFormat == Imf::HALF && bUseCompression) ? 2 : 4);
 	MemFile.Data.AddUninitialized(Width * Height * NumWriteComponents * OutputPixelSize);
 
 	Imf::FrameBuffer ImfFrameBuffer;
@@ -348,6 +350,8 @@ void FExrImageWrapper::Compress( int32 Quality )
 
 	const int32 MaxComponents = 4;
 
+	bUseCompression = (Quality != ImageCompression::CompressionQuality::Uncompressed);
+
 	switch (RawBitDepth)
 	{
 		case 8:
diff --git a/Engine/Source/Runtime/ImageWrapper/Private/ExrImageWrapper.h b/Engine/Source/Runtime/ImageWrapper/Private/ExrImageWrapper.h
--- a/Engine/Source/Runtime/ImageWrapper/Private/ExrImageWrapper.h
+++ b/Engine/Source/Runtime/ImageWrapper/Private/ExrImageWrapper.h
@@ -49,6 +49,8 @@ private:
 	template <Imf::PixelType OutputFormat, typename sourcetype>
 	void CompressRaw(const sourcetype* SrcData, bool bIgnoreAlpha);
 	const char* GetRawChannelName(int ChannelIndex) const;
+	
+	bool bUseCompression;
 };
 
 #endif // WITH_UNREALEXR
\ No newline at end of file
diff --git a/Engine/Source/Runtime/ImageWrapper/Public/Interfaces/IImageWrapper.h b/Engine/Source/Runtime/ImageWrapper/Public/Interfaces/IImageWrapper.h
--- a/Engine/Source/Runtime/ImageWrapper/Public/Interfaces/IImageWrapper.h
+++ b/Engine/Source/Runtime/ImageWrapper/Public/Interfaces/IImageWrapper.h
@@ -51,6 +51,14 @@ namespace ERGBFormat
 	};
 };
 
+namespace ImageCompression
+{
+	enum CompressionQuality
+	{
+		Default			= 0,
+		Uncompressed	= 1,
+	};
+};
 
 /** Type definition for shared pointers to instances of IImageWrapper. */
 typedef TSharedPtr<class IImageWrapper> IImageWrapperPtr;
diff --git a/Engine/Source/Runtime/Launch/Private/LaunchEngineLoop.cpp b/Engine/Source/Runtime/Launch/Private/LaunchEngineLoop.cpp
--- a/Engine/Source/Runtime/Launch/Private/LaunchEngineLoop.cpp
+++ b/Engine/Source/Runtime/Launch/Private/LaunchEngineLoop.cpp
@@ -2284,6 +2284,7 @@ int32 FEngineLoop::Init()
 	}
 
 	FModuleManager::Get().LoadModule(TEXT("SequenceRecorder"));
+	FModuleManager::Get().LoadModule(TEXT("SequenceRecorderSections"));
 #endif
 
 	GIsRunning = true;
diff --git a/Engine/Source/Runtime/LevelSequence/Private/LevelSequence.cpp b/Engine/Source/Runtime/LevelSequence/Private/LevelSequence.cpp
--- a/Engine/Source/Runtime/LevelSequence/Private/LevelSequence.cpp
+++ b/Engine/Source/Runtime/LevelSequence/Private/LevelSequence.cpp
@@ -5,8 +5,8 @@
 #include "LevelSequenceObject.h"
 #include "MovieScene.h"
 #include "MovieSceneCommonHelpers.h"
-#include "Engine/Blueprint.h"
 
+DEFINE_LOG_CATEGORY_STATIC(LogLevelSequence, Log, All);
 
 ULevelSequence::ULevelSequence(const FObjectInitializer& ObjectInitializer)
 	: Super(ObjectInitializer)
@@ -20,23 +20,42 @@ void ULevelSequence::Initialize()
 	MovieScene = NewObject<UMovieScene>(this, NAME_None, RF_Transactional);
 }
 
-bool ULevelSequence::Rename(const TCHAR* NewName, UObject* NewOuter, ERenameFlags Flags)
+void ULevelSequence::PostLoad()
 {
+	Super::PostLoad();
+
 #if WITH_EDITOR
-	if (Super::Rename(NewName, NewOuter, Flags))
+	TSet<FGuid> InvalidSpawnables;
+
+	for (int32 Index = 0; Index < MovieScene->GetSpawnableCount(); ++Index)
 	{
-		ForEachObjectWithOuter(MovieScene, [&](UObject* Object){
-			if (auto* Blueprint = Cast<UBlueprint>(Object))
+		FMovieSceneSpawnable& Spawnable = MovieScene->GetSpawnable(Index);
+		if (!Spawnable.GetObjectTemplate())
+		{
+			if (Spawnable.GeneratedClass_DEPRECATED && Spawnable.GeneratedClass_DEPRECATED->ClassGeneratedBy)
 			{
-				Blueprint->RenameGeneratedClasses(nullptr, MovieScene, Flags);
+				const FName TemplateName = MakeUniqueObjectName(MovieScene, UObject::StaticClass(), Spawnable.GeneratedClass_DEPRECATED->ClassGeneratedBy->GetFName());
+
+				UObject* NewTemplate = NewObject<UObject>(MovieScene, Spawnable.GeneratedClass_DEPRECATED, TemplateName);
+				if (NewTemplate)
+				{
+					Spawnable.CopyObjectTemplate(*NewTemplate, *MovieScene);
+				}
 			}
-		}, false);
+		}
 
-		return true;
+		if (!Spawnable.GetObjectTemplate())
+		{
+			InvalidSpawnables.Add(Spawnable.GetGuid());
+			UE_LOG(LogLevelSequence, Warning, TEXT("Discarding spawnable with ID '%s' since its generated class could not produce to a template actor"), *Spawnable.GetGuid().ToString());
+		}
+	}
+
+	for (FGuid& ID : InvalidSpawnables)
+	{
+		MovieScene->RemoveSpawnable(ID);
 	}
-#endif	//#if WITH_EDITOR
-	
-	return false;
+#endif
 }
 
 void ULevelSequence::ConvertPersistentBindingsToDefault(UObject* FixupContext)
@@ -85,7 +104,7 @@ UObject* ULevelSequence::FindPossessableObject(const FGuid& ObjectId, UObject* C
 
 FGuid ULevelSequence::FindPossessableObjectId(UObject& Object) const
 {
-	return FGuid();
+	return ObjectReferences.FindBindingId(&Object, Object.GetWorld());
 }
 
 UMovieScene* ULevelSequence::GetMovieScene() const
diff --git a/Engine/Source/Runtime/LevelSequence/Private/LevelSequencePlayer.cpp b/Engine/Source/Runtime/LevelSequence/Private/LevelSequencePlayer.cpp
--- a/Engine/Source/Runtime/LevelSequence/Private/LevelSequencePlayer.cpp
+++ b/Engine/Source/Runtime/LevelSequence/Private/LevelSequencePlayer.cpp
@@ -6,6 +6,7 @@
 #include "MovieSceneSequence.h"
 #include "MovieSceneSequenceInstance.h"
 #include "LevelSequenceSpawnRegister.h"
+#include "Engine/LevelStreaming.h"
 
 
 struct FTickAnimationPlayers : public FTickableGameObject
@@ -68,6 +69,7 @@ ULevelSequencePlayer::ULevelSequencePlayer(const FObjectInitializer& ObjectIniti
 	, LevelSequence(nullptr)
 	, bIsPlaying(false)
 	, TimeCursorPosition(0.0f)
+	, LastCursorPosition(0.0f)
 	, StartTime(0.f)
 	, EndTime(0.f)
 	, CurrentNumLoops(0)
@@ -156,12 +158,8 @@ float ULevelSequencePlayer::GetPlaybackPosition() const
 
 void ULevelSequencePlayer::SetPlaybackPosition(float NewPlaybackPosition)
 {
-	float LastTimePosition = TimeCursorPosition;
-	
-	TimeCursorPosition = NewPlaybackPosition;
-	OnCursorPositionChanged();
-
-	UpdateMovieSceneInstance(TimeCursorPosition, LastTimePosition);
+	UpdateTimeCursorPosition(NewPlaybackPosition);
+	UpdateMovieSceneInstance(TimeCursorPosition, LastCursorPosition);
 }
 
 float ULevelSequencePlayer::GetLength() const
@@ -187,29 +185,32 @@ void ULevelSequencePlayer::SetPlaybackRange( const float NewStartTime, const flo
 	TimeCursorPosition = FMath::Clamp(TimeCursorPosition, 0.f, GetLength());
 }
 
-void ULevelSequencePlayer::OnCursorPositionChanged()
+void ULevelSequencePlayer::UpdateTimeCursorPosition(float NewPosition)
 {
 	float Length = GetLength();
 
-	// Handle looping or stopping
-	if (TimeCursorPosition >= Length || TimeCursorPosition < 0)
+	if ((NewPosition >= Length) || (NewPosition < 0))
 	{
+		// loop playback
 		if (PlaybackSettings.LoopCount < 0 || CurrentNumLoops < PlaybackSettings.LoopCount)
 		{
 			++CurrentNumLoops;
-			const float Overplay = FMath::Fmod(TimeCursorPosition, Length);
+			const float Overplay = FMath::Fmod(NewPosition, Length);
 			TimeCursorPosition = Overplay < 0 ? Length + Overplay : Overplay;
+			LastCursorPosition = TimeCursorPosition;
 
 			SpawnRegister->ForgetExternallyOwnedSpawnedObjects(*this);
+
+			return;
 		}
-		else
-		{
-			// Stop playing without modifying the playback position
-			// @todo: trigger an event?
-			bIsPlaying = false;
-			CurrentNumLoops = 0;
-		}
+
+		// stop playback
+		bIsPlaying = false;
+		CurrentNumLoops = 0;
 	}
+
+	LastCursorPosition = TimeCursorPosition;
+	TimeCursorPosition = NewPosition;
 }
 
 /* ULevelSequencePlayer implementation
@@ -320,28 +321,37 @@ UObject* ULevelSequencePlayer::GetPlaybackContext() const
 	return World.Get();
 }
 
-UObject* ULevelSequencePlayer::GetEventContext() const
+TArray<UObject*> ULevelSequencePlayer::GetEventContexts() const
 {
+	TArray<UObject*> EventContexts;
 	if (World.IsValid())
 	{
-		return World->GetLevelScriptActor();
+		if (World->GetLevelScriptActor())
+		{
+			EventContexts.Add(World->GetLevelScriptActor());
+		}
+
+		for (ULevelStreaming* StreamingLevel : World->StreamingLevels)
+		{
+			if (StreamingLevel->GetLevelScriptActor())
+			{
+				EventContexts.Add(StreamingLevel->GetLevelScriptActor());
+			}
+		}
 	}
-	return nullptr;
+	return EventContexts;
 }
 
 void ULevelSequencePlayer::Update(const float DeltaSeconds)
 {
-	float LastTimePosition = TimeCursorPosition;
-
 	if (bIsPlaying)
 	{
-		TimeCursorPosition += DeltaSeconds * PlaybackSettings.PlayRate;
-		OnCursorPositionChanged();
-		UpdateMovieSceneInstance(TimeCursorPosition, LastTimePosition);
+		UpdateTimeCursorPosition(TimeCursorPosition + DeltaSeconds * PlaybackSettings.PlayRate);
+		UpdateMovieSceneInstance(TimeCursorPosition, LastCursorPosition);
 	}
 	else if (!bHasCleanedUpSequence && TimeCursorPosition >= GetLength())
 	{
-		UpdateMovieSceneInstance(TimeCursorPosition, LastTimePosition);
+		UpdateMovieSceneInstance(TimeCursorPosition, LastCursorPosition);
 
 		bHasCleanedUpSequence = true;
 		
diff --git a/Engine/Source/Runtime/LevelSequence/Private/LevelSequenceSpawnRegister.cpp b/Engine/Source/Runtime/LevelSequence/Private/LevelSequenceSpawnRegister.cpp
--- a/Engine/Source/Runtime/LevelSequence/Private/LevelSequenceSpawnRegister.cpp
+++ b/Engine/Source/Runtime/LevelSequence/Private/LevelSequenceSpawnRegister.cpp
@@ -30,8 +30,8 @@ UObject* FLevelSequenceSpawnRegister::SpawnObject(const FGuid& BindingId, FMovie
 		return nullptr;
 	}
 
-	UClass* SpawnableClass = Spawnable->GetClass();
-	if (!SpawnableClass || !SpawnableClass->IsChildOf(AActor::StaticClass()))
+	AActor* ObjectTemplate = Cast<AActor>(Spawnable->GetObjectTemplate());
+	if (!ObjectTemplate)
 	{
 		return nullptr;
 	}
@@ -55,38 +55,20 @@ UObject* FLevelSequenceSpawnRegister::SpawnObject(const FGuid& BindingId, FMovie
 	{
 		SpawnInfo.Name = ActorName;
 		SpawnInfo.ObjectFlags = ObjectFlags;
+		SpawnInfo.Template = ObjectTemplate;
 	}
 
 	FTransform SpawnTransform;
 
-	AActor* ActorCDO = CastChecked<AActor>(SpawnableClass->ClassDefaultObject);
-	if (USceneComponent* RootComponent = ActorCDO->GetRootComponent())
+	if (USceneComponent* RootComponent = ObjectTemplate->GetRootComponent())
 	{
 		SpawnTransform.SetTranslation(RootComponent->RelativeLocation);
 		SpawnTransform.SetRotation(RootComponent->RelativeRotation.Quaternion());
 	}
 
-	//@todo: This is in place of SpawnActorAbsolute which doesn't exist yet in the Orion branch
-	FTransform NewTransform = SpawnTransform;
 	{
-		AActor* Template = SpawnInfo.Template;
-
-		if(!Template)
-		{
-			// Use class's default actor as a template.
-			Template = SpawnableClass->GetDefaultObject<AActor>();
-		}
-
-		USceneComponent* TemplateRootComponent = (Template)? Template->GetRootComponent() : NULL;
-		if(TemplateRootComponent)
-		{
-			TemplateRootComponent->UpdateComponentToWorld();
-			NewTransform = TemplateRootComponent->GetComponentToWorld().Inverse() * NewTransform;
-			//TemplateRootComponent->bAutoActivate = false;
-		}
-
 		// Disable all particle components so that they don't auto fire as soon as the actor is spawned. The particles should be triggered through the particle track.
-		TArray<UActorComponent*> ParticleComponents = Template->GetComponentsByClass(UParticleSystem::StaticClass());
+		TArray<UActorComponent*> ParticleComponents = ObjectTemplate->GetComponentsByClass(UParticleSystemComponent::StaticClass());
 		for (int32 ComponentIdx = 0; ComponentIdx < ParticleComponents.Num(); ++ComponentIdx)
 		{
 			ParticleComponents[ComponentIdx]->bAutoActivate = false;
@@ -99,12 +81,52 @@ UObject* FLevelSequenceSpawnRegister::SpawnObject(const FGuid& BindingId, FMovie
 		WorldContext = GWorld;
 	}
 
-	AActor* SpawnedActor = WorldContext->SpawnActor(SpawnableClass, &NewTransform, SpawnInfo);
+	// @todo: Remove when instanced components work correctly on spawned templates
+	TArray<UActorComponent*> TemplateInstanceComponents = ObjectTemplate->GetInstanceComponents();
+	ObjectTemplate->ClearInstanceComponents(false);
+
+	AActor* SpawnedActor = WorldContext->SpawnActorAbsolute(ObjectTemplate->GetClass(), SpawnTransform, SpawnInfo);
 	if (!SpawnedActor)
 	{
 		return nullptr;
 	}
 
+	// First, duplicate instance components
+	for (UActorComponent* TemplateComponent : TemplateInstanceComponents)
+	{
+		ObjectTemplate->AddInstanceComponent(TemplateComponent);
+
+		UActorComponent* NewComponent = DuplicateObject<UActorComponent>(TemplateComponent, SpawnedActor, TemplateComponent->GetFName());
+		SpawnedActor->AddInstanceComponent(NewComponent);
+	}
+
+	// Second, attach the components. We do this as a different pass, since instance components may be attached to each other
+	for (UActorComponent* ActorComp : SpawnedActor->GetInstanceComponents())
+	{
+		USceneComponent* SceneComp = Cast<USceneComponent>(ActorComp);
+		if (!SceneComp)
+		{
+			continue;
+		}
+		
+		USceneComponent* NewParent = SpawnedActor->GetRootComponent();
+
+		USceneComponent* OldParent = SceneComp->GetAttachParent();
+		if (OldParent)
+		{
+			FString PathName = OldParent->GetPathName(OldParent->GetOwner());
+
+			NewParent = FindObject<USceneComponent>(SpawnedActor, *PathName);
+		}
+
+		if (ensure(NewParent))
+		{
+			SceneComp->AttachToComponent(NewParent, FAttachmentTransformRules::KeepRelativeTransform, SceneComp->GetAttachSocketName());
+		}
+	}
+
+	SpawnedActor->RegisterAllComponents();
+
 	// tag this actor so we know it was spawned by sequencer
 	SpawnedActor->Tags.Add(SequencerActorTag);
 
diff --git a/Engine/Source/Runtime/LevelSequence/Public/LevelSequence.h b/Engine/Source/Runtime/LevelSequence/Public/LevelSequence.h
--- a/Engine/Source/Runtime/LevelSequence/Public/LevelSequence.h
+++ b/Engine/Source/Runtime/LevelSequence/Public/LevelSequence.h
@@ -43,7 +43,7 @@ public:
 	virtual void UnbindPossessableObjects(const FGuid& ObjectId) override;
 	virtual bool AllowsSpawnableObjects() const override;
 
-	virtual bool Rename(const TCHAR* NewName = nullptr, UObject* NewOuter = nullptr, ERenameFlags Flags = REN_None) override;
+	virtual void PostLoad() override;
 	
 	/** Bind a posessable object with an explicitly-supplied ObjectReference */
 	void BindPossessableObject(const FGuid& ObjectId, const FLevelSequenceObjectReference& ObjectReference);
diff --git a/Engine/Source/Runtime/LevelSequence/Public/LevelSequencePlayer.h b/Engine/Source/Runtime/LevelSequence/Public/LevelSequencePlayer.h
--- a/Engine/Source/Runtime/LevelSequence/Public/LevelSequencePlayer.h
+++ b/Engine/Source/Runtime/LevelSequence/Public/LevelSequencePlayer.h
@@ -144,7 +144,7 @@ protected:
 	virtual TSharedRef<FMovieSceneSequenceInstance> GetRootMovieSceneSequenceInstance() const override;
 	virtual IMovieSceneSpawnRegister& GetSpawnRegister() override;
 	virtual UObject* GetPlaybackContext() const override;
-	virtual UObject* GetEventContext() const override;
+	virtual TArray<UObject*> GetEventContexts() const override;
 
 public:
 
@@ -152,11 +152,11 @@ public:
 
 private:
 
-	/** Update the movie scene instance from the specified previous position, to the specified time position */
+	/** Update the movie scene instance from the specified previous position, to the specified time position. */
 	void UpdateMovieSceneInstance(float CurrentPosition, float PreviousPosition);
 
-	/** Called when the cursor position has changed to implement looping */
-	void OnCursorPositionChanged();
+	/** Update the time cursor position and handle stopping & looping. */
+	void UpdateTimeCursorPosition(float NewPosition);
 
 private:
 
@@ -172,6 +172,9 @@ private:
 	UPROPERTY()
 	float TimeCursorPosition;
 
+	/** The time cursor position in the previous update. */
+	float LastCursorPosition;
+
 	/** Time time at which to start playing the sequence (defaults to the lower bound of the sequence's play range) */
 	float StartTime;
 
diff --git a/Engine/Source/Runtime/MovieScene/Private/MovieScene.cpp b/Engine/Source/Runtime/MovieScene/Private/MovieScene.cpp
--- a/Engine/Source/Runtime/MovieScene/Private/MovieScene.cpp
+++ b/Engine/Source/Runtime/MovieScene/Private/MovieScene.cpp
@@ -9,7 +9,7 @@
 
 UMovieScene::UMovieScene(const FObjectInitializer& ObjectInitializer)
 	: Super(ObjectInitializer)
-	, InOutRange(FFloatRange::Empty())
+	, SelectionRange(FFloatRange::Empty())
 	, PlaybackRange(FFloatRange::Empty())
 	, InTime_DEPRECATED(FLT_MAX)
 	, OutTime_DEPRECATED(-FLT_MAX)
@@ -25,13 +25,11 @@ UMovieScene::UMovieScene(const FObjectInitializer& ObjectInitializer)
 #if WITH_EDITOR
 
 // @todo sequencer: Some of these methods should only be used by tools, and should probably move out of MovieScene!
-FGuid UMovieScene::AddSpawnable( const FString& Name, UBlueprint* Blueprint )
+FGuid UMovieScene::AddSpawnable( const FString& Name, UObject& ObjectTemplate )
 {
-	check( (Blueprint != nullptr) && (Blueprint->GeneratedClass) );
-
 	Modify();
 
-	FMovieSceneSpawnable NewSpawnable( Name, Blueprint->GeneratedClass );
+	FMovieSceneSpawnable NewSpawnable( Name, ObjectTemplate );
 	Spawnables.Add( NewSpawnable );
 
 	// Add a new binding so that tracks can be added to it
@@ -52,25 +50,10 @@ bool UMovieScene::RemoveSpawnable( const FGuid& Guid )
 			if( CurSpawnable.GetGuid() == Guid )
 			{
 				Modify();
-
-				{
-					UClass* GeneratedClass = CurSpawnable.GetClass();
-					UBlueprint* Blueprint = GeneratedClass ? Cast<UBlueprint>(GeneratedClass->ClassGeneratedBy) : nullptr;
-
-					if (Blueprint)
-					{
-						// @todo sequencer: Also remove created Blueprint inner object.  Is this sufficient?  Needs to work with Undo too!
-						Blueprint->ClearFlags( RF_Standalone );	// @todo sequencer: Probably not needed for Blueprint
-						Blueprint->MarkPendingKill();
-					}
-				}
-
 				RemoveBinding( Guid );
 
-				// Found it!
 				Spawnables.RemoveAt( SpawnableIter.GetIndex() );
-
-
+				
 				bAnythingRemoved = true;
 				break;
 			}
diff --git a/Engine/Source/Runtime/MovieScene/Private/MovieSceneFolder.cpp b/Engine/Source/Runtime/MovieScene/Private/MovieSceneFolder.cpp
--- a/Engine/Source/Runtime/MovieScene/Private/MovieSceneFolder.cpp
+++ b/Engine/Source/Runtime/MovieScene/Private/MovieSceneFolder.cpp
@@ -74,25 +74,26 @@ void UMovieSceneFolder::RemoveChildObjectBinding( const FGuid& InObjectBinding )
 }
 
 
-void UMovieSceneFolder::PreSave()
+void UMovieSceneFolder::Serialize( FArchive& Archive )
 {
-	ChildObjectBindingStrings.Empty();
-	for ( const FGuid& ChildObjectBinding : ChildObjectBindings )
+	if ( Archive.IsLoading() )
 	{
-		ChildObjectBindingStrings.Add( ChildObjectBinding.ToString() );
+		Super::Serialize( Archive );
+		ChildObjectBindings.Empty();
+		for ( const FString& ChildObjectBindingString : ChildObjectBindingStrings )
+		{
+			FGuid ChildObjectBinding;
+			FGuid::Parse( ChildObjectBindingString, ChildObjectBinding );
+			ChildObjectBindings.Add( ChildObjectBinding );
+		}
 	}
-	Super::PreSave();
-}
-
-
-void UMovieSceneFolder::PostLoad()
-{
-	Super::PostLoad();
-	ChildObjectBindings.Empty();
-	for ( const FString& ChildObjectBindingString : ChildObjectBindingStrings )
+	else
 	{
-		FGuid ChildObjectBinding;
-		FGuid::Parse( ChildObjectBindingString, ChildObjectBinding );
-		ChildObjectBindings.Add( ChildObjectBinding );
+		ChildObjectBindingStrings.Empty();
+		for ( const FGuid& ChildObjectBinding : ChildObjectBindings )
+		{
+			ChildObjectBindingStrings.Add( ChildObjectBinding.ToString() );
+		}
+		Super::Serialize( Archive );
 	}
 }
diff --git a/Engine/Source/Runtime/MovieScene/Private/MovieSceneSequenceInstance.cpp b/Engine/Source/Runtime/MovieScene/Private/MovieSceneSequenceInstance.cpp
--- a/Engine/Source/Runtime/MovieScene/Private/MovieSceneSequenceInstance.cpp
+++ b/Engine/Source/Runtime/MovieScene/Private/MovieSceneSequenceInstance.cpp
@@ -7,11 +7,17 @@
 
 FMovieSceneSequenceInstance::FMovieSceneSequenceInstance(const UMovieSceneSequence& InMovieSceneSequence)
 	: MovieSceneSequence( &InMovieSceneSequence )
+#if WITH_EDITOR
+	, InstanceId(FGuid::NewGuid())
+#endif
 {
 	TimeRange = MovieSceneSequence->GetMovieScene()->GetPlaybackRange();
 }
 
 FMovieSceneSequenceInstance::FMovieSceneSequenceInstance(const UMovieSceneTrack& InMovieSceneTrack)
+#if WITH_EDITOR
+	: InstanceId(FGuid::NewGuid())
+#endif
 {
 	TimeRange = CastChecked<UMovieScene>(InMovieSceneTrack.GetOuter())->GetPlaybackRange();
 }
@@ -155,6 +161,20 @@ void FMovieSceneSequenceInstance::RestoreState(class IMovieScenePlayer& Player)
 	}
 }
 
+void FMovieSceneSequenceInstance::RestoreSpecificState(const FGuid& ObjectGuid, IMovieScenePlayer& Player)
+{
+	FMovieSceneObjectBindingInstance* ObjectBindingInstance = ObjectBindingInstances.Find(ObjectGuid);
+	if (!ObjectBindingInstance)
+	{
+		return;
+	}
+
+	for (auto& Pair : ObjectBindingInstance->TrackInstances)
+	{
+		Pair.Value->RestoreState(ObjectBindingInstance->RuntimeObjects, Player, *this);
+	}
+}
+
 void FMovieSceneSequenceInstance::Update( EMovieSceneUpdateData& UpdateData, class IMovieScenePlayer& Player )
 {
 	if(MovieSceneSequence.IsValid())
@@ -219,7 +239,8 @@ void FMovieSceneSequenceInstance::UpdatePassSingle( EMovieSceneUpdateData& Updat
 		// update each master track
 		for( FMovieSceneInstanceMap::TIterator It( MasterTrackInstances ); It; ++It )
 		{		
-			if (It.Value()->HasUpdatePasses() & UpdateData.UpdatePass && It.Value() != CameraCutTrackInstance)
+			if (It.Value()->HasUpdatePasses() & UpdateData.UpdatePass && It.Value() != CameraCutTrackInstance &&
+				( UpdateData.bSubSceneDeactivate == false || It.Value()->RequiresUpdateForSubSceneDeactivate () ))
 			{
 				It.Value()->Update( UpdateData, NoObjects, Player, *this);
 			}
@@ -233,7 +254,8 @@ void FMovieSceneSequenceInstance::UpdatePassSingle( EMovieSceneUpdateData& Updat
 		
 			for( FMovieSceneInstanceMap::TIterator It = ObjectBindingInstance.TrackInstances.CreateIterator(); It; ++It )
 			{
-				if (It.Value()->HasUpdatePasses() & UpdateData.UpdatePass && It.Value() != CameraCutTrackInstance)
+				if (It.Value()->HasUpdatePasses() & UpdateData.UpdatePass && It.Value() != CameraCutTrackInstance &&
+					( UpdateData.bSubSceneDeactivate == false || It.Value()->RequiresUpdateForSubSceneDeactivate() ) )
 				{
 					It.Value()->Update( UpdateData, ObjectBindingInstance.RuntimeObjects, Player, *this );
 				}
@@ -243,7 +265,8 @@ void FMovieSceneSequenceInstance::UpdatePassSingle( EMovieSceneUpdateData& Updat
 		// update camera cut track last to make sure spawnable cameras are there, and to override sub-shots
 		if (CameraCutTrackInstance.IsValid())
 		{
-			if (CameraCutTrackInstance->HasUpdatePasses() & UpdateData.UpdatePass)
+			if (CameraCutTrackInstance->HasUpdatePasses() & UpdateData.UpdatePass &&
+				( UpdateData.bSubSceneDeactivate == false || CameraCutTrackInstance->RequiresUpdateForSubSceneDeactivate() ) )
 			{
 				CameraCutTrackInstance->Update(UpdateData, NoObjects, Player, *this );
 			}
diff --git a/Engine/Source/Runtime/MovieScene/Private/MovieSceneSpawnable.cpp b/Engine/Source/Runtime/MovieScene/Private/MovieSceneSpawnable.cpp
new file mode 100644
--- /dev/null
+++ b/Engine/Source/Runtime/MovieScene/Private/MovieSceneSpawnable.cpp
@@ -0,0 +1,23 @@
+// Copyright 1998-2016 Epic Games, Inc. All Rights Reserved.
+
+#include "MovieScenePrivatePCH.h"
+#include "MovieSceneSpawnable.h"
+#include "MovieScene.h"
+
+
+void FMovieSceneSpawnable::CopyObjectTemplate(UObject& InSourceObject, UMovieScene& OwnerMovieScene)
+{
+	const FName ObjectName = ObjectTemplate ? ObjectTemplate->GetFName() : InSourceObject.GetFName();
+
+	if (ObjectTemplate)
+	{
+		ObjectTemplate->MarkPendingKill();
+		ObjectTemplate = nullptr;
+	}
+
+	ObjectTemplate = StaticDuplicateObject(&InSourceObject, &OwnerMovieScene, ObjectName, RF_AllFlags & ~(RF_Transient));
+	ObjectTemplate->SetFlags(RF_ArchetypeObject);
+
+	// @todo: this will mark the package as dirty whenever a spawnable is destroyed. We should diff the duplicated object, and only mark dirty where the spawnable has actually changed
+	OwnerMovieScene.MarkPackageDirty();
+}
\ No newline at end of file
diff --git a/Engine/Source/Runtime/MovieScene/Public/IMovieScenePlayer.h b/Engine/Source/Runtime/MovieScene/Public/IMovieScenePlayer.h
--- a/Engine/Source/Runtime/MovieScene/Public/IMovieScenePlayer.h
+++ b/Engine/Source/Runtime/MovieScene/Public/IMovieScenePlayer.h
@@ -129,9 +129,9 @@ public:
 	virtual UObject* GetPlaybackContext() const { return nullptr; }
 
 	/**
-	 * Access the event context for this movie scene player
+	 * Access the event contexts for this movie scene player
 	 */
-	virtual UObject* GetEventContext() const { return nullptr; }
+	virtual TArray<UObject*> GetEventContexts() const { return TArray<UObject*>(); }
 
 	/**
 	 * Test whether this is a preview player or not. As such, playback range becomes insignificant for things like spawnables
diff --git a/Engine/Source/Runtime/MovieScene/Public/IMovieSceneSpawnRegister.h b/Engine/Source/Runtime/MovieScene/Public/IMovieSceneSpawnRegister.h
--- a/Engine/Source/Runtime/MovieScene/Public/IMovieSceneSpawnRegister.h
+++ b/Engine/Source/Runtime/MovieScene/Public/IMovieSceneSpawnRegister.h
@@ -16,11 +16,11 @@ struct FMovieSceneSpawnable;
 /** Struct used for defining a new spawnable type */
 struct FNewSpawnable
 {
-	FNewSpawnable() : Blueprint(nullptr) {}
-	FNewSpawnable(UBlueprint* InBlueprint, FString InName) : Blueprint(InBlueprint), Name(MoveTemp(InName)) {}
+	FNewSpawnable() : ObjectTemplate(nullptr) {}
+	FNewSpawnable(UObject* InObjectTemplate, FString InName) : ObjectTemplate(InObjectTemplate), Name(MoveTemp(InName)) {}
 
-	/** The blueprint that defines the spawnable */
-	UBlueprint* Blueprint;
+	/** The archetype object template that defines the spawnable */
+	UObject* ObjectTemplate;
 
 	/** The desired name of the new spawnable */
 	FString Name;
diff --git a/Engine/Source/Runtime/MovieScene/Public/IMovieSceneTrackInstance.h b/Engine/Source/Runtime/MovieScene/Public/IMovieSceneTrackInstance.h
--- a/Engine/Source/Runtime/MovieScene/Public/IMovieSceneTrackInstance.h
+++ b/Engine/Source/Runtime/MovieScene/Public/IMovieSceneTrackInstance.h
@@ -20,6 +20,10 @@ struct EMovieSceneUpdateData
 	float LastPosition;
 	bool bPreroll;
 	bool bJumpCut;
+	/** Indicates that this update was caused by the owning movie scene stopping playback due
+	    to the active sub-scene being deactivated. */
+	bool bSubSceneDeactivate;
+
 	EMovieSceneUpdatePass UpdatePass;
 
 	EMovieSceneUpdateData()
@@ -28,6 +32,7 @@ struct EMovieSceneUpdateData
 		LastPosition = 0.0f;
 		bPreroll = false;
 		bJumpCut = false;
+		bSubSceneDeactivate = false;
 		UpdatePass = MSUP_PreUpdate;
 	}
 	EMovieSceneUpdateData(float InPosition, float InLastPosition)
@@ -36,6 +41,7 @@ struct EMovieSceneUpdateData
 		LastPosition = InLastPosition;
 		bPreroll = false;
 		bJumpCut = false;
+		bSubSceneDeactivate = false;
 		UpdatePass = MSUP_PreUpdate;
 	}
 };
@@ -76,6 +82,11 @@ public:
 	 */
 	virtual EMovieSceneUpdatePass HasUpdatePasses() { return MSUP_Update; }
 
+	/**
+	 * Whether or not this track instance needs to be updated when it's deactivated because it's in a sub-scene that's ending.
+	 */
+	virtual bool RequiresUpdateForSubSceneDeactivate() { return false; }
+
 	/**
 	 * Refreshes the current instance
 	 */
diff --git a/Engine/Source/Runtime/MovieScene/Public/MovieScene.h b/Engine/Source/Runtime/MovieScene/Public/MovieScene.h
--- a/Engine/Source/Runtime/MovieScene/Public/MovieScene.h
+++ b/Engine/Source/Runtime/MovieScene/Public/MovieScene.h
@@ -93,10 +93,10 @@ public:
 	 * These objects are stored as "inners" of the MovieScene.
 	 *
 	 * @param Name Name of the spawnable.
-	 * @param Blueprint	The blueprint to add.
+	 * @param ObjectTemplate The object template to use for the spawnable
 	 * @return Guid of the newly-added spawnable.
 	 */
-	FGuid AddSpawnable(const FString& Name, UBlueprint* Blueprint);
+	FGuid AddSpawnable(const FString& Name, UObject& ObjectTemplate);
 
 	/**
 	 * Removes a spawnable from this movie scene.
@@ -382,10 +382,10 @@ public:
 		return ObjectBindings;
 	}
 
-	/** Get the current in/out range. */
-	TRange<float> GetInOutRange() const
+	/** Get the current selection range. */
+	TRange<float> GetSelectionRange() const
 	{
-		return InOutRange;
+		return SelectionRange;
 	}
 
 	/**
@@ -415,10 +415,10 @@ public:
 		return ObjectsToLabels;
 	}
 
-	/** Set the in/out selection range. */
-	void SetInOutRange(TRange<float> Range)
+	/** Set the selection range. */
+	void SetSelectionRange(TRange<float> Range)
 	{
-		InOutRange = Range;
+		SelectionRange = Range;
 	}
 
 	/**
@@ -503,9 +503,9 @@ private:
 	UPROPERTY()
 	UMovieSceneTrack* CameraCutTrack;
 
-	/** User-defined in/out selection range. */
+	/** User-defined selection range. */
 	UPROPERTY()
-	FFloatRange InOutRange;
+	FFloatRange SelectionRange;
 
 	/** User-defined playback range for this movie scene. Must be a finite range. Relative to this movie-scene's 0-time origin. */
 	UPROPERTY()
diff --git a/Engine/Source/Runtime/MovieScene/Public/MovieSceneCommonHelpers.h b/Engine/Source/Runtime/MovieScene/Public/MovieSceneCommonHelpers.h
--- a/Engine/Source/Runtime/MovieScene/Public/MovieSceneCommonHelpers.h
+++ b/Engine/Source/Runtime/MovieScene/Public/MovieSceneCommonHelpers.h
@@ -102,12 +102,13 @@ public:
 	/**
 	 * Calls the setter function for a specific runtime object or if the setter function does not exist, the property is set directly
 	 *
-	 * @param InRuntimeObject	The runtime object whose function to call
-	 * @param FunctionParams	Parameters to pass to the function
+	 * @param InRuntimeObject The runtime object whose function to call
+	 * @param PropertyValue The new value to assign to the property (if bound to property)
+	 * @param FunctionParams Parameters to pass to the function (if bound to function)
 	 */
 	template <typename ValueType>
 	void 
-	CallFunction( UObject* InRuntimeObject, void* FunctionParams )
+	CallFunction( UObject* InRuntimeObject, void* PropertyValue, void* FunctionParams )
 	{
 		FPropertyAndFunction PropAndFunction = RuntimeObjectToFunctionMap.FindRef(InRuntimeObject);
 		if(PropAndFunction.Function)
@@ -116,11 +117,19 @@ public:
 		}
 		else
 		{
-			ValueType* NewValue = (ValueType*)(FunctionParams);
+			ValueType* NewValue = (ValueType*)(PropertyValue);
 			SetCurrentValue<ValueType>(InRuntimeObject, *NewValue);
 		}
 	}
 
+	/** For backwards compatibility. */
+	template <typename ValueType>
+	void 
+	CallFunction( UObject* InRuntimeObject, void* PropertyValue )
+	{
+		CallFunction<ValueType>(InRuntimeObject, PropertyValue, PropertyValue);
+	}
+
 	/**
 	 * Rebuilds the property and function mappings for a set of runtime objects
 	 *
diff --git a/Engine/Source/Runtime/MovieScene/Public/MovieSceneFolder.h b/Engine/Source/Runtime/MovieScene/Public/MovieSceneFolder.h
--- a/Engine/Source/Runtime/MovieScene/Public/MovieSceneFolder.h
+++ b/Engine/Source/Runtime/MovieScene/Public/MovieSceneFolder.h
@@ -45,8 +45,7 @@ class MOVIESCENE_API UMovieSceneFolder : public UObject
 	void RemoveChildObjectBinding( const FGuid& InObjectBinding );
 
 
-	virtual void PreSave() override;
-	virtual void PostLoad() override;
+	virtual void Serialize( FArchive& Archive );
 
 
 private:
diff --git a/Engine/Source/Runtime/MovieScene/Public/MovieSceneSection.h b/Engine/Source/Runtime/MovieScene/Public/MovieSceneSection.h
--- a/Engine/Source/Runtime/MovieScene/Public/MovieSceneSection.h
+++ b/Engine/Source/Runtime/MovieScene/Public/MovieSceneSection.h
@@ -163,9 +163,10 @@ public:
 	/**
 	 * Get the key handles for the keys on the curves within this section
 	 *
-	 * @param KeyHandles The key handles of the keys on the curves within this section
+	 * @param OutKeyHandles Will contain the key handles of the keys on the curves within this section
+	 * @param TimeRange Optional time range that the keys must be in (default = all)
 	 */
-	virtual void GetKeyHandles(TSet<FKeyHandle>& KeyHandles) const {};
+	virtual void GetKeyHandles(TSet<FKeyHandle>& OutKeyHandles, TRange<float> TimeRange) const { };
 
 	/**
 	 * Get the data structure representing the specified key.
diff --git a/Engine/Source/Runtime/MovieScene/Public/MovieSceneSequenceInstance.h b/Engine/Source/Runtime/MovieScene/Public/MovieSceneSequenceInstance.h
--- a/Engine/Source/Runtime/MovieScene/Public/MovieSceneSequenceInstance.h
+++ b/Engine/Source/Runtime/MovieScene/Public/MovieSceneSequenceInstance.h
@@ -76,6 +76,9 @@ public:
 	/** Restore state of the objects that this movie scene controls. */
 	MOVIESCENE_API void RestoreState(IMovieScenePlayer& Player);
 
+	/** Restore state of a specific object track that this movie scene controls. */
+	MOVIESCENE_API void RestoreSpecificState(const FGuid& ObjectGuid, IMovieScenePlayer& Player);
+
 	/**
 	 * Updates this movie scene.
 	 *
@@ -104,6 +107,8 @@ public:
 	 */
 	MOVIESCENE_API void UpdatePassSingle(EMovieSceneUpdateData& UpdateData, IMovieScenePlayer& Player);
 
+	MOVIESCENE_API void UpdateFromSubSceneDeactivate();
+
 	/**
 	 * Refreshes the existing instance.
 	 *
@@ -155,6 +160,13 @@ public:
 	/** Handle the sequence this instance refers to changing */
 	MOVIESCENE_API void HandleSequenceSectionChanged(UMovieSceneSequence* Sequence);
 
+#if WITH_EDITOR
+	/**
+	 * @return a transient but, unique identifier for this instance
+	 */
+	const FGuid& GetInstanceId() const { return InstanceId; }
+#endif
+
 protected:
 
 	void RefreshInstanceMap(const TArray<UMovieSceneTrack*>& Tracks, const TArray<TWeakObjectPtr<UObject>>& RuntimeObjects, FMovieSceneInstanceMap& TrackInstances, IMovieScenePlayer& Player);
@@ -195,4 +207,8 @@ private:
 
 	/** Cached time range for the movie scene */
 	TRange<float> TimeRange;
+
+#if WITH_EDITOR
+	FGuid InstanceId;
+#endif
 };
diff --git a/Engine/Source/Runtime/MovieScene/Public/MovieSceneSpawnable.h b/Engine/Source/Runtime/MovieScene/Public/MovieSceneSpawnable.h
--- a/Engine/Source/Runtime/MovieScene/Public/MovieSceneSpawnable.h
+++ b/Engine/Source/Runtime/MovieScene/Public/MovieSceneSpawnable.h
@@ -4,6 +4,8 @@
 
 #include "MovieSceneSpawnable.generated.h"
 
+class UMovieScene;
+
 UENUM()
 enum class ESpawnOwnership : uint8
 {
@@ -23,37 +25,53 @@ enum class ESpawnOwnership : uint8
 USTRUCT()
 struct FMovieSceneSpawnable 
 {
-	GENERATED_USTRUCT_BODY()
-
-public:
+	GENERATED_BODY()
 
-	/** FMovieSceneSpawnable default constructor */
-	FMovieSceneSpawnable() { }
+	FMovieSceneSpawnable()
+		: ObjectTemplate(nullptr)
+		, Ownership(ESpawnOwnership::InnerSequence)
+#if WITH_EDITORONLY_DATA
+		, bIgnoreOwnershipInEditor(false)
+		, GeneratedClass_DEPRECATED(nullptr)
+#endif
+	{
+	}
 
 	/** FMovieSceneSpawnable initialization constructor */
-	FMovieSceneSpawnable(const FString& InitName, UClass* InitClass)
+	FMovieSceneSpawnable(const FString& InitName, UObject& InObjectTemplate)
 		: Guid(FGuid::NewGuid())
 		, Name(InitName)
-		, GeneratedClass(InitClass)
+		, ObjectTemplate(&InObjectTemplate)
 		, Ownership(ESpawnOwnership::InnerSequence)
 #if WITH_EDITORONLY_DATA
 		, bIgnoreOwnershipInEditor(false)
+		, GeneratedClass_DEPRECATED(nullptr)
 #endif
-	{ }
+	{
+		ObjectTemplate->SetFlags(RF_ArchetypeObject);
+	}
 
 public:
 
 	/**
-	 * Get the Blueprint associated with the spawnable object.
+	 * Get the template object for this spawnable
 	 *
-	 * @return Blueprint class
+	 * @return Object template
 	 * @see GetGuid, GetName
 	 */
-	UClass* GetClass()
+	UObject* GetObjectTemplate()
 	{
-		return GeneratedClass;
+		return ObjectTemplate;
 	}
 
+	/**
+	 * Copy the specified object into this spawnable's template
+	 *
+	 * @param InSourceObject The source object to use. This object will be duplicated into the spawnable.
+	 * @param OwnerMovieScene The outer movie scene to which this spawnable belongs
+	 */
+	MOVIESCENE_API void CopyObjectTemplate(UObject& InSourceObject, UMovieScene& OwnerMovieScene);
+
 	/**
 	 * Get the unique identifier of the spawnable object.
 	 *
@@ -161,11 +179,9 @@ private:
 	// @todo sequencer: Should be editor-only probably
 	UPROPERTY()
 	FString Name;
-
-	/** Data-only blueprint-generated-class for this object */
-	// @todo sequencer: Could be weak object ptr here, IF blueprints that are inners are housekept properly without references
+	
 	UPROPERTY()
-	UClass* GeneratedClass;
+	UObject* ObjectTemplate;
 
 	/** Set of GUIDs to possessable object bindings that are bound to an object inside this spawnable */
 	// @todo sequencer: This should be a TSet, but they don't duplicate correctly atm
@@ -180,5 +196,10 @@ private:
 	/** When true, this spawnable will not respect its ownership sematics outside of the playback range, when spawned from inside a currently editing sequence */
 	UPROPERTY()
 	bool bIgnoreOwnershipInEditor;
+
+public:
+	/** Deprecated generated class */
+	UPROPERTY()
+	UClass* GeneratedClass_DEPRECATED;
 #endif
 };
diff --git a/Engine/Source/Runtime/MovieSceneTracks/MovieSceneTracks.Build.cs b/Engine/Source/Runtime/MovieSceneTracks/MovieSceneTracks.Build.cs
--- a/Engine/Source/Runtime/MovieSceneTracks/MovieSceneTracks.Build.cs
+++ b/Engine/Source/Runtime/MovieSceneTracks/MovieSceneTracks.Build.cs
@@ -25,8 +25,13 @@ public class MovieSceneTracks : ModuleRules
 
 		PrivateDependencyModuleNames.AddRange(
 			new string[] {
-				"SlateCore"
+				"SlateCore",
 			}
 		);
+
+		if (UEBuildConfiguration.bBuildEditor == true)
+		{
+			PrivateDependencyModuleNames.Add("UnrealEd");
+		}
 	}
 }
diff --git a/Engine/Source/Runtime/MovieSceneTracks/Private/Sections/MovieScene3DAttachSection.cpp b/Engine/Source/Runtime/MovieSceneTracks/Private/Sections/MovieScene3DAttachSection.cpp
--- a/Engine/Source/Runtime/MovieSceneTracks/Private/Sections/MovieScene3DAttachSection.cpp
+++ b/Engine/Source/Runtime/MovieSceneTracks/Private/Sections/MovieScene3DAttachSection.cpp
@@ -63,10 +63,18 @@ void UMovieScene3DAttachSection::Eval( USceneComponent* SceneComponent, float Po
 		}
 	}
 
-	if(!bFoundAttachment)
+	if (!bFoundAttachment)
 	{
-		OutTranslation = ParentActor->GetActorLocation();
-		OutRotation = ParentActor->GetActorRotation();
+		if (ParentActor->GetDefaultAttachComponent())
+		{
+			OutTranslation = ParentActor->GetDefaultAttachComponent()->GetComponentLocation();
+			OutRotation = ParentActor->GetDefaultAttachComponent()->GetComponentRotation();
+		}
+		else
+		{
+			OutTranslation = ParentActor->GetActorLocation();
+			OutRotation = ParentActor->GetActorRotation();
+		}
 	}
 
 	FTransform ParentTransform(OutRotation, OutTranslation);
diff --git a/Engine/Source/Runtime/MovieSceneTracks/Private/Sections/MovieScene3DPathSection.cpp b/Engine/Source/Runtime/MovieSceneTracks/Private/Sections/MovieScene3DPathSection.cpp
--- a/Engine/Source/Runtime/MovieSceneTracks/Private/Sections/MovieScene3DPathSection.cpp
+++ b/Engine/Source/Runtime/MovieSceneTracks/Private/Sections/MovieScene3DPathSection.cpp
@@ -129,14 +129,19 @@ void UMovieScene3DPathSection::DilateSection( float DilationFactor, float Origin
 }
 
 
-void UMovieScene3DPathSection::GetKeyHandles(TSet<FKeyHandle>& KeyHandles) const
+void UMovieScene3DPathSection::GetKeyHandles(TSet<FKeyHandle>& OutKeyHandles, TRange<float> TimeRange) const
 {
+	if (!TimeRange.Overlaps(GetRange()))
+	{
+		return;
+	}
+
 	for (auto It(TimingCurve.GetKeyHandleIterator()); It; ++It)
 	{
 		float Time = TimingCurve.GetKeyTime(It.Key());
-		if (IsTimeWithinSection(Time))
+		if (TimeRange.Contains(Time))
 		{
-			KeyHandles.Add(It.Key());
+			OutKeyHandles.Add(It.Key());
 		}
 	}
 }
diff --git a/Engine/Source/Runtime/MovieSceneTracks/Private/Sections/MovieScene3DTransformSection.cpp b/Engine/Source/Runtime/MovieSceneTracks/Private/Sections/MovieScene3DTransformSection.cpp
--- a/Engine/Source/Runtime/MovieSceneTracks/Private/Sections/MovieScene3DTransformSection.cpp
+++ b/Engine/Source/Runtime/MovieSceneTracks/Private/Sections/MovieScene3DTransformSection.cpp
@@ -192,34 +192,39 @@ void UMovieScene3DTransformSection::DilateSection(float DilationFactor, float Or
 }
 
 
-void UMovieScene3DTransformSection::GetKeyHandles(TSet<FKeyHandle>& KeyHandles) const
+void UMovieScene3DTransformSection::GetKeyHandles(TSet<FKeyHandle>& OutKeyHandles, TRange<float> TimeRange) const
 {
+	if (!TimeRange.Overlaps(GetRange()))
+	{
+		return;
+	}
+
 	for (int32 Axis = 0; Axis < 3; ++Axis)
 	{
 		for (auto It(Translation[Axis].GetKeyHandleIterator()); It; ++It)
 		{
 			float Time = Translation[Axis].GetKeyTime(It.Key());
-			if (IsTimeWithinSection(Time))
+			if (TimeRange.Contains(Time))
 			{
-				KeyHandles.Add(It.Key());
+				OutKeyHandles.Add(It.Key());
 			}
 		}
 
 		for (auto It(Rotation[Axis].GetKeyHandleIterator()); It; ++It)
 		{
 			float Time = Rotation[Axis].GetKeyTime(It.Key());
-			if (IsTimeWithinSection(Time))
+			if (TimeRange.Contains(Time))
 			{
-				KeyHandles.Add(It.Key());
+				OutKeyHandles.Add(It.Key());
 			}
 		}
 
 		for (auto It(Scale[Axis].GetKeyHandleIterator()); It; ++It)
 		{
 			float Time = Scale[Axis].GetKeyTime(It.Key());
-			if (IsTimeWithinSection(Time))
+			if (TimeRange.Contains(Time))
 			{
-				KeyHandles.Add(It.Key());
+				OutKeyHandles.Add(It.Key());
 			}
 		}
 	}
diff --git a/Engine/Source/Runtime/MovieSceneTracks/Private/Sections/MovieSceneActorReferenceSection.cpp b/Engine/Source/Runtime/MovieSceneTracks/Private/Sections/MovieSceneActorReferenceSection.cpp
--- a/Engine/Source/Runtime/MovieSceneTracks/Private/Sections/MovieSceneActorReferenceSection.cpp
+++ b/Engine/Source/Runtime/MovieSceneTracks/Private/Sections/MovieSceneActorReferenceSection.cpp
@@ -33,14 +33,19 @@ void UMovieSceneActorReferenceSection::DilateSection( float DilationFactor, floa
 }
 
 
-void UMovieSceneActorReferenceSection::GetKeyHandles(TSet<FKeyHandle>& KeyHandles) const
+void UMovieSceneActorReferenceSection::GetKeyHandles(TSet<FKeyHandle>& OutKeyHandles, TRange<float> TimeRange) const
 {
+	if (!TimeRange.Overlaps(GetRange()))
+	{
+		return;
+	}
+
 	for ( auto It( ActorGuidIndexCurve.GetKeyHandleIterator() ); It; ++It )
 	{
 		float Time = ActorGuidIndexCurve.GetKeyTime( It.Key() );
-		if (IsTimeWithinSection(Time))
+		if (TimeRange.Contains(Time))
 		{
-			KeyHandles.Add(It.Key());
+			OutKeyHandles.Add(It.Key());
 		}
 	}
 }
diff --git a/Engine/Source/Runtime/MovieSceneTracks/Private/Sections/MovieSceneAudioSection.cpp b/Engine/Source/Runtime/MovieSceneTracks/Private/Sections/MovieSceneAudioSection.cpp
--- a/Engine/Source/Runtime/MovieSceneTracks/Private/Sections/MovieSceneAudioSection.cpp
+++ b/Engine/Source/Runtime/MovieSceneTracks/Private/Sections/MovieSceneAudioSection.cpp
@@ -12,6 +12,7 @@ UMovieSceneAudioSection::UMovieSceneAudioSection( const FObjectInitializer& Obje
 	Sound = nullptr;
 	AudioStartTime = 0.f;
 	AudioDilationFactor = 1.f;
+	AudioVolume = 1.f;
 #if WITH_EDITORONLY_DATA
 	bShowIntensity = false;
 #endif
@@ -49,7 +50,7 @@ void UMovieSceneAudioSection::DilateSection( float DilationFactor, float Origin,
 }
 
 
-void UMovieSceneAudioSection::GetKeyHandles(TSet<FKeyHandle>& KeyHandles) const
+void UMovieSceneAudioSection::GetKeyHandles(TSet<FKeyHandle>& KeyHandles, TRange<float> TimeRange) const
 {
 	// do nothing
 }
diff --git a/Engine/Source/Runtime/MovieSceneTracks/Private/Sections/MovieSceneBoolSection.cpp b/Engine/Source/Runtime/MovieSceneTracks/Private/Sections/MovieSceneBoolSection.cpp
--- a/Engine/Source/Runtime/MovieSceneTracks/Private/Sections/MovieSceneBoolSection.cpp
+++ b/Engine/Source/Runtime/MovieSceneTracks/Private/Sections/MovieSceneBoolSection.cpp
@@ -37,14 +37,19 @@ void UMovieSceneBoolSection::DilateSection( float DilationFactor, float Origin,
 }
 
 
-void UMovieSceneBoolSection::GetKeyHandles(TSet<FKeyHandle>& KeyHandles) const
+void UMovieSceneBoolSection::GetKeyHandles(TSet<FKeyHandle>& OutKeyHandles, TRange<float> TimeRange) const
 {
+	if (!TimeRange.Overlaps(GetRange()))
+	{
+		return;
+	}
+
 	for (auto It(BoolCurve.GetKeyHandleIterator()); It; ++It)
 	{
 		float Time = BoolCurve.GetKeyTime(It.Key());
-		if (IsTimeWithinSection(Time))
+		if (TimeRange.Contains(Time))
 		{
-			KeyHandles.Add(It.Key());
+			OutKeyHandles.Add(It.Key());
 		}
 	}
 }
diff --git a/Engine/Source/Runtime/MovieSceneTracks/Private/Sections/MovieSceneByteSection.cpp b/Engine/Source/Runtime/MovieSceneTracks/Private/Sections/MovieSceneByteSection.cpp
--- a/Engine/Source/Runtime/MovieSceneTracks/Private/Sections/MovieSceneByteSection.cpp
+++ b/Engine/Source/Runtime/MovieSceneTracks/Private/Sections/MovieSceneByteSection.cpp
@@ -31,13 +31,20 @@ void UMovieSceneByteSection::DilateSection( float DilationFactor, float Origin,
 }
 
 
-void UMovieSceneByteSection::GetKeyHandles(TSet<FKeyHandle>& KeyHandles) const
+void UMovieSceneByteSection::GetKeyHandles(TSet<FKeyHandle>& OutKeyHandles, TRange<float> TimeRange) const
 {
+	if (!TimeRange.Overlaps(GetRange()))
+	{
+		return;
+	}
+
 	for (auto It(ByteCurve.GetKeyHandleIterator()); It; ++It)
 	{
 		float Time = ByteCurve.GetKeyTime(It.Key());
-		if (IsTimeWithinSection(Time))
-			KeyHandles.Add(It.Key());
+		if (TimeRange.Contains(Time))
+		{
+			OutKeyHandles.Add(It.Key());
+		}
 	}
 }
 
diff --git a/Engine/Source/Runtime/MovieSceneTracks/Private/Sections/MovieSceneCameraAnimSection.cpp b/Engine/Source/Runtime/MovieSceneTracks/Private/Sections/MovieSceneCameraAnimSection.cpp
--- a/Engine/Source/Runtime/MovieSceneTracks/Private/Sections/MovieSceneCameraAnimSection.cpp
+++ b/Engine/Source/Runtime/MovieSceneTracks/Private/Sections/MovieSceneCameraAnimSection.cpp
@@ -36,12 +36,17 @@ void UMovieSceneCameraAnimSection::DilateSection(float DilationFactor, float Ori
 }
 
 
-void UMovieSceneCameraAnimSection::GetKeyHandles(TSet<FKeyHandle>& KeyHandles) const
+void UMovieSceneCameraAnimSection::GetKeyHandles(TSet<FKeyHandle>& OutKeyHandles, TRange<float> TimeRange) const
 {
+	if (!TimeRange.Overlaps(GetRange()))
+	{
+		return;
+	}
+
 // 	for (auto It(AnimWeightCurve.GetKeyHandleIterator()); It; ++It)
 // 	{
 // 		float Time = AnimWeightCurve.GetKeyTime(It.Key());
-// 		if (IsTimeWithinSection(Time))
+// 		if (TimeRange.Contains(Time))
 // 		{
 // 			KeyHandles.Add(It.Key());
 // 		}
diff --git a/Engine/Source/Runtime/MovieSceneTracks/Private/Sections/MovieSceneCameraShakeSection.cpp b/Engine/Source/Runtime/MovieSceneTracks/Private/Sections/MovieSceneCameraShakeSection.cpp
--- a/Engine/Source/Runtime/MovieSceneTracks/Private/Sections/MovieSceneCameraShakeSection.cpp
+++ b/Engine/Source/Runtime/MovieSceneTracks/Private/Sections/MovieSceneCameraShakeSection.cpp
@@ -30,12 +30,17 @@ void UMovieSceneCameraShakeSection::DilateSection(float DilationFactor, float Or
 }
 
 
-void UMovieSceneCameraShakeSection::GetKeyHandles(TSet<FKeyHandle>& KeyHandles) const
+void UMovieSceneCameraShakeSection::GetKeyHandles(TSet<FKeyHandle>& OutKeyHandles, TRange<float> TimeRange) const
 {
+	if (!TimeRange.Overlaps(GetRange()))
+	{
+		return;
+	}
+
 // 	for (auto It(ShakeWeightCurve.GetKeyHandleIterator()); It; ++It)
 // 	{
 // 		float Time = ShakeWeightCurve.GetKeyTime(It.Key());
-// 		if (IsTimeWithinSection(Time))
+// 		if (TimeRange.Contains(Time))
 // 		{
 // 			KeyHandles.Add(It.Key());
 // 		}
diff --git a/Engine/Source/Runtime/MovieSceneTracks/Private/Sections/MovieSceneColorSection.cpp b/Engine/Source/Runtime/MovieSceneTracks/Private/Sections/MovieSceneColorSection.cpp
--- a/Engine/Source/Runtime/MovieSceneTracks/Private/Sections/MovieSceneColorSection.cpp
+++ b/Engine/Source/Runtime/MovieSceneTracks/Private/Sections/MovieSceneColorSection.cpp
@@ -64,40 +64,47 @@ void UMovieSceneColorSection::DilateSection(float DilationFactor, float Origin,
 }
 
 
-void UMovieSceneColorSection::GetKeyHandles(TSet<FKeyHandle>& KeyHandles) const
+void UMovieSceneColorSection::GetKeyHandles(TSet<FKeyHandle>& OutKeyHandles, TRange<float> TimeRange) const
 {
+	if (!TimeRange.Overlaps(GetRange()))
+	{
+		return;
+	}
+
 	for (auto It(RedCurve.GetKeyHandleIterator()); It; ++It)
 	{
 		float Time = RedCurve.GetKeyTime(It.Key());
-		if (IsTimeWithinSection(Time))
+		if (TimeRange.Contains(Time))
 		{
-			KeyHandles.Add(It.Key());
+			OutKeyHandles.Add(It.Key());
 		}
 	}
 
 	for (auto It(GreenCurve.GetKeyHandleIterator()); It; ++It)
 	{
 		float Time = GreenCurve.GetKeyTime(It.Key());
-		if (IsTimeWithinSection(Time))
+		if (TimeRange.Contains(Time))
 		{
-			KeyHandles.Add(It.Key());
+			OutKeyHandles.Add(It.Key());
 		}
 	}
 
 	for (auto It(BlueCurve.GetKeyHandleIterator()); It; ++It)
 	{
 		float Time = BlueCurve.GetKeyTime(It.Key());
-		if (IsTimeWithinSection(Time))
+		if (TimeRange.Contains(Time))
 		{
-			KeyHandles.Add(It.Key());
+			OutKeyHandles.Add(It.Key());
 		}
 	}
 
 	for (auto It(AlphaCurve.GetKeyHandleIterator()); It; ++It)
 	{
 		float Time = AlphaCurve.GetKeyTime(It.Key());
-		if (IsTimeWithinSection(Time))
-			KeyHandles.Add(It.Key());
+		if (TimeRange.Contains(Time))
+		{
+			OutKeyHandles.Add(It.Key());
+		}
 	}
 }
 
diff --git a/Engine/Source/Runtime/MovieSceneTracks/Private/Sections/MovieSceneEventSection.cpp b/Engine/Source/Runtime/MovieSceneTracks/Private/Sections/MovieSceneEventSection.cpp
--- a/Engine/Source/Runtime/MovieSceneTracks/Private/Sections/MovieSceneEventSection.cpp
+++ b/Engine/Source/Runtime/MovieSceneTracks/Private/Sections/MovieSceneEventSection.cpp
@@ -25,7 +25,7 @@ void UMovieSceneEventSection::AddKey(float Time, const FName& EventName, FKeyPar
 }
 
 
-void UMovieSceneEventSection::TriggerEvents(UObject* EventContextObject, float Position, float LastPosition)
+void UMovieSceneEventSection::TriggerEvents(TArray<UObject*> EventContextObjects, float Position, float LastPosition)
 {
 	const TArray<FNameCurveKey>& Keys = Events.GetKeys();
 
@@ -35,7 +35,7 @@ void UMovieSceneEventSection::TriggerEvents(UObject* EventContextObject, float P
 		{
 			if ((Key.Time >= LastPosition) && (Key.Time <= Position))
 			{
-				TriggerEvent(Key.Value, EventContextObject);
+				TriggerEvent(Key.Value, EventContextObjects);
 			}
 		}
 	}
@@ -47,7 +47,7 @@ void UMovieSceneEventSection::TriggerEvents(UObject* EventContextObject, float P
 
 			if ((Key.Time >= Position) && (Key.Time <= LastPosition))
 			{
-				TriggerEvent(Key.Value, EventContextObject);
+				TriggerEvent(Key.Value, EventContextObjects);
 			}		
 		}
 	}
@@ -65,12 +65,12 @@ void UMovieSceneEventSection::DilateSection(float DilationFactor, float Origin,
 }
 
 
-void UMovieSceneEventSection::GetKeyHandles(TSet<FKeyHandle>& KeyHandles) const
+void UMovieSceneEventSection::GetKeyHandles(TSet<FKeyHandle>& KeyHandles, TRange<float> TimeRange) const
 {
 	for (auto It(Events.GetKeyHandleIterator()); It; ++It)
 	{
 		float Time = Events.GetKeyTime(It.Key());
-		if (IsTimeWithinSection(Time))
+		if (TimeRange.Contains(Time))
 		{
 			KeyHandles.Add(It.Key());
 		}
@@ -89,22 +89,25 @@ void UMovieSceneEventSection::MoveSection(float DeltaPosition, TSet<FKeyHandle>&
 /* UMovieSceneSection overrides
  *****************************************************************************/
 
-void UMovieSceneEventSection::TriggerEvent(const FName& Event, UObject* EventContextObject)
+void UMovieSceneEventSection::TriggerEvent(const FName& Event, TArray<UObject*> EventContextObjects)
 {
-	UFunction* EventFunction = EventContextObject->FindFunction(Event);
-
-	if (EventFunction == nullptr)
-	{
-		// @todo sequencer: gmp: add external log category for MovieScene
-		//UE_LOG(LogMovieScene, Log, TEXT("UMovieSceneEventSection::TriggerEvent: Unable to find function '%s'"), *Event.ToString());
-	}
-	else if (EventFunction->NumParms != 0)
+	for (UObject* EventContextObject : EventContextObjects)
 	{
-		// @todo sequencer: gmp: add external log category for MovieScene
-		//UE_LOG(LogMovieScene, Log, TEXT("UMovieSceneEventSection::TriggerEvent: Function '%s' does not have zero parameters."), *Event.ToString());
-	}
-	else
-	{
-		EventContextObject->ProcessEvent(EventFunction, nullptr);
+		UFunction* EventFunction = EventContextObject->FindFunction(Event);
+
+		if (EventFunction == nullptr)
+		{
+			// @todo sequencer: gmp: add external log category for MovieScene
+			//UE_LOG(LogMovieScene, Log, TEXT("UMovieSceneEventSection::TriggerEvent: Unable to find function '%s'"), *Event.ToString());
+		}
+		else if (EventFunction->NumParms != 0)
+		{
+			// @todo sequencer: gmp: add external log category for MovieScene
+			//UE_LOG(LogMovieScene, Log, TEXT("UMovieSceneEventSection::TriggerEvent: Function '%s' does not have zero parameters."), *Event.ToString());
+		}
+		else
+		{
+			EventContextObject->ProcessEvent(EventFunction, nullptr);
+		}
 	}
 }
diff --git a/Engine/Source/Runtime/MovieSceneTracks/Private/Sections/MovieSceneFloatSection.cpp b/Engine/Source/Runtime/MovieSceneTracks/Private/Sections/MovieSceneFloatSection.cpp
--- a/Engine/Source/Runtime/MovieSceneTracks/Private/Sections/MovieSceneFloatSection.cpp
+++ b/Engine/Source/Runtime/MovieSceneTracks/Private/Sections/MovieSceneFloatSection.cpp
@@ -32,14 +32,19 @@ void UMovieSceneFloatSection::DilateSection( float DilationFactor, float Origin,
 }
 
 
-void UMovieSceneFloatSection::GetKeyHandles(TSet<FKeyHandle>& KeyHandles) const
+void UMovieSceneFloatSection::GetKeyHandles(TSet<FKeyHandle>& OutKeyHandles, TRange<float> TimeRange) const
 {
+	if (!TimeRange.Overlaps(GetRange()))
+	{
+		return;
+	}
+
 	for (auto It(FloatCurve.GetKeyHandleIterator()); It; ++It)
 	{
 		float Time = FloatCurve.GetKeyTime(It.Key());
-		if (IsTimeWithinSection(Time))
+		if (TimeRange.Contains(Time))
 		{
-			KeyHandles.Add(It.Key());
+			OutKeyHandles.Add(It.Key());
 		}
 	}
 }
diff --git a/Engine/Source/Runtime/MovieSceneTracks/Private/Sections/MovieSceneParameterSection.cpp b/Engine/Source/Runtime/MovieSceneTracks/Private/Sections/MovieSceneParameterSection.cpp
--- a/Engine/Source/Runtime/MovieSceneTracks/Private/Sections/MovieSceneParameterSection.cpp
+++ b/Engine/Source/Runtime/MovieSceneTracks/Private/Sections/MovieSceneParameterSection.cpp
@@ -276,8 +276,13 @@ void UMovieSceneParameterSection::DilateSection(float DilationFactor, float Orig
 	}
 }
 
-void UMovieSceneParameterSection::GetKeyHandles(TSet<FKeyHandle>& KeyHandles) const
+void UMovieSceneParameterSection::GetKeyHandles(TSet<FKeyHandle>& OutKeyHandles, TRange<float> TimeRange) const
 {
+	if (!TimeRange.Overlaps(GetRange()))
+	{
+		return;
+	}
+
 	TArray<const FRichCurve*> AllCurves;
 	GatherCurves(AllCurves);
 
@@ -286,9 +291,9 @@ void UMovieSceneParameterSection::GetKeyHandles(TSet<FKeyHandle>& KeyHandles) co
 		for (auto It(Curve->GetKeyHandleIterator()); It; ++It)
 		{
 			float Time = Curve->GetKeyTime(It.Key());
-			if (IsTimeWithinSection(Time))
+			if (TimeRange.Contains(Time))
 			{
-				KeyHandles.Add(It.Key());
+				OutKeyHandles.Add(It.Key());
 			}
 		}
 	}
diff --git a/Engine/Source/Runtime/MovieSceneTracks/Private/Sections/MovieSceneParticleSection.cpp b/Engine/Source/Runtime/MovieSceneTracks/Private/Sections/MovieSceneParticleSection.cpp
--- a/Engine/Source/Runtime/MovieSceneTracks/Private/Sections/MovieSceneParticleSection.cpp
+++ b/Engine/Source/Runtime/MovieSceneTracks/Private/Sections/MovieSceneParticleSection.cpp
@@ -36,14 +36,19 @@ void UMovieSceneParticleSection::DilateSection( float DilationFactor, float Orig
 	ParticleKeys.ScaleCurve( Origin, DilationFactor, KeyHandles );
 }
 
-void UMovieSceneParticleSection::GetKeyHandles( TSet<FKeyHandle>& KeyHandles ) const
+void UMovieSceneParticleSection::GetKeyHandles( TSet<FKeyHandle>& OutKeyHandles, TRange<float> TimeRange ) const
 {
+	if (!TimeRange.Overlaps(GetRange()))
+	{
+		return;
+	}
+
 	for ( auto It( ParticleKeys.GetKeyHandleIterator() ); It; ++It )
 	{
 		float Time = ParticleKeys.GetKeyTime( It.Key() );
-		if ( IsTimeWithinSection( Time ) )
+		if (TimeRange.Contains(Time))
 		{
-			KeyHandles.Add( It.Key() );
+			OutKeyHandles.Add( It.Key() );
 		}
 	}
 }
diff --git a/Engine/Source/Runtime/MovieSceneTracks/Private/Sections/MovieSceneSkeletalAnimationSection.cpp b/Engine/Source/Runtime/MovieSceneTracks/Private/Sections/MovieSceneSkeletalAnimationSection.cpp
--- a/Engine/Source/Runtime/MovieSceneTracks/Private/Sections/MovieSceneSkeletalAnimationSection.cpp
+++ b/Engine/Source/Runtime/MovieSceneTracks/Private/Sections/MovieSceneSkeletalAnimationSection.cpp
@@ -52,7 +52,7 @@ UMovieSceneSection* UMovieSceneSkeletalAnimationSection::SplitSection(float Spli
 }
 
 
-void UMovieSceneSkeletalAnimationSection::GetKeyHandles(TSet<FKeyHandle>& KeyHandles) const
+void UMovieSceneSkeletalAnimationSection::GetKeyHandles(TSet<FKeyHandle>& KeyHandles, TRange<float> TimeRange) const
 {
 	// do nothing
 }
diff --git a/Engine/Source/Runtime/MovieSceneTracks/Private/Sections/MovieSceneSubSection.cpp b/Engine/Source/Runtime/MovieSceneTracks/Private/Sections/MovieSceneSubSection.cpp
--- a/Engine/Source/Runtime/MovieSceneTracks/Private/Sections/MovieSceneSubSection.cpp
+++ b/Engine/Source/Runtime/MovieSceneTracks/Private/Sections/MovieSceneSubSection.cpp
@@ -98,4 +98,52 @@ void UMovieSceneSubSection::PostEditChangeProperty(FPropertyChangedEvent& Proper
 		OnSequenceChangedDelegate.ExecuteIfBound(SubSequence);
 	}
 }
-#endif
\ No newline at end of file
+#endif
+
+UMovieSceneSection* UMovieSceneSubSection::SplitSection( float SplitTime )
+{
+	if ( !IsTimeWithinSection( SplitTime ) )
+	{
+		return nullptr;
+	}
+
+	float InitialStartTime = GetStartTime();
+	float InitialStartOffset = StartOffset;
+	float NewStartOffset = ( SplitTime - InitialStartTime ) / TimeScale;
+	NewStartOffset += InitialStartOffset;
+
+	// Ensure start offset is not less than 0
+	NewStartOffset = FMath::Max( NewStartOffset, 0.f );
+
+	UMovieSceneSubSection* NewSection = Cast<UMovieSceneSubSection>( UMovieSceneSection::SplitSection( SplitTime ) );
+	if ( NewSection )
+	{
+		NewSection->StartOffset = NewStartOffset;
+		return NewSection;
+	}
+
+	return nullptr;
+}
+
+void UMovieSceneSubSection::TrimSection( float TrimTime, bool bTrimLeft )
+{
+	if ( !IsTimeWithinSection( TrimTime ) )
+	{
+		return;
+	}
+
+	float InitialStartTime = GetStartTime();
+	float InitialStartOffset = StartOffset;
+
+	UMovieSceneSection::TrimSection( TrimTime, bTrimLeft );
+
+	// If trimming off the left, set the offset of the shot
+	if ( bTrimLeft )
+	{
+		float NewStartOffset = ( TrimTime - InitialStartTime ) / TimeScale;
+		NewStartOffset += InitialStartOffset;
+
+		// Ensure start offset is not less than 0
+		StartOffset = FMath::Max( NewStartOffset, 0.f );
+	}
+}
\ No newline at end of file
diff --git a/Engine/Source/Runtime/MovieSceneTracks/Private/Sections/MovieSceneVectorSection.cpp b/Engine/Source/Runtime/MovieSceneTracks/Private/Sections/MovieSceneVectorSection.cpp
--- a/Engine/Source/Runtime/MovieSceneTracks/Private/Sections/MovieSceneVectorSection.cpp
+++ b/Engine/Source/Runtime/MovieSceneTracks/Private/Sections/MovieSceneVectorSection.cpp
@@ -67,16 +67,21 @@ void UMovieSceneVectorSection::DilateSection(float DilationFactor, float Origin,
 }
 
 
-void UMovieSceneVectorSection::GetKeyHandles(TSet<FKeyHandle>& KeyHandles) const
+void UMovieSceneVectorSection::GetKeyHandles(TSet<FKeyHandle>& OutKeyHandles, TRange<float> TimeRange) const
 {
+	if (!TimeRange.Overlaps(GetRange()))
+	{
+		return;
+	}
+
 	for (int32 i = 0; i < ChannelsUsed; ++i)
 	{
 		for (auto It(Curves[i].GetKeyHandleIterator()); It; ++It)
 		{
 			float Time = Curves[i].GetKeyTime(It.Key());
-			if (IsTimeWithinSection(Time))
+			if (TimeRange.Contains(Time))
 			{
-				KeyHandles.Add(It.Key());
+				OutKeyHandles.Add(It.Key());
 			}
 		}
 	}
diff --git a/Engine/Source/Runtime/MovieSceneTracks/Private/TrackInstances/MovieSceneAudioTrackInstance.cpp b/Engine/Source/Runtime/MovieSceneTracks/Private/TrackInstances/MovieSceneAudioTrackInstance.cpp
--- a/Engine/Source/Runtime/MovieSceneTracks/Private/TrackInstances/MovieSceneAudioTrackInstance.cpp
+++ b/Engine/Source/Runtime/MovieSceneTracks/Private/TrackInstances/MovieSceneAudioTrackInstance.cpp
@@ -161,7 +161,7 @@ void FMovieSceneAudioTrackInstance::PlaySound(UMovieSceneAudioSection* AudioSect
 	Component->bAllowSpatialization = !AudioTrack->IsAMasterTrack();
 	Component->Stop();
 	Component->SetSound(AudioSection->GetSound());
-	Component->SetVolumeMultiplier(1.f);
+	Component->SetVolumeMultiplier(AudioSection->GetAudioVolume());
 	Component->SetPitchMultiplier(PitchMultiplier);
 	Component->bIsUISound = true;
 	Component->Play(Time - AudioSection->GetAudioStartTime());
diff --git a/Engine/Source/Runtime/MovieSceneTracks/Private/TrackInstances/MovieSceneCameraAnimTrackInstance.cpp b/Engine/Source/Runtime/MovieSceneTracks/Private/TrackInstances/MovieSceneCameraAnimTrackInstance.cpp
--- a/Engine/Source/Runtime/MovieSceneTracks/Private/TrackInstances/MovieSceneCameraAnimTrackInstance.cpp
+++ b/Engine/Source/Runtime/MovieSceneTracks/Private/TrackInstances/MovieSceneCameraAnimTrackInstance.cpp
@@ -33,6 +33,7 @@ void FMovieSceneCameraAnimTrackInstance::Update(EMovieSceneUpdateData& UpdateDat
 			if (CameraComponent)
 			{
 				CameraComponent->ClearAdditiveOffset();
+				CameraComponent->ClearExtraPostProcessBlends();
 			}
 		}
 	}
@@ -66,7 +67,7 @@ void FMovieSceneCameraAnimTrackInstance::Update(EMovieSceneUpdateData& UpdateDat
 					SectionInstanceDataMap.Remove(DeadSection);
 				}
 
-				// #todo sort by start time to match application order of player camera
+				// sort by start time to match application order of player camera
 				ActiveSections.Sort(
 					[](const UMovieSceneCameraAnimSection& One, const UMovieSceneCameraAnimSection& Two)
 				{
@@ -135,6 +136,20 @@ void FMovieSceneCameraAnimTrackInstance::Update(EMovieSceneUpdateData& UpdateDat
 
 							SectionInstanceData.AdditiveFOVOffset = NewFOVToBaseFOV;
 							SectionInstanceData.AdditiveOffset = NewCameraToBaseCamera;
+
+							// harvest PP changes
+							{
+								UCameraComponent* AnimCamComp = GetTempCameraActor()->GetCameraComponent();
+								if (AnimCamComp)
+								{
+									SectionInstanceData.PostProcessingBlendWeight = AnimCamComp->PostProcessBlendWeight;
+									if (SectionInstanceData.PostProcessingBlendWeight > 0.f)
+									{
+										// avoid big struct copy if not necessary
+										SectionInstanceData.PostProcessingSettings = AnimCamComp->PostProcessSettings;
+									}
+								}
+							}
 						}
 					}
 				}
@@ -147,6 +162,11 @@ void FMovieSceneCameraAnimTrackInstance::Update(EMovieSceneUpdateData& UpdateDat
 					FMovieSceneCameraAnimSectionInstanceData& SectionInstanceData = SectionInstanceDataMap.FindChecked(ActiveSection);
 					TotalTransform = TotalTransform * SectionInstanceData.AdditiveOffset;
 					TotalFOVOffset += SectionInstanceData.AdditiveFOVOffset;
+
+					if (SectionInstanceData.PostProcessingBlendWeight > 0.f)
+					{
+						CameraComponent->AddExtraPostProcessBlend(SectionInstanceData.PostProcessingSettings, SectionInstanceData.PostProcessingBlendWeight);
+					}
 				}
 
 				CameraComponent->AddAdditiveOffset(TotalTransform, TotalFOVOffset);
diff --git a/Engine/Source/Runtime/MovieSceneTracks/Private/TrackInstances/MovieSceneCameraShakeTrackInstance.cpp b/Engine/Source/Runtime/MovieSceneTracks/Private/TrackInstances/MovieSceneCameraShakeTrackInstance.cpp
--- a/Engine/Source/Runtime/MovieSceneTracks/Private/TrackInstances/MovieSceneCameraShakeTrackInstance.cpp
+++ b/Engine/Source/Runtime/MovieSceneTracks/Private/TrackInstances/MovieSceneCameraShakeTrackInstance.cpp
@@ -33,6 +33,7 @@ void FMovieSceneCameraShakeTrackInstance::Update(EMovieSceneUpdateData& UpdateDa
 			if (CameraComponent)
 			{
 				CameraComponent->ClearAdditiveOffset();
+				CameraComponent->ClearExtraPostProcessBlends();
 			}
 		}
 	}
@@ -70,7 +71,7 @@ void FMovieSceneCameraShakeTrackInstance::Update(EMovieSceneUpdateData& UpdateDa
 					SectionInstanceDataMap.Remove(DeadSection);
 				}
 
-				// #todo sort by start time to match application order of player camera
+				// sort by start time to match application order of player camera
 				ActiveSections.Sort(
 					[](const UMovieSceneCameraShakeSection& One, const UMovieSceneCameraShakeSection& Two)
 				{
@@ -91,11 +92,10 @@ void FMovieSceneCameraShakeTrackInstance::Update(EMovieSceneUpdateData& UpdateDa
 						{
 							// make it root so GC doesn't take it away
 							ShakeInst->AddToRoot();
+							ShakeInst->SetTempCameraAnimActor(GetTempCameraActor());
 							ShakeInst->PlayShake(nullptr, ActiveSection->PlayScale, ActiveSection->PlaySpace, ActiveSection->UserDefinedPlaySpace);
 							if (ShakeInst->AnimInst)
 							{
-								// if there was a cameraanim component of the shake, give it our temp camera actor to use
-								ShakeInst->AnimInst->SetCameraActor(GetTempCameraActor());
 								ShakeInst->AnimInst->SetStopAutomatically(false);
 							}
 						}
@@ -144,6 +144,20 @@ void FMovieSceneCameraShakeTrackInstance::Update(EMovieSceneUpdateData& UpdateDa
 
 						SectionInstanceData.AdditiveFOVOffset = NewFOVToBaseFOV;
 						SectionInstanceData.AdditiveOffset = NewCameraToBaseCamera;
+
+						// harvest PP changes
+						{
+							UCameraComponent* AnimCamComp = GetTempCameraActor()->GetCameraComponent();
+							if (AnimCamComp)
+							{
+								SectionInstanceData.PostProcessingBlendWeight = AnimCamComp->PostProcessBlendWeight;
+								if (SectionInstanceData.PostProcessingBlendWeight > 0.f)
+								{
+									// avoid big struct copy if not necessary
+									SectionInstanceData.PostProcessingSettings = AnimCamComp->PostProcessSettings;
+								}
+							}
+						}
 					}
 				}
 
@@ -155,6 +169,11 @@ void FMovieSceneCameraShakeTrackInstance::Update(EMovieSceneUpdateData& UpdateDa
 					FMovieSceneCameraShakeSectionInstanceData& SectionInstanceData = SectionInstanceDataMap.FindChecked(ActiveSection);
 					TotalTransform = TotalTransform * SectionInstanceData.AdditiveOffset;
 					TotalFOVOffset += SectionInstanceData.AdditiveFOVOffset;
+
+					if (SectionInstanceData.PostProcessingBlendWeight > 0.f)
+					{
+						CameraComponent->AddExtraPostProcessBlend(SectionInstanceData.PostProcessingSettings, SectionInstanceData.PostProcessingBlendWeight);
+					}
 				}
 
 				CameraComponent->AddAdditiveOffset(TotalTransform, TotalFOVOffset);
diff --git a/Engine/Source/Runtime/MovieSceneTracks/Private/TrackInstances/MovieSceneColorTrackInstance.cpp b/Engine/Source/Runtime/MovieSceneTracks/Private/TrackInstances/MovieSceneColorTrackInstance.cpp
--- a/Engine/Source/Runtime/MovieSceneTracks/Private/TrackInstances/MovieSceneColorTrackInstance.cpp
+++ b/Engine/Source/Runtime/MovieSceneTracks/Private/TrackInstances/MovieSceneColorTrackInstance.cpp
@@ -38,13 +38,13 @@ void FMovieSceneColorTrackInstance::SaveState(const TArray<TWeakObjectPtr<UObjec
 				InitLinearColorMap.Add(Object, ColorValue);
 			}
 		}
-		else if( ColorType == EColorType::RegularColor )
+		else if( ColorType == EColorType::Color )
 		{
 			if (InitLinearColorMap.Find(Object) == nullptr)
 			{			
 				UProperty* Property = PropertyBindings->GetProperty(Object);
 				FColor ColorValue = PropertyBindings->GetCurrentValue<FColor>(Object);
-				FLinearColor LinearColorValue = ColorValue.ReinterpretAsLinear();
+				FLinearColor LinearColorValue = ColorValue;
 				InitLinearColorMap.Add(Object, ColorValue);
 			}
 		}
@@ -75,7 +75,7 @@ void FMovieSceneColorTrackInstance::RestoreState(const TArray<TWeakObjectPtr<UOb
 				PropertyBindings->CallFunction<FSlateColor>(Object, ColorValue);
 			}
 		}
-		else if( ColorType == EColorType::Linear || ColorType == EColorType::RegularColor )
+		else if( ColorType == EColorType::Linear || ColorType == EColorType::Color )
 		{
 			//todo
 			FLinearColor* ColorValue = InitLinearColorMap.Find(Object);
@@ -99,30 +99,53 @@ void FMovieSceneColorTrackInstance::Update(EMovieSceneUpdateData& UpdateData, co
 	for (auto ObjectPtr : RuntimeObjects)
 	{
 		UObject* Object = ObjectPtr.Get();
-
-		if( ColorType == EColorType::Slate )
+		if ( Object != nullptr )
 		{
-			FSlateColor ColorValue = PropertyBindings->GetCurrentValue<FSlateColor>(Object);
-			FLinearColor LinearColor = ColorValue.GetSpecifiedColor();
-			if(ColorTrack->Eval(UpdateData.Position, UpdateData.LastPosition, LinearColor))
+			if ( ColorType == EColorType::Slate )
 			{
-				FSlateColor NewColor(LinearColor);
-				PropertyBindings->CallFunction<FSlateColor>(Object, &NewColor);
+				FSlateColor ColorValue = PropertyBindings->GetCurrentValue<FSlateColor>( Object );
+				FLinearColor LinearColor = ColorValue.GetSpecifiedColor();
+
+				if ( ColorTrack->Eval( UpdateData.Position, UpdateData.LastPosition, LinearColor ) )
+				{
+					FSlateColor NewColor( LinearColor );
+					PropertyBindings->CallFunction<FSlateColor>( Object, &NewColor );
+				}
 			}
-		}
-		else
-		{
-			FLinearColor ColorValue = ColorType == EColorType::Linear ? PropertyBindings->GetCurrentValue<FLinearColor>(Object) : PropertyBindings->GetCurrentValue<FColor>(Object).ReinterpretAsLinear();
-			if(ColorTrack->Eval(UpdateData.Position, UpdateData.LastPosition, ColorValue))
+			else if ( ColorType == EColorType::Color )
 			{
-				// LightComponent's SetLightColor applies an sRGB conversion by default, so invert it here before applying the value 
-				if (ColorType == EColorType::RegularColor)
+				// We assume the color we get back is in sRGB, assigning it to a linear color will implicitly
+				// convert it to a linear color instead of using ReinterpretAsLinear which will just change the
+				// bytes into floats using divide by 255.
+				FColor SrgbColor = PropertyBindings->GetCurrentValue<FColor>( Object );
+				FLinearColor LinearColorValue = SrgbColor;
+				bool bConvertBackToSRgb = true;
+				if ( ColorTrack->Eval( UpdateData.Position, UpdateData.LastPosition, LinearColorValue ) )
 				{
-					const bool sRGB = false;
-					ColorValue = FLinearColor(ColorValue.ToFColor(sRGB));
+					// Light components have to be handled specially here because their set function takes two values, the linear color
+					// and whether or not the linear color needs to be converted back to sRGB.  All other other set function cases should
+					// follow the sequencer convention of having a single parameter of the correct type, which in this case is an FColor
+					// already in sRGB format.
+					ULightComponent* LightComponent = Cast<ULightComponent>( Object );
+					if ( LightComponent != nullptr )
+					{
+						LightComponent->SetLightColor( LinearColorValue, bConvertBackToSRgb );
+					}
+					else
+					{
+						FColor SRgbColorValue = LinearColorValue.ToFColor( bConvertBackToSRgb );
+						PropertyBindings->CallFunction<FColor>( Object, &SRgbColorValue );
+					}
 				}
+			}
+			else
+			{
+				FLinearColor ColorValue = PropertyBindings->GetCurrentValue<FLinearColor>( Object );
 
-				PropertyBindings->CallFunction<FLinearColor>(Object, &ColorValue);
+				if ( ColorTrack->Eval( UpdateData.Position, UpdateData.LastPosition, ColorValue ) )
+				{
+					PropertyBindings->CallFunction<FLinearColor>( Object, &ColorValue );
+				}
 			}
 		}
 	}
@@ -150,7 +173,7 @@ void FMovieSceneColorTrackInstance::RefreshInstance( const TArray<TWeakObjectPtr
 
 		if( StructName == NAME_Color )
 		{
-			ColorType = EColorType::RegularColor;
+			ColorType = EColorType::Color;
 		}
 		else if( StructName == SlateColor )
 		{
diff --git a/Engine/Source/Runtime/MovieSceneTracks/Private/TrackInstances/MovieSceneEventTrackInstance.cpp b/Engine/Source/Runtime/MovieSceneTracks/Private/TrackInstances/MovieSceneEventTrackInstance.cpp
--- a/Engine/Source/Runtime/MovieSceneTracks/Private/TrackInstances/MovieSceneEventTrackInstance.cpp
+++ b/Engine/Source/Runtime/MovieSceneTracks/Private/TrackInstances/MovieSceneEventTrackInstance.cpp
@@ -19,5 +19,5 @@ FMovieSceneEventTrackInstance::FMovieSceneEventTrackInstance(UMovieSceneEventTra
 
 void FMovieSceneEventTrackInstance::Update(EMovieSceneUpdateData& UpdateData, const TArray<TWeakObjectPtr<UObject>>& RuntimeObjects, IMovieScenePlayer& Player, FMovieSceneSequenceInstance& SequenceInstance)
 {
-	EventTrack->TriggerEvents(UpdateData.Position, UpdateData.LastPosition, Player.GetEventContext());
+	EventTrack->TriggerEvents(UpdateData.Position, UpdateData.LastPosition, Player.GetEventContexts());
 }
