Snap to not including scale should KeepWorld for scale not KeepRelative

diff --git a/Engine/Source/Runtime/Engine/Private/EngineTypes.cpp b/Engine/Source/Runtime/Engine/Private/EngineTypes.cpp
--- a/Engine/Source/Runtime/Engine/Private/EngineTypes.cpp
+++ b/Engine/Source/Runtime/Engine/Private/EngineTypes.cpp
@@ -5,7 +5,7 @@
 
 FAttachmentTransformRules FAttachmentTransformRules::KeepRelativeTransform(EAttachmentRule::KeepRelative, false);
 FAttachmentTransformRules FAttachmentTransformRules::KeepWorldTransform(EAttachmentRule::KeepWorld, false);
-FAttachmentTransformRules FAttachmentTransformRules::SnapToTargetNotIncludingScale(EAttachmentRule::SnapToTarget, EAttachmentRule::SnapToTarget, EAttachmentRule::KeepRelative, false);
+FAttachmentTransformRules FAttachmentTransformRules::SnapToTargetNotIncludingScale(EAttachmentRule::SnapToTarget, EAttachmentRule::SnapToTarget, EAttachmentRule::KeepWorld, false);
 FAttachmentTransformRules FAttachmentTransformRules::SnapToTargetIncludingScale(EAttachmentRule::SnapToTarget, false);
 
 FDetachmentTransformRules FDetachmentTransformRules::KeepRelativeTransform(EDetachmentRule::KeepRelative, true);
