Moved a desired change in 4.11.1 (had to be reverted for binary compatibility) back to Main.

Added input settings to:
- control if the viewport locks the mouse on acquire capture.
- control if the viewport acquires capture on the application launch (first window activate).

diff --git a/Engine/Source/Runtime/Engine/Classes/Engine/GameViewportClient.h b/Engine/Source/Runtime/Engine/Classes/Engine/GameViewportClient.h
--- a/Engine/Source/Runtime/Engine/Classes/Engine/GameViewportClient.h
+++ b/Engine/Source/Runtime/Engine/Classes/Engine/GameViewportClient.h
@@ -652,16 +652,40 @@ public:
 	}
 
 	/**
-	 * Sets whether or not the cursor is hidden when the viewport captures the mouse
+	 * Gets whether or not the viewport captures the Mouse on launch of the application
+	 * Technically this controls capture on the first window activate, so in situations 
+	 * where the application is launched but isn't activated the effect is delayed until
+	 * activation.
 	 */
+	virtual bool CaptureMouseOnLaunch() override;
+
+	/**
+	 * Sets whether or not the cursor is locked to the viewport when the viewport captures the mouse
+	 */
+	void SetLockDuringCapture(bool InLockDuringCapture)
+	{
+		bLockDuringCapture = InLockDuringCapture;
+	}
+
+	/**
+	 * Gets whether or not the cursor is locked to the viewport when the viewport captures the mouse
+	 */
+	virtual bool LockDuringCapture() override
+	{
+		return bHideCursorDuringCapture;
+	}
+
+	/**
+	* Sets whether or not the cursor is hidden when the viewport captures the mouse
+	*/
 	void SetHideCursorDuringCapture(bool InHideCursorDuringCapture)
 	{
 		bHideCursorDuringCapture = InHideCursorDuringCapture;
 	}
 
 	/**
-	 * Gets whether or not the cursor is hidden when the viewport captures the mouse
-	 */
+	* Gets whether or not the cursor is hidden when the viewport captures the mouse
+	*/
 	virtual bool HideCursorDuringCapture() override
 	{
 		return bHideCursorDuringCapture;
@@ -798,6 +822,9 @@ private:
 	/** Mouse capture behavior when the viewport is clicked */
 	EMouseCaptureMode MouseCaptureMode;
 
+	/** Whether or not the cursor is locked when the viewport captures the mouse */
+	bool bLockDuringCapture;
+
 	/** Whether or not the cursor is hidden when the viewport captures the mouse */
 	bool bHideCursorDuringCapture;
 
diff --git a/Engine/Source/Runtime/Engine/Classes/GameFramework/InputSettings.h b/Engine/Source/Runtime/Engine/Classes/GameFramework/InputSettings.h
--- a/Engine/Source/Runtime/Engine/Classes/GameFramework/InputSettings.h
+++ b/Engine/Source/Runtime/Engine/Classes/GameFramework/InputSettings.h
@@ -43,10 +43,18 @@ class ENGINE_API UInputSettings
 	UPROPERTY(config, EditAnywhere, Category="MouseProperties", AdvancedDisplay)
 	float DoubleClickTime;
 
+	/** Controls if the viewport will capture the mouse on Launch of the application */
+	UPROPERTY(config, EditAnywhere, Category = "ViewportProperties")
+	bool bCaptureMouseOnLaunch;
+	
 	/** The default mouse capture mode for the game viewport */
 	UPROPERTY(config, EditAnywhere, Category = "ViewportProperties")
 	TEnumAsByte<EMouseCaptureMode> DefaultViewportMouseCaptureMode;
 
+	/** The default mouse lock state when the viewport acquires capture */
+	UPROPERTY(config, EditAnywhere, Category = "ViewportProperties")
+	bool bDefaultViewportMouseLock;
+
 	/** List of Action Mappings */
 	UPROPERTY(config, EditAnywhere, Category="Bindings")
 	TArray<struct FInputActionKeyMapping> ActionMappings;
diff --git a/Engine/Source/Runtime/Engine/Private/GameViewportClient.cpp b/Engine/Source/Runtime/Engine/Private/GameViewportClient.cpp
--- a/Engine/Source/Runtime/Engine/Private/GameViewportClient.cpp
+++ b/Engine/Source/Runtime/Engine/Private/GameViewportClient.cpp
@@ -116,6 +116,7 @@ UGameViewportClient::UGameViewportClient(const FObjectInitializer& ObjectInitial
 	, HighResScreenshotDialog(NULL)
 	, bIgnoreInput(false)
 	, MouseCaptureMode(EMouseCaptureMode::CapturePermanently)
+	, bLockDuringCapture(true)
 	, bHideCursorDuringCapture(false)
 	, AudioDeviceHandle(INDEX_NONE)
 	, bHasAudioFocus(false)
@@ -175,6 +176,7 @@ UGameViewportClient::UGameViewportClient(FVTableHelper& Helper)
 	, bIgnoreInput(false)
 	, MouseCaptureMode(EMouseCaptureMode::CapturePermanently)
 	, bHideCursorDuringCapture(false)
+	, bLockDuringCapture(true)
 	, AudioDeviceHandle(INDEX_NONE)
 	, bHasAudioFocus(false)
 {
@@ -277,6 +279,7 @@ void UGameViewportClient::Init(struct FWorldContext& WorldContext, UGameInstance
 
 	// Set the projects default viewport mouse capture mode
 	MouseCaptureMode = GetDefault<UInputSettings>()->DefaultViewportMouseCaptureMode;
+	bLockDuringCapture = GetDefault<UInputSettings>()->bDefaultViewportMouseLock;
 
 	// Create the cursor Widgets
 	UUserInterfaceSettings* UISettings = GetMutableDefault<UUserInterfaceSettings>(UUserInterfaceSettings::StaticClass());
@@ -2035,7 +2038,10 @@ void UGameViewportClient::VerifyPathRenderingComponents()
 	}
 }
 
-
+bool UGameViewportClient::CaptureMouseOnLaunch()
+{
+	return GetDefault<UInputSettings>()->bCaptureMouseOnLaunch;
+}
 
 bool UGameViewportClient::Exec( UWorld* InWorld, const TCHAR* Cmd,FOutputDevice& Ar)
 {
diff --git a/Engine/Source/Runtime/Engine/Private/PlayerController.cpp b/Engine/Source/Runtime/Engine/Private/PlayerController.cpp
--- a/Engine/Source/Runtime/Engine/Private/PlayerController.cpp
+++ b/Engine/Source/Runtime/Engine/Private/PlayerController.cpp
@@ -4488,7 +4488,7 @@ void FInputModeDataBase::SetFocusAndLocking(FReply& SlateOperations, TSharedPtr<
 	}
 
 	if (bLockMouseToViewport)
-	{
+	{	
 		SlateOperations.LockMouseToWidget(InViewportWidget);
 	}
 	else
@@ -4506,6 +4506,7 @@ void FInputModeUIOnly::ApplyInputMode(FReply& SlateOperations, class UGameViewpo
 
 		SlateOperations.ReleaseMouseCapture();
 
+		GameViewportClient.SetLockDuringCapture(bLockMouseToViewport);
 		GameViewportClient.SetIgnoreInput(true);
 		GameViewportClient.SetCaptureMouseOnClick(EMouseCaptureMode::NoCapture);
 	}
@@ -4520,6 +4521,7 @@ void FInputModeGameAndUI::ApplyInputMode(FReply& SlateOperations, class UGameVie
 
 		SlateOperations.ReleaseMouseCapture();
 
+		GameViewportClient.SetLockDuringCapture(bLockMouseToViewport);
 		GameViewportClient.SetIgnoreInput(false);
 		GameViewportClient.SetHideCursorDuringCapture(bHideCursorDuringCapture);
 		GameViewportClient.SetCaptureMouseOnClick(EMouseCaptureMode::CaptureDuringMouseDown);
@@ -4535,6 +4537,7 @@ void FInputModeGameOnly::ApplyInputMode(FReply& SlateOperations, class UGameView
 		SlateOperations.UseHighPrecisionMouseMovement(ViewportWidgetRef);
 		SlateOperations.SetUserFocus(ViewportWidgetRef);
 		SlateOperations.LockMouseToWidget(ViewportWidgetRef);
+		GameViewportClient.SetLockDuringCapture(true);
 		GameViewportClient.SetIgnoreInput(false);
 		GameViewportClient.SetCaptureMouseOnClick(bConsumeCaptureMouseDown ? EMouseCaptureMode::CapturePermanently : EMouseCaptureMode::CapturePermanently_IncludingInitialMouseDown);
 	}
diff --git a/Engine/Source/Runtime/Engine/Private/Slate/SceneViewport.cpp b/Engine/Source/Runtime/Engine/Private/Slate/SceneViewport.cpp
--- a/Engine/Source/Runtime/Engine/Private/Slate/SceneViewport.cpp
+++ b/Engine/Source/Runtime/Engine/Private/Slate/SceneViewport.cpp
@@ -40,6 +40,11 @@ FSceneViewport::FSceneViewport( FViewportClient* InViewportClient, TSharedPtr<SV
 {
 	bIsSlateViewport = true;
 	RenderThreadSlateTexture = new FSlateRenderTargetRHI(nullptr, 0, 0);
+
+	if (InViewportClient)
+	{
+		bShouldCaptureMouseOnActivate = InViewportClient->CaptureMouseOnLaunch();
+	}
 }
 
 FSceneViewport::~FSceneViewport()
@@ -451,6 +456,13 @@ FReply FSceneViewport::AcquireFocusAndCapture(FIntPoint MousePosition)
 	UWorld* World = ViewportClient->GetWorld();
 	if (World && World->IsGameWorld() && World->GetGameInstance() && World->GetGameInstance()->GetFirstLocalPlayerController())
 	{
+		ReplyState.CaptureMouse(ViewportWidgetRef);
+
+		if (ViewportClient->LockDuringCapture())
+		{
+			ReplyState.LockMouseToWidget(ViewportWidgetRef);
+		}
+
 		bool bShouldShowMouseCursor = World->GetGameInstance()->GetFirstLocalPlayerController()->ShouldShowMouseCursor();
 		if (ViewportClient->HideCursorDuringCapture() && bShouldShowMouseCursor)
 		{
@@ -459,8 +471,6 @@ FReply FSceneViewport::AcquireFocusAndCapture(FIntPoint MousePosition)
 		}
 		if (bCursorHiddenDueToCapture || !bShouldShowMouseCursor)
 		{
-			ReplyState.CaptureMouse(ViewportWidgetRef);
-			ReplyState.LockMouseToWidget(ViewportWidgetRef);
 			ReplyState.UseHighPrecisionMouseMovement(ViewportWidgetRef);
 		}
 	}
diff --git a/Engine/Source/Runtime/Engine/Private/UserInterface/InputSettings.cpp b/Engine/Source/Runtime/Engine/Private/UserInterface/InputSettings.cpp
--- a/Engine/Source/Runtime/Engine/Private/UserInterface/InputSettings.cpp
+++ b/Engine/Source/Runtime/Engine/Private/UserInterface/InputSettings.cpp
@@ -13,7 +13,9 @@
 
 UInputSettings::UInputSettings(const FObjectInitializer& ObjectInitializer)
 	: Super(ObjectInitializer)
+	, bCaptureMouseOnLaunch(true)
 	, DefaultViewportMouseCaptureMode(EMouseCaptureMode::CapturePermanently_IncludingInitialMouseDown)
+	, bDefaultViewportMouseLock(true)
 {
 }
 
diff --git a/Engine/Source/Runtime/Engine/Public/UnrealClient.h b/Engine/Source/Runtime/Engine/Public/UnrealClient.h
--- a/Engine/Source/Runtime/Engine/Public/UnrealClient.h
+++ b/Engine/Source/Runtime/Engine/Public/UnrealClient.h
@@ -894,6 +894,19 @@ public:
 	 */
 	virtual EMouseCaptureMode CaptureMouseOnClick() { return EMouseCaptureMode::CapturePermanently; }
 
+	/**
+	 * Gets whether or not the viewport captures the Mouse on launch of the application
+	 * Technically this controls capture on the first window activate, so in situations
+	 * where the application is launched but isn't activated the effect is delayed until
+	 * activation.
+	 */
+	virtual bool CaptureMouseOnLaunch() { return true; }
+
+	/**
+	 * Gets whether or not the cursor is locked to the viewport when the viewport captures the mouse
+	 */
+	virtual bool LockDuringCapture() { return true; }
+
 	/**
 	 * Gets whether or not the cursor is hidden when the viewport captures the mouse
 	 */
