#ue4 - Don't reject low distance MTD, it could cause us to not process some valid overlaps.

This was missed in 4.11->Main merge. Relevant 4.11 CLs were 2918756, 2921814.

diff --git a/Engine/Source/Runtime/Engine/Private/PhysicsEngine/BodyInstance.cpp b/Engine/Source/Runtime/Engine/Private/PhysicsEngine/BodyInstance.cpp
--- a/Engine/Source/Runtime/Engine/Private/PhysicsEngine/BodyInstance.cpp
+++ b/Engine/Source/Runtime/Engine/Private/PhysicsEngine/BodyInstance.cpp
@@ -4486,8 +4486,8 @@ bool FBodyInstance::OverlapPhysX_AssumesLocked(const PxGeometry& PGeom, const Px
 		{
 			PxVec3 POutDirection;
 			float OutDistance;
-			//TODO: this 0 check is a temp fix because physx is reporting contacts with distance 0 which gives us nan
-			if (PxGeometryQuery::computePenetration(POutDirection, OutDistance, PGeom, ShapePose, PShape->getGeometry().any(), PxShapeExt::getGlobalPose(*PShape, *RigidBody)) && !FMath::IsNearlyZero(OutDistance))
+
+			if (PxGeometryQuery::computePenetration(POutDirection, OutDistance, PGeom, ShapePose, PShape->getGeometry().any(), PxShapeExt::getGlobalPose(*PShape, *RigidBody)))
 			{				
 				//TODO: there are some edge cases that give us nan results. In these cases we skip
 				if (!POutDirection.isFinite())
