Copying //UE4/Release-Staging-4.12 to //UE4/Dev-Main (Source: //UE4/Release-4.12 @ 2955635)

diff --git a/Engine/Source/Runtime/CoreUObject/Public/Internationalization/PackageLocalizationManager.h b/Engine/Source/Runtime/CoreUObject/Public/Internationalization/PackageLocalizationManager.h
--- a/Engine/Source/Runtime/CoreUObject/Public/Internationalization/PackageLocalizationManager.h
+++ b/Engine/Source/Runtime/CoreUObject/Public/Internationalization/PackageLocalizationManager.h
@@ -8,6 +8,20 @@ class IPackageLocalizationCache;
 class COREUOBJECT_API FPackageLocalizationManager
 {
 public:
+	typedef TFunction<void(FPackageLocalizationManager&)> FLazyInitFunc;
+
+	/**
+	 * Initialize the manager from the callback set by InitializeFromLazyCallback. It is expected that this callback calls one of the InitializeFromX functions.
+	 */
+	void PerformLazyInitialization();
+	
+	/**
+	 * Initialize the manager lazily using the given callback. It is expected that this callback calls one of the InitializeFromX functions.
+	 *
+	 * @param InLazyInitFunc	The function to call to initialize the manager.
+	 */
+	void InitializeFromLazyCallback(FLazyInitFunc InLazyInitFunc);
+
 	/**
 	 * Initialize the manager using the given cache. This will perform an initial scan for localized packages.
 	 *
@@ -27,7 +41,7 @@ public:
 	 *
 	 * @return The localized package name, or NAME_None if there is no localized package.
 	 */
-	FName FindLocalizedPackageName(const FName InSourcePackageName) const;
+	FName FindLocalizedPackageName(const FName InSourcePackageName);
 
 	/**
 	 * Try and find the localized package name for the given source package for the given culture.
@@ -37,7 +51,7 @@ public:
 	 *
 	 * @return The localized package name, or NAME_None if there is no localized package.
 	 */
-	FName FindLocalizedPackageNameForCulture(const FName InSourcePackageName, const FString& InCultureName) const;
+	FName FindLocalizedPackageNameForCulture(const FName InSourcePackageName, const FString& InCultureName);
 
 	/**
 	 * Singleton accessor.
@@ -57,6 +71,9 @@ private:
 	 */
 	FName FindLocalizedPackageNameNoCache(const FName InSourcePackageName, const FString& InCultureName) const;
 
+	/** Function to call to lazily initialize the manager. */
+	FLazyInitFunc LazyInitFunc;
+
 	/** Pointer to our currently active cache. Only valid after Initialize has been called. */
 	TSharedPtr<IPackageLocalizationCache> ActiveCache;
 };
diff --git a/Engine/Source/Runtime/CoreUObject/Public/Internationalization/PackageLocalizationUtil.h b/Engine/Source/Runtime/CoreUObject/Public/Internationalization/PackageLocalizationUtil.h
new file mode 100644
--- /dev/null
+++ b/Engine/Source/Runtime/CoreUObject/Public/Internationalization/PackageLocalizationUtil.h
@@ -0,0 +1,41 @@
+// Copyright 1998-2016 Epic Games, Inc. All Rights Reserved.
+
+#pragma once
+
+/** Utility functions for dealing with localized package names */
+struct COREUOBJECT_API FPackageLocalizationUtil
+{
+	/**
+	 * Converts a localized version of a package path to the source version (by removing /L10N/<code> from the package path, if present)
+	 * Note: This does not check whether the source package exists
+	 *
+	 * @param InLocalized	Path to the localized package.
+	 * @param OutSource		Path to the source package.
+	 *
+	 * @returns True if the conversion happened, false otherwise
+	 */
+	static bool ConvertLocalizedToSource(const FString& InLocalized, FString& OutSource);
+
+	/**
+	 * Converts a source version of a package path to the localized version for the given culture (by adding /L10N/<code> to the package path)
+	 * Note: This does not check whether the source package exists
+	 *
+	 * @param InSource		Path to the source package.
+	 * @param InCulture		Culture code to use.
+	 * @param OutLocalized	Path to the localized package.
+	 *
+	 * @returns True if the conversion happened, false otherwise
+	 */
+	static bool ConvertSourceToLocalized(const FString& InSource, const FString& InCulture, FString& OutLocalized);
+
+	/**
+	 * Given a package path, get the localized root package for the given culture (eg, if given "/Game/MyFolder/MyAsset" and a culture of "fr", this would return "/Game/L10N/fr")
+	 *
+	 * @param InPath		Package path to use.
+	 * @param InCulture		Culture code to use (if any).
+	 * @param OutLocalized	Localized package root.
+	 *
+	 * @returns True if the conversion happened, false otherwise
+	 */
+	static bool GetLocalizedRoot(const FString& InPath, const FString& InCulture, FString& OutLocalized);
+};
diff --git a/Engine/Source/Runtime/Engine/Classes/Animation/AnimComposite.h b/Engine/Source/Runtime/Engine/Classes/Animation/AnimComposite.h
--- a/Engine/Source/Runtime/Engine/Classes/Animation/AnimComposite.h
+++ b/Engine/Source/Runtime/Engine/Classes/Animation/AnimComposite.h
@@ -32,6 +32,8 @@ public:
 	virtual bool IsValidAdditive() const override { return GetAdditiveAnimType() != AAT_None; }
 	virtual void EnableRootMotionSettingFromMontage(bool bInEnableRootMotion, const ERootMotionRootLock::Type InRootMotionRootLock) override;
 	virtual bool HasRootMotion() const override;
+	virtual void GetAnimNotifiesFromDeltaPositions(const float& PreviousPosition, const float & CurrentPosition, TArray<const FAnimNotifyEvent *>& OutActiveNotifies) const override;
+	virtual bool IsNotifyAvailable() const override;
 	//~ End UAnimSequenceBase Interface
 	//~ Begin UAnimSequence Interface
 #if WITH_EDITOR
diff --git a/Engine/Source/Runtime/Engine/Classes/Animation/AnimCompositeBase.h b/Engine/Source/Runtime/Engine/Classes/Animation/AnimCompositeBase.h
--- a/Engine/Source/Runtime/Engine/Classes/Animation/AnimCompositeBase.h
+++ b/Engine/Source/Runtime/Engine/Classes/Animation/AnimCompositeBase.h
@@ -159,6 +159,11 @@ struct FAnimSegment
 	 * return true if valid, false otherwise, only invalid if we contains recursive reference
 	 **/
 	bool IsValid() const { return bValid;  }
+
+	/** 
+	 * return true if anim notify is available 
+	 */
+	bool IsNotifyAvailable() const { return IsValid() && AnimReference && AnimReference->IsNotifyAvailable(); }
 private:
 
 	/**
@@ -249,6 +254,8 @@ struct FAnimTrack
 	*/
 	void GetAnimNotifiesFromTrackPositions(const float& PreviousTrackPosition, const float& CurrentTrackPosition, TArray<const FAnimNotifyEvent *> & OutActiveNotifies) const;
 
+	/** return true if anim notify is available */
+	bool IsNotifyAvailable() const;
 };
 
 UCLASS(abstract, MinimalAPI)
diff --git a/Engine/Source/Runtime/Engine/Classes/Animation/AnimInstance.h b/Engine/Source/Runtime/Engine/Classes/Animation/AnimInstance.h
--- a/Engine/Source/Runtime/Engine/Classes/Animation/AnimInstance.h
+++ b/Engine/Source/Runtime/Engine/Classes/Animation/AnimInstance.h
@@ -1022,6 +1022,9 @@ public:
 	/** Called after updates are completed, dispatches notifies etc. */
 	void PostUpdateAnimation();
 
+	/** Called on the game thread pre-evaluation. */
+	void PreEvaluateAnimation();
+
 	DEPRECATED(4.11, "This function should no longer be used. Use ParallelEvaluateAnimation")
 	void EvaluateAnimation(struct FPoseContext& Output);
 
diff --git a/Engine/Source/Runtime/Engine/Classes/Animation/AnimMontage.h b/Engine/Source/Runtime/Engine/Classes/Animation/AnimMontage.h
--- a/Engine/Source/Runtime/Engine/Classes/Animation/AnimMontage.h
+++ b/Engine/Source/Runtime/Engine/Classes/Animation/AnimMontage.h
@@ -35,7 +35,6 @@ struct FCompositeSection : public FAnimLinkableElement
 	UPROPERTY()
 	FName NextSectionName;
 
-private:
 	/** Meta data that can be saved with the asset 
 	 * 
 	 * You can query by GetMetaData function
diff --git a/Engine/Source/Runtime/Engine/Classes/Animation/AnimSequenceBase.h b/Engine/Source/Runtime/Engine/Classes/Animation/AnimSequenceBase.h
--- a/Engine/Source/Runtime/Engine/Classes/Animation/AnimSequenceBase.h
+++ b/Engine/Source/Runtime/Engine/Classes/Animation/AnimSequenceBase.h
@@ -78,7 +78,7 @@ class UAnimSequenceBase : public UAnimationAsset
 	 * Supports playing backwards (CurrentPosition<PreviousPosition).
 	 * Only supports contiguous range, does NOT support looping and wrapping over.
 	 */
-	ENGINE_API void GetAnimNotifiesFromDeltaPositions(const float& PreviousPosition, const float & CurrentPosition, TArray<const FAnimNotifyEvent *>& OutActiveNotifies) const;
+	ENGINE_API virtual void GetAnimNotifiesFromDeltaPositions(const float& PreviousPosition, const float & CurrentPosition, TArray<const FAnimNotifyEvent *>& OutActiveNotifies) const;
 
 	/** Evaluate curve data to Instance at the time of CurrentTime **/
 	ENGINE_API virtual void EvaluateCurveData(FBlendedCurve& OutCurve, float CurrentTime, bool bForceUseRawData=false) const;
@@ -181,6 +181,8 @@ public:
 	virtual class UAnimSequence* GetAdditiveBasePose() const { return nullptr; }
 
 #endif
+	// return true if anim notify is available 
+	virtual bool IsNotifyAvailable() const;
 
 protected:
 	template <typename DataType>
diff --git a/Engine/Source/Runtime/Engine/Classes/Animation/AnimSingleNodeInstance.h b/Engine/Source/Runtime/Engine/Classes/Animation/AnimSingleNodeInstance.h
--- a/Engine/Source/Runtime/Engine/Classes/Animation/AnimSingleNodeInstance.h
+++ b/Engine/Source/Runtime/Engine/Classes/Animation/AnimSingleNodeInstance.h
@@ -23,10 +23,10 @@ class ENGINE_API UAnimSingleNodeInstance : public UAnimInstance
 	PRAGMA_ENABLE_DEPRECATION_WARNINGS
 
 	/** Current Asset being played **/
-	DEPRECATED(4.11, "Please use FAnimSingleNodeInstanceProxy::CurrentAsset")
+	UPROPERTY(Transient)
 	class UAnimationAsset* CurrentAsset;
 
-	DEPRECATED(4.11, "Please use FAnimSingleNodeInstanceProxy::CurrentVertexAnim")
+	UPROPERTY(Transient)
 	class UVertexAnimation* CurrentVertexAnim;
 
 	/** Random cached values to play each asset **/
diff --git a/Engine/Source/Runtime/Engine/Classes/Components/ChildActorComponent.h b/Engine/Source/Runtime/Engine/Classes/Components/ChildActorComponent.h
--- a/Engine/Source/Runtime/Engine/Classes/Components/ChildActorComponent.h
+++ b/Engine/Source/Runtime/Engine/Classes/Components/ChildActorComponent.h
@@ -50,7 +50,7 @@ private:
 	TSubclassOf<AActor>	ChildActorClass;
 
 	/** The actor that we spawned and own */
-	UPROPERTY(BlueprintReadOnly, Category=ChildActorComponent, TextExportTransient, NonPIEDuplicateTransient, meta=(AllowPrivateAccess="true"))
+	UPROPERTY(Replicated, BlueprintReadOnly, Category=ChildActorComponent, TextExportTransient, NonPIEDuplicateTransient, meta=(AllowPrivateAccess="true"))
 	AActor*	ChildActor;
 
 	/** We try to keep the child actor's name as best we can, so we store it off here when destroying */
@@ -68,6 +68,8 @@ public:
 	virtual void PostLoad() override;
 #endif
 	virtual void BeginDestroy() override;
+	virtual void GetLifetimeReplicatedProps(TArray<FLifetimeProperty>& OutLifetimeProps) const override;
+	virtual void PostRepNotifies() override;
 	//~ End Object Interface.
 
 	//~ Begin ActorComponent Interface.
@@ -89,7 +91,7 @@ public:
 	FName GetChildActorName() const { return ChildActorName; }
 
 	/** Kill any currently present child actor */
-	void DestroyChildActor(const bool bRequiresRename = true);
+	void DestroyChildActor();
 };
 
 
diff --git a/Engine/Source/Runtime/Engine/Classes/Components/SkeletalMeshComponent.h b/Engine/Source/Runtime/Engine/Classes/Components/SkeletalMeshComponent.h
--- a/Engine/Source/Runtime/Engine/Classes/Components/SkeletalMeshComponent.h
+++ b/Engine/Source/Runtime/Engine/Classes/Components/SkeletalMeshComponent.h
@@ -940,6 +940,8 @@ protected:
 	
 	virtual bool ComponentOverlapComponentImpl(class UPrimitiveComponent* PrimComp, const FVector Pos, const FQuat& Quat, const FCollisionQueryParams& Params) override;
 
+	virtual bool MoveComponentImpl(const FVector& Delta, const FQuat& NewRotation, bool bSweep, FHitResult* OutHit = NULL, EMoveComponentFlags MoveFlags = MOVECOMP_NoFlags, ETeleportType Teleport = ETeleportType::None) override;
+
 public:
 
 	virtual class UBodySetup* GetBodySetup() override;
diff --git a/Engine/Source/Runtime/Engine/Classes/Components/SplineComponent.h b/Engine/Source/Runtime/Engine/Classes/Components/SplineComponent.h
--- a/Engine/Source/Runtime/Engine/Classes/Components/SplineComponent.h
+++ b/Engine/Source/Runtime/Engine/Classes/Components/SplineComponent.h
@@ -76,11 +76,16 @@ class ENGINE_API USplineComponent : public UPrimitiveComponent
 	UPROPERTY(EditAnywhere, Category = Spline, meta=(DisplayName="Override Construction Script"))
 	bool bSplineHasBeenEdited;
 
+	/**
+	 * Whether the spline points should be passed to the User Construction Script so they can be further manipulated by it.
+	 * If false, they will not be visible to it, and it will not be able to influence the per-instance positions set in the editor.
+	 */
 	UPROPERTY(EditAnywhere, BlueprintReadWrite, Category = Spline)
 	bool bInputSplinePointsToConstructionScript;
 
+	/** If true, the spline will be rendered if the Splines showflag is set. */
 	UPROPERTY(EditAnywhere, BlueprintReadWrite, Category = Spline)
-	bool bAlwaysRenderInEditor;
+	bool bDrawDebug;
 
 private:
 	/**
@@ -94,16 +99,16 @@ public:
 	/** Default up vector in local space to be used when calculating transforms along the spline */
 	UPROPERTY(EditAnywhere, BlueprintReadWrite, Category = Spline)
 	FVector DefaultUpVector;
-
-#if WITH_EDITORONLY_DATA
+#if !UE_BUILD_SHIPPING
 	/** Color of an unselected spline component segment in the editor */
 	UPROPERTY(EditAnywhere, Category = Editor)
 	FLinearColor EditorUnselectedSplineSegmentColor;
-
 	/** Color of a selected spline component segment in the editor */
 	UPROPERTY(EditAnywhere, Category = Editor)
 	FLinearColor EditorSelectedSplineSegmentColor;
+#endif
 
+#if WITH_EDITORONLY_DATA
 	/** Whether scale visualization should be displayed */
 	UPROPERTY(EditAnywhere, Category = Editor)
 	bool bShouldVisualizeScale;
@@ -126,7 +131,7 @@ public:
 	virtual FActorComponentInstanceData* GetComponentInstanceData() const override;
 	//~ End UActorComponent Interface.
 
-#if WITH_EDITOR
+#if !UE_BUILD_SHIPPING
 	//~ Begin UPrimitiveComponent Interface.
 	virtual FPrimitiveSceneProxy* CreateSceneProxy() override;
 	//~ End UPrimitiveComponent Interface.
@@ -183,6 +188,10 @@ public:
 	UFUNCTION(BlueprintCallable, Category = Editor)
 	void SetSelectedSplineSegmentColor(const FLinearColor& SegmentColor);
 
+	/** Specify whether this spline should be rendered when the Editor/Game spline show flag is set */
+	UFUNCTION(BlueprintCallable, Category = Spline)
+	void SetDrawDebug(bool bShow);
+
 	/** Specify whether the spline is a closed loop or not */
 	UFUNCTION(BlueprintCallable, Category = Spline)
 	void SetClosedLoop(bool bInClosedLoop);
diff --git a/Engine/Source/Runtime/Engine/Classes/Engine/BlueprintGeneratedClass.h b/Engine/Source/Runtime/Engine/Classes/Engine/BlueprintGeneratedClass.h
--- a/Engine/Source/Runtime/Engine/Classes/Engine/BlueprintGeneratedClass.h
+++ b/Engine/Source/Runtime/Engine/Classes/Engine/BlueprintGeneratedClass.h
@@ -696,7 +696,7 @@ protected:
 	void BuildCustomPropertyListForPostConstruction(FCustomPropertyListNode*& InPropertyList, UStruct* InStruct, const uint8* DataPtr, const uint8* DefaultDataPtr);
 
 	/** Internal helper method used to recursively build a custom property list from an array property used for post-construct initialization. */
-	void BuildCustomArrayPropertyListForPostConstruction(UArrayProperty* ArrayProperty, FCustomPropertyListNode*& InPropertyList, const uint8* DataPtr, const uint8* DefaultDataPtr);
+	void BuildCustomArrayPropertyListForPostConstruction(UArrayProperty* ArrayProperty, FCustomPropertyListNode*& InPropertyList, const uint8* DataPtr, const uint8* DefaultDataPtr, int32 StartIndex = 0);
 
 private:
 	/** List of native class-owned properties that differ from defaults. This is used to optimize property initialization during post-construction by minimizing the number of native class-owned property values that get copied to the new instance. */
diff --git a/Engine/Source/Runtime/Engine/Classes/Engine/DemoNetDriver.h b/Engine/Source/Runtime/Engine/Classes/Engine/DemoNetDriver.h
--- a/Engine/Source/Runtime/Engine/Classes/Engine/DemoNetDriver.h
+++ b/Engine/Source/Runtime/Engine/Classes/Engine/DemoNetDriver.h
@@ -245,4 +245,7 @@ public:
 
 	/** Called when a "go to time" operation is completed. */
 	void NotifyGotoTimeFinished(bool bWasSuccessful);
+
+	/** Read the streaming level information from the metadata after the level is loaded */
+	void PendingNetGameLoadMapCompleted();
 };
diff --git a/Engine/Source/Runtime/Engine/Classes/Engine/DemoPendingNetGame.h b/Engine/Source/Runtime/Engine/Classes/Engine/DemoPendingNetGame.h
--- a/Engine/Source/Runtime/Engine/Classes/Engine/DemoPendingNetGame.h
+++ b/Engine/Source/Runtime/Engine/Classes/Engine/DemoPendingNetGame.h
@@ -4,10 +4,17 @@
 
 #include "DemoPendingNetGame.generated.h"
 
+class UEngine;
+struct FWorldContext;
 
 UCLASS(transient, config=Engine)
 class UDemoPendingNetGame
 	: public UPendingNetGame
 {
 	GENERATED_UCLASS_BODY()
+
+	// UPendingNetGame interface.
+	virtual void Tick( float DeltaTime ) override;
+	virtual void SendJoin() override;
+	virtual void LoadMapCompleted(UEngine* Engine, FWorldContext& Context, bool bLoadedMapSuccessfully, const FString& LoadMapError) override;
 };
diff --git a/Engine/Source/Runtime/Engine/Classes/Engine/GameViewportClient.h b/Engine/Source/Runtime/Engine/Classes/Engine/GameViewportClient.h
--- a/Engine/Source/Runtime/Engine/Classes/Engine/GameViewportClient.h
+++ b/Engine/Source/Runtime/Engine/Classes/Engine/GameViewportClient.h
@@ -810,9 +810,6 @@ private:
 	/** Those sound stat flags which are enabled on this viewport */
 	static ESoundShowFlags::Type SoundShowFlags;
 
-	/** Number of viewports which have enabled 'show collision' */
-	static int32 NumViewportsShowingCollision;
-
 	/** Disables splitscreen, useful when game code is in menus, and doesn't want splitscreen on */
 	bool bDisableSplitScreenOverride;
 
@@ -828,9 +825,6 @@ private:
 	/** Whether or not the cursor is hidden when the viewport captures the mouse */
 	bool bHideCursorDuringCapture;
 
-	/** Handle to the registered ShowCollisionOnSpawnedActors delegate */
-	FDelegateHandle ShowCollisionOnSpawnedActorsDelegateHandle;
-
 	/** Handle to the audio device created for this viewport. Each viewport (for multiple PIE) will have its own audio device. */
 	uint32 AudioDeviceHandle;
 
diff --git a/Engine/Source/Runtime/Engine/Classes/Engine/PendingNetGame.h b/Engine/Source/Runtime/Engine/Classes/Engine/PendingNetGame.h
--- a/Engine/Source/Runtime/Engine/Classes/Engine/PendingNetGame.h
+++ b/Engine/Source/Runtime/Engine/Classes/Engine/PendingNetGame.h
@@ -3,6 +3,9 @@
 #pragma once
 #include "PendingNetGame.generated.h"
 
+struct FWorldContext;
+class UEngine;
+
 /**
  * Accepting connection response codes
  */
@@ -135,4 +138,7 @@ public:
 
 	/** Create the peer net driver and a socket to listen for new client peer connections. */
 	void InitPeerListen();
+
+	/** Called by the engine after it calls LoadMap for this PendingNetGame. */
+	virtual void LoadMapCompleted(UEngine* Engine, FWorldContext& Context, bool bLoadedMapSuccessfully, const FString& LoadMapError);
 };
diff --git a/Engine/Source/Runtime/Engine/Classes/Engine/RendererSettings.h b/Engine/Source/Runtime/Engine/Classes/Engine/RendererSettings.h
--- a/Engine/Source/Runtime/Engine/Classes/Engine/RendererSettings.h
+++ b/Engine/Source/Runtime/Engine/Classes/Engine/RendererSettings.h
@@ -164,6 +164,12 @@ class ENGINE_API URendererSettings : public UDeveloperSettings
 		ConfigRestartRequired=true))
 	uint32 bUseDXT5NormalMaps:1;
 
+	UPROPERTY(config, EditAnywhere, Category = Materials, meta =(
+		ConfigRestartRequired = true,
+		ConsoleVariable = "r.ClearCoatNormal",
+		ToolTip = "Use a separate normal map for the bottom layer of a clear coat material. This is a higher quality feature that is expensive."))
+		uint32 bClearCoatEnableSecondNormal : 1;
+
 	UPROPERTY(config, EditAnywhere, Category = Textures, meta = (
 		ConsoleVariable = "r.ReflectionCaptureResolution", DisplayName = "Reflection Capture Resolution",
 		ToolTip = "The cubemap resolution for all reflection capture probes. Must be power of 2."))
diff --git a/Engine/Source/Runtime/Engine/Classes/GameFramework/Actor.h b/Engine/Source/Runtime/Engine/Classes/GameFramework/Actor.h
--- a/Engine/Source/Runtime/Engine/Classes/GameFramework/Actor.h
+++ b/Engine/Source/Runtime/Engine/Classes/GameFramework/Actor.h
@@ -2322,9 +2322,6 @@ public:
 	 */
 	void ExecuteConstruction(const FTransform& Transform, const class FComponentInstanceDataCache* InstanceDataCache, bool bIsDefaultTransform = false);
 
-	// Forces PostLoad call on natively created components, that were originally in SCS
-	virtual void FixComponentsFromDynamicClass();
-
 	/**
 	 * Called when an instance of this class is placed (in editor) or spawned.
 	 * @param	Transform			The transform the actor was constructed at.
@@ -2718,6 +2715,12 @@ public:
 	 */
 	void UpdateAllReplicatedComponents();
 
+	/** Returns whether replication is enabled or not. */
+	FORCEINLINE bool GetIsReplicated() const
+	{
+		return bReplicates;
+	}
+
 	/** Returns a constant reference to the replicated components set
 	 */
 	const TSet<UActorComponent*>& GetReplicatedComponents() const 
diff --git a/Engine/Source/Runtime/Engine/Classes/Materials/MaterialExpressionClearCoatNormalCustomOutput.h b/Engine/Source/Runtime/Engine/Classes/Materials/MaterialExpressionClearCoatNormalCustomOutput.h
new file mode 100644
--- /dev/null
+++ b/Engine/Source/Runtime/Engine/Classes/Materials/MaterialExpressionClearCoatNormalCustomOutput.h
@@ -0,0 +1,39 @@
+// Copyright 1998-2015 Epic Games, Inc. All Rights Reserved.
+
+
+#pragma once
+
+#include "Materials/MaterialExpressionCustomOutput.h"
+#include "MaterialExpressionClearCoatNormalCustomOutput.generated.h"
+
+
+
+UCLASS()
+class UMaterialExpressionClearCoatNormalCustomOutput : public UMaterialExpressionCustomOutput
+{
+	GENERATED_UCLASS_BODY()
+
+	UPROPERTY(meta = (RequiredInput = "true"))
+	FExpressionInput Input;
+
+#if WITH_EDITOR
+	virtual int32 Compile(class FMaterialCompiler* Compiler, int32 OutputIndex, int32 MultiplexIndex) override;
+	virtual void GetCaption(TArray<FString>& OutCaptions) const override;
+#endif
+
+	virtual FExpressionInput* GetInput(int32 InputIndex) override;
+
+
+	// Begin UObject Interface
+#if WITH_EDITOR
+
+	virtual uint32 GetInputType(int32 InputIndex) override { return MCT_Float3; }
+#endif
+
+
+	virtual FString GetFunctionName() const override { return TEXT("ClearCoatBottomNormal"); }
+
+};
+
+
+
diff --git a/Engine/Source/Runtime/Engine/Private/Actor.cpp b/Engine/Source/Runtime/Engine/Private/Actor.cpp
--- a/Engine/Source/Runtime/Engine/Private/Actor.cpp
+++ b/Engine/Source/Runtime/Engine/Private/Actor.cpp
@@ -496,6 +496,26 @@ void AActor::Serialize(FArchive& Ar)
 			}
 		}
 	}
+
+	// When duplicating for PIE all components need to be gathered up and duplicated even if there are no other property references to them
+	// otherwise we can end up with Attach Parents that do not get redirected to the correct component. However, if there is a transient component
+	// we'll let that drop
+	if (Ar.GetPortFlags() & PPF_DuplicateForPIE)
+	{
+		TArray<UActorComponent*> DuplicatingComponents;
+		if (Ar.IsSaving())
+		{
+			DuplicatingComponents.Reserve(OwnedComponents.Num());
+			for (UActorComponent* OwnedComponent : OwnedComponents)
+			{
+				if (!OwnedComponent->HasAnyFlags(RF_Transient))
+				{
+					DuplicatingComponents.Add(OwnedComponent);
+				}
+			}
+		}
+		Ar << DuplicatingComponents;
+	}
 #endif
 
 	Super::Serialize(Ar);
diff --git a/Engine/Source/Runtime/Engine/Private/ActorConstruction.cpp b/Engine/Source/Runtime/Engine/Private/ActorConstruction.cpp
--- a/Engine/Source/Runtime/Engine/Private/ActorConstruction.cpp
+++ b/Engine/Source/Runtime/Engine/Private/ActorConstruction.cpp
@@ -459,18 +459,11 @@ void AActor::RerunConstructionScripts()
 	}
 }
 
-void AActor::FixComponentsFromDynamicClass()
-{
-
-}
-
 void AActor::ExecuteConstruction(const FTransform& Transform, const FComponentInstanceDataCache* InstanceDataCache, bool bIsDefaultTransform)
 {
 	check(!IsPendingKill());
 	check(!HasAnyFlags(RF_BeginDestroyed|RF_FinishDestroyed));
 
-	FixComponentsFromDynamicClass();
-
 	// ensure that any existing native root component gets this new transform
 	// we can skip this in the default case as the given transform will be the root component's transform
 	if (RootComponent && !bIsDefaultTransform)
diff --git a/Engine/Source/Runtime/Engine/Private/Animation/AnimComposite.cpp b/Engine/Source/Runtime/Engine/Private/Animation/AnimComposite.cpp
--- a/Engine/Source/Runtime/Engine/Private/Animation/AnimComposite.cpp
+++ b/Engine/Source/Runtime/Engine/Private/Animation/AnimComposite.cpp
@@ -28,46 +28,47 @@ void UAnimComposite::ReplaceReferredAnimations(const TMap<UAnimationAsset*, UAni
 }
 #endif
 
-void UAnimComposite::HandleAssetPlayerTickedInternal(FAnimAssetTickContext &Context, const float PreviousTime, const float MoveDelta, const FAnimTickRecord &Instance, struct FAnimNotifyQueue& NotifyQueue) const
+bool UAnimComposite::IsNotifyAvailable() const
 {
-	Super::HandleAssetPlayerTickedInternal(Context, PreviousTime, MoveDelta, Instance, NotifyQueue);
+	return (Super::IsNotifyAvailable() || AnimationTrack.IsNotifyAvailable());
+}
 
-	// if should generate notifies
-	if (Context.ShouldGenerateNotifies())
-	{
-		// make sure to check animation track notifies
-		TArray<const FAnimNotifyEvent*> AnimNotifies;
-		const bool bMovingForward = (RateScale >= 0.f);
+void UAnimComposite::GetAnimNotifiesFromDeltaPositions(const float& PreviousPosition, const float & CurrentPosition, TArray<const FAnimNotifyEvent *>& OutActiveNotifies) const
+{
+	const bool bMovingForward = (RateScale >= 0.f);
+
+	Super::GetAnimNotifiesFromDeltaPositions(PreviousPosition, CurrentPosition, OutActiveNotifies);
 
-		if (bMovingForward)
+	if (bMovingForward)
+	{
+		if (PreviousPosition <= CurrentPosition)
 		{
-			const float CurrentTime = *Instance.TimeAccumulator;
-			if (PreviousTime <= CurrentTime)
-			{
-				AnimationTrack.GetAnimNotifiesFromTrackPositions(PreviousTime, CurrentTime, AnimNotifies);
-			}
-			else
-			{
-				AnimationTrack.GetAnimNotifiesFromTrackPositions(PreviousTime, SequenceLength, AnimNotifies);
-				AnimationTrack.GetAnimNotifiesFromTrackPositions(0.f, CurrentTime, AnimNotifies);
-			}
+			AnimationTrack.GetAnimNotifiesFromTrackPositions(PreviousPosition, CurrentPosition, OutActiveNotifies);
 		}
 		else
 		{
-			const float CurrentTime = *Instance.TimeAccumulator;
-			if (PreviousTime >= CurrentTime)
-			{
-				AnimationTrack.GetAnimNotifiesFromTrackPositions(PreviousTime, CurrentTime, AnimNotifies);
-			}
-			else
-			{
-				AnimationTrack.GetAnimNotifiesFromTrackPositions(PreviousTime, 0.f, AnimNotifies);
-				AnimationTrack.GetAnimNotifiesFromTrackPositions(SequenceLength, CurrentTime, AnimNotifies);
-			}
+			AnimationTrack.GetAnimNotifiesFromTrackPositions(PreviousPosition, SequenceLength, OutActiveNotifies);
+			AnimationTrack.GetAnimNotifiesFromTrackPositions(0.f, CurrentPosition, OutActiveNotifies);
 		}
-
-		NotifyQueue.AddAnimNotifies(AnimNotifies, Instance.EffectiveBlendWeight);
 	}
+	else
+	{
+		if (PreviousPosition >= CurrentPosition)
+		{
+			AnimationTrack.GetAnimNotifiesFromTrackPositions(PreviousPosition, CurrentPosition, OutActiveNotifies);
+		}
+		else
+		{
+			AnimationTrack.GetAnimNotifiesFromTrackPositions(PreviousPosition, 0.f, OutActiveNotifies);
+			AnimationTrack.GetAnimNotifiesFromTrackPositions(SequenceLength, CurrentPosition, OutActiveNotifies);
+		}
+	}
+}
+
+void UAnimComposite::HandleAssetPlayerTickedInternal(FAnimAssetTickContext &Context, const float PreviousTime, const float MoveDelta, const FAnimTickRecord &Instance, struct FAnimNotifyQueue& NotifyQueue) const
+{
+	Super::HandleAssetPlayerTickedInternal(Context, PreviousTime, MoveDelta, Instance, NotifyQueue);
+
 	ExtractRootMotionFromTrack(AnimationTrack, PreviousTime, PreviousTime + MoveDelta, Context.RootMotionMovementParams);
 }
 
diff --git a/Engine/Source/Runtime/Engine/Private/Animation/AnimCompositeBase.cpp b/Engine/Source/Runtime/Engine/Private/Animation/AnimCompositeBase.cpp
--- a/Engine/Source/Runtime/Engine/Private/Animation/AnimCompositeBase.cpp
+++ b/Engine/Source/Runtime/Engine/Private/Animation/AnimCompositeBase.cpp
@@ -630,6 +630,10 @@ void FAnimTrack::InvalidateRecursiveAsset(class UAnimCompositeBase* CheckAsset)
 				AnimSegment.bValid = IsValidToAdd(CompositeBase);
 			}
 		}
+		else
+		{
+			AnimSegment.bValid = IsValidToAdd(AnimSegment.AnimReference);
+		}
 	}
 }
 
@@ -665,6 +669,19 @@ void FAnimTrack::GetAnimNotifiesFromTrackPositions(const float& PreviousTrackPos
 	}
 }
 
+bool FAnimTrack::IsNotifyAvailable() const
+{
+	for (int32 SegmentIndex = 0; SegmentIndex < AnimSegments.Num(); ++SegmentIndex)
+	{
+		if (AnimSegments[SegmentIndex].IsNotifyAvailable())
+		{
+			return true;
+		}
+	}
+
+	return false;
+}
+
 bool FAnimTrack::IsValidToAdd(const UAnimSequenceBase* SequenceBase) const
 {
 	bool bValid = false;
diff --git a/Engine/Source/Runtime/Engine/Private/Animation/AnimInstance.cpp b/Engine/Source/Runtime/Engine/Private/Animation/AnimInstance.cpp
--- a/Engine/Source/Runtime/Engine/Private/Animation/AnimInstance.cpp
+++ b/Engine/Source/Runtime/Engine/Private/Animation/AnimInstance.cpp
@@ -389,22 +389,12 @@ PRAGMA_ENABLE_DEPRECATION_WARNINGS
 	// so that node knows where montage is
 	UpdateMontage(DeltaSeconds);
 
-	if(NeedsImmediateUpdate(DeltaSeconds))
+	if(bNeedsValidRootMotion || NeedsImmediateUpdate(DeltaSeconds))
 	{
 		// cant use parallel update, so just do the work here
 		Proxy.UpdateAnimation();
 		PostUpdateAnimation();
 	}
-	else if(bNeedsValidRootMotion && RootMotionMode == ERootMotionMode::RootMotionFromMontagesOnly)
-	{
-		// We may have just partially blended root motion, so make it up to 1 by
-		// blending in identity too
-		FRootMotionMovementParams& ExtractedRootMotion = Proxy.GetExtractedRootMotion();
-		if (ExtractedRootMotion.bHasRootMotion)
-		{
-			ExtractedRootMotion.MakeUpToFullWeight();
-		}
-	}
 }
 
 void UAnimInstance::PreUpdateAnimation(float DeltaSeconds)
@@ -495,7 +485,6 @@ bool UAnimInstance::NeedsImmediateUpdate(float DeltaSeconds) const
 		bEvaluationPhaseSkipped || 
 		CVarUseParallelAnimUpdate.GetValueOnGameThread() == 0 ||
 		CVarUseParallelAnimationEvaluation.GetValueOnGameThread() == 0 ||
-		GetWorld()->IsServer() ||
 		!bUseParallelUpdateAnimation ||
 		DeltaSeconds == 0.0f ||
 		RootMotionMode == ERootMotionMode::RootMotionFromEverything;
@@ -506,6 +495,11 @@ bool UAnimInstance::NeedsUpdate() const
 	return bNeedsUpdate;
 }
 
+void UAnimInstance::PreEvaluateAnimation()
+{
+	GetProxyOnGameThread<FAnimInstanceProxy>().PreEvaluateAnimation(this);
+}
+
 void UAnimInstance::EvaluateAnimation(FPoseContext& Output)
 {
 	GetProxyOnGameThread<FAnimInstanceProxy>().EvaluateAnimation(Output);
@@ -563,6 +557,8 @@ void UAnimInstance::PostEvaluateAnimation()
 		SCOPE_CYCLE_COUNTER(STAT_BlueprintPostEvaluateAnimation);
 		BlueprintPostEvaluateAnimation();
 	}
+
+	GetProxyOnGameThread<FAnimInstanceProxy>().ClearObjects();
 }
 
 void UAnimInstance::NativeInitializeAnimation()
diff --git a/Engine/Source/Runtime/Engine/Private/Animation/AnimInstanceProxy.cpp b/Engine/Source/Runtime/Engine/Private/Animation/AnimInstanceProxy.cpp
--- a/Engine/Source/Runtime/Engine/Private/Animation/AnimInstanceProxy.cpp
+++ b/Engine/Source/Runtime/Engine/Private/Animation/AnimInstanceProxy.cpp
@@ -41,15 +41,7 @@ void FAnimInstanceProxy::Initialize(UAnimInstance* InAnimInstance)
 
 	AnimClassInterface = IAnimClassInterface::GetFromClass(InAnimInstance->GetClass());
 
-	USkeletalMeshComponent* OwnerComponent = InAnimInstance->GetSkelMeshComponent();
-	if (OwnerComponent->SkeletalMesh != NULL)
-	{
-		Skeleton = OwnerComponent->SkeletalMesh->Skeleton;
-	}
-	else
-	{
-		Skeleton = nullptr;
-	}
+	InitializeObjects(InAnimInstance);
 
 	if (AnimClassInterface)
 	{
@@ -98,8 +90,6 @@ void FAnimInstanceProxy::Initialize(UAnimInstance* InAnimInstance)
 #endif
 	}
 
-	SkeletalMeshComponent = InAnimInstance->GetSkelMeshComponent();
-
 #if !NO_LOGGING
 	ActorName = GetNameSafe(InAnimInstance->GetOwningActor());
 #endif
@@ -169,6 +159,8 @@ void FAnimInstanceProxy::PreUpdate(UAnimInstance* InAnimInstance, float DeltaSec
 	RootMotionMode = InAnimInstance->RootMotionMode;
 	bShouldExtractRootMotion = InAnimInstance->ShouldExtractRootMotion();
 
+	InitializeObjects(InAnimInstance);
+
 	// Save off LOD level that we're currently using.
 	LODLevel = InAnimInstance->GetSkelMeshComponent()->PredictedLODLevel;
 
@@ -214,6 +206,25 @@ void FAnimInstanceProxy::PostUpdate(UAnimInstance* InAnimInstance) const
 	InAnimInstance->NotifyQueue.ApplyMontageNotifies(*this);
 }
 
+void FAnimInstanceProxy::InitializeObjects(UAnimInstance* InAnimInstance)
+{
+	SkeletalMeshComponent = InAnimInstance->GetSkelMeshComponent();
+	if (SkeletalMeshComponent->SkeletalMesh != nullptr)
+	{
+		Skeleton = SkeletalMeshComponent->SkeletalMesh->Skeleton;
+	}
+	else
+	{
+		Skeleton = nullptr;
+	}
+}
+
+void FAnimInstanceProxy::ClearObjects()
+{
+	SkeletalMeshComponent = nullptr;
+	Skeleton = nullptr;
+}
+
 FAnimTickRecord& FAnimInstanceProxy::CreateUninitializedTickRecord(int32 GroupIndex, FAnimGroupInstance*& OutSyncGroupPtr)
 {
 	// Find or create the sync group if there is one
@@ -661,6 +672,11 @@ void FAnimInstanceProxy::UpdateAnimation()
 	TickAssetPlayerInstances(CurrentDeltaSeconds);
 }
 
+void FAnimInstanceProxy::PreEvaluateAnimation(UAnimInstance* InAnimInstance)
+{
+	InitializeObjects(InAnimInstance);
+}
+
 void FAnimInstanceProxy::EvaluateAnimation(FPoseContext& Output)
 {
 	ANIM_MT_SCOPE_CYCLE_COUNTER(EvaluateAnimInstance, !IsInGameThread());
diff --git a/Engine/Source/Runtime/Engine/Private/Animation/AnimSequenceBase.cpp b/Engine/Source/Runtime/Engine/Private/Animation/AnimSequenceBase.cpp
--- a/Engine/Source/Runtime/Engine/Private/Animation/AnimSequenceBase.cpp
+++ b/Engine/Source/Runtime/Engine/Private/Animation/AnimSequenceBase.cpp
@@ -113,6 +113,10 @@ void UAnimSequenceBase::SortNotifies()
 	Notifies.Sort();
 }
 
+bool UAnimSequenceBase::IsNotifyAvailable() const
+{
+	return (Notifies.Num() != 0) && (SequenceLength > 0.f);
+}
 /** 
  * Retrieves AnimNotifies given a StartTime and a DeltaTime.
  * Time will be advanced and support looping if bAllowLooping is true.
@@ -121,12 +125,17 @@ void UAnimSequenceBase::SortNotifies()
  */
 void UAnimSequenceBase::GetAnimNotifies(const float& StartTime, const float& DeltaTime, const bool bAllowLooping, TArray<const FAnimNotifyEvent *> & OutActiveNotifies) const
 {
-	// Early out if we have no notifies
-	if( (Notifies.Num() == 0) || (DeltaTime == 0.f) || (SequenceLength==0.f))
+	if(DeltaTime == 0.f)
 	{
 		return;
 	}
 
+	// Early out if we have no notifies
+	if (!IsNotifyAvailable())
+	{
+		return;
+	}
+	
 	bool const bPlayingBackwards = (DeltaTime < 0.f);
 	float PreviousPosition = StartTime;
 	float CurrentPosition = StartTime;
diff --git a/Engine/Source/Runtime/Engine/Private/Animation/AnimSingleNodeInstance.cpp b/Engine/Source/Runtime/Engine/Private/Animation/AnimSingleNodeInstance.cpp
--- a/Engine/Source/Runtime/Engine/Private/Animation/AnimSingleNodeInstance.cpp
+++ b/Engine/Source/Runtime/Engine/Private/Animation/AnimSingleNodeInstance.cpp
@@ -29,7 +29,41 @@ UAnimSingleNodeInstance::UAnimSingleNodeInstance(const FObjectInitializer& Objec
 
 void UAnimSingleNodeInstance::SetAnimationAsset(class UAnimationAsset* NewAsset, bool bInIsLooping, float InPlayRate)
 {
+	if (NewAsset != CurrentAsset)
+	{
+		CurrentAsset = NewAsset;
+	}
+
 	FAnimSingleNodeInstanceProxy& Proxy = GetProxyOnGameThread<FAnimSingleNodeInstanceProxy>();
+
+	if (
+#if WITH_EDITOR
+		!Proxy.CanProcessAdditiveAnimations() &&
+#endif
+		NewAsset && NewAsset->IsValidAdditive())
+	{
+		UE_LOG(LogAnimation, Warning, TEXT("Setting an additive animation (%s) on an AnimSingleNodeInstance is not allowed. This will not function correctly in cooked builds!"), *NewAsset->GetName());
+	}
+
+	USkeletalMeshComponent* MeshComponent = GetSkelMeshComponent();
+	if (MeshComponent)
+	{
+		if (MeshComponent->SkeletalMesh == nullptr)
+		{
+			// if it does not have SkeletalMesh, we nullify it
+			CurrentAsset = nullptr;
+		}
+		else if (CurrentAsset != nullptr)
+		{
+			// if we have an asset, make sure their skeleton matches, otherwise, null it
+			if (MeshComponent->SkeletalMesh == nullptr || MeshComponent->SkeletalMesh->Skeleton != CurrentAsset->GetSkeleton())
+			{
+				// clear asset since we do not have matching skeleton
+				CurrentAsset = nullptr;
+			}
+		}
+	}
+	
 	Proxy.SetAnimationAsset(NewAsset, GetSkelMeshComponent(), bInIsLooping, InPlayRate);
 
 	// if composite, we want to make sure this is valid
@@ -58,9 +92,32 @@ void UAnimSingleNodeInstance::SetAnimationAsset(class UAnimationAsset* NewAsset,
 	}
 }
 
-void UAnimSingleNodeInstance::SetVertexAnimation(UVertexAnimation * NewVertexAnim, bool bIsLooping, float InPlayRate)
+void UAnimSingleNodeInstance::SetVertexAnimation(UVertexAnimation* NewVertexAnim, bool bIsLooping, float InPlayRate)
 {
-	GetProxyOnGameThread<FAnimSingleNodeInstanceProxy>().SetVertexAnimation(NewVertexAnim, bIsLooping, InPlayRate);
+	if (NewVertexAnim != CurrentVertexAnim)
+	{
+		CurrentVertexAnim = NewVertexAnim;
+	}
+
+	if (USkeletalMeshComponent * MeshComponent = GetSkelMeshComponent())
+	{
+		if (MeshComponent->SkeletalMesh == nullptr)
+		{
+			// if it does not have SkeletalMesh, we nullify it
+			CurrentVertexAnim = NULL;
+		}
+		else if (CurrentVertexAnim != nullptr)
+		{
+			// if we have an anim, make sure their mesh matches, otherwise, null it
+			if (MeshComponent->SkeletalMesh != CurrentVertexAnim->BaseSkelMesh)
+			{
+				// clear asset since we do not have matching skeleton
+				CurrentVertexAnim = nullptr;
+			}
+		}
+	}
+
+	GetProxyOnGameThread<FAnimSingleNodeInstanceProxy>().SetVertexAnimation(CurrentVertexAnim, bIsLooping, InPlayRate);
 
 	// reinitialize
 	InitializeAnimation();
@@ -131,10 +188,10 @@ void UAnimSingleNodeInstance::UpdateBlendspaceSamples(FVector InBlendInput)
 
 void UAnimSingleNodeInstance::RestartMontage(UAnimMontage* Montage, FName FromSection)
 {
-	FAnimSingleNodeInstanceProxy& Proxy = GetProxyOnGameThread<FAnimSingleNodeInstanceProxy>();
-
-	if( Montage == Proxy.GetCurrentAsset() )
+	if(Montage == CurrentAsset)
 	{
+		FAnimSingleNodeInstanceProxy& Proxy = GetProxyOnGameThread<FAnimSingleNodeInstanceProxy>();
+
 		Proxy.ResetWeightInfo();
 		Montage_Play(Montage, Proxy.GetPlayRate());
 		if( FromSection != NAME_None )
@@ -161,7 +218,7 @@ void UAnimSingleNodeInstance::NativePostEvaluateAnimation()
 void UAnimSingleNodeInstance::OnMontageInstanceStopped(FAnimMontageInstance& StoppedMontageInstance) 
 {
 	FAnimSingleNodeInstanceProxy& Proxy = GetProxyOnGameThread<FAnimSingleNodeInstanceProxy>();
-	if (StoppedMontageInstance.Montage == Proxy.GetCurrentAsset())
+	if (StoppedMontageInstance.Montage == CurrentAsset)
 	{
 		Proxy.SetCurrentTime(StoppedMontageInstance.GetPosition());
 	}
@@ -199,7 +256,7 @@ void UAnimSingleNodeInstance::SetLooping(bool bIsLooping)
 	FAnimSingleNodeInstanceProxy& Proxy = GetProxyOnGameThread<FAnimSingleNodeInstanceProxy>();
 	Proxy.SetLooping(bIsLooping);
 
-	if (UAnimMontage* Montage = Cast<UAnimMontage>(GetCurrentAsset()))
+	if (UAnimMontage* Montage = Cast<UAnimMontage>(CurrentAsset))
 	{
 		SetMontageLoop(Montage, Proxy.IsLooping(), Montage_GetCurrentSection());
 	}
@@ -216,7 +273,7 @@ void UAnimSingleNodeInstance::SetPlaying(bool bIsPlaying)
 	}
 	else if (Proxy.IsPlaying())
 	{
-		UAnimMontage* Montage = Cast<UAnimMontage>(GetCurrentAsset());
+		UAnimMontage* Montage = Cast<UAnimMontage>(CurrentAsset);
 		if (Montage)
 		{
 			RestartMontage(Montage);
@@ -262,12 +319,12 @@ void UAnimSingleNodeInstance::SetPlayRate(float InPlayRate)
 
 UAnimationAsset* UAnimSingleNodeInstance::GetCurrentAsset()
 {
-	return GetProxyOnGameThread<FAnimSingleNodeInstanceProxy>().GetCurrentAsset();
+	return CurrentAsset;
 }
 
 UVertexAnimation* UAnimSingleNodeInstance::GetCurrentVertexAnimation()
 {
-	return GetProxyOnGameThread<FAnimSingleNodeInstanceProxy>().GetCurrentVertexAnimation();
+	return CurrentVertexAnim;
 }
 
 float UAnimSingleNodeInstance::GetCurrentTime() const
@@ -296,7 +353,7 @@ void UAnimSingleNodeInstance::SetPositionWithPreviousTime(float InPosition, floa
 	// this will need to handle manually, emptying, it and collect it, and trigger them at once. 
 	if (bFireNotifies)
 	{
-		UAnimSequenceBase * SequenceBase = Cast<UAnimSequenceBase> (Proxy.GetCurrentAsset());
+		UAnimSequenceBase * SequenceBase = Cast<UAnimSequenceBase> (CurrentAsset);
 		if (SequenceBase)
 		{
 			NotifyQueue.Reset(GetSkelMeshComponent());
@@ -331,12 +388,24 @@ void UAnimSingleNodeInstance::SetBlendSpaceInput(const FVector& InBlendInput)
 
 float UAnimSingleNodeInstance::GetLength()
 {
-	return GetProxyOnGameThread<FAnimSingleNodeInstanceProxy>().GetLength();
+	if ((CurrentAsset != NULL))
+	{
+		if (UBlendSpace* BlendSpace = Cast<UBlendSpace>(CurrentAsset))
+		{
+			return BlendSpace->AnimLength;
+		}
+		else if (UAnimSequenceBase* SequenceBase = Cast<UAnimSequenceBase>(CurrentAsset))
+		{
+			return SequenceBase->SequenceLength;
+		}
+	}
+
+	return 0.f;
 }
 
 void UAnimSingleNodeInstance::StepForward()
 {
-	if (UAnimSequence* Sequence = Cast<UAnimSequence>(GetCurrentAsset()))
+	if (UAnimSequence* Sequence = Cast<UAnimSequence>(CurrentAsset))
 	{
 		FAnimSingleNodeInstanceProxy& Proxy = GetProxyOnGameThread<FAnimSingleNodeInstanceProxy>();
 
@@ -349,7 +418,7 @@ void UAnimSingleNodeInstance::StepForward()
 
 void UAnimSingleNodeInstance::StepBackward()
 {
-	if (UAnimSequence* Sequence = Cast<UAnimSequence>(GetCurrentAsset()))
+	if (UAnimSequence* Sequence = Cast<UAnimSequence>(CurrentAsset))
 	{
 		FAnimSingleNodeInstanceProxy& Proxy = GetProxyOnGameThread<FAnimSingleNodeInstanceProxy>();
 
@@ -367,7 +436,7 @@ FAnimInstanceProxy* UAnimSingleNodeInstance::CreateAnimInstanceProxy()
 
 FVector UAnimSingleNodeInstance::GetFilterLastOutput()
 {
-	if (UBlendSpaceBase* Blendspace = Cast<UBlendSpaceBase>(GetCurrentAsset()))
+	if (UBlendSpaceBase* Blendspace = Cast<UBlendSpaceBase>(CurrentAsset))
 	{
 		FAnimSingleNodeInstanceProxy& Proxy = GetProxyOnGameThread<FAnimSingleNodeInstanceProxy>();
 		return Proxy.GetFilterLastOutput();
diff --git a/Engine/Source/Runtime/Engine/Private/Animation/AnimSingleNodeInstanceProxy.cpp b/Engine/Source/Runtime/Engine/Private/Animation/AnimSingleNodeInstanceProxy.cpp
--- a/Engine/Source/Runtime/Engine/Private/Animation/AnimSingleNodeInstanceProxy.cpp
+++ b/Engine/Source/Runtime/Engine/Private/Animation/AnimSingleNodeInstanceProxy.cpp
@@ -6,6 +6,7 @@
 #include "Animation/AnimMontage.h"
 #include "Animation/VertexAnim/VertexAnimation.h"
 #include "Animation/BlendSpace.h"
+#include "Animation/AnimSingleNodeInstance.h"
 
 void FAnimSingleNodeInstanceProxy::Initialize(UAnimInstance* InAnimInstance)
 {
@@ -259,6 +260,23 @@ void FAnimSingleNodeInstanceProxy::PostUpdate(UAnimInstance* InAnimInstance) con
 	}
 }
 
+void FAnimSingleNodeInstanceProxy::InitializeObjects(UAnimInstance* InAnimInstance)
+{
+	FAnimInstanceProxy::InitializeObjects(InAnimInstance);
+
+	UAnimSingleNodeInstance* AnimSingleNodeInstance = CastChecked<UAnimSingleNodeInstance>(InAnimInstance);
+	CurrentAsset = AnimSingleNodeInstance->CurrentAsset;
+	CurrentVertexAnim = AnimSingleNodeInstance->CurrentVertexAnim;
+}
+
+void FAnimSingleNodeInstanceProxy::ClearObjects()
+{
+	FAnimInstanceProxy::ClearObjects();
+
+	CurrentAsset = nullptr;
+	CurrentVertexAnim = nullptr;
+}
+
 void FAnimSingleNodeInstanceProxy::InternalBlendSpaceEvaluatePose(class UBlendSpaceBase* BlendSpace, TArray<FBlendSampleData>& BlendSampleDataCache, FPoseContext& OutContext)
 {
 	if (BlendSpace->IsValidAdditive())
@@ -296,38 +314,6 @@ void FAnimSingleNodeInstanceProxy::InternalBlendSpaceEvaluatePose(class UBlendSp
 
 void FAnimSingleNodeInstanceProxy::SetAnimationAsset(class UAnimationAsset* NewAsset, USkeletalMeshComponent* MeshComponent, bool bIsLooping, float InPlayRate)
 {
-	if (NewAsset != CurrentAsset)
-	{
-		CurrentAsset = NewAsset;
-	}
-
-	if (
-#if WITH_EDITOR
-		!bCanProcessAdditiveAnimations &&
-#endif
-		NewAsset && NewAsset->IsValidAdditive())
-	{
-		UE_LOG(LogAnimation, Warning, TEXT("Setting an additve animation (%s) on an AnimSingleNodeInstance is not allowed. This will not function correctly in cooked builds!"), *NewAsset->GetName());
-	}
-
-	if (MeshComponent)
-	{
-		if (MeshComponent->SkeletalMesh == NULL)
-		{
-			// if it does not have SkeletalMesh, we nullify it
-			CurrentAsset = NULL;
-		}
-		else if (CurrentAsset != NULL)
-		{
-			// if we have an asset, make sure their skeleton matches, otherwise, null it
-			if (GetSkeleton() != CurrentAsset->GetSkeleton())
-			{
-				// clear asset since we do not have matching skeleton
-				CurrentAsset = NULL;
-			}
-		}
-	}
-
 	bLooping = bIsLooping;
 	PlayRate = InPlayRate;
 	CurrentTime = 0.f;
@@ -359,29 +345,6 @@ void FAnimSingleNodeInstanceProxy::UpdateBlendspaceSamples(FVector InBlendInput)
 
 void FAnimSingleNodeInstanceProxy::SetVertexAnimation(UVertexAnimation * NewVertexAnim, bool bIsLooping, float InPlayRate)
 {
-	if (NewVertexAnim != CurrentVertexAnim)
-	{
-		CurrentVertexAnim = NewVertexAnim;
-	}
-
-	if (USkeletalMeshComponent * MeshComponent = GetSkelMeshComponent())
-	{
-		if (MeshComponent->SkeletalMesh == NULL)
-		{
-			// if it does not have SkeletalMesh, we nullify it
-			CurrentVertexAnim = NULL;
-		}
-		else if (CurrentVertexAnim != NULL)
-		{
-			// if we have an anim, make sure their mesh matches, otherwise, null it
-			if (MeshComponent->SkeletalMesh != CurrentVertexAnim->BaseSkelMesh)
-			{
-				// clear asset since we do not have matching skeleton
-				CurrentVertexAnim = NULL;
-			}
-		}
-	}
-
 	bLooping = bIsLooping;
 	PlayRate = InPlayRate;
 }
@@ -414,23 +377,6 @@ void FAnimSingleNodeInstanceProxy::SetReverse(bool bInReverse)
 	}*/
 }
 
-float FAnimSingleNodeInstanceProxy::GetLength()
-{
-	if ((CurrentAsset != NULL))
-	{
-		if (UBlendSpace* BlendSpace = Cast<UBlendSpace>(CurrentAsset))
-		{
-			return BlendSpace->AnimLength;
-		}
-		else if (UAnimSequenceBase* SequenceBase = Cast<UAnimSequenceBase>(CurrentAsset))
-		{
-			return SequenceBase->SequenceLength;
-		}
-	}	
-
-	return 0.f;
-}
-
 void FAnimSingleNodeInstanceProxy::SetBlendSpaceInput(const FVector& InBlendInput)
 {
 	BlendSpaceInput = InBlendInput;
diff --git a/Engine/Source/Runtime/Engine/Private/BlueprintGeneratedClass.cpp b/Engine/Source/Runtime/Engine/Private/BlueprintGeneratedClass.cpp
--- a/Engine/Source/Runtime/Engine/Private/BlueprintGeneratedClass.cpp
+++ b/Engine/Source/Runtime/Engine/Private/BlueprintGeneratedClass.cpp
@@ -363,19 +363,20 @@ void UBlueprintGeneratedClass::BuildCustomPropertyListForPostConstruction(FCusto
 	}
 }
 
-void UBlueprintGeneratedClass::BuildCustomArrayPropertyListForPostConstruction(UArrayProperty* ArrayProperty, FCustomPropertyListNode*& InPropertyList, const uint8* DataPtr, const uint8* DefaultDataPtr)
+void UBlueprintGeneratedClass::BuildCustomArrayPropertyListForPostConstruction(UArrayProperty* ArrayProperty, FCustomPropertyListNode*& InPropertyList, const uint8* DataPtr, const uint8* DefaultDataPtr, int32 StartIndex)
 {
 	FCustomPropertyListNode** CurrentArrayNodePtr = &InPropertyList;
 
 	FScriptArrayHelper ArrayValueHelper(ArrayProperty, DataPtr);
 	FScriptArrayHelper DefaultArrayValueHelper(ArrayProperty, DefaultDataPtr);
 
-	for (int32 ArrayValueIndex = 0; ArrayValueIndex < ArrayValueHelper.Num(); ++ArrayValueIndex)
+	for (int32 ArrayValueIndex = StartIndex; ArrayValueIndex < ArrayValueHelper.Num(); ++ArrayValueIndex)
 	{
-		if (ArrayValueIndex < DefaultArrayValueHelper.Num())
+		const int32 DefaultArrayValueIndex = ArrayValueIndex - StartIndex;
+		if (DefaultArrayValueIndex < DefaultArrayValueHelper.Num())
 		{
 			const uint8* ArrayPropertyValue = ArrayValueHelper.GetRawPtr(ArrayValueIndex);
-			const uint8* DefaultArrayPropertyValue = DefaultArrayValueHelper.GetRawPtr(ArrayValueIndex);
+			const uint8* DefaultArrayPropertyValue = DefaultArrayValueHelper.GetRawPtr(DefaultArrayValueIndex);
 
 			if (UStructProperty* InnerStructProperty = Cast<UStructProperty>(ArrayProperty->Inner))
 			{
@@ -434,8 +435,25 @@ void UBlueprintGeneratedClass::BuildCustomArrayPropertyListForPostConstruction(U
 		}
 		else
 		{
-			// Signals the "end" of differences with the default value (signals that remaining values should be copied in full)
-			*CurrentArrayNodePtr = new(CustomPropertyListForPostConstruction) FCustomPropertyListNode(nullptr, ArrayValueIndex);
+			// Create a temp default array as a placeholder to compare against the remaining elements in the value.
+			FScriptArray TempDefaultArray;
+			const int32 Count = ArrayValueHelper.Num() - DefaultArrayValueHelper.Num();
+			TempDefaultArray.Add(Count, ArrayProperty->Inner->ElementSize);
+			uint8 *Dest = (uint8*)TempDefaultArray.GetData();
+			if (ArrayProperty->Inner->PropertyFlags & CPF_ZeroConstructor)
+			{
+				FMemory::Memzero(Dest, Count * ArrayProperty->Inner->ElementSize);
+			}
+			else
+			{
+				for (int32 i = 0; i < Count; i++, Dest += ArrayProperty->Inner->ElementSize)
+				{
+					ArrayProperty->Inner->InitializeValue(Dest);
+				}
+			}
+
+			// Recursively fill out the property list for the remainder of the elements in the value that extend beyond the size of the default value.
+			BuildCustomArrayPropertyListForPostConstruction(ArrayProperty, *CurrentArrayNodePtr, DataPtr, (uint8*)&TempDefaultArray, ArrayValueIndex);
 
 			// Don't need to record anything else.
 			break;
diff --git a/Engine/Source/Runtime/Engine/Private/Components/AudioComponent.cpp b/Engine/Source/Runtime/Engine/Private/Components/AudioComponent.cpp
--- a/Engine/Source/Runtime/Engine/Private/Components/AudioComponent.cpp
+++ b/Engine/Source/Runtime/Engine/Private/Components/AudioComponent.cpp
@@ -362,7 +362,7 @@ void UAudioComponent::PlaybackCompleted(bool bFailedToStart)
 	}
 
 	// Auto destruction is handled via marking object for deletion.
-	if (bAutoDestroy)
+	if (bAutoDestroy && !bIsActive)
 	{
 		DestroyComponent();
 	}
diff --git a/Engine/Source/Runtime/Engine/Private/Components/ChildActorComponent.cpp b/Engine/Source/Runtime/Engine/Private/Components/ChildActorComponent.cpp
--- a/Engine/Source/Runtime/Engine/Private/Components/ChildActorComponent.cpp
+++ b/Engine/Source/Runtime/Engine/Private/Components/ChildActorComponent.cpp
@@ -2,6 +2,7 @@
 
 #include "EnginePrivate.h"
 #include "Components/ChildActorComponent.h"
+#include "Net/UnrealNetwork.h"
 
 DEFINE_LOG_CATEGORY_STATIC(LogChildActorComponent, Warning, All);
 
@@ -34,6 +35,9 @@ void UChildActorComponent::OnRegister()
 				// so moving Attach to happen in Register
 				ChildRoot->AttachToComponent(this, FAttachmentTransformRules::SnapToTargetNotIncludingScale);
 			}
+
+			// Ensure the components replication is correctly initialized
+			SetIsReplicated(ChildActor->GetIsReplicated());
 		}
 	}
 	else if (ChildActorClass)
@@ -83,6 +87,29 @@ void UChildActorComponent::PostEditUndo()
 }
 #endif
 
+void UChildActorComponent::GetLifetimeReplicatedProps(TArray<FLifetimeProperty>& OutLifetimeProps) const
+{
+	Super::GetLifetimeReplicatedProps(OutLifetimeProps);
+
+	DOREPLIFETIME(UChildActorComponent, ChildActor);
+}
+
+void UChildActorComponent::PostRepNotifies()
+{
+	Super::PostRepNotifies();
+
+	if (ChildActor)
+	{
+		ChildActorClass = ChildActor->GetClass();
+		ChildActorName = ChildActor->GetFName();
+	}
+	else
+	{
+		ChildActorClass = nullptr;
+		ChildActorName = NAME_None;
+	}
+}
+
 void UChildActorComponent::OnComponentCreated()
 {
 	Super::OnComponentCreated();
@@ -94,16 +121,14 @@ void UChildActorComponent::OnComponentDestroyed(bool bDestroyingHierarchy)
 {
 	Super::OnComponentDestroyed(bDestroyingHierarchy);
 
-	const UWorld* const MyWorld = GetWorld();
-	DestroyChildActor(MyWorld && !MyWorld->IsGameWorld());
+	DestroyChildActor();
 }
 
 void UChildActorComponent::OnUnregister()
 {
 	Super::OnUnregister();
 
-	const UWorld* const MyWorld = GetWorld();
-	DestroyChildActor(MyWorld && !MyWorld->IsGameWorld());
+	DestroyChildActor();
 }
 
 FChildActorComponentInstanceData::FChildActorComponentInstanceData(const UChildActorComponent* Component)
@@ -262,6 +287,18 @@ void UChildActorComponent::PostLoad()
 
 void UChildActorComponent::CreateChildActor()
 {
+	AActor* MyOwner = GetOwner();
+
+	if (MyOwner && !MyOwner->HasAuthority())
+	{
+		AActor* ChildClassCDO = (ChildActorClass ? ChildActorClass->GetDefaultObject<AActor>() : nullptr);
+		if (ChildClassCDO && ChildClassCDO->GetIsReplicated())
+		{
+			// If we belong to an actor that is not authoritative and the child class is replicated then we expect that Actor will be replicated across so don't spawn one
+			return;
+		}
+	}
+
 	// Kill spawned actor if we have one
 	DestroyChildActor();
 
@@ -273,7 +310,6 @@ void UChildActorComponent::CreateChildActor()
 		{
 			// Before we spawn let's try and prevent cyclic disaster
 			bool bSpawn = true;
-			AActor* MyOwner = GetOwner();
 			AActor* Actor = MyOwner;
 			while (Actor && bSpawn)
 			{
@@ -323,6 +359,8 @@ void UChildActorComponent::CreateChildActor()
 					ChildActor->FinishSpawning(ComponentToWorld, false, ComponentInstanceData);
 
 					ChildActor->AttachToComponent(this, FAttachmentTransformRules::SnapToTargetNotIncludingScale);
+
+					SetIsReplicated(ChildActor->GetIsReplicated());
 				}
 			}
 		}
@@ -336,50 +374,62 @@ void UChildActorComponent::CreateChildActor()
 	}
 }
 
-void UChildActorComponent::DestroyChildActor(const bool bRequiresRename)
+void UChildActorComponent::DestroyChildActor()
 {
-	// If we own an Actor, kill it now
-	if (ChildActor != nullptr && !GExitPurge)
+	// If we own an Actor, kill it now unless we don't have authority on it, for that we rely on the server
+	if (ChildActor && ChildActor->HasAuthority())
 	{
-		// if still alive, destroy, otherwise just clear the pointer
-		if (!ChildActor->IsPendingKillOrUnreachable())
+		if (!GExitPurge)
 		{
-#if WITH_EDITOR
-			if (CachedInstanceData)
+			// if still alive, destroy, otherwise just clear the pointer
+			if (!ChildActor->IsPendingKillOrUnreachable())
 			{
-				delete CachedInstanceData;
-				CachedInstanceData = nullptr;
-			}
+#if WITH_EDITOR
+				if (CachedInstanceData)
+				{
+					delete CachedInstanceData;
+					CachedInstanceData = nullptr;
+				}
 #else
-			check(!CachedInstanceData);
+				check(!CachedInstanceData);
 #endif
-			// If we're already tearing down we won't be needing this
-			if (!HasAnyFlags(RF_BeginDestroyed))
-			{
-				CachedInstanceData = new FChildActorComponentInstanceData(this);
-			}
+				// If we're already tearing down we won't be needing this
+				if (!HasAnyFlags(RF_BeginDestroyed))
+				{
+					CachedInstanceData = new FChildActorComponentInstanceData(this);
+				}
 
-			UWorld* World = ChildActor->GetWorld();
-			// World may be nullptr during shutdown
-			if (World != nullptr)
-			{
-				UClass* ChildClass = ChildActor->GetClass();
+				UWorld* World = ChildActor->GetWorld();
+				// World may be nullptr during shutdown
+				if (World != nullptr)
+				{
+					UClass* ChildClass = ChildActor->GetClass();
 
-				// We would like to make certain that our name is not going to accidentally get taken from us while we're destroyed
-				// so we increment ClassUnique beyond our index to be certain of it.  This is ... a bit hacky.
-				int32& ClassUnique = ChildActor->GetOutermost()->ClassUniqueNameIndexMap.FindOrAdd(ChildClass->GetFName());
-				ClassUnique = FMath::Max(ClassUnique, ChildActor->GetFName().GetNumber());
+					// We would like to make certain that our name is not going to accidentally get taken from us while we're destroyed
+					// so we increment ClassUnique beyond our index to be certain of it.  This is ... a bit hacky.
+					int32& ClassUnique = ChildActor->GetOutermost()->ClassUniqueNameIndexMap.FindOrAdd(ChildClass->GetFName());
+					ClassUnique = FMath::Max(ClassUnique, ChildActor->GetFName().GetNumber());
 
-				if (bRequiresRename)
-				{
-					const FString ObjectBaseName = FString::Printf(TEXT("DESTROYED_%s_CHILDACTOR"), *ChildClass->GetName());
-					const ERenameFlags RenameFlags = ((GetWorld()->IsGameWorld() || IsLoading()) ? REN_DoNotDirty | REN_ForceNoResetLoaders : REN_DoNotDirty);
-					ChildActor->Rename(*MakeUniqueObjectName(ChildActor->GetOuter(), ChildClass, *ObjectBaseName).ToString(), nullptr, RenameFlags);
+					// If we are getting here due to garbage collection we can't rename, so we'll have to abandon this child actor name and pick up a new one
+					if (!IsGarbageCollecting())
+					{
+						const FString ObjectBaseName = FString::Printf(TEXT("DESTROYED_%s_CHILDACTOR"), *ChildClass->GetName());
+						const ERenameFlags RenameFlags = ((GetWorld()->IsGameWorld() || IsLoading()) ? REN_DoNotDirty | REN_ForceNoResetLoaders : REN_DoNotDirty);
+						ChildActor->Rename(*MakeUniqueObjectName(ChildActor->GetOuter(), ChildClass, *ObjectBaseName).ToString(), nullptr, RenameFlags);
+					}
+					else
+					{
+						ChildActorName = NAME_None;
+						if (CachedInstanceData)
+						{
+							CachedInstanceData->ChildActorName = NAME_None;
+						}
+					}
+					World->DestroyActor(ChildActor);
 				}
-				World->DestroyActor(ChildActor);
 			}
 		}
-	}
 
-	ChildActor = nullptr;
+		ChildActor = nullptr;
+	}
 }
diff --git a/Engine/Source/Runtime/Engine/Private/Components/SkeletalMeshComponent.cpp b/Engine/Source/Runtime/Engine/Private/Components/SkeletalMeshComponent.cpp
--- a/Engine/Source/Runtime/Engine/Private/Components/SkeletalMeshComponent.cpp
+++ b/Engine/Source/Runtime/Engine/Private/Components/SkeletalMeshComponent.cpp
@@ -332,6 +332,10 @@ void USkeletalMeshComponent::OnRegister()
 
 void USkeletalMeshComponent::OnUnregister()
 {
+	const bool bBlockOnTask = true; // wait on evaluation task so we complete any work before this component goes away
+	const bool bPerformPostAnimEvaluation = false; // Skip post evaluation, it would be wasted work
+	HandleExistingParallelEvaluationTask(bBlockOnTask, bPerformPostAnimEvaluation);
+
 #if WITH_APEX_CLOTHING
 	//clothing actors will be re-created in TickClothing
 	ReleaseAllClothingResources();
@@ -1239,6 +1243,11 @@ void USkeletalMeshComponent::RefreshBoneTransforms(FActorComponentTickFunction*
 		CachedCurve.Empty();
 	}
 
+	if(AnimScriptInstance)
+	{
+		AnimScriptInstance->PreEvaluateAnimation();
+	}
+
 	if (bDoParallelEvaluation)
 	{
 		QUICK_SCOPE_CYCLE_COUNTER(STAT_USkeletalMeshComponent_RefreshBoneTransforms_SetupParallel); 
@@ -2478,3 +2487,10 @@ void USkeletalMeshComponent::UnregisterOnPhysicsCreatedDelegate(const FDelegateH
 {
 	OnSkelMeshPhysicsCreated.Remove(DelegateHandle);
 }
+
+bool USkeletalMeshComponent::MoveComponentImpl(const FVector& Delta, const FQuat& NewRotation, bool bSweep, FHitResult* OutHit /*= NULL*/, EMoveComponentFlags MoveFlags /*= MOVECOMP_NoFlags*/, ETeleportType Teleport /*= ETeleportType::None*/)
+{
+	// If we're simulating physics ignore teleport type, we just want to teleport.
+	ETeleportType NewTeleportType = BodyInstance.bSimulatePhysics ? ETeleportType::TeleportPhysics : Teleport;
+	return Super::MoveComponentImpl(Delta, NewRotation, bSweep, OutHit, MoveFlags, NewTeleportType);
+}
diff --git a/Engine/Source/Runtime/Engine/Private/Components/SplineComponent.cpp b/Engine/Source/Runtime/Engine/Private/Components/SplineComponent.cpp
--- a/Engine/Source/Runtime/Engine/Private/Components/SplineComponent.cpp
+++ b/Engine/Source/Runtime/Engine/Private/Components/SplineComponent.cpp
@@ -17,12 +17,14 @@ USplineComponent::USplineComponent(const FObjectInitializer& ObjectInitializer)
 	, bStationaryEndpoints(false)
 	, bSplineHasBeenEdited(false)
 	, bInputSplinePointsToConstructionScript(false)
-	, bAlwaysRenderInEditor(true)
+	, bDrawDebug(true)
 	, bClosedLoop(false)
 	, DefaultUpVector(FVector::UpVector)
-#if WITH_EDITORONLY_DATA
+#if !UE_BUILD_SHIPPING
 	, EditorUnselectedSplineSegmentColor(FLinearColor(1.0f, 1.0f, 1.0f))
 	, EditorSelectedSplineSegmentColor(FLinearColor(1.0f, 0.0f, 0.0f))
+#endif
+#if WITH_EDITORONLY_DATA
 	, bShouldVisualizeScale(false)
 	, ScaleVisualizationWidth(30.0f)
 #endif
@@ -439,7 +441,7 @@ bool USplineComponent::IsClosedLoop() const
 
 void USplineComponent::SetUnselectedSplineSegmentColor(const FLinearColor& Color)
 {
-#if WITH_EDITOR
+#if !UE_BUILD_SHIPPING
 	EditorUnselectedSplineSegmentColor = Color;
 #endif
 }
@@ -447,12 +449,19 @@ void USplineComponent::SetUnselectedSplineSegmentColor(const FLinearColor& Color
 
 void USplineComponent::SetSelectedSplineSegmentColor(const FLinearColor& Color)
 {
-#if WITH_EDITOR
+#if !UE_BUILD_SHIPPING
 	EditorSelectedSplineSegmentColor = Color;
 #endif
 }
 
 
+void USplineComponent::SetDrawDebug(bool bShow)
+{
+	bDrawDebug = bShow;
+	MarkRenderStateDirty();
+}
+
+
 void USplineComponent::ClearSplinePoints()
 {
 	SplineInfo.Points.Reset();
@@ -1115,10 +1124,10 @@ FTransform USplineComponent::FindTransformClosestToWorldLocation(const FVector&
 	return GetTransformAtSplineInputKey(Param, CoordinateSpace, bUseScale);
 }
 
-#if WITH_EDITOR
+#if !UE_BUILD_SHIPPING
 FPrimitiveSceneProxy* USplineComponent::CreateSceneProxy()
 {
-	if (!bAlwaysRenderInEditor)
+	if (!bDrawDebug)
 	{
 		return Super::CreateSceneProxy();
 	}
@@ -1129,7 +1138,7 @@ FPrimitiveSceneProxy* USplineComponent::CreateSceneProxy()
 
 		FSplineSceneProxy(const USplineComponent* InComponent)
 			: FPrimitiveSceneProxy(InComponent)
-			, bAlwaysRenderInEditor(InComponent->bAlwaysRenderInEditor)
+			, bDrawDebug(InComponent->bDrawDebug)
 			, SplineInfo(InComponent->SplineInfo)
 			, LineColor(InComponent->EditorUnselectedSplineSegmentColor)
 		{}
@@ -1218,7 +1227,7 @@ FPrimitiveSceneProxy* USplineComponent::CreateSceneProxy()
 		virtual FPrimitiveViewRelevance GetViewRelevance(const FSceneView* View) const override
 		{
 			FPrimitiveViewRelevance Result;
-			Result.bDrawRelevance = bAlwaysRenderInEditor && !IsSelected() && IsShown(View) && View->Family->EngineShowFlags.Splines;
+			Result.bDrawRelevance = bDrawDebug && !IsSelected() && IsShown(View) && View->Family->EngineShowFlags.Splines;
 			Result.bDynamicRelevance = true;
 			Result.bShadowRelevance = IsShadowCast(View);
 			Result.bEditorPrimitiveRelevance = UseEditorCompositing(View);
@@ -1229,7 +1238,7 @@ FPrimitiveSceneProxy* USplineComponent::CreateSceneProxy()
 		uint32 GetAllocatedSize(void) const { return FPrimitiveSceneProxy::GetAllocatedSize(); }
 
 	private:
-		bool bAlwaysRenderInEditor;
+		bool bDrawDebug;
 		FInterpCurveVector SplineInfo;
 		FLinearColor LineColor;
 	};
@@ -1240,7 +1249,7 @@ FPrimitiveSceneProxy* USplineComponent::CreateSceneProxy()
 
 FBoxSphereBounds USplineComponent::CalcBounds(const FTransform& LocalToWorld) const
 {
-	if (!bAlwaysRenderInEditor)
+	if (!bDrawDebug)
 	{
 		// Do as little as possible if not rendering anything
 		return Super::CalcBounds(LocalToWorld);
diff --git a/Engine/Source/Runtime/Engine/Private/DemoNetDriver.cpp b/Engine/Source/Runtime/Engine/Private/DemoNetDriver.cpp
--- a/Engine/Source/Runtime/Engine/Private/DemoNetDriver.cpp
+++ b/Engine/Source/Runtime/Engine/Private/DemoNetDriver.cpp
@@ -537,43 +537,6 @@ bool UDemoNetDriver::InitConnectInternal( FString& Error )
 
 	// Create fake control channel
 	ServerConnection->CreateChannel( CHTYPE_Control, 1 );
-
-	// Attempt to read metadata if it exists
-	FArchive* MetadataAr = ReplayStreamer->GetMetadataArchive();
-
-	FNetworkDemoMetadataHeader MetadataHeader;
-
-	if ( MetadataAr != nullptr )
-	{
-		(*MetadataAr) << MetadataHeader;
-
-		// Check metadata magic value
-		if ( MetadataHeader.Magic != NETWORK_DEMO_METADATA_MAGIC )
-		{
-			Error = FString( TEXT( "Demo metadata file is corrupt" ) );
-			UE_LOG( LogDemo, Error, TEXT( "UDemoNetDriver::InitConnect: %s" ), *Error );
-			GameInstance->HandleDemoPlaybackFailure( EDemoPlayFailure::Corrupt, Error );
-			return false;
-		}
-
-		// Check version
-		if ( MetadataHeader.Version != NETWORK_DEMO_METADATA_VERSION )
-		{
-			Error = FString( TEXT( "Demo metadata file version is incorrect" ) );
-			UE_LOG( LogDemo, Error, TEXT( "UDemoNetDriver::InitConnect: %s" ), *Error );
-			GameInstance->HandleDemoPlaybackFailure( EDemoPlayFailure::InvalidVersion, Error );
-			return false;
-		}
-
-		DemoTotalFrames = MetadataHeader.NumFrames;
-		DemoTotalTime	= MetadataHeader.TotalTime;
-
-		UE_LOG( LogDemo, Log, TEXT( "Starting demo playback with full demo and metadata. Filename: %s, Frames: %i, Version %i" ), *DemoURL.Map, DemoTotalFrames, DemoHeader.Version );
-	}
-	else
-	{
-		UE_LOG( LogDemo, Log, TEXT( "Starting demo playback with streaming demo, metadata file not found. Filename: %s, Version %i" ), *DemoURL.Map, DemoHeader.Version );
-	}
 	
 	// Default async world loading to the cvar value...
 	bool bAsyncLoadWorld = CVarDemoAsyncLoadWorld.GetValueOnGameThread() > 0;
@@ -627,70 +590,12 @@ bool UDemoNetDriver::InitConnectInternal( FString& Error )
 
 		auto NewPendingNetGame = NewObject<UDemoPendingNetGame>();
 
+		// Set up the pending net game so that the engine can call LoadMap on the next tick.
 		NewPendingNetGame->DemoNetDriver = this;
+		NewPendingNetGame->URL = DemoURL;
+		NewPendingNetGame->bSuccessfullyConnected = true;
 
 		WorldContext->PendingNetGame = NewPendingNetGame;
-
-		bool bSuccess = GEngine->LoadMap( *WorldContext, LocalDemoURL, NewPendingNetGame, LoadMapError );
-
-#if !(UE_BUILD_SHIPPING || UE_BUILD_TEST)
-		if ( CVarDemoForceFailure.GetValueOnGameThread() == 2 )
-		{
-			bSuccess = false;
-		}
-#endif
-
-		if ( !bSuccess )
-		{
-			StopDemo();
-
-			// If we don't have a world that means we failed loading the new world.
-			// Since there is no world, we must free the net driver ourselves
-			// Technically the pending net game should handle it, but things aren't quite setup properly to handle that either
-			if ( WorldContext->World() == NULL )
-			{
-				GEngine->DestroyNamedNetDriver( WorldContext->PendingNetGame, NetDriverName );
-			}
-
-			WorldContext->PendingNetGame = NULL;
-
-			GEngine->BrowseToDefaultMap( *WorldContext );
-
-			Error = LoadMapError;
-			UE_LOG( LogDemo, Error, TEXT( "UDemoNetDriver::InitConnect: LoadMap failed: failed: %s" ), *Error );
-			GameInstance->HandleDemoPlaybackFailure( EDemoPlayFailure::Generic, FString( TEXT( "LoadMap failed" ) ) );
-			return false;
-		}
-
-		SetWorld( WorldContext->World() );
-		WorldContext->World()->DemoNetDriver = this;
-		WorldContext->PendingNetGame = NULL;
-
-		// Read meta data, if it exists
-		for ( int32 i = 0; i < MetadataHeader.NumStreamingLevels; ++i )
-		{
-			ULevelStreamingKismet* StreamingLevel = NewObject<ULevelStreamingKismet>( GetWorld(), NAME_None, RF_NoFlags, NULL );
-
-			StreamingLevel->bShouldBeLoaded		= true;
-			StreamingLevel->bShouldBeVisible	= true;
-			StreamingLevel->bShouldBlockOnLoad	= false;
-			StreamingLevel->bInitiallyLoaded	= true;
-			StreamingLevel->bInitiallyVisible	= true;
-
-			FString PackageName;
-			FString PackageNameToLoad;
-
-			( *MetadataAr ) << PackageName;
-			( *MetadataAr ) << PackageNameToLoad;
-			( *MetadataAr ) << StreamingLevel->LevelTransform;
-
-			StreamingLevel->PackageNameToLoad = FName( *PackageNameToLoad );
-			StreamingLevel->SetWorldAssetByPackageName( FName( *PackageName ) );
-
-			GetWorld()->StreamingLevels.Add( StreamingLevel );
-
-			UE_LOG( LogDemo, Log, TEXT( "  Loading streamingLevel: %s, %s" ), *PackageName, *PackageNameToLoad );
-		}
 	}
 
 	return true;
@@ -1530,9 +1435,20 @@ void UDemoNetDriver::TickDemoRecord( float DeltaSeconds )
 
 			Actor->CallPreReplication( this );
 
-			const bool bUpdatedExternalData = ( FindOrCreateRepChangedPropertyTracker( Actor ).Get()->ExternalData.Num() > 0 );
+			const bool bDidReplicateActor = DemoReplicateActor( Actor, ClientConnections[0], SpectatorController, false );
 
-			if ( DemoReplicateActor( Actor, ClientConnections[0], SpectatorController, false ) || bUpdatedExternalData )
+			TSharedPtr<FRepChangedPropertyTracker> PropertyTracker = FindOrCreateRepChangedPropertyTracker( Actor );
+
+			if (!bDidReplicateActor)
+			{
+				// Clear external data if the actor didn't replicate for some reason
+				PropertyTracker->ExternalData.Empty();
+				PropertyTracker->ExternalDataNumBits = 0;
+			}
+
+			const bool bUpdatedExternalData = PropertyTracker->ExternalData.Num() > 0;
+
+			if ( bDidReplicateActor || bUpdatedExternalData )
 			{
 				// Choose an optimal time, we choose 70% of the actual rate to allow frequency to go up if needed
 				ActorInfo->OptimalNetUpdateDelta = FMath::Clamp( LastReplicateDelta * 0.7f, MinOptimalDelta, MaxOptimalDelta );
@@ -2680,6 +2596,81 @@ void UDemoNetDriver::NotifyGotoTimeFinished(bool bWasSuccessful)
 	}
 }
 
+void UDemoNetDriver::PendingNetGameLoadMapCompleted()
+{
+	// Attempt to read metadata if it exists
+	FArchive* const MetadataAr = ReplayStreamer.IsValid() ? ReplayStreamer->GetMetadataArchive() : nullptr;
+	UGameInstance* const GameInstance = GetWorld()->GetGameInstance();
+
+	if ( !GameInstance )
+	{
+		UE_LOG( LogDemo, Log, TEXT( "UDemoNetDriver::ReadStreamingLevelMetadata: GameInstance is null." ) );
+		return;
+	}
+
+	if ( MetadataAr != nullptr )
+	{
+		FNetworkDemoMetadataHeader MetadataHeader;
+
+		(*MetadataAr) << MetadataHeader;
+
+		FString Error;
+
+		// Check metadata magic value
+		if ( MetadataHeader.Magic != NETWORK_DEMO_METADATA_MAGIC )
+		{
+			Error = FString( TEXT( "Demo metadata file is corrupt" ) );
+			UE_LOG( LogDemo, Error, TEXT( "UDemoNetDriver::InitConnect: %s" ), *Error );
+			GameInstance->HandleDemoPlaybackFailure( EDemoPlayFailure::Corrupt, Error );
+			return;
+		}
+
+		// Check version
+		if ( MetadataHeader.Version != NETWORK_DEMO_METADATA_VERSION )
+		{
+			Error = FString( TEXT( "Demo metadata file version is incorrect" ) );
+			UE_LOG( LogDemo, Error, TEXT( "UDemoNetDriver::InitConnect: %s" ), *Error );
+			GameInstance->HandleDemoPlaybackFailure( EDemoPlayFailure::InvalidVersion, Error );
+			return;
+		}
+
+		DemoTotalFrames = MetadataHeader.NumFrames;
+		DemoTotalTime	= MetadataHeader.TotalTime;
+
+		// Read meta data, if it exists, after the world was loaded.
+		for ( int32 i = 0; i < MetadataHeader.NumStreamingLevels; ++i )
+		{
+			ULevelStreamingKismet* StreamingLevel = NewObject<ULevelStreamingKismet>( GetWorld(), NAME_None, RF_NoFlags, NULL );
+
+			StreamingLevel->bShouldBeLoaded = true;
+			StreamingLevel->bShouldBeVisible = true;
+			StreamingLevel->bShouldBlockOnLoad = false;
+			StreamingLevel->bInitiallyLoaded = true;
+			StreamingLevel->bInitiallyVisible = true;
+
+			FString PackageName;
+			FString PackageNameToLoad;
+
+			( *MetadataAr ) << PackageName;
+			( *MetadataAr ) << PackageNameToLoad;
+			( *MetadataAr ) << StreamingLevel->LevelTransform;
+
+			StreamingLevel->PackageNameToLoad = FName( *PackageNameToLoad );
+			StreamingLevel->SetWorldAssetByPackageName( FName( *PackageName ) );
+
+			GetWorld()->StreamingLevels.Add( StreamingLevel );
+
+			UE_LOG( LogDemo, Log, TEXT( "  Loading streamingLevel: %s, %s" ), *PackageName, *PackageNameToLoad );
+		}
+
+		UE_LOG( LogDemo, Log, TEXT( "Starting demo playback with full demo and metadata. Filename: %s, Frames: %i." ), *DemoURL.Map, DemoTotalFrames );
+	}
+	else
+	{
+		UE_LOG( LogDemo, Log, TEXT( "Starting demo playback with streaming demo, metadata file not found. Filename: %s." ), *DemoURL.Map );
+	}
+}
+
 /*-----------------------------------------------------------------------------
 	UDemoPendingNetGame.
 -----------------------------------------------------------------------------*/
@@ -2687,3 +2678,52 @@ void UDemoNetDriver::NotifyGotoTimeFinished(bool bWasSuccessful)
 UDemoPendingNetGame::UDemoPendingNetGame( const FObjectInitializer& ObjectInitializer ) : Super( ObjectInitializer )
 {
 }
+
+void UDemoPendingNetGame::Tick( float DeltaTime )
+{
+	// Replays don't need to do anything here
+}
+
+void UDemoPendingNetGame::SendJoin()
+{
+	// Don't send a join request to a replay
+}
+
+void UDemoPendingNetGame::LoadMapCompleted(UEngine* Engine, FWorldContext& Context, bool bLoadedMapSuccessfully, const FString& LoadMapError)
+{
+#if !( UE_BUILD_SHIPPING || UE_BUILD_TEST )
+	if ( CVarDemoForceFailure.GetValueOnGameThread() == 2 )
+	{
+		bLoadedMapSuccessfully = false;
+	}
+#endif
+
+	// If we have a demo pending net game we should have a demo net driver
+	check(DemoNetDriver);
+
+	if ( !bLoadedMapSuccessfully )
+	{
+		DemoNetDriver->StopDemo();
+
+		// If we don't have a world that means we failed loading the new world.
+		// Since there is no world, we must free the net driver ourselves
+		// Technically the pending net game should handle it, but things aren't quite setup properly to handle that either
+		if ( Context.World() == NULL )
+		{
+			GEngine->DestroyNamedNetDriver( Context.PendingNetGame, DemoNetDriver->NetDriverName );
+		}
+
+		Context.PendingNetGame = NULL;
+
+		GEngine->BrowseToDefaultMap( Context );
+
+		UE_LOG( LogDemo, Error, TEXT( "UDemoPendingNetGame::HandlePostLoadMap: LoadMap failed: %s" ), *LoadMapError );
+		if ( Context.OwningGameInstance )
+		{
+			Context.OwningGameInstance->HandleDemoPlaybackFailure( EDemoPlayFailure::Generic, FString( TEXT( "LoadMap failed" ) ) );
+		}
+		return;
+	}
+	
+	DemoNetDriver->PendingNetGameLoadMapCompleted();
+}
diff --git a/Engine/Source/Runtime/Engine/Private/GameViewportClient.cpp b/Engine/Source/Runtime/Engine/Private/GameViewportClient.cpp
--- a/Engine/Source/Runtime/Engine/Private/GameViewportClient.cpp
+++ b/Engine/Source/Runtime/Engine/Private/GameViewportClient.cpp
@@ -35,6 +35,7 @@
 #include "IMovieSceneCapture.h"
 #include "MovieSceneCaptureSettings.h"
 #include "ActorEditorUtils.h"
+#include "ComponentRecreateRenderStateContext.h"
 
 #define LOCTEXT_NAMESPACE "GameViewport"
 
@@ -57,9 +58,6 @@ FSimpleMulticastDelegate UGameViewportClient::CreatedDelegate;
 /** A list of all the stat names which are enabled for this viewport (static so they persist between runs) */
 TArray<FString> UGameViewportClient::EnabledStats;
 
-/** The number of GameViewportClients which have enabled 'show collision' */
-int32 UGameViewportClient::NumViewportsShowingCollision = 0;
-
 /** Those sound stat flags which are enabled on this viewport */
 FViewportClient::ESoundShowFlags::Type UGameViewportClient::SoundShowFlags = FViewportClient::ESoundShowFlags::Disabled;
 
@@ -1376,7 +1374,7 @@ void UGameViewportClient::LostFocus(FViewport* InViewport)
 	// We need to reset some key inputs, since keyup events will sometimes not be processed (such as going into immersive/maximized mode).  
 	// Resetting them will prevent them from "sticking"
 	UWorld* const ViewportWorld = GetWorld();
-	if (ViewportWorld)
+	if (ViewportWorld && !ViewportWorld->bIsTearingDown)
 	{
 		for (FConstPlayerControllerIterator Iterator = ViewportWorld->GetPlayerControllerIterator(); Iterator; ++Iterator)
 		{
@@ -2173,26 +2171,6 @@ bool UGameViewportClient::HandleForceFullscreenCommand( const TCHAR* Cmd, FOutpu
 	return true;
 }
 
-/** Contains the previous state of a primitive before turning on collision visibility */
-struct CollVisibilityState
-{
-	bool bHiddenInGame;
-	bool bVisible;
-
-	CollVisibilityState(bool InHidden, bool InVisible) :
-		bHiddenInGame(InHidden),
-		bVisible(InVisible)
-	{
-	}
-};
-
-typedef TMap<TWeakObjectPtr<UPrimitiveComponent>, CollVisibilityState> CollisionComponentVisibilityMap;
-CollisionComponentVisibilityMap& GetCollisionComponentVisibilityMap()
-{
-	static CollisionComponentVisibilityMap Mapping;
-	return Mapping;
-}
-
 bool UGameViewportClient::HandleShowCommand( const TCHAR* Cmd, FOutputDevice& Ar, UWorld* InWorld )
 {
 #if UE_BUILD_SHIPPING
@@ -2384,62 +2362,17 @@ void UGameViewportClient::ToggleShowCollision()
 			EngineShowFlags.SetVolumes(false);
 			ToggleShowVolumes();
 		}
-
-		NumViewportsShowingCollision++;
-		if (World != nullptr)
-		{
-			ShowCollisionOnSpawnedActorsDelegateHandle = World->AddOnActorSpawnedHandler(FOnActorSpawned::FDelegate::CreateUObject(this, &UGameViewportClient::ShowCollisionOnSpawnedActors));
-		}
 	}
-	else
-	{
-		NumViewportsShowingCollision--;
-		check(NumViewportsShowingCollision >= 0);
-		if (World != nullptr)
-		{
-			World->RemoveOnActorSpawnedHandler(ShowCollisionOnSpawnedActorsDelegateHandle);
-		}
-	}
-
-	CollisionComponentVisibilityMap& Mapping = GetCollisionComponentVisibilityMap();
 
-	// Restore state to any object in the map above
-	for (CollisionComponentVisibilityMap::TIterator It(Mapping); It; ++It)
+	if (World != nullptr)
 	{
-		TWeakObjectPtr<UPrimitiveComponent>& PrimitiveComponent = It.Key();
-		if (PrimitiveComponent.IsValid())
-		{
-			const CollVisibilityState& VisState = It.Value();
-			PrimitiveComponent->SetHiddenInGame(VisState.bHiddenInGame);
-			PrimitiveComponent->SetVisibility(VisState.bVisible);
-		}
-	}
-	Mapping.Empty();
+		// Tell engine to create proxies for hidden components, so we can still draw collision
+		World->bCreateRenderStateForHiddenComponents = bIsShowingCollision;
 
-	if (World == nullptr)
-	{
-		return;
+		// Need to recreate scene proxies when this flag changes.
+		FGlobalComponentRecreateRenderStateContext Recreate;
 	}
 
-	if (NumViewportsShowingCollision > 0)
-	{
-		for (TObjectIterator<UPrimitiveComponent> It; It; ++It)
-		{
-			UPrimitiveComponent* PrimitiveComponent = *It;
-			if (!PrimitiveComponent->IsVisible() && PrimitiveComponent->IsCollisionEnabled() && PrimitiveComponent->GetScene() == World->Scene)
-			{
-				AActor* Owner = PrimitiveComponent->GetOwner();
-
-				if (Owner && Owner->GetWorld() && Owner->GetWorld()->IsGameWorld() && !FActorEditorUtils::IsABuilderBrush(Owner))
-				{
-					// Save state before modifying the collision visibility
-					Mapping.Add(PrimitiveComponent, CollVisibilityState(PrimitiveComponent->bHiddenInGame, PrimitiveComponent->bVisible));
-					PrimitiveComponent->SetHiddenInGame(false);
-					PrimitiveComponent->SetVisibility(true);
-				}
-			}
-		}
-	}
 
 #if !(UE_BUILD_SHIPPING || UE_BUILD_TEST)
 	if (EngineShowFlags.Collision)
@@ -2467,28 +2400,6 @@ void UGameViewportClient::ToggleShowCollision()
 #endif
 }
 
-void UGameViewportClient::ShowCollisionOnSpawnedActors(AActor* Actor)
-{
-	CollisionComponentVisibilityMap& Mapping = GetCollisionComponentVisibilityMap();
-
-	TInlineComponentArray<UPrimitiveComponent*> Components;
-	check(Actor != nullptr);
-	Actor->GetComponents(Components);
-
-	for (auto Component : Components)
-	{
-		if (!Mapping.Contains(Component) && !Component->IsVisible() && Component->IsCollisionEnabled() && Component->GetScene() == GetWorld()->Scene)
-		{
-			check(Component->GetOwner() && Component->GetOwner()->GetWorld() && Component->GetOwner()->GetWorld()->IsGameWorld());
-
-			// Save state before modifying the collision visibility
-			Mapping.Add(Component, CollVisibilityState(Component->bHiddenInGame, Component->bVisible));
-			Component->SetHiddenInGame(false);
-			Component->SetVisibility(true);
-		}
-	}
-}
-
 bool UGameViewportClient::HandleShowLayerCommand( const TCHAR* Cmd, FOutputDevice& Ar, UWorld* InWorld )
 {
 	FString LayerName = FParse::Token(Cmd, 0);
diff --git a/Engine/Source/Runtime/Engine/Private/Materials/MaterialExpressions.cpp b/Engine/Source/Runtime/Engine/Private/Materials/MaterialExpressions.cpp
--- a/Engine/Source/Runtime/Engine/Private/Materials/MaterialExpressions.cpp
+++ b/Engine/Source/Runtime/Engine/Private/Materials/MaterialExpressions.cpp
@@ -140,6 +140,7 @@
 #include "Materials/MaterialExpressionDistanceFieldGradient.h"
 #include "Materials/MaterialFunction.h"
 #include "Materials/MaterialParameterCollection.h"
+#include "Materials/MaterialExpressionClearCoatNormalCustomOutput.h"
 
 #include "EditorSupportDelegates.h"
 #include "MaterialCompiler.h"
@@ -10062,4 +10063,55 @@ void UMaterialExpressionTangentOutput::GetCaption(TArray<FString>& OutCaptions)
 }
 #endif // WITH_EDITOR
 
+///////////////////////////////////////////////////////////////////////////////
+// Clear Coat Custom Normal Input
+///////////////////////////////////////////////////////////////////////////////
+
+UMaterialExpressionClearCoatNormalCustomOutput::UMaterialExpressionClearCoatNormalCustomOutput(const FObjectInitializer& ObjectInitializer)
+: Super(ObjectInitializer)
+{
+	// Structure to hold one-time initialization
+	struct FConstructorStatics
+	{
+		FText NAME_Utility;
+		FConstructorStatics()
+			: NAME_Utility(LOCTEXT("Utility", "Utility"))
+		{
+		}
+	};
+	static FConstructorStatics ConstructorStatics;
+
+	MenuCategories.Add(ConstructorStatics.NAME_Utility);
+	bCollapsed = true;
+
+	// No outputs
+	Outputs.Reset();
+}
+
+#if WITH_EDITOR
+int32  UMaterialExpressionClearCoatNormalCustomOutput::Compile(class FMaterialCompiler* Compiler, int32 OutputIndex, int32 MultiplexIndex)
+{
+	if (Input.Expression)
+	{
+		return Compiler->CustomOutput(this, OutputIndex, Input.Compile(Compiler, MultiplexIndex));
+	}
+	else
+	{
+		return CompilerError(Compiler, TEXT("Input missing"));
+	}
+	return INDEX_NONE;
+}
+
+
+void UMaterialExpressionClearCoatNormalCustomOutput::GetCaption(TArray<FString>& OutCaptions) const
+{
+	OutCaptions.Add(FString(TEXT("ClearCoatBottomNormal")));
+}
+#endif // WITH_EDITOR
+
+FExpressionInput* UMaterialExpressionClearCoatNormalCustomOutput::GetInput(int32 InputIndex)
+{
+	return &Input;
+}
+
 #undef LOCTEXT_NAMESPACE
diff --git a/Engine/Source/Runtime/Engine/Private/Particles/ParticleComponents.cpp b/Engine/Source/Runtime/Engine/Private/Particles/ParticleComponents.cpp
--- a/Engine/Source/Runtime/Engine/Private/Particles/ParticleComponents.cpp
+++ b/Engine/Source/Runtime/Engine/Private/Particles/ParticleComponents.cpp
@@ -5720,10 +5720,15 @@ void UParticleSystemComponent::InitializeSystem()
 		// Don't warn in a commandlet, we're expected not to have a scene / FX system.
 		if (!IsRunningCommandlet() && !IsRunningDedicatedServer())
 		{
-			UE_LOG(LogParticles,Warning,TEXT("InitializeSystem called on an unregistered component. Template=%s Component=%s"),
-				Template ? *Template->GetPathName() : TEXT("NULL"),
-				*GetPathName()
-				);
+			// We're also not expected to have a scene / FX system when we belong to an inactive world.
+			UWorld* OwnerWorld = GetWorld();
+			if (!OwnerWorld || OwnerWorld->WorldType != EWorldType::Inactive)
+			{
+				UE_LOG(LogParticles,Warning,TEXT("InitializeSystem called on an unregistered component. Template=%s Component=%s"),
+					Template ? *Template->GetPathName() : TEXT("NULL"),
+					*GetPathName()
+					);
+			}
 		}
 		return;
 	}
diff --git a/Engine/Source/Runtime/Engine/Private/PendingNetGame.cpp b/Engine/Source/Runtime/Engine/Private/PendingNetGame.cpp
--- a/Engine/Source/Runtime/Engine/Private/PendingNetGame.cpp
+++ b/Engine/Source/Runtime/Engine/Private/PendingNetGame.cpp
@@ -100,6 +100,28 @@ void UPendingNetGame::AddReferencedObjects(UObject* InThis, FReferenceCollector&
 	Super::AddReferencedObjects( This, Collector );
 }
 
+void UPendingNetGame::LoadMapCompleted(UEngine* Engine, FWorldContext& Context, bool bLoadedMapSuccessfully, const FString& LoadMapError)
+{
+	if (!bLoadedMapSuccessfully || LoadMapError != TEXT(""))
+	{
+		// we can't guarantee the current World is in a valid state, so travel to the default map
+		Engine->BrowseToDefaultMap(Context);
+		Engine->BroadcastTravelFailure(Context.World(), ETravelFailure::LoadMapFailure, LoadMapError);
+		check(Context.World() != NULL);
+	}
+	else
+	{
+		// Show connecting message, cause precaching to occur.
+		Engine->TransitionType = TT_Connecting;
+
+		Engine->RedrawViewports();
+
+		// Send join.
+		Context.PendingNetGame->SendJoin();
+		Context.PendingNetGame->NetDriver = NULL;
+	}
+}
+
 EAcceptConnection::Type UPendingNetGame::NotifyAcceptingConnection()
 {
 	return EAcceptConnection::Reject; 
diff --git a/Engine/Source/Runtime/Engine/Private/PlayerCameraManager.cpp b/Engine/Source/Runtime/Engine/Private/PlayerCameraManager.cpp
--- a/Engine/Source/Runtime/Engine/Private/PlayerCameraManager.cpp
+++ b/Engine/Source/Runtime/Engine/Private/PlayerCameraManager.cpp
@@ -592,13 +592,12 @@ UCameraModifier* APlayerCameraManager::FindCameraModifierByClass(TSubclassOf<UCa
 	return nullptr;
 }
 
-
 bool APlayerCameraManager::AddCameraModifierToList(UCameraModifier* NewModifier)
 {
 	if (NewModifier)
 	{
 		// Look through current modifier list and find slot for this priority
-		int32 BestIdx = 0;
+		int32 BestIdx = ModifierList.Num();
 		for (int32 ModifierIdx = 0; ModifierIdx < ModifierList.Num(); ModifierIdx++)
 		{
 			UCameraModifier* const M = ModifierList[ModifierIdx];
@@ -619,12 +618,12 @@ bool APlayerCameraManager::AddCameraModifierToList(UCameraModifier* NewModifier)
 						return false;
 					}
 
+					// Update best index
+					BestIdx = ModifierIdx;
+
 					break;
 				}
 			}
-
-			// Update best index
-			BestIdx = ModifierIdx;
 		}
 
 		// Insert self into best index
diff --git a/Engine/Source/Runtime/Engine/Private/ShaderCompiler/ShaderCompiler.cpp b/Engine/Source/Runtime/Engine/Private/ShaderCompiler/ShaderCompiler.cpp
--- a/Engine/Source/Runtime/Engine/Private/ShaderCompiler/ShaderCompiler.cpp
+++ b/Engine/Source/Runtime/Engine/Private/ShaderCompiler/ShaderCompiler.cpp
@@ -2301,6 +2301,11 @@ void GlobalBeginCompileShader(
 		}
 	}
 
+	{
+		static const auto CVar = IConsoleManager::Get().FindTConsoleVariableDataInt(TEXT("r.ClearCoatNormal"));
+		Input.Environment.SetDefine(TEXT("CLEAR_COAT_BOTTOM_NORMAL"), CVar ? (CVar->GetValueOnGameThread() != 0) : 0);
+	}
+
 	{
 		static const auto CVar = IConsoleManager::Get().FindTConsoleVariableDataInt(TEXT("Compat.UseDXT5NormalMaps"));
 		Input.Environment.SetDefine(TEXT("DXT5_NORMALMAPS"), CVar ? (CVar->GetValueOnGameThread() != 0) : 0);
diff --git a/Engine/Source/Runtime/Engine/Private/SimpleConstructionScript.cpp b/Engine/Source/Runtime/Engine/Private/SimpleConstructionScript.cpp
--- a/Engine/Source/Runtime/Engine/Private/SimpleConstructionScript.cpp
+++ b/Engine/Source/Runtime/Engine/Private/SimpleConstructionScript.cpp
@@ -896,9 +896,10 @@ USceneComponent* USimpleConstructionScript::GetSceneRootComponentTemplate(USCS_N
 
 		for(int32 StackIndex = 0; StackIndex < BPStack.Num(); ++StackIndex)
 		{
-			if(BPStack[StackIndex] && BPStack[StackIndex]->SimpleConstructionScript)
+			if(BPStack[StackIndex] && BPStack[StackIndex]->SimpleConstructionScript && !SCSStack.Contains(BPStack[StackIndex]->SimpleConstructionScript))
 			{
-				SCSStack.AddUnique(BPStack[StackIndex]->SimpleConstructionScript);
+				// UBlueprint::GetBlueprintHierarchyFromClass returns first children then parents. So we need to revert the order.
+				SCSStack.Insert(BPStack[StackIndex]->SimpleConstructionScript, 0);
 			}
 		}
 
diff --git a/Engine/Source/Runtime/Engine/Private/Slate/SceneViewport.cpp b/Engine/Source/Runtime/Engine/Private/Slate/SceneViewport.cpp
--- a/Engine/Source/Runtime/Engine/Private/Slate/SceneViewport.cpp
+++ b/Engine/Source/Runtime/Engine/Private/Slate/SceneViewport.cpp
@@ -1364,7 +1364,7 @@ void FSceneViewport::EnqueueBeginRenderFrame()
 	// If we dont have the ViewportRHI then we need to get it before rendering
 	// Note, we need ViewportRHI even if UseSeparateRenderTarget() is true when stereo rendering
 	// is enabled.
-	if( !IsValidRef(ViewportRHI) && (!UseSeparateRenderTarget() || (GEngine->StereoRenderingDevice.IsValid())) )
+	if (!IsValidRef(ViewportRHI) && (!UseSeparateRenderTarget() || (GEngine->StereoRenderingDevice.IsValid() && GEngine->StereoRenderingDevice->IsStereoEnabled())))
 	{
 		// Get the viewport for this window from the renderer so we can render directly to the backbuffer
 		TSharedPtr<FSlateRenderer> Renderer = FSlateApplication::Get().GetRenderer();
diff --git a/Engine/Source/Runtime/Engine/Private/UnrealEngine.cpp b/Engine/Source/Runtime/Engine/Private/UnrealEngine.cpp
--- a/Engine/Source/Runtime/Engine/Private/UnrealEngine.cpp
+++ b/Engine/Source/Runtime/Engine/Private/UnrealEngine.cpp
@@ -1443,6 +1443,20 @@ void LoadSpecialMaterial(const FString& MaterialName, UMaterial*& Material, bool
 }
 
 
+template<typename ClassType>
+void LoadEngineClass(const FStringClassReference& ClassName, TSubclassOf<ClassType>& EngineClassRef)
+{
+	if ( EngineClassRef == nullptr )
+	{
+		EngineClassRef = LoadClass<ClassType>(nullptr, *ClassName.ToString());
+		if (EngineClassRef == nullptr)
+		{
+			EngineClassRef = ClassType::StaticClass();
+			UE_LOG(LogEngine, Error, TEXT("Failed to load '%s', falling back to '%s'"), *ClassName.ToString(), *EngineClassRef->GetName());
+		}
+	}
+}
+
 /**
  * Loads all Engine object references from their corresponding config entries.
  */
@@ -1516,9 +1530,9 @@ void UEngine::InitializeObjectReferences()
 	}
 
 	if( HighFrequencyNoiseTexture == NULL )
-	{
+		{
 		HighFrequencyNoiseTexture = LoadObject<UTexture2D>(NULL, *HighFrequencyNoiseTextureName.ToString(), NULL, LOAD_None, NULL);
-	}
+		}
 
 	if( DefaultBokehTexture == NULL )
 	{
@@ -1552,52 +1566,15 @@ void UEngine::InitializeObjectReferences()
 		checkf(DefaultPhysMaterial != NULL, TEXT("The default material (%s) is not found. Please make sure you have default material set up correctly."), *DefaultPhysMaterialName.ToString());
 	}
 
-	if ( ConsoleClass == NULL )
-	{
-		ConsoleClass = LoadClass<UConsole>(NULL, *ConsoleClassName.ToString(), NULL, LOAD_None, NULL);
-	}
-
-	if ( GameViewportClientClass == NULL )
-	{
-		GameViewportClientClass = LoadClass<UGameViewportClient>(NULL, *GameViewportClientClassName.ToString(), NULL, LOAD_None, NULL);
-
-		checkf(GameViewportClientClass != NULL, TEXT("Engine config value GameViewportClientClassName is not a valid class name."));
-	}
-
-	if ( LocalPlayerClass == NULL )
-	{
-		LocalPlayerClass = LoadClass<ULocalPlayer>(NULL, *LocalPlayerClassName.ToString(), NULL, LOAD_None, NULL);
-	}
-
-	if ( WorldSettingsClass == NULL )
-	{
-		WorldSettingsClass = LoadClass<AWorldSettings>(NULL, *WorldSettingsClassName.ToString(), NULL, LOAD_None, NULL);
-	}
-
-	if ( NavigationSystemClass == NULL )
-	{
-		NavigationSystemClass = LoadClass<UNavigationSystem>(NULL, *NavigationSystemClassName.ToString(), NULL, LOAD_None, NULL);
-	}
-
-	if ( AvoidanceManagerClass == NULL )
-	{
-		AvoidanceManagerClass = LoadClass<UAvoidanceManager>(NULL, *AvoidanceManagerClassName.ToString(), NULL, LOAD_None, NULL);
-	}
-
-	if ( PhysicsCollisionHandlerClass == NULL )
-	{
-		PhysicsCollisionHandlerClass = LoadClass<UPhysicsCollisionHandler>(NULL, *PhysicsCollisionHandlerClassName.ToString(), NULL, LOAD_None, NULL);
-	}
-
-	if ( GameUserSettingsClass == NULL )
-	{
-		GameUserSettingsClass = LoadClass<UGameUserSettings>(NULL, *GameUserSettingsClassName.ToString(), NULL, LOAD_None, NULL);
-	}
-
-	if ( LevelScriptActorClass == NULL )
-	{
-		LevelScriptActorClass = LoadClass<ALevelScriptActor>(NULL, *LevelScriptActorClassName.ToString(), NULL, LOAD_None, NULL);
-	}
+	LoadEngineClass<UConsole>(ConsoleClassName, ConsoleClass);
+	LoadEngineClass<UGameViewportClient>(GameViewportClientClassName, GameViewportClientClass);
+	LoadEngineClass<ULocalPlayer>(LocalPlayerClassName, LocalPlayerClass);
+	LoadEngineClass<AWorldSettings>(WorldSettingsClassName, WorldSettingsClass);
+	LoadEngineClass<UNavigationSystem>(NavigationSystemClassName, NavigationSystemClass);
+	LoadEngineClass<UAvoidanceManager>(AvoidanceManagerClassName, AvoidanceManagerClass);
+	LoadEngineClass<UPhysicsCollisionHandler>(PhysicsCollisionHandlerClassName, PhysicsCollisionHandlerClass);
+	LoadEngineClass<UGameUserSettings>(GameUserSettingsClassName, GameUserSettingsClass);
+	LoadEngineClass<ALevelScriptActor>(LevelScriptActorClassName, LevelScriptActorClass);
 
 	// set the font object pointers, unless on server
 	if (!IsRunningDedicatedServer())
@@ -9221,24 +9198,7 @@ void UEngine::TickWorldTravel(FWorldContext& Context, float DeltaSeconds)
 
 				const bool bLoadedMapSuccessfully = LoadMap(Context, Context.PendingNetGame->URL, Context.PendingNetGame, Error);
 
-				if (!bLoadedMapSuccessfully || Error != TEXT(""))
-				{
-					// we can't guarantee the current World is in a valid state, so travel to the default map
-					BrowseToDefaultMap(Context);
-					BroadcastTravelFailure(Context.World(), ETravelFailure::LoadMapFailure, Error);
-					check(Context.World() != NULL);
-				}
-				else
-				{
-					// Show connecting message, cause precaching to occur.
-					TransitionType = TT_Connecting;
-
-					RedrawViewports();
-
-					// Send join.
-					Context.PendingNetGame->SendJoin();
-					Context.PendingNetGame->NetDriver = NULL;
-				}
+				Context.PendingNetGame->LoadMapCompleted(this, Context, bLoadedMapSuccessfully, Error);
 
 				// Kill the pending level.
 				Context.PendingNetGame = NULL;
diff --git a/Engine/Source/Runtime/Engine/Private/UserInterface/Console.cpp b/Engine/Source/Runtime/Engine/Private/UserInterface/Console.cpp
--- a/Engine/Source/Runtime/Engine/Private/UserInterface/Console.cpp
+++ b/Engine/Source/Runtime/Engine/Private/UserInterface/Console.cpp
@@ -694,6 +694,64 @@ bool UConsole::InputKey_InputLine( int32 ControllerId, FKey Key, EInputEvent Eve
 		return true;
 	}
 
+	auto DecrementCursor = [this]()
+	{
+		if (AutoCompleteCursor > 0)
+		{
+			// move cursor within displayed region
+			--AutoCompleteCursor;
+		}
+		else
+		{
+			// can we scroll?
+			if (AutoCompleteIndex > 0)
+			{
+				--AutoCompleteIndex;
+			}
+			else
+			{
+				// wrap around
+				AutoCompleteIndex = FMath::Max(0, AutoComplete.Num() - (int32)MAX_AUTOCOMPLETION_LINES - 1);
+				if (AutoComplete.Num() <= MAX_AUTOCOMPLETION_LINES)
+				{
+					AutoCompleteCursor = AutoComplete.Num() + AutoCompleteCursor - 1;
+				}
+				else
+				{
+					// skip the "x more matches" line when wrapping
+					AutoCompleteIndex++;
+					AutoCompleteCursor = MAX_AUTOCOMPLETION_LINES + AutoCompleteCursor - 1;
+				}
+			}
+			bAutoCompleteLocked = false;
+		}
+	};
+
+	auto IncrementCursor = [this]()
+	{
+		if (AutoCompleteCursor + 1 < FMath::Min((int32)MAX_AUTOCOMPLETION_LINES, AutoComplete.Num()))
+		{
+			// move cursor within displayed region
+			++AutoCompleteCursor;
+		}
+		else
+		{
+			// can be negative
+			int32 ScrollRegionSize = AutoComplete.Num() - (int32)MAX_AUTOCOMPLETION_LINES;
+
+			// can we scroll?
+			if (AutoCompleteIndex < ScrollRegionSize)
+			{
+				++AutoCompleteIndex;
+			}
+			else
+			{
+				// wrap around
+				AutoCompleteIndex = AutoCompleteCursor = 0;
+			}
+		}
+	};
+
 	// if user input is open
 	if(ConsoleState != NAME_None)
 	{
@@ -779,34 +837,13 @@ bool UConsole::InputKey_InputLine( int32 ControllerId, FKey Key, EInputEvent Eve
 			{
 				if(AutoComplete.Num())
 				{
-					if (AutoCompleteCursor > 0)
+					if (ConsoleSettings->bOrderTopToBottom)
 					{
-						// move cursor within displayed region
-						--AutoCompleteCursor;
+						DecrementCursor();
 					}
 					else
 					{
-						// can we scroll?
-						if(AutoCompleteIndex > 0)
-						{
-							--AutoCompleteIndex;
-						}
-						else
-						{
-							// wrap around
-							AutoCompleteIndex = FMath::Max(0, AutoComplete.Num() - (int32)MAX_AUTOCOMPLETION_LINES - 1);
-							if (AutoComplete.Num() <= MAX_AUTOCOMPLETION_LINES)
-							{
-								AutoCompleteCursor = AutoComplete.Num() + AutoCompleteCursor - 1;
-							}
-							else
-							{
-								// skip the "x more matches" line when wrapping
-								AutoCompleteIndex++;
-								AutoCompleteCursor = MAX_AUTOCOMPLETION_LINES + AutoCompleteCursor - 1;
-							}
-						}
-						bAutoCompleteLocked = false;
+						IncrementCursor();
 					}
 				}
 				else
@@ -818,42 +855,36 @@ bool UConsole::InputKey_InputLine( int32 ControllerId, FKey Key, EInputEvent Eve
 		}
 		else if( Key == EKeys::Down || Key == EKeys::Tab )
 		{
-			if (!bCtrl && AutoComplete.Num())
+			if (!bCtrl)
 			{
-				bool bScroll = AutoComplete.Num() > 1;
-
-				if (Key == EKeys::Tab)
+				if (AutoComplete.Num())
 				{
-					bCaptureKeyInput = true;
+					bool bScroll = AutoComplete.Num() > 1;
 
-					// If this is a repeated tab press, we want to scroll. Otherwise complete the current command
-					bScroll = bScroll && LastAutoCompletedCommand == TypedStr;
-				}
-
-				if (bScroll)
-				{
-					if (AutoCompleteCursor + 1 < FMath::Min((int32)MAX_AUTOCOMPLETION_LINES, AutoComplete.Num()))
+					if (Key == EKeys::Tab)
 					{
-						// move cursor within displayed region
-						++AutoCompleteCursor;
+						bCaptureKeyInput = true;
+
+						// If this is a repeated tab press, we want to scroll. Otherwise complete the current command
+						bScroll = bScroll && LastAutoCompletedCommand == TypedStr;
 					}
-					else
-					{
-						// can be negative
-						int32 ScrollRegionSize = AutoComplete.Num() - (int32)MAX_AUTOCOMPLETION_LINES;
 
-						// can we scroll?
-						if (AutoCompleteIndex < ScrollRegionSize)
+					if (bScroll)
+					{
+						if (ConsoleSettings->bOrderTopToBottom)
 						{
-							++AutoCompleteIndex;
+							IncrementCursor();
 						}
 						else
 						{
-							// wrap around
-							AutoCompleteIndex = AutoCompleteCursor = 0;
+							DecrementCursor();
 						}
 					}
 				}
+				else
+				{
+					SetAutoCompleteFromHistory();
+				}
 
 				SetInputLineFromAutoComplete();
 			}
@@ -1331,10 +1362,8 @@ void UConsole::PostRender_InputLine(UCanvas* Canvas, FIntPoint UserInputLinePos)
 		Canvas->DrawItem(ConsoleOutline, UserInputLinePos.X + xl - Border, y + yl - Height - Border);
 
 		// auto completion elements
-		for(int32 Num = AutoCompleteElements.Num(), i = Num - 1; i >= 0; --i)
+		auto DrawElement = [&](const FAutoCompleteCommand& AutoCompleteElement, int32 i, int32 Num)
 		{
-			const FAutoCompleteCommand& AutoCompleteElement = *AutoCompleteElements[i];
-
 			const bool bCursorLineColor = (i == AutoCompleteCursor);
 			const bool bMoreMatches = (Num > MAX_AUTOCOMPLETION_LINES && i == Num - 1);
 			const bool bHistory = AutoCompleteElement.IsHistory();
@@ -1390,6 +1419,21 @@ void UConsole::PostRender_InputLine(UCanvas* Canvas, FIntPoint UserInputLinePos)
 			ConsoleText.Text = FText::FromString(AutoCompleteElement.GetRight());
 			Canvas->DrawItem( ConsoleText, DescriptionX, y );
 			y -= yl;
+		};
+
+		if (ConsoleSettings->bOrderTopToBottom)
+		{
+			for (int32 Num = AutoCompleteElements.Num(), i = Num - 1; i >= 0; --i)
+			{
+				DrawElement(*AutoCompleteElements[i], i, Num);
+			}
+		}
+		else
+		{
+			for (int32 Num = AutoCompleteElements.Num(), i = 0; i < Num; ++i)
+			{
+				DrawElement(*AutoCompleteElements[i], i, Num);
+			}
 		}
 	}
 
