Copying //UE4/Dev-Framework to //UE4/Dev-Main (Source: //UE4/Dev-Framework @ 2964666)

diff --git a/Engine/Source/Runtime/Engine/Classes/Animation/AnimSequence.h b/Engine/Source/Runtime/Engine/Classes/Animation/AnimSequence.h
--- a/Engine/Source/Runtime/Engine/Classes/Animation/AnimSequence.h
+++ b/Engine/Source/Runtime/Engine/Classes/Animation/AnimSequence.h
@@ -7,6 +7,7 @@
  *
  */
 
+#include "AnimTypes.h"
 #include "AnimSequenceBase.h"
 #include "AnimSequence.generated.h"
 
@@ -460,6 +461,10 @@ class ENGINE_API UAnimSequence : public UAnimSequenceBase
 	UPROPERTY(EditAnywhere, AssetRegistrySearchable, Category=Animation)
 	FName RetargetSource;
 
+	/** This defines how values between keys are calculated **/
+	UPROPERTY(EditAnywhere, AssetRegistrySearchable, Category = Animation)
+	EAnimInterpolationType Interpolation;
+	
 	/** If this is on, it will allow extracting of root motion **/
 	UPROPERTY(EditAnywhere, AssetRegistrySearchable, Category = RootMotion, meta = (DisplayName = "EnableRootMotion"))
 	bool bEnableRootMotion;
diff --git a/Engine/Source/Runtime/Engine/Classes/Components/SceneComponent.h b/Engine/Source/Runtime/Engine/Classes/Components/SceneComponent.h
--- a/Engine/Source/Runtime/Engine/Classes/Components/SceneComponent.h
+++ b/Engine/Source/Runtime/Engine/Classes/Components/SceneComponent.h
@@ -85,7 +85,6 @@ DECLARE_DYNAMIC_MULTICAST_DELEGATE_OneParam(FPhysicsVolumeChanged, class APhysic
  * Useful as a 'dummy' component in the hierarchy to offset others.
  * @see [Scene Components](https://docs.unrealengine.com/latest/INT/Programming/UnrealArchitecture/Actors/Components/index.html#scenecomponents)
  */
-PRAGMA_DISABLE_DEPRECATION_WARNINGS // Autogenerated destructor requires disabling of deprecation for the class DEPRECATED(4.12)
 UCLASS(ClassGroup=(Utility, Common), BlueprintType, hideCategories=(Trigger, PhysicsVolume), meta=(BlueprintSpawnableComponent, IgnoreCategoryKeywordsInSubclasses, ShortTooltip="A Scene Component is a component that has a scene transform and can be attached to other scene components."))
 class ENGINE_API USceneComponent : public UActorComponent
 {
@@ -93,14 +92,6 @@ class ENGINE_API USceneComponent : public UActorComponent
 
 public:
 
-	// Explicitly define the destructor while the deprecation warning pragmas are disabled so no
-	// warnings are thrown when used elsewhere and the compiler is forced to generate them
-#if PLATFORM_COMPILER_HAS_DEFAULTED_FUNCTIONS
-	virtual ~USceneComponent() = default;
-#else
-	virtual ~USceneComponent() { }
-#endif
-
 	/** The name to use for the default scene root variable */
 	static FName GetDefaultSceneRootVariableName();
 
@@ -112,20 +103,19 @@ public:
 	virtual void GetLifetimeReplicatedProps(TArray<FLifetimeProperty>& OutLifetimeProps) const override;
 
 	/** What we are currently attached to. If valid, RelativeLocation etc. are used relative to this object */
-	DEPRECATED(4.12, "Direct access of AttachParent has been deprecated. Use SetupAttachment, AttachToComponent, or GetAttachParent instead.")
+private:
 	UPROPERTY(ReplicatedUsing = OnRep_AttachParent)
 	USceneComponent* AttachParent;
 
 	/** List of child SceneComponents that are attached to us. */
-	DEPRECATED(4.12, "Direct access of AttachChildren has been deprecated. Use GetAttachChildren instead.")
 	UPROPERTY(Replicated, transient)
 	TArray<USceneComponent*> AttachChildren;
 
 	/** Optional socket name on AttachParent that we are attached to. */
-	DEPRECATED(4.12, "Direct access of AttachSocketName has been deprecated. Use SetupAttachment, AttachToComponent, or GetAttachSocketName instead.")
 	UPROPERTY(ReplicatedUsing = OnRep_AttachSocketName)
 	FName AttachSocketName;
 
+public:
 	DEPRECATED(4.11, "SceneComponent no longer supports custom locations.")
 	uint32 bRequiresCustomLocation:1;
 
@@ -542,14 +532,6 @@ public:
 	UFUNCTION(BlueprintCallable, Category="Physics")
 	virtual bool IsAnySimulatingPhysics() const;
 
-private:
-	// Temporary accessor functions to avoid deprecation messages, once variables
-	// are properly private likely these should be removed and the direct references
-	// replaced
-	TArray<USceneComponent*>& GetMutableAttachChildren();
-	USceneComponent*& GetMutableAttachParent();
-	FName& GetMutableAttachSocketName();
-
 public:
 
 	/** Get the SceneComponents that are attached to this component. */
@@ -1172,7 +1154,6 @@ private:
 
 	friend struct FDirectAttachChildrenAccessor;
 };
-PRAGMA_ENABLE_DEPRECATION_WARNINGS // DEPRECATED(4.12)
 
 /** 
   * Struct to allow direct access to the AttachChildren array for a handful of cases that will require more work than can be done  
@@ -1183,7 +1164,7 @@ struct FDirectAttachChildrenAccessor
 private:
 	static TArray<USceneComponent*>& Get(USceneComponent* Component)
 	{ 
-		return Component->GetMutableAttachChildren();
+		return Component->AttachChildren;
 	}
 
 	friend class UChildActorComponent;
@@ -1198,39 +1179,21 @@ private:
 //////////////////////////////////////////////////////////////////////////
 // USceneComponent inlines
 
-PRAGMA_DISABLE_DEPRECATION_WARNINGS // DEPRECATED(4.12)
-FORCEINLINE TArray<USceneComponent*>& USceneComponent::GetMutableAttachChildren()
-{
-	return AttachChildren;
-}
-
 FORCEINLINE const TArray<USceneComponent*>& USceneComponent::GetAttachChildren() const
 {
 	return AttachChildren;
 }
 
-FORCEINLINE USceneComponent*& USceneComponent::GetMutableAttachParent()
-{
-	return AttachParent;
-}
-
 FORCEINLINE USceneComponent* USceneComponent::GetAttachParent() const
 {
 	return AttachParent;
 }
 
-FORCEINLINE FName& USceneComponent::GetMutableAttachSocketName() 
-{
-	return AttachSocketName;
-}
-
 FORCEINLINE FName USceneComponent::GetAttachSocketName() const
 {
 	return AttachSocketName;
 }
 
-PRAGMA_ENABLE_DEPRECATION_WARNINGS // DEPRECATED(4.12)
-
 FORCEINLINE_DEBUGGABLE void USceneComponent::ConditionalUpdateComponentToWorld()
 {
 	if (!bWorldToComponentUpdated)
diff --git a/Engine/Source/Runtime/Engine/Classes/Engine/StaticMesh.h b/Engine/Source/Runtime/Engine/Classes/Engine/StaticMesh.h
--- a/Engine/Source/Runtime/Engine/Classes/Engine/StaticMesh.h
+++ b/Engine/Source/Runtime/Engine/Classes/Engine/StaticMesh.h
@@ -529,6 +529,7 @@ public:
 	/**
 	 * Returns the number of LODs used by the mesh.
 	 */
+	UFUNCTION(BlueprintCallable, Category = "StaticMesh")
 	ENGINE_API int32 GetNumLODs() const;
 
 	/**
@@ -548,6 +549,10 @@ public:
 	UFUNCTION(BlueprintCallable, Category="StaticMesh")
 	ENGINE_API FBox GetBoundingBox() const;
 
+	/** Returns number of Sections that this StaticMesh has, in the supplied LOD (LOD 0 is the highest) */
+	UFUNCTION(BlueprintCallable, Category = "StaticMesh")
+	ENGINE_API int32 GetNumSections(int32 InLOD) const;
+
 	/**
 	 * Gets a Material given a Material Index and an LOD number
 	 *
diff --git a/Engine/Source/Runtime/Engine/Classes/GameFramework/CharacterMovementComponent.h b/Engine/Source/Runtime/Engine/Classes/GameFramework/CharacterMovementComponent.h
--- a/Engine/Source/Runtime/Engine/Classes/GameFramework/CharacterMovementComponent.h
+++ b/Engine/Source/Runtime/Engine/Classes/GameFramework/CharacterMovementComponent.h
@@ -555,6 +555,12 @@ protected:
 	/** Computes the analog input modifier based on current input vector and/or acceleration. */
 	virtual float ComputeAnalogInputModifier() const;
 
+	/** Used for throttling "stuck in geometry" logging. */
+	float LastStuckWarningTime;
+
+	/** Used when throttling "stuck in geometry" logging, to output the number of events we skipped if throttling. */
+	uint32 StuckWarningCountSinceNotify;
+
 public:
 
 	/** Get the value of ServerLastTransformUpdateTimeStamp. */
@@ -593,6 +599,37 @@ public:
 	UPROPERTY(Category="Character Movement (General Settings)", EditAnywhere, BlueprintReadWrite, AdvancedDisplay, meta=(ClampMin="1", ClampMax="25", UIMin="1", UIMax="25"))
 	int32 MaxSimulationIterations;
 
+	/**
+	* Max distance we allow simulated proxies to depenetrate when moving out of anything but Pawns.
+	* This is generally more tolerant than with Pawns, because other geometry is either not moving, or is moving predictably with a bit of delay compared to on the server.
+	* @see MaxDepenetrationWithGeometryAsProxy, MaxDepenetrationWithPawn, MaxDepenetrationWithPawnAsProxy
+	*/
+	UPROPERTY(Category="Character Movement (General Settings)", EditAnywhere, BlueprintReadWrite, AdvancedDisplay, meta=(ClampMin="0", UIMin="0"))
+	float MaxDepenetrationWithGeometry;
+
+	/**
+	* Max distance we allow simulated proxies to depenetrate when moving out of anything but Pawns.
+	* This is generally more tolerant than with Pawns, because other geometry is either not moving, or is moving predictably with a bit of delay compared to on the server.
+	* @see MaxDepenetrationWithGeometry, MaxDepenetrationWithPawn, MaxDepenetrationWithPawnAsProxy
+	*/
+	UPROPERTY(Category="Character Movement (General Settings)", EditAnywhere, BlueprintReadWrite, AdvancedDisplay, meta=(ClampMin="0", UIMin="0"))
+	float MaxDepenetrationWithGeometryAsProxy;
+
+	/**
+	* Max distance we are allowed to depenetrate when moving out of other Pawns.
+	* @see MaxDepenetrationWithGeometry, MaxDepenetrationWithGeometryAsProxy, MaxDepenetrationWithPawnAsProxy
+	*/
+	UPROPERTY(Category="Character Movement (General Settings)", EditAnywhere, BlueprintReadWrite, AdvancedDisplay, meta=(ClampMin="0", UIMin="0"))
+	float MaxDepenetrationWithPawn;
+
+	/**
+	 * Max distance we allow simulated proxies to depenetrate when moving out of other Pawns.
+	 * Typically we don't want a large value, because we receive a server authoritative position that we should not then ignore by pushing them out of the local player.
+	 * @see MaxDepenetrationWithGeometry, MaxDepenetrationWithGeometryAsProxy, MaxDepenetrationWithPawn
+	 */
+	UPROPERTY(Category="Character Movement (General Settings)", EditAnywhere, BlueprintReadWrite, AdvancedDisplay, meta=(ClampMin="0", UIMin="0"))
+	float MaxDepenetrationWithPawnAsProxy;
+
 	/**
 	 * How long to take to smoothly interpolate from the old pawn position on the client to the corrected one sent by the server. Not used by Linear smoothing.
 	 */
@@ -1507,6 +1544,9 @@ protected:
 	 */
 	virtual void MaintainHorizontalGroundVelocity();
 
+	/** Overridden to enforce max distances based on hit geometry. */
+	virtual FVector GetPenetrationAdjustment(const FHitResult& Hit) const override;
+
 	/** Overridden to set bJustTeleported to true, so we don't make incorrect velocity calculations based on adjusted movement. */
 	virtual bool ResolvePenetrationImpl(const FVector& Adjustment, const FHitResult& Hit, const FQuat& NewRotation) override;
 
diff --git a/Engine/Source/Runtime/Engine/Classes/PhysicsEngine/BodySetup.h b/Engine/Source/Runtime/Engine/Classes/PhysicsEngine/BodySetup.h
--- a/Engine/Source/Runtime/Engine/Classes/PhysicsEngine/BodySetup.h
+++ b/Engine/Source/Runtime/Engine/Classes/PhysicsEngine/BodySetup.h
@@ -95,7 +95,7 @@ class UBodySetup : public UObject
 	UPROPERTY(EditAnywhere, Category=Collision, meta=(DisplayName = "Collision Complexity"))
 	TEnumAsByte<enum ECollisionTraceFlag> CollisionTraceFlag;
 
-	TEnumAsByte<enum ECollisionTraceFlag> GetCollisionTraceFlag() const;
+	ENGINE_API TEnumAsByte<enum ECollisionTraceFlag> GetCollisionTraceFlag() const;
 
 	/** Default properties of the body instance, copied into objects on instantiation, was URB_BodyInstance */
 	UPROPERTY(EditAnywhere, Category=Collision, meta=(FullyExpand = "true"))
diff --git a/Engine/Source/Runtime/Engine/Classes/PhysicsEngine/ConstraintInstance.h b/Engine/Source/Runtime/Engine/Classes/PhysicsEngine/ConstraintInstance.h
--- a/Engine/Source/Runtime/Engine/Classes/PhysicsEngine/ConstraintInstance.h
+++ b/Engine/Source/Runtime/Engine/Classes/PhysicsEngine/ConstraintInstance.h
@@ -239,7 +239,9 @@ struct ENGINE_API FConstraintInstance
 	float TwistLimitDamping;
 
 	/** Specifies the angular offset between the two frames of reference. By default limit goes from (-Angle, +Angle)
-	  * This allows you to bias the limit for swing1 swing2 and twist. */
+	  * This allows you to bias the limit for swing1 swing2 and twist.
+      * This is also used as the reference orientation for Angular Orientation Drive.
+      */
 	UPROPERTY(EditAnywhere, Category = Angular)
 	FRotator AngularRotationOffset;
 	
diff --git a/Engine/Source/Runtime/Engine/Classes/PhysicsEngine/ConvexElem.h b/Engine/Source/Runtime/Engine/Classes/PhysicsEngine/ConvexElem.h
--- a/Engine/Source/Runtime/Engine/Classes/PhysicsEngine/ConvexElem.h
+++ b/Engine/Source/Runtime/Engine/Classes/PhysicsEngine/ConvexElem.h
@@ -60,8 +60,11 @@ struct FKConvexElem : public FKShapeElem
 	/** Calculate a bounding box for this convex element with the specified transform and scale */
 	ENGINE_API FBox	CalcAABB(const FTransform& BoneTM, const FVector& Scale3D) const;
 
+	/** Get set of planes that define this convex hull */
+	ENGINE_API void GetPlanes(TArray<FPlane>& Planes) const;
+
 	/** Utility for creating a convex hull from a set of planes. Will reset current state of this elem. */
-	bool	HullFromPlanes(const TArray<FPlane>& InPlanes, const TArray<FVector>& SnapVerts);
+	ENGINE_API bool HullFromPlanes(const TArray<FPlane>& InPlanes, const TArray<FVector>& SnapVerts);
 
 	/** Returns the volume of this element */
 	float GetVolume(const FVector& Scale) const;
diff --git a/Engine/Source/Runtime/Engine/Private/AI/Navigation/NavigationSystem.cpp b/Engine/Source/Runtime/Engine/Private/AI/Navigation/NavigationSystem.cpp
--- a/Engine/Source/Runtime/Engine/Private/AI/Navigation/NavigationSystem.cpp
+++ b/Engine/Source/Runtime/Engine/Private/AI/Navigation/NavigationSystem.cpp
@@ -1203,11 +1203,18 @@ void UNavigationSystem::SimpleMoveToActor(AController* Controller, const AActor*
 		return;
 	}
 
-	if (PFollowComp->HasReached(*Goal))
+	const bool bAlreadyAtGoal = PFollowComp->HasReached(*Goal);
+
+	// script source, keep only one move request at time
+	if (PFollowComp->GetStatus() != EPathFollowingStatus::Idle)
 	{
-		// make sure previous move request gets aborted
-		PFollowComp->AbortMove(TEXT("Aborting move due to new move request finishing with AlreadyAtGoal"), FAIRequestID::AnyRequest);
-		PFollowComp->SetLastMoveAtGoal(true);
+		PFollowComp->AbortMove(*NavSys, FPathFollowingResultFlags::ForcedScript | FPathFollowingResultFlags::NewRequest
+			, FAIRequestID::AnyRequest, bAlreadyAtGoal ? EPathFollowingVelocityMode::Reset : EPathFollowingVelocityMode::Keep);
+	}
+
+	if (bAlreadyAtGoal)
+	{
+		PFollowComp->RequestMoveWithImmediateFinish(EPathFollowingResult::Success);
 	}
 	else
 	{
@@ -1219,13 +1226,11 @@ void UNavigationSystem::SimpleMoveToActor(AController* Controller, const AActor*
 			if (Result.IsSuccessful())
 			{
 				Result.Path->SetGoalActorObservation(*Goal, 100.0f);
-
-				PFollowComp->RequestMove(Result.Path, Goal);
+				PFollowComp->RequestMove(FAIMoveRequest(Goal), Result.Path);
 			}
 			else if (PFollowComp->GetStatus() != EPathFollowingStatus::Idle)
 			{
-				PFollowComp->AbortMove(TEXT("Aborting move due to new move request failing to generate a path"), FAIRequestID::AnyRequest);
-				PFollowComp->SetLastMoveAtGoal(false);
+				PFollowComp->RequestMoveWithImmediateFinish(EPathFollowingResult::Invalid);
 			}
 		}
 	}
@@ -1262,11 +1267,18 @@ void UNavigationSystem::SimpleMoveToLocation(AController* Controller, const FVec
 		return;
 	}
 
+	const bool bAlreadyAtGoal = PFollowComp->HasReached(GoalLocation);
+
+	// script source, keep only one move request at time
+	if (PFollowComp->GetStatus() != EPathFollowingStatus::Idle)
+	{
+		PFollowComp->AbortMove(*NavSys, FPathFollowingResultFlags::ForcedScript | FPathFollowingResultFlags::NewRequest
+			, FAIRequestID::AnyRequest, bAlreadyAtGoal ? EPathFollowingVelocityMode::Reset : EPathFollowingVelocityMode::Keep);
+	}
+
 	if (PFollowComp->HasReached(GoalLocation))
 	{
-		// make sure previous move request gets aborted
-		PFollowComp->AbortMove(TEXT("Aborting move due to new move request finishing with AlreadyAtGoal"), FAIRequestID::AnyRequest);
-		PFollowComp->SetLastMoveAtGoal(true);
+		PFollowComp->RequestMoveWithImmediateFinish(EPathFollowingResult::Success);
 	}
 	else
 	{
@@ -1277,12 +1289,11 @@ void UNavigationSystem::SimpleMoveToLocation(AController* Controller, const FVec
 			FPathFindingResult Result = NavSys->FindPathSync(Query);
 			if (Result.IsSuccessful())
 			{
-				PFollowComp->RequestMove(Result.Path, NULL);
+				PFollowComp->RequestMove(FAIMoveRequest(GoalLocation), Result.Path);
 			}
 			else if (PFollowComp->GetStatus() != EPathFollowingStatus::Idle)
 			{
-				PFollowComp->AbortMove(TEXT("Aborting move due to new move request failing to generate a path"), FAIRequestID::AnyRequest);
-				PFollowComp->SetLastMoveAtGoal(false);
+				PFollowComp->RequestMoveWithImmediateFinish(EPathFollowingResult::Invalid);
 			}
 		}
 	}
diff --git a/Engine/Source/Runtime/Engine/Private/AI/Navigation/RecastNavMesh.cpp b/Engine/Source/Runtime/Engine/Private/AI/Navigation/RecastNavMesh.cpp
--- a/Engine/Source/Runtime/Engine/Private/AI/Navigation/RecastNavMesh.cpp
+++ b/Engine/Source/Runtime/Engine/Private/AI/Navigation/RecastNavMesh.cpp
@@ -1074,7 +1074,7 @@ void ARecastNavMesh::BatchProjectPoints(TArray<FNavigationProjectionWork>& Workl
 	}
 }
 
-bool ARecastNavMesh::GetPolysInBox(const FBox& Box, TArray<FNavPoly>& Polys, FSharedConstNavQueryFilter Filter, const UObject* Owner) const
+bool ARecastNavMesh::GetPolysInBox(const FBox& Box, TArray<FNavPoly>& Polys, FSharedConstNavQueryFilter Filter, const UObject* InOwner) const
 {
 	// sanity check
 	if (RecastNavMeshImpl->GetRecastMesh() == NULL)
@@ -1085,7 +1085,7 @@ bool ARecastNavMesh::GetPolysInBox(const FBox& Box, TArray<FNavPoly>& Polys, FSh
 	bool bSuccess = false;
 
 	const FNavigationQueryFilter& FilterToUse = GetRightFilterRef(Filter);
-	FRecastSpeciaLinkFilter LinkFilter(UNavigationSystem::GetCurrent(GetWorld()), Owner);
+	FRecastSpeciaLinkFilter LinkFilter(UNavigationSystem::GetCurrent(GetWorld()), InOwner);
 	INITIALIZE_NAVQUERY_WLINKFILTER(NavQuery, FilterToUse.GetMaxSearchNodes(), LinkFilter);
 
 	const dtQueryFilter* QueryFilter = ((const FRecastQueryFilter*)(FilterToUse.GetImplementation()))->GetAsDetourQueryFilter();
diff --git a/Engine/Source/Runtime/Engine/Private/AnimBlueprint.cpp b/Engine/Source/Runtime/Engine/Private/AnimBlueprint.cpp
--- a/Engine/Source/Runtime/Engine/Private/AnimBlueprint.cpp
+++ b/Engine/Source/Runtime/Engine/Private/AnimBlueprint.cpp
@@ -138,14 +138,14 @@ void UAnimBlueprint::PostLoad()
 	Super::PostLoad();
 #if WITH_EDITOR
 	// Validate animation overrides
-	UAnimBlueprintGeneratedClass* GeneratedClass = GetAnimBlueprintGeneratedClass();
+	UAnimBlueprintGeneratedClass* AnimBPGeneratedClass = GetAnimBlueprintGeneratedClass();
 	
-	if (GeneratedClass)
+	if (AnimBPGeneratedClass)
 	{
 		// If there is no index for the guid, remove the entry.
-		ParentAssetOverrides.RemoveAll([GeneratedClass](const FAnimParentNodeAssetOverride& Element)
+		ParentAssetOverrides.RemoveAll([AnimBPGeneratedClass](const FAnimParentNodeAssetOverride& Element)
 		{
-			if (!GeneratedClass->GetNodePropertyIndexFromGuid(Element.ParentNodeGuid, EPropertySearchMode::Hierarchy))
+			if (!AnimBPGeneratedClass->GetNodePropertyIndexFromGuid(Element.ParentNodeGuid, EPropertySearchMode::Hierarchy))
 			{
 				return true;
 			}
diff --git a/Engine/Source/Runtime/Engine/Private/Animation/AnimPhysicsSolver.cpp b/Engine/Source/Runtime/Engine/Private/Animation/AnimPhysicsSolver.cpp
--- a/Engine/Source/Runtime/Engine/Private/Animation/AnimPhysicsSolver.cpp
+++ b/Engine/Source/Runtime/Engine/Private/Animation/AnimPhysicsSolver.cpp
@@ -778,6 +778,43 @@ void FAnimPhys::ConstrainPlanar(float DeltaTime, TArray<FAnimPhysLinearLimit>& L
 	LimitContainer.Add(FAnimPhysLinearLimit(nullptr, Body, Position0, Position1, LimitPlane.GetSafeNormal(), TargetSpeed, TargetSpeed, FVector2D(0.0f, MAX_flt)));
 }
 
+void FAnimPhys::ConstrainSphericalInner(float DeltaTime, TArray<FAnimPhysLinearLimit>& LimitContainer, FAnimPhysRigidBody* Body, const FTransform& SphereTransform, float SphereRadius)
+{
+	FVector SphereToBody = Body->Pose.Position - SphereTransform.GetLocation();
+	FVector LimitNormal = SphereToBody.GetSafeNormal();
+	float DistanceToLimit = SphereToBody.Size() - SphereRadius;
+
+	if(Body->CollisionType != AnimPhysCollisionType::CoM)
+	{
+		DistanceToLimit += Body->SphereCollisionRadius;
+	}
+
+	float TargetSpeed = DistanceToLimit / DeltaTime;
+
+	FVector Position0 = SphereTransform.GetLocation();
+	FVector Position1 = FVector::ZeroVector;
+
+	LimitContainer.Add(FAnimPhysLinearLimit(nullptr, Body, Position0, Position1, LimitNormal, TargetSpeed, TargetSpeed, FVector2D(-MAX_flt, 0.0f)));
+}
+
+void FAnimPhys::ConstrainSphericalOuter(float DeltaTime, TArray<FAnimPhysLinearLimit>& LimitContainer, FAnimPhysRigidBody* Body, const FTransform& SphereTransform, float SphereRadius)
+{
+	FVector SphereToBody = Body->Pose.Position - SphereTransform.GetLocation();
+	float DistanceToLimit = SphereToBody.Size() - SphereRadius;
+
+	if(Body->CollisionType != AnimPhysCollisionType::CoM)
+	{
+		DistanceToLimit -= Body->SphereCollisionRadius;
+	}
+
+	float TargetSpeed = DistanceToLimit / DeltaTime;
+
+	FVector Position0 = SphereTransform.GetLocation();
+	FVector Position1 = FVector::ZeroVector;
+
+	LimitContainer.Add(FAnimPhysLinearLimit(nullptr, Body, Position0, Position1, SphereToBody.GetSafeNormal(), TargetSpeed, TargetSpeed, FVector2D(0.0f, MAX_flt)));
+}
+
 void FAnimPhys::CreateSpring(TArray<FAnimPhysSpring>& SpringContainer, FAnimPhysRigidBody* Body0, FVector Position0, FAnimPhysRigidBody* Body1, FVector Position1)
 {
 	FAnimPhysSpring NewSpring;
diff --git a/Engine/Source/Runtime/Engine/Private/Animation/AnimSequence.cpp b/Engine/Source/Runtime/Engine/Private/Animation/AnimSequence.cpp
--- a/Engine/Source/Runtime/Engine/Private/Animation/AnimSequence.cpp
+++ b/Engine/Source/Runtime/Engine/Private/Animation/AnimSequence.cpp
@@ -118,6 +118,7 @@ float GetIntervalPerKey(int32 NumFrames, float SequenceLength)
 
 UAnimSequence::UAnimSequence(const FObjectInitializer& ObjectInitializer)
 	: Super(ObjectInitializer)
+	, Interpolation(EAnimInterpolationType::Linear)
 	, bEnableRootMotion(false)
 	, RootMotionRootLock(ERootMotionRootLock::RefPose)
 	, bRootMotionSettingsCopiedFromMontage(false)
@@ -740,6 +741,11 @@ void UAnimSequence::ExtractBoneTransform(const struct FRawAnimSequenceTrack& Raw
 	const bool bHasScaleKey = (RawTrack.ScaleKeys.Num() > 0);
 	static const FVector DefaultScale3D = FVector(1.f);
 
+	if (Interpolation == EAnimInterpolationType::Step) 
+	{
+		Alpha = 0.f;
+	}
+
 	if(Alpha <= 0.f)
 	{
 		const int32 PosKeyIndex1 = FMath::Min(KeyIndex1, RawTrack.PosKeys.Num()-1);
@@ -3359,10 +3365,10 @@ int32 FindParentNodeIndex(URig* Rig, USkeleton* Skeleton, FName ParentNodeName)
 
 int32 UAnimSequence::GetSpaceBasedAnimationData(TArray< TArray<FTransform> >& AnimationDataInComponentSpace, FAnimSequenceTrackContainer * RiggingAnimationData) const
 {
-	USkeleton* Skeleton = GetSkeleton();
+	USkeleton* MySkeleton = GetSkeleton();
 
-	check(Skeleton);
-	const FReferenceSkeleton& RefSkeleton = Skeleton->GetReferenceSkeleton();
+	check(MySkeleton);
+	const FReferenceSkeleton& RefSkeleton = MySkeleton->GetReferenceSkeleton();
 	int32 NumBones = RefSkeleton.GetNum();
 
 	AnimationDataInComponentSpace.Empty(NumBones);
@@ -3380,7 +3386,7 @@ int32 UAnimSequence::GetSpaceBasedAnimationData(TArray< TArray<FTransform> >& An
 
 	if (RiggingAnimationData)
 	{
-		const URig* Rig = Skeleton->GetRig();
+		const URig* Rig = MySkeleton->GetRig();
 
 		check(Rig);
 
@@ -3398,7 +3404,7 @@ int32 UAnimSequence::GetSpaceBasedAnimationData(TArray< TArray<FTransform> >& An
 				if ( !BoneEvaluated[BoneIndex] )
 				{
 					const FName& BoneName = RefSkeleton.GetBoneName(BoneIndex);
-					const FName& NodeName = Skeleton->GetRigNodeNameFromBoneName(BoneName);
+					const FName& NodeName = MySkeleton->GetRigNodeNameFromBoneName(BoneName);
 					const FTransformBase* TransformBase = Rig->GetTransformBaseByNodeName(NodeName);
 					const int32 NodeIndex = RiggingAnimationData->TrackNames.Find(NodeName);
 					if(NodeIndex != INDEX_NONE)
@@ -3416,7 +3422,7 @@ int32 UAnimSequence::GetSpaceBasedAnimationData(TArray< TArray<FTransform> >& An
 						// this is easy since we just make sure it's evaluated or not
 						{
 							const FName& ParentNodeName = RotTransformConstraints[0].ParentSpace;
-							const FName& ParentBoneName = Skeleton->GetRigBoneMapping(ParentNodeName);
+							const FName& ParentBoneName = MySkeleton->GetRigBoneMapping(ParentNodeName);
 							const int32& ParentBoneIndex = RefSkeleton.FindBoneIndex(ParentBoneName);
 
 							if(ParentBoneIndex != INDEX_NONE)
@@ -3453,7 +3459,7 @@ int32 UAnimSequence::GetSpaceBasedAnimationData(TArray< TArray<FTransform> >& An
 						if (BoneEvaluated[BoneIndex])
 						{
 							const FName& ParentNodeName = PosTransformConstraints[0].ParentSpace;
-							const FName& ParentBoneName = Skeleton->GetRigBoneMapping(ParentNodeName);
+							const FName& ParentBoneName = MySkeleton->GetRigBoneMapping(ParentNodeName);
 							const int32& ParentBoneIndex = RefSkeleton.FindBoneIndex(ParentBoneName);
 
 							if(ParentBoneIndex != INDEX_NONE)
@@ -3535,9 +3541,9 @@ int32 UAnimSequence::GetSpaceBasedAnimationData(TArray< TArray<FTransform> >& An
 		// this one calculates aniamted space per bones and per key
 		for (int32 BoneIndex = 0; BoneIndex < NumBones; ++BoneIndex)
 		{
-			FName BoneName = Skeleton->GetReferenceSkeleton().GetBoneName(BoneIndex);
+			FName BoneName = MySkeleton->GetReferenceSkeleton().GetBoneName(BoneIndex);
 			int32 TrackIndex = AnimationTrackNames.Find(BoneName);
-			int32 ParentBoneIndex = Skeleton->GetReferenceSkeleton().GetParentIndex(BoneIndex);
+			int32 ParentBoneIndex = MySkeleton->GetReferenceSkeleton().GetParentIndex(BoneIndex);
 
 			if (TrackIndex != INDEX_NONE)
 			{
@@ -3561,7 +3567,7 @@ int32 UAnimSequence::GetSpaceBasedAnimationData(TArray< TArray<FTransform> >& An
 			else
 			{
 				// get local spaces from refpose and use that to fill it up
-				FTransform LocalTransform = Skeleton->GetReferenceSkeleton().GetRefBonePose()[BoneIndex];
+				FTransform LocalTransform = MySkeleton->GetReferenceSkeleton().GetRefBonePose()[BoneIndex];
 
 				for (int32 Key = 0; Key < NumKeys; ++Key)
 				{
@@ -3574,7 +3580,7 @@ int32 UAnimSequence::GetSpaceBasedAnimationData(TArray< TArray<FTransform> >& An
 						AnimationDataInComponentSpace[BoneIndex][Key] = LocalTransform;
 					}
 				}
-			}
+			}	
 		}
 
 	}
@@ -3584,12 +3590,12 @@ int32 UAnimSequence::GetSpaceBasedAnimationData(TArray< TArray<FTransform> >& An
 
 bool UAnimSequence::ConvertAnimationDataToRiggingData(FAnimSequenceTrackContainer& RiggingAnimationData)
 {
-	USkeleton* Skeleton = GetSkeleton();
-	if (Skeleton && Skeleton->GetRig())
+	USkeleton* MySkeleton = GetSkeleton();
+	if (MySkeleton && MySkeleton->GetRig())
 	{
-		const URig* Rig = Skeleton->GetRig();
+		const URig* Rig = MySkeleton->GetRig();
 		TArray<FName> ValidNodeNames;
-		int32 NumNodes = Skeleton->GetMappedValidNodes(ValidNodeNames);
+		int32 NumNodes = MySkeleton->GetMappedValidNodes(ValidNodeNames);
 		TArray< TArray<FTransform> > AnimationDataInComponentSpace;
 		int32 NumBones = GetSpaceBasedAnimationData(AnimationDataInComponentSpace, NULL);
 
@@ -3602,8 +3608,8 @@ bool UAnimSequence::ConvertAnimationDataToRiggingData(FAnimSequenceTrackContaine
 			{
 				struct FRawAnimSequenceTrack& Track = RiggingAnimationData.AnimationTracks[NodeIndex];
 				const FName& NodeName = ValidNodeNames[NodeIndex];
-				const FName& BoneName = Skeleton->GetRigBoneMapping(NodeName);
-				const int32& BoneIndex = Skeleton->GetReferenceSkeleton().FindBoneIndex(BoneName);
+				const FName& BoneName = MySkeleton->GetRigBoneMapping(NodeName);
+				const int32& BoneIndex = MySkeleton->GetReferenceSkeleton().FindBoneIndex(BoneName);
 
 				if (ensure(BoneIndex != INDEX_NONE))
 				{
@@ -3626,8 +3632,8 @@ bool UAnimSequence::ConvertAnimationDataToRiggingData(FAnimSequenceTrackContaine
 						if (RotationTransformConstraint.Num() > 0)
 						{
 							const FName& ParentSpace = RotationTransformConstraint[0].ParentSpace;
-							const FName& ParentBoneName = Skeleton->GetRigBoneMapping(ParentSpace);
-							const int32& ParentBoneIndex = Skeleton->GetReferenceSkeleton().FindBoneIndex(ParentBoneName);
+							const FName& ParentBoneName = MySkeleton->GetRigBoneMapping(ParentSpace);
+							const int32& ParentBoneIndex = MySkeleton->GetReferenceSkeleton().FindBoneIndex(ParentBoneName);
 							if (ParentBoneIndex != INDEX_NONE)
 							{
 								// if no rig control, component space is used
@@ -3662,8 +3668,8 @@ bool UAnimSequence::ConvertAnimationDataToRiggingData(FAnimSequenceTrackContaine
 						if (TranslationTransformConstraint.Num() > 0)
 						{
 							const FName& ParentSpace = TranslationTransformConstraint[0].ParentSpace;
-							const FName& ParentBoneName = Skeleton->GetRigBoneMapping(ParentSpace);
-							const int32& ParentBoneIndex = Skeleton->GetReferenceSkeleton().FindBoneIndex(ParentBoneName);
+							const FName& ParentBoneName = MySkeleton->GetRigBoneMapping(ParentSpace);
+							const int32& ParentBoneIndex = MySkeleton->GetReferenceSkeleton().FindBoneIndex(ParentBoneName);
 							if (ParentBoneIndex != INDEX_NONE)
 							{
 								// if no rig control, component space is used
@@ -3720,11 +3726,11 @@ bool UAnimSequence::ConvertRiggingDataToAnimationData(FAnimSequenceTrackContaine
 		TArray< TArray<FTransform> > AnimationDataInComponentSpace;
 		int32 NumBones = GetSpaceBasedAnimationData(AnimationDataInComponentSpace, &RiggingAnimationData);
 
-		USkeleton* Skeleton = GetSkeleton();
+		USkeleton* MySkeleton = GetSkeleton();
 		TArray<FRawAnimSequenceTrack> OldAnimationData = RawAnimationData;
 		TArray<FName> OldAnimationTrackNames = AnimationTrackNames;
 		TArray<FName> ValidNodeNames;
-		Skeleton->GetMappedValidNodes(ValidNodeNames);
+		MySkeleton->GetMappedValidNodes(ValidNodeNames);
 		// remove from ValidNodeNames if it doesn't belong to AnimationTrackNames
 		for (int32 NameIndex=0; NameIndex<ValidNodeNames.Num(); ++NameIndex)
 		{
@@ -3750,11 +3756,11 @@ bool UAnimSequence::ConvertRiggingDataToAnimationData(FAnimSequenceTrackContaine
 			ClearBakedTransformData();
 		}
 
-		const FReferenceSkeleton& RefSkeleton = Skeleton->GetReferenceSkeleton();
-		const URig* Rig = Skeleton->GetRig();
+		const FReferenceSkeleton& RefSkeleton = MySkeleton->GetReferenceSkeleton();
+		const URig* Rig = MySkeleton->GetRig();
 		for (int32 NodeIndex = 0; NodeIndex < ValidNumNodes; ++NodeIndex)
 		{
-			FName BoneName = Skeleton->GetRigBoneMapping(ValidNodeNames[NodeIndex]);
+			FName BoneName = MySkeleton->GetRigBoneMapping(ValidNodeNames[NodeIndex]);
 			int32 BoneIndex = RefSkeleton.FindBoneIndex(BoneName);
 
 			if (BoneIndex != INDEX_NONE)
@@ -3805,7 +3811,7 @@ bool UAnimSequence::ConvertRiggingDataToAnimationData(FAnimSequenceTrackContaine
 		int32 TrackIdx = 0;
 		for (auto TrackName : AnimationTrackNames)
 		{
-			TrackToSkeletonMapTable[TrackIdx++].BoneTreeIndex = Skeleton->GetReferenceSkeleton().FindBoneIndex(TrackName);
+			TrackToSkeletonMapTable[TrackIdx++].BoneTreeIndex = MySkeleton->GetReferenceSkeleton().FindBoneIndex(TrackName);
 		}
 		PostProcessSequence();
 
diff --git a/Engine/Source/Runtime/Engine/Private/Animation/AnimSequenceBase.cpp b/Engine/Source/Runtime/Engine/Private/Animation/AnimSequenceBase.cpp
--- a/Engine/Source/Runtime/Engine/Private/Animation/AnimSequenceBase.cpp
+++ b/Engine/Source/Runtime/Engine/Private/Animation/AnimSequenceBase.cpp
@@ -20,11 +20,11 @@ UAnimSequenceBase::UAnimSequenceBase(const FObjectInitializer& ObjectInitializer
 }
 
 template <typename DataType>
-void UAnimSequenceBase::VerifyCurveNames(USkeleton* Skeleton, const FName& NameContainer, TArray<DataType>& CurveList)
+void UAnimSequenceBase::VerifyCurveNames(USkeleton* InSkeleton, const FName& NameContainer, TArray<DataType>& CurveList)
 {
 	// since this is verify function that makes sure it exists after loaded
 	// we should add it if it doesn't exist
-	const FSmartNameMapping* NameMapping = Skeleton->GetOrAddSmartNameContainer(NameContainer);
+	const FSmartNameMapping* NameMapping = InSkeleton->GetOrAddSmartNameContainer(NameContainer);
 
 	TArray<DataType*> UnlinkedCurves;
 	for(DataType& Curve : CurveList)
@@ -48,7 +48,7 @@ void UAnimSequenceBase::VerifyCurveNames(USkeleton* Skeleton, const FName& NameC
 
 	for(DataType* Curve : UnlinkedCurves)
 	{
-		Skeleton->AddSmartNameAndModify(NameContainer, Curve->LastObservedName, Curve->CurveUid);
+		InSkeleton->AddSmartNameAndModify(NameContainer, Curve->LastObservedName, Curve->CurveUid);
 	}
 }
 
@@ -80,7 +80,7 @@ void UAnimSequenceBase::PostLoad()
 #endif
 	RefreshCacheData();
 
-	if(USkeleton* Skeleton = GetSkeleton())
+	if(USkeleton* MySkeleton = GetSkeleton())
 	{
 		// Fix up the existing curves to work with smartnames
 		if(GetLinkerUE4Version() < VER_UE4_SKELETON_ADD_SMARTNAMES)
@@ -88,16 +88,16 @@ void UAnimSequenceBase::PostLoad()
 			for(FFloatCurve& Curve : RawCurveData.FloatCurves)
 			{
 				// Add the names of the curves into the smartname mapping and store off the curve uid which will be saved next time the sequence saves.
-				Skeleton->AddSmartNameAndModify(USkeleton::AnimCurveMappingName, Curve.LastObservedName, Curve.CurveUid);
+				MySkeleton->AddSmartNameAndModify(USkeleton::AnimCurveMappingName, Curve.LastObservedName, Curve.CurveUid);
 			}
 		}
 		else
 		{
-			VerifyCurveNames<FFloatCurve>(Skeleton, USkeleton::AnimCurveMappingName, RawCurveData.FloatCurves);
+			VerifyCurveNames<FFloatCurve>(MySkeleton, USkeleton::AnimCurveMappingName, RawCurveData.FloatCurves);
 		}
 
 #if WITH_EDITOR
-		VerifyCurveNames<FTransformCurve>(Skeleton, USkeleton::AnimTrackCurveMappingName, RawCurveData.TransformCurves);
+		VerifyCurveNames<FTransformCurve>(MySkeleton, USkeleton::AnimTrackCurveMappingName, RawCurveData.TransformCurves);
 #endif
 	}
 }
@@ -494,9 +494,9 @@ void UAnimSequenceBase::Serialize(FArchive& Ar)
 
 	if(Ar.ArIsSaving && Ar.UE4Ver() >= VER_UE4_SKELETON_ADD_SMARTNAMES)
 	{
-		if(USkeleton* Skeleton = GetSkeleton())
+		if(USkeleton* MySkeleton = GetSkeleton())
 		{
-			const FSmartNameMapping* Mapping = GetSkeleton()->GetSmartNameContainer(USkeleton::AnimCurveMappingName);
+			const FSmartNameMapping* Mapping = MySkeleton->GetSmartNameContainer(USkeleton::AnimCurveMappingName);
 			check(Mapping); // Should always exist
 			RawCurveData.UpdateLastObservedNames(Mapping);
 		}
@@ -505,12 +505,12 @@ void UAnimSequenceBase::Serialize(FArchive& Ar)
 #if WITH_EDITORONLY_DATA
 	if(Ar.ArIsSaving && Ar.UE4Ver() >= VER_UE4_ANIMATION_ADD_TRACKCURVES)
 	{
-		if(USkeleton* Skeleton = GetSkeleton())
+		if(USkeleton* MySkeleton = GetSkeleton())
 		{
 			// we don't add track curve container unless it has been edited. 
 			if ( RawCurveData.TransformCurves.Num() > 0 )
 			{
-				const FSmartNameMapping* Mapping = GetSkeleton()->GetSmartNameContainer(USkeleton::AnimTrackCurveMappingName);
+				const FSmartNameMapping* Mapping = MySkeleton->GetSmartNameContainer(USkeleton::AnimTrackCurveMappingName);
 				// this name might not exists because it's only available if you edit animation
 				if (Mapping)
 				{
diff --git a/Engine/Source/Runtime/Engine/Private/AudioDevice.cpp b/Engine/Source/Runtime/Engine/Private/AudioDevice.cpp
--- a/Engine/Source/Runtime/Engine/Private/AudioDevice.cpp
+++ b/Engine/Source/Runtime/Engine/Private/AudioDevice.cpp
@@ -12,6 +12,10 @@
 #include "Sound/SoundWave.h"
 #include "IAudioExtensionPlugin.h"
 
+#if WITH_EDITOR
+#include "Editor/UnrealEd/Classes/Editor/EditorEngine.h"
+#endif
+
 /*-----------------------------------------------------------------------------
 FDynamicParameter implementation.
 -----------------------------------------------------------------------------*/
@@ -1751,7 +1755,21 @@ void FAudioDevice::ApplyClassAdjusters(USoundMix* SoundMix, float InterpValue, f
 		{
 			if (Entry.bApplyToChildren)
 			{
-				RecursiveApplyAdjuster(Entry, Entry.SoundClassObject);
+				// If we're using the override, Entry will already have interpolated values
+				if (bUsingOverride)
+				{
+					RecursiveApplyAdjuster(Entry, Entry.SoundClassObject);
+				}
+				else
+				{
+					// Copy the entry with the interpolated values before applying it recursively
+					FSoundClassAdjuster EntryCopy = Entry;
+					EntryCopy.VolumeAdjuster = InterpolateAdjuster(Entry.VolumeAdjuster, InterpValue);
+					EntryCopy.PitchAdjuster = InterpolateAdjuster(Entry.PitchAdjuster, InterpValue);
+					EntryCopy.VoiceCenterChannelVolumeAdjuster = InterpolateAdjuster(Entry.VoiceCenterChannelVolumeAdjuster, InterpValue);
+
+					RecursiveApplyAdjuster(EntryCopy, Entry.SoundClassObject);
+				}
 			}
 			else
 			{
@@ -3464,3 +3482,15 @@ void FAudioDevice::OnEndPIE(const bool bIsSimulating)
 	}
 }
 #endif
+
+bool FAudioDevice::CanUseVRAudioDevice()
+{
+#if WITH_EDITOR
+	if (GIsEditor)
+	{
+		UEditorEngine* EdEngine = Cast<UEditorEngine>(GEngine);
+		return EdEngine->bUseVRPreviewForPlayWorld;
+	}
+#endif
+	return true;
+}
diff --git a/Engine/Source/Runtime/Engine/Private/Character.cpp b/Engine/Source/Runtime/Engine/Private/Character.cpp
--- a/Engine/Source/Runtime/Engine/Private/Character.cpp
+++ b/Engine/Source/Runtime/Engine/Private/Character.cpp
@@ -154,9 +154,9 @@ UPawnMovementComponent* ACharacter::GetMovementComponent() const
 }
 
 
-void ACharacter::SetupPlayerInputComponent(UInputComponent* InputComponent)
+void ACharacter::SetupPlayerInputComponent(UInputComponent* InInputComponent)
 {
-	check(InputComponent);
+	check(InInputComponent);
 }
 
 
diff --git a/Engine/Source/Runtime/Engine/Private/Collision/KAggregateGeom.cpp b/Engine/Source/Runtime/Engine/Private/Collision/KAggregateGeom.cpp
--- a/Engine/Source/Runtime/Engine/Private/Collision/KAggregateGeom.cpp
+++ b/Engine/Source/Runtime/Engine/Private/Collision/KAggregateGeom.cpp
@@ -8,6 +8,7 @@
 #include "PhysicsEngine/AggregateGeom.h"
 #include "Collision.h"
 #include "Engine/Polys.h"
+#include "PhysicsEngine/PhysXSupport.h"
 
 #define MIN_HULL_VERT_DISTANCE		(0.1f)
 #define MIN_HULL_VALID_DIMENSION	(0.5f)
@@ -347,7 +348,31 @@ FBox FKConvexElem::CalcAABB(const FTransform& BoneTM, const FVector& Scale3D) co
 	// Zero out rotation and location so we transform by scale along
 	const FTransform LocalToWorld = FTransform(FQuat::Identity, FVector::ZeroVector, Scale3D) * BoneTM;
 
-	return ElemBox.TransformBy(LocalToWorld);
+	return ElemBox.TransformBy(Transform * LocalToWorld);
+}
+
+void FKConvexElem::GetPlanes(TArray<FPlane>& Planes) const
+{
+#if WITH_PHYSX
+	if (ConvexMesh != nullptr)
+	{
+		Planes.Empty();
+
+		PxU32 NumPolys = ConvexMesh->getNbPolygons();
+		for (PxU32 PolyIndex = 0; PolyIndex < NumPolys; PolyIndex++)
+		{
+			PxHullPolygon Data;
+			bool bStatus = ConvexMesh->getPolygonData(PolyIndex, Data);
+			check(bStatus);
+
+			// Convert to UE type
+			FPlane Plane(Data.mPlane[0], Data.mPlane[1], Data.mPlane[2], -Data.mPlane[3]);
+
+			// Add to output array
+			Planes.Add(Plane);
+		}
+	}
+#endif // WITH_PHYSX
 }
 
 
@@ -511,7 +536,7 @@ bool FKConvexElem::HullFromPlanes(const TArray<FPlane>& InPlanes, const TArray<F
 	// We shouldn't have the same vertex multiple times (using AddVertexIfNotPresent above)
 	if(VertexData.Num() < 4)
 	{
-		return true;
+		return false;
 	}
 
 	// Check that not all vertices lie on a line (ie. find plane)
@@ -536,7 +561,7 @@ bool FKConvexElem::HullFromPlanes(const TArray<FPlane>& InPlanes, const TArray<F
 
 	if(!bFound)
 	{
-		return true;
+		return false;
 	}
 
 	// Now we check that not all vertices lie on a plane, by checking at least one lies off the plane we have formed.
@@ -558,7 +583,7 @@ bool FKConvexElem::HullFromPlanes(const TArray<FPlane>& InPlanes, const TArray<F
 	// If we did not find a vert off the line - discard this hull.
 	if(!bFound)
 	{
-		return true;
+		return false;
 	}
 
 	// calc bounding box of verts
diff --git a/Engine/Source/Runtime/Engine/Private/Components/CharacterMovementComponent.cpp b/Engine/Source/Runtime/Engine/Private/Components/CharacterMovementComponent.cpp
--- a/Engine/Source/Runtime/Engine/Private/Components/CharacterMovementComponent.cpp
+++ b/Engine/Source/Runtime/Engine/Private/Components/CharacterMovementComponent.cpp
@@ -86,6 +86,25 @@ namespace CharacterMovementComponentStatics
 }
 
 // CVars
+namespace CharacterCVars
+{
+
+	// Logging when character is stuck. Off by default in shipping.
+#if UE_BUILD_SHIPPING
+	static float StuckWarningPeriod = -1.f;
+#else
+	static float StuckWarningPeriod = 1.f;
+#endif
+
+	static FAutoConsoleVariableRef CVarStuckWarningPeriod(
+		TEXT("p.CharacterStuckWarningPeriod"),
+		StuckWarningPeriod,
+		TEXT("How often (in seconds) we are allowed to log a message about being stuck in geometry.\n")
+		TEXT("<0: Disable, >=0: Enable and log this often, in seconds."),
+		ECVF_Default);
+}
+
+
 static TAutoConsoleVariable<int32> CVarNetEnableMoveCombining(
 	TEXT("p.NetEnableMoveCombining"),
 	1,
@@ -232,6 +251,11 @@ UCharacterMovementComponent::UCharacterMovementComponent(const FObjectInitialize
 	MaxSimulationTimeStep = 0.05f;
 	MaxSimulationIterations = 8;
 
+	MaxDepenetrationWithGeometry = 500.f;
+	MaxDepenetrationWithGeometryAsProxy = 100.f;
+	MaxDepenetrationWithPawn = 100.f;
+	MaxDepenetrationWithPawnAsProxy = 2.f;
+
 	NetworkSimulatedSmoothLocationTime = 0.100f;
 	NetworkSimulatedSmoothRotationTime = 0.033f;
 	NetworkMaxSmoothUpdateDistance = 256.f;
@@ -915,17 +939,7 @@ void UCharacterMovementComponent::PerformAirControlForPathFollowing(FVector Dire
 	// use air control if low grav or above destination and falling towards it
 	if ( CharacterOwner && Velocity.Z < 0.f && (ZDiff < 0.f || GetGravityZ() > 0.9f * GetWorld()->GetDefaultGravityZ()))
 	{
-		if ( ZDiff > 0.f )
-		{
-			if ( ZDiff > 2.f * GetMaxJumpHeight() )
-			{
-				if (PathFollowingComp.IsValid())
-				{
-					PathFollowingComp->AbortMove(TEXT("missed jump"));
-				}
-			}
-		}
-		else
+		if ( ZDiff < 0.f )
 		{
 			if ( (Velocity.X == 0.f) && (Velocity.Y == 0.f) )
 			{
@@ -1096,8 +1110,9 @@ void UCharacterMovementComponent::TickComponent(float DeltaTime, enum ELevelTick
 
 void UCharacterMovementComponent::PostPhysicsTickComponent(float DeltaTime, FCharacterMovementComponentPostPhysicsTickFunction& ThisTickFunction)
 {
-	if(bDeferUpdateBasedMovement)
+	if (bDeferUpdateBasedMovement)
 	{
+		FScopedMovementUpdate ScopedMovementUpdate(UpdatedComponent, bEnableScopedMovementUpdates ? EScopedUpdate::DeferredUpdates : EScopedUpdate::ImmediateUpdates);
 		UpdateBasedMovement(DeltaTime);
 		SaveBaseLocation();
 		bDeferUpdateBasedMovement = false;
@@ -2371,6 +2386,27 @@ float UCharacterMovementComponent::GetMaxSpeed() const
 }
 
 
+FVector UCharacterMovementComponent::GetPenetrationAdjustment(const FHitResult& Hit) const
+{
+	FVector Result = Super::GetPenetrationAdjustment(Hit);
+
+	if (CharacterOwner)
+	{
+		const bool bIsProxy = (CharacterOwner->Role == ROLE_SimulatedProxy);
+		float MaxDistance = bIsProxy ? MaxDepenetrationWithGeometryAsProxy : MaxDepenetrationWithGeometry;
+		const AActor* HitActor = Hit.GetActor();
+		if (Cast<APawn>(HitActor))
+		{
+			MaxDistance = bIsProxy ? MaxDepenetrationWithPawnAsProxy : MaxDepenetrationWithPawn;
+		}
+
+		Result = Result.GetClampedToMaxSize(MaxDistance);
+	}
+
+	return Result;
+}
+
+
 bool UCharacterMovementComponent::ResolvePenetrationImpl(const FVector& Adjustment, const FHitResult& Hit, const FQuat& NewRotation)
 {
 	// If movement occurs, mark that we teleported, so we don't incorrectly adjust velocity based on a potentially very different movement than our movement direction.
@@ -3910,11 +3946,6 @@ void UCharacterMovementComponent::RevertMove(const FVector& OldLocation, UPrimit
 		Velocity = FVector::ZeroVector;
 		Acceleration = FVector::ZeroVector;
 		//UE_LOG(LogCharacterMovement, Log, TEXT("%s FAILMOVE RevertMove"), *CharacterOwner->GetName());
-
-		if (PathFollowingComp.IsValid())
-		{
-			PathFollowingComp->AbortMove(TEXT("RevertMove"));
-		}
 	}
 }
 
@@ -3945,22 +3976,37 @@ FVector UCharacterMovementComponent::ComputeGroundMovementDelta(const FVector& D
 
 void UCharacterMovementComponent::OnCharacterStuckInGeometry(const FHitResult* Hit)
 {
-	if (Hit == nullptr)
-	{
-		UE_LOG(LogCharacterMovement, Log, TEXT("%s is stuck and failed to move!"), *CharacterOwner->GetName());
-	}
-	else
+	if (CharacterCVars::StuckWarningPeriod >= 0)
 	{
-		UE_LOG(LogCharacterMovement, Log, TEXT("%s is stuck and failed to move! Velocity: X=%3.3f Y=%3.3f Z=%3.3f Location: X=%3.3f Y=%3.3f Z=%3.3f Normal: X=%3.3f Y=%3.3f Z=%3.3f PenetrationDepth:%.3f Actor:%s Component:%s BoneName:%s"),
-			   *GetNameSafe(CharacterOwner),
-			   Velocity.X, Velocity.Y, Velocity.Z,
-			   Hit->Location.X, Hit->Location.Y, Hit->Location.Z,
-			   Hit->Normal.X, Hit->Normal.Y, Hit->Normal.Z,
-			   Hit->PenetrationDepth,
-			   *GetNameSafe(Hit->GetActor()),
-			   *GetNameSafe(Hit->GetComponent()),
-			   Hit->BoneName.IsValid() ? *Hit->BoneName.ToString() : TEXT("None")
-			   );
+		UWorld* MyWorld = GetWorld();
+		const float RealTimeSeconds = MyWorld->GetRealTimeSeconds();
+		if ((RealTimeSeconds - LastStuckWarningTime) >= CharacterCVars::StuckWarningPeriod)
+		{
+			LastStuckWarningTime = RealTimeSeconds;
+			if (Hit == nullptr)
+			{
+				UE_LOG(LogCharacterMovement, Log, TEXT("%s is stuck and failed to move! (%d other events since notify)"), *CharacterOwner->GetName(), StuckWarningCountSinceNotify);
+			}
+			else
+			{
+				UE_LOG(LogCharacterMovement, Log, TEXT("%s is stuck and failed to move! Velocity: X=%3.2f Y=%3.2f Z=%3.2f Location: X=%3.2f Y=%3.2f Z=%3.2f Normal: X=%3.2f Y=%3.2f Z=%3.2f PenetrationDepth:%.3f Actor:%s Component:%s BoneName:%s (%d other events since notify)"),
+					   *GetNameSafe(CharacterOwner),
+					   Velocity.X, Velocity.Y, Velocity.Z,
+					   Hit->Location.X, Hit->Location.Y, Hit->Location.Z,
+					   Hit->Normal.X, Hit->Normal.Y, Hit->Normal.Z,
+					   Hit->PenetrationDepth,
+					   *GetNameSafe(Hit->GetActor()),
+					   *GetNameSafe(Hit->GetComponent()),
+					   Hit->BoneName.IsValid() ? *Hit->BoneName.ToString() : TEXT("None"),
+					   StuckWarningCountSinceNotify
+					   );
+			}
+			StuckWarningCountSinceNotify = 0;
+		}
+		else
+		{
+			StuckWarningCountSinceNotify += 1;
+		}
 	}
 
 	// Don't update velocity based on our (failed) change in position this update since we're stuck.
@@ -4983,7 +5029,7 @@ void UCharacterMovementComponent::PhysicsVolumeChanged( APhysicsVolume* NewVolum
 			// AI needs to stop any current moves
 			if (PathFollowingComp.IsValid())
 			{
-				PathFollowingComp->AbortMove(TEXT("water"));
+				PathFollowingComp->AbortMove(*this, FPathFollowingResultFlags::MovementStop);
 			}			
 		}
 		else if ( !IsSwimming() )
diff --git a/Engine/Source/Runtime/Engine/Private/Components/NavMovementComponent.cpp b/Engine/Source/Runtime/Engine/Private/Components/NavMovementComponent.cpp
--- a/Engine/Source/Runtime/Engine/Private/Components/NavMovementComponent.cpp
+++ b/Engine/Source/Runtime/Engine/Private/Components/NavMovementComponent.cpp
@@ -62,7 +62,7 @@ void UNavMovementComponent::StopActiveMovement()
 {
 	if (PathFollowingComp.IsValid() && bStopMovementAbortPaths)
 	{
-		PathFollowingComp->AbortMove("StopActiveMovement");
+		PathFollowingComp->AbortMove(*this, FPathFollowingResultFlags::MovementStop);
 	}
 }
 
diff --git a/Engine/Source/Runtime/Engine/Private/Components/SceneComponent.cpp b/Engine/Source/Runtime/Engine/Private/Components/SceneComponent.cpp
--- a/Engine/Source/Runtime/Engine/Private/Components/SceneComponent.cpp
+++ b/Engine/Source/Runtime/Engine/Private/Components/SceneComponent.cpp
@@ -445,8 +445,8 @@ void USceneComponent::OnRegister()
 		if (AttachToComponent(GetAttachParent(), FAttachmentTransformRules::KeepRelativeTransform, GetAttachSocketName()) == false)
 		{
 			// Failed to attach, we need to clear AttachParent so we don't think we're actually attached when we're not.
-			GetMutableAttachParent() = nullptr;
-			GetMutableAttachSocketName() = NAME_None;
+			AttachParent = nullptr;
+			AttachSocketName = NAME_None;
 		}
 	}
 	
@@ -726,10 +726,10 @@ void USceneComponent::DestroyComponent(bool bPromoteChildren/*= false*/)
 				{
 					// Attach the child node that we're promoting to the parent and move it to the same position as the old node was in the array
 					ChildToPromote->AttachToComponent(CachedAttachParent, FAttachmentTransformRules::KeepWorldTransform);
-					CachedAttachParent->GetMutableAttachChildren().Remove(ChildToPromote);
+					CachedAttachParent->AttachChildren.Remove(ChildToPromote);
 
 					Index = FMath::Clamp<int32>(Index, 0, AttachSiblings.Num());
-					CachedAttachParent->GetMutableAttachChildren().Insert(ChildToPromote, Index);
+					CachedAttachParent->AttachChildren.Insert(ChildToPromote, Index);
 				}
 			}
 
@@ -769,7 +769,6 @@ void USceneComponent::OnComponentDestroyed(bool bDestroyingHierarchy)
 	{
 		// If we're destroying the hierarchy we only have to make sure that we detach children from other Actor's
 		AActor* MyOwner = GetOwner();
-		TArray<USceneComponent*>& AttachedChildren = GetMutableAttachChildren();
 
 		if (bDestroyingHierarchy)
 		{
@@ -778,7 +777,7 @@ void USceneComponent::OnComponentDestroyed(bool bDestroyingHierarchy)
 			bool bExternalAttachParentDetermined = false;
 			USceneComponent* ExternalAttachParent = nullptr;
 
-			int32 ChildCount = AttachedChildren.Num();
+			int32 ChildCount = AttachChildren.Num();
 
 			// We cache the actual children to put back after the detach process
 			TArray<USceneComponent*> CachedChildren;
@@ -786,7 +785,7 @@ void USceneComponent::OnComponentDestroyed(bool bDestroyingHierarchy)
 
 			while (ChildCount > 0)
 			{
-				USceneComponent* Child = AttachedChildren.Last();
+				USceneComponent* Child = AttachChildren.Last();
 				if (Child && Child->GetOwner() != MyOwner)
 				{
 					if (Child->GetAttachParent())
@@ -828,7 +827,7 @@ void USceneComponent::OnComponentDestroyed(bool bDestroyingHierarchy)
 								// so instead of crashing, output an error and gracefully handle
 								UE_LOG(LogSceneComponent, Error, TEXT("Component '%s' has '%s' in its AttachChildren array, however, '%s' believes it is attached to '%s'"), *GetFullName(), *Child->GetFullName(), *Child->GetFullName(), *Child->GetAttachParent()->GetFullName());
 							}
-							AttachedChildren.Pop(false);
+							AttachChildren.Pop(false);
 						}
 					}
 					else 
@@ -840,28 +839,28 @@ void USceneComponent::OnComponentDestroyed(bool bDestroyingHierarchy)
 						{
 							UE_LOG(LogSceneComponent, Error, TEXT("Component '%s' has '%s' in its AttachChildren array, however, '%s' believes it is not attached to anything"), *GetFullName(), *Child->GetFullName(), *Child->GetFullName());
 						}
-						AttachedChildren.Pop(false);
+						AttachChildren.Pop(false);
 					}
-					checkf(ChildCount > AttachedChildren.Num(), TEXT("AttachChildren count increased while detaching '%s', likely caused by OnAttachmentChanged introducing new children, which could lead to an infinite loop."), *Child->GetName());
+					checkf(ChildCount > AttachChildren.Num(), TEXT("AttachChildren count increased while detaching '%s', likely caused by OnAttachmentChanged introducing new children, which could lead to an infinite loop."), *Child->GetName());
 				}
 				else
 				{
-					AttachedChildren.Pop(false);
+					AttachChildren.Pop(false);
 					if (Child)
 					{
 						CachedChildren.Add(Child);
 					}
 				}
-				ChildCount = AttachedChildren.Num();
+				ChildCount = AttachChildren.Num();
 			}
-			AttachedChildren = MoveTemp(CachedChildren);
+			AttachChildren = MoveTemp(CachedChildren);
 		}
 		else
 		{
-			int32 ChildCount = AttachedChildren.Num();
+			int32 ChildCount = AttachChildren.Num();
 			while (ChildCount > 0)
 			{
-				if (USceneComponent* Child = AttachedChildren.Last())
+				if (USceneComponent* Child = AttachChildren.Last())
 				{
 					if (Child->GetAttachParent())
 					{
@@ -888,7 +887,7 @@ void USceneComponent::OnComponentDestroyed(bool bDestroyingHierarchy)
 								// so instead of crashing, output an error and gracefully handle
 								UE_LOG(LogSceneComponent, Error, TEXT("Component '%s' has '%s' in its AttachChildren array, however, '%s' believes it is attached to '%s'"), *GetFullName(), *Child->GetFullName(), *Child->GetFullName(), *Child->GetAttachParent()->GetFullName());
 							}
-							AttachedChildren.Pop(false);
+							AttachChildren.Pop(false);
 						}
 					}
 					else 
@@ -900,15 +899,15 @@ void USceneComponent::OnComponentDestroyed(bool bDestroyingHierarchy)
 						{
 							UE_LOG(LogSceneComponent, Error, TEXT("Component '%s' has '%s' in its AttachChildren array, however, '%s' believes it is not attached to anything"), *GetFullName(), *Child->GetFullName(), *Child->GetFullName());
 						}
-						AttachedChildren.Pop(false);
+						AttachChildren.Pop(false);
 					}
-					checkf(ChildCount > AttachedChildren.Num(), TEXT("AttachChildren count increased while detaching '%s', likely caused by OnAttachmentChanged introducing new children, which could lead to an infinite loop."), *Child->GetName());
+					checkf(ChildCount > AttachChildren.Num(), TEXT("AttachChildren count increased while detaching '%s', likely caused by OnAttachmentChanged introducing new children, which could lead to an infinite loop."), *Child->GetName());
 				}
 				else
 				{
-					AttachedChildren.Pop(false);
+					AttachChildren.Pop(false);
 				}
-				ChildCount = AttachedChildren.Num();
+				ChildCount = AttachChildren.Num();
 			}
 		}
 
@@ -1406,7 +1405,6 @@ void USceneComponent::AppendDescendants(TArray<USceneComponent*>& Children) cons
 	}
 }
 
-PRAGMA_DISABLE_DEPRECATION_WARNINGS // DEPRECATED(4.12)
 void USceneComponent::SetupAttachment(class USceneComponent* InParent, FName InSocketName)
 {
 	if (ensureMsgf(!bRegistered, TEXT("SetupAttachment should only be used to initialize AttachParent and AttachSocketName for a future AttachTo. Once a component is registered you must use AttachTo.")))
@@ -1418,7 +1416,6 @@ void USceneComponent::SetupAttachment(class USceneComponent* InParent, FName InS
 		}
 	}
 }
-PRAGMA_ENABLE_DEPRECATION_WARNINGS // DEPRECATED(4.12)
 
 //This function is used for giving AttachTo different bWeldSimulatedBodies default, but only when called from BP
 bool USceneComponent::K2_AttachTo(class USceneComponent* InParent, FName InSocketName, EAttachLocation::Type AttachLocationType, bool bWeldSimulatedBodies /*= true*/)
@@ -1626,19 +1623,19 @@ bool USceneComponent::AttachToComponent(USceneComponent* Parent, const FAttachme
 		PrimaryComponentTick.AddPrerequisite(Parent, Parent->PrimaryComponentTick); // force us to tick after the parent does
 
 		// Save pointer from child to parent
-		GetMutableAttachParent() = Parent;
-		GetMutableAttachSocketName() = SocketName;
+		AttachParent = Parent;
+		AttachSocketName = SocketName;
 
 		OnAttachmentChanged();
 
 		// Preserve order of previous attachment if valid (in case we're doing a reattach operation inside a loop that might assume the AttachChildren order won't change)
 		if(LastAttachIndex != INDEX_NONE)
 		{
-			Parent->GetMutableAttachChildren().Insert(this, LastAttachIndex);
+			Parent->AttachChildren.Insert(this, LastAttachIndex);
 		}
 		else
 		{
-			Parent->GetMutableAttachChildren().Add(this);
+			Parent->AttachChildren.Add(this);
 		}
 
 		// Now apply attachment rules
@@ -1808,7 +1805,7 @@ void USceneComponent::DetachFromComponent(const FDetachmentTransformRules& Detac
 
 		PrimaryComponentTick.RemovePrerequisite(GetAttachParent(), GetAttachParent()->PrimaryComponentTick); // no longer required to tick after the attachment
 
-		GetAttachParent()->GetMutableAttachChildren().Remove(this);
+		GetAttachParent()->AttachChildren.Remove(this);
 		GetAttachParent()->OnChildDetached(this);
 
 #if WITH_EDITOR
@@ -1820,8 +1817,8 @@ void USceneComponent::DetachFromComponent(const FDetachmentTransformRules& Detac
 			}
 		}
 #endif
-		GetMutableAttachParent() = nullptr;
-		GetMutableAttachSocketName() = NAME_None;
+		AttachParent = nullptr;
+		AttachSocketName = NAME_None;
 
 		OnAttachmentChanged();
 
@@ -2760,7 +2757,7 @@ void USceneComponent::PostNetReceive()
 		if (ParentRoot != this)
 		{
 			bNetUpdateAttachment = true;
-			GetMutableAttachParent() = ParentRoot;
+			AttachParent = ParentRoot;
 		}
 	}
 }
@@ -2769,10 +2766,8 @@ void USceneComponent::PostRepNotifies()
 {
 	if (bNetUpdateAttachment)
 	{
-		PRAGMA_DISABLE_DEPRECATION_WARNINGS // DEPRECATED(4.12)
 		Exchange(NetOldAttachParent, AttachParent);
 		Exchange(NetOldAttachSocketName, AttachSocketName);
-		PRAGMA_ENABLE_DEPRECATION_WARNINGS // DEPRECATED(4.12)
 		AttachToComponent(NetOldAttachParent, FAttachmentTransformRules::KeepRelativeTransform, NetOldAttachSocketName);
 		bNetUpdateAttachment = false;
 	}
@@ -2792,11 +2787,9 @@ void USceneComponent::GetLifetimeReplicatedProps(TArray< FLifetimeProperty > & O
 	DOREPLIFETIME(USceneComponent, bAbsoluteRotation);
 	DOREPLIFETIME(USceneComponent, bAbsoluteScale);
 	DOREPLIFETIME(USceneComponent, bVisible);
-	PRAGMA_DISABLE_DEPRECATION_WARNINGS // DEPRECATED(4.12)
 	DOREPLIFETIME(USceneComponent, AttachParent);
 	DOREPLIFETIME(USceneComponent, AttachChildren);
 	DOREPLIFETIME(USceneComponent, AttachSocketName);
-	PRAGMA_ENABLE_DEPRECATION_WARNINGS // DEPRECATED(4.12)
 	DOREPLIFETIME(USceneComponent, RelativeLocation);
 	DOREPLIFETIME(USceneComponent, RelativeRotation);
 	DOREPLIFETIME(USceneComponent, RelativeScale3D);
diff --git a/Engine/Source/Runtime/Engine/Private/Components/SpotLightComponent.cpp b/Engine/Source/Runtime/Engine/Private/Components/SpotLightComponent.cpp
--- a/Engine/Source/Runtime/Engine/Private/Components/SpotLightComponent.cpp
+++ b/Engine/Source/Runtime/Engine/Private/Components/SpotLightComponent.cpp
@@ -217,9 +217,9 @@ FSphere USpotLightComponent::GetBoundingSphere() const
 	return FSphere(ComponentToWorld.GetLocation() + .5f * GetDirection() * AttenuationRadius, BoundsRadius);
 }
 
-bool USpotLightComponent::AffectsBounds(const FBoxSphereBounds& Bounds) const
+bool USpotLightComponent::AffectsBounds(const FBoxSphereBounds& InBounds) const
 {
-	if(!Super::AffectsBounds(Bounds))
+	if(!Super::AffectsBounds(InBounds))
 	{
 		return false;
 	}
@@ -230,17 +230,17 @@ bool USpotLightComponent::AffectsBounds(const FBoxSphereBounds& Bounds) const
 	float	Sin = FMath::Sin(ClampedOuterConeAngle),
 			Cos = FMath::Cos(ClampedOuterConeAngle);
 
-	FVector	U = GetComponentLocation() - (Bounds.SphereRadius / Sin) * GetDirection(),
-			D = Bounds.Origin - U;
+	FVector	U = GetComponentLocation() - (InBounds.SphereRadius / Sin) * GetDirection(),
+			D = InBounds.Origin - U;
 	float	dsqr = D | D,
 			E = GetDirection() | D;
 	if(E > 0.0f && E * E >= dsqr * FMath::Square(Cos))
 	{
-		D = Bounds.Origin - GetComponentLocation();
+		D = InBounds.Origin - GetComponentLocation();
 		dsqr = D | D;
 		E = -(GetDirection() | D);
 		if(E > 0.0f && E * E >= dsqr * FMath::Square(Sin))
-			return dsqr <= FMath::Square(Bounds.SphereRadius);
+			return dsqr <= FMath::Square(InBounds.SphereRadius);
 		else
 			return true;
 	}
diff --git a/Engine/Source/Runtime/Engine/Private/Components/StaticMeshComponent.cpp b/Engine/Source/Runtime/Engine/Private/Components/StaticMeshComponent.cpp
--- a/Engine/Source/Runtime/Engine/Private/Components/StaticMeshComponent.cpp
+++ b/Engine/Source/Runtime/Engine/Private/Components/StaticMeshComponent.cpp
@@ -987,9 +987,9 @@ void UStaticMeshComponent::CopyInstanceVertexColorsIfCompatible( UStaticMeshComp
 	{
 		Modify();
 
-		bool bRegistered = IsRegistered();
+		bool bIsRegistered = IsRegistered();
 		FComponentReregisterContext ReregisterContext(this);
-		if (bRegistered)
+		if (bIsRegistered)
 		{
 			FlushRenderingCommands(); // don't sync threads unless we have to
 		}
diff --git a/Engine/Source/Runtime/Engine/Private/Components/TextRenderComponent.cpp b/Engine/Source/Runtime/Engine/Private/Components/TextRenderComponent.cpp
--- a/Engine/Source/Runtime/Engine/Private/Components/TextRenderComponent.cpp
+++ b/Engine/Source/Runtime/Engine/Private/Components/TextRenderComponent.cpp
@@ -985,14 +985,14 @@ void UTextRenderComponent::SetWorldSize(float Value)
 
 FVector UTextRenderComponent::GetTextLocalSize() const
 {
-	const FBoxSphereBounds Bounds = CalcBounds(FTransform::Identity);
-	return Bounds.GetBox().GetSize();
+	const FBoxSphereBounds TextBounds = CalcBounds(FTransform::Identity);
+	return TextBounds.GetBox().GetSize();
 }
 
 FVector UTextRenderComponent::GetTextWorldSize() const
 {
-	const FBoxSphereBounds Bounds = CalcBounds(ComponentToWorld);
-	return Bounds.GetBox().GetSize();
+	const FBoxSphereBounds TextBounds = CalcBounds(ComponentToWorld);
+	return TextBounds.GetBox().GetSize();
 }
 
 void UTextRenderComponent::PostLoad()
diff --git a/Engine/Source/Runtime/Engine/Private/ContentStreaming.cpp b/Engine/Source/Runtime/Engine/Private/ContentStreaming.cpp
--- a/Engine/Source/Runtime/Engine/Private/ContentStreaming.cpp
+++ b/Engine/Source/Runtime/Engine/Private/ContentStreaming.cpp
@@ -851,16 +851,16 @@ int32 FStreamingManagerCollection::BlockTillAllRequestsFinished( float TimeLimit
 /** Returns the number of resources that currently wants to be streamed in. */
 int32 FStreamingManagerCollection::GetNumWantingResources() const
 {
-	int32 NumWantingResources = 0;
+	int32 NumResourcesWantingStreaming = 0;
 
 	// Route to streaming managers.
 	for( int32 ManagerIndex=0; ManagerIndex<StreamingManagers.Num(); ManagerIndex++ )
 	{
 		IStreamingManager* StreamingManager = StreamingManagers[ManagerIndex];
-		NumWantingResources += StreamingManager->GetNumWantingResources();
+		NumResourcesWantingStreaming += StreamingManager->GetNumWantingResources();
 	}
 
-	return NumWantingResources;
+	return NumResourcesWantingStreaming;
 }
 
 /**
@@ -871,16 +871,16 @@ int32 FStreamingManagerCollection::GetNumWantingResources() const
  */
 int32 FStreamingManagerCollection::GetNumWantingResourcesID() const
 {
-	int32 NumWantingResourcesCounter = MAX_int32;
+	int32 NumWantingResourcesStreamingCounter = MAX_int32;
 
 	// Route to streaming managers.
 	for( int32 ManagerIndex=0; ManagerIndex<StreamingManagers.Num(); ManagerIndex++ )
 	{
 		IStreamingManager* StreamingManager = StreamingManagers[ManagerIndex];
-		NumWantingResourcesCounter = FMath::Min( NumWantingResourcesCounter, StreamingManager->GetNumWantingResourcesID() );
+		NumWantingResourcesStreamingCounter = FMath::Min( NumWantingResourcesStreamingCounter, StreamingManager->GetNumWantingResourcesID() );
 	}
 
-	return NumWantingResourcesCounter;
+	return NumWantingResourcesStreamingCounter;
 }
 
 /**
diff --git a/Engine/Source/Runtime/Engine/Private/Controller.cpp b/Engine/Source/Runtime/Engine/Private/Controller.cpp
--- a/Engine/Source/Runtime/Engine/Private/Controller.cpp
+++ b/Engine/Source/Runtime/Engine/Private/Controller.cpp
@@ -671,7 +671,7 @@ void AController::StopMovement()
 	UPathFollowingComponent* PathFollowingComp = FindComponentByClass<UPathFollowingComponent>();
 	if (PathFollowingComp != NULL)
 	{
-		PathFollowingComp->AbortMove(TEXT("StopMovement"));
+		PathFollowingComp->AbortMove(*this, FPathFollowingResultFlags::MovementStop);
 	}
 }
 
diff --git a/Engine/Source/Runtime/Engine/Private/GPUSkinVertexFactory.cpp b/Engine/Source/Runtime/Engine/Private/GPUSkinVertexFactory.cpp
--- a/Engine/Source/Runtime/Engine/Private/GPUSkinVertexFactory.cpp
+++ b/Engine/Source/Runtime/Engine/Private/GPUSkinVertexFactory.cpp
@@ -207,7 +207,7 @@ void FGPUBaseSkinVertexFactory::FShaderDataType::GoToNextFrame(uint32 FrameNumbe
 }
 
 bool FGPUBaseSkinVertexFactory::FShaderDataType::UpdateBoneData(FRHICommandListImmediate& RHICmdList, const TArray<FMatrix>& ReferenceToLocalMatrices,
-	const TArray<FBoneIndexType>& BoneMap, uint32 FrameNumber, ERHIFeatureLevel::Type FeatureLevel, bool bUseSkinCache)
+	const TArray<FBoneIndexType>& BoneMap, uint32 FrameNumber, ERHIFeatureLevel::Type InFeatureLevel, bool bUseSkinCache)
 {
 	const uint32 NumBones = BoneMap.Num();
 	check(NumBones <= MaxGPUSkinBones);
@@ -215,7 +215,7 @@ bool FGPUBaseSkinVertexFactory::FShaderDataType::UpdateBoneData(FRHICommandListI
 
 	FVertexBufferAndSRV* CurrentBoneBuffer = 0;
 
-	if (FeatureLevel >= ERHIFeatureLevel::ES3_1)
+	if (InFeatureLevel >= ERHIFeatureLevel::ES3_1)
 	{
 		check(IsInRenderingThread());
 		GoToNextFrame(FrameNumber);
@@ -274,7 +274,7 @@ bool FGPUBaseSkinVertexFactory::FShaderDataType::UpdateBoneData(FRHICommandListI
 			RefToLocal.To3x4MatrixTranspose( (float*)BoneMat.M );
 		}
 	}
-	if (FeatureLevel >= ERHIFeatureLevel::ES3_1)
+	if (InFeatureLevel >= ERHIFeatureLevel::ES3_1)
 	{
 		if (NumBones)
 		{
diff --git a/Engine/Source/Runtime/Engine/Private/GameEngine.cpp b/Engine/Source/Runtime/Engine/Private/GameEngine.cpp
--- a/Engine/Source/Runtime/Engine/Private/GameEngine.cpp
+++ b/Engine/Source/Runtime/Engine/Private/GameEngine.cpp
@@ -1120,10 +1120,10 @@ void UGameEngine::Tick( float DeltaSeconds, bool bIdleMode )
 	}
 
 	// Update Audio. This needs to occur after rendering as the rendering code updates the listener position.
-	FAudioDeviceManager* AudioDeviceManager = GEngine->GetAudioDeviceManager();
-	if (AudioDeviceManager)
+	FAudioDeviceManager* GameAudioDeviceManager = GEngine->GetAudioDeviceManager();
+	if (GameAudioDeviceManager)
 	{
-		AudioDeviceManager->UpdateActiveAudioDevices(bIsAnyNonPreviewWorldUnpaused);
+		GameAudioDeviceManager->UpdateActiveAudioDevices(bIsAnyNonPreviewWorldUnpaused);
 	}
 
 	// rendering thread commands
diff --git a/Engine/Source/Runtime/Engine/Private/GameMode.cpp b/Engine/Source/Runtime/Engine/Private/GameMode.cpp
--- a/Engine/Source/Runtime/Engine/Private/GameMode.cpp
+++ b/Engine/Source/Runtime/Engine/Private/GameMode.cpp
@@ -1063,29 +1063,29 @@ bool AGameMode::MustSpectate_Implementation(APlayerController* NewPlayerControll
 	return NewPlayerController->PlayerState->bOnlySpectator;
 }
 
-APlayerController* AGameMode::Login(UPlayer* NewPlayer, ENetRole RemoteRole, const FString& Portal, const FString& Options, const TSharedPtr<const FUniqueNetId>& UniqueId, FString& ErrorMessage)
+APlayerController* AGameMode::Login(UPlayer* NewPlayer, ENetRole InRemoteRole, const FString& Portal, const FString& Options, const TSharedPtr<const FUniqueNetId>& UniqueId, FString& ErrorMessage)
 {
 	ErrorMessage = GameSession->ApproveLogin(Options);
 	if (!ErrorMessage.IsEmpty())
 	{
-		return NULL;
+		return nullptr;
 	}
 
-	APlayerController* NewPlayerController = SpawnPlayerController(RemoteRole, FVector::ZeroVector, FRotator::ZeroRotator);
+	APlayerController* NewPlayerController = SpawnPlayerController(InRemoteRole, FVector::ZeroVector, FRotator::ZeroRotator);
 
 	// Handle spawn failure.
-	if (NewPlayerController == NULL)
+	if (NewPlayerController == nullptr)
 	{
 		UE_LOG(LogGameMode, Log, TEXT("Couldn't spawn player controller of class %s"), PlayerControllerClass ? *PlayerControllerClass->GetName() : TEXT("NULL"));
 		ErrorMessage = FString::Printf(TEXT("Failed to spawn player controller"));
-		return NULL;
+		return nullptr;
 	}
 
 	// Customize incoming player based on URL options
 	ErrorMessage = InitNewPlayer(NewPlayerController, UniqueId, Options, Portal);
 	if (!ErrorMessage.IsEmpty())
 	{
-		return NULL;
+		return nullptr;
 	}
 
 	// Set up spectating
@@ -1259,7 +1259,7 @@ void AGameMode::PreLogin(const FString& Options, const FString& Address, const T
 	ErrorMessage = GameSession->ApproveLogin(Options);
 }
 
-APlayerController* AGameMode::SpawnPlayerController(ENetRole RemoteRole, FVector const& SpawnLocation, FRotator const& SpawnRotation)
+APlayerController* AGameMode::SpawnPlayerController(ENetRole InRemoteRole, FVector const& SpawnLocation, FRotator const& SpawnRotation)
 {
 	FActorSpawnParameters SpawnInfo;
 	SpawnInfo.Instigator = Instigator;	
@@ -1268,7 +1268,7 @@ APlayerController* AGameMode::SpawnPlayerController(ENetRole RemoteRole, FVector
 	APlayerController* NewPC = GetWorld()->SpawnActor<APlayerController>(PlayerControllerClass, SpawnLocation, SpawnRotation, SpawnInfo);
 	if (NewPC)
 	{
-		if (RemoteRole == ROLE_SimulatedProxy)
+		if (InRemoteRole == ROLE_SimulatedProxy)
 		{
 			// This is a local player because it has no authority/autonomous remote role
 			NewPC->SetAsLocalPlayerController();
diff --git a/Engine/Source/Runtime/Engine/Private/HierarchicalInstancedStaticMesh.cpp b/Engine/Source/Runtime/Engine/Private/HierarchicalInstancedStaticMesh.cpp
--- a/Engine/Source/Runtime/Engine/Private/HierarchicalInstancedStaticMesh.cpp
+++ b/Engine/Source/Runtime/Engine/Private/HierarchicalInstancedStaticMesh.cpp
@@ -1582,12 +1582,12 @@ void FHierarchicalStaticMeshSceneProxy::GetDynamicMeshElements(const TArray<cons
 				const int32 NumUnbuiltInstances = LastUnbuiltIndex - FirstUnbuiltIndex + 1;
 				if (NumUnbuiltInstances < 1000)
 				{
-					const int32 LODs = RenderData->LODResources.Num();
+					const int32 NumLODs = RenderData->LODResources.Num();
 
 					int32 Force = CVarForceLOD.GetValueOnRenderThread();
 					if (Force >= 0)
 					{
-						Force = FMath::Clamp(Force, 0, LODs - 1);
+						Force = FMath::Clamp(Force, 0, NumLODs - 1);
 						InstanceParams.AddRun(Force, Force, FirstUnbuiltIndex, LastUnbuiltIndex);
 					}
 					else
@@ -1605,7 +1605,7 @@ void FHierarchicalStaticMeshSceneProxy::GetDynamicMeshElements(const TArray<cons
 						const float MaxDrawDistanceScale = GetCachedScalabilityCVars().ViewDistanceScale;
 						const float SphereRadius = RenderData->Bounds.SphereRadius;
 
-						for (int32 LODIndex = 0; LODIndex < LODs; LODIndex++)
+						for (int32 LODIndex = 0; LODIndex < NumLODs; LODIndex++)
 						{
 							LODPlanesMax[LODIndex] = MIN_flt;
 							LODPlanesMin[LODIndex] = MAX_flt;
@@ -1629,29 +1629,29 @@ void FHierarchicalStaticMeshSceneProxy::GetDynamicMeshElements(const TArray<cons
 							FinalCull *= MaxDrawDistanceScale;
 							ElementParams.FinalCullDistance = FMath::Max(ElementParams.FinalCullDistance, FinalCull);
 
-							for (int32 LODIndex = 1; LODIndex < LODs; LODIndex++)
+							for (int32 LODIndex = 1; LODIndex < NumLODs; LODIndex++)
 							{
 								float Val = FMath::Min(FMath::Sqrt(Fac / RenderData->ScreenSize[LODIndex]), FinalCull);
 								LODPlanesMin[LODIndex - 1] = FMath::Min(LODPlanesMin[LODIndex - 1], Val - LODRandom);
 								LODPlanesMax[LODIndex - 1] = FMath::Max(LODPlanesMax[LODIndex - 1], Val);
 							}
-							LODPlanesMin[LODs - 1] = FMath::Min(LODPlanesMin[LODs - 1], FinalCull - LODRandom);
-							LODPlanesMax[LODs - 1] = FMath::Max(LODPlanesMax[LODs - 1], FinalCull);
+							LODPlanesMin[NumLODs - 1] = FMath::Min(LODPlanesMin[NumLODs - 1], FinalCull - LODRandom);
+							LODPlanesMax[NumLODs - 1] = FMath::Max(LODPlanesMax[NumLODs - 1], FinalCull);
 						}
 
 						// calculate runs
 						int32 MinLOD = 0;
-						int32 MaxLOD = LODs;
+						int32 MaxLOD = NumLODs;
 						CalcLOD(MinLOD, MaxLOD, UnbuiltBounds[0].Min, UnbuiltBounds[0].Max, ViewOriginInLocalZero, ViewOriginInLocalOne, LODPlanesMin, LODPlanesMax);
 						int32 FirstIndexInRun = 0;
 						for (int32 Index = 1; Index < NumUnbuiltInstances; ++Index)
 						{
 							int32 TempMinLOD = 0;
-							int32 TempMaxLOD = LODs;
+							int32 TempMaxLOD = NumLODs;
 							CalcLOD(TempMinLOD, TempMaxLOD, UnbuiltBounds[Index].Min, UnbuiltBounds[Index].Max, ViewOriginInLocalZero, ViewOriginInLocalOne, LODPlanesMin, LODPlanesMax);
 							if (TempMinLOD != MinLOD)
 							{
-								if (MinLOD < LODs)
+								if (MinLOD < NumLODs)
 								{
 									InstanceParams.AddRun(MinLOD, MinLOD, FirstIndexInRun + FirstUnbuiltIndex, (Index - 1) + FirstUnbuiltIndex);
 								}
@@ -2595,12 +2595,12 @@ int32 UHierarchicalInstancedStaticMeshComponent::GetOverlappingSphereCount(const
 	TArray<FTransform> Transforms;
 	const FBox AABB(Sphere.Center - FVector(Sphere.W), Sphere.Center + FVector(Sphere.W));
 	GatherInstanceTransformsInArea(*this, AABB, 0, Transforms);
-	const FBoxSphereBounds Bounds = StaticMesh->GetBounds();
+	const FBoxSphereBounds MeshBounds = StaticMesh->GetBounds();
 
 	for (const FTransform& TM : Transforms)
 	{
 		const FVector Center = TM.GetLocation();
-		const FSphere InstanceSphere(Center, Bounds.SphereRadius);
+		const FSphere InstanceSphere(Center, MeshBounds.SphereRadius);
 		
 		if (Sphere.Intersects(InstanceSphere))
 		{
@@ -2617,11 +2617,11 @@ int32 UHierarchicalInstancedStaticMeshComponent::GetOverlappingBoxCount(const FB
 	GatherInstanceTransformsInArea(*this, Box, 0, Transforms);
 	
 	int32 Count = 0;
-	const FBoxSphereBounds Bounds = StaticMesh->GetBounds();
+	const FBoxSphereBounds MeshBounds = StaticMesh->GetBounds();
 	for(FTransform& T : Transforms)
 	{
 		const FVector Centre = T.GetLocation();
-		const FBox OtherBox(FVector(Centre - Bounds.BoxExtent), FVector(Centre + Bounds.BoxExtent));
+		const FBox OtherBox(FVector(Centre - MeshBounds.BoxExtent), FVector(Centre + MeshBounds.BoxExtent));
 
 		if(Box.Intersect(OtherBox))
 		{
@@ -2636,13 +2636,13 @@ void UHierarchicalInstancedStaticMeshComponent::GetOverlappingBoxTransforms(cons
 {
 	GatherInstanceTransformsInArea(*this, Box, 0, OutTransforms);
 
-	const FBoxSphereBounds Bounds = StaticMesh->GetBounds();
+	const FBoxSphereBounds MeshBounds = StaticMesh->GetBounds();
 	int32 NumTransforms = OutTransforms.Num();
 	for(int32 Idx = NumTransforms - 1 ; Idx >= 0 ; --Idx)
 	{
 		FTransform& TM = OutTransforms[Idx];
 		const FVector Centre = TM.GetLocation();
-		const FBox OtherBox(FVector(Centre - Bounds.BoxExtent), FVector(Centre + Bounds.BoxExtent));
+		const FBox OtherBox(FVector(Centre - MeshBounds.BoxExtent), FVector(Centre + MeshBounds.BoxExtent));
 
 		if(!Box.Intersect(OtherBox))
 		{
diff --git a/Engine/Source/Runtime/Engine/Private/Interpolation.cpp b/Engine/Source/Runtime/Engine/Private/Interpolation.cpp
--- a/Engine/Source/Runtime/Engine/Private/Interpolation.cpp
+++ b/Engine/Source/Runtime/Engine/Private/Interpolation.cpp
@@ -1691,17 +1691,17 @@ void AMatineeActor::RecaptureActorState()
 	for(int32 i=0; i<GroupInst.Num(); i++)
 	{
 		UInterpGroupInst* GrInst = GroupInst[i];
-		AActor* GroupActor = GrInst->GetGroupActor();
-		if( GroupActor )
+		AActor* InstGroupActor = GrInst->GetGroupActor();
+		if( InstGroupActor )
 		{
-			ConditionallySaveActorState( GrInst, GroupActor );
+			ConditionallySaveActorState( GrInst, InstGroupActor );
 		}
 	}
 	UpdateInterp( SavedScrubPosition, true );
 #endif // WITH_EDITORONLY_DATA
 }
 
-void AMatineeActor::InitGroupActorForGroup(class UInterpGroup* InGroup, class AActor* GroupActor)
+void AMatineeActor::InitGroupActorForGroup(class UInterpGroup* InGroup, class AActor* InGroupActor)
 {
 	bool bFoundGroup = false;
 	for( int32 GroupIndex = 0; GroupIndex < GroupActorInfos.Num(); ++GroupIndex )
@@ -1710,11 +1710,11 @@ void AMatineeActor::InitGroupActorForGroup(class UInterpGroup* InGroup, class AA
 		if( Info.ObjectName == InGroup->GetFName() )
 		{
 			bFoundGroup = true;
-			Info.Actors.AddUnique( GroupActor );
+			Info.Actors.AddUnique( InGroupActor );
 
-			if (GroupActor)
+			if (InGroupActor)
 			{
-				GroupActor->AddControllingMatineeActor(*this);
+				InGroupActor->AddControllingMatineeActor(*this);
 			}
 		}
 	}
@@ -1723,13 +1723,13 @@ void AMatineeActor::InitGroupActorForGroup(class UInterpGroup* InGroup, class AA
 	{
 		FInterpGroupActorInfo NewInfo;
 		NewInfo.ObjectName = InGroup->GetFName();
-		NewInfo.Actors.Add( GroupActor );
+		NewInfo.Actors.Add( InGroupActor );
 
 		GroupActorInfos.Add( NewInfo );
 
-		if (GroupActor)
+		if (InGroupActor)
 		{
-			GroupActor->AddControllingMatineeActor(*this);
+			InGroupActor->AddControllingMatineeActor(*this);
 		}
 	}
 
@@ -6647,12 +6647,12 @@ void UInterpTrackDirector::PreviewUpdateTrack(float NewPosition, UInterpTrackIns
 }
 
 #if WITH_EDITOR
-bool UInterpTrackDirector::UpdatePreviewCamera(AMatineeActor* MatineeActor, bool bIsSelected)
+bool UInterpTrackDirector::UpdatePreviewCamera(AMatineeActor* MatineeActor, bool bInIsSelected)
 {
 	check( MatineeActor );
 #if WITH_EDITORONLY_DATA
 	// If selected, get the viewed actor in the matinee
-	AActor* Actor = bIsSelected ? MatineeActor->FindViewedActor() : NULL;
+	AActor* Actor = bInIsSelected ? MatineeActor->FindViewedActor() : NULL;
 	if ( PreviewCamera != Actor )
 	{
 		// Try casting it to a camera actor
diff --git a/Engine/Source/Runtime/Engine/Private/LODActor.cpp b/Engine/Source/Runtime/Engine/Private/LODActor.cpp
--- a/Engine/Source/Runtime/Engine/Private/LODActor.cpp
+++ b/Engine/Source/Runtime/Engine/Private/LODActor.cpp
@@ -423,10 +423,10 @@ void ALODActor::SetIsDirty(const bool bNewState)
 	if (IsDirty())
 	{
 		// If this LODActor is a SubActor at a higher LOD level mark parent dirty as well
-		UPrimitiveComponent* ParentComponent = StaticMeshComponent->GetLODParentPrimitive();
-		if (ParentComponent)
+		UPrimitiveComponent* LODParentComponent = StaticMeshComponent->GetLODParentPrimitive();
+		if (LODParentComponent)
 		{
-			ALODActor* LODParentActor = Cast<ALODActor>(ParentComponent->GetOwner());
+			ALODActor* LODParentActor = Cast<ALODActor>(LODParentComponent->GetOwner());
 			if (LODParentActor)
 			{
 				LODParentActor->Modify();
diff --git a/Engine/Source/Runtime/Engine/Private/LightMap.cpp b/Engine/Source/Runtime/Engine/Private/LightMap.cpp
--- a/Engine/Source/Runtime/Engine/Private/LightMap.cpp
+++ b/Engine/Source/Runtime/Engine/Private/LightMap.cpp
@@ -721,15 +721,15 @@ bool FLightMapPendingTexture::AddElement(FLightMapAllocationGroup& AllocationGro
 	{
 		auto& Allocation = AllocationGroup.Allocations[iAllocation];
 		uint32 BaseX, BaseY;
-		const uint32 SizeX = Allocation->MappedRect.Width();
-		const uint32 SizeY = Allocation->MappedRect.Height();
-		if (FTextureLayout::AddElement(BaseX, BaseY, SizeX, SizeY))
+		const uint32 MappedRectWidth = Allocation->MappedRect.Width();
+		const uint32 MappedRectHeight = Allocation->MappedRect.Height();
+		if (FTextureLayout::AddElement(BaseX, BaseY, MappedRectWidth, MappedRectHeight))
 		{
 			Allocation->OffsetX = BaseX;
 			Allocation->OffsetY = BaseY;
 
 			// Assumes bAlignByFour
-			NewUnallocatedTexels -= ((SizeX + 3) & ~3) * ((SizeY + 3) & ~3);
+			NewUnallocatedTexels -= ((MappedRectWidth + 3) & ~3) * ((MappedRectHeight + 3) & ~3);
 		}
 		else
 		{
@@ -744,9 +744,9 @@ bool FLightMapPendingTexture::AddElement(FLightMapAllocationGroup& AllocationGro
 		for (--iAllocation; iAllocation >= 0; --iAllocation)
 		{
 			auto& Allocation = AllocationGroup.Allocations[iAllocation];
-			const uint32 SizeX = Allocation->MappedRect.Width();
-			const uint32 SizeY = Allocation->MappedRect.Height();
-			verify(FTextureLayout::RemoveElement(Allocation->OffsetX, Allocation->OffsetY, SizeX, SizeY));
+			const uint32 MappedRectWidth = Allocation->MappedRect.Width();
+			const uint32 MappedRectHeight = Allocation->MappedRect.Height();
+			verify(FTextureLayout::RemoveElement(Allocation->OffsetX, Allocation->OffsetY, MappedRectWidth, MappedRectHeight));
 		}
 		return false;
 	}
diff --git a/Engine/Source/Runtime/Engine/Private/Materials/Material.cpp b/Engine/Source/Runtime/Engine/Private/Materials/Material.cpp
--- a/Engine/Source/Runtime/Engine/Private/Materials/Material.cpp
+++ b/Engine/Source/Runtime/Engine/Private/Materials/Material.cpp
@@ -231,13 +231,13 @@ public:
 			static FName NameSubsurfaceProfile(TEXT("__SubsurfaceProfile"));
 			if (ParameterName == NameSubsurfaceProfile)
 			{
-				const USubsurfaceProfile* SubsurfaceProfileRT = GetSubsurfaceProfileRT();
+				const USubsurfaceProfile* MySubsurfaceProfileRT = GetSubsurfaceProfileRT();
 
 				int32 AllocationId = 0;
-				if(SubsurfaceProfileRT)
+				if(MySubsurfaceProfileRT)
 				{
 					// can be optimized (cached)
-					AllocationId = GSubsurfaceProfileTextureObject.FindAllocationId(SubsurfaceProfileRT);
+					AllocationId = GSubsurfaceProfileTextureObject.FindAllocationId(MySubsurfaceProfileRT);
 				}
 				else
 				{
diff --git a/Engine/Source/Runtime/Engine/Private/Materials/MaterialInstance.cpp b/Engine/Source/Runtime/Engine/Private/Materials/MaterialInstance.cpp
--- a/Engine/Source/Runtime/Engine/Private/Materials/MaterialInstance.cpp
+++ b/Engine/Source/Runtime/Engine/Private/Materials/MaterialInstance.cpp
@@ -80,7 +80,7 @@ FMaterialInstanceResource::FMaterialInstanceResource(UMaterialInstance* InOwner,
 {
 }
 
-const FMaterial* FMaterialInstanceResource::GetMaterial(ERHIFeatureLevel::Type FeatureLevel) const
+const FMaterial* FMaterialInstanceResource::GetMaterial(ERHIFeatureLevel::Type InFeatureLevel) const
 {
 	checkSlow(IsInParallelRenderingThread());
 
@@ -89,7 +89,7 @@ const FMaterial* FMaterialInstanceResource::GetMaterial(ERHIFeatureLevel::Type F
 		if (Owner->bHasStaticPermutationResource)
 		{
 			EMaterialQualityLevel::Type ActiveQualityLevel = GetCachedScalabilityCVars().MaterialQualityLevel;
-			FMaterialResource* StaticPermutationResource = Owner->StaticPermutationMaterialResources[ActiveQualityLevel][FeatureLevel];
+			FMaterialResource* StaticPermutationResource = Owner->StaticPermutationMaterialResources[ActiveQualityLevel][InFeatureLevel];
 
 			if (StaticPermutationResource->GetRenderingThreadShaderMap())
 			{
@@ -104,23 +104,23 @@ const FMaterial* FMaterialInstanceResource::GetMaterial(ERHIFeatureLevel::Type F
 				EMaterialDomain Domain = (EMaterialDomain)StaticPermutationResource->GetMaterialDomain();
 				UMaterial* FallbackMaterial = UMaterial::GetDefaultMaterial(Domain);
 				//there was an error, use the default material's resource
-				return FallbackMaterial->GetRenderProxy(IsSelected(), IsHovered())->GetMaterial(FeatureLevel);
+				return FallbackMaterial->GetRenderProxy(IsSelected(), IsHovered())->GetMaterial(InFeatureLevel);
 			}
 		}
 		else
 		{
 			//use the parent's material resource
-			return Parent->GetRenderProxy(IsSelected(), IsHovered())->GetMaterial(FeatureLevel);
+			return Parent->GetRenderProxy(IsSelected(), IsHovered())->GetMaterial(InFeatureLevel);
 		}
 	}
 	else 
 	{
 		UMaterial* FallbackMaterial = UMaterial::GetDefaultMaterial(MD_Surface);
-		return FallbackMaterial->GetRenderProxy(IsSelected(), IsHovered())->GetMaterial(FeatureLevel);
+		return FallbackMaterial->GetRenderProxy(IsSelected(), IsHovered())->GetMaterial(InFeatureLevel);
 	}
 }
 
-FMaterial* FMaterialInstanceResource::GetMaterialNoFallback(ERHIFeatureLevel::Type FeatureLevel) const
+FMaterial* FMaterialInstanceResource::GetMaterialNoFallback(ERHIFeatureLevel::Type InFeatureLevel) const
 {
 	checkSlow(IsInParallelRenderingThread());
 
@@ -129,7 +129,7 @@ FMaterial* FMaterialInstanceResource::GetMaterialNoFallback(ERHIFeatureLevel::Ty
 		if (Owner->bHasStaticPermutationResource)
 		{
 			EMaterialQualityLevel::Type ActiveQualityLevel = GetCachedScalabilityCVars().MaterialQualityLevel;
-			FMaterialResource* StaticPermutationResource = Owner->StaticPermutationMaterialResources[ActiveQualityLevel][FeatureLevel];
+			FMaterialResource* StaticPermutationResource = Owner->StaticPermutationMaterialResources[ActiveQualityLevel][InFeatureLevel];
 			return StaticPermutationResource;
 		}
 		else
@@ -138,7 +138,7 @@ FMaterial* FMaterialInstanceResource::GetMaterialNoFallback(ERHIFeatureLevel::Ty
 
 			if (ParentProxy)
 			{
-				return ParentProxy->GetMaterialNoFallback(FeatureLevel);
+				return ParentProxy->GetMaterialNoFallback(InFeatureLevel);
 			}
 		}
 	}
@@ -156,13 +156,13 @@ bool FMaterialInstanceResource::GetScalarValue(
 	static FName NameSubsurfaceProfile(TEXT("__SubsurfaceProfile"));
 	if (ParameterName == NameSubsurfaceProfile)
 	{
-		const USubsurfaceProfile* SubsurfaceProfileRT = GetSubsurfaceProfileRT();
+		const USubsurfaceProfile* MySubsurfaceProfileRT = GetSubsurfaceProfileRT();
 
 		int32 AllocationId = 0;
-		if (SubsurfaceProfileRT)
+		if (MySubsurfaceProfileRT)
 		{
 			// can be optimized (cached)
-			AllocationId = GSubsurfaceProfileTextureObject.FindAllocationId(SubsurfaceProfileRT);
+			AllocationId = GSubsurfaceProfileTextureObject.FindAllocationId(MySubsurfaceProfileRT);
 		}
 		else
 		{
diff --git a/Engine/Source/Runtime/Engine/Private/ModelRender.cpp b/Engine/Source/Runtime/Engine/Private/ModelRender.cpp
--- a/Engine/Source/Runtime/Engine/Private/ModelRender.cpp
+++ b/Engine/Source/Runtime/Engine/Private/ModelRender.cpp
@@ -418,8 +418,6 @@ public:
 			// Determine the DPG the primitive should be drawn in.
 			uint8 PrimitiveDPG = GetStaticDepthPriorityGroup();
 
-			const FMatrix LocalToWorld = Component->GetRenderMatrix();
-
 			for(int32 ElementIndex = 0;ElementIndex < Elements.Num();ElementIndex++)
 			{
 				const FModelElement& ModelElement = Component->GetElements()[ElementIndex];
diff --git a/Engine/Source/Runtime/Engine/Private/NiagaraEffectRenderer.cpp b/Engine/Source/Runtime/Engine/Private/NiagaraEffectRenderer.cpp
--- a/Engine/Source/Runtime/Engine/Private/NiagaraEffectRenderer.cpp
+++ b/Engine/Source/Runtime/Engine/Private/NiagaraEffectRenderer.cpp
@@ -638,8 +638,8 @@ DynamicDataRender(NULL)
 		{
 			//FMaterialRenderProxy* MaterialProxy = MeshMaterials[SectionIndex]->GetRenderProxy(bSelected);
 			const FStaticMeshSection& Section = LODModel.Sections[SectionIndex];
-			UMaterialInterface *Material = Properties->ParticleMesh->GetMaterial(Section.MaterialIndex);
-			Material->CheckMaterialUsage_Concurrent(MATUSAGE_MeshParticles);
+			UMaterialInterface* ParticleMeshMaterial = Properties->ParticleMesh->GetMaterial(Section.MaterialIndex);
+			ParticleMeshMaterial->CheckMaterialUsage_Concurrent(MATUSAGE_MeshParticles);
 		}
 	}
 
@@ -854,8 +854,8 @@ void NiagaraEffectRendererMeshes::GetDynamicMeshElements(const TArray<const FSce
 				for (int32 SectionIndex = 0; SectionIndex < LODModel.Sections.Num(); SectionIndex++)
 				{
 					const FStaticMeshSection& Section = LODModel.Sections[SectionIndex];
-					UMaterialInterface *Material = Properties->ParticleMesh->GetMaterial(Section.MaterialIndex);
-					FMaterialRenderProxy* MaterialProxy = Material->GetRenderProxy(false, false);
+					UMaterialInterface* ParticleMeshMaterial = Properties->ParticleMesh->GetMaterial(Section.MaterialIndex);
+					FMaterialRenderProxy* MaterialProxy = ParticleMeshMaterial->GetRenderProxy(false, false);
 
 					if ((Section.NumTriangles == 0) || (MaterialProxy == NULL))
 					{
diff --git a/Engine/Source/Runtime/Engine/Private/PhysicsVolume.cpp b/Engine/Source/Runtime/Engine/Private/PhysicsVolume.cpp
--- a/Engine/Source/Runtime/Engine/Private/PhysicsVolume.cpp
+++ b/Engine/Source/Runtime/Engine/Private/PhysicsVolume.cpp
@@ -61,14 +61,14 @@ void APhysicsVolume::EndPlay(const EEndPlayReason::Type EndPlayReason)
 		if ((EndPlayReason == EEndPlayReason::RemovedFromWorld) || (EndPlayReason == EEndPlayReason::Destroyed))
 		{
 			// Prevent UpdatePhysicsVolume() calls below from returning this component.
-			UPrimitiveComponent* Brush = GetBrushComponent();
-			const bool bSavedGenerateOverlapEvents = Brush->bGenerateOverlapEvents;
-			Brush->bGenerateOverlapEvents = false;
+			UPrimitiveComponent* VolumeBrushComponent = GetBrushComponent();
+			const bool bSavedGenerateOverlapEvents = VolumeBrushComponent->bGenerateOverlapEvents;
+			VolumeBrushComponent->bGenerateOverlapEvents = false;
 
 			// Refresh physics volume on any components touching this volume.
 			// TODO: Physics volume tracking code needs a cleanup, ideally it just uses normal begin/end overlap events,
 			// but there are subtleties with the stacking and priority rules.
-			const TArray<FOverlapInfo>& Overlaps = Brush->GetOverlapInfos();
+			const TArray<FOverlapInfo>& Overlaps = VolumeBrushComponent->GetOverlapInfos();
 			for (const FOverlapInfo& Info : Overlaps)
 			{
 				UPrimitiveComponent* OtherPrim = Info.OverlapInfo.GetComponent();
@@ -79,7 +79,7 @@ void APhysicsVolume::EndPlay(const EEndPlayReason::Type EndPlayReason)
 			}
 
 			// Restore saved flag, since we may stream back in.
-			Brush->bGenerateOverlapEvents = bSavedGenerateOverlapEvents;
+			VolumeBrushComponent->bGenerateOverlapEvents = bSavedGenerateOverlapEvents;
 		}
 	}
 	Super::EndPlay(EndPlayReason);
diff --git a/Engine/Source/Runtime/Engine/Private/PlayerController.cpp b/Engine/Source/Runtime/Engine/Private/PlayerController.cpp
--- a/Engine/Source/Runtime/Engine/Private/PlayerController.cpp
+++ b/Engine/Source/Runtime/Engine/Private/PlayerController.cpp
@@ -2600,12 +2600,13 @@ void APlayerController::ServerViewPrevPlayer_Implementation()
 APlayerState* APlayerController::GetNextViewablePlayer(int32 dir)
 {
 	int32 CurrentIndex = -1;
+	UWorld* MyWorld = GetWorld();
 	if (PlayerCameraManager->ViewTarget.PlayerState )
 	{
 		// Find index of current viewtarget's PlayerState
-		for ( int32 i=0; i<GetWorld()->GameState->PlayerArray.Num(); i++ )
+		for ( int32 i=0; i<MyWorld->GameState->PlayerArray.Num(); i++ )
 		{
-			if (PlayerCameraManager->ViewTarget.PlayerState == GetWorld()->GameState->PlayerArray[i])
+			if (PlayerCameraManager->ViewTarget.PlayerState == MyWorld->GameState->PlayerArray[i])
 			{
 				CurrentIndex = i;
 				break;
@@ -2615,39 +2616,39 @@ APlayerState* APlayerController::GetNextViewablePlayer(int32 dir)
 
 	// Find next valid viewtarget in appropriate direction
 	int32 NewIndex;
-	for ( NewIndex=CurrentIndex+dir; (NewIndex>=0)&&(NewIndex<GetWorld()->GameState->PlayerArray.Num()); NewIndex=NewIndex+dir )
+	for ( NewIndex=CurrentIndex+dir; (NewIndex>=0)&&(NewIndex<MyWorld->GameState->PlayerArray.Num()); NewIndex=NewIndex+dir )
 	{
-		APlayerState* const PlayerState = GetWorld()->GameState->PlayerArray[NewIndex];
-		if ( (PlayerState != NULL) && (Cast<AController>(PlayerState->GetOwner()) != NULL) && (Cast<AController>(PlayerState->GetOwner())->GetPawn() != NULL)
-			&& GetWorld()->GetAuthGameMode()->CanSpectate(this, PlayerState) )
+		APlayerState* const NextPlayerState = MyWorld->GameState->PlayerArray[NewIndex];
+		AController* NextController = (NextPlayerState ? Cast<AController>(NextPlayerState->GetOwner()) : nullptr);
+		if ( NextController && NextController->GetPawn() != nullptr && MyWorld->GetAuthGameMode()->CanSpectate(this, PlayerState) )
 		{
 			return PlayerState;
 		}
 	}
 
 	// wrap around
-	CurrentIndex = (NewIndex < 0) ? GetWorld()->GameState->PlayerArray.Num() : -1;
+	CurrentIndex = (NewIndex < 0) ? MyWorld->GameState->PlayerArray.Num() : -1;
 	for ( NewIndex=CurrentIndex+dir; (NewIndex>=0)&&(NewIndex<GetWorld()->GameState->PlayerArray.Num()); NewIndex=NewIndex+dir )
 	{
-		APlayerState* const PlayerState = GetWorld()->GameState->PlayerArray[NewIndex];
-		if ( (PlayerState != NULL) && (Cast<AController>(PlayerState->GetOwner()) != NULL) && (Cast<AController>(PlayerState->GetOwner())->GetPawn() != NULL) &&
-			GetWorld()->GetAuthGameMode()->CanSpectate(this, PlayerState) )
+		APlayerState* const NextPlayerState = MyWorld->GameState->PlayerArray[NewIndex];
+		AController* NextController = (NextPlayerState ? Cast<AController>(NextPlayerState->GetOwner()) : nullptr);
+		if ( NextController && NextController->GetPawn() != nullptr && MyWorld->GetAuthGameMode()->CanSpectate(this, PlayerState) )
 		{
 			return PlayerState;
 		}
 	}
 
-	return NULL;
+	return nullptr;
 }
 
 
 void APlayerController::ViewAPlayer(int32 dir)
 {
-	APlayerState* const PlayerState = GetNextViewablePlayer(dir);
+	APlayerState* const NextPlayerState = GetNextViewablePlayer(dir);
 
-	if ( PlayerState != NULL )
+	if ( NextPlayerState != nullptr )
 	{
-		SetViewTarget(PlayerState);
+		SetViewTarget(NextPlayerState);
 	}
 }
 
@@ -2781,11 +2782,11 @@ void APlayerController::DisplayDebug(class UCanvas* Canvas, const FDebugDisplayI
 
 		for(int32 i=InputStack.Num() - 1; i >= 0; --i)
 		{
-			AActor* Owner = InputStack[i]->GetOwner();
+			AActor* InputComponentOwner = InputStack[i]->GetOwner();
 			DisplayDebugManager.SetDrawColor(FColor::White);
-			if (Owner)
+			if (InputComponentOwner)
 			{
-				DisplayDebugManager.DrawString(FString::Printf(TEXT(" %s.%s"), *Owner->GetName(), *InputStack[i]->GetName()));
+				DisplayDebugManager.DrawString(FString::Printf(TEXT(" %s.%s"), *InputComponentOwner->GetName(), *InputStack[i]->GetName()));
 			}
 			else
 			{
@@ -4062,12 +4063,12 @@ void APlayerController::SetSpectatorPawn(class ASpectatorPawn* NewSpectatorPawn)
 		else
 		{
 			// clearing the spectator pawn, try to attach to the regular pawn
-			APawn* const Pawn = GetPawn();
-			AttachToPawn(Pawn);
-			AddPawnTickDependency(Pawn);
-			if (Pawn)
+			APawn* const MyPawn = GetPawn();
+			AttachToPawn(MyPawn);
+			AddPawnTickDependency(MyPawn);
+			if (MyPawn)
 			{
-				AutoManageActiveCameraTarget(Pawn);
+				AutoManageActiveCameraTarget(MyPawn);
 			}
 			else
 			{
@@ -4245,12 +4246,12 @@ void APlayerController::SetupInactiveStateInputComponent(UInputComponent* InComp
 }
 
 
-void APlayerController::PushInputComponent(UInputComponent* InputComponent)
+void APlayerController::PushInputComponent(UInputComponent* InInputComponent)
 {
-	if (InputComponent)
+	if (InInputComponent)
 	{
 		bool bPushed = false;
-		CurrentInputStack.RemoveSingle(InputComponent);
+		CurrentInputStack.RemoveSingle(InInputComponent);
 		for (int32 Index = CurrentInputStack.Num() - 1; Index >= 0; --Index)
 		{
 			UInputComponent* IC = CurrentInputStack[Index].Get();
@@ -4258,27 +4259,27 @@ void APlayerController::PushInputComponent(UInputComponent* InputComponent)
 			{
 				CurrentInputStack.RemoveAt(Index);
 			}
-			else if (IC->Priority <= InputComponent->Priority)
+			else if (IC->Priority <= InInputComponent->Priority)
 			{
-				CurrentInputStack.Insert(InputComponent, Index + 1);
+				CurrentInputStack.Insert(InInputComponent, Index + 1);
 				bPushed = true;
 				break;
 			}
 		}
 		if (!bPushed)
 		{
-			CurrentInputStack.Insert(InputComponent, 0);
+			CurrentInputStack.Insert(InInputComponent, 0);
 		}
 	}
 }
 
-bool APlayerController::PopInputComponent(UInputComponent* InputComponent)
+bool APlayerController::PopInputComponent(UInputComponent* InInputComponent)
 {
-	if (InputComponent)
+	if (InInputComponent)
 	{
-		if (CurrentInputStack.RemoveSingle(InputComponent) > 0)
+		if (CurrentInputStack.RemoveSingle(InInputComponent) > 0)
 		{
-			InputComponent->ClearBindingValues();
+			InInputComponent->ClearBindingValues();
 			return true;
 		}
 	}
diff --git a/Engine/Source/Runtime/Engine/Private/Polygon.cpp b/Engine/Source/Runtime/Engine/Private/Polygon.cpp
--- a/Engine/Source/Runtime/Engine/Private/Polygon.cpp
+++ b/Engine/Source/Runtime/Engine/Private/Polygon.cpp
@@ -7,9 +7,7 @@
 #include "EnginePrivate.h"
 #include "Model.h"
 
-#if WITH_EDITOR
 #include "GeomTools.h"
-#endif
 #include "Engine/Polys.h"
 
 
@@ -754,7 +752,7 @@ int32 FPoly::Triangulate( ABrush* InOwnerBrush, TArray<FPoly>& OutTriangles )
 
 	// Attempt to triangulate this polygon
 	TArray<FClipSMTriangle> Triangles;
-	if( TriangulatePoly( Triangles, Polygon ) )
+	if( FGeomTools::TriangulatePoly( Triangles, Polygon ) )
 	{
 		// Create a set of FPolys from the triangles
 
@@ -778,7 +776,6 @@ int32 FPoly::Triangulate( ABrush* InOwnerBrush, TArray<FPoly>& OutTriangles )
 			}
 		}
 	}
-
 #endif
 
 	return OutTriangles.Num();
diff --git a/Engine/Source/Runtime/Engine/Private/SceneManagement.cpp b/Engine/Source/Runtime/Engine/Private/SceneManagement.cpp
--- a/Engine/Source/Runtime/Engine/Private/SceneManagement.cpp
+++ b/Engine/Source/Runtime/Engine/Private/SceneManagement.cpp
@@ -154,22 +154,22 @@ void FSimpleElementCollector::RegisterDynamicResource(FDynamicPrimitiveResource*
 	DynamicResource->InitPrimitiveResource();
 }
 
-void FSimpleElementCollector::DrawBatchedElements(FRHICommandList& RHICmdList, const FSceneView& View, FTexture2DRHIRef DepthTexture, EBlendModeFilter::Type Filter) const
+void FSimpleElementCollector::DrawBatchedElements(FRHICommandList& RHICmdList, const FSceneView& InView, FTexture2DRHIRef DepthTexture, EBlendModeFilter::Type Filter) const
 {
 	// Mobile HDR does not execute post process, so does not need to render flipped
-	const bool bNeedToSwitchVerticalAxis = RHINeedsToSwitchVerticalAxis(View.GetShaderPlatform()) && !bIsMobileHDR;
+	const bool bNeedToSwitchVerticalAxis = RHINeedsToSwitchVerticalAxis(InView.GetShaderPlatform()) && !bIsMobileHDR;
 
 	// Draw the batched elements.
 	BatchedElements.Draw(
 		RHICmdList,
-		View.GetFeatureLevel(),
+		InView.GetFeatureLevel(),
 		bNeedToSwitchVerticalAxis,
-		View.ViewProjectionMatrix,
-		View.ViewRect.Width(),
-		View.ViewRect.Height(),
-		View.Family->EngineShowFlags.HitProxies,
+		InView.ViewProjectionMatrix,
+		InView.ViewRect.Width(),
+		InView.ViewRect.Height(),
+		InView.Family->EngineShowFlags.HitProxies,
 		1.0f,
-		&View,
+		&InView,
 		DepthTexture,
 		Filter
 		);
