Fix Linux compile errors due to misplaced ENGINE_API

diff --git a/Engine/Source/Runtime/Engine/Public/GeomTools.h b/Engine/Source/Runtime/Engine/Public/GeomTools.h
--- a/Engine/Source/Runtime/Engine/Public/GeomTools.h
+++ b/Engine/Source/Runtime/Engine/Public/GeomTools.h
@@ -3,7 +3,7 @@
 #pragma once
 
 /** An edge in 3D space, used by these utility functions. */
-ENGINE_API struct FUtilEdge3D
+struct ENGINE_API FUtilEdge3D
 {
 	/** Start of edge in 3D space. */
 	FVector V0;
@@ -12,7 +12,7 @@ ENGINE_API struct FUtilEdge3D
 };
 
 /** An edge in 2D space, used by these utility functions. */
-ENGINE_API struct FUtilEdge2D
+struct ENGINE_API FUtilEdge2D
 {
 	/** Start of edge in 2D space. */
 	FVector2D V0;
@@ -21,7 +21,7 @@ ENGINE_API struct FUtilEdge2D
 };
 
 /** A triangle vertex in 2D space, with UV information. */
-ENGINE_API struct FUtilVertex2D
+struct ENGINE_API FUtilVertex2D
 {
 	FVector2D Pos;
 	FColor Color;
@@ -40,27 +40,27 @@ ENGINE_API struct FUtilVertex2D
 };
 
 /** A polygon in 2D space, used by utility function. */
-ENGINE_API struct FUtilPoly2D
+struct ENGINE_API FUtilPoly2D
 {
 	/** Set of verts, in order, for polygon. */
 	TArray<FUtilVertex2D> Verts;
 };
 
 /** A set of 2D polygons, along with a transform for going into world space. */
-ENGINE_API struct FUtilPoly2DSet
+struct ENGINE_API FUtilPoly2DSet
 {
 	TArray<FUtilPoly2D>	Polys;
 	FMatrix PolyToWorld;
 };
 
 /** Triangle in 2D space. */
-ENGINE_API struct FUtilTri2D
+struct ENGINE_API FUtilTri2D
 {
 	FUtilVertex2D Verts[3];
 };
 
 /** Temp vertex struct for one vert of a static mesh triangle. */
-ENGINE_API struct FClipSMVertex
+struct ENGINE_API FClipSMVertex
 {
 	FVector Pos;
 	FVector TangentX;
@@ -71,7 +71,7 @@ ENGINE_API struct FClipSMVertex
 };
 
 /** Properties of a clipped static mesh face. */
-ENGINE_API struct FClipSMFace
+struct ENGINE_API FClipSMFace
 {
 	int32 MaterialIndex;
 	uint32 SmoothingMask;
@@ -93,7 +93,7 @@ ENGINE_API struct FClipSMFace
 };
 
 /** Properties of a clipped static mesh triangle. */
-ENGINE_API struct FClipSMTriangle : FClipSMFace
+struct ENGINE_API FClipSMTriangle : FClipSMFace
 {
 	FClipSMVertex Vertices[3];
 
@@ -107,7 +107,7 @@ ENGINE_API struct FClipSMTriangle : FClipSMFace
 };
 
 /** Properties of a clipped static mesh polygon. */
-ENGINE_API struct FClipSMPolygon : FClipSMFace
+struct ENGINE_API FClipSMPolygon : FClipSMFace
 {
 	TArray<FClipSMVertex> Vertices;
 
