Suppresed warnings from main log

diff --git a/Engine/Source/Runtime/Engine/Private/ShaderCompiler/ShaderCompiler.cpp b/Engine/Source/Runtime/Engine/Private/ShaderCompiler/ShaderCompiler.cpp
--- a/Engine/Source/Runtime/Engine/Private/ShaderCompiler/ShaderCompiler.cpp
+++ b/Engine/Source/Runtime/Engine/Private/ShaderCompiler/ShaderCompiler.cpp
@@ -1590,7 +1590,7 @@ void FShaderCompilingManager::ProcessCompiledShaderMaps(
 									*LegacyShaderPlatformToShaderFormat(ShaderMap->GetShaderPlatform()).ToString());
 								for (int32 ErrorIndex = 0; ErrorIndex < Errors.Num(); ErrorIndex++)
 								{
-									UE_LOG(LogShaderCompilers, Warning, TEXT("	%s"), *Errors[ErrorIndex]);
+									UE_LOG(LogShaders, Warning, TEXT("	%s"), *Errors[ErrorIndex]);
 								}
 							}
 						}
