Fixed game window jumping around when FSceneViewport::ResizeFrame was called

InPosX and InPosY had never been used by this function, but a recent HMD change had caused them to be passed into a call to ReshapeWindow. Since their default values were -1,-1 this caused the window to jump.

These change means that those values are no longer used, and I'll make a subsequent change to deprecate the function signature that takes them.

diff --git a/Engine/Source/Runtime/Engine/Private/Slate/SceneViewport.cpp b/Engine/Source/Runtime/Engine/Private/Slate/SceneViewport.cpp
--- a/Engine/Source/Runtime/Engine/Private/Slate/SceneViewport.cpp
+++ b/Engine/Source/Runtime/Engine/Private/Slate/SceneViewport.cpp
@@ -1065,7 +1065,12 @@ void FSceneViewport::ResizeFrame(uint32 NewWindowSizeX, uint32 NewWindowSizeY, E
 				NewWindowSizeY = DisplayMetrics.PrimaryDisplayHeight;
 			}
 
-			FVector2D NewWindowPos(InPosX, InPosY);
+			// Resize window
+			FVector2D WindowPos = WindowToResize->GetPositionInScreen();
+			FVector2D WindowSize = WindowToResize->GetSizeInScreen();
+			FVector2D NewWindowSize(NewWindowSizeX, NewWindowSizeY);
+			FVector2D NewWindowPos = WindowPos;
+
 			IHeadMountedDisplay::MonitorInfo MonitorInfo;
 			if (GEngine->HMDDevice.IsValid() && GEngine->HMDDevice->GetHMDMonitorInfo(MonitorInfo))
 			{
@@ -1077,11 +1082,6 @@ void FSceneViewport::ResizeFrame(uint32 NewWindowSizeX, uint32 NewWindowSizeY, E
 				}
 			}
 
-			// Resize window
-			FVector2D WindowPos = WindowToResize->GetPositionInScreen();
-			FVector2D WindowSize = WindowToResize->GetSizeInScreen();
-			FVector2D NewWindowSize(NewWindowSizeX, NewWindowSizeY);
-
 			if(NewWindowSize != WindowSize || NewWindowPos != WindowPos || DesiredWindowMode != WindowToResize->GetWindowMode())
 			{
 				WindowToResize->SetWindowMode(DesiredWindowMode);
