Copying //UE4/UT-Staging to //UE4/Dev-Main (Source: //UE4/UT-Staging @ 2975641)

Changelist 2970587 by [REDACTED_PERSON]

foliage.MinimumScreenSize is now ECVF_Scalability (Cleared with Martin Mittring)
r.Streaming.Boost is now ECVF_Scalability (Cleared with Martin Mittring)

Changelist 2970537 by [REDACTED_PERSON]
Changelist 2965258 by [REDACTED_PERSON]

CameraComponent::SetCameraMesh() needs to be public

Changelist 2967296 by [REDACTED_PERSON]

UnrealTournament MCP config (not sure what our policy is here, but it's in NotForLicensees and is required for testing UT)

Changelist 2965853 by [REDACTED_PERSON]

Use ENGINE_API to export FNavMeshRenderingHelpers::AddVertex and FNavMeshRenderingHelpers::AddTriangle used by UT nav mesh rendering
Use ENGINE_API to export FCharacterReplaySample << operator

Changelist 2965258 by [REDACTED_PERSON]

FNavMeshRenderingHelpers::AddVertex and FNavMeshRenderingHelpers::AddTriangle put in the header so that UT can consume them
UDemoNetDriver::DemoURL accessor function

Changelist 2964256 by [REDACTED_PERSON]

Remove LexicalConversion::FromString( FString& OutValue, const TCHAR* Buffer ) from CrashReporter project as UT needs it as a lower level

Changelist 2964110 by [REDACTED_PERSON]

Have SSF point at NotForLicensees directory

Changelist 2963798 by [REDACTED_PERSON]

IRepChangedPropertyTracker::SetExternalData and IRepChangedPropertyTracker::IsReplay added
UCanvasRenderTarget2D::FastUpdateResource path added so that resource thrashing doesn't occur
UDemoNetDriver support for game specific headers
SMenuAnchor SetIsOpen needed to be virtual
Added SSF back into NotForLicensees folder (Ben Marsh needs to comment on how this should actually go, I may just have had a really old engine starting point so might not matter)

Changelist 2963380 by [REDACTED_PERSON]

LexicalConversion::FromString( FString& OutValue, const TCHAR* Buffer ) needed
UGameInstance::DelayPendingNetGameTravel() virtual function (approved by [REDACTED_PERSON] in the past, default behavior should be nop and low cost anyway)

Changelist 2807141 by [REDACTED_PERSON]

Change to UT builder email notification preferences

diff --git a/Engine/Source/Programs/AutomationTool/GUBP/GUBP.cs b/Engine/Source/Programs/AutomationTool/GUBP/GUBP.cs
--- a/Engine/Source/Programs/AutomationTool/GUBP/GUBP.cs
+++ b/Engine/Source/Programs/AutomationTool/GUBP/GUBP.cs
@@ -162,7 +162,7 @@ public partial class GUBP : BuildCommand
 		BuildGraph Graph = new BuildGraph(AggregateNodeDefinitions, BuildNodeDefinitions);
 		FindCompletionState(Graph.BuildNodes, TempStorageDir);
 		ComputeDependentFrequencies(Graph.BuildNodes);
-		ValidateTriggerDependencies(Graph.BuildNodes);
+		//ValidateTriggerDependencies(Graph.BuildNodes);
 
 		int TimeIndex = ParseParamInt("UserTimeIndex", -1);
 		if (TimeIndex == -1)
diff --git a/Engine/Source/Runtime/Engine/Classes/AI/Navigation/NavMeshRenderingComponent.h b/Engine/Source/Runtime/Engine/Classes/AI/Navigation/NavMeshRenderingComponent.h
--- a/Engine/Source/Runtime/Engine/Classes/AI/Navigation/NavMeshRenderingComponent.h
+++ b/Engine/Source/Runtime/Engine/Classes/AI/Navigation/NavMeshRenderingComponent.h
@@ -176,3 +176,10 @@ protected:
 	uint32 bForceUpdate : 1;
 	FTimerHandle TimerHandle;
 };
+
+namespace FNavMeshRenderingHelpers
+{
+	ENGINE_API void AddVertex(FNavMeshSceneProxyData::FDebugMeshData& MeshData, const FVector& Pos, const FColor Color = FColor::White);
+
+	ENGINE_API void AddTriangle(FNavMeshSceneProxyData::FDebugMeshData& MeshData, int32 V0, int32 V1, int32 V2);
+}
diff --git a/Engine/Source/Runtime/Engine/Classes/Camera/CameraComponent.h b/Engine/Source/Runtime/Engine/Classes/Camera/CameraComponent.h
--- a/Engine/Source/Runtime/Engine/Classes/Camera/CameraComponent.h
+++ b/Engine/Source/Runtime/Engine/Classes/Camera/CameraComponent.h
@@ -116,6 +116,10 @@ public:
 	UFUNCTION(BlueprintCallable, Category = "Rendering")
 	void AddOrUpdateBlendable(TScriptInterface<IBlendableInterface> InBlendableObject, float InWeight = 1.0f) { PostProcessSettings.AddBlendable(InBlendableObject, InWeight); }
 
+#if WITH_EDITORONLY_DATA
+	virtual void SetCameraMesh(UStaticMesh* Mesh);
+#endif
+
 protected:
 #if WITH_EDITORONLY_DATA
 	// The frustum component used to show visually where the camera field of view is
@@ -126,9 +130,7 @@ protected:
 
 	// The camera mesh to show visually where the camera is placed
 	class UStaticMeshComponent* ProxyMeshComponent;
-
-	virtual void SetCameraMesh(UStaticMesh* Mesh);
-
+	
 	virtual void ResetProxyMeshTransform();
 #endif
 
diff --git a/Engine/Source/Runtime/Engine/Classes/Engine/CanvasRenderTarget2D.h b/Engine/Source/Runtime/Engine/Classes/Engine/CanvasRenderTarget2D.h
--- a/Engine/Source/Runtime/Engine/Classes/Engine/CanvasRenderTarget2D.h
+++ b/Engine/Source/Runtime/Engine/Classes/Engine/CanvasRenderTarget2D.h
@@ -79,8 +79,13 @@ public:
 	// UObject overrides
 	virtual UWorld* GetWorld() const override;
 
+	/** Don't delete the underlying resource if it already exists */
+	void FastUpdateResource();
+
 protected:
 
+	void RepaintCanvas();
+
 	/* The world this render target will be used with */
 	UPROPERTY()
 	TWeakObjectPtr<UWorld> World;
diff --git a/Engine/Source/Runtime/Engine/Classes/Engine/DemoNetDriver.h b/Engine/Source/Runtime/Engine/Classes/Engine/DemoNetDriver.h
--- a/Engine/Source/Runtime/Engine/Classes/Engine/DemoNetDriver.h
+++ b/Engine/Source/Runtime/Engine/Classes/Engine/DemoNetDriver.h
@@ -177,6 +177,8 @@ public:
 	bool IsRecording() const;
 	bool IsPlaying() const;
 
+	FString GetDemoURL() { return DemoURL.ToString(); }
+
 public:
 
 	// FExec interface
@@ -248,4 +250,16 @@ public:
 
 	/** Read the streaming level information from the metadata after the level is loaded */
 	void PendingNetGameLoadMapCompleted();
+
+protected:
+	/** allows subclasses to write game specific data to demo header which is then handled by ProcessGameSpecificDemoHeader */
+	virtual void WriteGameSpecificDemoHeader(TArray<FString>& GameSpecificData)
+	{}
+	/** allows subclasses to read game specific data from demo
+	 * return false to cancel playback
+	 */
+	virtual bool ProcessGameSpecificDemoHeader(const TArray<FString>& GameSpecificData, FString& Error)
+	{
+		return true;
+	}
 };
diff --git a/Engine/Source/Runtime/Engine/Classes/Engine/GameInstance.h b/Engine/Source/Runtime/Engine/Classes/Engine/GameInstance.h
--- a/Engine/Source/Runtime/Engine/Classes/Engine/GameInstance.h
+++ b/Engine/Source/Runtime/Engine/Classes/Engine/GameInstance.h
@@ -235,6 +235,14 @@ public:
 	 */
 	virtual void HandleGameNetControlMessage(class UNetConnection* Connection, uint8 MessageByte, const FString& MessageStr)
 	{}
+	
+	/** return true to delay an otherwise ready-to-join PendingNetGame performing LoadMap() and finishing up
+	 * useful to wait for content downloads, etc
+	 */
+	virtual bool DelayPendingNetGameTravel()
+	{
+		return false;
+	}
 
 	FTimerManager* TimerManager;
 	FLatentActionManager* LatentActionManager;
diff --git a/Engine/Source/Runtime/Engine/Classes/GameFramework/CharacterMovementComponent.h b/Engine/Source/Runtime/Engine/Classes/GameFramework/CharacterMovementComponent.h
--- a/Engine/Source/Runtime/Engine/Classes/GameFramework/CharacterMovementComponent.h
+++ b/Engine/Source/Runtime/Engine/Classes/GameFramework/CharacterMovementComponent.h
@@ -2302,7 +2302,7 @@ public:
 	{
 	}
 
-	friend FArchive& operator<<( FArchive& Ar, FCharacterReplaySample& V );
+	friend ENGINE_API FArchive& operator<<( FArchive& Ar, FCharacterReplaySample& V );
 
 	FVector			Location;
 	FRotator		Rotation;
diff --git a/Engine/Source/Runtime/Engine/Private/CanvasRenderTarget2D.cpp b/Engine/Source/Runtime/Engine/Private/CanvasRenderTarget2D.cpp
--- a/Engine/Source/Runtime/Engine/Private/CanvasRenderTarget2D.cpp
+++ b/Engine/Source/Runtime/Engine/Private/CanvasRenderTarget2D.cpp
@@ -10,20 +10,45 @@ UCanvasRenderTarget2D::UCanvasRenderTarget2D( const FObjectInitializer& ObjectIn
 	bNeedsTwoCopies = false;
 }
 
+
 void UCanvasRenderTarget2D::UpdateResource()
 {
 	// Call parent implementation
 	Super::UpdateResource();
-	
+
 	// Don't allocate canvas object for CRT2D CDO
-	if(IsTemplate())
+	if (IsTemplate())
 	{
 		return;
 	}
 
+	RepaintCanvas();
+}
+
+void UCanvasRenderTarget2D::FastUpdateResource()
+{
+	if (Resource == nullptr)
+	{
+		// We don't have a resource, just take the fast path
+		UpdateResource();
+	}
+	else
+	{
+		// Don't allocate canvas object for CRT2D CDO
+		if (IsTemplate())
+		{
+			return;
+		}
+
+		RepaintCanvas();
+	}
+}
+
+void UCanvasRenderTarget2D::RepaintCanvas()
+{
 	// Create or find the canvas object to use to render onto the texture.  Multiple canvas render target textures can share the same canvas.
-	static const FName CanvasName( TEXT( "CanvasRenderTarget2DCanvas" ) );
-	UCanvas* Canvas = (UCanvas*)StaticFindObjectFast(UCanvas::StaticClass(), GetTransientPackage(), CanvasName );
+	static const FName CanvasName(TEXT("CanvasRenderTarget2DCanvas"));
+	UCanvas* Canvas = (UCanvas*)StaticFindObjectFast(UCanvas::StaticClass(), GetTransientPackage(), CanvasName);
 	if (Canvas == nullptr)
 	{
 		Canvas = NewObject<UCanvas>(GetTransientPackage(), CanvasName);
@@ -39,15 +64,15 @@ void UCanvasRenderTarget2D::UpdateResource()
 
 	// Enqueue the rendering command to set up the rendering canvas.
 	ENQUEUE_UNIQUE_RENDER_COMMAND_ONEPARAMETER
-	(
-		CanvasRenderTargetMakeCurrentCommand,
-		FTextureRenderTarget2DResource*,
-		TextureRenderTarget,
-		(FTextureRenderTarget2DResource*)GameThread_GetRenderTargetResource(),
-		{
-			SetRenderTarget(RHICmdList, TextureRenderTarget->GetRenderTargetTexture(), FTexture2DRHIRef(), true);
-			RHICmdList.SetViewport(0, 0, 0.0f, TextureRenderTarget->GetSizeXY().X, TextureRenderTarget->GetSizeXY().Y, 1.0f);
-		}
+		(
+			CanvasRenderTargetMakeCurrentCommand,
+			FTextureRenderTarget2DResource*,
+			TextureRenderTarget,
+			(FTextureRenderTarget2DResource*)GameThread_GetRenderTargetResource(),
+			{
+				SetRenderTarget(RHICmdList, TextureRenderTarget->GetRenderTargetTexture(), FTexture2DRHIRef(), true);
+	RHICmdList.SetViewport(0, 0, 0.0f, TextureRenderTarget->GetSizeXY().X, TextureRenderTarget->GetSizeXY().Y, 1.0f);
+			}
 	);
 
 	// Create the FCanvas which does the actual rendering.
@@ -70,14 +95,14 @@ void UCanvasRenderTarget2D::UpdateResource()
 	// Enqueue the rendering command to copy the freshly rendering texture resource back to the render target RHI 
 	// so that the texture is updated and available for rendering.
 	ENQUEUE_UNIQUE_RENDER_COMMAND_ONEPARAMETER
-	(
-		CanvasRenderTargetResolveCommand,
-		FTextureRenderTargetResource*,
-		RenderTargetResource,
-		GameThread_GetRenderTargetResource(),
-		{
-			RHICmdList.CopyToResolveTarget(RenderTargetResource->GetRenderTargetTexture(), RenderTargetResource->TextureRHI, true, FResolveParams());
-		}
+		(
+			CanvasRenderTargetResolveCommand,
+			FTextureRenderTargetResource*,
+			RenderTargetResource,
+			GameThread_GetRenderTargetResource(),
+			{
+				RHICmdList.CopyToResolveTarget(RenderTargetResource->GetRenderTargetTexture(), RenderTargetResource->TextureRHI, true, FResolveParams());
+			}
 	);
 }
 
diff --git a/Engine/Source/Runtime/Engine/Private/Components/NavMeshRenderingComponent.cpp b/Engine/Source/Runtime/Engine/Private/Components/NavMeshRenderingComponent.cpp
--- a/Engine/Source/Runtime/Engine/Private/Components/NavMeshRenderingComponent.cpp
+++ b/Engine/Source/Runtime/Engine/Private/Components/NavMeshRenderingComponent.cpp
@@ -141,7 +141,7 @@ namespace FNavMeshRenderingHelpers
 		return FColor(((Col >> 1) & 0x007f7f7f) | (Col & 0xff000000));
 	}
 
-	void AddVertex(FNavMeshSceneProxyData::FDebugMeshData& MeshData, const FVector& Pos, const FColor Color = FColor::White)
+	void AddVertex(FNavMeshSceneProxyData::FDebugMeshData& MeshData, const FVector& Pos, const FColor Color)
 	{
 		const int32 VertexIndex = MeshData.Vertices.Num();
 		FDynamicMeshVertex* Vertex = new(MeshData.Vertices) FDynamicMeshVertex;
diff --git a/Engine/Source/Runtime/Engine/Private/DemoNetDriver.cpp b/Engine/Source/Runtime/Engine/Private/DemoNetDriver.cpp
--- a/Engine/Source/Runtime/Engine/Private/DemoNetDriver.cpp
+++ b/Engine/Source/Runtime/Engine/Private/DemoNetDriver.cpp
@@ -366,6 +366,7 @@ struct FNetworkDemoHeader
 	uint32	InternalProtocolVersion;	// Version of the engine internal network format
 	uint32	EngineNetVersion;			// Version of engine networking format
 	FString LevelName;					// Name of level loaded for demo
+	TArray<FString> GameSpecificData;   // Area for subclasses to write stuff
 	
 	FNetworkDemoHeader() : 
 		Magic( NETWORK_DEMO_MAGIC ), 
@@ -409,6 +410,8 @@ struct FNetworkDemoHeader
 		Ar << Header.EngineNetVersion;
 		Ar << Header.LevelName;
 
+		Ar << Header.GameSpecificData;
+
 		return Ar;
 	}
 };
@@ -535,6 +538,13 @@ bool UDemoNetDriver::InitConnectInternal( FString& Error )
 		return false;
 	}
 
+	if (!ProcessGameSpecificDemoHeader(DemoHeader.GameSpecificData, Error))
+	{
+		UE_LOG(LogDemo, Error, TEXT("UDemoNetDriver::InitConnect: (Game Specific) %s"), *Error);
+		GameInstance->HandleDemoPlaybackFailure(EDemoPlayFailure::Generic, Error);
+		return false;
+	}
+
 	// Create fake control channel
 	ServerConnection->CreateChannel( CHTYPE_Control, 1 );
 	
@@ -664,6 +674,8 @@ bool UDemoNetDriver::InitListen( FNetworkNotify* InNotify, FURL& ListenURL, bool
 
 	DemoHeader.LevelName = World->GetCurrentLevel()->GetOutermost()->GetName();
 
+	WriteGameSpecificDemoHeader(DemoHeader.GameSpecificData);
+
 	// Write the header
 	(*FileAr) << DemoHeader;
 	FileAr->Flush();
diff --git a/Engine/Source/Runtime/Engine/Private/HierarchicalInstancedStaticMesh.cpp b/Engine/Source/Runtime/Engine/Private/HierarchicalInstancedStaticMesh.cpp
--- a/Engine/Source/Runtime/Engine/Private/HierarchicalInstancedStaticMesh.cpp
+++ b/Engine/Source/Runtime/Engine/Private/HierarchicalInstancedStaticMesh.cpp
@@ -54,7 +54,8 @@ static TAutoConsoleVariable<int32> CVarMaxTrianglesToRender(
 TAutoConsoleVariable<float> CVarFoliageMinimumScreenSize(
 	TEXT("foliage.MinimumScreenSize"),
 	0.000005f,
-	TEXT("This controls the screen size at which we cull foliage instances entirely."));
+	TEXT("This controls the screen size at which we cull foliage instances entirely."),
+	ECVF_Scalability);
 
 TAutoConsoleVariable<float> CVarFoliageLODDistanceScale(
 	TEXT("foliage.LODDistanceScale"),
diff --git a/Engine/Source/Runtime/Engine/Private/Streaming/TextureStreamingHelpers.cpp b/Engine/Source/Runtime/Engine/Private/Streaming/TextureStreamingHelpers.cpp
--- a/Engine/Source/Runtime/Engine/Private/Streaming/TextureStreamingHelpers.cpp
+++ b/Engine/Source/Runtime/Engine/Private/Streaming/TextureStreamingHelpers.cpp
@@ -72,7 +72,7 @@ TAutoConsoleVariable<float> CVarStreamingBoost(
 	TEXT("=1.0: normal\n")
 	TEXT("<1.0: decrease wanted mip levels\n")
 	TEXT(">1.0: increase wanted mip levels"),
-	ECVF_Default
+	ECVF_Scalability
 	);
 
 TAutoConsoleVariable<float> CVarStreamingScreenSizeEffectiveMax(
diff --git a/Engine/Source/Runtime/Engine/Private/UnrealEngine.cpp b/Engine/Source/Runtime/Engine/Private/UnrealEngine.cpp
--- a/Engine/Source/Runtime/Engine/Private/UnrealEngine.cpp
+++ b/Engine/Source/Runtime/Engine/Private/UnrealEngine.cpp
@@ -9189,7 +9189,7 @@ void UEngine::TickWorldTravel(FWorldContext& Context, float DeltaSeconds)
 			BroadcastNetworkFailure(NULL, Context.PendingNetGame->NetDriver, ENetworkFailure::PendingConnectionFailure, Context.PendingNetGame->ConnectionError);
 			CancelPending(Context);
 		}
-		else if( Context.PendingNetGame && Context.PendingNetGame->bSuccessfullyConnected && !Context.PendingNetGame->bSentJoinRequest )
+		else if (Context.PendingNetGame && Context.PendingNetGame->bSuccessfullyConnected && !Context.PendingNetGame->bSentJoinRequest && (Context.OwningGameInstance == NULL || !Context.OwningGameInstance->DelayPendingNetGameTravel()))
 		{
 			if (!MakeSureMapNameIsValid(Context.PendingNetGame->URL.Map))
 			{
diff --git a/Engine/Source/Runtime/Slate/Public/Widgets/Input/SMenuAnchor.h b/Engine/Source/Runtime/Slate/Public/Widgets/Input/SMenuAnchor.h
--- a/Engine/Source/Runtime/Slate/Public/Widgets/Input/SMenuAnchor.h
+++ b/Engine/Source/Runtime/Slate/Public/Widgets/Input/SMenuAnchor.h
@@ -77,9 +77,9 @@ public:
 	 * Open or close the popup
 	 *
 	 * @param InIsOpen    If true, open the popup. Otherwise close it.
-	 * @param bFocusMenu  Shoudl we focus the popup as soon as it opens?
+	 * @param bFocusMenu  Should we focus the popup as soon as it opens?
 	 */
-	void SetIsOpen( bool InIsOpen, const bool bFocusMenu = true );
+	virtual void SetIsOpen( bool InIsOpen, const bool bFocusMenu = true );
 	
 	/** @return true if the popup is open; false otherwise. */
 	bool IsOpen() const;
