Fixed some more fallout from changes to use the window position when changing the game viewport mode

- FSceneViewport::ResizeFrame:
  - Fixed the HMD monitor info setting the wrong variables.
  - Fixed SetWindowMode and ResizeViewport potentially being passed two different modes.
  - We now only move the window if we need to (this avoids issues with WindowedFullscreen window positioning).

- FWindowsWindow::MoveWindowTo:
  - Now treats the screen space position it's given as relative to the top-left of the window, rather than the top-left of the windows' client area.

- FWindowsApplication:
  - WM_MOVE was passing a screen space position relative to the top-left of the windows' client area, rather than its window area like Slate expected.

diff --git a/Engine/Source/Runtime/Core/Private/Windows/WindowsApplication.cpp b/Engine/Source/Runtime/Core/Private/Windows/WindowsApplication.cpp
--- a/Engine/Source/Runtime/Core/Private/Windows/WindowsApplication.cpp
+++ b/Engine/Source/Runtime/Core/Private/Windows/WindowsApplication.cpp
@@ -902,15 +902,21 @@ int32 FWindowsApplication::ProcessMessage( HWND hwnd, uint32 msg, WPARAM wParam,
 
 		case WM_MOVE:
 			{
-				// client area position
-				const int32 NewX = (int)(short)(LOWORD(lParam));
-				const int32 NewY = (int)(short)(HIWORD(lParam));
-				FIntPoint NewPosition(NewX,NewY);
+				// client area position (needed to test for minimized windows)
+				int32 NewX = (int)(short)(LOWORD(lParam));
+				int32 NewY = (int)(short)(HIWORD(lParam));
 
 				// Only cache the screen position if its not minimized
-				if ( FWindowsApplication::MinimizedWindowPosition != NewPosition )
+				if (FWindowsApplication::MinimizedWindowPosition.X != NewX && FWindowsApplication::MinimizedWindowPosition.Y != NewY)
 				{
-					MessageHandler->OnMovedWindow( CurrentNativeEventWindow, NewX, NewY );
+					// Slate expects screen space, but WM_MOVE is given in client space - get the real window rect
+					RECT WndRect;
+					::GetWindowRect(hwnd, &WndRect);
+
+					NewX = WndRect.left;
+					NewY = WndRect.top;
+
+					MessageHandler->OnMovedWindow(CurrentNativeEventWindow, NewX, NewY);
 
 					return 0;
 				}
diff --git a/Engine/Source/Runtime/Core/Private/Windows/WindowsWindow.cpp b/Engine/Source/Runtime/Core/Private/Windows/WindowsWindow.cpp
--- a/Engine/Source/Runtime/Core/Private/Windows/WindowsWindow.cpp
+++ b/Engine/Source/Runtime/Core/Private/Windows/WindowsWindow.cpp
@@ -416,21 +416,7 @@ bool FWindowsWindow::GetFullScreenInfo( int32& X, int32& Y, int32& Width, int32&
 /** Native windows should implement MoveWindowTo by relocating the platform-specific window to (X,Y). */
 void FWindowsWindow::MoveWindowTo( int32 X, int32 Y )
 {
-	RECT WindowRect;
-	::GetWindowRect(HWnd, &WindowRect);
-
-	POINT ClientPoint;
-	ClientPoint.x = 0;
-	ClientPoint.y = 0;
-	::ClientToScreen( HWnd, &ClientPoint );
-
-	const int32 XMoveDistance = X - ClientPoint.x;
-	const int32 YMoveDistance = Y - ClientPoint.y;
-
-	X = WindowRect.left + XMoveDistance;
-	Y = WindowRect.top + YMoveDistance;
-
-	::MoveWindow( HWnd, X, Y, WindowRect.right - WindowRect.left, WindowRect.bottom - WindowRect.top, true );
+	::SetWindowPos(HWnd, nullptr, X, Y, 0, 0, SWP_NOACTIVATE | SWP_NOSIZE | SWP_NOZORDER);
 }
 
 /** Native windows should implement BringToFront by making this window the top-most window (i.e. focused).
diff --git a/Engine/Source/Runtime/Engine/Private/Slate/SceneViewport.cpp b/Engine/Source/Runtime/Engine/Private/Slate/SceneViewport.cpp
--- a/Engine/Source/Runtime/Engine/Private/Slate/SceneViewport.cpp
+++ b/Engine/Source/Runtime/Engine/Private/Slate/SceneViewport.cpp
@@ -1052,12 +1052,12 @@ void FSceneViewport::ResizeFrame(uint32 NewWindowSizeX, uint32 NewWindowSizeY, E
 
 		if( WindowToResize.IsValid() )
 		{
-			EWindowMode::Type DesiredWindowMode = GetWindowModeType(NewWindowMode);
+			NewWindowMode = GetWindowModeType(NewWindowMode);
 
 			// If we're going into windowed fullscreen mode, we always want the window to fill the entire screen.
 			// When we calculate the scene view, we'll check the fullscreen mode and configure the screen percentage
 			// scaling so we actual render to the resolution we've been asked for.
-			if (DesiredWindowMode == EWindowMode::WindowedFullscreen)
+			if (NewWindowMode == EWindowMode::WindowedFullscreen)
 			{
 				FDisplayMetrics DisplayMetrics;
 				FSlateApplication::Get().GetInitialDisplayMetrics(DisplayMetrics);
@@ -1067,42 +1067,49 @@ void FSceneViewport::ResizeFrame(uint32 NewWindowSizeX, uint32 NewWindowSizeY, E
 
 			// Resize window
 			FVector2D WindowPos = WindowToResize->GetPositionInScreen();
-			FVector2D WindowSize = WindowToResize->GetSizeInScreen();
+			FVector2D WindowSize = WindowToResize->GetClientSizeInScreen();
 			FVector2D NewWindowSize(NewWindowSizeX, NewWindowSizeY);
-			FVector2D NewWindowPos = WindowPos;
+			TOptional<FVector2D> NewWindowPos;
 
 			IHeadMountedDisplay::MonitorInfo MonitorInfo;
 			if (GEngine->HMDDevice.IsValid() && GEngine->HMDDevice->GetHMDMonitorInfo(MonitorInfo))
 			{
 				if (MonitorInfo.DesktopX > 0 || MonitorInfo.DesktopY > 0)
 				{
-					NewWindowSizeX = MonitorInfo.ResolutionX;
-					NewWindowSizeY = MonitorInfo.ResolutionY;
+					NewWindowSize.X = MonitorInfo.ResolutionX;
+					NewWindowSize.Y = MonitorInfo.ResolutionY;
 					NewWindowPos = FVector2D(MonitorInfo.DesktopX, MonitorInfo.DesktopY);
 				}
 			}
 
-			if(NewWindowSize != WindowSize || NewWindowPos != WindowPos || DesiredWindowMode != WindowToResize->GetWindowMode())
+			if (NewWindowSize != WindowSize || (NewWindowPos.IsSet() && NewWindowPos != WindowPos) || NewWindowMode != WindowMode)
 			{
-				WindowToResize->SetWindowMode(DesiredWindowMode);
+				WindowToResize->SetWindowMode(NewWindowMode);
 				LockMouseToViewport(!CurrentReplyState.ShouldReleaseMouseLock());
-				WindowToResize->ReshapeWindow(NewWindowPos, NewWindowSize);
+				if (NewWindowPos.IsSet())
+				{
+					WindowToResize->ReshapeWindow(NewWindowPos.GetValue(), NewWindowSize);
+				}
+				else
+				{
+					WindowToResize->Resize(NewWindowSize);
+				}
 			}
 
 			// Resize viewport
 			FVector2D ViewportSize(SizeX, SizeY);
 			FVector2D NewViewportSize = WindowToResize->GetViewportSize();
 
-			if(NewViewportSize != ViewportSize || NewWindowMode != WindowMode)
+			if (NewViewportSize != ViewportSize || NewWindowMode != WindowMode)
 			{
-				ResizeViewport(NewViewportSize.X, NewViewportSize.Y, NewWindowMode, NewWindowPos.X, NewWindowPos.Y);
+				ResizeViewport(NewViewportSize.X, NewViewportSize.Y, NewWindowMode, 0, 0);
 			}
 
 			// Resize backbuffer
 			FVector2D BackBufferSize = WindowToResize->IsMirrorWindow() ? WindowSize : ViewportSize;
 			FVector2D NewBackbufferSize = WindowToResize->IsMirrorWindow() ? NewWindowSize : NewViewportSize;
 			
-			if(NewBackbufferSize != BackBufferSize)
+			if (NewBackbufferSize != BackBufferSize)
 			{
 				FSlateApplicationBase::Get().GetRenderer()->UpdateFullscreenState(WindowToResize.ToSharedRef(), NewBackbufferSize.X, NewBackbufferSize.Y);
 			}
