Copying //UE4/Dev-Rendering to //UE4/Dev-Main (Source: //UE4/Dev-Rendering @ 2981742)

diff --git a/Engine/Source/Runtime/Engine/Classes/Kismet/KismetRenderingLibrary.h b/Engine/Source/Runtime/Engine/Classes/Kismet/KismetRenderingLibrary.h
new file mode 100644
--- /dev/null
+++ b/Engine/Source/Runtime/Engine/Classes/Kismet/KismetRenderingLibrary.h
@@ -0,0 +1,48 @@
+// Copyright 1998-2016 Epic Games, Inc. All Rights Reserved.
+
+#pragma once
+
+#include "KismetRenderingLibrary.generated.h"
+
+USTRUCT(BlueprintType)
+struct FDrawToRenderTargetContext
+{
+	GENERATED_USTRUCT_BODY()
+
+	FDrawToRenderTargetContext() :
+		RenderTarget(NULL),
+		DrawEvent(NULL)
+	{}
+
+	UPROPERTY()
+	UTextureRenderTarget2D* RenderTarget;
+
+	TDrawEvent<FRHICommandList>* DrawEvent;
+};
+
+UCLASS(MinimalAPI)
+class UKismetRenderingLibrary : public UBlueprintFunctionLibrary
+{
+	GENERATED_UCLASS_BODY()
+
+	/** 
+	 * Renders a quad with the material applied to the specified render target.   
+	 * This sets the render target even if it is already set, which is an expensive operation. 
+	 * Use BeginDrawCanvasToRenderTarget / EndDrawCanvasToRenderTarget instead if rendering multiple primitives to the same render target.
+	 */
+	UFUNCTION(BlueprintCallable, Category="Rendering", meta=(Keywords="DrawMaterialToRenderTarget", WorldContext="WorldContextObject"))
+	static ENGINE_API void DrawMaterialToRenderTarget(UObject* WorldContextObject, UTextureRenderTarget2D* TextureRenderTarget, UMaterialInterface* Material);
+
+	/** 
+	 * Returns a Canvas object that can be used to draw to the specified render target.
+	 * Be sure to call EndDrawCanvasToRenderTarget to complete the rendering!
+	 */
+	UFUNCTION(BlueprintCallable, Category="Rendering", meta=(Keywords="BeginDrawCanvasToRenderTarget", WorldContext="WorldContextObject"))
+	static ENGINE_API void BeginDrawCanvasToRenderTarget(UObject* WorldContextObject, UTextureRenderTarget2D* TextureRenderTarget, UCanvas*& Canvas, FVector2D& Size, FDrawToRenderTargetContext& Context);
+
+	/**  
+	 * Must be paired with a BeginDrawCanvasToRenderTarget to complete rendering to a render target.
+	 */
+	UFUNCTION(BlueprintCallable, Category="Rendering", meta=(Keywords="EndDrawCanvasToRenderTarget", WorldContext="WorldContextObject"))
+	static ENGINE_API void EndDrawCanvasToRenderTarget(UObject* WorldContextObject, const FDrawToRenderTargetContext& Context);
+};
diff --git a/Engine/Source/Runtime/Engine/Classes/Materials/Material.h b/Engine/Source/Runtime/Engine/Classes/Materials/Material.h
--- a/Engine/Source/Runtime/Engine/Classes/Materials/Material.h
+++ b/Engine/Source/Runtime/Engine/Classes/Materials/Material.h
@@ -443,6 +443,10 @@ public:
 	UPROPERTY(EditAnywhere, Category=Material, AdvancedDisplay)
 	uint32 DitherOpacityMask:1;
 
+	/** Whether the material should allow outputting negative emissive color values.  Only allowed on unlit materials. */
+	UPROPERTY(EditAnywhere, Category=Material, AdvancedDisplay)
+	uint32 bAllowNegativeEmissiveColor:1;
+
 	/** Number of customized UV inputs to display.  Unconnected customized UV inputs will just pass through the vertex UVs. */
 	UPROPERTY(EditAnywhere, Category=Material, AdvancedDisplay, meta=(ClampMin=0))
 	int32 NumCustomizedUVs;
@@ -739,9 +743,14 @@ public:
 	UPROPERTY(EditAnywhere, BlueprintReadWrite, Category=PostProcessMaterial, meta=(DisplayName = "Blendable Priority"))
 	int32 BlendablePriority;
 
+	/** Controls how the Refraction input is interpreted and how the refraction offset into scene color is computed for this material. */
+	UPROPERTY(EditAnywhere, Category=Refraction)
+	TEnumAsByte<enum ERefractionMode> RefractionMode;
+
 	/** This is the refraction depth bias, larger values offset distortion to prevent closer objects from rendering into the distorted surface at acute viewing angles but increases the disconnect between surface and where the refraction starts. */
-	UPROPERTY(EditAnywhere, Category=Material, AdvancedDisplay)
+	UPROPERTY(EditAnywhere, Category=Refraction)
 	float RefractionDepthBias;
+
 	/** 
 	 * Guid that uniquely identifies this material. 
 	 * Any changes to the state of the material that do not appear separately in the shadermap DDC keys must cause this guid to be regenerated!
diff --git a/Engine/Source/Runtime/Engine/Classes/Materials/MaterialExpressionTransform.h b/Engine/Source/Runtime/Engine/Classes/Materials/MaterialExpressionTransform.h
--- a/Engine/Source/Runtime/Engine/Classes/Materials/MaterialExpressionTransform.h
+++ b/Engine/Source/Runtime/Engine/Classes/Materials/MaterialExpressionTransform.h
@@ -23,6 +23,9 @@ enum EMaterialVectorCoordTransformSource
 	/** Camera space */
 	TRANSFORMSOURCE_Camera  UMETA(DisplayName="Camera Space"),
 
+	/** Particle space */
+	TRANSFORMSOURCE_ParticleWorld  UMETA(DisplayName = "Mesh particle space"),
+
 	TRANSFORMSOURCE_MAX,
 };
 
@@ -43,7 +46,10 @@ enum EMaterialVectorCoordTransform
 
 	/** Camera space */
 	TRANSFORM_Camera  UMETA(DisplayName="Camera Space"),
-	
+
+	/** Particle space */
+	TRANSFORM_ParticleWorld UMETA(DisplayName = "Mesh particle space"),
+
 	TRANSFORM_MAX,
 };
 
diff --git a/Engine/Source/Runtime/Engine/Classes/Materials/MaterialExpressionTransformPosition.h b/Engine/Source/Runtime/Engine/Classes/Materials/MaterialExpressionTransformPosition.h
--- a/Engine/Source/Runtime/Engine/Classes/Materials/MaterialExpressionTransformPosition.h
+++ b/Engine/Source/Runtime/Engine/Classes/Materials/MaterialExpressionTransformPosition.h
@@ -22,7 +22,10 @@ enum EMaterialPositionTransformSource
 
 	/** Camera space */
 	TRANSFORMPOSSOURCE_Camera  UMETA(DisplayName="Camera Space"),
-	
+
+	/** Particle space */
+	TRANSFORMPOSSOURCE_Particle UMETA(DisplayName = "Mesh Particle Space"),
+
 	TRANSFORMPOSSOURCE_MAX,
 };
 
diff --git a/Engine/Source/Runtime/Engine/Classes/Materials/MaterialParameterCollection.h b/Engine/Source/Runtime/Engine/Classes/Materials/MaterialParameterCollection.h
--- a/Engine/Source/Runtime/Engine/Classes/Materials/MaterialParameterCollection.h
+++ b/Engine/Source/Runtime/Engine/Classes/Materials/MaterialParameterCollection.h
@@ -81,7 +81,9 @@ class UMaterialParameterCollection : public UObject
 	virtual void PreEditChange(class FEditPropertyChain& PropertyAboutToChange) override;
 	virtual void PostEditChangeProperty(FPropertyChangedEvent& PropertyChangedEvent) override;
 #endif // WITH_EDITOR
+	virtual void PostInitProperties() override;
 	virtual void PostLoad() override;
+	virtual void BeginDestroy() override;
 	//~ End UObject Interface.
 
 	/** Finds a parameter name given an Id, returns NAME_None if the parameter was not found. */
@@ -111,9 +113,17 @@ class UMaterialParameterCollection : public UObject
 
 private:
 
+	/** Default resource used when no instance is available. */
+	class FMaterialParameterCollectionInstanceResource* DefaultResource;
+
 	TScopedPointer<FUniformBufferStruct> UniformBufferStruct;
 
 	void CreateBufferStruct();
+
+	/** Gets default values into data to be set on the uniform buffer. */
+	void GetDefaultParameterData(TArray<FVector4>& ParameterData) const;
+
+	void UpdateDefaultResource();
 };
 
 
diff --git a/Engine/Source/Runtime/Engine/Classes/Particles/Beam/ParticleModuleBeamSource.h b/Engine/Source/Runtime/Engine/Classes/Particles/Beam/ParticleModuleBeamSource.h
--- a/Engine/Source/Runtime/Engine/Classes/Particles/Beam/ParticleModuleBeamSource.h
+++ b/Engine/Source/Runtime/Engine/Classes/Particles/Beam/ParticleModuleBeamSource.h
@@ -91,6 +91,8 @@ class UParticleModuleBeamSource : public UParticleModuleBeamBase
 				FBeam2TypeDataPayload* BeamData, const uint8* ParticleBase, 
 				int32& CurrentOffset, int32	ParticleIndex, bool bSpawning,
 				FBeamParticleModifierPayloadData* ModifierData);
+
+	int32 LastSelectedParticleIndex;
 };
 
 
diff --git a/Engine/Source/Runtime/Engine/Classes/Particles/Beam/ParticleModuleBeamTarget.h b/Engine/Source/Runtime/Engine/Classes/Particles/Beam/ParticleModuleBeamTarget.h
--- a/Engine/Source/Runtime/Engine/Classes/Particles/Beam/ParticleModuleBeamTarget.h
+++ b/Engine/Source/Runtime/Engine/Classes/Particles/Beam/ParticleModuleBeamTarget.h
@@ -87,6 +87,8 @@ class UParticleModuleBeamTarget : public UParticleModuleBeamBase
 				FBeam2TypeDataPayload* BeamData, const uint8* ParticleBase, 
 				int32& CurrentOffset, int32	ParticleIndex, bool bSpawning,
 				FBeamParticleModifierPayloadData* ModifierData);
+
+	int32 LastSelectedParticleIndex;
 };
 
 
diff --git a/Engine/Source/Runtime/Engine/Classes/Particles/ParticleSystemComponent.h b/Engine/Source/Runtime/Engine/Classes/Particles/ParticleSystemComponent.h
--- a/Engine/Source/Runtime/Engine/Classes/Particles/ParticleSystemComponent.h
+++ b/Engine/Source/Runtime/Engine/Classes/Particles/ParticleSystemComponent.h
@@ -797,6 +797,95 @@ public:
 	UFUNCTION(BlueprintCallable, Category="Effects|Components|ParticleSystem")
 	virtual void SetBeamTargetStrength(int32 EmitterIndex, float NewTargetStrength, int32 TargetIndex);
 
+	/**
+	*	Get the beam end point
+	*
+	*	@param	EmitterIndex		The index of the emitter to get the value of
+	*
+	*	@return	true		EmitterIndex is valid and End point is set - OutEndPoint is valid
+	*			false		EmitterIndex invalid or End point is not set - OutEndPoint is invalid
+	*/
+	UFUNCTION(BlueprintCallable, Category = "Effects|Components|ParticleSystem")
+	virtual bool GetBeamEndPoint(int32 EmitterIndex, FVector& OutEndPoint) const;
+	
+	/**
+	*	Get the beam source point
+	*
+	*	@param	EmitterIndex		The index of the emitter to get
+	*	@param	SourceIndex			Which beam within the emitter to get
+	*	@param	OutSourcePoint		Value of source point
+	*
+	*	@return	true		EmitterIndex and SourceIndex are valid - OutSourcePoint is valid
+	*			false		EmitterIndex or SourceIndex is invalid - OutSourcePoint is invalid
+	*/
+	UFUNCTION(BlueprintCallable, Category = "Effects|Components|ParticleSystem")
+	virtual bool GetBeamSourcePoint(int32 EmitterIndex, int32 SourceIndex, FVector& OutSourcePoint) const;
+
+	/**
+	*	Get the beam source tangent
+	*
+	*	@param	EmitterIndex		The index of the emitter to get
+	*	@param	SourceIndex			Which beam within the emitter to get
+	*	@param	OutTangentPoint		Value of source tangent
+	*
+	*	@return	true		EmitterIndex and SourceIndex are valid - OutTangentPoint is valid
+	*			false		EmitterIndex or SourceIndex is invalid - OutTangentPoint is invalid
+	*/
+	UFUNCTION(BlueprintCallable, Category = "Effects|Components|ParticleSystem")
+	virtual bool GetBeamSourceTangent(int32 EmitterIndex, int32 SourceIndex, FVector& OutTangentPoint) const;
+
+	/**
+	*	Get the beam source strength
+	*
+	*	@param	EmitterIndex		The index of the emitter to get
+	*	@param	SourceIndex			Which beam within the emitter to get
+	*	@param	OutSourceStrength		Value of source tangent
+	*
+	*	@return	true		EmitterIndex and SourceIndex are valid - OutSourceStrength is valid
+	*			false		EmitterIndex or SourceIndex is invalid - OutSourceStrength is invalid
+	*/
+	UFUNCTION(BlueprintCallable, Category = "Effects|Components|ParticleSystem")
+	virtual bool GetBeamSourceStrength(int32 EmitterIndex, int32 SourceIndex, float& OutSourceStrength) const;
+
+	/**
+	*	Get the beam target point
+	*
+	*	@param	EmitterIndex		The index of the emitter to get
+	*	@param	TargetIndex			Which beam within the emitter to get
+	*	@param	OutTargetPoint		Value of target point
+	*
+	*	@return	true		EmitterIndex and TargetIndex are valid - OutTargetPoint is valid
+	*			false		EmitterIndex or TargetIndex is invalid - OutTargetPoint is invalid
+	*/
+	UFUNCTION(BlueprintCallable, Category = "Effects|Components|ParticleSystem")
+	virtual bool GetBeamTargetPoint(int32 EmitterIndex, int32 TargetIndex, FVector& OutTargetPoint) const;
+
+	/**
+	*	Get the beam target tangent
+	*
+	*	@param	EmitterIndex		The index of the emitter to get
+	*	@param	TargetIndex			Which beam within the emitter to get
+	*	@param	OutTangentPoint		Value of target tangent
+	*
+	*	@return	true		EmitterIndex and TargetIndex are valid - OutTangentPoint is valid
+	*			false		EmitterIndex or TargetIndex is invalid - OutTangentPoint is invalid
+	*/
+	UFUNCTION(BlueprintCallable, Category = "Effects|Components|ParticleSystem")
+	virtual bool GetBeamTargetTangent(int32 EmitterIndex, int32 TargetIndex, FVector& OutTangentPoint) const;
+	
+	/**
+	*	Get the beam target strength
+	*
+	*	@param	EmitterIndex		The index of the emitter to get
+	*	@param	TargetIndex			Which beam within the emitter to get
+	*	@param	OutTargetStrength	Value of target tangent
+	*
+	*	@return	true		EmitterIndex and TargetIndex are valid - OutTargetStrength is valid
+	*			false		EmitterIndex or TargetIndex is invalid - OutTargetStrength is invalid
+	*/
+	UFUNCTION(BlueprintCallable, Category = "Effects|Components|ParticleSystem")
+	virtual bool GetBeamTargetStrength(int32 EmitterIndex, int32 TargetIndex, float& OutTargetStrength) const;
+	
 	/**
 	 *	Enables/Disables a sub-emitter
 	 *
diff --git a/Engine/Source/Runtime/Engine/Private/AudioStreaming.cpp b/Engine/Source/Runtime/Engine/Private/AudioStreaming.cpp
--- a/Engine/Source/Runtime/Engine/Private/AudioStreaming.cpp
+++ b/Engine/Source/Runtime/Engine/Private/AudioStreaming.cpp
@@ -436,7 +436,7 @@ void FAudioStreamingManager::SetDisregardWorldResourcesForFrames(int32 NumFrames
 {
 }
 
-void FAudioStreamingManager::AddPreparedLevel(class ULevel* Level)
+void FAudioStreamingManager::AddLevel(class ULevel* Level)
 {
 }
 
diff --git a/Engine/Source/Runtime/Engine/Private/Components/ModelComponent.cpp b/Engine/Source/Runtime/Engine/Private/Components/ModelComponent.cpp
--- a/Engine/Source/Runtime/Engine/Private/Components/ModelComponent.cpp
+++ b/Engine/Source/Runtime/Engine/Private/Components/ModelComponent.cpp
@@ -372,10 +372,13 @@ void UModelComponent::GetStreamingTextureInfo(FStreamingTextureLevelContext& Lev
 				// Add each texture to the output with the appropriate parameters.
 				for(int32 TextureIndex = 0;TextureIndex < Textures.Num();TextureIndex++)
 				{
+					UTexture2D* Texture2D = Cast<UTexture2D>(Textures[TextureIndex]);
+					if (!Texture2D) continue;
+
 					FStreamingTexturePrimitiveInfo& StreamingTexture = *new(OutStreamingTextures) FStreamingTexturePrimitiveInfo;
 					StreamingTexture.Bounds = SurfaceBoundingSphere;
 					StreamingTexture.TexelFactor = TexelFactor;
-					StreamingTexture.Texture = Textures[TextureIndex];
+					StreamingTexture.Texture = Texture2D;
 				}
 			}
 		}
diff --git a/Engine/Source/Runtime/Engine/Private/Components/PrimitiveComponent.cpp b/Engine/Source/Runtime/Engine/Private/Components/PrimitiveComponent.cpp
--- a/Engine/Source/Runtime/Engine/Private/Components/PrimitiveComponent.cpp
+++ b/Engine/Source/Runtime/Engine/Private/Components/PrimitiveComponent.cpp
@@ -28,6 +28,7 @@
 #include "GameFramework/CheatManager.h"
 #include "GameFramework/DamageType.h"
 #include "Components/ChildActorComponent.h"
+#include "Streaming/TextureStreamingHelpers.h"
 
 #define LOCTEXT_NAMESPACE "PrimitiveComponent"
 
@@ -246,10 +247,20 @@ void UPrimitiveComponent::GetStreamingTextureInfoWithNULLRemoval(FStreamingTextu
 	GetStreamingTextureInfo(LevelContext, OutStreamingTextures);
 	for (int32 Index = 0; Index < OutStreamingTextures.Num(); Index++)
 	{
-		if (!OutStreamingTextures[Index].Texture || !Cast<UTexture2D>(OutStreamingTextures[Index].Texture))
+		const FStreamingTexturePrimitiveInfo& Info = OutStreamingTextures[Index];
+		if (!IsStreamingTexture(Info.Texture))
 		{
 			OutStreamingTextures.RemoveAt(Index--);
 		}
+		else
+		{
+			// Other wise check that everything is setup right.
+			const bool bCanBeStreamedByDistance = Info.TexelFactor > SMALL_NUMBER && Info.Bounds.SphereRadius > SMALL_NUMBER && ensure(FMath::IsFinite(Info.TexelFactor));
+			if (!bForceMipStreaming && !bCanBeStreamedByDistance && !(Info.TexelFactor < 0 && Info.Texture->LODGroup == TEXTUREGROUP_Terrain_Heightmap))
+			{
+				OutStreamingTextures.RemoveAt(Index--);
+			}
+		}
 	}
 }
 
diff --git a/Engine/Source/Runtime/Engine/Private/Components/SkinnedMeshComponent.cpp b/Engine/Source/Runtime/Engine/Private/Components/SkinnedMeshComponent.cpp
--- a/Engine/Source/Runtime/Engine/Private/Components/SkinnedMeshComponent.cpp
+++ b/Engine/Source/Runtime/Engine/Private/Components/SkinnedMeshComponent.cpp
@@ -604,10 +604,13 @@ void USkinnedMeshComponent::GetStreamingTextureInfo(FStreamingTextureLevelContex
 				MaterialInterface->GetUsedTextures(Textures, EMaterialQualityLevel::Num, false, World ? World->FeatureLevel : GMaxRHIFeatureLevel, false);
 				for(int32 TextureIndex = 0; TextureIndex < Textures.Num(); TextureIndex++ )
 				{
+					UTexture2D* Texture2D = Cast<UTexture2D>(Textures[TextureIndex]);
+					if (!Texture2D) continue;
+
 					FStreamingTexturePrimitiveInfo& StreamingTexture = *new(OutStreamingTextures) FStreamingTexturePrimitiveInfo;
 					StreamingTexture.Bounds = Bounds.GetSphere();
 					StreamingTexture.TexelFactor = WorldTexelFactor;
-					StreamingTexture.Texture = Textures[TextureIndex];
+					StreamingTexture.Texture = Texture2D;
 				}
 			}
 		}
diff --git a/Engine/Source/Runtime/Engine/Private/Components/SkyLightComponent.cpp b/Engine/Source/Runtime/Engine/Private/Components/SkyLightComponent.cpp
--- a/Engine/Source/Runtime/Engine/Private/Components/SkyLightComponent.cpp
+++ b/Engine/Source/Runtime/Engine/Private/Components/SkyLightComponent.cpp
@@ -175,7 +175,7 @@ void USkyLightComponent::SetCaptureIsDirty()
 
 void USkyLightComponent::SanatizeCubemapSize()
 {
-	static const int32 MaxCubemapResolution = 4096;
+	static const int32 MaxCubemapResolution = 1024;
 	CubemapResolution = FMath::Min(int32(FMath::RoundUpToPowerOfTwo(CubemapResolution)), MaxCubemapResolution);
 }
 
diff --git a/Engine/Source/Runtime/Engine/Private/Components/SplineMeshComponent.cpp b/Engine/Source/Runtime/Engine/Private/Components/SplineMeshComponent.cpp
--- a/Engine/Source/Runtime/Engine/Private/Components/SplineMeshComponent.cpp
+++ b/Engine/Source/Runtime/Engine/Private/Components/SplineMeshComponent.cpp
@@ -101,12 +101,30 @@ FVertexFactoryShaderParameters* FSplineMeshVertexFactory::ConstructShaderParamet
 
 void FSplineMeshSceneProxy::InitVertexFactory(USplineMeshComponent* InComponent, int32 InLODIndex, FColorVertexBuffer* InOverrideColorVertexBuffer)
 {
+	uint32 TangentXOffset = 0;
+	uint32 TangetnZOffset = 0;
+	uint32 UVsBaseOffset = 0;
+
+	auto& RD = InComponent->StaticMesh->RenderData->LODResources[InLODIndex];
+	SELECT_STATIC_MESH_VERTEX_TYPE(
+		RD.VertexBuffer.GetUseHighPrecisionTangentBasis(),
+		RD.VertexBuffer.GetUseFullPrecisionUVs(),
+		RD.VertexBuffer.GetNumTexCoords(),
+		{
+			TangentXOffset = STRUCT_OFFSET(VertexType, TangentX);
+			TangetnZOffset = STRUCT_OFFSET(VertexType, TangentZ);
+			UVsBaseOffset = STRUCT_OFFSET(VertexType, UVs);
+		});
+
 	// Initialize the static mesh's vertex factory.
-	ENQUEUE_UNIQUE_RENDER_COMMAND_THREEPARAMETER(
+	ENQUEUE_UNIQUE_RENDER_COMMAND_SIXPARAMETER(
 		InitSplineMeshVertexFactory,
 		FStaticMeshLODResources*, RenderData, &InComponent->StaticMesh->RenderData->LODResources[InLODIndex],
 		UStaticMesh*, Parent, InComponent->StaticMesh,
 		bool, bOverrideColorVertexBuffer, !!InOverrideColorVertexBuffer,
+		uint32, TangentXOffset, TangentXOffset,
+		uint32, TangetnZOffset, TangetnZOffset,
+		uint32, UVsBaseOffset, UVsBaseOffset,
 		{
 
 		if ((RenderData->SplineVertexFactory && !bOverrideColorVertexBuffer) || (RenderData->SplineVertexFactoryOverrideColorVertexBuffer && bOverrideColorVertexBuffer))
@@ -134,17 +152,22 @@ void FSplineMeshSceneProxy::InitVertexFactory(USplineMeshComponent* InComponent,
 
 		Data.TangentBasisComponents[0] = FVertexStreamComponent(
 			&RenderData->VertexBuffer,
-			STRUCT_OFFSET(FStaticMeshFullVertex, RawTangentX),
+			TangentXOffset,
 			RenderData->VertexBuffer.GetStride(),
-			RenderData->VertexBuffer.GetUseHighPrecisionTangentBasis() ? VET_URGB10A2N : VET_PackedNormal
+			RenderData->VertexBuffer.GetUseHighPrecisionTangentBasis() ?
+				TStaticMeshVertexTangentTypeSelector<EStaticMeshVertexTangentBasisType::HighPrecision>::VertexElementType : 
+				TStaticMeshVertexTangentTypeSelector<EStaticMeshVertexTangentBasisType::Default>::VertexElementType
 			);
 
 		Data.TangentBasisComponents[1] = FVertexStreamComponent(
 			&RenderData->VertexBuffer,
-			STRUCT_OFFSET(FStaticMeshFullVertex, TangentZ),
+			TangetnZOffset,
 			RenderData->VertexBuffer.GetStride(),
-			VET_UShort2N
+			RenderData->VertexBuffer.GetUseHighPrecisionTangentBasis() ?
+				TStaticMeshVertexTangentTypeSelector<EStaticMeshVertexTangentBasisType::HighPrecision>::VertexElementType : 
+				TStaticMeshVertexTangentTypeSelector<EStaticMeshVertexTangentBasisType::Default>::VertexElementType
 			);
+
 		if (bOverrideColorVertexBuffer)
 		{
 			Data.ColorComponent = FVertexStreamComponent(
@@ -160,83 +183,56 @@ void FSplineMeshSceneProxy::InitVertexFactory(USplineMeshComponent* InComponent,
 		{
 			FColorVertexBuffer* LODColorVertexBuffer = &RenderData->ColorVertexBuffer;
 			if (LODColorVertexBuffer->GetNumVertices() > 0)
-		{
-			Data.ColorComponent = FVertexStreamComponent(
-				LODColorVertexBuffer,
-				0,	// Struct offset to color
-				LODColorVertexBuffer->GetStride(),
-				VET_Color
-				);
-		}
+			{
+				Data.ColorComponent = FVertexStreamComponent(
+					LODColorVertexBuffer,
+					0,	// Struct offset to color
+					LODColorVertexBuffer->GetStride(),
+					VET_Color
+					);
+			}
 		}
 
 		Data.TextureCoordinates.Empty();
 
-		if (!RenderData->VertexBuffer.GetUseFullPrecisionUVs())
-		{
-			int32 UVIndex;
-			for (UVIndex = 0; UVIndex < (int32)RenderData->VertexBuffer.GetNumTexCoords() - 1; UVIndex += 2)
-			{
-				Data.TextureCoordinates.Add(FVertexStreamComponent(
-					&RenderData->VertexBuffer,
-					STRUCT_OFFSET(TStaticMeshFullVertexFloat16UVs<MAX_STATIC_TEXCOORDS>, UVs) + sizeof(FVector2DHalf) * UVIndex,
-					RenderData->VertexBuffer.GetStride(),
-					VET_Half4
-					));
-			}
-			// possible last UV channel if we have an odd number
-			if (UVIndex < (int32)RenderData->VertexBuffer.GetNumTexCoords())
-			{
-				Data.TextureCoordinates.Add(FVertexStreamComponent(
-					&RenderData->VertexBuffer,
-					STRUCT_OFFSET(TStaticMeshFullVertexFloat16UVs<MAX_STATIC_TEXCOORDS>, UVs) + sizeof(FVector2DHalf) * UVIndex,
-					RenderData->VertexBuffer.GetStride(),
-					VET_Half2
-					));
-			}
+		uint32 UVSizeInBytes = RenderData->VertexBuffer.GetUseFullPrecisionUVs() ?
+			sizeof(TStaticMeshVertexUVsTypeSelector<EStaticMeshVertexUVType::HighPrecision>::UVsTypeT) : sizeof(TStaticMeshVertexUVsTypeSelector<EStaticMeshVertexUVType::Default>::UVsTypeT);
 
-			if (Parent->LightMapCoordinateIndex >= 0 && (uint32)Parent->LightMapCoordinateIndex < RenderData->VertexBuffer.GetNumTexCoords())
-			{
-				Data.LightMapCoordinateComponent = FVertexStreamComponent(
-					&RenderData->VertexBuffer,
-					STRUCT_OFFSET(TStaticMeshFullVertexFloat16UVs<MAX_STATIC_TEXCOORDS>, UVs) + sizeof(FVector2DHalf) * Parent->LightMapCoordinateIndex,
-					RenderData->VertexBuffer.GetStride(),
-					VET_Half2
-					);
-			}
+		EVertexElementType UVDoubleWideVertexElementType = RenderData->VertexBuffer.GetUseFullPrecisionUVs() ?
+			VET_Float4 : VET_Half4;
+
+		EVertexElementType UVVertexElementType = RenderData->VertexBuffer.GetUseFullPrecisionUVs() ?
+			VET_Float2 : VET_Half2;
+
+		int32 UVIndex;
+		for (UVIndex = 0; UVIndex < (int32)RenderData->VertexBuffer.GetNumTexCoords() - 1; UVIndex += 2)
+		{
+			Data.TextureCoordinates.Add(FVertexStreamComponent(
+				&RenderData->VertexBuffer,
+				UVsBaseOffset + UVSizeInBytes * UVIndex,
+				RenderData->VertexBuffer.GetStride(),
+				UVDoubleWideVertexElementType
+				));
 		}
-		else
+		// possible last UV channel if we have an odd number
+		if (UVIndex < (int32)RenderData->VertexBuffer.GetNumTexCoords())
 		{
-			int32 UVIndex;
-			for (UVIndex = 0; UVIndex < (int32)RenderData->VertexBuffer.GetNumTexCoords() - 1; UVIndex += 2)
-			{
-				Data.TextureCoordinates.Add(FVertexStreamComponent(
-					&RenderData->VertexBuffer,
-					STRUCT_OFFSET(TStaticMeshFullVertexFloat32UVs<MAX_STATIC_TEXCOORDS>, UVs) + sizeof(FVector2D) * UVIndex,
-					RenderData->VertexBuffer.GetStride(),
-					VET_Float4
-					));
-			}
-			// possible last UV channel if we have an odd number
-			if (UVIndex < (int32)RenderData->VertexBuffer.GetNumTexCoords())
-			{
-				Data.TextureCoordinates.Add(FVertexStreamComponent(
-					&RenderData->VertexBuffer,
-					STRUCT_OFFSET(TStaticMeshFullVertexFloat32UVs<MAX_STATIC_TEXCOORDS>, UVs) + sizeof(FVector2D) * UVIndex,
-					RenderData->VertexBuffer.GetStride(),
-					VET_Float2
-					));
-			}
+			Data.TextureCoordinates.Add(FVertexStreamComponent(
+				&RenderData->VertexBuffer,
+				UVsBaseOffset + UVSizeInBytes * UVIndex,
+				RenderData->VertexBuffer.GetStride(),
+				UVVertexElementType
+				));
+		}
 
-			if (Parent->LightMapCoordinateIndex >= 0 && (uint32)Parent->LightMapCoordinateIndex < RenderData->VertexBuffer.GetNumTexCoords())
-			{
-				Data.LightMapCoordinateComponent = FVertexStreamComponent(
-					&RenderData->VertexBuffer,
-					STRUCT_OFFSET(TStaticMeshFullVertexFloat32UVs<MAX_STATIC_TEXCOORDS>, UVs) + sizeof(FVector2D) * Parent->LightMapCoordinateIndex,
-					RenderData->VertexBuffer.GetStride(),
-					VET_Float2
-					);
-			}
+		if (Parent->LightMapCoordinateIndex >= 0 && (uint32)Parent->LightMapCoordinateIndex < RenderData->VertexBuffer.GetNumTexCoords())
+		{
+			Data.LightMapCoordinateComponent = FVertexStreamComponent(
+				&RenderData->VertexBuffer,
+				UVsBaseOffset + UVSizeInBytes * Parent->LightMapCoordinateIndex,
+				RenderData->VertexBuffer.GetStride(),
+				UVVertexElementType
+				);
 		}
 
 		VertexFactory->SetData(Data);
diff --git a/Engine/Source/Runtime/Engine/Private/Components/StaticMeshComponent.cpp b/Engine/Source/Runtime/Engine/Private/Components/StaticMeshComponent.cpp
--- a/Engine/Source/Runtime/Engine/Private/Components/StaticMeshComponent.cpp
+++ b/Engine/Source/Runtime/Engine/Private/Components/StaticMeshComponent.cpp
@@ -744,7 +744,9 @@ void UStaticMeshComponent::GetStreamingTextureInfo(FStreamingTextureLevelContext
 		const ERHIFeatureLevel::Type FeatureLevel = GetWorld() ? GetWorld()->FeatureLevel : GMaxRHIFeatureLevel;
 		const bool bUseNewMetrics = CVarStreamingUseNewMetrics.GetValueOnGameThread() != 0;
 
-		LevelContext.BindComponent(bUseNewMetrics ? &StreamingTextureData : nullptr, Bounds, WorldTexelFactor, FMath::Max(0.0f, StreamingDistanceMultiplier));
+		const float MeshExtraScale = FMath::Max(0.0f, StaticMesh->StreamingDistanceMultiplier);
+		const float PrimitiveExtraScale = FMath::Max(0.0f, StreamingDistanceMultiplier);
+		LevelContext.BindComponent(bUseNewMetrics ? &StreamingTextureData : nullptr, Bounds, MeshExtraScale * PrimitiveExtraScale, WorldTexelFactor * PrimitiveExtraScale);
 
 		for (int32 LODIndex = 0; LODIndex < StaticMesh->RenderData->LODResources.Num(); ++LODIndex)
 		{
diff --git a/Engine/Source/Runtime/Engine/Private/ContentStreaming.cpp b/Engine/Source/Runtime/Engine/Private/ContentStreaming.cpp
--- a/Engine/Source/Runtime/Engine/Private/ContentStreaming.cpp
+++ b/Engine/Source/Runtime/Engine/Private/ContentStreaming.cpp
@@ -44,73 +44,6 @@ void FlushResourceStreaming()
 	IStreamingManager::Get().BlockTillAllRequestsFinished();
 }
 
-// ----------------------------------------------------------------------
-
-FFloatMipLevel::FFloatMipLevel() 
-{
-	Invalidate();
-}
-
-int32 FFloatMipLevel::ComputeMip(const FStreamingTexture* StreamingTexture, float MipBias, bool bOptimal) const
-{	
-	check(MipBias >= 0);
-
-	int32 MinAllowedMips = 1;
-	int32 MaxAllowedMips = 16;
-
-	if(StreamingTexture)
-	{
-		MinAllowedMips = StreamingTexture->MinAllowedMips;
-		MaxAllowedMips = bOptimal ? StreamingTexture->MaxAllowedOptimalMips : StreamingTexture->MaxAllowedMips;
-
-		if (!StreamingTexture->bCanBeAffectedByMipBias)
-		{
-			MipBias = 0.0f;
-		}
-	}
-
-	// bias the clamped mip level
-	int32 LocalMipLevel = (int32)(FMath::Min(MipLevel, (float)MaxAllowedMips) - MipBias);
-
-	LocalMipLevel = FMath::Max(LocalMipLevel, MinAllowedMips);
-
-	return LocalMipLevel;
-}
-
-bool FFloatMipLevel::operator>=(const FFloatMipLevel& rhs) const
-{
-	return MipLevel >= rhs.MipLevel;
-}
-
-FFloatMipLevel FFloatMipLevel::FromScreenSizeInTexels(float ScreenSizeInTexels)
-{
-	//check(ScreenSizeInTexels >= 0);  @todo: should we figure out why this is negative?
-	const float ScreenSizeInTexelsClamped = FMath::Max(0.f, ScreenSizeInTexels);
-
-	FFloatMipLevel Ret;
-
-	// int
-//	Ret.MipLevel = 1 + FMath::CeilLogTwo( FMath::TruncToInt(ScreenSizeInTexels) );
-	// float
-	Ret.MipLevel = 1.0f + FMath::Log2(ScreenSizeInTexelsClamped);
-
-	return Ret;
-}
-
-FFloatMipLevel FFloatMipLevel::FromMipLevel(int32 InMipLevel)
-{
-	FFloatMipLevel Ret;
-
-	Ret.MipLevel = (float)InMipLevel;
-
-	return Ret;
-}
-
-void FFloatMipLevel::Invalidate()
-{
-	MipLevel = -1.0f;
-}
-
 /*-----------------------------------------------------------------------------
 	Texture tracking.
 -----------------------------------------------------------------------------*/
@@ -128,7 +61,6 @@ struct FTrackedTextureEvent
 	,	WantedMips(0)
 	,	StreamingStatus(0)
 	,	Timestamp(0.0f)
-	,	StreamType(StreamType_Other)
 	,	BoostFactor(1.0f)
 	{
 	}
@@ -138,7 +70,6 @@ struct FTrackedTextureEvent
 	,	WantedMips(0)
 	,	StreamingStatus(0)
 	,	Timestamp(0.0f)
-	,	StreamType(StreamType_Other)
 	,	BoostFactor(1.0f)
 	{
 		TextureName = new TCHAR[InTextureNameString.Len() + 1];
@@ -152,14 +83,10 @@ struct FTrackedTextureEvent
 	int32			NumRequestedMips;
 	/** Number of wanted mips. */
 	int32			WantedMips;
-	/** Number of wanted mips, as calculated by the heuristics for dynamic primitives. */
-	FFloatMipLevel	DynamicWantedMips;
 	/** Streaming status. */
 	int32			StreamingStatus;
 	/** Timestamp, in seconds from startup. */
 	float			Timestamp;
-	/** Which streaming heuristic this texture uses. */
-	ETextureStreamingType	StreamType;
 	/** Currently used boost factor for the streaming distance. */
 	float			BoostFactor;
 };
@@ -277,7 +204,7 @@ bool TrackTextureEvent( FStreamingTexture* StreamingTexture, UTexture2D* Texture
 		TrackTextureInit();
 	}
 
-	float MipBias = Manager ? Manager->ThreadSettings.MipBias : 0;
+	float MipBias =  Manager ? Manager->MipBias : 0;
 
 	int32 NumTrackedTextures = GTrackedTextureNames.Num();
 	if ( NumTrackedTextures )
@@ -312,14 +239,10 @@ bool TrackTextureEvent( FStreamingTexture* StreamingTexture, UTexture2D* Texture
 
 					int32 StreamingStatus	= Texture->PendingMipChangeRequestStatus.GetValue();
 					int32 WantedMips		= Texture->RequestedMips;
-					FFloatMipLevel DynamicWantedMips;
 					float BoostFactor		= 1.0f;
-					ETextureStreamingType StreamType = bForceMipLevelsToBeResident ? StreamType_Forced : StreamType_Other;
 					if ( StreamingTexture )
 					{
 						WantedMips			= StreamingTexture->WantedMips;
-						//@TODO: DynamicWantedMips	= FFloatMipLevel::FromScreenSizeInTexels(StreamingTexture->DynamicScreenSize);
-						StreamType			= StreamingTexture->GetStreamingType();
 						BoostFactor			= StreamingTexture->BoostFactor;
 					}
 
@@ -327,8 +250,6 @@ bool TrackTextureEvent( FStreamingTexture* StreamingTexture, UTexture2D* Texture
 						 LastEvent->NumRequestedMips != Texture->RequestedMips ||
 						 LastEvent->StreamingStatus != StreamingStatus ||
 						 LastEvent->WantedMips != WantedMips ||
-						 LastEvent->DynamicWantedMips.ComputeMip(StreamingTexture, MipBias, false) != DynamicWantedMips.ComputeMip(StreamingTexture, MipBias, false) ||
-						 LastEvent->StreamType != StreamType ||
 						 LastEvent->BoostFactor != BoostFactor ||
 						 bIsDestroying )
 					{
@@ -338,15 +259,12 @@ bool TrackTextureEvent( FStreamingTexture* StreamingTexture, UTexture2D* Texture
 						NewEvent.NumResidentMips		= LastEvent->NumResidentMips	= Texture->ResidentMips;
 						NewEvent.NumRequestedMips		= LastEvent->NumRequestedMips	= Texture->RequestedMips;
 						NewEvent.WantedMips				= LastEvent->WantedMips			= WantedMips;
-						NewEvent.DynamicWantedMips		= LastEvent->DynamicWantedMips	= DynamicWantedMips;
 						NewEvent.StreamingStatus		= LastEvent->StreamingStatus	= StreamingStatus;
 						NewEvent.Timestamp				= LastEvent->Timestamp			= float(FPlatformTime::Seconds() - GStartTime);
-						NewEvent.StreamType				= LastEvent->StreamType			= StreamType;
 						NewEvent.BoostFactor			= LastEvent->BoostFactor		= BoostFactor;
-						UE_LOG(LogContentStreaming, Log, TEXT("Texture: \"%s\", ResidentMips: %d/%d, RequestedMips: %d, WantedMips: %d, DynamicWantedMips: %d, StreamingStatus: %d, StreamType: %s, Boost: %.1f (%s)"),
-							TextureName, LastEvent->NumResidentMips, Texture->GetNumMips(), bIsDestroying ? 0 : LastEvent->NumRequestedMips,
-							LastEvent->WantedMips, LastEvent->DynamicWantedMips.ComputeMip(StreamingTexture, MipBias, false), LastEvent->StreamingStatus, GStreamTypeNames[StreamType],
-							BoostFactor, bIsDestroying ? TEXT("DESTROYED") : TEXT("updated") );
+						UE_LOG(LogContentStreaming, Log, TEXT("Texture: \"%s\", ResidentMips: %d/%d, RequestedMips: %d, WantedMips: %d, StreamingStatus: %d, Boost: %.1f (%s)"),
+							TextureName, LastEvent->NumResidentMips, Texture->GetNumMips(), bIsDestroying ? 0 : LastEvent->NumRequestedMips, LastEvent->WantedMips, 
+							LastEvent->StreamingStatus, BoostFactor, bIsDestroying ? TEXT("DESTROYED") : TEXT("updated") );
 					}
 				}
 				return true;
@@ -1011,15 +929,11 @@ void FStreamingManagerCollection::AddLevel( ULevel* Level )
 	UE_LOG(LogContentStreaming, Log, TEXT("FStreamingManagerCollection::AddLevel(\"%s\")"), *Level->GetOutermost()->GetName());
 #endif
 
-	if ( Level->bTextureStreamingBuilt == false )
-	{
-		// If it hasn't been done yet for this level, build its texture streaming data now.
-		// This will also add the level.
-		Level->BuildStreamingData();
-	}
-	else
+	// Route to streaming managers.
+	for( int32 ManagerIndex=0; ManagerIndex<StreamingManagers.Num(); ManagerIndex++ )
 	{
-		AddPreparedLevel( Level );
+		IStreamingManager* StreamingManager = StreamingManagers[ManagerIndex];
+		StreamingManager->AddLevel( Level );
 	}
 }
 
@@ -1038,23 +952,6 @@ void FStreamingManagerCollection::RemoveLevel( ULevel* Level )
 	}
 }
 
-/** Adds a ULevel that has already built StreamingData to the streaming manager. */
-void FStreamingManagerCollection::AddPreparedLevel( class ULevel* Level )
-{
-#if STREAMING_LOG_LEVELS
-	UE_LOG(LogContentStreaming, Log, TEXT("FStreamingManagerCollection::AddPreparedLevel(\"%s\")"), *Level->GetOutermost()->GetName());
-#endif
-
-	check( Level->bTextureStreamingBuilt );
-
-	// Route to streaming managers.
-	for( int32 ManagerIndex=0; ManagerIndex<StreamingManagers.Num(); ManagerIndex++ )
-	{
-		IStreamingManager* StreamingManager = StreamingManagers[ManagerIndex];
-		StreamingManager->AddPreparedLevel( Level );
-	}
-}
-
 /** Called when an actor is spawned. */
 void FStreamingManagerCollection::NotifyActorSpawned( AActor* Actor )
 {
@@ -1167,9 +1064,6 @@ void FStreamingManagerCollection::AddOrRemoveTextureStreamingManagerIfNeeded(boo
 			GConfig->GetFloat( TEXT("TextureStreaming"), TEXT("LoadMapTimeLimit"), LoadMapTimeLimit, GEngineIni );
 			// Create the streaming manager and add the default streamers.
 			TextureStreamingManager = new FStreamingManagerTexture();
-			TextureStreamingManager->AddTextureStreamingHandler( new FStreamingHandlerTextureDynamic() );
-			TextureStreamingManager->AddTextureStreamingHandler( new FStreamingHandlerTextureStatic() );
-			TextureStreamingManager->AddTextureStreamingHandler( new FStreamingHandlerTextureLevelForced() );
 			AddStreamingManager( TextureStreamingManager );		
 				
 			//Need to work out if all textures should be streamable and added to the texture streaming manager.
@@ -1339,7 +1233,7 @@ FArchive& operator<<( FArchive& Ar, FStreamableTextureInstance& TextureInstance
 	else if (Ar.IsLoading())
 	{
 		TextureInstance.MinDistance = 0;
-		TextureInstance.MaxDistance = MAX_FLT;
+		TextureInstance.MaxDistance = FLT_MAX;
 	}
 
 	Ar << TextureInstance.TexelFactor;
diff --git a/Engine/Source/Runtime/Engine/Private/InstancedStaticMesh.cpp b/Engine/Source/Runtime/Engine/Private/InstancedStaticMesh.cpp
--- a/Engine/Source/Runtime/Engine/Private/InstancedStaticMesh.cpp
+++ b/Engine/Source/Runtime/Engine/Private/InstancedStaticMesh.cpp
@@ -351,21 +351,38 @@ void FInstancedStaticMeshRenderData::InitStaticMeshVertexFactories(
 			VET_Float3
 			);
 
+		uint32 TangentXOffset = 0;
+		uint32 TangetnZOffset = 0;
+		uint32 UVsBaseOffset = 0;
+
+		SELECT_STATIC_MESH_VERTEX_TYPE(
+			RenderData->VertexBuffer.GetUseHighPrecisionTangentBasis(),
+			RenderData->VertexBuffer.GetUseFullPrecisionUVs(),
+			RenderData->VertexBuffer.GetNumTexCoords(),
+			{
+				TangentXOffset = STRUCT_OFFSET(VertexType, TangentX);
+				TangetnZOffset = STRUCT_OFFSET(VertexType, TangentZ);
+				UVsBaseOffset = STRUCT_OFFSET(VertexType, UVs);
+			});
+
 		Data.TangentBasisComponents[0] = FVertexStreamComponent(
 			&RenderData->VertexBuffer,
-			STRUCT_OFFSET(FStaticMeshFullVertex, RawTangentX),
+			TangentXOffset,
 			RenderData->VertexBuffer.GetStride(),
-			RenderData->VertexBuffer.GetUseHighPrecisionTangentBasis() ? VET_URGB10A2N : VET_PackedNormal
+			RenderData->VertexBuffer.GetUseHighPrecisionTangentBasis() ?
+				TStaticMeshVertexTangentTypeSelector<EStaticMeshVertexTangentBasisType::HighPrecision>::VertexElementType : 
+				TStaticMeshVertexTangentTypeSelector<EStaticMeshVertexTangentBasisType::Default>::VertexElementType
 			);
 
 		Data.TangentBasisComponents[1] = FVertexStreamComponent(
 			&RenderData->VertexBuffer,
-			STRUCT_OFFSET(FStaticMeshFullVertex, TangentZ),
+			TangetnZOffset,
 			RenderData->VertexBuffer.GetStride(),
-			VET_UShort2N
+			RenderData->VertexBuffer.GetUseHighPrecisionTangentBasis() ?
+				TStaticMeshVertexTangentTypeSelector<EStaticMeshVertexTangentBasisType::HighPrecision>::VertexElementType : 
+				TStaticMeshVertexTangentTypeSelector<EStaticMeshVertexTangentBasisType::Default>::VertexElementType
 			);
 
-
 		if( RenderData->ColorVertexBuffer.GetNumVertices() > 0 )
 		{
 			Data.ColorComponent = FVertexStreamComponent(
@@ -377,73 +394,48 @@ void FInstancedStaticMeshRenderData::InitStaticMeshVertexFactories(
 		}
 
 		Data.TextureCoordinates.Empty();
+
+		uint32 UVSizeInBytes = RenderData->VertexBuffer.GetUseFullPrecisionUVs() ?
+			sizeof(TStaticMeshVertexUVsTypeSelector<EStaticMeshVertexUVType::HighPrecision>::UVsTypeT) : sizeof(TStaticMeshVertexUVsTypeSelector<EStaticMeshVertexUVType::Default>::UVsTypeT);
+
+		EVertexElementType UVDoubleWideVertexElementType = RenderData->VertexBuffer.GetUseFullPrecisionUVs() ?
+			VET_Float4 : VET_Half4;
+
+		EVertexElementType UVVertexElementType = RenderData->VertexBuffer.GetUseFullPrecisionUVs() ?
+			VET_Float2 : VET_Half2;
+
 		// Only bind InstancedStaticMeshMaxTexCoord, even if the mesh has more.
 		int32 NumTexCoords = FMath::Min<int32>((int32)RenderData->VertexBuffer.GetNumTexCoords(), InstancedStaticMeshMaxTexCoord);
-		if( !RenderData->VertexBuffer.GetUseFullPrecisionUVs() )
-		{
-			int32 UVIndex;
-			for (UVIndex = 0; UVIndex < NumTexCoords - 1; UVIndex += 2)
-			{
-				Data.TextureCoordinates.Add(FVertexStreamComponent(
-					&RenderData->VertexBuffer,
-					STRUCT_OFFSET(TStaticMeshFullVertexFloat16UVs<MAX_STATIC_TEXCOORDS>, UVs) + sizeof(FVector2DHalf)* UVIndex,
-					RenderData->VertexBuffer.GetStride(),
-					VET_Half4
-					));
-			}
-			// possible last UV channel if we have an odd number
-			if( UVIndex < NumTexCoords )
-			{
-				Data.TextureCoordinates.Add(FVertexStreamComponent(
-					&RenderData->VertexBuffer,
-					STRUCT_OFFSET(TStaticMeshFullVertexFloat16UVs<MAX_STATIC_TEXCOORDS>,UVs) + sizeof(FVector2DHalf) * UVIndex,
-					RenderData->VertexBuffer.GetStride(),
-					VET_Half2
-					));
-			}
 
-			if (Parent->LightMapCoordinateIndex >= 0 && Parent->LightMapCoordinateIndex < NumTexCoords)
-			{
-				Data.LightMapCoordinateComponent = FVertexStreamComponent(
-					&RenderData->VertexBuffer,
-					STRUCT_OFFSET(TStaticMeshFullVertexFloat16UVs<MAX_STATIC_TEXCOORDS>, UVs) + sizeof(FVector2DHalf)* Parent->LightMapCoordinateIndex,
-					RenderData->VertexBuffer.GetStride(),
-					VET_Half2
-					);
-			}
+		int32 UVIndex;
+		for (UVIndex = 0; UVIndex < NumTexCoords - 1; UVIndex += 2)
+		{
+			Data.TextureCoordinates.Add(FVertexStreamComponent(
+				&RenderData->VertexBuffer,
+				UVsBaseOffset + UVSizeInBytes * UVIndex,
+				RenderData->VertexBuffer.GetStride(),
+				UVDoubleWideVertexElementType
+				));
 		}
-		else
+		// possible last UV channel if we have an odd number
+		if( UVIndex < NumTexCoords )
 		{
-			int32 UVIndex;
-			for (UVIndex = 0; UVIndex < NumTexCoords - 1; UVIndex += 2)
-			{
-				Data.TextureCoordinates.Add(FVertexStreamComponent(
-					&RenderData->VertexBuffer,
-					STRUCT_OFFSET(TStaticMeshFullVertexFloat32UVs<MAX_STATIC_TEXCOORDS>, UVs) + sizeof(FVector2D)* UVIndex,
-					RenderData->VertexBuffer.GetStride(),
-					VET_Float4
-					));
-			}
-			// possible last UV channel if we have an odd number
-			if (UVIndex < NumTexCoords)
-			{
-				Data.TextureCoordinates.Add(FVertexStreamComponent(
-					&RenderData->VertexBuffer,
-					STRUCT_OFFSET(TStaticMeshFullVertexFloat32UVs<MAX_STATIC_TEXCOORDS>,UVs) + sizeof(FVector2D) * UVIndex,
-					RenderData->VertexBuffer.GetStride(),
-					VET_Float2
-					));
-			}
+			Data.TextureCoordinates.Add(FVertexStreamComponent(
+				&RenderData->VertexBuffer,
+				UVsBaseOffset + UVSizeInBytes * UVIndex,
+				RenderData->VertexBuffer.GetStride(),
+				UVVertexElementType
+				));
+		}
 
-			if (Parent->LightMapCoordinateIndex >= 0 && Parent->LightMapCoordinateIndex < NumTexCoords)
-			{
-				Data.LightMapCoordinateComponent = FVertexStreamComponent(
-					&RenderData->VertexBuffer,
-					STRUCT_OFFSET(TStaticMeshFullVertexFloat32UVs<InstancedStaticMeshMaxTexCoord>, UVs) + sizeof(FVector2D)* Parent->LightMapCoordinateIndex,
-					RenderData->VertexBuffer.GetStride(),
-					VET_Float2
-					);
-			}
+		if (Parent->LightMapCoordinateIndex >= 0 && Parent->LightMapCoordinateIndex < NumTexCoords)
+		{
+			Data.LightMapCoordinateComponent = FVertexStreamComponent(
+				&RenderData->VertexBuffer,
+				UVsBaseOffset + UVSizeInBytes * Parent->LightMapCoordinateIndex,
+				RenderData->VertexBuffer.GetStride(),
+				UVVertexElementType
+				);
 		}
 
 		if (bInstanced)
diff --git a/Engine/Source/Runtime/Engine/Private/KismetRenderingLibrary.cpp b/Engine/Source/Runtime/Engine/Private/KismetRenderingLibrary.cpp
new file mode 100644
--- /dev/null
+++ b/Engine/Source/Runtime/Engine/Private/KismetRenderingLibrary.cpp
@@ -0,0 +1,181 @@
+// Copyright 1998-2016 Epic Games, Inc. All Rights Reserved.
+
+#include "EnginePrivate.h"
+#include "Kismet/KismetRenderingLibrary.h"
+#include "MessageLog.h"
+#include "UObjectToken.h"
+#include "Runtime/Engine/Classes/Engine/TextureRenderTarget2D.h"
+
+//////////////////////////////////////////////////////////////////////////
+// UKismetRenderingLibrary
+
+#define LOCTEXT_NAMESPACE "KismetRenderingLibrary"
+
+UKismetRenderingLibrary::UKismetRenderingLibrary(const FObjectInitializer& ObjectInitializer)
+	: Super(ObjectInitializer)
+{
+}
+
+UCanvas* GCanvasForRenderingToTarget = NULL;
+
+void UKismetRenderingLibrary::DrawMaterialToRenderTarget(UObject* WorldContextObject, UTextureRenderTarget2D* TextureRenderTarget, UMaterialInterface* Material)
+{
+	check(WorldContextObject);
+	UWorld* World = GEngine->GetWorldFromContextObject(WorldContextObject);
+
+	if (TextureRenderTarget 
+		&& TextureRenderTarget->Resource 
+		&& World 
+		&& Material)
+	{
+		if (!GCanvasForRenderingToTarget)
+		{
+			GCanvasForRenderingToTarget = NewObject<UCanvas>(GetTransientPackage(), NAME_None);
+			GCanvasForRenderingToTarget->AddToRoot();
+		}
+
+		UCanvas* Canvas = GCanvasForRenderingToTarget;
+
+		Canvas->Init(TextureRenderTarget->SizeX, TextureRenderTarget->SizeY, nullptr);
+		Canvas->Update();
+
+		FCanvas RenderCanvas(
+			TextureRenderTarget->GameThread_GetRenderTargetResource(), 
+			nullptr, 
+			FApp::GetCurrentTime() - GStartTime, 
+			FApp::GetDeltaTime(), 
+			FApp::GetCurrentTime() - GStartTime, 
+			World->FeatureLevel);
+
+		Canvas->Canvas = &RenderCanvas;
+
+		TDrawEvent<FRHICommandList>* DrawMaterialToTargetEvent = new TDrawEvent<FRHICommandList>();
+
+		ENQUEUE_UNIQUE_RENDER_COMMAND_TWOPARAMETER(
+			BeginDrawEventCommand,
+			FName,RTName,TextureRenderTarget->GetFName(),
+			TDrawEvent<FRHICommandList>*,DrawMaterialToTargetEvent,DrawMaterialToTargetEvent,
+		{
+			BEGIN_DRAW_EVENTF(
+				RHICmdList, 
+				DrawCanvasToTarget, 
+				(*DrawMaterialToTargetEvent), 
+				*RTName.ToString());
+		});
+
+		Canvas->K2_DrawMaterial(Material, FVector2D(0, 0), FVector2D(TextureRenderTarget->SizeX, TextureRenderTarget->SizeY), FVector2D(0, 0));
+
+		RenderCanvas.Flush_GameThread();
+		Canvas->Canvas = NULL;
+
+		ENQUEUE_UNIQUE_RENDER_COMMAND_TWOPARAMETER(
+			CanvasRenderTargetResolveCommand, FTextureRenderTargetResource*, RenderTargetResource, TextureRenderTarget->GameThread_GetRenderTargetResource(),
+			TDrawEvent<FRHICommandList>*,DrawMaterialToTargetEvent,DrawMaterialToTargetEvent,
+			{
+				RHICmdList.CopyToResolveTarget(RenderTargetResource->GetRenderTargetTexture(), RenderTargetResource->TextureRHI, true, FResolveParams());
+				STOP_DRAW_EVENT((*DrawMaterialToTargetEvent));
+				delete DrawMaterialToTargetEvent;
+			}
+		);
+	}
+	else if (!Material)
+	{
+		FMessageLog("Blueprint").Warning(LOCTEXT("InvalidMaterial", "DrawMaterialToRenderTarget: Material must be non-null."));
+	}
+	else if (!TextureRenderTarget)
+	{
+		FMessageLog("Blueprint").Warning(LOCTEXT("InvalidTextureRenderTarget", "DrawMaterialToRenderTarget: TextureRenderTarget must be non-null."));
+	}
+}
+
+void UKismetRenderingLibrary::BeginDrawCanvasToRenderTarget(UObject* WorldContextObject, UTextureRenderTarget2D* TextureRenderTarget, UCanvas*& Canvas, FVector2D& Size, FDrawToRenderTargetContext& Context)
+{
+	check(WorldContextObject);
+	UWorld* World = GEngine->GetWorldFromContextObject(WorldContextObject);
+
+	Canvas = NULL;
+	Size = FVector2D(0, 0);
+	Context = FDrawToRenderTargetContext();
+
+	if (TextureRenderTarget 
+		&& TextureRenderTarget->Resource 
+		&& World)
+	{
+		Context.RenderTarget = TextureRenderTarget;
+
+		if (!GCanvasForRenderingToTarget)
+		{
+			GCanvasForRenderingToTarget = NewObject<UCanvas>(GetTransientPackage(), NAME_None);
+			GCanvasForRenderingToTarget->AddToRoot();
+		}
+
+		Canvas = GCanvasForRenderingToTarget;
+		Size = FVector2D(TextureRenderTarget->SizeX, TextureRenderTarget->SizeY);
+
+		Canvas->Init(TextureRenderTarget->SizeX, TextureRenderTarget->SizeY, nullptr);
+		Canvas->Update();
+
+		Canvas->Canvas = new FCanvas(
+			TextureRenderTarget->GameThread_GetRenderTargetResource(), 
+			nullptr, 
+			FApp::GetCurrentTime() - GStartTime, 
+			FApp::GetDeltaTime(), 
+			FApp::GetCurrentTime() - GStartTime, 
+			World->FeatureLevel, 
+			// Draw immediately so that interleaved SetVectorParameter (etc) function calls work as expected
+			FCanvas::CDM_ImmediateDrawing);
+
+		Context.DrawEvent = new TDrawEvent<FRHICommandList>();
+
+		ENQUEUE_UNIQUE_RENDER_COMMAND_TWOPARAMETER(
+			BeginDrawEventCommand,
+			FName,RTName,TextureRenderTarget->GetFName(),
+			TDrawEvent<FRHICommandList>*,DrawEvent,Context.DrawEvent,
+		{
+			BEGIN_DRAW_EVENTF(
+				RHICmdList, 
+				DrawCanvasToTarget, 
+				(*DrawEvent), 
+				*RTName.ToString());
+		});
+	}
+	else if (!TextureRenderTarget)
+	{
+		FMessageLog("Blueprint").Warning(LOCTEXT("InvalidTextureRenderTarget", "BeginDrawCanvasToRenderTarget: TextureRenderTarget must be non-null."));
+	}
+}
+
+void UKismetRenderingLibrary::EndDrawCanvasToRenderTarget(UObject* WorldContextObject, const FDrawToRenderTargetContext& Context)
+{
+	check(WorldContextObject);
+	UWorld* World = GEngine->GetWorldFromContextObject(WorldContextObject);
+
+	if (World)
+	{
+		if (GCanvasForRenderingToTarget && GCanvasForRenderingToTarget->Canvas)
+		{
+			GCanvasForRenderingToTarget->Canvas->Flush_GameThread();
+			delete GCanvasForRenderingToTarget->Canvas;
+			GCanvasForRenderingToTarget->Canvas = nullptr;
+		}
+		
+		if (Context.RenderTarget)
+		{
+			ENQUEUE_UNIQUE_RENDER_COMMAND_TWOPARAMETER(
+				CanvasRenderTargetResolveCommand, FTextureRenderTargetResource*, RenderTargetResource, Context.RenderTarget->GameThread_GetRenderTargetResource(),
+				TDrawEvent<FRHICommandList>*,DrawEvent,Context.DrawEvent,
+				{
+					RHICmdList.CopyToResolveTarget(RenderTargetResource->GetRenderTargetTexture(), RenderTargetResource->TextureRHI, true, FResolveParams());
+					STOP_DRAW_EVENT((*DrawEvent));
+					delete DrawEvent;
+				}
+			);
+		}
+		else if (!Context.RenderTarget)
+		{
+			FMessageLog("Blueprint").Warning(LOCTEXT("InvalidContext", "EndDrawCanvasToRenderTarget: Context must be valid."));
+		}
+	}
+}
+
+#undef LOCTEXT_NAMESPACE
\ No newline at end of file
diff --git a/Engine/Source/Runtime/Engine/Private/Level.cpp b/Engine/Source/Runtime/Engine/Private/Level.cpp
--- a/Engine/Source/Runtime/Engine/Private/Level.cpp
+++ b/Engine/Source/Runtime/Engine/Private/Level.cpp
@@ -18,6 +18,7 @@ Level.cpp: Level-related functions
 #include "BlueprintUtilities.h"
 #include "DynamicMeshBuilder.h"
 #include "Engine/LevelBounds.h"
+#include "RenderingObjectVersion.h"
 #if WITH_EDITOR
 #include "Editor/UnrealEd/Public/Kismet2/KismetEditorUtilities.h"
 #include "Editor/UnrealEd/Public/Kismet2/BlueprintEditorUtils.h"
@@ -281,57 +282,6 @@ void ULevel::AddReferencedObjects(UObject* InThis, FReferenceCollector& Collecto
 {
 	ULevel* This = CastChecked<ULevel>(InThis);
 
-	// List of textures to remove
-	TArray< UTexture2D*, TInlineAllocator<32> > TexturesToRemove;
-
-	// Let GC know that we're referencing some UTexture2D objects
-	for (auto& It : This->TextureToInstancesMap)
-	{
-		UTexture2D* Texture2D = It.Key;
-		Collector.AddReferencedObject(Texture2D, This);
-		if (!Texture2D)
-		{
-			// The texture has been probably marked as pending kill and has been removed by GC
-			// We need to remove it after we exit this loop.
-			TexturesToRemove.Add(It.Key);
-		}
-	}
-	for (UTexture2D* Texture : TexturesToRemove)
-	{
-		This->TextureToInstancesMap.Remove(Texture);
-	}
-	TexturesToRemove.Reset();
-
-	// Let GC know that we're referencing some UTexture2D objects
-	for( auto& It : This->DynamicTextureInstances )
-	{
-		UPrimitiveComponent* Primitive = It.Key.Get();
-		const TArray<FDynamicTextureInstance>& TextureInstances = It.Value;
-
-		for ( FDynamicTextureInstance& Instance : It.Value )
-		{
-			Collector.AddReferencedObject( Instance.Texture, This );
-		}
-	}
-
-	// Let GC know that we're referencing some UTexture2D objects
-	for (auto& It : This->ForceStreamTextures)
-	{
-		UTexture2D* Texture2D = It.Key;
-		Collector.AddReferencedObject(Texture2D, This);
-		if (!Texture2D)
-		{
-			// The texture has been probably marked as pending kill and has been removed by GC
-			// We need to remove it after we exit this loop.
-			TexturesToRemove.Add(It.Key);
-		}
-	}
-	for (UTexture2D* Texture : TexturesToRemove)
-	{
-		This->ForceStreamTextures.Remove(Texture);
-	}
-	TexturesToRemove.Reset();
-
 	// Let GC know that we're referencing some AActor objects
 	for (auto& Actor : This->Actors)
 	{
@@ -375,37 +325,18 @@ void ULevel::Serialize( FArchive& Ar )
 	if( !Ar.IsTransacting() )
 	{
 		Ar << LevelScriptActor;
+	}
 
-		if( ( Ar.IsLoading() && !FPlatformProperties::SupportsTextureStreaming() ) ||
-			( Ar.IsCooking() && !Ar.CookingTarget()->SupportsFeature(ETargetPlatformFeatures::TextureStreaming) ) )
-		{
-			// Strip for unsupported platforms
-			TMap< UTexture2D*, TArray< FStreamableTextureInstance > >		Dummy0;
-			TMap< UPrimitiveComponent*, TArray< FDynamicTextureInstance > >	Dummy1;
-			Ar << Dummy0;
-			Ar << Dummy1;
-		}
-		else
-		{
-			Ar << TextureToInstancesMap;
-			Ar << DynamicTextureInstances;
-		}
-
-		bool bIsCooked = Ar.IsCooking();
-		if (Ar.UE4Ver() >= VER_UE4_REBUILD_TEXTURE_STREAMING_DATA_ON_LOAD)
-		{
-			Ar << bIsCooked;
-		}
-		if (Ar.UE4Ver() < VER_UE4_REBUILD_TEXTURE_STREAMING_DATA_ON_LOAD)
-		{
-			bool bTextureStreamingBuiltDummy = bIsCooked;
-			Ar << bTextureStreamingBuiltDummy;
-		}
-		if (Ar.IsLoading())
-		{
-			// Always rebuild texture streaming data after loading
-			bTextureStreamingBuilt = false;
-		}
+	// Stop serializing deprecated classes with new versions
+	if ( Ar.IsLoading() && Ar.CustomVer(FRenderingObjectVersion::GUID) < FRenderingObjectVersion::RemovedTextureStreamingLevelData )
+	{
+		// Strip for unsupported platforms
+		TMap< UTexture2D*, TArray< FStreamableTextureInstance > >		Dummy0;
+		TMap< UPrimitiveComponent*, TArray< FDynamicTextureInstance > >	Dummy1;
+		bool Dummy2;
+		Ar << Dummy0;
+		Ar << Dummy1;
+		Ar << Dummy2;
 
 		//@todo legacy, useless
 		if (Ar.IsLoading())
@@ -426,16 +357,8 @@ void ULevel::Serialize( FArchive& Ar )
 			Ar << DummySetup;
 		}
 
-		if( ( Ar.IsLoading() && !FPlatformProperties::SupportsTextureStreaming() ) ||
-			( Ar.IsCooking() && !Ar.CookingTarget()->SupportsFeature(ETargetPlatformFeatures::TextureStreaming) ) )
-		{
-			TMap< UTexture2D*, bool > Dummy;
-			Ar << Dummy;
-		}
-		else
-		{
-			Ar << ForceStreamTextures;
-		}
+		TMap< UTexture2D*, bool > Dummy3;
+		Ar << Dummy3;
 	}
 
 	// Mark archive and package as containing a map if we're serializing to disk.
@@ -619,9 +542,6 @@ void ULevel::PostLoad()
 		SortActorList();
 	}
 
-	// Remove UTexture2D references that are NULL (missing texture).
-	ForceStreamTextures.Remove( NULL );
-
 	// Validate navigable geometry
 	if (Model == NULL || Model->NumUniqueVertices == 0)
 	{
@@ -1500,8 +1420,6 @@ void ULevel::BuildStreamingData(UWorld* World, ULevel* TargetLevel/*=NULL*/, UTe
 #if WITH_EDITORONLY_DATA
 	double StartTime = FPlatformTime::Seconds();
 
-	bool bUseDynamicStreaming = false;
-	GConfig->GetBool(TEXT("TextureStreaming"), TEXT("UseDynamicStreaming"), bUseDynamicStreaming, GEngineIni);
 
 	TArray<ULevel* > LevelsToCheck;
 	if ( TargetLevel )
@@ -1528,7 +1446,13 @@ void ULevel::BuildStreamingData(UWorld* World, ULevel* TargetLevel/*=NULL*/, UTe
 	for ( int32 LevelIndex=0; LevelIndex < LevelsToCheck.Num(); LevelIndex++ )
 	{
 		ULevel* Level = LevelsToCheck[LevelIndex];
-		Level->BuildStreamingData( UpdateSpecificTextureOnly );
+		if (!Level) continue;
+
+		if (Level->bIsVisible || Level->IsPersistentLevel())
+		{
+			IStreamingManager::Get().AddLevel(Level);
+		}
+		//@todo : handle UpdateSpecificTextureOnly
 	}
 
 	UE_LOG(LogLevel, Verbose, TEXT("ULevel::BuildStreamingData took %.3f seconds."), FPlatformTime::Seconds() - StartTime);
@@ -1537,220 +1461,6 @@ void ULevel::BuildStreamingData(UWorld* World, ULevel* TargetLevel/*=NULL*/, UTe
 #endif
 }
 
-void ULevel::BuildStreamingData(UTexture2D* UpdateSpecificTextureOnly/*=NULL*/)
-{
-	bool bUseDynamicStreaming = false;
-	GConfig->GetBool(TEXT("TextureStreaming"), TEXT("UseDynamicStreaming"), bUseDynamicStreaming, GEngineIni);
-
-	if ( UpdateSpecificTextureOnly == NULL )
-	{
-		// Reset the streaming manager, when building data for a whole Level
-		IStreamingManager::Get().RemoveLevel( this );
-		TextureToInstancesMap.Empty();
-		DynamicTextureInstances.Empty();
-		ForceStreamTextures.Empty();
-		bTextureStreamingBuilt = false;
-	}
-
-	TArray<UObject *> ObjectsInOuter;
-	GetObjectsWithOuter(this, ObjectsInOuter);
-
-	FStreamingTextureLevelContext LevelContext(this);
-
-	for( int32 Index = 0; Index < ObjectsInOuter.Num(); Index++ )
-	{
-		UPrimitiveComponent* Primitive = Cast<UPrimitiveComponent>(ObjectsInOuter[Index]);
-		if (Primitive)
-		{
-			const bool bIsClassDefaultObject = Primitive->IsTemplate(RF_ClassDefaultObject);
-			if ( !bIsClassDefaultObject && Primitive->IsRegistered() )
-			{
-				const AActor* const Owner = Primitive->GetOwner();
-				const bool bIsStatic = Owner == NULL || Primitive->Mobility == EComponentMobility::Static || Primitive->Mobility == EComponentMobility::Stationary;
-
-				TArray<FStreamingTexturePrimitiveInfo> PrimitiveStreamingTextures;
-
-				// Ask the primitive to enumerate the streaming textures it uses.
-				Primitive->GetStreamingTextureInfoWithNULLRemoval(LevelContext, PrimitiveStreamingTextures);
-
-				for(int32 TextureIndex = 0;TextureIndex < PrimitiveStreamingTextures.Num();TextureIndex++)
-				{
-					const FStreamingTexturePrimitiveInfo& PrimitiveStreamingTexture = PrimitiveStreamingTextures[TextureIndex];
-					UTexture2D* Texture2D = Cast<UTexture2D>(PrimitiveStreamingTexture.Texture);
-					bool bCanBeStreamedByDistance = !FMath::IsNearlyZero(PrimitiveStreamingTexture.TexelFactor) && !FMath::IsNearlyZero(PrimitiveStreamingTexture.Bounds.SphereRadius);
-
-					// Only handle 2D textures that match the target texture.
-					const bool bIsTargetTexture = (!UpdateSpecificTextureOnly || UpdateSpecificTextureOnly == Texture2D);
-					bool bShouldHandleTexture = (Texture2D && bIsTargetTexture);
-
-					// Check if this is a lightmap/shadowmap that shouldn't be streamed.
-					if ( bShouldHandleTexture )
-					{
-						UShadowMapTexture2D* ShadowMap2D	= Cast<UShadowMapTexture2D>(Texture2D);
-						ULightMapTexture2D* Lightmap2D		= Cast<ULightMapTexture2D>(Texture2D);
-						if ( (Lightmap2D && (Lightmap2D->LightmapFlags & LMF_Streamed) == 0) ||
-							(ShadowMap2D && (ShadowMap2D->ShadowmapFlags & SMF_Streamed) == 0) )
-						{
-							bShouldHandleTexture			= false;
-						}
-					}
-
-					// Check if this is a duplicate texture
-					if (bShouldHandleTexture)
-					{
-						for (int32 HandledTextureIndex = 0; HandledTextureIndex < TextureIndex; HandledTextureIndex++)
-						{
-							const FStreamingTexturePrimitiveInfo& HandledStreamingTexture = PrimitiveStreamingTextures[HandledTextureIndex];
-							if ( PrimitiveStreamingTexture.Texture == HandledStreamingTexture.Texture &&
-								FMath::IsNearlyEqual(PrimitiveStreamingTexture.TexelFactor, HandledStreamingTexture.TexelFactor) &&
-								PrimitiveStreamingTexture.Bounds.GetSphere().Equals( HandledStreamingTexture.Bounds.GetSphere() ) )
-							{
-								// It's a duplicate, don't handle this one.
-								bShouldHandleTexture = false;
-								break;
-							}
-						}
-					}
-
-					if(bShouldHandleTexture)
-					{
-						// Is the primitive set to force its textures to be resident?
-						if ( Primitive->bForceMipStreaming )
-						{
-							// Add them to the ForceStreamTextures set.
-							ForceStreamTextures.Add(Texture2D,true);
-						}
-						// Is this texture used by a static object?
-						else if ( bIsStatic && bCanBeStreamedByDistance )
-						{
-							// Texture instance information.
-							FStreamableTextureInstance TextureInstance;
-							TextureInstance.Bounds = PrimitiveStreamingTexture.Bounds;
-							TextureInstance.TexelFactor = PrimitiveStreamingTexture.TexelFactor;
-
-							// HLOD support.
-							TextureInstance.MinDistance = Primitive->MinDrawDistance;
-							const UPrimitiveComponent* LODParentPrimitive = Primitive->GetLODParentPrimitive();
-							if (LODParentPrimitive) // Max distance when HLOD becomes visible.
-							{
-								TextureInstance.MaxDistance = LODParentPrimitive->MinDrawDistance;
-								// Taken into account the streaming distance offsets.
-								TextureInstance.MaxDistance += (Primitive->Bounds.Origin - LODParentPrimitive->Bounds.Origin).Size();
-							}
-							else
-							{
-								TextureInstance.MaxDistance = FLT_MAX;
-							}
-
-							// See whether there already is an instance in the level.
-							TArray<FStreamableTextureInstance>* TextureInstances = TextureToInstancesMap.Find( Texture2D );
-							// We have existing instances.
-							if( TextureInstances )
-							{
-								// Add to the array.
-								TextureInstances->Add( TextureInstance );
-							}
-							// This is the first instance.
-							else
-							{
-								// Create array with current instance as the only entry.
-								TArray<FStreamableTextureInstance> NewTextureInstances;
-								NewTextureInstances.Add( TextureInstance );
-								// And set it.
-								TextureToInstancesMap.Add( Texture2D, NewTextureInstances );
-							}
-						}
-						// Is the texture used by a dynamic object that we can track at run-time.
-						else if ( bUseDynamicStreaming && Owner && bCanBeStreamedByDistance )
-						{
-							// Texture instance information.
-							FDynamicTextureInstance TextureInstance;
-							TextureInstance.Texture = Texture2D;
-							TextureInstance.Bounds = PrimitiveStreamingTexture.Bounds;
-							TextureInstance.TexelFactor	= PrimitiveStreamingTexture.TexelFactor;
-							TextureInstance.OriginalRadius = PrimitiveStreamingTexture.Bounds.SphereRadius;
-
-							// See whether there already is an instance in the level.
-							TArray<FDynamicTextureInstance>* TextureInstances = DynamicTextureInstances.Find( Primitive );
-							// We have existing instances.
-							if( TextureInstances )
-							{
-								// Add to the array.
-								TextureInstances->Add( TextureInstance );
-							}
-							// This is the first instance.
-							else
-							{
-								// Create array with current instance as the only entry.
-								TArray<FDynamicTextureInstance> NewTextureInstances;
-								NewTextureInstances.Add( TextureInstance );
-								// And set it.
-								DynamicTextureInstances.Add( Primitive, NewTextureInstances );
-							}
-						}
-					}
-				}
-			}
-		}
-	}
-
-	if ( UpdateSpecificTextureOnly == NULL )
-	{
-		// Normalize the texelfactor for lightmaps and shadowmaps
-		NormalizeLightmapTexelFactor();
-
-		bTextureStreamingBuilt = true;
-
-		// Update the streaming manager.
-		IStreamingManager::Get().AddPreparedLevel( this );
-	}
-}
-
-
-void ULevel::NormalizeLightmapTexelFactor()
-{
-	for ( TMap<UTexture2D*,TArray<FStreamableTextureInstance> >::TIterator It(TextureToInstancesMap); It; ++It )
-	{
-		UTexture2D* Texture2D = It.Key();
-		if ( Texture2D->LODGroup == TEXTUREGROUP_Lightmap || Texture2D->LODGroup == TEXTUREGROUP_Shadowmap)
-		{
-			TArray<FStreamableTextureInstance>& TextureInstances = It.Value();
-
-			// Clamp texelfactors to 20-80% range.
-			// This is to prevent very low-res or high-res charts to dominate otherwise decent streaming.
-			struct FCompareTexelFactor
-			{
-				FORCEINLINE bool operator()( const FStreamableTextureInstance& A, const FStreamableTextureInstance& B ) const
-				{
-					return A.TexelFactor < B.TexelFactor;
-				}
-			};
-			TextureInstances.Sort( FCompareTexelFactor() );
-
-			float MinTexelFactor = TextureInstances[ TextureInstances.Num() * 0.2f ].TexelFactor;
-			float MaxTexelFactor = TextureInstances[ TextureInstances.Num() * 0.8f ].TexelFactor;
-			for ( int32 InstanceIndex=0; InstanceIndex < TextureInstances.Num(); ++InstanceIndex )
-			{
-				FStreamableTextureInstance& Instance = TextureInstances[InstanceIndex];
-				Instance.TexelFactor = FMath::Clamp( Instance.TexelFactor, MinTexelFactor, MaxTexelFactor );
-			}
-		}
-	}
-}
-
-TArray<FStreamableTextureInstance>* ULevel::GetStreamableTextureInstances(UTexture2D*& TargetTexture)
-{
-	typedef TArray<FStreamableTextureInstance>	STIA_Type;
-	for (TMap<UTexture2D*,STIA_Type>::TIterator It(TextureToInstancesMap); It; ++It)
-	{
-		TArray<FStreamableTextureInstance>& TSIA = It.Value();
-		TargetTexture = It.Key();
-		return &TSIA;
-	}		
-
-	return NULL;
-}
-
 ABrush* ULevel::GetBrush() const
 {
 	return GetDefaultBrush();
@@ -2073,23 +1783,12 @@ void ULevel::ApplyWorldOffset(const FVector& InWorldOffset, bool bWorldShift)
 {
 	QUICK_SCOPE_CYCLE_COUNTER(STAT_ULevel_ApplyWorldOffset);
 
-	if (bTextureStreamingBuilt && !InWorldOffset.IsZero())
+	if (!InWorldOffset.IsZero())
 	{
-		QUICK_SCOPE_CYCLE_COUNTER(STAT_ULevel_ApplyWorldOffset_TextureStreaming);
-		// Update texture streaming data to account for the move
-		for (TMap< UTexture2D*, TArray<FStreamableTextureInstance> >::TIterator It(TextureToInstancesMap); It; ++It)
-		{
-			TArray<FStreamableTextureInstance>& TextureInfo = It.Value();
-			for (int32 i = 0; i < TextureInfo.Num(); i++)
-			{
-				TextureInfo[i].Bounds.Origin += InWorldOffset;
-			}
-		}
-
 		// Re-add level data to a manager, in case level is visible it this point
 		if (bIsVisible)
 		{
-			IStreamingManager::Get().AddPreparedLevel( this );
+			IStreamingManager::Get().AddLevel( this );
 		}
 	}
 
diff --git a/Engine/Source/Runtime/Engine/Private/LevelTick.cpp b/Engine/Source/Runtime/Engine/Private/LevelTick.cpp
--- a/Engine/Source/Runtime/Engine/Private/LevelTick.cpp
+++ b/Engine/Source/Runtime/Engine/Private/LevelTick.cpp
@@ -1046,6 +1046,35 @@ static FAutoConsoleVariableRef CVarTimeBetweenPurgingPendingKillObjects(
 
 #include "GameFramework/SpawnActorTimer.h"
 
+TDrawEvent<FRHICommandList>* BeginTickDrawEvent()
+{
+	TDrawEvent<FRHICommandList>* TickDrawEvent = new TDrawEvent<FRHICommandList>();
+
+	ENQUEUE_UNIQUE_RENDER_COMMAND_ONEPARAMETER(
+		BeginDrawEventCommand,
+		TDrawEvent<FRHICommandList>*,TickDrawEvent,TickDrawEvent,
+	{
+		BEGIN_DRAW_EVENTF(
+			RHICmdList, 
+			WorldTick, 
+			(*TickDrawEvent),
+			TEXT("WorldTick"));
+	});
+
+	return TickDrawEvent;
+}
+
+void EndTickDrawEvent(TDrawEvent<FRHICommandList>* TickDrawEvent)
+{
+	ENQUEUE_UNIQUE_RENDER_COMMAND_ONEPARAMETER(
+		EndDrawEventCommand,
+		TDrawEvent<FRHICommandList>*,TickDrawEvent,TickDrawEvent,
+	{
+		STOP_DRAW_EVENT((*TickDrawEvent));
+		delete TickDrawEvent;
+	});
+}
+
 /**
  * Update the level after a variable amount of time, DeltaSeconds, has passed.
  * All child actors are ticked after their owners have been ticked.
@@ -1057,6 +1086,8 @@ void UWorld::Tick( ELevelTick TickType, float DeltaSeconds )
 		return;
 	}
 
+	TDrawEvent<FRHICommandList>* TickDrawEvent = BeginTickDrawEvent();
+
 	FWorldDelegates::OnWorldTickStart.Broadcast(TickType, DeltaSeconds);
 
 	//Tick game and other thread trackers.
@@ -1481,6 +1512,8 @@ void UWorld::Tick( ELevelTick TickType, float DeltaSeconds )
 		}
 	}
 	);
+
+	EndTickDrawEvent(TickDrawEvent);
 }
 
 /**
diff --git a/Engine/Source/Runtime/Engine/Private/LevelUtils.cpp b/Engine/Source/Runtime/Engine/Private/LevelUtils.cpp
--- a/Engine/Source/Runtime/Engine/Private/LevelUtils.cpp
+++ b/Engine/Source/Runtime/Engine/Private/LevelUtils.cpp
@@ -394,15 +394,6 @@ void FLevelUtils::ApplyLevelTransform( ULevel* Level, const FTransform& LevelTra
 			Level->bTextureStreamingRotationChanged = true;
 		}
 
-		for (TMap< UTexture2D*, TArray<FStreamableTextureInstance> >::TIterator It(Level->TextureToInstancesMap); It; ++It)
-		{
-			TArray<FStreamableTextureInstance>& TextureInfo = It.Value();
-			for (int32 i = 0; i < TextureInfo.Num(); i++)
-			{
-				TextureInfo[i].Bounds.Origin = LevelTransform.TransformPosition(TextureInfo[i].Bounds.Origin);
-			}
-		}
-
 		// Iterate over all actors in the level and transform them
 		for( int32 ActorIndex=0; ActorIndex<Level->Actors.Num(); ActorIndex++ )
 		{
diff --git a/Engine/Source/Runtime/Engine/Private/Materials/HLSLMaterialTranslator.h b/Engine/Source/Runtime/Engine/Private/Materials/HLSLMaterialTranslator.h
--- a/Engine/Source/Runtime/Engine/Private/Materials/HLSLMaterialTranslator.h
+++ b/Engine/Source/Runtime/Engine/Private/Materials/HLSLMaterialTranslator.h
@@ -466,6 +466,11 @@ public:
 				Errorf(TEXT("Post process materials must use unlit."));
 			}
 
+			if (Material->AllowNegativeEmissiveColor() && MaterialShadingModel != MSM_Unlit)
+			{
+				Errorf(TEXT("Only unlit materials can output negative emissive color."));
+			}
+
 			if (MaterialCompilationOutput.bNeedsSceneTextures)
 			{
 				if (Domain != MD_DeferredDecal && Domain != MD_PostProcess)
@@ -3747,6 +3752,11 @@ protected:
 						CodeStr = TEXT("<A>");
 					}
 				}
+				else if (DestCoordBasis == MCB_MeshParticle)
+				{
+					CodeStr = TEXT("mul(<A>, <MATRIX>(Parameters.Particle.LocalToWorld))");
+				}
+
 				// else use MCB_TranslatedWorld as intermediary basis
 				IntermediaryBasis = MCB_TranslatedWorld;
 				break;
@@ -3771,6 +3781,18 @@ protected:
 				IntermediaryBasis = MCB_TranslatedWorld;
 				break;
 			}
+			case MCB_MeshParticle:
+			{
+				if (DestCoordBasis == MCB_World)
+				{
+					CodeStr = TEXT("mul(<MATRIX>(Parameters.Particle.LocalToWorld), <A>)");
+				}
+				else
+				{
+					return Errorf(TEXT("Can transform only to world space from particle space"));
+				}
+				break;
+			}
 			default:
 				check(0);
 				break;
diff --git a/Engine/Source/Runtime/Engine/Private/Materials/MaterialExpressions.cpp b/Engine/Source/Runtime/Engine/Private/Materials/MaterialExpressions.cpp
--- a/Engine/Source/Runtime/Engine/Private/Materials/MaterialExpressions.cpp
+++ b/Engine/Source/Runtime/Engine/Private/Materials/MaterialExpressions.cpp
@@ -5568,6 +5568,7 @@ static EMaterialCommonBasis GetMaterialCommonBasis(EMaterialVectorCoordTransform
 		MCB_World,						// TRANSFORMSOURCE_World
 		MCB_View,						// TRANSFORMSOURCE_View
 		MCB_Camera,						// TRANSFORMSOURCE_Camera
+		MCB_MeshParticle,
 	};
 	return ConversionTable[X];
 }
@@ -5580,6 +5581,7 @@ static EMaterialCommonBasis GetMaterialCommonBasis(EMaterialVectorCoordTransform
 		MCB_World,						// TRANSFORM_World
 		MCB_View,						// TRANSFORM_View
 		MCB_Camera,						// TRANSFORM_Camera
+		MCB_MeshParticle,
 	};
 	return ConversionTable[X];
 }
@@ -5675,6 +5677,7 @@ static EMaterialCommonBasis GetMaterialCommonBasis(EMaterialPositionTransformSou
 		MCB_TranslatedWorld,			// TRANSFORMPOSSOURCE_TranslatedWorld
 		MCB_View,						// TRANSFORMPOSSOURCE_View
 		MCB_Camera,						// TRANSFORMPOSSOURCE_Camera
+		MCB_MeshParticle,	
 	};
 	return ConversionTable[X];
 }
diff --git a/Engine/Source/Runtime/Engine/Private/Materials/MaterialInstance.cpp b/Engine/Source/Runtime/Engine/Private/Materials/MaterialInstance.cpp
--- a/Engine/Source/Runtime/Engine/Private/Materials/MaterialInstance.cpp
+++ b/Engine/Source/Runtime/Engine/Private/Materials/MaterialInstance.cpp
@@ -1140,20 +1140,19 @@ void UMaterialInstanceDynamic::CopyScalarAndVectorParameters(const UMaterialInte
 					const FMaterialUniformExpressionScalarParameter* ScalarExpression = static_cast<const FMaterialUniformExpressionScalarParameter*>(UniformExpression);
 
 					float Value;
-					if (ScalarExpression->GetGameThreadNumberValue(&SourceMaterialToCopyFrom, Value))
-					{
-						FName ParameterName = ScalarExpression->GetParameterName();
+					ScalarExpression->GetGameThreadNumberValue(&SourceMaterialToCopyFrom, Value);
 
-						FScalarParameterValue* ParameterValue = GameThread_FindParameterByName(ScalarParameterValues, ParameterName);
+					FName ParameterName = ScalarExpression->GetParameterName();
 
-						if (!ParameterValue)
-						{
-							ParameterValue = new(ScalarParameterValues)FScalarParameterValue;
-							ParameterValue->ParameterName = ParameterName;
-						}
+					FScalarParameterValue* ParameterValue = GameThread_FindParameterByName(ScalarParameterValues, ParameterName);
 
-						ParameterValue->ParameterValue = Value;
+					if (!ParameterValue)
+					{
+						ParameterValue = new(ScalarParameterValues)FScalarParameterValue;
+						ParameterValue->ParameterName = ParameterName;
 					}
+
+					ParameterValue->ParameterValue = Value;
 				}
 			}
 		}
@@ -1172,20 +1171,19 @@ void UMaterialInstanceDynamic::CopyScalarAndVectorParameters(const UMaterialInte
 					const FMaterialUniformExpressionVectorParameter* VectorExpression = static_cast<const FMaterialUniformExpressionVectorParameter*>(UniformExpression);
 
 					FLinearColor Value;
-					if (VectorExpression->GetGameThreadNumberValue(&SourceMaterialToCopyFrom, Value))
-					{
-						FName ParameterName = VectorExpression->GetParameterName();
+					VectorExpression->GetGameThreadNumberValue(&SourceMaterialToCopyFrom, Value);
 
-						FVectorParameterValue* ParameterValue = GameThread_FindParameterByName(VectorParameterValues, ParameterName);
+					FName ParameterName = VectorExpression->GetParameterName();
 
-						if (!ParameterValue)
-						{
-							ParameterValue = new(VectorParameterValues)FVectorParameterValue;
-							ParameterValue->ParameterName = ParameterName;
-						}
+					FVectorParameterValue* ParameterValue = GameThread_FindParameterByName(VectorParameterValues, ParameterName);
 
-						ParameterValue->ParameterValue = Value;
+					if (!ParameterValue)
+					{
+						ParameterValue = new(VectorParameterValues)FVectorParameterValue;
+						ParameterValue->ParameterName = ParameterName;
 					}
+
+					ParameterValue->ParameterValue = Value;
 				}
 			}
 		}
diff --git a/Engine/Source/Runtime/Engine/Private/Materials/MaterialShared.cpp b/Engine/Source/Runtime/Engine/Private/Materials/MaterialShared.cpp
--- a/Engine/Source/Runtime/Engine/Private/Materials/MaterialShared.cpp
+++ b/Engine/Source/Runtime/Engine/Private/Materials/MaterialShared.cpp
@@ -394,7 +394,12 @@ void FMaterial::GetShaderMapId(EShaderPlatform Platform, FMaterialShaderMapId& O
 EMaterialTessellationMode FMaterial::GetTessellationMode() const 
 { 
 	return MTM_NoTessellation; 
-};
+}
+
+ERefractionMode FMaterial::GetRefractionMode() const 
+{ 
+	return RM_IndexOfRefraction; 
+}
 
 void FMaterial::GetShaderMapIDsWithUnfinishedCompilation(TArray<int32>& ShaderMapIds)
 {
@@ -955,6 +960,11 @@ EBlendMode FMaterialResource::GetBlendMode() const
 	return MaterialInstance ? MaterialInstance->GetBlendMode() : Material->GetBlendMode();
 }
 
+ERefractionMode FMaterialResource::GetRefractionMode() const
+{
+	return Material->RefractionMode;
+}
+
 EMaterialShadingModel FMaterialResource::GetShadingModel() const 
 {
 	return MaterialInstance ? MaterialInstance->GetShadingModel() : Material->GetShadingModel();
@@ -980,6 +990,11 @@ bool FMaterialResource::IsDitherMasked() const
 	return Material->DitherOpacityMask;
 }
 
+bool FMaterialResource::AllowNegativeEmissiveColor() const 
+{
+	return Material->bAllowNegativeEmissiveColor;
+}
+
 bool FMaterialResource::IsDistorted() const { return Material->bUsesDistortion && IsTranslucentBlendMode(GetBlendMode()); }
 float FMaterialResource::GetTranslucencyDirectionalLightingIntensity() const { return Material->TranslucencyDirectionalLightingIntensity; }
 float FMaterialResource::GetTranslucentShadowDensityScale() const { return Material->TranslucentShadowDensityScale; }
@@ -1379,6 +1394,15 @@ void FMaterial::SetupMaterialEnvironment(
 		OutEnvironment.SetDefine(TEXT("MATERIALDECALRESPONSEMASK"), MaterialDecalResponseMask);
 	}
 
+	switch(GetRefractionMode())
+	{
+	case RM_IndexOfRefraction: OutEnvironment.SetDefine(TEXT("REFRACTION_USE_INDEX_OF_REFRACTION"),TEXT("1")); break;
+	case RM_PixelNormalOffset: OutEnvironment.SetDefine(TEXT("REFRACTION_USE_PIXEL_NORMAL_OFFSET"),TEXT("1")); break;
+	default: 
+		UE_LOG(LogMaterial, Warning, TEXT("Unknown material refraction mode: %u  Setting to RM_IndexOfRefraction"),(int32)GetRefractionMode());
+		OutEnvironment.SetDefine(TEXT("REFRACTION_USE_INDEX_OF_REFRACTION"),TEXT("1"));
+	}
+
 	OutEnvironment.SetDefine(TEXT("USE_DITHERED_LOD_TRANSITION_FROM_MATERIAL"), IsDitheredLODTransition());
 	OutEnvironment.SetDefine(TEXT("MATERIAL_TWOSIDED"), IsTwoSided());
 	OutEnvironment.SetDefine(TEXT("MATERIAL_TANGENTSPACENORMAL"), IsTangentSpaceNormal());
@@ -1393,6 +1417,7 @@ void FMaterial::SetupMaterialEnvironment(
 	OutEnvironment.SetDefine(TEXT("MATERIAL_SSR"), ShouldDoSSR());
 	OutEnvironment.SetDefine(TEXT("MATERIAL_BLOCK_GI"), ShouldBlockGI());
 	OutEnvironment.SetDefine(TEXT("MATERIAL_DITHER_OPACITY_MASK"), IsDitherMasked());
+	OutEnvironment.SetDefine(TEXT("MATERIAL_ALLOW_NEGATIVE_EMISSIVECOLOR"), AllowNegativeEmissiveColor());
 
 	{
 		auto DecalBlendMode = (EDecalBlendMode)GetDecalBlendMode();
diff --git a/Engine/Source/Runtime/Engine/Private/Materials/MaterialUniformExpressions.cpp b/Engine/Source/Runtime/Engine/Private/Materials/MaterialUniformExpressions.cpp
--- a/Engine/Source/Runtime/Engine/Private/Materials/MaterialUniformExpressions.cpp
+++ b/Engine/Source/Runtime/Engine/Private/Materials/MaterialUniformExpressions.cpp
@@ -513,7 +513,7 @@ bool FMaterialUniformExpressionTexture::IsIdentical(const FMaterialUniformExpres
 	return TextureIndex == OtherTextureExpression->TextureIndex;
 }
 
-bool FMaterialUniformExpressionVectorParameter::GetGameThreadNumberValue(const UMaterialInterface* SourceMaterialToCopyFrom, FLinearColor& OutValue) const
+void FMaterialUniformExpressionVectorParameter::GetGameThreadNumberValue(const UMaterialInterface* SourceMaterialToCopyFrom, FLinearColor& OutValue) const
 {
 	check(IsInGameThread());
 	checkSlow(SourceMaterialToCopyFrom);
@@ -530,7 +530,7 @@ bool FMaterialUniformExpressionVectorParameter::GetGameThreadNumberValue(const U
 			if(ParameterValue)
 			{
 				OutValue = ParameterValue->ParameterValue;
-				return true;
+				break;
 			}
 
 			// go up the hierarchy
@@ -541,12 +541,12 @@ bool FMaterialUniformExpressionVectorParameter::GetGameThreadNumberValue(const U
 			// we reached the base material
 			// get the copy form the base material
 			GetDefaultValue(OutValue);
-			return false;
+			break;
 		}
 	}
 }
 
-bool FMaterialUniformExpressionScalarParameter::GetGameThreadNumberValue(const UMaterialInterface* SourceMaterialToCopyFrom, float& OutValue) const
+void FMaterialUniformExpressionScalarParameter::GetGameThreadNumberValue(const UMaterialInterface* SourceMaterialToCopyFrom, float& OutValue) const
 {
 	check(IsInGameThread());
 	checkSlow(SourceMaterialToCopyFrom);
@@ -563,7 +563,7 @@ bool FMaterialUniformExpressionScalarParameter::GetGameThreadNumberValue(const U
 			if(ParameterValue)
 			{
 				OutValue = ParameterValue->ParameterValue;
-				return true;
+				break;
 			}
 
 			// go up the hierarchy
@@ -574,7 +574,7 @@ bool FMaterialUniformExpressionScalarParameter::GetGameThreadNumberValue(const U
 			// we reached the base material
 			// get the copy form the base material
 			GetDefaultValue(OutValue);
-			return false;
+			break;
 		}
 	}
 }
diff --git a/Engine/Source/Runtime/Engine/Private/Materials/MaterialUniformExpressions.h b/Engine/Source/Runtime/Engine/Private/Materials/MaterialUniformExpressions.h
--- a/Engine/Source/Runtime/Engine/Private/Materials/MaterialUniformExpressions.h
+++ b/Engine/Source/Runtime/Engine/Private/Materials/MaterialUniformExpressions.h
@@ -147,7 +147,7 @@ public:
 	}
 
 	// faster than GetNumberValue(), good for run-time use
-	bool GetGameThreadNumberValue(const UMaterialInterface* SourceMaterialToCopyFrom, FLinearColor& OutValue) const;
+	void GetGameThreadNumberValue(const UMaterialInterface* SourceMaterialToCopyFrom, FLinearColor& OutValue) const;
 
 	virtual bool IsConstant() const
 	{
@@ -225,7 +225,7 @@ public:
 	}
 	
 	// faster than GetNumberValue(), good for run-time use
-	bool GetGameThreadNumberValue(const UMaterialInterface* SourceMaterialToCopyFrom, float& OutValue) const;
+	void GetGameThreadNumberValue(const UMaterialInterface* SourceMaterialToCopyFrom, float& OutValue) const;
 
 	virtual bool IsConstant() const
 	{
diff --git a/Engine/Source/Runtime/Engine/Private/Materials/ParameterCollection.cpp b/Engine/Source/Runtime/Engine/Private/Materials/ParameterCollection.cpp
--- a/Engine/Source/Runtime/Engine/Private/Materials/ParameterCollection.cpp
+++ b/Engine/Source/Runtime/Engine/Private/Materials/ParameterCollection.cpp
@@ -6,9 +6,23 @@
 #include "Materials/MaterialParameterCollection.h"
 #include "Materials/MaterialExpressionCollectionParameter.h"
 
+TMap<FGuid, FMaterialParameterCollectionInstanceResource*> GDefaultMaterialParameterCollectionInstances;
+
 UMaterialParameterCollection::UMaterialParameterCollection(const FObjectInitializer& ObjectInitializer)
 	: Super(ObjectInitializer)
-{}
+{
+	DefaultResource = NULL;
+}
+
+void UMaterialParameterCollection::PostInitProperties()
+{
+	Super::PostInitProperties();
+
+	if (!HasAnyFlags(RF_ClassDefaultObject))
+	{
+		DefaultResource = new FMaterialParameterCollectionInstanceResource();
+	}
+}
 
 void UMaterialParameterCollection::PostLoad()
 {
@@ -27,6 +41,26 @@ void UMaterialParameterCollection::PostLoad()
 		UWorld* CurrentWorld = *It;
 		CurrentWorld->AddParameterCollectionInstance(this, true);
 	}
+
+	UpdateDefaultResource();
+}
+
+void UMaterialParameterCollection::BeginDestroy()
+{
+	if (DefaultResource)
+	{
+		ENQUEUE_UNIQUE_RENDER_COMMAND_ONEPARAMETER(
+			FRemoveDefaultResourceCommand,
+			FGuid,Id,StateId,
+		{
+			GDefaultMaterialParameterCollectionInstances.Remove(Id);
+		});
+
+		DefaultResource->GameThread_Destroy();
+		DefaultResource = NULL;
+	}
+
+	Super::BeginDestroy();
 }
 
 #if WITH_EDITOR
@@ -73,7 +107,7 @@ FName CreateUniqueName(TArray<ParameterType>& Parameters, int32 RenameParameterI
 }
 
 template<typename ParameterType>
-void SanatizeParameters(TArray<ParameterType>& Parameters)
+void SanitizeParameters(TArray<ParameterType>& Parameters)
 {
 	for (int32 i = 0; i < Parameters.Num() - 1; ++i)
 	{
@@ -105,8 +139,8 @@ void UMaterialParameterCollection::PreEditChange(class FEditPropertyChain& Prope
 
 void UMaterialParameterCollection::PostEditChangeProperty(FPropertyChangedEvent& PropertyChangedEvent)
 {
-	SanatizeParameters(ScalarParameters);
-	SanatizeParameters(VectorParameters);
+	SanitizeParameters(ScalarParameters);
+	SanitizeParameters(VectorParameters);
 
 	// If the array counts have changed, an element has been added or removed, and we need to update the uniform buffer layout,
 	// Which also requires recompiling any referencing materials
@@ -220,6 +254,8 @@ void UMaterialParameterCollection::PostEditChangeProperty(FPropertyChangedEvent&
 		CurrentWorld->UpdateParameterCollectionInstances(true);
 	}
 
+	UpdateDefaultResource();
+
 	PreviousScalarParameters.Empty();
 	PreviousVectorParameters.Empty();
 
@@ -386,6 +422,50 @@ void UMaterialParameterCollection::CreateBufferStruct()
 		);
 }
 
+void UMaterialParameterCollection::GetDefaultParameterData(TArray<FVector4>& ParameterData) const
+{
+	// The memory layout created here must match the index assignment in UMaterialParameterCollection::GetParameterIndex
+
+	ParameterData.Empty(FMath::DivideAndRoundUp(ScalarParameters.Num(), 4) + VectorParameters.Num());
+
+	for (int32 ParameterIndex = 0; ParameterIndex < ScalarParameters.Num(); ParameterIndex++)
+	{
+		const FCollectionScalarParameter& Parameter = ScalarParameters[ParameterIndex];
+
+		// Add a new vector for each packed vector
+		if (ParameterIndex % 4 == 0)
+		{
+			ParameterData.Add(FVector4(0, 0, 0, 0));
+		}
+
+		FVector4& CurrentVector = ParameterData.Last();
+		// Pack into the appropriate component of this packed vector
+		CurrentVector[ParameterIndex % 4] = Parameter.DefaultValue;
+	}
+
+	for (int32 ParameterIndex = 0; ParameterIndex < VectorParameters.Num(); ParameterIndex++)
+	{
+		const FCollectionVectorParameter& Parameter = VectorParameters[ParameterIndex];
+		ParameterData.Add(Parameter.DefaultValue);
+	}
+}
+
+void UMaterialParameterCollection::UpdateDefaultResource()
+{
+	// Propagate the new values to the rendering thread
+	TArray<FVector4> ParameterData;
+	GetDefaultParameterData(ParameterData);
+	DefaultResource->GameThread_UpdateContents(StateId, ParameterData);
+
+	ENQUEUE_UNIQUE_RENDER_COMMAND_TWOPARAMETER(
+		FUpdateDefaultResourceCommand,
+		FGuid,Id,StateId,
+		FMaterialParameterCollectionInstanceResource*,Resource,DefaultResource,
+	{
+		GDefaultMaterialParameterCollectionInstances.Add(Id, Resource);
+	});
+}
+
 UMaterialParameterCollectionInstance::UMaterialParameterCollectionInstance(const FObjectInitializer& ObjectInitializer)
 	: Super(ObjectInitializer)
 {
diff --git a/Engine/Source/Runtime/Engine/Private/NiagaraEffectRenderer.cpp b/Engine/Source/Runtime/Engine/Private/NiagaraEffectRenderer.cpp
--- a/Engine/Source/Runtime/Engine/Private/NiagaraEffectRenderer.cpp
+++ b/Engine/Source/Runtime/Engine/Private/NiagaraEffectRenderer.cpp
@@ -656,45 +656,55 @@ void NiagaraEffectRendererMeshes::SetupVertexFactory(FMeshParticleVertexFactory
 		VET_Float3
 		);
 
+	uint32 TangentXOffset = 0;
+	uint32 TangetnZOffset = 0;
+	uint32 UVsBaseOffset = 0;
+
+	SELECT_STATIC_MESH_VERTEX_TYPE(
+		LODResources.VertexBuffer.GetUseHighPrecisionTangentBasis(),
+		LODResources.VertexBuffer.GetUseFullPrecisionUVs(),
+		LODResources.VertexBuffer.GetNumTexCoords(),
+		{
+			TangentXOffset = STRUCT_OFFSET(VertexType, TangentX);
+			TangetnZOffset = STRUCT_OFFSET(VertexType, TangentZ);
+			UVsBaseOffset = STRUCT_OFFSET(VertexType, UVs);
+		});
+
 	Data.TangentBasisComponents[0] = FVertexStreamComponent(
 		&LODResources.VertexBuffer,
-		STRUCT_OFFSET(FStaticMeshFullVertex, RawTangentX),
+		TangentXOffset,
 		LODResources.VertexBuffer.GetStride(),
-		LODResources.VertexBuffer.GetUseHighPrecisionTangentBasis() ? VET_URGB10A2N : VET_PackedNormal
+		LODResources.VertexBuffer.GetUseHighPrecisionTangentBasis() ?
+			TStaticMeshVertexTangentTypeSelector<EStaticMeshVertexTangentBasisType::HighPrecision>::VertexElementType : 
+			TStaticMeshVertexTangentTypeSelector<EStaticMeshVertexTangentBasisType::Default>::VertexElementType
 		);
 
 	Data.TangentBasisComponents[1] = FVertexStreamComponent(
 		&LODResources.VertexBuffer,
-		STRUCT_OFFSET(FStaticMeshFullVertex, TangentZ),
+		TangetnZOffset,
 		LODResources.VertexBuffer.GetStride(),
-		VET_UShort2N
+		LODResources.VertexBuffer.GetUseHighPrecisionTangentBasis() ?
+			TStaticMeshVertexTangentTypeSelector<EStaticMeshVertexTangentBasisType::HighPrecision>::VertexElementType : 
+			TStaticMeshVertexTangentTypeSelector<EStaticMeshVertexTangentBasisType::Default>::VertexElementType
 		);
 
 	Data.TextureCoordinates.Empty();
-	if (!LODResources.VertexBuffer.GetUseFullPrecisionUVs())
-	{
-		uint32 NumTexCoords = FMath::Min<uint32>(LODResources.VertexBuffer.GetNumTexCoords(), MAX_TEXCOORDS);
-		for (uint32 UVIndex = 0; UVIndex < NumTexCoords; UVIndex++)
-		{
-			Data.TextureCoordinates.Add(FVertexStreamComponent(
-				&LODResources.VertexBuffer,
-				STRUCT_OFFSET(TStaticMeshFullVertexFloat16UVs<MAX_TEXCOORDS>, UVs) + sizeof(FVector2DHalf)* UVIndex,
-				LODResources.VertexBuffer.GetStride(),
-				VET_Half2
-				));
-		}
-	}
-	else
+
+	uint32 UVSizeInBytes = LODResources.VertexBuffer.GetUseFullPrecisionUVs() ?
+		sizeof(TStaticMeshVertexUVsTypeSelector<EStaticMeshVertexUVType::HighPrecision>::UVsTypeT) : sizeof(TStaticMeshVertexUVsTypeSelector<EStaticMeshVertexUVType::Default>::UVsTypeT);
+
+	EVertexElementType UVVertexElementType = LODResources.VertexBuffer.GetUseFullPrecisionUVs() ?
+		VET_Float2 : VET_Half2;
+
+	uint32 NumTexCoords = FMath::Min<uint32>(LODResources.VertexBuffer.GetNumTexCoords(), MAX_TEXCOORDS);
+	for (uint32 UVIndex = 0; UVIndex < NumTexCoords; UVIndex++)
 	{
-		for (uint32 UVIndex = 0; UVIndex < LODResources.VertexBuffer.GetNumTexCoords(); UVIndex++)
-		{
-			Data.TextureCoordinates.Add(FVertexStreamComponent(
-				&LODResources.VertexBuffer,
-				STRUCT_OFFSET(TStaticMeshFullVertexFloat32UVs<MAX_TEXCOORDS>, UVs) + sizeof(FVector2D)* UVIndex,
-				LODResources.VertexBuffer.GetStride(),
-				VET_Float2
-				));
-		}
+		Data.TextureCoordinates.Add(FVertexStreamComponent(
+			&LODResources.VertexBuffer,
+			UVsBaseOffset + UVSizeInBytes * UVIndex,
+			LODResources.VertexBuffer.GetStride(),
+			UVVertexElementType
+			));
 	}
 
 	if (LODResources.ColorVertexBuffer.GetNumVertices() > 0)
diff --git a/Engine/Source/Runtime/Engine/Private/Particles/ParticleBeam2EmitterInstance.cpp b/Engine/Source/Runtime/Engine/Private/Particles/ParticleBeam2EmitterInstance.cpp
--- a/Engine/Source/Runtime/Engine/Private/Particles/ParticleBeam2EmitterInstance.cpp
+++ b/Engine/Source/Runtime/Engine/Private/Particles/ParticleBeam2EmitterInstance.cpp
@@ -225,6 +225,144 @@ void FParticleBeam2EmitterInstance::SetBeamTargetStrength(float NewTargetStrengt
 	UserSetTargetStrengthArray[TargetIndex] = NewTargetStrength;
 }
 
+/**
+*	Get the end point position
+*
+*	@param	OutEndPoint	Current end point position
+*
+*	@return	true		End point is set - OutEndPoint is valid
+*			false		End point is not set - OutEndPoint is invalid
+*/
+bool FParticleBeam2EmitterInstance::GetBeamEndPoint(FVector& OutEndPoint) const
+{
+	if (UserSetTargetArray.Num() < 1)
+	{
+		return false;
+	}
+	OutEndPoint = UserSetTargetArray[0];
+	return true;
+}
+
+/**
+*	Get the source point at the given index
+*
+*	@param	SourceIndex			The index of the source point
+*	@param	OutSourcePoint		Value of of source point
+*
+*	@return	true		SourceIndex is valid - OutSourcePoint is valid
+*			false		SourceIndex is invalid - OutSourcePoint is invalid
+*/
+bool FParticleBeam2EmitterInstance::GetBeamSourcePoint(int32 SourceIndex, FVector& OutSourcePoint) const
+{  
+	if (SourceIndex < 0 || (UserSetSourceArray.Num() < (SourceIndex + 1)))
+	{
+		return false;
+	}
+
+	OutSourcePoint = UserSetSourceArray[SourceIndex];
+	return true;
+}
+
+/**
+*	Get the source tangent at the given index
+*
+*	@param	SourceIndex			The index of the source point
+*	@param	OutSourceTangent	Value of of source tangent
+*
+*	@return	true		SourceIndex is valid - OutSourceTangent is valid
+*			false		SourceIndex is invalid - OutSourceTangent is invalid
+*/
+bool FParticleBeam2EmitterInstance::GetBeamSourceTangent(int32 SourceIndex, FVector& OutSourcePoint) const
+{
+	if (SourceIndex < 0 || (UserSetSourceTangentArray.Num() < (SourceIndex + 1)))
+	{
+		return false;
+	}
+
+	OutSourcePoint = UserSetSourceTangentArray[SourceIndex];
+	return true;
+}
+
+/**
+*	Get the source strength at the given index
+*
+*	@param	SourceIndex			The index of the source point
+*	@param	OutSourceStrength	Value of of source strength
+*
+*	@return	true		SourceIndex is valid - OutSourceStrength is valid
+*			false		SourceIndex is invalid - OutSourceStrength is invalid
+*/
+bool FParticleBeam2EmitterInstance::GetBeamSourceStrength(int32 SourceIndex, float& OutSourceStrength) const
+{
+	if (SourceIndex < 0 || (UserSetSourceStrengthArray.Num() < (SourceIndex + 1)))
+	{
+		return false;
+	}
+
+	OutSourceStrength = UserSetSourceStrengthArray[SourceIndex];
+	return true;
+}
+
+/**
+*	Get the target point at the given index
+*
+*	@param	TargetIndex			The index of the target point
+*	@param	OutTargetPoint		Value of of target point
+*
+*	@return	true		TargetIndex is valid - OutTargetPoint is valid
+*			false		TargetIndex is invalid - OutTargetPoint is invalid
+*/
+bool FParticleBeam2EmitterInstance::GetBeamTargetPoint(int32 TargetIndex, FVector& OutTargetPoint) const
+{
+	if (TargetIndex < 0 || (UserSetTargetArray.Num() < (TargetIndex + 1)))
+	{
+		return false;
+	}
+
+	OutTargetPoint = UserSetTargetArray[TargetIndex];
+	return true;
+}
+
+/**
+*	Get the target tangent at the given index
+*
+*	@param	TargetIndex			The index of the target point
+*	@param	OutTangentPoint		Value of of target tangent
+*
+*	@return	true		TargetIndex is valid - OutTangentPoint is valid
+*			false		TargetIndex is invalid - OutTangentPoint is invalid
+*/
+bool FParticleBeam2EmitterInstance::GetBeamTargetTangent(int32 TargetIndex, FVector& OutTangentPoint) const
+{
+	if (TargetIndex < 0 || (UserSetTargetTangentArray.Num() < (TargetIndex + 1)))
+	{
+		return false;
+	}
+
+	OutTangentPoint = UserSetTargetTangentArray[TargetIndex];
+	return true;
+}
+
+/**
+*	Get the target strength at the given index
+*
+*	@param	TargetIndex			The index of the target point
+*	@param	OutTargetStrength	Value of of target strength
+*
+*	@return	true		TargetIndex is valid - OutTargetStrength is valid
+*			false		TargetIndex is invalid - OutTargetStrength is invalid
+*/
+bool FParticleBeam2EmitterInstance::GetBeamTargetStrength(int32 TargetIndex, float& OutTargetStrength) const
+{
+	if (TargetIndex < 0 || (UserSetTargetStrengthArray.Num() < (TargetIndex + 1)))
+	{
+		return false;
+	}
+
+	OutTargetStrength = UserSetTargetStrengthArray[TargetIndex];
+	return true;
+}
+
 void FParticleBeam2EmitterInstance::ApplyWorldOffset(FVector InOffset, bool bWorldShift)
 {
 	FParticleEmitterInstance::ApplyWorldOffset(InOffset, bWorldShift);
diff --git a/Engine/Source/Runtime/Engine/Private/Particles/ParticleBeamModules.cpp b/Engine/Source/Runtime/Engine/Private/Particles/ParticleBeamModules.cpp
--- a/Engine/Source/Runtime/Engine/Private/Particles/ParticleBeamModules.cpp
+++ b/Engine/Source/Runtime/Engine/Private/Particles/ParticleBeamModules.cpp
@@ -1883,10 +1883,8 @@ bool UParticleModuleBeamSource::ResolveSourceData(FParticleBeam2EmitterInstance*
 					if (BeamInst->SourceEmitter)
 					{
 						FVector CalcSourcePosition;
-						FVector CalcSourceTangent;
 
-						static int32 SourceIndex = -1;
-						static int32 LastSelectedParticleIndex = -1;
+						int32 SourceIndex = LastSelectedParticleIndex;
 
 						if (BeamInst->SourceEmitter && BeamInst->SourceEmitter->ParticleIndices)
 						{
@@ -1905,7 +1903,7 @@ bool UParticleModuleBeamSource::ResolveSourceData(FParticleBeam2EmitterInstance*
 								int32 Index = 0;
 								// TODO: add selection method and random selection
 								/*
-								switch (SourceModule->SelectionMethod)
+								switch (SelectionMethod)
 								{
 								case EPSSM_Random:
 								{
@@ -1955,14 +1953,12 @@ bool UParticleModuleBeamSource::ResolveSourceData(FParticleBeam2EmitterInstance*
 								else
 								{
 									CalcSourcePosition = SourceParticle->Location + WorldOrigin;
-									CalcSourceTangent = SourceParticle->Location - SourceParticle->OldLocation;
 								}
 							}
 							else
 							{
 								// Fall back to the emitter location??
 								CalcSourcePosition = Comp->GetComponentLocation();
-								CalcSourceTangent = Comp->PartSysVelocity;
 								//@todo. How to handle this... can potentially cause a jump from the emitter to the
 								// particle...
 								SourceIndex = -1;//No valid particle source;
@@ -1971,7 +1967,6 @@ bool UParticleModuleBeamSource::ResolveSourceData(FParticleBeam2EmitterInstance*
 
 						//@todo. fill this in correctly...
 						BeamData->SourcePoint = CalcSourcePosition; 
-						BeamData->SourceTangent = CalcSourceTangent;
 						bSetSource = true;
 					}
 				}
@@ -2354,10 +2349,8 @@ bool UParticleModuleBeamTarget::ResolveTargetData(FParticleBeam2EmitterInstance*
 					if (BeamInst->TargetEmitter)
 					{
 						FVector LocalTargetPosition;
-						FVector LocalTargetTangent;
 
-						static int32 TargetIndex = -1;
-						static int32 LastSelectedParticleIndex = -1;
+						int32 TargetIndex = LastSelectedParticleIndex;
 
 						if (BeamInst->TargetEmitter && BeamInst->TargetEmitter->ParticleIndices)
 						{
@@ -2426,14 +2419,12 @@ bool UParticleModuleBeamTarget::ResolveTargetData(FParticleBeam2EmitterInstance*
 								else
 								{
 									LocalTargetPosition = TargetParticle->Location + WorldOrigin;
-									LocalTargetTangent = TargetParticle->Location - TargetParticle->OldLocation;
 								}
 							}
 							else
 							{
 								// Fall back to the emitter location??
 								LocalTargetPosition = BeamInst->TargetEmitter->Component->GetComponentLocation();
-								LocalTargetTangent = BeamInst->TargetEmitter->Component->PartSysVelocity;
 								//@todo. How to handle this... can potentially cause a jump from the emitter to the
 								// particle...
 								TargetIndex = -1;//No valid particle Target;
@@ -2441,7 +2432,6 @@ bool UParticleModuleBeamTarget::ResolveTargetData(FParticleBeam2EmitterInstance*
 
 							//@todo. fill this in correctly...
 							BeamData->TargetPoint = LocalTargetPosition; //BeamInst->TargetEmitter->Component->GetComponentLocation();
-							BeamData->TargetTangent = LocalTargetTangent;
 							bSetTarget = true;
 						}
 					}
