Copying //UE4/Dev-Rendering to //UE4/Dev-Main (Source: //UE4/Dev-Rendering @ 2981742)

diff --git a/Engine/Source/Runtime/Engine/Private/Particles/ParticleComponents.cpp b/Engine/Source/Runtime/Engine/Private/Particles/ParticleComponents.cpp
--- a/Engine/Source/Runtime/Engine/Private/Particles/ParticleComponents.cpp
+++ b/Engine/Source/Runtime/Engine/Private/Particles/ParticleComponents.cpp
@@ -3138,7 +3138,7 @@ UParticleSystemComponent::UParticleSystemComponent(const FObjectInitializer& Obj
 
 void UParticleSystemComponent::SetRequiredSignificance(EParticleSignificanceLevel NewRequiredSignificance)
 {
-	if (ensure(Template))
+	if (Template)
 	{
 		RequiredSignificance = NewRequiredSignificance;
 
@@ -3643,6 +3643,9 @@ FDynamicEmitterDataBase* UParticleSystemComponent::CreateDynamicDataFromReplay(
 	check( EmitterReplayData != NULL );
 
 	FScopeCycleCounterEmitter AdditionalScope(EmitterInstance);
+#if WITH_EDITOR
+	uint32 StartTime = FPlatformTime::Cycles();
+#endif
 
 	// Allocate the appropriate type of emitter data
 	FDynamicEmitterDataBase* EmitterData = NULL;
@@ -3754,6 +3757,12 @@ FDynamicEmitterDataBase* UParticleSystemComponent::CreateDynamicDataFromReplay(
 		EmitterData->StatID = EmitterInstance->SpriteTemplate->GetStatID();
 	}
 #endif
+
+#if WITH_EDITOR
+	uint32 EndTime = FPlatformTime::Cycles();
+	EmitterInstance->LastTickDurationMs += FPlatformTime::ToMilliseconds(EndTime - StartTime);
+#endif
+
 	return EmitterData;
 }
 
@@ -3897,6 +3906,9 @@ FParticleDynamicData* UParticleSystemComponent::CreateDynamicData(ERHIFeatureLev
 				if (EmitterInst)
 				{
 					FScopeCycleCounterEmitter AdditionalScope(EmitterInst);
+#if WITH_EDITOR
+					uint32 StartTime = FPlatformTime::Cycles();
+#endif
 
 					// Generate the dynamic data for this emitter
 					{
@@ -3945,6 +3957,10 @@ FParticleDynamicData* UParticleSystemComponent::CreateDynamicData(ERHIFeatureLev
 							NewEmitterReplayFrame->FrameState = NewEmitterReplayData;
 						}
 					}
+#if WITH_EDITOR
+					uint32 EndTime = FPlatformTime::Cycles();
+					EmitterInst->LastTickDurationMs += FPlatformTime::ToMilliseconds(EndTime - StartTime);
+#endif
 				}
 			}
 		}
@@ -4475,6 +4491,17 @@ void UParticleSystemComponent::TickComponent(float DeltaTime, enum ELevelTick Ti
 	PlayerLocations.Reset();
 	PlayerLODDistanceFactor.Reset();
 
+#if WITH_EDITOR
+	// clear tick timers
+	for (auto Instance : EmitterInstances)
+	{
+		if (Instance)
+		{
+			Instance->LastTickDurationMs = 0.0f;
+		}
+	}
+#endif
+
 	if (World->IsGameWorld())
 	{
 		for (FConstPlayerControllerIterator Iterator = World->GetPlayerControllerIterator(); Iterator; ++Iterator)
@@ -4605,6 +4632,9 @@ void UParticleSystemComponent::ComputeTickComponent_Concurrent()
 	{
 		FParticleEmitterInstance* Instance = EmitterInstances[EmitterIndex];
 		FScopeCycleCounterEmitter AdditionalScopeInner(Instance);
+#if WITH_EDITOR
+		uint32 StartTime = FPlatformTime::Cycles();
+#endif
 
 		if (EmitterIndex + 1 < EmitterInstances.Num())
 		{
@@ -4658,6 +4688,11 @@ void UParticleSystemComponent::ComputeTickComponent_Concurrent()
 				}
 				TotalActiveParticles += Instance->ActiveParticles;
 			}
+
+#if WITH_EDITOR
+			uint32 EndTime = FPlatformTime::Cycles();
+			Instance->LastTickDurationMs += FPlatformTime::ToMilliseconds(EndTime - StartTime);
+#endif
 		}
 	}
 	if (bAsyncWorkOutstanding)
@@ -5180,7 +5215,7 @@ void UParticleSystemComponent::ActivateSystem(bool bFlagAsJustAttached)
 	// System settings may have been lowered. Support late deactivation.
 	const bool bDetailModeAllowsRendering = DetailMode <= GetCurrentDetailMode();
 
-	if( GIsAllowingParticles && bDetailModeAllowsRendering )
+	if( GIsAllowingParticles && bDetailModeAllowsRendering && Template )
 	{
 		// Auto attach if requested
 		const bool bWasAutoAttached = bDidAutoAttach;
@@ -5326,6 +5361,12 @@ void UParticleSystemComponent::ActivateSystem(bool bFlagAsJustAttached)
 			WarmupTime = 0.0f;
 			bSkipUpdateDynamicDataDuringTick = bSaveSkipUpdate;
 		}
+
+		//We are definitely insignificant already so set insignificant before we ever begin ticking.
+		if (bIsManagingSignificance && Template->GetHighestSignificance() < RequiredSignificance && Template->InsignificanceDelay == 0.0f)
+		{
+			OnSignificanceChanged(false, true);
+		}
 	}
 
 	// Mark render state dirty to ensure the scene proxy is added and registered with the scene.
@@ -5335,12 +5376,6 @@ void UParticleSystemComponent::ActivateSystem(bool bFlagAsJustAttached)
 	{
 		LastRenderTime = GetWorld()->GetTimeSeconds();
 	}
-
-	//We are definitely insignificant already so set insignificant before we ever begin ticking.
-	if (bIsManagingSignificance && Template->GetHighestSignificance() < RequiredSignificance && Template->InsignificanceDelay == 0.0f)
-	{
-		OnSignificanceChanged(false, true);
-	}
 }
 
 void UParticleSystemComponent::Complete()
@@ -5999,6 +6034,103 @@ void UParticleSystemComponent::SetBeamTargetStrength(int32 EmitterIndex,float Ne
 	}
 }
 
+bool UParticleSystemComponent::GetBeamEndPoint(int32 EmitterIndex, FVector& OutSourcePoint) const
+{
+	if ((EmitterIndex >= 0) && (EmitterIndex < EmitterInstances.Num()))
+	{
+		FParticleEmitterInstance* EmitterInst = EmitterInstances[EmitterIndex];
+		if (EmitterInst)
+		{
+			return EmitterInst->GetBeamEndPoint(OutSourcePoint);
+		}
+	}
+
+	return false;
+}
+
+bool UParticleSystemComponent::GetBeamSourcePoint(int32 EmitterIndex, int32 SourceIndex, FVector& OutSourcePoint) const
+{
+	if ((EmitterIndex >= 0) && (EmitterIndex < EmitterInstances.Num()))
+	{
+		FParticleEmitterInstance* EmitterInst = EmitterInstances[EmitterIndex];
+		if (EmitterInst)
+		{
+			return EmitterInst->GetBeamSourcePoint(SourceIndex, OutSourcePoint);
+		}
+	}
+
+	return false;
+}
+
+bool UParticleSystemComponent::GetBeamSourceTangent(int32 EmitterIndex, int32 SourceIndex, FVector& OutSourcePoint) const
+{
+	if ((EmitterIndex >= 0) && (EmitterIndex < EmitterInstances.Num()))
+	{
+		FParticleEmitterInstance* EmitterInst = EmitterInstances[EmitterIndex];
+		if (EmitterInst)
+		{
+			return EmitterInst->GetBeamSourceTangent(SourceIndex, OutSourcePoint);
+		}
+	}
+
+	return false;
+}
+
+bool UParticleSystemComponent::GetBeamSourceStrength(int32 EmitterIndex, int32 SourceIndex, float& OutSourceStrength) const
+{
+	if ((EmitterIndex >= 0) && (EmitterIndex < EmitterInstances.Num()))
+	{
+		FParticleEmitterInstance* EmitterInst = EmitterInstances[EmitterIndex];
+		if (EmitterInst)
+		{
+			return EmitterInst->GetBeamSourceStrength(SourceIndex, OutSourceStrength);
+		}
+	}
+
+	return false;
+}
+bool UParticleSystemComponent::GetBeamTargetPoint(int32 EmitterIndex, int32 TargetIndex, FVector& OutTargetPoint) const
+{
+	if ((EmitterIndex >= 0) && (EmitterIndex < EmitterInstances.Num()))
+	{
+		FParticleEmitterInstance* EmitterInst = EmitterInstances[EmitterIndex];
+		if (EmitterInst)
+		{
+			return EmitterInst->GetBeamTargetPoint(TargetIndex, OutTargetPoint);
+		}
+	}
+
+	return false;
+}
+
+bool UParticleSystemComponent::GetBeamTargetTangent(int32 EmitterIndex, int32 TargetIndex, FVector& OutTangentPoint) const
+{
+	if ((EmitterIndex >= 0) && (EmitterIndex < EmitterInstances.Num()))
+	{
+		FParticleEmitterInstance* EmitterInst = EmitterInstances[EmitterIndex];
+		if (EmitterInst)
+		{
+			return EmitterInst->GetBeamTargetTangent(TargetIndex, OutTangentPoint);
+		}
+	}
+
+	return false;
+}
+
+bool UParticleSystemComponent::GetBeamTargetStrength(int32 EmitterIndex, int32 TargetIndex, float& OutTargetStrength) const
+{
+	if ((EmitterIndex >= 0) && (EmitterIndex < EmitterInstances.Num()))
+	{
+		FParticleEmitterInstance* EmitterInst = EmitterInstances[EmitterIndex];
+		if (EmitterInst)
+		{
+			return EmitterInst->GetBeamTargetStrength(TargetIndex, OutTargetStrength);
+		}
+	}
+
+	return false;
+}
+
 
 void UParticleSystemComponent::SetEmitterEnable(FName EmitterName, bool bNewEnableState)
 {
diff --git a/Engine/Source/Runtime/Engine/Private/Particles/ParticleModules_Location.cpp b/Engine/Source/Runtime/Engine/Private/Particles/ParticleModules_Location.cpp
--- a/Engine/Source/Runtime/Engine/Private/Particles/ParticleModules_Location.cpp
+++ b/Engine/Source/Runtime/Engine/Private/Particles/ParticleModules_Location.cpp
@@ -517,6 +517,22 @@ void UParticleModuleLocationEmitter::Spawn(FParticleEmitterInstance* Owner, int3
 					// If the ScreenAlignment of the source emitter is PSA_Velocity, 
 					// and that of the local is not, then the rotation will NOT be correct!
 					Particle.Rotation		+= pkParticle->Rotation * InheritSourceRotationScale;
+
+					// for mesh emitters only: get the mesh rotation payloads for both emitters and update the rotations accordingly; if the offset is 0, the module
+					// doesn't exist, so we can't transfer rotation; if the offsets exist, the payload should never be nullptr.
+					//
+					const int32 MeshRotationOffset = Owner->GetMeshRotationOffset();
+					const int32 SrcMeshRotationOffset = LocationEmitterInst->GetMeshRotationOffset();
+					if (MeshRotationOffset && SrcMeshRotationOffset)
+					{
+						FMeshRotationPayloadData* DestPayloadData = (FMeshRotationPayloadData*)((uint8*)&Particle + MeshRotationOffset);
+						FMeshRotationPayloadData* SrcPayloadData = (FMeshRotationPayloadData*)((uint8*)pkParticle + SrcMeshRotationOffset);
+
+						ensure(DestPayloadData != nullptr && SrcPayloadData != nullptr);
+
+						DestPayloadData->Rotation += SrcPayloadData->Rotation;
+						DestPayloadData->InitialOrientation += SrcPayloadData->InitialOrientation;
+					}
 				}
 			}
 		}
diff --git a/Engine/Source/Runtime/Engine/Private/Particles/ParticleSystemRender.cpp b/Engine/Source/Runtime/Engine/Private/Particles/ParticleSystemRender.cpp
--- a/Engine/Source/Runtime/Engine/Private/Particles/ParticleSystemRender.cpp
+++ b/Engine/Source/Runtime/Engine/Private/Particles/ParticleSystemRender.cpp
@@ -2332,46 +2332,56 @@ void FDynamicMeshEmitterData::SetupVertexFactory( FMeshParticleVertexFactory* In
 			VET_Float3
 			);
 
+		uint32 TangentXOffset = 0;
+		uint32 TangetnZOffset = 0;
+		uint32 UVsBaseOffset = 0;
+
+		SELECT_STATIC_MESH_VERTEX_TYPE(
+			LODResources.VertexBuffer.GetUseHighPrecisionTangentBasis(),
+			LODResources.VertexBuffer.GetUseFullPrecisionUVs(),
+			LODResources.VertexBuffer.GetNumTexCoords(),
+			{
+				TangentXOffset = STRUCT_OFFSET(VertexType, TangentX);
+				TangetnZOffset = STRUCT_OFFSET(VertexType, TangentZ);
+				UVsBaseOffset = STRUCT_OFFSET(VertexType, UVs);
+			});
+
 		Data.TangentBasisComponents[0] = FVertexStreamComponent(
 			&LODResources.VertexBuffer,
-			STRUCT_OFFSET(FStaticMeshFullVertex, RawTangentX),
+			TangentXOffset,
 			LODResources.VertexBuffer.GetStride(),
-			LODResources.VertexBuffer.GetUseHighPrecisionTangentBasis() ? VET_URGB10A2N : VET_PackedNormal
+			LODResources.VertexBuffer.GetUseHighPrecisionTangentBasis() ?
+				TStaticMeshVertexTangentTypeSelector<EStaticMeshVertexTangentBasisType::HighPrecision>::VertexElementType : 
+				TStaticMeshVertexTangentTypeSelector<EStaticMeshVertexTangentBasisType::Default>::VertexElementType
 			);
 
 		Data.TangentBasisComponents[1] = FVertexStreamComponent(
 			&LODResources.VertexBuffer,
-			STRUCT_OFFSET(FStaticMeshFullVertex, TangentZ),
+			TangetnZOffset,
 			LODResources.VertexBuffer.GetStride(),
-			VET_UShort2N
+			LODResources.VertexBuffer.GetUseHighPrecisionTangentBasis() ?
+				TStaticMeshVertexTangentTypeSelector<EStaticMeshVertexTangentBasisType::HighPrecision>::VertexElementType : 
+				TStaticMeshVertexTangentTypeSelector<EStaticMeshVertexTangentBasisType::Default>::VertexElementType
 			);
 
 		Data.TextureCoordinates.Empty();
-		if( !LODResources.VertexBuffer.GetUseFullPrecisionUVs() )
+
+		uint32 UVSizeInBytes = LODResources.VertexBuffer.GetUseFullPrecisionUVs() ?
+			sizeof(TStaticMeshVertexUVsTypeSelector<EStaticMeshVertexUVType::HighPrecision>::UVsTypeT) : sizeof(TStaticMeshVertexUVsTypeSelector<EStaticMeshVertexUVType::Default>::UVsTypeT);
+
+		EVertexElementType UVVertexElementType = LODResources.VertexBuffer.GetUseFullPrecisionUVs() ?
+			VET_Float2 : VET_Half2;
+
+		uint32 NumTexCoords = FMath::Min<uint32>(LODResources.VertexBuffer.GetNumTexCoords(), MAX_TEXCOORDS);
+		for (uint32 UVIndex = 0; UVIndex < NumTexCoords; UVIndex++)
 		{
-			uint32 NumTexCoords = FMath::Min<uint32>(LODResources.VertexBuffer.GetNumTexCoords(),MAX_TEXCOORDS);
-			for(uint32 UVIndex = 0;UVIndex < NumTexCoords;UVIndex++)
-			{
-				Data.TextureCoordinates.Add(FVertexStreamComponent(
-					&LODResources.VertexBuffer,
-					STRUCT_OFFSET(TStaticMeshFullVertexFloat16UVs<MAX_TEXCOORDS>,UVs) + sizeof(FVector2DHalf) * UVIndex,
-					LODResources.VertexBuffer.GetStride(),
-					VET_Half2
-					));
-			}
+			Data.TextureCoordinates.Add(FVertexStreamComponent(
+				&LODResources.VertexBuffer,
+				UVsBaseOffset + UVSizeInBytes * UVIndex,
+				LODResources.VertexBuffer.GetStride(),
+				UVVertexElementType
+				));
 		}
-		else
-		{
-			for(uint32 UVIndex = 0;UVIndex < LODResources.VertexBuffer.GetNumTexCoords();UVIndex++)
-			{
-				Data.TextureCoordinates.Add(FVertexStreamComponent(
-					&LODResources.VertexBuffer,
-					STRUCT_OFFSET(TStaticMeshFullVertexFloat32UVs<MAX_TEXCOORDS>,UVs) + sizeof(FVector2D) * UVIndex,
-					LODResources.VertexBuffer.GetStride(),
-					VET_Float2
-					));
-			}
-		}	
 
 		if(LODResources.ColorVertexBuffer.GetNumVertices() > 0)
 		{
@@ -6622,7 +6632,7 @@ FParticleSystemSceneProxy::FParticleSystemSceneProxy(const UParticleSystemCompon
 	, FirstFreeMeshBatch(0)
 	, bVertexFactoriesDirty(false)
 {
-	WireframeColor = FLinearColor(1.0f, 0.0f, 0.0f);
+	WireframeColor = FLinearColor(3.0f, 0.0f, 0.0f);
 	LevelColor = FLinearColor(1.0f, 1.0f, 0.0f);
 	PropertyColor = FLinearColor(1.0f, 1.0f, 1.0f);
 
diff --git a/Engine/Source/Runtime/Engine/Private/PreviewScene.cpp b/Engine/Source/Runtime/Engine/Private/PreviewScene.cpp
--- a/Engine/Source/Runtime/Engine/Private/PreviewScene.cpp
+++ b/Engine/Source/Runtime/Engine/Private/PreviewScene.cpp
@@ -34,8 +34,6 @@ FPreviewScene::FPreviewScene(FPreviewScene::ConstructionValues CVS)
 										.SetTransactional(CVS.bTransactional));
 	PreviewWorld->InitializeActorsForPlay(FURL());
 
-	GetScene()->UpdateDynamicSkyLight(FLinearColor::White * CVS.SkyBrightness, FLinearColor::Black);
-
 	DirectionalLight = NewObject<UDirectionalLightComponent>(GetTransientPackage());
 	DirectionalLight->Intensity = CVS.LightBrightness;
 	DirectionalLight->LightColor = FColor::White;
@@ -184,7 +182,6 @@ void FPreviewScene::SetLightColor(const FColor& LightColor)
 
 void FPreviewScene::SetSkyBrightness(float SkyBrightness)
 {
-	GetScene()->UpdateDynamicSkyLight(FLinearColor::White * SkyBrightness, FLinearColor::Black);
 }
 
 void FPreviewScene::LoadSettings(const TCHAR* Section)
diff --git a/Engine/Source/Runtime/Engine/Private/PrimitiveDrawingUtils.cpp b/Engine/Source/Runtime/Engine/Private/PrimitiveDrawingUtils.cpp
--- a/Engine/Source/Runtime/Engine/Private/PrimitiveDrawingUtils.cpp
+++ b/Engine/Source/Runtime/Engine/Private/PrimitiveDrawingUtils.cpp
@@ -27,13 +27,17 @@ void BeginMeshDrawEvent_Inner(FRHICommandList& RHICmdList, const FPrimitiveScene
 			*Mesh.MaterialRenderProxy->GetMaterial(PrimitiveSceneProxy ? PrimitiveSceneProxy->GetScene().GetFeatureLevel() : GMaxRHIFeatureLevel)->GetFriendlyName(),
 			PrimitiveSceneProxy->GetResourceName().IsValid() ? *PrimitiveSceneProxy->GetResourceName().ToString() : TEXT(""));
 
+		bool bIssueAdditionalDrawEvents = false;
+		if (bIssueAdditionalDrawEvents)
+		{
 			// Show Actor, level and resource name inside the material name
 			// These are separate draw events since some platforms have a limit on draw event length
 			// Note: empty leaf events are culled from profilegpu by default so these won't show up
 			{
-				//SCOPED_CONDITIONAL_DRAW_EVENTF(RHICmdList, LevelEvent, PrimitiveSceneProxy->GetLevelName() != NAME_None, PrimitiveSceneProxy->GetLevelName().IsValid() ? *PrimitiveSceneProxy->GetLevelName().ToString() : TEXT(""));
+				SCOPED_CONDITIONAL_DRAW_EVENTF(RHICmdList, LevelEvent, PrimitiveSceneProxy->GetLevelName() != NAME_None, PrimitiveSceneProxy->GetLevelName().IsValid() ? *PrimitiveSceneProxy->GetLevelName().ToString() : TEXT(""));
 			}
-			//SCOPED_CONDITIONAL_DRAW_EVENTF(RHICmdList, OwnerEvent,PrimitiveSceneProxy->GetOwnerName() != NAME_None, *PrimitiveSceneProxy->GetOwnerName().ToString());
+			SCOPED_CONDITIONAL_DRAW_EVENTF(RHICmdList, OwnerEvent,PrimitiveSceneProxy->GetOwnerName() != NAME_None, *PrimitiveSceneProxy->GetOwnerName().ToString());
+		}
 	}
 	else
 	{
diff --git a/Engine/Source/Runtime/Engine/Private/PrimitiveSceneProxy.cpp b/Engine/Source/Runtime/Engine/Private/PrimitiveSceneProxy.cpp
--- a/Engine/Source/Runtime/Engine/Private/PrimitiveSceneProxy.cpp
+++ b/Engine/Source/Runtime/Engine/Private/PrimitiveSceneProxy.cpp
@@ -9,6 +9,13 @@
 #include "Components/BrushComponent.h"
 #include "PrimitiveSceneInfo.h"
 
+static TAutoConsoleVariable<int32> CVarForceSingleSampleShadowingFromStationary(
+	TEXT("r.Shadow.ForceSingleSampleShadowingFromStationary"),
+	0,
+	TEXT("Whether to force all components to act as if they have bSingleSampleShadowFromStationaryLights enabled.  Useful for scalability when dynamic shadows are disabled."),
+	ECVF_RenderThreadSafe | ECVF_Scalability
+	);
+
 FPrimitiveSceneProxy::FPrimitiveSceneProxy(const UPrimitiveComponent* InComponent, FName InResourceName)
 :	WireframeColor(FLinearColor::White)
 ,	LevelColor(FLinearColor::White)
@@ -283,6 +290,11 @@ void FPrimitiveSceneProxy::ApplyLateUpdateTransform(const FMatrix& LateUpdateTra
 	SetTransform(AdjustedLocalToWorld, Bounds, LocalBounds, ActorPosition);
 }
 
+bool FPrimitiveSceneProxy::UseSingleSampleShadowFromStationaryLights() const 
+{ 
+	return bSingleSampleShadowFromStationaryLights || CVarForceSingleSampleShadowingFromStationary.GetValueOnRenderThread() != 0; 
+}
+
 /**
  * Updates selection for the primitive proxy. This is called in the rendering thread by SetSelection_GameThread.
  * @param bInSelected - true if the parent actor is selected in the editor
diff --git a/Engine/Source/Runtime/Engine/Private/RenderingSettings.cpp b/Engine/Source/Runtime/Engine/Private/RenderingSettings.cpp
--- a/Engine/Source/Runtime/Engine/Private/RenderingSettings.cpp
+++ b/Engine/Source/Runtime/Engine/Private/RenderingSettings.cpp
@@ -78,6 +78,6 @@ void URendererSettings::PostEditChangeProperty(FPropertyChangedEvent& PropertyCh
 
 void URendererSettings::SanatizeReflectionCaptureResolution()
 {
-	static const int32 MaxReflectionCaptureResolution = 4096;
+	static const int32 MaxReflectionCaptureResolution = 1024;
 	ReflectionCaptureResolution = FMath::Min(int32(FMath::RoundUpToPowerOfTwo(ReflectionCaptureResolution)), MaxReflectionCaptureResolution);
 }
diff --git a/Engine/Source/Runtime/Engine/Private/SceneView.cpp b/Engine/Source/Runtime/Engine/Private/SceneView.cpp
--- a/Engine/Source/Runtime/Engine/Private/SceneView.cpp
+++ b/Engine/Source/Runtime/Engine/Private/SceneView.cpp
@@ -2115,7 +2115,7 @@ public:
 		{
 			if(CObj->TestFlags(ECVF_Scalability) || CObj->TestFlags(ECVF_ScalabilityGroup))
 			{
-				// float shoudl work on int32 as well
+				// float should work on int32 as well
 				float Value = CVar->GetFloat();
 				Crc = FCrc::MemCrc32(&Value, sizeof(Value), Crc);
 			}
@@ -2141,7 +2141,7 @@ static void DisplayInternals(FRHICommandListImmediate& RHICmdList, FSceneView& I
 		auto State = InView.State;
 
 		FCanvas Canvas((FRenderTarget*)Family->RenderTarget, NULL, Family->CurrentRealTime, Family->CurrentWorldTime, Family->DeltaWorldTime, InView.GetFeatureLevel());
-		Canvas.SetRenderTargetRect(InView.ViewRect);
+		Canvas.SetRenderTargetRect(FIntRect(0, 0, Family->RenderTarget->GetSizeXY().X, Family->RenderTarget->GetSizeXY().Y));
 
 		SetRenderTarget(RHICmdList, Family->RenderTarget->GetRenderTargetTexture(), FTextureRHIRef());
 
@@ -2154,7 +2154,8 @@ static void DisplayInternals(FRHICommandListImmediate& RHICmdList, FSceneView& I
 		const int32 FontSizeY = 14;
 
 		// dark background
-		Canvas.DrawTile(Pos.X - 4, Pos.Y - 4, 500 + 8, FontSizeY * 26 + 8, 0, 0, 1, 1, FLinearColor(0,0,0,0.6f), 0, true);
+		const uint32 BackgroundHeight = 28;
+		Canvas.DrawTile(Pos.X - 4, Pos.Y - 4, 500 + 8, FontSizeY * BackgroundHeight + 8, 0, 0, 1, 1, FLinearColor(0,0,0,0.6f), 0, true);
 
 		UFont* Font = GEngine->GetSmallFont();
 		FCanvasTextItem SmallTextItem( Pos, FText::GetEmpty(), GEngine->GetSmallFont(), FLinearColor::White );
@@ -2163,42 +2164,61 @@ static void DisplayInternals(FRHICommandListImmediate& RHICmdList, FSceneView& I
 		SmallTextItem.Text = FText::FromString(FString::Printf(TEXT("r.DisplayInternals = %d"), Family->DisplayInternalsData.DisplayInternalsCVarValue));
 		Canvas.DrawItem(SmallTextItem, Pos);
 		SmallTextItem.SetColor(FLinearColor::Gray);
-		Pos.Y += FontSizeY;
-		Pos.Y += FontSizeY;
-
-		SmallTextItem.Text = FText::FromString(TEXT("command line options:")); Canvas.DrawItem(SmallTextItem, Pos); Pos.Y += FontSizeY;
-		{
-			bool bDeterministic = FApp::UseFixedTimeStep() && FApp::bUseFixedSeed;
-			SmallTextItem.SetColor(bDeterministic ? FLinearColor::Gray : FLinearColor::Red);
-			SmallTextItem.Text = FText::FromString(FString::Printf(TEXT("  -UseFixedTimeStep: %u"), FApp::UseFixedTimeStep())); Canvas.DrawItem(SmallTextItem, Pos); Pos.Y += FontSizeY;
-			SmallTextItem.Text = FText::FromString(FString::Printf(TEXT("  -FixedSeed: %u"), FApp::bUseFixedSeed)); Canvas.DrawItem(SmallTextItem, Pos); Pos.Y += FontSizeY;
-			SmallTextItem.SetColor(FLinearColor::Gray);
-			SmallTextItem.Text = FText::FromString(FString::Printf(TEXT("  -gABC= (changelist): %d"), GetChangeListNumberForPerfTesting())); Canvas.DrawItem(SmallTextItem, Pos); Pos.Y += FontSizeY;
-		}
-
-		SmallTextItem.Text = FText::FromString(TEXT("Global:")); Canvas.DrawItem(SmallTextItem, Pos); Pos.Y += FontSizeY;
-		SmallTextItem.Text = FText::FromString(FString::Printf(TEXT("  Scalability CVar Hash: %x (use console command \"Scalability\")"), ComputeScalabilityCVarHash())); Canvas.DrawItem(SmallTextItem, Pos); Pos.Y += FontSizeY;
-		SmallTextItem.Text = FText::FromString(FString::Printf(TEXT("  FrameNumberRT: %u"), GFrameNumberRenderThread)); Canvas.DrawItem(SmallTextItem, Pos); Pos.Y += FontSizeY;
-		SmallTextItem.Text = FText::FromString(FString::Printf(TEXT("  FrameCounter: %u"), GFrameCounter)); Canvas.DrawItem(SmallTextItem, Pos); Pos.Y += FontSizeY;
-		SmallTextItem.Text = FText::FromString(FString::Printf(TEXT("  rand()/SRand: %x/%x"), FMath::Rand(), FMath::GetRandSeed())); Canvas.DrawItem(SmallTextItem, Pos); Pos.Y += FontSizeY;
-
-		SmallTextItem.Text = FText::FromString(TEXT("State:")); Canvas.DrawItem(SmallTextItem, Pos); Pos.Y += FontSizeY;
-		SmallTextItem.Text = FText::FromString(FString::Printf(TEXT("  TemporalAASample: %u"), State->GetCurrentTemporalAASampleIndex())); Canvas.DrawItem(SmallTextItem, Pos); Pos.Y += FontSizeY;
-		SmallTextItem.Text = FText::FromString(FString::Printf(TEXT("  LODTransition: %.2f"), State->GetTemporalLODTransition())); Canvas.DrawItem(SmallTextItem, Pos); Pos.Y += FontSizeY;
-
-		SmallTextItem.Text = FText::FromString(TEXT("Family:")); Canvas.DrawItem(SmallTextItem, Pos); Pos.Y += FontSizeY;
-		SmallTextItem.Text = FText::FromString(FString::Printf(TEXT("  Time (Real/World/DeltaWorld): %.2f/%.2f/%.2f"), Family->CurrentRealTime, Family->CurrentWorldTime, Family->DeltaWorldTime)); Canvas.DrawItem(SmallTextItem, Pos); Pos.Y += FontSizeY;
-		SmallTextItem.Text = FText::FromString(FString::Printf(TEXT("  MatineeTime: %f"), Family->DisplayInternalsData.MatineeTime)); Canvas.DrawItem(SmallTextItem, Pos); Pos.Y += FontSizeY;
-		SmallTextItem.Text = FText::FromString(FString::Printf(TEXT("  FrameNumber: %u"), Family->FrameNumber)); Canvas.DrawItem(SmallTextItem, Pos); Pos.Y += FontSizeY;
-		SmallTextItem.Text = FText::FromString(FString::Printf(TEXT("  ExposureSettings: %s"), *Family->ExposureSettings.ToString())); Canvas.DrawItem(SmallTextItem, Pos); Pos.Y += FontSizeY;
-		SmallTextItem.Text = FText::FromString(FString::Printf(TEXT("  GammaCorrection: %.2f"), Family->GammaCorrection)); Canvas.DrawItem(SmallTextItem, Pos); Pos.Y += FontSizeY;
-
-		SmallTextItem.Text = FText::FromString(TEXT("View:")); Canvas.DrawItem(SmallTextItem, Pos); Pos.Y += FontSizeY;
-		SmallTextItem.Text = FText::FromString(FString::Printf(TEXT("  TemporalJitter: %.2f/%.2f"), InView.TemporalJitterPixelsX, InView.TemporalJitterPixelsY)); Canvas.DrawItem(SmallTextItem, Pos); Pos.Y += FontSizeY;
-		SmallTextItem.Text = FText::FromString(FString::Printf(TEXT("  ViewProjectionMatrix Hash: %x"), InView.ViewProjectionMatrix.ComputeHash())); Canvas.DrawItem(SmallTextItem, Pos); Pos.Y += FontSizeY;
-		SmallTextItem.Text = FText::FromString(FString::Printf(TEXT("  ViewLocation: %s"), *InView.ViewLocation.ToString())); Canvas.DrawItem(SmallTextItem, Pos); Pos.Y += FontSizeY;
-		SmallTextItem.Text = FText::FromString(FString::Printf(TEXT("  ViewRotation: %s"), *InView.ViewRotation.ToString())); Canvas.DrawItem(SmallTextItem, Pos); Pos.Y += FontSizeY;
-		SmallTextItem.Text = FText::FromString(FString::Printf(TEXT("  ViewRect: %s"), *InView.ViewRect.ToString())); Canvas.DrawItem(SmallTextItem, Pos); Pos.Y += FontSizeY;
+		Pos.Y += 2 * FontSizeY;
+
+#define CANVAS_HEADER(txt) \
+		{ \
+			SmallTextItem.SetColor(FLinearColor::Gray); \
+			SmallTextItem.Text = FText::FromString(txt); \
+			Canvas.DrawItem(SmallTextItem, Pos); \
+			Pos.Y += FontSizeY; \
+		}
+#define CANVAS_LINE(bHighlight, txt, ... ) \
+		{ \
+			SmallTextItem.SetColor(bHighlight ? FLinearColor::Red : FLinearColor::Gray); \
+			SmallTextItem.Text = FText::FromString(FString::Printf(txt, __VA_ARGS__)); \
+			Canvas.DrawItem(SmallTextItem, Pos); \
+			Pos.Y += FontSizeY; \
+		}
+
+		CANVAS_HEADER(TEXT("command line options:"))
+		{
+			bool bHighlight = FApp::UseFixedTimeStep() && FApp::bUseFixedSeed;
+			CANVAS_LINE(bHighlight, TEXT("  -UseFixedTimeStep: %u"), FApp::UseFixedTimeStep())
+			CANVAS_LINE(bHighlight, TEXT("  -FixedSeed: %u"), FApp::bUseFixedSeed)
+			CANVAS_LINE(false, TEXT("  -gABC= (changelist): %d"), GetChangeListNumberForPerfTesting())
+		}
+
+		CANVAS_HEADER(TEXT("Global:"))
+		CANVAS_LINE(false, TEXT("  FrameNumberRT: %u"), GFrameNumberRenderThread)
+		CANVAS_LINE(false, TEXT("  Scalability CVar Hash: %x (use console command \"Scalability\")"), ComputeScalabilityCVarHash())
+		CANVAS_LINE(false, TEXT("  FrameNumberRT: %u"), GFrameNumberRenderThread)
+		CANVAS_LINE(false, TEXT("  FrameCounter: %u"), GFrameCounter)
+		CANVAS_LINE(false, TEXT("  rand()/SRand: %x/%x"), FMath::Rand(), FMath::GetRandSeed())
+		{
+			bool bHighlight = Family->DisplayInternalsData.NumPendingStreamingRequests != 0;
+			CANVAS_LINE(bHighlight, TEXT("  FStreamAllResourcesLatentCommand: %d"), bHighlight)
+		}
+
+		CANVAS_HEADER(TEXT("State:"))
+		CANVAS_LINE(false, TEXT("  TemporalAASample: %u"), State->GetCurrentTemporalAASampleIndex())
+		CANVAS_LINE(false, TEXT("  LODTransition: %.2f"), State->GetTemporalLODTransition())
+
+		CANVAS_HEADER(TEXT("Family:"))
+		CANVAS_LINE(false, TEXT("  Time (Real/World/DeltaWorld): %.2f/%.2f/%.2f"), Family->CurrentRealTime, Family->CurrentWorldTime, Family->DeltaWorldTime)
+		CANVAS_LINE(false, TEXT("  MatineeTime: %f"), Family->DisplayInternalsData.MatineeTime)
+		CANVAS_LINE(false, TEXT("  FrameNumber: %u"), Family->FrameNumber)
+		CANVAS_LINE(false, TEXT("  ExposureSettings: %s"), *Family->ExposureSettings.ToString())
+		CANVAS_LINE(false, TEXT("  GammaCorrection: %.2f"), Family->GammaCorrection)
+
+		CANVAS_HEADER(TEXT("View:"))
+		CANVAS_LINE(false, TEXT("  TemporalJitter: %.2f/%.2f"), InView.TemporalJitterPixelsX, InView.TemporalJitterPixelsY)
+		CANVAS_LINE(false, TEXT("  ViewProjectionMatrix Hash: %x"), InView.ViewProjectionMatrix.ComputeHash())
+		CANVAS_LINE(false, TEXT("  ViewLocation: %s"), *InView.ViewLocation.ToString())
+		CANVAS_LINE(false, TEXT("  ViewRotation: %s"), *InView.ViewRotation.ToString())
+		CANVAS_LINE(false, TEXT("  ViewRect: %s"), *InView.ViewRect.ToString())
+
+#undef CANVAS_LINE
 
 		Canvas.Flush_RenderThread(RHICmdList);
 	}
diff --git a/Engine/Source/Runtime/Engine/Private/ShaderCompiler/ShaderCompiler.cpp b/Engine/Source/Runtime/Engine/Private/ShaderCompiler/ShaderCompiler.cpp
--- a/Engine/Source/Runtime/Engine/Private/ShaderCompiler/ShaderCompiler.cpp
+++ b/Engine/Source/Runtime/Engine/Private/ShaderCompiler/ShaderCompiler.cpp
@@ -2380,11 +2380,6 @@ void GlobalBeginCompileShader(
 		Input.Environment.SetDefine(TEXT("SELECTIVE_BASEPASS_OUTPUTS"), CVar ? (CVar->GetValueOnGameThread() != 0) : 0);
 	}
 
-	{
-		static const auto CVar = IConsoleManager::Get().FindTConsoleVariableDataInt(TEXT("r.GBuffer"));
-		Input.Environment.SetDefine(TEXT("NO_GBUFFER"), CVar ? (CVar->GetValueOnGameThread() == 0) : 0);
-	}
-
 	{
 		static IConsoleVariable* CVar = IConsoleManager::Get().FindConsoleVariable(TEXT("r.DBuffer"));
 		Input.Environment.SetDefine(TEXT("USE_DBUFFER"), CVar ? CVar->GetInt() : 0);
@@ -2673,7 +2668,7 @@ bool RecompileShaders(const TCHAR* Cmd, FOutputDevice& Ar)
 				UMaterialInterface::IterateOverActiveFeatureLevels([&](ERHIFeatureLevel::Type InFeatureLevel) {
 					auto ShaderPlatform = GShaderPlatformForFeatureLevel[InFeatureLevel];
 					BeginRecompileGlobalShaders(ShaderTypes, ShaderPipelineTypes, ShaderPlatform);
-					UMaterial::UpdateMaterialShaders(ShaderTypes, ShaderPipelineTypes, FactoryTypes, ShaderPlatform);
+					//UMaterial::UpdateMaterialShaders(ShaderTypes, ShaderPipelineTypes, FactoryTypes, ShaderPlatform);
 					FinishRecompileGlobalShaders();
 				});
 			}
diff --git a/Engine/Source/Runtime/Engine/Private/ShowFlags.cpp b/Engine/Source/Runtime/Engine/Private/ShowFlags.cpp
--- a/Engine/Source/Runtime/Engine/Private/ShowFlags.cpp
+++ b/Engine/Source/Runtime/Engine/Private/ShowFlags.cpp
@@ -361,19 +361,6 @@ void EngineShowFlagOverride(EShowFlagInitMode ShowFlagInitMode, EViewModeIndex V
 		}
 	}
 
-	{
-		static const auto ICVar = IConsoleManager::Get().FindTConsoleVariableDataInt(TEXT("r.GBuffer"));
-		if(ICVar->GetValueOnGameThread() == 0)
-		{
-			EngineShowFlags.AmbientOcclusion = 0;
-			EngineShowFlags.ScreenSpaceAO = 0;
-			EngineShowFlags.Decals = 0;
-			EngineShowFlags.DynamicShadows = 0;
-			EngineShowFlags.GlobalIllumination = 0;
-			EngineShowFlags.ScreenSpaceReflections = 0;
-		}
-	}
-
 	{
 		static const auto ICVar = IConsoleManager::Get().FindTConsoleVariableDataInt(TEXT("r.EyeAdaptationQuality"));
 		if(ICVar->GetValueOnGameThread() <= 0)
diff --git a/Engine/Source/Runtime/Engine/Private/SkeletalRenderCPUSkin.cpp b/Engine/Source/Runtime/Engine/Private/SkeletalRenderCPUSkin.cpp
--- a/Engine/Source/Runtime/Engine/Private/SkeletalRenderCPUSkin.cpp
+++ b/Engine/Source/Runtime/Engine/Private/SkeletalRenderCPUSkin.cpp
@@ -546,8 +546,8 @@ FORCEINLINE void RebuildTangentBasis( VertexType& DestVertex )
 {
 	// derive the new tangent by orthonormalizing the new normal against
 	// the base tangent vector (assuming these are normalized)
-	FVector Tangent( DestVertex.TangentX );
-	FVector Normal( DestVertex.TangentZ );
+	FVector Tangent = DestVertex.TangentX;
+	FVector Normal = DestVertex.TangentZ;
 	Tangent = Tangent - ((Tangent | Normal) * Normal);
 	Tangent.Normalize();
 	DestVertex.TangentX = Tangent;
@@ -564,8 +564,11 @@ FORCEINLINE void ApplyMorphBlend( VertexType& DestVertex, const FVertexAnimDelta
 
 	// Save W before = operator. That overwrites W to be 127.
 	uint8 W = DestVertex.TangentZ.Vector.W;
+
+	FVector TanZ = DestVertex.TangentZ;
+
 	// add normal offset. can only apply normal deltas up to a weight of 1
-	DestVertex.TangentZ = FVector(FVector(DestVertex.TangentZ) + SrcMorph.TangentZDelta * FMath::Min(Weight,1.0f)).GetUnsafeNormal();
+	DestVertex.TangentZ = FVector(TanZ + SrcMorph.TangentZDelta * FMath::Min(Weight,1.0f)).GetUnsafeNormal();
 	// Recover W
 	DestVertex.TangentZ.Vector.W = W;
 } 
diff --git a/Engine/Source/Runtime/Engine/Private/StaticMesh.cpp b/Engine/Source/Runtime/Engine/Private/StaticMesh.cpp
--- a/Engine/Source/Runtime/Engine/Private/StaticMesh.cpp
+++ b/Engine/Source/Runtime/Engine/Private/StaticMesh.cpp
@@ -63,7 +63,7 @@ FStaticMeshVertexBuffer::FStaticMeshVertexBuffer():
 	Stride(0),
 	NumVertices(0),
 	bUseFullPrecisionUVs(false),
-	bUseHighPrecisionTangentBasis(true)
+	bUseHighPrecisionTangentBasis(false)
 {}
 
 FStaticMeshVertexBuffer::~FStaticMeshVertexBuffer()
@@ -145,72 +145,43 @@ void FStaticMeshVertexBuffer::RemoveLegacyShadowVolumeVertices(uint32 InNumVerti
 	Data = VertexData->GetDataPointer();
 }
 
-/**
-* Convert the existing data in this mesh from 16 bit to 32 bit UVs.
-* Without rebuilding the mesh (loss of precision)
-*/
-template<int32 NumTexCoordsT>
-void FStaticMeshVertexBuffer::ConvertToFullPrecisionUVs()
-{
-	if( !bUseFullPrecisionUVs )
-	{
-		check(NumTexCoords == NumTexCoordsT);
-		// create temp array to store 32 bit values
-		TArray< TStaticMeshFullVertexFloat32UVs<NumTexCoordsT> > DestVertexData;
-		// source vertices
-		TStaticMeshVertexData< TStaticMeshFullVertexFloat16UVs<NumTexCoordsT> >& SrcVertexData = 
-			*(TStaticMeshVertexData< TStaticMeshFullVertexFloat16UVs<NumTexCoordsT> >*)VertexData;
-		// copy elements from source vertices to temp data
-		DestVertexData.AddUninitialized(SrcVertexData.Num());
-		for( int32 VertIdx=0; VertIdx < SrcVertexData.Num(); VertIdx++ )
-		{
-			TStaticMeshFullVertexFloat32UVs<NumTexCoordsT>& DestVert = DestVertexData[VertIdx];
-			TStaticMeshFullVertexFloat16UVs<NumTexCoordsT>& SrcVert = SrcVertexData[VertIdx];		
-			FMemory::Memcpy(&DestVert,&SrcVert,sizeof(FStaticMeshFullVertex));
-			for( int32 UVIdx=0; UVIdx < NumTexCoordsT; UVIdx++ )
-			{
-				DestVert.UVs[UVIdx] = FVector2D(SrcVert.UVs[UVIdx]);
-			}
-		}
-		// force 32 bit UVs
-		bUseFullPrecisionUVs = true;
-		AllocateData();
-		*(TStaticMeshVertexData< TStaticMeshFullVertexFloat32UVs<NumTexCoordsT> >*)VertexData = DestVertexData;
-		Data = VertexData->GetDataPointer();
-		Stride = VertexData->GetStride();
+template<typename SrcVertexTypeT, typename DstVertexTypeT>
+void FStaticMeshVertexBuffer::ConvertVertexFormat()
+{
+	if (SrcVertexTypeT::TangentBasisType == DstVertexTypeT::TangentBasisType &&
+		SrcVertexTypeT::UVType == DstVertexTypeT::UVType)
+	{
+		return;
 	}
-}
 
-template<typename StaticMeshVertexT>
-void ConvertInPlaceLowPrecisionTangentBasis(StaticMeshVertexT* VertexData)
-{
-	auto& SrcVertexData = *VertexData;
+	check(SrcVertexTypeT::NumTexCoords == DstVertexTypeT::NumTexCoords);
 
-	for (int32 VertIdx = 0; VertIdx < SrcVertexData.Num(); ++VertIdx)
-	{
-		auto& SrcVert = SrcVertexData[VertIdx];
+	auto& SrcVertexData = *static_cast<TStaticMeshVertexData<SrcVertexTypeT>*>(VertexData);
 
-		FVector4 FullVector = *reinterpret_cast<FPackedRGB10A2*>(&SrcVert.RawTangentX);
-		*reinterpret_cast<FPackedNormal*>(&SrcVert.RawTangentX) = FullVector;
-	}
-}
+	TArray<DstVertexTypeT> DstVertexData;
+	DstVertexData.AddUninitialized(SrcVertexData.Num());
 
-template<int32 NumTexCoordsT>
-void FStaticMeshVertexBuffer::ConvertToLowPrecisionTangentBasis()
-{
-	if (bUseHighPrecisionTangentBasis)
+	for (int32 VertIdx = 0; VertIdx < SrcVertexData.Num(); VertIdx++)
 	{
-		if (GetUseFullPrecisionUVs())
-		{
-			ConvertInPlaceLowPrecisionTangentBasis((TStaticMeshVertexData<TStaticMeshFullVertexFloat32UVs<NumTexCoordsT> >*)VertexData);
-		}
-		else
+		SrcVertexTypeT& SrcVert = SrcVertexData[VertIdx];
+		DstVertexTypeT& DstVert = DstVertexData[VertIdx];
+
+		DstVert.SetTangents(SrcVert.GetTangentX(), SrcVert.GetTangentY(), SrcVert.GetTangentZ());
+
+		for (int32 UVIdx = 0; UVIdx < DstVertexTypeT::NumTexCoords; UVIdx++)
 		{
-			ConvertInPlaceLowPrecisionTangentBasis((TStaticMeshVertexData<TStaticMeshFullVertexFloat16UVs<NumTexCoordsT> >*)VertexData);
+			DstVert.SetUV(UVIdx, SrcVert.GetUV(UVIdx));
 		}
-
-		bUseHighPrecisionTangentBasis = false;
 	}
+
+	bUseHighPrecisionTangentBasis = DstVertexTypeT::TangentBasisType == EStaticMeshVertexTangentBasisType::HighPrecision;
+	bUseFullPrecisionUVs = DstVertexTypeT::UVType == EStaticMeshVertexUVType::HighPrecision;
+
+	AllocateData();
+	*static_cast<TStaticMeshVertexData<DstVertexTypeT>*>(VertexData) = DstVertexData;
+
+	Data = VertexData->GetDataPointer();
+	Stride = VertexData->GetStride();
 }
 
 /**
@@ -277,36 +248,12 @@ void FStaticMeshVertexBuffer::AllocateData( bool bNeedsCPUAccess /*= true*/ )
 	// Clear any old VertexData before allocating.
 	CleanUp();
 
-	if( !bUseFullPrecisionUVs )
-	{
-		switch(NumTexCoords)
-		{
-		case 1: VertexData = new TStaticMeshVertexData< TStaticMeshFullVertexFloat16UVs<1> >(bNeedsCPUAccess); break;
-		case 2: VertexData = new TStaticMeshVertexData< TStaticMeshFullVertexFloat16UVs<2> >(bNeedsCPUAccess); break;
-		case 3: VertexData = new TStaticMeshVertexData< TStaticMeshFullVertexFloat16UVs<3> >(bNeedsCPUAccess); break;
-		case 4: VertexData = new TStaticMeshVertexData< TStaticMeshFullVertexFloat16UVs<4> >(bNeedsCPUAccess); break;
-		case 5: VertexData = new TStaticMeshVertexData< TStaticMeshFullVertexFloat16UVs<5> >(bNeedsCPUAccess); break;
-		case 6: VertexData = new TStaticMeshVertexData< TStaticMeshFullVertexFloat16UVs<6> >(bNeedsCPUAccess); break;
-		case 7: VertexData = new TStaticMeshVertexData< TStaticMeshFullVertexFloat16UVs<7> >(bNeedsCPUAccess); break;
-		case 8: VertexData = new TStaticMeshVertexData< TStaticMeshFullVertexFloat16UVs<8> >(bNeedsCPUAccess); break;
-		default: UE_LOG(LogStaticMesh, Fatal,TEXT("Invalid number of texture coordinates"));
-		};		
-	}
-	else
-	{
-		switch(NumTexCoords)
-		{
-		case 1: VertexData = new TStaticMeshVertexData< TStaticMeshFullVertexFloat32UVs<1> >(bNeedsCPUAccess); break;
-		case 2: VertexData = new TStaticMeshVertexData< TStaticMeshFullVertexFloat32UVs<2> >(bNeedsCPUAccess); break;
-		case 3: VertexData = new TStaticMeshVertexData< TStaticMeshFullVertexFloat32UVs<3> >(bNeedsCPUAccess); break;
-		case 4: VertexData = new TStaticMeshVertexData< TStaticMeshFullVertexFloat32UVs<4> >(bNeedsCPUAccess); break;
-		case 5: VertexData = new TStaticMeshVertexData< TStaticMeshFullVertexFloat32UVs<5> >(bNeedsCPUAccess); break;
-		case 6: VertexData = new TStaticMeshVertexData< TStaticMeshFullVertexFloat32UVs<6> >(bNeedsCPUAccess); break;
-		case 7: VertexData = new TStaticMeshVertexData< TStaticMeshFullVertexFloat32UVs<7> >(bNeedsCPUAccess); break;
-		case 8: VertexData = new TStaticMeshVertexData< TStaticMeshFullVertexFloat32UVs<8> >(bNeedsCPUAccess); break;
-		default: UE_LOG(LogStaticMesh, Fatal,TEXT("Invalid number of texture coordinates"));
-		};		
-	}	
+	SELECT_STATIC_MESH_VERTEX_TYPE(
+		GetUseHighPrecisionTangentBasis(),
+		GetUseFullPrecisionUVs(),
+		GetNumTexCoords(),
+		VertexData = new TStaticMeshVertexData<VertexType>(bNeedsCPUAccess);
+		);
 
 	// Calculate the vertex stride.
 	Stride = VertexData->GetStride();
@@ -426,10 +373,27 @@ void FStaticMeshLODResources::InitVertexFactory(
 	Params.bOverrideColorVertexBuffer = bInOverrideColorVertexBuffer;
 	Params.Parent = InParentMesh;
 
+	uint32 TangentXOffset = 0;
+	uint32 TangetnZOffset = 0;
+	uint32 UVsBaseOffset = 0;
+
+	SELECT_STATIC_MESH_VERTEX_TYPE(
+		Params.LODResources->VertexBuffer.GetUseHighPrecisionTangentBasis(),
+		Params.LODResources->VertexBuffer.GetUseFullPrecisionUVs(),
+		Params.LODResources->VertexBuffer.GetNumTexCoords(),
+		{
+			TangentXOffset = STRUCT_OFFSET(VertexType, TangentX);
+			TangetnZOffset = STRUCT_OFFSET(VertexType, TangentZ);
+			UVsBaseOffset = STRUCT_OFFSET(VertexType, UVs);
+		});
+
 	// Initialize the static mesh's vertex factory.
-	ENQUEUE_UNIQUE_RENDER_COMMAND_ONEPARAMETER(
+	ENQUEUE_UNIQUE_RENDER_COMMAND_FOURPARAMETER(
 		InitStaticMeshVertexFactory,
-		InitStaticMeshVertexFactoryParams,Params,Params,
+		InitStaticMeshVertexFactoryParams, Params, Params,
+		uint32, TangentXOffset, TangentXOffset,
+		uint32, TangetnZOffset, TangetnZOffset,
+		uint32, UVsBaseOffset, UVsBaseOffset,
 		{
 			FLocalVertexFactory::FDataType Data;
 			Data.PositionComponent = FVertexStreamComponent(
@@ -441,16 +405,20 @@ void FStaticMeshLODResources::InitVertexFactory(
 
 			Data.TangentBasisComponents[0] = FVertexStreamComponent(
 				&Params.LODResources->VertexBuffer,
-				STRUCT_OFFSET(FStaticMeshFullVertex, RawTangentX),
+				TangentXOffset,
 				Params.LODResources->VertexBuffer.GetStride(),
-				Params.LODResources->VertexBuffer.GetUseHighPrecisionTangentBasis() ? VET_URGB10A2N : VET_PackedNormal
+				Params.LODResources->VertexBuffer.GetUseHighPrecisionTangentBasis() ? 
+					TStaticMeshVertexTangentTypeSelector<EStaticMeshVertexTangentBasisType::HighPrecision>::VertexElementType : 
+					TStaticMeshVertexTangentTypeSelector<EStaticMeshVertexTangentBasisType::Default>::VertexElementType
 				);
 
 			Data.TangentBasisComponents[1] = FVertexStreamComponent(
 				&Params.LODResources->VertexBuffer,
-				STRUCT_OFFSET(FStaticMeshFullVertex, TangentZ),
+				TangetnZOffset,
 				Params.LODResources->VertexBuffer.GetStride(),
-				VET_UShort2N
+				Params.LODResources->VertexBuffer.GetUseHighPrecisionTangentBasis() ?
+					TStaticMeshVertexTangentTypeSelector<EStaticMeshVertexTangentBasisType::HighPrecision>::VertexElementType : 
+					TStaticMeshVertexTangentTypeSelector<EStaticMeshVertexTangentBasisType::Default>::VertexElementType
 				);
 
 			// Use the "override" color vertex buffer if one was supplied.  Otherwise, the color vertex stream
@@ -482,72 +450,46 @@ void FStaticMeshLODResources::InitVertexFactory(
 
 			Data.TextureCoordinates.Empty();
 
-			if( !Params.LODResources->VertexBuffer.GetUseFullPrecisionUVs() )
-			{
-				int32 UVIndex;
-				for (UVIndex = 0; UVIndex < (int32)Params.LODResources->VertexBuffer.GetNumTexCoords() - 1; UVIndex += 2)
-				{
-					Data.TextureCoordinates.Add(FVertexStreamComponent(
-						&Params.LODResources->VertexBuffer,
-						STRUCT_OFFSET(TStaticMeshFullVertexFloat16UVs<MAX_STATIC_TEXCOORDS>, UVs) + sizeof(FVector2DHalf)* UVIndex,
-						Params.LODResources->VertexBuffer.GetStride(),
-						VET_Half4
-						));
-				}
-				// possible last UV channel if we have an odd number
-				if (UVIndex < (int32)Params.LODResources->VertexBuffer.GetNumTexCoords())
-				{
-					Data.TextureCoordinates.Add(FVertexStreamComponent(
-						&Params.LODResources->VertexBuffer,
-						STRUCT_OFFSET(TStaticMeshFullVertexFloat16UVs<MAX_STATIC_TEXCOORDS>, UVs) + sizeof(FVector2DHalf)* UVIndex,
-						Params.LODResources->VertexBuffer.GetStride(),
-						VET_Half2
-						));
-				}
+			uint32 UVSizeInBytes = Params.LODResources->VertexBuffer.GetUseFullPrecisionUVs() ?
+				sizeof(TStaticMeshVertexUVsTypeSelector<EStaticMeshVertexUVType::HighPrecision>::UVsTypeT) : sizeof(TStaticMeshVertexUVsTypeSelector<EStaticMeshVertexUVType::Default>::UVsTypeT);
 
-				if(	Params.Parent->LightMapCoordinateIndex >= 0 && (uint32)Params.Parent->LightMapCoordinateIndex < Params.LODResources->VertexBuffer.GetNumTexCoords())
-				{
-					Data.LightMapCoordinateComponent = FVertexStreamComponent(
-						&Params.LODResources->VertexBuffer,
-						STRUCT_OFFSET(TStaticMeshFullVertexFloat16UVs<MAX_STATIC_TEXCOORDS>,UVs) + sizeof(FVector2DHalf) * Params.Parent->LightMapCoordinateIndex,
-						Params.LODResources->VertexBuffer.GetStride(),
-						VET_Half2
-						);
-				}
+			EVertexElementType UVDoubleWideVertexElementType = Params.LODResources->VertexBuffer.GetUseFullPrecisionUVs() ?
+				VET_Float4 : VET_Half4;
+
+			EVertexElementType UVVertexElementType = Params.LODResources->VertexBuffer.GetUseFullPrecisionUVs() ?
+				VET_Float2 : VET_Half2;
+
+			int32 UVIndex;
+			for (UVIndex = 0; UVIndex < (int32)Params.LODResources->VertexBuffer.GetNumTexCoords() - 1; UVIndex += 2)
+			{
+				Data.TextureCoordinates.Add(FVertexStreamComponent(
+					&Params.LODResources->VertexBuffer,
+					UVsBaseOffset + UVSizeInBytes * UVIndex,
+					Params.LODResources->VertexBuffer.GetStride(),
+					UVDoubleWideVertexElementType
+					));
 			}
-			else
+
+			// possible last UV channel if we have an odd number
+			if (UVIndex < (int32)Params.LODResources->VertexBuffer.GetNumTexCoords())
 			{
-				int32 UVIndex;
-				for (UVIndex = 0; UVIndex < (int32)Params.LODResources->VertexBuffer.GetNumTexCoords() - 1; UVIndex += 2)
-				{
-					Data.TextureCoordinates.Add(FVertexStreamComponent(
-						&Params.LODResources->VertexBuffer,
-						STRUCT_OFFSET(TStaticMeshFullVertexFloat32UVs<MAX_STATIC_TEXCOORDS>, UVs) + sizeof(FVector2D)* UVIndex,
-						Params.LODResources->VertexBuffer.GetStride(),
-						VET_Float4
-						));
-				}
-				// possible last UV channel if we have an odd number
-				if (UVIndex < (int32)Params.LODResources->VertexBuffer.GetNumTexCoords())
-				{
-					Data.TextureCoordinates.Add(FVertexStreamComponent(
-						&Params.LODResources->VertexBuffer,
-						STRUCT_OFFSET(TStaticMeshFullVertexFloat32UVs<MAX_STATIC_TEXCOORDS>, UVs) + sizeof(FVector2D)* UVIndex,
-						Params.LODResources->VertexBuffer.GetStride(),
-						VET_Float2
-						));
-				}
+				Data.TextureCoordinates.Add(FVertexStreamComponent(
+					&Params.LODResources->VertexBuffer,
+					UVsBaseOffset + UVSizeInBytes * UVIndex,
+					Params.LODResources->VertexBuffer.GetStride(),
+					UVVertexElementType
+					));
+			}
 
-				if(	Params.Parent->LightMapCoordinateIndex >= 0 && (uint32)Params.Parent->LightMapCoordinateIndex < Params.LODResources->VertexBuffer.GetNumTexCoords())
-				{
-					Data.LightMapCoordinateComponent = FVertexStreamComponent(
-						&Params.LODResources->VertexBuffer,
-						STRUCT_OFFSET(TStaticMeshFullVertexFloat32UVs<MAX_STATIC_TEXCOORDS>,UVs) + sizeof(FVector2D) * Params.Parent->LightMapCoordinateIndex,
-						Params.LODResources->VertexBuffer.GetStride(),
-						VET_Float2
-						);
-				}
-			}	
+			if(	Params.Parent->LightMapCoordinateIndex >= 0 && (uint32)Params.Parent->LightMapCoordinateIndex < Params.LODResources->VertexBuffer.GetNumTexCoords())
+			{
+				Data.LightMapCoordinateComponent = FVertexStreamComponent(
+					&Params.LODResources->VertexBuffer,
+					UVsBaseOffset + UVSizeInBytes * Params.Parent->LightMapCoordinateIndex,
+					Params.LODResources->VertexBuffer.GetStride(),
+					UVVertexElementType
+					);
+			}
 
 			Params.VertexFactory->SetData(Data);
 		});
@@ -1131,6 +1073,7 @@ FArchive& operator<<(FArchive& Ar, FMeshBuildSettings& BuildSettings)
 	Ar << BuildSettings.bRemoveDegenerates;
 	Ar << BuildSettings.bBuildAdjacencyBuffer;
 	Ar << BuildSettings.bBuildReversedIndexBuffer;
+	Ar << BuildSettings.bUseHighPrecisionTangentBasis;
 	Ar << BuildSettings.bUseFullPrecisionUVs;
 	Ar << BuildSettings.bGenerateLightmapUVs;
 
@@ -1165,7 +1108,7 @@ FArchive& operator<<(FArchive& Ar, FMeshBuildSettings& BuildSettings)
 // differences, etc.) replace the version GUID below with a new one.
 // In case of merge conflicts with DDC versions, you MUST generate a new GUID
 // and set this new GUID as the version.                                       
-#define STATICMESH_DERIVEDDATA_VER TEXT("7C4164E1813428CB338EDD50CC040049")
+#define STATICMESH_DERIVEDDATA_VER TEXT("6F4494992CB61A37E4B1403C6156E6A")
 
 static const FString& GetStaticMeshDerivedDataVersion()
 {
@@ -1651,7 +1594,8 @@ bool UStaticMesh::GetStreamingTextureFactor(float& OutTexelFactor, FBoxSphereBou
 	}
 
 	OutBounds = TransformedSectionBox;
-	OutTexelFactor = WeightedTexelFactorSum / AreaSum * FMath::Max(0.0f, StreamingDistanceMultiplier);
+	OutTexelFactor = WeightedTexelFactorSum / AreaSum; 
+	// Don't take into account StreamingDistanceMultiplier here (but rather in the components using it). That allows realtime feedback without requiring a TextureStreamingBuild.
 	return true;
 
 #else
@@ -2268,7 +2212,6 @@ void UStaticMesh::Serialize(FArchive& Ar)
 //
 //	UStaticMesh::PostLoad
 //
-
 void UStaticMesh::PostLoad()
 {
 	Super::PostLoad();
@@ -2341,43 +2284,23 @@ void UStaticMesh::PostLoad()
 #endif // #if WITH_EDITOR
 
 	EnforceLightmapRestrictions();
-	
-	if (!GVertexElementTypeSupport.IsSupported(VET_Half2))
-	{
-		for (int32 LODIndex = 0; LODIndex < RenderData->LODResources.Num(); ++LODIndex)
-		{
-			if (RenderData->LODResources.IsValidIndex(LODIndex))
-			{
-				FStaticMeshLODResources& LOD = RenderData->LODResources[LODIndex];
-				// Determine the correct version of ConvertToFullPrecisionUVs based on the number of UVs in the vertex buffer
-				const uint32 NumTexCoords = LOD.VertexBuffer.GetNumTexCoords();
-				switch(NumTexCoords)
-				{
-				case 1: LOD.VertexBuffer.ConvertToFullPrecisionUVs<1>(); break;
-				case 2: LOD.VertexBuffer.ConvertToFullPrecisionUVs<2>(); break; 
-				case 3: LOD.VertexBuffer.ConvertToFullPrecisionUVs<3>(); break; 
-				case 4: LOD.VertexBuffer.ConvertToFullPrecisionUVs<4>(); break; 
-				}
-			}
-		}
-	}
 
-	if (!GVertexElementTypeSupport.IsSupported(VET_URGB10A2N))
+	if (!GVertexElementTypeSupport.IsSupported(VET_Half2))
 	{
 		for (int32 LODIndex = 0; LODIndex < RenderData->LODResources.Num(); ++LODIndex)
 		{
 			if (RenderData->LODResources.IsValidIndex(LODIndex))
 			{
 				FStaticMeshLODResources& LOD = RenderData->LODResources[LODIndex];
-
-				const uint32 NumTexCoords = LOD.VertexBuffer.GetNumTexCoords();
-				switch (NumTexCoords)
-				{
-				case 1: LOD.VertexBuffer.ConvertToLowPrecisionTangentBasis<1>(); break;
-				case 2: LOD.VertexBuffer.ConvertToLowPrecisionTangentBasis<2>(); break;
-				case 3: LOD.VertexBuffer.ConvertToLowPrecisionTangentBasis<3>(); break;
-				case 4: LOD.VertexBuffer.ConvertToLowPrecisionTangentBasis<4>(); break;
-				}
+				
+				SELECT_STATIC_MESH_VERTEX_TYPE(
+					LOD.VertexBuffer.GetUseHighPrecisionTangentBasis(),
+					LOD.VertexBuffer.GetUseFullPrecisionUVs(),
+					LOD.VertexBuffer.GetNumTexCoords(),
+					{
+						typedef TStaticMeshFullVertex<VertexType::TangentBasisType, EStaticMeshVertexUVType::HighPrecision, VertexType::NumTexCoords> DstVertexType;
+						LOD.VertexBuffer.ConvertVertexFormat<VertexType, DstVertexType>();
+					});
 			}
 		}
 	}
diff --git a/Engine/Source/Runtime/Engine/Private/StaticMeshBuild.cpp b/Engine/Source/Runtime/Engine/Private/StaticMeshBuild.cpp
--- a/Engine/Source/Runtime/Engine/Private/StaticMeshBuild.cpp
+++ b/Engine/Source/Runtime/Engine/Private/StaticMeshBuild.cpp
@@ -320,14 +320,18 @@ void RemapPaintedVertexColors(
 		if ( PointsToConsider.Num() > 0 )
 		{
 			FPaintedVertex BestVertex = PointsToConsider[0];
+			FVector BestVertexNormal = BestVertex.Normal;
+
 			float BestDistanceSquared = ( BestVertex.Position - CurPosition ).SizeSquared();
-			float BestNormalDot = FVector( BestVertex.Normal ) | CurNormal;
+			float BestNormalDot = BestVertexNormal | CurNormal;
 
 			for ( int32 ConsiderationIndex = 1; ConsiderationIndex < PointsToConsider.Num(); ++ConsiderationIndex )
 			{
 				FPaintedVertex& Vertex = PointsToConsider[ ConsiderationIndex ];
+				FVector VertexNormal = Vertex.Normal;
+
 				const float DistSqrd = ( Vertex.Position - CurPosition ).SizeSquared();
-				const float NormalDot = FVector( Vertex.Normal ) | CurNormal;
+				const float NormalDot = VertexNormal | CurNormal;
 				if ( DistSqrd < BestDistanceSquared - DistanceOverNormalThreshold )
 				{
 					BestVertex = Vertex;
diff --git a/Engine/Source/Runtime/Engine/Private/StaticMeshRender.cpp b/Engine/Source/Runtime/Engine/Private/StaticMeshRender.cpp
--- a/Engine/Source/Runtime/Engine/Private/StaticMeshRender.cpp
+++ b/Engine/Source/Runtime/Engine/Private/StaticMeshRender.cpp
@@ -99,7 +99,7 @@ FStaticMeshSceneProxy::FStaticMeshSceneProxy(UStaticMeshComponent* InComponent):
 #if WITH_EDITORONLY_DATA
 	, StreamingSectionData(InComponent->StreamingSectionData)
 	, StreamingDistanceMultiplier(InComponent->StreamingDistanceMultiplier)
-	, StreamingTexelFactor(0)
+	, StreamingTexelFactor(1.f)
 	, SectionIndexPreview(InComponent->SectionIndexPreview)
 #endif
 #if !(UE_BUILD_SHIPPING || UE_BUILD_TEST)
@@ -414,14 +414,16 @@ bool FStaticMeshSceneProxy::GetWireframeMeshElement(int32 LODIndex, int32 BatchI
 }
 
 #if WITH_EDITORONLY_DATA
-const FStreamingSectionBuildInfo* FStaticMeshSceneProxy::GetStreamingSectionData(float& OutDistanceMultiplier, int32 LODIndex, int32 ElementIndex) const
+const FStreamingSectionBuildInfo* FStaticMeshSceneProxy::GetStreamingSectionData(float& OutComponentExtraScale, float& OutMeshExtraScale, int32 LODIndex, int32 ElementIndex) const
 {
 	const bool bUseNewMetrics = CVarStreamingUseNewMetrics.GetValueOnRenderThread() != 0;
 
-	OutDistanceMultiplier = StreamingDistanceMultiplier;
+	OutComponentExtraScale = StreamingDistanceMultiplier;
 
 	if (!bUseNewMetrics)
 	{
+		OutMeshExtraScale = 1.f; // No extra scale as this scale is already taken into account in StreamingTexelFactor
+
 		// In this case the element is not in the build data.
 		static FStreamingSectionBuildInfo FallbackData;
 
@@ -429,6 +431,7 @@ const FStreamingSectionBuildInfo* FStaticMeshSceneProxy::GetStreamingSectionData
 		FallbackData.BoxExtent = GetBounds().BoxExtent;
 		for (int32 I = 0; I < FMaterialTexCoordBuildInfo::MAX_NUM_TEX_COORD; ++I)
 		{
+			// This fallback factor has already the mesh extra scale in it.
 			FallbackData.TexelFactors[I] = StreamingTexelFactor;
 		}
 		return &FallbackData;
@@ -437,6 +440,8 @@ const FStreamingSectionBuildInfo* FStaticMeshSceneProxy::GetStreamingSectionData
 	{
 		if (StreamingSectionData.IsValid())
 		{
+			OutMeshExtraScale = StaticMesh ? StaticMesh->StreamingDistanceMultiplier : 1.f;
+
 			for (const FStreamingSectionBuildInfo& SectionData : *StreamingSectionData)
 			{
 				if (SectionData.LODIndex == LODIndex && SectionData.ElementIndex == ElementIndex)
diff --git a/Engine/Source/Runtime/Engine/Private/Streaming/AsyncTextureStreaming.cpp b/Engine/Source/Runtime/Engine/Private/Streaming/AsyncTextureStreaming.cpp
--- a/Engine/Source/Runtime/Engine/Private/Streaming/AsyncTextureStreaming.cpp
+++ b/Engine/Source/Runtime/Engine/Private/Streaming/AsyncTextureStreaming.cpp
@@ -7,7 +7,32 @@ AsyncTextureStreaming.cpp: Definitions of classes used for texture streaming asy
 #include "EnginePrivate.h"
 #include "AsyncTextureStreaming.h"
 
-void FAsyncTextureStreaming::ClearRemovedTextureReferences()
+void FAsyncTextureStreamingData::Init(TArray<FStreamingViewInfo> InViewInfos, TArray<FLevelTextureManager>& LevelTextureManagers, FDynamicComponentTextureManager& DynamicComponentManager)
+{
+	ViewInfos = InViewInfos;
+
+	DynamicInstancesView = DynamicComponentManager.GetAsyncView();
+
+	StaticInstancesViews.Reset();
+	for (FLevelTextureManager& LevelManager : LevelTextureManagers)
+	{
+		StaticInstancesViews.Push(LevelManager.GetAsyncView());
+	}
+}
+
+void FAsyncTextureStreamingData::Update_Async()
+{
+	const bool bUseNewMetrics = CVarStreamingUseNewMetrics.GetValueOnAnyThread() > 0;
+	const float MaxEffectiveScreenSize = CVarStreamingScreenSizeEffectiveMax.GetValueOnAnyThread();
+
+	for (FTextureInstanceAsyncView& StaticInstancesView : StaticInstancesViews)
+	{
+		StaticInstancesView.Update_Async(ViewInfos, bUseNewMetrics, MaxEffectiveScreenSize);
+	}
+	DynamicInstancesView.Update_Async(ViewInfos, bUseNewMetrics, MaxEffectiveScreenSize);
+}
+
+void FAsyncTextureStreamingTask::ClearRemovedTextureReferences()
 {
 	// CleanUp Prioritized Textures, account for removed textures.
 	for (int32 TexPrioIndex = 0; TexPrioIndex < PrioritizedTextures.Num(); ++TexPrioIndex)
@@ -24,39 +49,33 @@ void FAsyncTextureStreaming::ClearRemovedTextureReferences()
 	}
 }
 
-void FAsyncTextureStreaming::DoWork()
+void FAsyncTextureStreamingTask::DoWork()
 {
-	DECLARE_SCOPE_CYCLE_COUNTER(TEXT("FAsyncTextureStreaming::DoWork"), STAT_AsyncTextureStreaming_DoWork, STATGROUP_StreamingDetails);
+	DECLARE_SCOPE_CYCLE_COUNTER(TEXT("FAsyncTextureStreamingTask::DoWork"), STAT_AsyncTextureStreaming_DoWork, STATGROUP_StreamingDetails);
 	PrioritizedTextures.Empty( StreamingManager.StreamingTextures.Num() );
 	PrioTexIndicesSortedByLoadOrder.Empty( StreamingManager.StreamingTextures.Num() );
 
 	const int32 SplitRequestSizeThreshold = CVarStreamingSplitRequestSizeThreshold.GetValueOnAnyThread();
 
-	const bool bUseNewMetrics = CVarStreamingUseNewMetrics.GetValueOnAnyThread() > 0;
-	const float MaxEffectiveScreenSize = CVarStreamingScreenSizeEffectiveMax.GetValueOnAnyThread();
-	StreamingManager.ThreadSettings.TextureBoundsVisibility->ComputeBoundsViewInfos(StreamingManager.ThreadSettings.ThreadViewInfos, bUseNewMetrics, MaxEffectiveScreenSize);
-
 	// Number of textures that want more mips.
 	ThreadStats.NumWantingTextures = 0;
 	ThreadStats.NumVisibleTexturesWithLowResolutions = 0;
 
+	StreamingData.Update_Async();
+
 	for ( int32 Index=0; Index < StreamingManager.StreamingTextures.Num() && !IsAborted(); ++Index )
 	{
 		FStreamingTexture& StreamingTexture = StreamingManager.StreamingTextures[ Index ];
 
-		StreamingTexture.bUsesStaticHeuristics = false;
-		StreamingTexture.bUsesDynamicHeuristics = false;
-		StreamingTexture.bUsesLastRenderHeuristics = false;
-		StreamingTexture.bUsesForcedHeuristics = false;
-		StreamingTexture.bUsesOrphanedHeuristics = false;
 		StreamingTexture.bHasSplitRequest  = false;
 		StreamingTexture.bIsLastSplitRequest = false;
+		StreamingTexture.bAsNeverStream = false;
 
-			// Figure out max number of miplevels allowed by LOD code.
-			StreamingManager.CalcMinMaxMips( StreamingTexture );
+		// Figure out max number of miplevels allowed by LOD code.
+		StreamingManager.CalcMinMaxMips( StreamingTexture );
 
-			// Determine how many mips this texture should have in memory.
-			StreamingManager.CalcWantedMips( StreamingTexture );
+		// Determine how many mips this texture should have in memory.
+		StreamingManager.CalcWantedMips( StreamingData, StreamingTexture, false, false );
 
 		ThreadStats.TotalResidentSize += StreamingTexture.GetSize( StreamingTexture.ResidentMips );
 		ThreadStats.TotalPossibleResidentSize += StreamingTexture.GetSize( FMath::Max<int32>(StreamingTexture.ResidentMips, StreamingTexture.RequestedMips) );
@@ -91,7 +110,7 @@ void FAsyncTextureStreaming::DoWork()
 			ThreadStats.NumWantingTextures++;
 		}
 
-		bool bTrackedTexture = TrackTextureEvent( &StreamingTexture, StreamingTexture.Texture, StreamingTexture.bForceFullyLoad, &StreamingManager );
+		bool bTrackedTexture = TrackTextureEvent( &StreamingTexture, StreamingTexture.Texture, StreamingTexture.bAsNeverStream, &StreamingManager );
 
 		// Add to sort list, if it wants to stream in or could potentially stream out.
 		if ( StreamingTexture.MaxAllowedMips > StreamingTexture.NumNonStreamingMips)
diff --git a/Engine/Source/Runtime/Engine/Private/Streaming/AsyncTextureStreaming.h b/Engine/Source/Runtime/Engine/Private/Streaming/AsyncTextureStreaming.h
--- a/Engine/Source/Runtime/Engine/Private/Streaming/AsyncTextureStreaming.h
+++ b/Engine/Source/Runtime/Engine/Private/Streaming/AsyncTextureStreaming.h
@@ -8,16 +8,46 @@ AsyncTextureStreaming.h: Definitions of classes used for texture streaming async
 
 #include "StreamingManagerTexture.h"
 
+
+/** Thread-safe helper struct for streaming information. */
+class FAsyncTextureStreamingData
+{
+public:
+
+	/** Set the data but do as little as possible, called from the game thread. */
+	void Init(TArray<FStreamingViewInfo> InViewInfos, TArray<FLevelTextureManager>& LevelTextureManagers, FDynamicComponentTextureManager& DynamicComponentManager);
+
+	/** Update everything internally so to allow calls to CalcWantedMips */
+	void Update_Async();
+
+	uint32 GetAllocatedSize() const { return ViewInfos.GetAllocatedSize() + StaticInstancesViews.GetAllocatedSize(); }
+
+	FORCEINLINE const FTextureInstanceAsyncView& GetDynamicInstancesView() const { return DynamicInstancesView; }
+	FORCEINLINE const TArray<FTextureInstanceAsyncView>& GetStaticInstancesViews() const { return StaticInstancesViews; }
+	FORCEINLINE const TArray<FStreamingViewInfo>& GetViewInfos() const { return ViewInfos; }
+
+private:
+
+	/** Cached from FStreamingManagerBase. */
+	TArray<FStreamingViewInfo> ViewInfos;
+	
+	FTextureInstanceAsyncView DynamicInstancesView;
+
+	/** Cached from each ULevel. */
+	TArray<FTextureInstanceAsyncView> StaticInstancesViews;
+};
+
+
 /*-----------------------------------------------------------------------------
 	Asynchronous texture streaming.
 -----------------------------------------------------------------------------*/
 
 /** Async work class for calculating priorities for all textures. */
 // this could implement a better abandon, but give how it is used, it does that anyway via the abort mechanism
-class FAsyncTextureStreaming : public FNonAbandonableTask
+class FAsyncTextureStreamingTask : public FNonAbandonableTask
 {
 public:
-	FAsyncTextureStreaming( FStreamingManagerTexture* InStreamingManager )
+	FAsyncTextureStreamingTask( FStreamingManagerTexture* InStreamingManager )
 	:	StreamingManager( *InStreamingManager )
 	,	ThreadContext( false, NULL, false )
 	,	bAbort( false )
@@ -118,14 +148,16 @@ public:
 		return ThreadStats;
 	}
 
+	FAsyncTextureStreamingData StreamingData;
+
 private:
-	friend class FAsyncTask<FAsyncTextureStreaming>;
+	friend class FAsyncTask<FAsyncTextureStreamingTask>;
 	/** Performs the async work. */
 	void DoWork();
 
 	FORCEINLINE TStatId GetStatId() const
 	{
-		RETURN_QUICK_DECLARE_CYCLE_STAT(FAsyncTextureStreaming, STATGROUP_ThreadPoolAsyncTasks);
+		RETURN_QUICK_DECLARE_CYCLE_STAT(FAsyncTextureStreamingTask, STATGROUP_ThreadPoolAsyncTasks);
 	}
 
 	/** Reference to the owning streaming manager, for accessing the thread-safe data. */
