Copying //UE4/Dev-Rendering to //UE4/Dev-Main (Source: //UE4/Dev-Rendering @ 2981742)

diff --git a/Engine/Source/Runtime/Engine/Private/Streaming/StreamingManagerTexture.cpp b/Engine/Source/Runtime/Engine/Private/Streaming/StreamingManagerTexture.cpp
--- a/Engine/Source/Runtime/Engine/Private/Streaming/StreamingManagerTexture.cpp
+++ b/Engine/Source/Runtime/Engine/Private/Streaming/StreamingManagerTexture.cpp
@@ -50,8 +50,8 @@ FStreamingManagerTexture::FStreamingManagerTexture()
 ,	bTriggerDumpTextureGroupStats( false )
 ,	bDetailedDumpTextureGroupStats( false )
 ,	bTriggerInvestigateTexture( false )
+,	MipBias(0)
 ,	AsyncWork( nullptr )
-,   DynamicComponentTextureManager( nullptr )
 ,	ProcessingStage( 0 )
 ,	NumTextureProcessingStages(5)
 ,	bUseDynamicStreaming( false )
@@ -173,26 +173,20 @@ FStreamingManagerTexture::FStreamingManagerTexture()
 	for ( int32 LODGroup=0; LODGroup < TEXTUREGROUP_MAX; ++LODGroup )
 	{
 		const FTextureLODGroup& TexGroup = UDeviceProfileManager::Get().GetActiveProfile()->GetTextureLODSettings()->GetTextureLODGroup(TextureGroup(LODGroup));
-		ThreadSettings.NumStreamedMips[LODGroup] = TexGroup.NumStreamedMips;
+		NumStreamedMips[LODGroup] = TexGroup.NumStreamedMips;
 	}
 
 	// setup the streaming resource flush function pointer
 	GFlushStreamingFunc = &FlushResourceStreaming;
 
 	ProcessingStage = 0;
-	AsyncWork = new FAsyncTask<FAsyncTextureStreaming>(this);
-
-	// We use a dynamic allocation here to reduce header dependencies.
-	DynamicComponentTextureManager = new FDynamicComponentTextureManager();
-	ThreadSettings.TextureBoundsVisibility = new FTextureBoundsVisibility();
+	AsyncWork = new FAsyncTask<FAsyncTextureStreamingTask>(this);
 }
 
 FStreamingManagerTexture::~FStreamingManagerTexture()
 {
 	AsyncWork->EnsureCompletion();
 	delete AsyncWork;
-	delete DynamicComponentTextureManager;
-	delete ThreadSettings.TextureBoundsVisibility;
 }
 
 /**
@@ -363,116 +357,70 @@ bool FStreamingManagerTexture::StreamOutTextureData( int64 RequiredMemorySize )
 	return bSucceeded;
 }
 
-/**
- * Adds new textures and level data on the gamethread (while the worker thread isn't active).
- */
-void FStreamingManagerTexture::UpdateThreadData()
+void FStreamingManagerTexture::IncrementalUpdate( float Percentage )
 {
-	// Add new textures.
-	StreamingTextures.Reserve( StreamingTextures.Num() + PendingStreamingTextures.Num() );
-	for ( int32 TextureIndex=0; TextureIndex < PendingStreamingTextures.Num(); ++TextureIndex )
-	{
-		UTexture2D* Texture = PendingStreamingTextures[ TextureIndex ];
-		FStreamingTexture* StreamingTexture = new (StreamingTextures) FStreamingTexture( Texture );
-		StreamingTexture->SetStreamingIndex( StreamingTextures.Num() - 1 );
-	}
-	PendingStreamingTextures.Empty();
+	QUICK_SCOPE_CYCLE_COUNTER(STAT_FStreamingManagerTexture_IncrementalUpdate);
 
-	// Remove old levels. Note: Don't try to access the ULevel object, it may have been deleted already!
-	for ( int32 LevelIndex=0; LevelIndex < ThreadSettings.LevelData.Num(); ++LevelIndex )
+	FRemovedTextureArray RemovedTextures;
+
+#if !(UE_BUILD_SHIPPING || UE_BUILD_TEST)
+	// This code is to trigger an update to the a streaming data update when the console variable is changed.
+	static int32 LastMetricsUsed = CVarStreamingUseNewMetrics.GetValueOnGameThread();
+	double Time0 = FPlatformTime::Seconds(), Time1 = 0;
+	if (CVarStreamingUseNewMetrics.GetValueOnGameThread() != LastMetricsUsed)
 	{
-		FLevelData& LevelData = ThreadSettings.LevelData[ LevelIndex ];
-		if ( LevelData.Value.bRemove )
+		TArray<ULevel*> Levels;
+		for (FLevelTextureManager& LevelManager : LevelTextureManagers)
 		{
-			ThreadSettings.LevelData.RemoveAtSwap( LevelIndex-- );
+			Levels.Push(LevelManager.GetLevel());
 		}
+		for (ULevel* Level : Levels)
+		{
+			AddLevel(Level);
+		}
+		Time1 = FPlatformTime::Seconds();
+		UE_LOG(LogContentStreaming, Log,  TEXT("Removing levels took %.2f ms."),  (float)(Time1 - Time0) * 1000.f);
 	}
+#endif
 
+	for (FLevelTextureManager& LevelManager : LevelTextureManagers)
 	{
-		STAT( double StartTime = FPlatformTime::Seconds() );
-
-		FRemovedTextureArray RemovedTextures;
-		DynamicComponentTextureManager->IncrementalTick(RemovedTextures, 1.f); // Perform a full update at this point.
-		SetTexturesRemovedTimestamp(RemovedTextures);
-
-		// Duplicate the data to be used on the worker thread
-		ThreadSettings.TextureBoundsVisibility->QuickSet(DynamicComponentTextureManager->GetTextureInstances());
-
-		STAT( GStreamingDynamicPrimitivesTime += FPlatformTime::Seconds() - StartTime );
-		}
+		LevelManager.IncrementalUpdate(DynamicComponentManager, RemovedTextures, Percentage, bUseDynamicStreaming); // Complete the incremental update.
+	}
 
-	// Add new levels.
-	for ( int32 LevelIndex=0; LevelIndex < PendingLevels.Num(); ++LevelIndex )
+#if !(UE_BUILD_SHIPPING || UE_BUILD_TEST)
+	if (CVarStreamingUseNewMetrics.GetValueOnGameThread() != LastMetricsUsed)
 	{
-		ULevel* Level = PendingLevels[ LevelIndex ];
-		FLevelData* LevelData = new (ThreadSettings.LevelData) FLevelData( Level );
-		FThreadLevelData& ThreadLevelData = LevelData->Value;
-
-		// Increase the the force-fully-load refcount.
-		for ( TMap<UTexture2D*,bool>::TIterator It(Level->ForceStreamTextures); It; ++It )
-		{
-			UTexture2D* Texture2D = It.Key();
-			if ( Texture2D && IsStreamingTexture( Texture2D ) )
-			{
-				FStreamingTexture& StreamingTexture = GetStreamingTexture( Texture2D );
-				check( StreamingTexture.ForceLoadRefCount >= 0 );
-				StreamingTexture.ForceLoadRefCount++;
-			}
-		}
-
-		// Copy all texture instances into ThreadLevelData.
-		//Note: ULevel could save it in this form, except it's harder to debug.
-		ThreadLevelData.ThreadTextureInstances.Empty( Level->TextureToInstancesMap.Num() );
-		for ( TMap<UTexture2D*,TArray<FStreamableTextureInstance> >::TIterator It(Level->TextureToInstancesMap); It; ++It )
-		{
-			// Convert to SIMD-friendly form
-			const UTexture2D* Texture2D = It.Key();
-			if ( Texture2D )
-			{
-				TArray<FStreamableTextureInstance>& TextureInstances = It.Value();
-				TArray<FStreamableTextureInstance4>& TextureInstances4 = ThreadLevelData.ThreadTextureInstances.Add( Texture2D, TArray<FStreamableTextureInstance4>() );
-				TextureInstances4.Reserve( (TextureInstances.Num() + 3) / 4 );
+		LastMetricsUsed = CVarStreamingUseNewMetrics.GetValueOnGameThread();
+		UE_LOG(LogContentStreaming, Log,  TEXT("Adding levels took %.2f ms."),  (float)(FPlatformTime::Seconds() - Time1) * 1000.f);
+	}
+#endif
 
-				for ( int32 InstanceIndex=0; InstanceIndex < TextureInstances.Num(); ++InstanceIndex )
-				{
-					if ( (InstanceIndex & 3) == 0 )
-					{
-						TextureInstances4.Add(FStreamableTextureInstance4());
-					}
-					const FStreamableTextureInstance& Instance = TextureInstances[ InstanceIndex ];
-					FStreamableTextureInstance4& Instance4 = TextureInstances4[ InstanceIndex/4 ];
-					Instance4.BoundsOriginX[ InstanceIndex & 3 ] = Instance.Bounds.Origin.X;
-					Instance4.BoundsOriginY[ InstanceIndex & 3 ] = Instance.Bounds.Origin.Y;
-					Instance4.BoundsOriginZ[ InstanceIndex & 3 ] = Instance.Bounds.Origin.Z;
-					Instance4.BoxExtentSizeX[ InstanceIndex & 3 ] = Instance.Bounds.BoxExtent.X;
-					Instance4.BoxExtentSizeY[ InstanceIndex & 3 ] = Instance.Bounds.BoxExtent.Y;
-					Instance4.BoxExtentSizeZ[ InstanceIndex & 3 ] = Instance.Bounds.BoxExtent.Z;
-					Instance4.BoundingSphereRadius[ InstanceIndex & 3 ] = Instance.Bounds.SphereRadius;
-					Instance4.MinDistanceSq[ InstanceIndex & 3 ] = FMath::Square(FMath::Max(0.f, Instance.MinDistance - RangePrefetchDistance));
-					Instance4.MaxDistanceSq[ InstanceIndex & 3 ] = Instance.MaxDistance == MAX_FLT ? MAX_FLT : FMath::Square(Instance.MaxDistance + RangePrefetchDistance);
-					Instance4.TexelFactor[ InstanceIndex & 3 ] = Instance.TexelFactor;
-					ensureMsgf(!FMath::IsNearlyZero(Instance.TexelFactor), TEXT("Texture instance %d has a texel factor of zero: %s"), TextureInstances.Num(), *Texture2D->GetPathName());
-				}
+	DynamicComponentManager.IncrementalUpdate(RemovedTextures, Percentage); // Complete the incremental update.
 
-				// Note: The old instance array must remain, in case a level is added/removed and then added again.
-			}
-		}
-	}
-	PendingLevels.Empty();
+	SetTexturesRemovedTimestamp(RemovedTextures);
+}
 
-	// If any views specified an actor to boost, do it now
-	for (int32 i = 0; i < CurrentViewInfos.Num(); ++i)
+/**
+ * Adds new textures and level data on the gamethread (while the worker thread isn't active).
+ */
+void FStreamingManagerTexture::UpdateThreadData()
+{
+	// Add new textures.
+	StreamingTextures.Reserve( StreamingTextures.Num() + PendingStreamingTextures.Num() );
+	for ( int32 TextureIndex=0; TextureIndex < PendingStreamingTextures.Num(); ++TextureIndex )
 	{
-		if (CurrentViewInfos[i].ActorToBoost.IsValid())
-		{
-			BoostTextures(CurrentViewInfos[i].ActorToBoost.Get(false), BoostPlayerTextures);
-		}
+		UTexture2D* Texture = PendingStreamingTextures[ TextureIndex ];
+		FStreamingTexture* StreamingTexture = new (StreamingTextures) FStreamingTexture( Texture );
+		StreamingTexture->SetStreamingIndex( StreamingTextures.Num() - 1 );
 	}
+	PendingStreamingTextures.Empty();
 
-	// Cache the view information.
-	ThreadSettings.ThreadViewInfos = CurrentViewInfos;
-	
-	ThreadSettings.MipBias = FMath::Max(CVarStreamingMipBias.GetValueOnGameThread(), 0.0f);
+	// Fully complete all pending update.
+	IncrementalUpdate(1.f);
+	MipBias = FMath::Max(CVarStreamingMipBias.GetValueOnGameThread(), 0.0f);
+	// Update the thread data.
+	AsyncWork->GetTask().StreamingData.Init(CurrentViewInfos, LevelTextureManagers, DynamicComponentManager);
 }
 
 /**
@@ -498,11 +446,10 @@ void FStreamingManagerTexture::BoostTextures( AActor* Actor, float BoostFactor )
 				Primitive->GetUsedTextures( Textures, EMaterialQualityLevel::Num );
 				for ( int32 TextureIndex=0; TextureIndex < Textures.Num(); ++TextureIndex )
 				{
-					UTexture2D* Texture2D = Cast<UTexture2D>( Textures[ TextureIndex ] );
-					if ( Texture2D && IsManagedStreamingTexture(Texture2D) )
+					FStreamingTexture* StreamingTexture = GetStreamingTexture( Cast<UTexture2D>( Textures[ TextureIndex ] ) );
+					if ( StreamingTexture )
 					{
-						FStreamingTexture& StreamingTexture = GetStreamingTexture( Texture2D );
-						StreamingTexture.BoostFactor = FMath::Max( StreamingTexture.BoostFactor, BoostFactor );
+						StreamingTexture->BoostFactor = FMath::Max( StreamingTexture->BoostFactor, BoostFactor );
 					}
 				}
 			}
@@ -511,76 +458,32 @@ void FStreamingManagerTexture::BoostTextures( AActor* Actor, float BoostFactor )
 }
 
 /** Adds a ULevel to the streaming manager. */
-void FStreamingManagerTexture::AddPreparedLevel( ULevel* Level )
+void FStreamingManagerTexture::AddLevel( ULevel* Level )
 {
-	PendingLevels.AddUnique( Level );
+	// Remove any data related to this level first.
+	RemoveLevel(Level);
 
-	if ( bUseDynamicStreaming )
-	{
-		// Add all dynamic primitives from the ULevel.
-		for ( TMap<TWeakObjectPtr<UPrimitiveComponent>,TArray<FDynamicTextureInstance> >::TIterator It(Level->DynamicTextureInstances); It; ++It )
-		{
-			UPrimitiveComponent* Primitive = It.Key().Get();
-			TArray<FDynamicTextureInstance>& TextureInstances = It.Value();
-			NotifyPrimitiveAttached( Primitive, DPT_Level );
-
-			//@TODO: This variable is deprecated, remove from serialization.
-			TextureInstances.Empty();
-		}
-	}
+	// If the level was not already there, add a new entry.
+	new (LevelTextureManagers) FLevelTextureManager(Level);
 }
 
 /** Removes a ULevel from the streaming manager. */
 void FStreamingManagerTexture::RemoveLevel( ULevel* Level )
 {
-	// Remove from pending new levels, if it's been added very recently (this frame)...
-	PendingLevels.Remove( Level );
+	const double CurrentTime = FApp::GetCurrentTime();
 
-	// Mark the level for removal (will take effect when we're syncing threads).
-	FLevelData* LevelData = ThreadSettings.LevelData.FindByKey( Level );
-	if ( LevelData && !LevelData->Value.bRemove )
+	for (int32 Index = 0; Index < LevelTextureManagers.Num(); ++Index)
 	{
-		FThreadLevelData& ThreadLevelData = LevelData->Value;
-		ThreadLevelData.bRemove = true;
+		FLevelTextureManager& LevelManager = LevelTextureManagers[Index];
+		if (LevelManager.GetLevel() != Level) continue;
 
-		// Mark all textures with a timestamp. They're about to lose their location-based heuristic and we don't want them to
-		// start using LastRenderTime heuristic for a few seconds until they are garbage collected!
-		for( TMap<UTexture2D*,TArray<FStreamableTextureInstance> >::TIterator It(Level->TextureToInstancesMap); It; ++It )
-		{
-			const UTexture2D* Texture2D = It.Key();
-			if ( Texture2D && IsManagedStreamingTexture( Texture2D ) )
-			{
-				FStreamingTexture& StreamingTexture = GetStreamingTexture( Texture2D );
-				StreamingTexture.InstanceRemovedTimestamp = FApp::GetCurrentTime();
-			}
-		}
-
-		// Decrease the the force-fully-load refcount. 
-		// Note: Only update ForceLoadRefCount if the level was in the ThreadSettings and therefore bumped ForceLoadRefCount earlier.
-		// Note: We do this immediately, since ULevel may be being deleted.
-		for ( TMap<UTexture2D*,bool>::TIterator It(Level->ForceStreamTextures); It; ++It )
-		{
-			UTexture2D* Texture2D = It.Key();
-			if ( Texture2D && IsManagedStreamingTexture( Texture2D ) )
-			{
-				FStreamingTexture& StreamingTexture = GetStreamingTexture( Texture2D );
-				// Note: This refcount should always be positive, but for some reason it may not be when BuildStreamingData() is called in the editor...
-				if ( StreamingTexture.ForceLoadRefCount > 0 )
-				{
-					StreamingTexture.ForceLoadRefCount--;
-				}
-			}
-		}
-	}
+		FRemovedTextureArray RemovedTextures;
+		LevelManager.Remove(DynamicComponentManager, RemovedTextures);
+		SetTexturesRemovedTimestamp(RemovedTextures);
 
-	if ( bUseDynamicStreaming )
-	{
-		// Mark all dynamic primitives from the ULevel (they'll be removed completely next sync).
-		for ( TMap<TWeakObjectPtr<UPrimitiveComponent>,TArray<FDynamicTextureInstance> >::TIterator It(Level->DynamicTextureInstances); It; ++It )
-		{
-			UPrimitiveComponent* Primitive = It.Key().GetEvenIfUnreachable();
-			NotifyPrimitiveDetached( Primitive );
-		}
+		// Remove the level entry. The async task view will still be valid as it uses a shared ptr.
+		LevelTextureManagers.RemoveAtSwap(Index);
+		break;
 	}
 }
 
@@ -682,7 +585,7 @@ void FStreamingManagerTexture::NotifyPrimitiveAttached( const UPrimitiveComponen
 		UE_LOG(LogContentStreaming, Log, TEXT("NotifyPrimitiveAttached(0x%08x \"%s\"), IsRegistered=%d"), SIZE_T(Primitive), *Primitive->GetReadableName(), Primitive->IsRegistered());
 #endif
 		FRemovedTextureArray RemovedTextures;
-		DynamicComponentTextureManager->Add(Primitive, DynamicType, RemovedTextures);
+		DynamicComponentManager.Add(Primitive, DynamicType, RemovedTextures);
 		SetTexturesRemovedTimestamp(RemovedTextures);
 	}
 }
@@ -699,7 +602,13 @@ void FStreamingManagerTexture::NotifyPrimitiveDetached( const UPrimitiveComponen
 		UE_LOG(LogContentStreaming, Log, TEXT("NotifyPrimitiveDetached(0x%08x \"%s\"), IsRegistered=%d"), SIZE_T(Primitive), *Primitive->GetReadableName(), Primitive->IsRegistered());
 #endif
 		FRemovedTextureArray RemovedTextures;
-		DynamicComponentTextureManager->Remove(Primitive, RemovedTextures);
+
+		for (FLevelTextureManager& LevelManager: LevelTextureManagers)
+		{
+			LevelManager.RemoveComponentReferences(Primitive);
+		}
+
+		DynamicComponentManager.Remove(Primitive, RemovedTextures);
 		SetTexturesRemovedTimestamp(RemovedTextures);
 	}
 }
@@ -712,14 +621,13 @@ void FStreamingManagerTexture::NotifyPrimitiveDetached( const UPrimitiveComponen
 */
 void FStreamingManagerTexture::SetTexturesRemovedTimestamp(const FRemovedTextureArray& RemovedTextures)
 {
-	double CurrentTime = FApp::GetCurrentTime();
+	const double CurrentTime = FApp::GetCurrentTime();
 	for ( int32 TextureIndex=0; TextureIndex < RemovedTextures.Num(); ++TextureIndex )
 	{
-		const UTexture2D* Texture2D = RemovedTextures[TextureIndex];
-		if ( Texture2D && IsManagedStreamingTexture(Texture2D) )
+		FStreamingTexture* StreamingTexture = GetStreamingTexture( RemovedTextures[TextureIndex] );
+		if ( StreamingTexture )
 		{
-			FStreamingTexture& StreamingTexture = GetStreamingTexture( Texture2D );
-			StreamingTexture.InstanceRemovedTimestamp = CurrentTime;
+			StreamingTexture->InstanceRemovedTimestamp = CurrentTime;
 		}
 	}
 }
@@ -732,7 +640,7 @@ void FStreamingManagerTexture::SetTexturesRemovedTimestamp(const FRemovedTexture
 void FStreamingManagerTexture::NotifyPrimitiveUpdated( const UPrimitiveComponent* Primitive )
 {
 	FRemovedTextureArray RemovedTextures;
-	DynamicComponentTextureManager->Update(Primitive, RemovedTextures);
+	DynamicComponentManager.Update(Primitive, RemovedTextures);
 	SetTexturesRemovedTimestamp(RemovedTextures);
 }
 
@@ -745,23 +653,21 @@ void FStreamingManagerTexture::NotifyPrimitiveUpdated( const UPrimitiveComponent
  */
 void FStreamingManagerTexture::NotifyTimedPrimitiveAttached( const UPrimitiveComponent* Primitive, EDynamicPrimitiveType DynamicType )
 {
-	if ( Primitive  && Primitive->IsRegistered())
+	const double CurrentTime = FApp::GetCurrentTime();
+	if ( Primitive  && Primitive->IsRegistered() )
 	{
 		FStreamingTextureLevelContext LevelContext;
 		TArray<FStreamingTexturePrimitiveInfo> TextureInstanceInfos;
 		Primitive->GetStreamingTextureInfoWithNULLRemoval( LevelContext, TextureInstanceInfos );
 
-		TArray<FSpawnedTextureInstance>* TextureInstances = NULL;
-		for ( int32 InstanceIndex=0; InstanceIndex < TextureInstanceInfos.Num(); ++InstanceIndex )
+		for ( const FStreamingTexturePrimitiveInfo& Info :  TextureInstanceInfos )
 		{
-			FStreamingTexturePrimitiveInfo& Info = TextureInstanceInfos[InstanceIndex];
-			UTexture2D* Texture2D = Cast<UTexture2D>(Info.Texture);
-			if ( Texture2D && IsManagedStreamingTexture(Texture2D) )
+			FStreamingTexture* StreamingTexture = GetStreamingTexture( Info.Texture );
+			if ( StreamingTexture )
 			{
 				// Note: Doesn't have to be cycle-perfect for thread safety.
-				FStreamingTexture& StreamingTexture = GetStreamingTexture( Texture2D );
-				StreamingTexture.LastRenderTimeRefCount++;
-				StreamingTexture.LastRenderTimeRefCountTimestamp = FApp::GetCurrentTime();
+				StreamingTexture->LastRenderTimeRefCount++;
+				StreamingTexture->LastRenderTimeRefCountTimestamp = CurrentTime;
 			}
 		}
 	}
@@ -776,26 +682,20 @@ void FStreamingManagerTexture::NotifyTimedPrimitiveAttached( const UPrimitiveCom
  */
 void FStreamingManagerTexture::NotifyTimedPrimitiveDetached( const UPrimitiveComponent* Primitive )
 {
+	const double CurrentTime = FApp::GetCurrentTime();
 	if ( Primitive )
 	{
 		FStreamingTextureLevelContext LevelContext;
 		TArray<FStreamingTexturePrimitiveInfo> TextureInstanceInfos;
 		Primitive->GetStreamingTextureInfoWithNULLRemoval( LevelContext, TextureInstanceInfos );
 
-		TArray<FSpawnedTextureInstance>* TextureInstances = NULL;
-		for ( int32 InstanceIndex=0; InstanceIndex < TextureInstanceInfos.Num(); ++InstanceIndex )
+		for ( const FStreamingTexturePrimitiveInfo& Info :  TextureInstanceInfos )
 		{
-			FStreamingTexturePrimitiveInfo& Info = TextureInstanceInfos[InstanceIndex];
-			UTexture2D* Texture2D = Cast<UTexture2D>(Info.Texture);
-			if ( Texture2D && IsManagedStreamingTexture(Texture2D) )
+			FStreamingTexture* StreamingTexture = GetStreamingTexture( Info.Texture );
+			if ( StreamingTexture && StreamingTexture->LastRenderTimeRefCount > 0 )
 			{
-				// Note: Doesn't have to be cycle-perfect for thread safety.
-				FStreamingTexture& StreamingTexture = GetStreamingTexture( Texture2D );
-				if ( StreamingTexture.LastRenderTimeRefCount > 0 )
-				{
-					StreamingTexture.LastRenderTimeRefCount--;
-					StreamingTexture.LastRenderTimeRefCountTimestamp = FApp::GetCurrentTime();
-				}
+				StreamingTexture->LastRenderTimeRefCount--;
+				StreamingTexture->LastRenderTimeRefCountTimestamp = CurrentTime;
 			}
 		}
 	}
@@ -804,15 +704,16 @@ void FStreamingManagerTexture::NotifyTimedPrimitiveDetached( const UPrimitiveCom
 /**
  * Returns the corresponding FStreamingTexture for a UTexture2D.
  */
-FStreamingTexture& FStreamingManagerTexture::GetStreamingTexture( const UTexture2D* Texture2D )
+FStreamingTexture* FStreamingManagerTexture::GetStreamingTexture( const UTexture2D* Texture2D )
 {
-	return StreamingTextures[ Texture2D->StreamingIndex ];
-}
-
-/** Returns true if this is a streaming texture that is managed by the streaming manager. */
-bool FStreamingManagerTexture::IsManagedStreamingTexture( const UTexture2D* Texture2D )
-{
-	return Texture2D->StreamingIndex >= 0 && Texture2D->StreamingIndex < StreamingTextures.Num() && IsStreamingTexture( Texture2D );
+	if (Texture2D && StreamingTextures.IsValidIndex(Texture2D->StreamingIndex))
+	{
+		return &StreamingTextures[Texture2D->StreamingIndex];
+	}
+	else
+	{
+		return nullptr;
+	}
 }
 
 /**
@@ -995,7 +896,7 @@ void FStreamingManagerTexture::StreamTextures( bool bProcessEverything )
 	const TArray<FTexturePriority>& PrioritizedTextures = AsyncWork->GetTask().GetPrioritizedTextures();
 	const TArray<int32>& PrioTexIndicesSortedByLoadOrder = AsyncWork->GetTask().GetPrioTexIndicesSortedByLoadOrder();
 
-	FAsyncTextureStreaming::FAsyncStats ThreadStats = AsyncWork->GetTask().GetStats();
+	FAsyncTextureStreamingTask::FAsyncStats ThreadStats = AsyncWork->GetTask().GetStats();
 	Context.AddStats( AsyncWork->GetTask().GetContext() );
 
 	FIOSystem::Get().FlushLog();
@@ -1038,7 +939,7 @@ void FStreamingManagerTexture::StreamTextures( bool bProcessEverything )
 			int32 CurrDroppedMips	= TextureStat.NumMips - TextureStat.ResidentMips;
 			int32 WantedDroppedMips	= TextureStat.NumMips - TextureStat.WantedMips;
 			int32 MostDroppedMips	= TextureStat.NumMips - TextureStat.MostResidentMips;
-			TextureStatsReport.Add( FString::Printf( TEXT(",%.3f,%ix%i,%ix%i,%ix%i,%ix%i,%i,%i,%i,%i,%s,%3f sec,%.1f,%s"),
+			TextureStatsReport.Add( FString::Printf( TEXT(",%.3f,%ix%i,%ix%i,%ix%i,%ix%i,%i,%i,%i,%i,%3f sec,%.1f,%s"),
 				TextureStat.Priority,
 				FMath::Max(TextureStat.SizeX >> CurrDroppedMips, 1),
 				FMath::Max(TextureStat.SizeY >> CurrDroppedMips, 1),
@@ -1052,7 +953,6 @@ void FStreamingManagerTexture::StreamTextures( bool bProcessEverything )
 				TextureStat.WantedSize/1024,
 				TextureStat.MaxSize/1024,
 				TextureStat.MostResidentSize/1024,
-				GStreamTypeNames[TextureStat.StreamType],
 				TextureStat.LastRenderTime,
 				TextureStat.BoostFactor,
 				*TextureStat.TextureName
@@ -1442,9 +1342,7 @@ void FStreamingManagerTexture::UpdateResourceStreaming( float DeltaTime, bool bP
 	}
 	else
 	{
-		FRemovedTextureArray RemovedTextures;
-		DynamicComponentTextureManager->IncrementalTick(RemovedTextures, 1.f / (float)FMath::Max(NumTextureProcessingStages - 1, 1)); // -1 since we don't want to do anything at stage 0.
-		SetTexturesRemovedTimestamp(RemovedTextures);
+		IncrementalUpdate(1.f / (float)FMath::Max(NumTextureProcessingStages - 1, 1)); // -1 since we don't want to do anything at stage 0.)
 	}
 
 	// Non-threaded data collection.
@@ -1555,7 +1453,7 @@ int32 FStreamingManagerTexture::BlockTillAllRequestsFinished( float TimeLimit /*
 				{
 					PendingTextures[ 1 - CurrentArray ].Add( TextureIndex );
 				}
-				TrackTextureEvent( &StreamingTexture, StreamingTexture.Texture, StreamingTexture.bForceFullyLoad, this );
+				TrackTextureEvent( &StreamingTexture, StreamingTexture.Texture, StreamingTexture.bAsNeverStream, this );
 			}
 		}
 
@@ -1598,30 +1496,6 @@ int32 FStreamingManagerTexture::BlockTillAllRequestsFinished( float TimeLimit /*
 	return NumPendingUpdates;
 }
 
-/**
- * Adds a textures streaming handler to the array of handlers used to determine which
- * miplevels need to be streamed in/ out.
- *
- * @param TextureStreamingHandler	Handler to add
- */
-void FStreamingManagerTexture::AddTextureStreamingHandler( FStreamingHandlerTextureBase* TextureStreamingHandler )
-{
-	AsyncWork->EnsureCompletion();
-	TextureStreamingHandlers.Add( TextureStreamingHandler );
-}
-
-/**
- * Removes a textures streaming handler from the array of handlers used to determine which
- * miplevels need to be streamed in/ out.
- *
- * @param TextureStreamingHandler	Handler to remove
- */
-void FStreamingManagerTexture::RemoveTextureStreamingHandler( FStreamingHandlerTextureBase* TextureStreamingHandler )
-{
-	AsyncWork->EnsureCompletion();
-	TextureStreamingHandlers.Remove( TextureStreamingHandler );
-}
-
 /**
  * Calculates the minimum and maximum number of mip-levels for a streaming texture.
  */
@@ -1636,33 +1510,6 @@ void FStreamingManagerTexture::CalcMinMaxMips( FStreamingTexture& StreamingTextu
 		TextureLODBias = FMath::Max( TextureLODBias - StreamingTexture.NumCinematicMipLevels, 0 );
 	}
 
-	// We only stream in skybox textures as you are always within the bounding box and those textures
-	// tend to be huge. A streaming fudge factor fluctuating around 1 will cause them to be streamed in
-	// and out, making it likely for the allocator to fragment.
-	if( StreamingTexture.LODGroup == TEXTUREGROUP_Skybox )
-	{
-		StreamingTexture.bForceFullyLoad = true;
-	}
-
-	// Don't stream in all referenced textures but rather only those that have been rendered in the last 5 minutes if
-	// we only stream in textures. This means you still might see texture popping, but the option is designed to avoid
-	// hitching due to CPU overhead, which is still taken care off by the 5 minute rule.
-	static auto CVarOnlyStreamInTextures = IConsoleManager::Get().FindTConsoleVariableDataInt(TEXT("r.OnlyStreamInTextures"));
-	if( CVarOnlyStreamInTextures->GetValueOnAnyThread() != 0 )
-	{
-		float SecondsSinceLastRender = StreamingTexture.LastRenderTime;
-		if( SecondsSinceLastRender < 300 )
-		{
-			StreamingTexture.bForceFullyLoad = true;
-		}
-	}
-
-	const int32 HLODStategy = CVarStreamingHLODStrategy.GetValueOnAnyThread();
-	if (StreamingTexture.bIsHLODTexture && HLODStategy == 2) // disable streaming
-	{
-		StreamingTexture.bForceFullyLoad = true;
-	}
-
 	// Calculate the min/max allowed mips.
 	UTexture2D::CalcAllowedMips(
 		StreamingTexture.MipCount,
@@ -1672,25 +1519,20 @@ void FStreamingManagerTexture::CalcMinMaxMips( FStreamingTexture& StreamingTextu
 		StreamingTexture.MaxAllowedMips );
 
 	// Take the reduced texture pool in to account.
-	StreamingTexture.MaxAllowedOptimalMips = StreamingTexture.MaxAllowedMips;
 	if ( GIsOperatingWithReducedTexturePool )
 	{				
 		int32 MaxTextureMipCount = FMath::Max( GMaxTextureMipCount - 2, 0 );
 		StreamingTexture.MaxAllowedMips = FMath::Min( StreamingTexture.MaxAllowedMips, MaxTextureMipCount );
 	}
 
-	// Should this texture be fully streamed in?
-	if ( StreamingTexture.bForceFullyLoad )
-	{
-		StreamingTexture.MinAllowedMips = StreamingTexture.MaxAllowedMips;
-	}
 	// Check if the texture LOD group restricts the number of streaming mip-levels (in absolute terms).
-	else if ( ThreadSettings.NumStreamedMips[StreamingTexture.LODGroup] >= 0 )
+	else if ( NumStreamedMips[StreamingTexture.LODGroup] >= 0 )
 	{
-		StreamingTexture.MinAllowedMips = FMath::Clamp( StreamingTexture.MipCount - ThreadSettings.NumStreamedMips[StreamingTexture.LODGroup], StreamingTexture.MinAllowedMips, StreamingTexture.MaxAllowedMips );
+		StreamingTexture.MinAllowedMips = FMath::Clamp( StreamingTexture.MipCount - NumStreamedMips[StreamingTexture.LODGroup], StreamingTexture.MinAllowedMips, StreamingTexture.MaxAllowedMips );
 	}
 
-	if ( StreamingTexture.bIsHLODTexture && HLODStategy == 1 && StreamingTexture.MaxAllowedMips > StreamingTexture.MinAllowedMips ) // stream only mip 0
+	const int32 HLODStategy = CVarStreamingHLODStrategy.GetValueOnAnyThread();
+	if ( StreamingTexture.LODGroup == TEXTUREGROUP_HierarchicalLOD && HLODStategy == 1 && StreamingTexture.MaxAllowedMips > StreamingTexture.MinAllowedMips ) // stream only mip 0
 	{
 		StreamingTexture.MinAllowedMips = StreamingTexture.MaxAllowedMips - 1;
 	}
@@ -1707,7 +1549,7 @@ void FStreamingManagerTexture::UpdateFrameStats( FStreamingContext& Context, FSt
 #if STATS_FAST || STATS
 	STAT(int64 ResidentSize = StreamingTexture.GetSize(StreamingTexture.ResidentMips));
 	STAT(Context.ThisFrameTotalStreamingTexturesSize += ResidentSize);
-	STAT(int32 PerfectWantedMips = FMath::Clamp(StreamingTexture.PerfectWantedMips, StreamingTexture.MinAllowedMips, StreamingTexture.MaxAllowedOptimalMips) );
+	STAT(int32 PerfectWantedMips = FMath::Clamp(StreamingTexture.PerfectWantedMips, StreamingTexture.MinAllowedMips, StreamingTexture.MaxAllowedMips) );
 	STAT(int64 PerfectWantedSize = StreamingTexture.GetSize( PerfectWantedMips ));
 	STAT(int64 MostResidentSize = StreamingTexture.GetSize( StreamingTexture.MostResidentMips ) );
 	STAT(Context.ThisFrameOptimalWantedSize += PerfectWantedSize );
@@ -1715,147 +1557,106 @@ void FStreamingManagerTexture::UpdateFrameStats( FStreamingContext& Context, FSt
 
 #if STATS
 	int64 PotentialSize = StreamingTexture.GetSize(StreamingTexture.MaxAllowedMips);
-	ETextureStreamingType StreamType = StreamingTexture.GetStreamingType();
-	switch ( StreamType )
-	{
-		case StreamType_Forced:
-			Context.ThisFrameTotalForcedHeuristicSize += PerfectWantedSize;
-			break;
-		case StreamType_LastRenderTime:
-			Context.ThisFrameTotalLastRenderHeuristicSize += PerfectWantedSize;
-			break;
-		case StreamType_Dynamic:
-			Context.ThisFrameTotalDynamicTextureHeuristicSize += PerfectWantedSize;
-			break;
-		case StreamType_Static:
-			Context.ThisFrameTotalStaticTextureHeuristicSize += PerfectWantedSize;
-			break;
-		case StreamType_Orphaned:
-		case StreamType_Other:
-			break;
-	}
+
 	if ( Context.bCollectTextureStats )
 	{
 		FString TextureName = StreamingTexture.Texture->GetFullName();
 		if ( CollectTextureStatsName.Len() == 0 || TextureName.Contains( CollectTextureStatsName) )
 		{
-			new (Context.TextureStats) FTextureStreamingStats( StreamingTexture.Texture, StreamType, StreamingTexture.ResidentMips, PerfectWantedMips, StreamingTexture.MostResidentMips, ResidentSize, PerfectWantedSize, PotentialSize, MostResidentSize, StreamingTexture.BoostFactor, StreamingTexture.CalcLoadOrderPriority(), TextureIndex );
+			new (Context.TextureStats) FTextureStreamingStats( StreamingTexture.Texture, StreamingTexture.ResidentMips, PerfectWantedMips, StreamingTexture.MostResidentMips, ResidentSize, PerfectWantedSize, PotentialSize, MostResidentSize, StreamingTexture.BoostFactor, StreamingTexture.CalcLoadOrderPriority(), TextureIndex );
 		}
 	}
 	Context.ThisFrameNumStreamingTextures++;
 	Context.ThisFrameTotalStreamingTexturesMaxSize += PotentialSize;
 	Context.ThisFrameTotalLightmapMemorySize += StreamingTexture.bIsLightmap ? ResidentSize : 0;
 	Context.ThisFrameTotalLightmapDiskSize += StreamingTexture.bIsLightmap ? PotentialSize : 0;
-	Context.ThisFrameTotalHLODMemorySize += StreamingTexture.bIsHLODTexture ? ResidentSize : 0;
-	Context.ThisFrameTotalHLODDiskSize += StreamingTexture.bIsHLODTexture ? PotentialSize : 0;
+	Context.ThisFrameTotalHLODMemorySize += StreamingTexture.LODGroup == TEXTUREGROUP_HierarchicalLOD ? ResidentSize : 0;
+	Context.ThisFrameTotalHLODDiskSize += StreamingTexture.LODGroup == TEXTUREGROUP_HierarchicalLOD ? PotentialSize : 0;
+#endif
+}
+
+void FStreamingManagerTexture::LogHeuristic(bool bOutputToLog, bool bUseHeuristic, const TCHAR*& Heuristic, const TCHAR* NewHeuristic)
+{
+#if !UE_BUILD_SHIPPING
+	if (bUseHeuristic)
+	{
+		if (bOutputToLog)
+		{
+			UE_LOG(LogContentStreaming, Log,  TEXT("  %s Reference."), NewHeuristic);
+		}
+		Heuristic = NewHeuristic;
+	}
 #endif
 }
 
 /**
  * Calculates the number of mip-levels we would like to have in memory for a texture.
  */
-void FStreamingManagerTexture::CalcWantedMips( FStreamingTexture& StreamingTexture )
+const TCHAR* FStreamingManagerTexture::CalcWantedMips(const FAsyncTextureStreamingData& StreamingData, FStreamingTexture& StreamingTexture, bool bGetHeuristic, bool bOutputToLog)
 {
-	FFloatMipLevel WantedMips;
-	float MinDistance = FLT_MAX;
-	FFloatMipLevel PerfectWantedMips = WantedMips;
+	const TCHAR* Heuristic = nullptr; 
 
-	// Figure out miplevels to request based on handlers.
-	if ( StreamingTexture.MinAllowedMips != StreamingTexture.MaxAllowedMips )
+	float MaxSize = 0;
+	float MaxSize_VisibleOnly = 0;
+
+	if ( StreamingTexture.bForceFullyLoad || StreamingTexture.MinAllowedMips == StreamingTexture.MaxAllowedMips )
 	{
-		// Iterate over all handlers and figure out the maximum requested number of mips.
-		for( int32 HandlerIndex=0; HandlerIndex<TextureStreamingHandlers.Num(); HandlerIndex++ )
+		LogHeuristic(bOutputToLog, bGetHeuristic, Heuristic, TEXT("Forced"));
+		MaxSize_VisibleOnly = FLT_MAX;
+	}
+	else
+	{
+		LogHeuristic(bOutputToLog, bGetHeuristic && StreamingData.GetDynamicInstancesView().HasTextureReferences(StreamingTexture.Texture), Heuristic, TEXT("Dynamic"));
+		StreamingData.GetDynamicInstancesView().GetTexelSize(StreamingTexture.Texture, MaxSize, MaxSize_VisibleOnly, bOutputToLog);
+
+		for (const FTextureInstanceAsyncView& StaticInstancesView : StreamingData.GetStaticInstancesViews())
 		{
-			FStreamingHandlerTextureBase* TextureStreamingHandler = TextureStreamingHandlers[HandlerIndex];
-			float HandlerDistance = FLT_MAX;
-			FFloatMipLevel HandlerWantedMips = TextureStreamingHandler->GetWantedMips( *this, StreamingTexture, HandlerDistance );
-			WantedMips = FMath::Max( WantedMips, HandlerWantedMips );
-			PerfectWantedMips = FMath::Max( PerfectWantedMips, HandlerWantedMips );
-			MinDistance = FMath::Min( MinDistance, HandlerDistance );
+			// No need to iterate more if texture is already at maximum resolution.
+			if (MaxSize_VisibleOnly >= MAX_TEXTURE_SIZE && !bOutputToLog)
+			{
+				break;
+			}
+
+			LogHeuristic(bOutputToLog, bGetHeuristic && StaticInstancesView.HasTextureReferences(StreamingTexture.Texture), Heuristic, TEXT("Static"));
+			StaticInstancesView.GetTexelSize(StreamingTexture.Texture, MaxSize, MaxSize_VisibleOnly, bOutputToLog);
 		}
 
-		bool bShouldAlsoUseLastRenderTime = false;
-		if ( StreamingTexture.LastRenderTimeRefCount > 0 || (FApp::GetCurrentTime() - StreamingTexture.LastRenderTimeRefCountTimestamp) < 91.0 )
+		/**
+		 * Fallback handler to catch textures that have been orphaned recently.
+		 * This handler prevents massive spike in texture memory usage.
+		 * Orphaned textures were previously streamed based on distance but those instance locations have been removed -
+		 * for instance because a ULevel is being unloaded. These textures are still active and in memory until they are garbage collected,
+		 * so we must ensure that they do not start using the LastRenderTime fallback handler and suddenly stream in all their mips -
+		 * just to be destroyed shortly after by a garbage collect.
+		 */
+		const float TimeSinceRemoved = (float)(FApp::GetCurrentTime() - StreamingTexture.InstanceRemovedTimestamp);
+		if (MaxSize == 0 && MaxSize_VisibleOnly == 0 && TimeSinceRemoved < 91.f && StreamingTexture.LastRenderTime > TimeSinceRemoved - 5.f)
 		{
-			bShouldAlsoUseLastRenderTime = true;
-
-#if UE_BUILD_DEBUG
-			//@DEBUG: To be able to set a break-point here.
-			if ( WantedMips.ComputeMip(&StreamingTexture, ThreadSettings.MipBias, false) != INDEX_NONE )
+			LogHeuristic(bOutputToLog, true, Heuristic, TEXT("Orphaned"));
+			MaxSize = (float)(0x1 << (FMath::Min<int32>(StreamingTexture.MaxAllowedMips - 1, StreamingTexture.ResidentMips) - 1));
+			if (StreamingTexture.LastRenderTime < 5)
 			{
-				int32 Q=0;
+				MaxSize_VisibleOnly = MaxSize;
 			}
-#endif
 		}
 
-		// Not handled by any handler, use fallback handlers.
-		if ( !WantedMips.IsHandled() || bShouldAlsoUseLastRenderTime )
+		const bool bUseLastRenderTime = (StreamingTexture.LastRenderTimeRefCount > 0 || FApp::GetCurrentTime() - StreamingTexture.LastRenderTimeRefCountTimestamp < 91.0);
+		if (((MaxSize == 0 && MaxSize_VisibleOnly == 0) || bUseLastRenderTime) && GStreamWithTimeFactor)
 		{
-			// Try the handler for textures that has recently lost instance locations.
+			LogHeuristic(bOutputToLog, bGetHeuristic, Heuristic, TEXT("LastRenderTime"));
+			if( StreamingTexture.LastRenderTime < 5.0f)
 			{
-				float HandlerDistance = FLT_MAX;
-				FFloatMipLevel HandlerWantedMips = GetWantedMipsForOrphanedTexture(StreamingTexture, HandlerDistance);
-				WantedMips = FMath::Max(WantedMips, HandlerWantedMips);
-				MinDistance = FMath::Min(MinDistance, HandlerDistance);
-				PerfectWantedMips = FMath::Max(PerfectWantedMips, HandlerWantedMips);
+				MaxSize_VisibleOnly = FLT_MAX;
 			}
-
-			// Still wasn't handled? Use the LastRenderTime handler as last resort.
-			if (!WantedMips.IsHandled() || bShouldAlsoUseLastRenderTime )
+			else if( StreamingTexture.LastRenderTime < 90.0f)
 			{
-				// Fallback handler used if texture is not handled by any other handlers. Guaranteed to handle texture and not return INDEX_NONE.
-				FStreamingHandlerTextureLastRender FallbackStreamingHandler;
-				float HandlerDistance = FLT_MAX;
-				FFloatMipLevel HandlerWantedMips = FallbackStreamingHandler.GetWantedMips( *this, StreamingTexture, HandlerDistance );
-				WantedMips = FMath::Max( WantedMips, HandlerWantedMips );
-				MinDistance = FMath::Min( MinDistance, HandlerDistance );
-				PerfectWantedMips = FMath::Max( PerfectWantedMips, HandlerWantedMips );
+				MaxSize = FLT_MAX;
 			}
 		}
 	}
-	// Stream in all allowed miplevels if requested.
-	else
-	{
-		PerfectWantedMips = FFloatMipLevel::FromMipLevel(StreamingTexture.MaxAllowedOptimalMips);
-		WantedMips = FFloatMipLevel::FromMipLevel(StreamingTexture.MaxAllowedMips);
-	}
 
-	StreamingTexture.WantedMips = WantedMips.ComputeMip(&StreamingTexture, ThreadSettings.MipBias, false);
-	StreamingTexture.MinDistance = MinDistance;
-	StreamingTexture.PerfectWantedMips = PerfectWantedMips.ComputeMip(&StreamingTexture, ThreadSettings.MipBias, true);
-}
-
-/**
- * Fallback handler to catch textures that have been orphaned recently.
- * This handler prevents massive spike in texture memory usage.
- * Orphaned textures were previously streamed based on distance but those instance locations have been removed -
- * for instance because a ULevel is being unloaded. These textures are still active and in memory until they are garbage collected,
- * so we must ensure that they do not start using the LastRenderTime fallback handler and suddenly stream in all their mips -
- * just to be destroyed shortly after by a garbage collect.
- */
-FFloatMipLevel FStreamingManagerTexture::GetWantedMipsForOrphanedTexture( FStreamingTexture& StreamingTexture, float& Distance )
-{
-	FFloatMipLevel WantedMips;
-
-	// Did we recently remove instance locations for this texture?
-	const float TimeSinceInstanceWasRemoved = float(FApp::GetCurrentTime() - StreamingTexture.InstanceRemovedTimestamp);
-
-	// Was it less than 91 seconds ago?
-	if ( TimeSinceInstanceWasRemoved < 91.0f )
-	{
-		const float TimeSinceTextureWasRendered = StreamingTexture.LastRenderTime;
-
-		// Check if it hasn't been rendered since the instances were removed (with five second margin).
-		if ( (TimeSinceTextureWasRendered - TimeSinceInstanceWasRemoved) > -5.0f )
-		{
-			// It may be garbage-collected soon. Stream out the highest mip, if currently loaded.
-			WantedMips = FFloatMipLevel::FromMipLevel(FMath::Min( StreamingTexture.ResidentMips, StreamingTexture.MaxAllowedMips - 1 ));
-			Distance = 1000.0f;
-			StreamingTexture.bUsesOrphanedHeuristics = true;
-		}
-	}
-
-	return WantedMips;
+	StreamingTexture.SetPerfectWantedMips(MaxSize, MaxSize_VisibleOnly, MipBias);
+	return Heuristic;
 }
 
 /**
@@ -2152,7 +1953,7 @@ bool FStreamingManagerTexture::HandleDebugTrackedTexturesCommand( const TCHAR* C
 								StreamingTexture.Texture->GetNumMips(),
 								LastEvent->NumRequestedMips,
 								LastEvent->WantedMips,
-								LastEvent->DynamicWantedMips.ComputeMip(&StreamingTexture, ThreadSettings.MipBias, false),
+								LastEvent->DynamicWantedMips.ComputeMip(&StreamingTexture, MipBias, false),
 								LastEvent->StreamingStatus,
 								GStreamTypeNames[StreamingTexture.GetStreamingType()],
 								StreamingTexture.BoostFactor
@@ -2169,13 +1970,14 @@ bool FStreamingManagerTexture::HandleDebugTrackedTexturesCommand( const TCHAR* C
 						for( int32 HandlerIndex=0; HandlerIndex<TextureStreamingHandlers.Num(); HandlerIndex++ )
 						{
 							FStreamingHandlerTextureBase* TextureStreamingHandler = TextureStreamingHandlers[HandlerIndex];
-							float HandlerDistance = FLT_MAX;
-							FFloatMipLevel HandlerWantedMips = TextureStreamingHandler->GetWantedMips(*this, StreamingTexture, HandlerDistance);
+							float MaxSize = 0;
+							float MaxSize_VisibleOnly = 0;
+							FFloatMipLevel HandlerWantedMips = TextureStreamingHandler->GetWantedSize(*this, StreamingTexture, HandlerDistance);
 							Ar.Logf(
 								TEXT("    Handler %s: WantedMips: %d, PerfectWantedMips: %d, Distance: %f"),
 								TextureStreamingHandler->HandlerName,
-								HandlerWantedMips.ComputeMip(&StreamingTexture, ThreadSettings.MipBias, false),
-								HandlerWantedMips.ComputeMip(&StreamingTexture, ThreadSettings.MipBias, true),
+								HandlerWantedMips.ComputeMip(&StreamingTexture, MipBias, false),
+								HandlerWantedMips.ComputeMip(&StreamingTexture, MipBias, true),
 								HandlerDistance
 								);
 						}
@@ -2250,42 +2052,21 @@ bool FStreamingManagerTexture::HandlePauseTextureStreamingCommand( const TCHAR*
 
 bool FStreamingManagerTexture::HandleStreamingManagerMemoryCommand( const TCHAR* Cmd, FOutputDevice& Ar, UWorld* InWorld )
 {
-	int32 OldStreamingTextures = 0;
-	int32 NewStreamingTextures = StreamingTextures.Num() * sizeof( FStreamingTexture );
-	for ( TObjectIterator<UTexture2D> It; It; ++It )
-	{
-		OldStreamingTextures += sizeof(TLinkedList<UTexture2D*>);
-		NewStreamingTextures += sizeof(int32);
-	}
+	AsyncWork->EnsureCompletion();
 
-	int32 OldSystem = 0;
-	int32 OldSystemSlack = 0;
-	for( int32 LevelIndex=0; LevelIndex<InWorld->GetNumLevels(); LevelIndex++ )
-	{
-		// Find instances of the texture in this level.
-		ULevel* Level = InWorld->GetLevel(LevelIndex);
-		for ( TMap<UTexture2D*,TArray<FStreamableTextureInstance> >::TConstIterator It(Level->TextureToInstancesMap); It; ++It )
-		{
-			const TArray<FStreamableTextureInstance>& TextureInstances = It.Value();
-			OldSystem += TextureInstances.Num() * sizeof(FStreamableTextureInstance);
-			OldSystemSlack += TextureInstances.GetSlack() * sizeof(FStreamableTextureInstance);
-		}
-	}
+	uint32 MemSize = sizeof(FStreamingManagerTexture);
+	MemSize += StreamingTextures.GetAllocatedSize();
+	MemSize += DynamicComponentManager.GetAllocatedSize();
+	MemSize += PendingStreamingTextures.GetAllocatedSize();
+	MemSize += LevelTextureManagers.GetAllocatedSize();
+	MemSize += AsyncWork->GetTask().StreamingData.GetAllocatedSize();
 
-	int32 NewSystem = 0;
-	int32 NewSystemSlack = 0;
-	for ( int32 LevelIndex=0; LevelIndex < ThreadSettings.LevelData.Num(); ++LevelIndex )
+	for (const FLevelTextureManager& LevelManager : LevelTextureManagers)
 	{
-		const FLevelData& LevelData = ThreadSettings.LevelData[ LevelIndex ];
-		for ( TMap<const UTexture2D*,TArray<FStreamableTextureInstance4> >::TConstIterator It(LevelData.Value.ThreadTextureInstances); It; ++It )
-		{
-			const TArray<FStreamableTextureInstance4>& TextureInstances = It.Value();
-			NewSystem += TextureInstances.Num() * sizeof(FStreamableTextureInstance4);
-			NewSystemSlack += TextureInstances.GetSlack() * sizeof(FStreamableTextureInstance4);
-		}
+		MemSize += LevelManager.GetAllocatedSize();
 	}
-	Ar.Logf( TEXT("Old: %.2f KB used, %.2f KB slack, %.2f KB texture info"), OldSystem/1024.0f, OldSystemSlack/1024.0f, OldStreamingTextures/1024.0f );
-	Ar.Logf( TEXT("New: %.2f KB used, %.2f KB slack, %.2f KB texture info"), NewSystem/1024.0f, NewSystemSlack/1024.0f, NewStreamingTextures/1024.0f );
+
+	Ar.Logf(TEXT("StreamingManagerTexture: %.2f KB used"), MemSize / 1024.0f);
 
 	return true;
 }
@@ -2451,23 +2232,23 @@ void FStreamingManagerTexture::DumpTextureGroupStats( bool bDetailedStats )
 	for( TObjectIterator<UTexture> It; It; ++It )
 	{
 		UTexture* Texture = *It;
+		UTexture2D* Texture2D = Cast<UTexture2D>(Texture);
 		FTextureGroupStats& Stat = TextureGroupStats[Texture->LODGroup];
 		FTextureGroupStats& Waste = TextureGroupWaste[Texture->LODGroup];
-		UTexture2D* Texture2D = Cast<UTexture2D>(Texture);
+		FStreamingTexture* StreamingTexture = GetStreamingTexture(Texture2D);
 		uint32 TextureAlign = 0;
-		if ( Texture2D && Texture2D->StreamingIndex >= 0 )
+		if ( StreamingTexture )
 		{
-			FStreamingTexture& StreamingTexture = GetStreamingTexture( Texture2D );
 			Stat.NumTextures++;
-			Stat.CurrentTextureSize += StreamingTexture.GetSize( StreamingTexture.ResidentMips );
-			Stat.WantedTextureSize += StreamingTexture.GetSize( StreamingTexture.WantedMips );
-			Stat.MaxTextureSize += StreamingTexture.GetSize( StreamingTexture.MaxAllowedMips );
+			Stat.CurrentTextureSize += StreamingTexture->GetSize( StreamingTexture->ResidentMips );
+			Stat.WantedTextureSize += StreamingTexture->GetSize( StreamingTexture->WantedMips );
+			Stat.MaxTextureSize += StreamingTexture->GetSize( StreamingTexture->MaxAllowedMips );
 			
-			int64 WasteCurrent = StreamingTexture.GetSize( StreamingTexture.ResidentMips ) - RHICalcTexture2DPlatformSize(Texture2D->GetSizeX(), Texture2D->GetSizeY(), Texture2D->GetPixelFormat(), StreamingTexture.ResidentMips, 1, 0, TextureAlign);			
+			int64 WasteCurrent = StreamingTexture->GetSize( StreamingTexture->ResidentMips ) - RHICalcTexture2DPlatformSize(Texture2D->GetSizeX(), Texture2D->GetSizeY(), Texture2D->GetPixelFormat(), StreamingTexture->ResidentMips, 1, 0, TextureAlign);			
 
-			int64 WasteWanted = StreamingTexture.GetSize( StreamingTexture.WantedMips ) - RHICalcTexture2DPlatformSize(Texture2D->GetSizeX(), Texture2D->GetSizeY(), Texture2D->GetPixelFormat(), StreamingTexture.WantedMips, 1, 0, TextureAlign);			
+			int64 WasteWanted = StreamingTexture->GetSize( StreamingTexture->WantedMips ) - RHICalcTexture2DPlatformSize(Texture2D->GetSizeX(), Texture2D->GetSizeY(), Texture2D->GetPixelFormat(), StreamingTexture->WantedMips, 1, 0, TextureAlign);			
 
-			int64 WasteMaxSize = StreamingTexture.GetSize( StreamingTexture.MaxAllowedMips ) - RHICalcTexture2DPlatformSize(Texture2D->GetSizeX(), Texture2D->GetSizeY(), Texture2D->GetPixelFormat(), StreamingTexture.MaxAllowedMips, 1, 0, TextureAlign);			
+			int64 WasteMaxSize = StreamingTexture->GetSize( StreamingTexture->MaxAllowedMips ) - RHICalcTexture2DPlatformSize(Texture2D->GetSizeX(), Texture2D->GetSizeY(), Texture2D->GetPixelFormat(), StreamingTexture->MaxAllowedMips, 1, 0, TextureAlign);			
 
 			Waste.NumTextures++;
 			Waste.CurrentTextureSize += FMath::Max<int64>(WasteCurrent,0);
@@ -2476,6 +2257,7 @@ void FStreamingManagerTexture::DumpTextureGroupStats( bool bDetailedStats )
 		}
 		else
 		{
+
 			bool bIsPooledTexture = Texture->Resource && IsValidRef(Texture->Resource->TextureRHI) && appIsPoolTexture( Texture->Resource->TextureRHI );
 			int64 TextureSize = Texture->CalcTextureMemorySizeEnum(TMC_ResidentMips);
 			Stat.NumNonStreamingTextures++;
@@ -2626,31 +2408,15 @@ void FStreamingManagerTexture::InvestigateTexture( const FString& InInvestigateT
 		if (TextureName.Contains(InInvestigateTextureName))
 		{
 			UTexture2D* Texture2D = StreamingTexture.Texture;
-			ETextureStreamingType StreamType = StreamingTexture.GetStreamingType();
 			int32 CurrentMipIndex = FMath::Max(Texture2D->GetNumMips() - StreamingTexture.ResidentMips, 0);
 			int32 WantedMipIndex = FMath::Max(Texture2D->GetNumMips() - StreamingTexture.WantedMips, 0);
 			int32 MaxMipIndex = FMath::Max(Texture2D->GetNumMips() - StreamingTexture.MaxAllowedMips, 0);
 			UE_LOG(LogContentStreaming, Log,  TEXT("Texture: %s"), *TextureName );
 			UE_LOG(LogContentStreaming, Log,  TEXT("  Texture group:   %s"), UTexture::GetTextureGroupString(StreamingTexture.LODGroup) );
-			if ( StreamType != StreamType_LastRenderTime && StreamingTexture.bUsesLastRenderHeuristics )
-			{
-				UE_LOG(LogContentStreaming, Log,  TEXT("  Stream logic:    %s and %s (%d references)"), GStreamTypeNames[StreamType], GStreamTypeNames[StreamType_LastRenderTime], StreamingTexture.LastRenderTimeRefCount );
-			}
-			else
-			{
-				UE_LOG(LogContentStreaming, Log,  TEXT("  Stream logic:    %s"), GStreamTypeNames[StreamType] );
-			}
-			if ( StreamingTexture.LastRenderTime > 1000000.0f )
-			{
-				UE_LOG(LogContentStreaming, Log,  TEXT("  Last rendertime: Never") );
-			}
-			else
-			{
-				UE_LOG(LogContentStreaming, Log,  TEXT("  Last rendertime: %.3f seconds ago"), StreamingTexture.LastRenderTime );
-			}
-			if ( StreamingTexture.ForceLoadRefCount > 0 )
+
+			if ( StreamingTexture.bAsNeverStream )
 			{
-				UE_LOG(LogContentStreaming, Log,  TEXT("  Force all mips:  %d level references"), StreamingTexture.ForceLoadRefCount );
+				UE_LOG(LogContentStreaming, Log,  TEXT("  Force all mips: bForceFullyLoad") );
 			}
 			else if ( Texture2D->bGlobalForceMipLevelsToBeResident )
 			{
@@ -2677,477 +2443,28 @@ void FStreamingManagerTexture::InvestigateTexture( const FString& InInvestigateT
 			UE_LOG(LogContentStreaming, Log,  TEXT("  Retention Priority: %.3f"), StreamingTexture.CalcRetentionPriority() );
 			UE_LOG(LogContentStreaming, Log,  TEXT("  Boost factor:    %.1f"), StreamingTexture.BoostFactor );
 			UE_LOG(LogContentStreaming, Log,  TEXT("  Allowed mips:    %d-%d"), StreamingTexture.MinAllowedMips, StreamingTexture.MaxAllowedMips );
-			UE_LOG(LogContentStreaming, Log,  TEXT("  Global mip bias: %.1f"), ThreadSettings.MipBias );
+			UE_LOG(LogContentStreaming, Log,  TEXT("  Global mip bias: %.1f"), MipBias );
 
-			float ScreenSizeFactor;
-			switch ( StreamingTexture.LODGroup )
-			{
-				case TEXTUREGROUP_Lightmap:
-					ScreenSizeFactor = GLightmapStreamingFactor;
-					break;
-				case TEXTUREGROUP_Shadowmap:
-					ScreenSizeFactor = GShadowmapStreamingFactor;
-					break;
-				default:
-					ScreenSizeFactor = 1.0f;
-			}
-			ScreenSizeFactor *= StreamingTexture.BoostFactor;
+			// Make sure the async task is idle (also implies Update_Async is finished and that the distances are valid).
+			AsyncWork->EnsureCompletion();
+			FAsyncTextureStreamingData& StreamingData = AsyncWork->GetTask().StreamingData;
+			StreamingData.Update_Async();
+
+			float MaxSize = 0;
+			float MaxSize_VisibleOnly = 0;
 
-			for( int32 ViewIndex=0; ViewIndex < ThreadNumViews(); ViewIndex++ )
+			for( int32 ViewIndex=0; ViewIndex < StreamingData.GetViewInfos().Num(); ViewIndex++ )
 			{
 				// Calculate distance of viewer to bounding sphere.
-				const FStreamingViewInfo& ViewInfo = ThreadGetView(ViewIndex);
+				const FStreamingViewInfo& ViewInfo = StreamingData.GetViewInfos()[ViewIndex];
 				UE_LOG(LogContentStreaming, Log,  TEXT("  View%d: Position=(%d,%d,%d) ScreenSize=%f Boost=%f"), ViewIndex,
 					int32(ViewInfo.ViewOrigin.X), int32(ViewInfo.ViewOrigin.Y), int32(ViewInfo.ViewOrigin.Z),
 					ViewInfo.ScreenSize, ViewInfo.BoostFactor);
 			}
 
-			// Iterate over all associated/ visible levels.
-			for ( int32 LevelIndex=0; LevelIndex < ThreadSettings.LevelData.Num(); ++LevelIndex )
-			{
-				// Find instances of the texture in this level.
-				FStreamingManagerTexture::FThreadLevelData& LevelData = ThreadSettings.LevelData[ LevelIndex ].Value;
-				TArray<FStreamableTextureInstance4>* TextureInstances = LevelData.ThreadTextureInstances.Find( StreamingTexture.Texture );
-				if ( TextureInstances )
-				{
-					for ( int32 InstanceIndex=0; InstanceIndex < TextureInstances->Num(); ++InstanceIndex )
-					{
-						const FStreamableTextureInstance4& Instance = (*TextureInstances)[InstanceIndex];
-						const FVector4& CenterX = Instance.BoundsOriginX;
-						const FVector4& CenterY = Instance.BoundsOriginY;
-						const FVector4& CenterZ = Instance.BoundsOriginZ;
-						for ( int32 PartialIndex=0; PartialIndex < 4; ++PartialIndex )
-						{
-							if ( CenterX[PartialIndex] < 3.402823466e+30F )
-							{
-								FVector Center( CenterX[PartialIndex], CenterY[PartialIndex], CenterZ[PartialIndex] );
-								float Radius = Instance.BoundingSphereRadius[PartialIndex];
-								float TexelFactor = Instance.TexelFactor[PartialIndex];
-								float MinRelevantDistance = FMath::Sqrt(Instance.MinDistanceSq[PartialIndex]);
-								float MaxRelevantDistance = SqrtKeepMax(Instance.MaxDistanceSq[PartialIndex]);
-
-								float MinDistance = MAX_FLT;
-								float OutOfRangeDistance = MAX_FLT;
-								FFloatMipLevel WantedMipCount;
-								bool bInside = false;
-								bool bIsInRange = false;
-
-								for( int32 ViewIndex=0; ViewIndex < ThreadNumViews(); ViewIndex++ )
-								{
-									// Calculate distance of viewer to bounding sphere.
-									const FStreamingViewInfo& ViewInfo = ThreadGetView(ViewIndex);
-									float Distance = (ViewInfo.ViewOrigin - Center).Size();
-									if (Distance >= MinRelevantDistance && Distance <= MaxRelevantDistance)
-									{
-										bIsInRange = true;
-										MinDistance = FMath::Min(MinDistance, Distance);
-										float DistSqMinusRadiusSq = ClampMeshToCameraDistanceSquared(FMath::Square(Distance) - FMath::Square(Radius));
-										if( DistSqMinusRadiusSq > 1.f )
-										{
-											// Outside the texture instance bounding sphere, calculate miplevel based on screen space size of bounding sphere.
-											// Calculate the maximum screen space dimension in pixels.
-											const float ScreenSize = ViewInfo.ScreenSize * ViewInfo.BoostFactor * ScreenSizeFactor;
-											const float	ScreenSizeInTexels = TexelFactor * FMath::InvSqrtEst( DistSqMinusRadiusSq ) * ScreenSize;
-											// WantedMipCount is the number of mips so we need to adjust with "+ 1".
-											WantedMipCount = FMath::Max( WantedMipCount, FFloatMipLevel::FromScreenSizeInTexels(ScreenSizeInTexels) );
-										}
-										else
-										{
-											// Request all miplevels to be loaded if we're inside the bounding sphere.
-											WantedMipCount = FFloatMipLevel::FromMipLevel(StreamingTexture.MaxAllowedMips);
-											bInside = true;
-											break;
-										}
-									}
-									else
-									{
-										// Here we store distance in another variable since the relevant distance (from another view) might actually be greater.
-										OutOfRangeDistance = FMath::Min(OutOfRangeDistance, Distance);
-									}
-								}
-
-								FString RangeString;
-								if (MinRelevantDistance != 0 || MaxRelevantDistance != FLT_MAX)
-								{
-									if (MaxRelevantDistance == FLT_MAX)
-									{
-										RangeString = FString::Printf(TEXT("Range: [%.0f, INF], "), MinRelevantDistance);
-									}
-									else
-									{
-										RangeString = FString::Printf(TEXT("Range: [%.0f, %.0f], "), MinRelevantDistance, MaxRelevantDistance);
-									}
-								}
-									
-
-								if (bIsInRange)
-								{
-									int32 IntWantedMipCount = WantedMipCount.ComputeMip(&StreamingTexture, ThreadSettings.MipBias, false);
-									int32 WantedMip = FMath::Max(Texture2D->GetNumMips() - IntWantedMipCount, 0);
-
-									UE_LOG(LogContentStreaming, Log, TEXT("Static: Wanted=%dx%d, Distance=%.1f, %sTexelFactor=%.2f, Radius=%5.1f, Position=(%d,%d,%d)%s"),
-										Texture2D->PlatformData->Mips[WantedMip].SizeX, Texture2D->PlatformData->Mips[WantedMip].SizeY,
-										MinDistance, *RangeString, TexelFactor, Radius, int32(Center.X), int32(Center.Y), int32(Center.Z),
-										bInside ? TEXT(" [all mips]") : TEXT(""));
-								}
-								else // Here there is no use for distance it will be FLT_MAX
-								{
-									UE_LOG(LogContentStreaming, Log, TEXT("Static: Distance(OutOfRange)=%.1f, %sTexelFactor=%.2f, Radius=%5.1f, Position=(%d,%d,%d)%s"),
-										   OutOfRangeDistance, *RangeString, TexelFactor, Radius, int32(Center.X), int32(Center.Y), int32(Center.Z),
-										   bInside ? TEXT(" [all mips]") : TEXT(""));
-								}
-							}
-						}
-					}
-				}
-			}
+			CalcWantedMips(StreamingData, StreamingTexture, true, true);
 		}
 	}
 #endif // !UE_BUILD_SHIPPING
 }
 
-void FStreamingManagerTexture::DumpTextureInstances( const UPrimitiveComponent* Primitive, UTexture2D* Texture2D )
-{
-#if !UE_BUILD_SHIPPING && 0
-	//@TODO
-			UE_LOG(LogContentStreaming, Log, TEXT("%s: Wanted=%dx%d, Distance=%.1f, TexelFactor=%.2f, CurrentRadius=%5.1f, OriginalRadius=%5.1f, Position=(%d,%d,%d), IsRegistered=%d, Component=\"%s\""),
-				TypeString, Texture2D->PlatformData->Mips[WantedMipIndex].SizeX, Texture2D->PlatformData->Mips[WantedMipIndex].SizeY,
-				MinDistance, Instance.TexelFactor, PrimitiveData.BoundingSphere.W, 1.0f / Instance.InvOriginalRadius,
-				int32(PrimitiveData.BoundingSphere.Center.X), int32(PrimitiveData.BoundingSphere.Center.Y), int32(PrimitiveData.BoundingSphere.Center.Z),
-				PrimitiveData.bAttached ? true : false,
-				*Primitive->GetReadableName() );
-#endif // !UE_BUILD_SHIPPING
-}
-
-/*-----------------------------------------------------------------------------
-	Streaming handler implementations.
------------------------------------------------------------------------------*/
-
-/**
- * Returns mip count wanted by this handler for the passed in texture.
- * New version for the priority system, using SIMD and no fudge factor.
- * 
- * @param	StreamingManager	Streaming manager
- * @param	Streaming Texture	Texture to determine wanted mip count for
- * @param	MinDistance			[out] Distance to the closest instance of this texture, in world space units.
- * @return	Number of miplevels that should be streamed in or INDEX_NONE if texture isn't handled by this handler.
- */
-FFloatMipLevel FStreamingHandlerTextureStatic::GetWantedMips( FStreamingManagerTexture& StreamingManager, FStreamingTexture& StreamingTexture, float& MinDistance )
-{
-	FFloatMipLevel WantedMipCount;
-	bool bShouldAbortLoop = false;
-	bool bEntryFound = false; // True an entry for this texture exists­.
-
-	const bool bUseNewMetrics = CVarStreamingUseNewMetrics.GetValueOnAnyThread() > 0;
-	const float HiddenScale = CVarStreamingHiddenPrimitiveScale.GetValueOnAnyThread();
-	const float MaxResolution = (float)(0x1 << (StreamingTexture.MaxAllowedMips - 1));
-
-	// Nothing do to if there are no views or instances.
-	if( StreamingManager.ThreadNumViews() /*&& StreamingTexture.Instances.Num() > 0*/ )
-	{
-		float ScreenSizeFactor;
-		switch ( StreamingTexture.LODGroup )
-		{
-			case TEXTUREGROUP_Lightmap:
-				ScreenSizeFactor = GLightmapStreamingFactor;
-				break;
-			case TEXTUREGROUP_Shadowmap:
-				ScreenSizeFactor = GShadowmapStreamingFactor;
-				break;
-			default:
-				ScreenSizeFactor = 1.0f;
-		}
-
-		ScreenSizeFactor *= StreamingTexture.BoostFactor;
-
-		const VectorRegister MaxFloat4 = VectorSet((float)FLT_MAX, (float)FLT_MAX, (float)FLT_MAX, (float)FLT_MAX);
-		VectorRegister MinDistanceSq4 = MaxFloat4;
-		VectorRegister MaxTexels = VectorSet(-(float)FLT_MAX, -(float)FLT_MAX, -(float)FLT_MAX, -(float)FLT_MAX);
-		for ( int32 LevelIndex=0; LevelIndex < StreamingManager.ThreadSettings.LevelData.Num(); ++LevelIndex )
-		{
-			FStreamingManagerTexture::FThreadLevelData& LevelData = StreamingManager.ThreadSettings.LevelData[ LevelIndex ].Value;
-			TArray<FStreamableTextureInstance4>* TextureInstances = LevelData.ThreadTextureInstances.Find( StreamingTexture.Texture );
-			if ( TextureInstances )
-			{
-				bEntryFound = true;
-
-				for ( int32 InstanceIndex=0; InstanceIndex < TextureInstances->Num() && !bShouldAbortLoop; ++InstanceIndex )
-				{
-					const FStreamableTextureInstance4& TextureInstance = (*TextureInstances)[InstanceIndex];
-
-					// Calculate distance of viewer to bounding sphere.
-					const VectorRegister CenterX = VectorLoadAligned( &TextureInstance.BoundsOriginX );
-					const VectorRegister CenterY = VectorLoadAligned( &TextureInstance.BoundsOriginY );
-					const VectorRegister CenterZ = VectorLoadAligned( &TextureInstance.BoundsOriginZ );
-
-					// Calculate distance of viewer to bounding sphere.
-					const VectorRegister ExtentX = VectorLoadAligned( &TextureInstance.BoxExtentSizeX );
-					const VectorRegister ExtentY = VectorLoadAligned( &TextureInstance.BoxExtentSizeY );
-					const VectorRegister ExtentZ = VectorLoadAligned( &TextureInstance.BoxExtentSizeZ );
-
-					// Iterate over all view infos.
-					for( int32 ViewIndex=0; ViewIndex < StreamingManager.ThreadNumViews() && !bShouldAbortLoop; ViewIndex++ )
-					{
-						const FStreamingViewInfo& ViewInfo = StreamingManager.ThreadGetView(ViewIndex);
-
-						const float MaxEffectiveScreenSize = CVarStreamingScreenSizeEffectiveMax.GetValueOnAnyThread();
-						const float EffectiveScreenSize = (MaxEffectiveScreenSize > 0.0f) ? FMath::Min(MaxEffectiveScreenSize, ViewInfo.ScreenSize) : ViewInfo.ScreenSize;
-						const float ScreenSizeFloat = EffectiveScreenSize * ViewInfo.BoostFactor * ScreenSizeFactor;
-						const VectorRegister ScreenSize = VectorLoadFloat1( &ScreenSizeFloat );
-						const VectorRegister ViewOriginX = VectorLoadFloat1( &ViewInfo.ViewOrigin.X );
-						const VectorRegister ViewOriginY = VectorLoadFloat1( &ViewInfo.ViewOrigin.Y );
-						const VectorRegister ViewOriginZ = VectorLoadFloat1( &ViewInfo.ViewOrigin.Z );
-
-						//const float DistSq = FVector::DistSquared( ViewInfo.ViewOrigin, TextureInstance.BoundingSphere.Center );
-						VectorRegister Temp = VectorSubtract( ViewOriginX, CenterX );
-						VectorRegister DistSq = VectorMultiply( Temp, Temp );
-						Temp = VectorSubtract( ViewOriginY, CenterY );
-						DistSq = VectorMultiplyAdd( Temp, Temp, DistSq );
-						Temp = VectorSubtract( ViewOriginZ, CenterZ );
-						DistSq = VectorMultiplyAdd( Temp, Temp, DistSq );
-
-						// The range is handle as the distance to the center of the bounds (and not to the bounds), 
-						// because it needs to match how HLOD visibility is handled.
-						VectorRegister ClampedDistSq = VectorMax( VectorLoadAligned( &TextureInstance.MinDistanceSq ), DistSq );
-						ClampedDistSq = VectorMin( VectorLoadAligned( &TextureInstance.MaxDistanceSq ), ClampedDistSq );
-						const VectorRegister InRangeMask = VectorCompareEQ(DistSq, ClampedDistSq);
-
-						VectorRegister DistSqMinusRadiusSq = VectorZero();
-						if (bUseNewMetrics)
-						{
-							// In this case DistSqMinusRadiusSq will contain the distance to the box^2
-							Temp = VectorSubtract( ViewOriginX, CenterX );
-							Temp = VectorAbs( Temp );
-							VectorRegister BoxRef = VectorMin( Temp, ExtentX );
-							Temp = VectorSubtract( Temp, BoxRef );
-							DistSqMinusRadiusSq = VectorMultiply( Temp, Temp );
-
-							Temp = VectorSubtract( ViewOriginY, CenterY );
-							Temp = VectorAbs( Temp );
-							BoxRef = VectorMin( Temp, ExtentY );
-							Temp = VectorSubtract( Temp, BoxRef );
-							DistSqMinusRadiusSq = VectorMultiplyAdd( Temp, Temp, DistSqMinusRadiusSq );
-
-							Temp = VectorSubtract( ViewOriginZ, CenterZ );
-							Temp = VectorAbs( Temp );
-							BoxRef = VectorMin( Temp, ExtentZ );
-							Temp = VectorSubtract( Temp, BoxRef );
-							DistSqMinusRadiusSq = VectorMultiplyAdd( Temp, Temp, DistSqMinusRadiusSq );
-						}
-						else
-						{
-							// const float DistSqMinusRadiusSq = DistSq - FMath::Square(TextureInstance.BoundingSphere.W);
-							DistSqMinusRadiusSq = VectorLoadAligned( &TextureInstance.BoundingSphereRadius );
-							DistSqMinusRadiusSq = VectorMultiply( DistSqMinusRadiusSq, DistSqMinusRadiusSq );
-							DistSqMinusRadiusSq = VectorSubtract( DistSq, DistSqMinusRadiusSq );
-						}
-						// Force at least distance of 2 for entry not in their range, in order to enter in next conditional block.
-						// Also required in order to not enter in the final block computing the WantedMipCount.
-						DistSqMinusRadiusSq = VectorSelect( InRangeMask, DistSqMinusRadiusSq, VectorMax(DistSqMinusRadiusSq, VectorSet(2.f, 2.f, 2.f, 2.f) ) );
-
-						MinDistanceSq4 = VectorMin( MinDistanceSq4, DistSqMinusRadiusSq );
-
-						if ( VectorAllGreaterThan( DistSqMinusRadiusSq, VectorOne() ) )
-						{
-							if (StreamingTexture.LODGroup != TEXTUREGROUP_Terrain_Heightmap)
-							{
-								// Calculate the maximum screen space dimension in pixels.
-								//const float ScreenSizeInTexels = TextureInstance.TexelFactor * FMath::InvSqrtEst( DistSqMinusRadiusSq ) * ScreenSize;
-								DistSqMinusRadiusSq = VectorMax( DistSqMinusRadiusSq, VectorOne() );
-								VectorRegister Texels = VectorMultiply( VectorLoadAligned( &TextureInstance.TexelFactor ), VectorReciprocalSqrt(DistSqMinusRadiusSq) );
-								Texels = VectorMultiply( Texels, ScreenSize );
-
-								// Clamp to max resolution before actually scaling, otherwise this could have not effect.
-								Texels = VectorMin(VectorSet(MaxResolution, MaxResolution, MaxResolution, MaxResolution), Texels);
-								VectorRegister VisibilityScale = VectorSelect( InRangeMask, VectorOne(), VectorSet(HiddenScale, HiddenScale, HiddenScale, HiddenScale) );
-								Texels = VectorMultiply(Texels, VisibilityScale);
-
-								MaxTexels = VectorMax( MaxTexels, Texels );
-							}
-							else
-							{
-								// To check Forced LOD...
-								VectorRegister MinForcedLODs = VectorLoadAligned( &TextureInstance.TexelFactor );
-								bool AllForcedLOD = !!VectorAllLesserThan(MinForcedLODs, VectorZero());
-								MinForcedLODs = VectorMin( MinForcedLODs, VectorSwizzle(MinForcedLODs, 2, 3, 0, 1) );
-								MinForcedLODs = VectorMin( MinForcedLODs, VectorReplicate(MinForcedLODs, 1) );
-								float MinLODValue;
-								VectorStoreFloat1( MinForcedLODs, &MinLODValue );
-
-								if (MinLODValue <= 0)
-								{
-									WantedMipCount = FMath::Max(WantedMipCount, FFloatMipLevel::FromMipLevel(StreamingTexture.MaxAllowedMips - 13 - FMath::FloorToInt(MinLODValue)));
-									if (WantedMipCount >= FFloatMipLevel::FromMipLevel(StreamingTexture.MaxAllowedMips))
-									{
-										MinDistance = 1.0f;
-										bShouldAbortLoop = true;
-									}
-								}
-
-								if (!AllForcedLOD)
-								{
-									// Calculate the maximum screen space dimension in pixels.
-									DistSqMinusRadiusSq = VectorMax( DistSqMinusRadiusSq, VectorOne() );
-									VectorRegister Texels = VectorMultiply( VectorLoadAligned( &TextureInstance.TexelFactor ), VectorReciprocalSqrt(DistSqMinusRadiusSq) );
-									Texels = VectorMultiply( Texels, ScreenSize );
-									MaxTexels = VectorMax( MaxTexels, Texels );
-								}
-							}
-						}
-						else
-						{
-							WantedMipCount = FFloatMipLevel::FromMipLevel(StreamingTexture.MaxAllowedMips);
-							MinDistance = 1.0f;
-							bShouldAbortLoop = true;
-						}
-						StreamingTexture.bUsesStaticHeuristics = true;
-					}
-				}
-
-				if (StreamingTexture.LODGroup == TEXTUREGROUP_Terrain_Heightmap && VectorAllLesserThan(MaxTexels, VectorZero())) // All ForcedLOD cases...
-				{
-					MinDistance = 1.0f;
-					bShouldAbortLoop = true;
-				}
-
-				if ( StreamingTexture.bUsesStaticHeuristics && !bShouldAbortLoop )
-				{
-					MinDistanceSq4 = VectorMin( MinDistanceSq4, VectorSwizzle(MinDistanceSq4, 2, 3, 0, 1) );
-					MinDistanceSq4 = VectorMin( MinDistanceSq4, VectorReplicate(MinDistanceSq4, 1) );
-					float MinDistanceSq;
-					VectorStoreFloat1( MinDistanceSq4, &MinDistanceSq );
-					MinDistanceSq = ClampMeshToCameraDistanceSquared(MinDistanceSq);
-					if ( MinDistanceSq > 1.0f )
-					{
-						MaxTexels = VectorMax( MaxTexels, VectorSwizzle(MaxTexels, 2, 3, 0, 1) );
-						MaxTexels = VectorMax( MaxTexels, VectorReplicate(MaxTexels, 1) );
-						float ScreenSizeInTexels;
-						VectorStoreFloat1( MaxTexels, &ScreenSizeInTexels );
-						// If every entry is out of range, we risk returning -1 from the handler, which would be the same as not handled.
-						// Not handled textures, will use fallback handlers, where here we rather want to prevent streaming the texture.
-						if (bEntryFound)
-						{
-							ScreenSizeInTexels = FMath::Max(1.f, ScreenSizeInTexels); // Ensure IsHandled()
-						}
-						// WantedMipCount is the number of mips so we need to adjust with "+ 1".
-						WantedMipCount = FMath::Max(WantedMipCount, FFloatMipLevel::FromScreenSizeInTexels(ScreenSizeInTexels));
-						MinDistance = FMath::Min( MinDistanceSq, FMath::Sqrt( MinDistanceSq ) );
-					}
-					else
-					{
-						WantedMipCount = FFloatMipLevel::FromMipLevel(StreamingTexture.MaxAllowedMips);
-						MinDistance = 1.0f;
-					}
-				}
-			}
-		}
-	}
-
-	return WantedMipCount;
-}
-
-/*-----------------------------------------------------------------------------
-	Streaming handler implementations.
------------------------------------------------------------------------------*/
-
-/**
- * Returns mip count wanted by this handler for the passed in texture.
- * New version for the priority system, using SIMD and no fudge factor.
- * 
- * @param	StreamingManager	Streaming manager
- * @param	Streaming Texture	Texture to determine wanted mip count for
- * @param	MinDistance			[out] Distance to the closest instance of this texture, in world space units.
- * @return	Number of miplevels that should be streamed in or INDEX_NONE if texture isn't handled by this handler.
- */
-FFloatMipLevel FStreamingHandlerTextureDynamic::GetWantedMips( FStreamingManagerTexture& StreamingManager, FStreamingTexture& StreamingTexture, float& MinDistance )
-{
-	float MaxSize_InRange, MinDistanceSq_InRange, MaxSize_AnyRange, MinDistanceSq_AnyRange;
-	if (StreamingManager.ThreadSettings.TextureBoundsVisibility->GetTexelSize(StreamingTexture.Texture, MaxSize_InRange, MinDistanceSq_InRange, MaxSize_AnyRange, MinDistanceSq_AnyRange))
-	{
-		//*//*//*//*//*//*//*//*//*//*//*//*//*//*//*//*
-		StreamingTexture.bUsesDynamicHeuristics = true;
-		//*//*//*//*//*//*//*//*//*//*//*//*//*//*//*//*
-
-		const float HiddenScale = CVarStreamingHiddenPrimitiveScale.GetValueOnAnyThread();
-		const float MaxResolution = (float)(0x1 << (StreamingTexture.MaxAllowedMips - 1));
-
-		float ExtraScale = StreamingTexture.BoostFactor;
-		switch ( StreamingTexture.LODGroup )
-		{
-			case TEXTUREGROUP_Lightmap:
-				ExtraScale *= GLightmapStreamingFactor;
-				break;
-			case TEXTUREGROUP_Shadowmap:
-				ExtraScale *= GShadowmapStreamingFactor;
-				break;
-			default:
-				break;
-		}
-
-		MaxSize_InRange *= ExtraScale;
-		MaxSize_AnyRange = FMath::Min(MaxResolution, MaxSize_AnyRange) * ExtraScale * HiddenScale; // Clamp by max resolution first otherwise it could do no effect.
-
-		check(MinDistanceSq_AnyRange >= 1.f);
-		MinDistance = FMath::Min(MinDistance, FMath::Sqrt(MinDistanceSq_AnyRange));
-
-		return FFloatMipLevel::FromScreenSizeInTexels(FMath::Max(MaxSize_InRange, MaxSize_AnyRange));
-	}
-	else
-	{
-		return FFloatMipLevel();
-	}
-}
-
-/**
- * Returns mip count wanted by this handler for the passed in texture. 
- * 
- * @param	StreamingManager	Streaming manager
- * @param	Streaming Texture	Texture to determine wanted mip count for
- * @param	MinDistance			[out] Distance to the closest instance of this texture, in world space units.
- * @return	Number of miplevels that should be streamed in or INDEX_NONE if texture isn't handled by this handler.
- */
-FFloatMipLevel FStreamingHandlerTextureLastRender::GetWantedMips( FStreamingManagerTexture& StreamingManager, FStreamingTexture& StreamingTexture, float& MinDistance )
-{
-	FFloatMipLevel Ret;
-
-	float SecondsSinceLastRender = StreamingTexture.LastRenderTime;	//(FApp::GetCurrentTime() - Texture->Resource->LastRenderTime);;
-	StreamingTexture.bUsesLastRenderHeuristics = true;
-
-	if( SecondsSinceLastRender < 45.0f && GStreamWithTimeFactor )
-	{
-		Ret = FFloatMipLevel::FromMipLevel(StreamingTexture.MaxAllowedMips);
-		MinDistance = 0;
-	}
-	else if( SecondsSinceLastRender < 90.0f && GStreamWithTimeFactor )
-	{
-		Ret = FFloatMipLevel::FromMipLevel(StreamingTexture.MaxAllowedMips - 1);
-		MinDistance = 1000.0f;
-	}
-	else
-	{
-		Ret = FFloatMipLevel::FromMipLevel(0);
-		MinDistance = 10000.0f;
-	}
-
-	return Ret;
-}
-
-/**
- * Returns mip count wanted by this handler for the passed in texture. 
- * New version for the priority system, using SIMD and no fudge factor.
- * 
- * @param	StreamingManager	Streaming manager
- * @param	Streaming Texture	Texture to determine wanted mip count for
- * @param	MinDistance			[out] Distance to the closest instance of this texture, in world space units.
- * @return	Number of miplevels that should be streamed in or INDEX_NONE if texture isn't handled by this handler.
- */
-FFloatMipLevel FStreamingHandlerTextureLevelForced::GetWantedMips( FStreamingManagerTexture& StreamingManager, FStreamingTexture& StreamingTexture, float& MinDistance )
-{
-	FFloatMipLevel Ret;
-	if ( StreamingTexture.ForceLoadRefCount > 0 )
-	{
-		StreamingTexture.bUsesForcedHeuristics = true;
-		Ret = FFloatMipLevel::FromMipLevel(StreamingTexture.MaxAllowedMips);
-	}
-	return Ret;
-}
\ No newline at end of file
diff --git a/Engine/Source/Runtime/Engine/Private/Streaming/StreamingManagerTexture.h b/Engine/Source/Runtime/Engine/Private/Streaming/StreamingManagerTexture.h
--- a/Engine/Source/Runtime/Engine/Private/Streaming/StreamingManagerTexture.h
+++ b/Engine/Source/Runtime/Engine/Private/Streaming/StreamingManagerTexture.h
@@ -9,92 +9,12 @@ TextureStreamingManager.h: Definitions of classes used for texture streaming.
 #include "StreamingTexture.h"
 #include "TextureInstanceManager.h"
 
+class FAsyncTextureStreamingData;
+
 /*-----------------------------------------------------------------------------
 	Texture streaming.
 -----------------------------------------------------------------------------*/
 
-
-/**
- * Structure containing all information needed for determining the screen space
- * size of an object/ texture instance.
- */
-struct FStreamableTextureInstance4
-{
-	FStreamableTextureInstance4()
-	:	BoundsOriginX( MAX_FLT, MAX_FLT, MAX_FLT, MAX_FLT )
-	,	BoundsOriginY( 0, 0, 0, 0 )
-	,	BoundsOriginZ( 0, 0, 0, 0 )
-	,	BoxExtentSizeX( 0, 0, 0, 0 )
-	,	BoxExtentSizeY( 0, 0, 0, 0 )
-	,	BoxExtentSizeZ( 0, 0, 0, 0 )
-	,	BoundingSphereRadius( 0, 0, 0, 0 )
-	,	MinDistanceSq( 0, 0, 0, 0 )
-	,	MaxDistanceSq( MAX_FLT, MAX_FLT, MAX_FLT, MAX_FLT )
-	,	TexelFactor( 0, 0, 0, 0 )
-	{
-	}
-
-	/** X coordinates for the bounds origin of 4 texture instances */
-	FVector4 BoundsOriginX;
-	/** Y coordinates for the bounds origin of 4 texture instances */
-	FVector4 BoundsOriginY;
-	/** Z coordinates for the bounds origin of 4 texture instances */
-	FVector4 BoundsOriginZ;
-
-	/** X size of the bounds box extent of 4 texture instances */
-	FVector4 BoxExtentSizeX;
-	/** Y size of the bounds box extent of 4 texture instances */
-	FVector4 BoxExtentSizeY;
-	/** Z size of the bounds box extent of 4 texture instances */
-	FVector4 BoxExtentSizeZ;
-
-	/** Sphere radii for the bounding sphere of 4 texture instances */
-	FVector4 BoundingSphereRadius;
-	/** Minimal distance ^2 (between the bounding sphere origin and the view origin) for which this entry is valid */
-	FVector4 MinDistanceSq;
-	/** Maximal distance ^2 (between the bounding sphere origin and the view origin) for which this entry is valid */
-	FVector4 MaxDistanceSq;
-
-	/** Texel scale factors for 4 texture instances */
-	FVector4 TexelFactor;
-};
-
-
-// To compute the maximum mip size required by any heuristic
-// float to avoid float->int conversion, unclamped until the end
-class FFloatMipLevel
-{
-public:
-	// constructor
-	FFloatMipLevel();
-
-	/**
-	 * Helper function to map screen texels to a mip level, includes the global bias
-	 * @param StreamingTexture can be 0 (if texture is getting destroyed)
-	 * @apram MipBias from cvar
-	 * @param bOptimal only needed for stats
-	 */
-	int32 ComputeMip(const FStreamingTexture* StreamingTexture, float MipBias, bool bOptimal) const;
-	//
-	bool operator>=(const FFloatMipLevel& rhs) const;
-	//
-	static FFloatMipLevel FromScreenSizeInTexels(float ScreenSizeInTexels);
-	
-	static FFloatMipLevel FromMipLevel(int32 InMipLevel);
-
-	void Invalidate();
-
-	// @return true: means this handle has a mip level that needs to be considered
-	bool IsHandled() const
-	{
-		return MipLevel > 0.0f;
-	}
-
-private:
-	// -1: not set yet
-	float MipLevel;
-};
-
 /**
  * Streaming manager dealing with textures.
  */
@@ -139,22 +59,6 @@ struct FStreamingManagerTexture : public ITextureStreamingManager
 	 */
 	virtual void NotifyLevelChange() override;
 
-	/**
-	 * Adds a textures streaming handler to the array of handlers used to determine which
-	 * miplevels need to be streamed in/ out.
-	 *
-	 * @param TextureStreamingHandler	Handler to add
-	 */
-	void AddTextureStreamingHandler( FStreamingHandlerTextureBase* TextureStreamingHandler );
-
-	/**
-	 * Removes a textures streaming handler from the array of handlers used to determine which
-	 * miplevels need to be streamed in/ out.
-	 *
-	 * @param TextureStreamingHandler	Handler to remove
-	 */
-	void RemoveTextureStreamingHandler( FStreamingHandlerTextureBase* TextureStreamingHandler );
-
 	/** Don't stream world resources for the next NumFrames. */
 	virtual void SetDisregardWorldResourcesForFrames( int32 NumFrames ) override;
 
@@ -216,7 +120,7 @@ struct FStreamingManagerTexture : public ITextureStreamingManager
 	virtual void RemoveStreamingTexture( UTexture2D* Texture ) override;
 
 	/** Adds a ULevel to the streaming manager. */
-	virtual void AddPreparedLevel( class ULevel* Level ) override;
+	virtual void AddLevel( class ULevel* Level ) override;
 
 	/** Removes a ULevel from the streaming manager. */
 	virtual void RemoveLevel( class ULevel* Level ) override;
@@ -264,10 +168,7 @@ struct FStreamingManagerTexture : public ITextureStreamingManager
 	virtual void NotifyPrimitiveUpdated( const UPrimitiveComponent* Primitive ) override;
 
 	/** Returns the corresponding FStreamingTexture for a UTexture2D. */
-	FStreamingTexture& GetStreamingTexture( const UTexture2D* Texture2D );
-
-	/** Returns true if this is a streaming texture that is managed by the streaming manager. */
-	virtual bool IsManagedStreamingTexture( const UTexture2D* Texture2D ) override;
+	FStreamingTexture* GetStreamingTexture( const UTexture2D* Texture2D );
 
 	/** Updates the I/O state of a texture (allowing it to progress to the next stage) and some stats. */
 	void UpdateTextureStatus( FStreamingTexture& StreamingTexture, FStreamingContext& Context );
@@ -291,18 +192,6 @@ struct FStreamingManagerTexture : public ITextureStreamingManager
 	 */
 	void UpdateStreamingStats( const FStreamingContext& Context, bool bAllTexturesProcessed );
 
-	/** Returns the number of cached view infos (thread-safe). */
-	int32 ThreadNumViews()
-	{
-		return ThreadSettings.ThreadViewInfos.Num();
-	}
-
-	/** Returns a cached view info for the specified index (thread-safe). */
-	const FStreamingViewInfo& ThreadGetView( int32 ViewIndex )
-	{
-		return ThreadSettings.ThreadViewInfos[ ViewIndex ];
-	}
-
 #if STATS
 	/** Ringbuffer of bandwidth samples for streaming in mip-levels (MB/sec). */
 	static float BandwidthSamples[NUM_BANDWIDTHSAMPLES];
@@ -326,25 +215,16 @@ struct FStreamingManagerTexture : public ITextureStreamingManager
 
 protected:
 //BEGIN: Thread-safe functions and data
-		friend class FAsyncTextureStreaming;
-		friend struct FStreamingHandlerTextureStatic;
-		friend struct FStreamingHandlerTextureDynamic;
+		friend class FAsyncTextureStreamingTask;
 
 		/** Calculates the minimum and maximum number of mip-levels for a streaming texture. */
 		void CalcMinMaxMips( FStreamingTexture& StreamingTexture );
 
-		/** Calculates the number of mip-levels we would like to have in memory for a texture. */
-		void CalcWantedMips( FStreamingTexture& StreamingTexture );
 
-		/**
-		 * Fallback handler to catch textures that have been orphaned recently.
-		 * This handler prevents massive spike in texture memory usage.
-		 * Orphaned textures were previously streamed based on distance but those instance locations have been removed -
-		 * for instance because a ULevel is being unloaded. These textures are still active and in memory until they are garbage collected,
-		 * so we must ensure that they do not start using the LastRenderTime fallback handler and suddenly stream in all their mips -
-		 * just to be destroyed shortly after by a garbage collect.
-		 */
-		FFloatMipLevel GetWantedMipsForOrphanedTexture( FStreamingTexture& StreamingTexture, float& Distance );
+		static FORCEINLINE void LogHeuristic(bool bOutputToLog, bool bUseHeuristic, const TCHAR*& Heuristic, const TCHAR* NewHeuristic);
+
+		/** Calculates the number of mip-levels we would like to have in memory for a texture. */
+		const TCHAR* CalcWantedMips( const FAsyncTextureStreamingData& StreamingData, FStreamingTexture& StreamingTexture, bool bGetHeuristic, bool bOutputToLog );
 
 		/** Updates this frame's STATs by one texture. */
 		void UpdateFrameStats( FStreamingContext& Context, FStreamingTexture& StreamingTexture, int32 TextureIndex );
@@ -399,47 +279,6 @@ protected:
 		 */
 		void DropForcedMips( FStreamingContext& Context );
 
-		/** Thread-safe helper struct for per-level information. */
-		struct FThreadLevelData
-		{
-			FThreadLevelData()
-			:	bRemove( false )
-			{
-			}
-
-			/** Whether this level has been removed. */
-			bool bRemove;
-
-			/** Texture instances used in this level, stored in SIMD-friendly format. */
-			TMap<const UTexture2D*,TArray<FStreamableTextureInstance4> > ThreadTextureInstances;
-		};
-
-		typedef TKeyValuePair< class ULevel*, FThreadLevelData >	FLevelData;
-
-		/** Thread-safe helper struct for streaming information. */
-		struct FThreadSettings
-		{
-			FThreadSettings() : TextureBoundsVisibility(nullptr) {}
-			
-			/** Cached from the system settings. */
-			int32 NumStreamedMips[TEXTUREGROUP_MAX];
-
-			/** Cached from each ULevel. */
-			TArray< FLevelData > LevelData;
-
-			/** Cached from FStreamingManagerBase. */
-			TArray<FStreamingViewInfo> ThreadViewInfos;
-
-			FTextureBoundsVisibility* TextureBoundsVisibility;
-
-			/** from cvar, >=0 */
-			float MipBias;
-		};
-
-
-		/** Thread-safe helper data for streaming information. */
-		FThreadSettings	ThreadSettings;
-
 		/** All streaming UTexture2D objects. */
 		TArray<FStreamingTexture> StreamingTextures;
 
@@ -457,6 +296,8 @@ protected:
 
 	void	DumpTextureGroupStats( bool bDetailedStats );
 
+	void	IncrementalUpdate( float Percentage );
+
 	/**
 	 * Prints out detailed information about streaming textures that has a name that contains the given string.
 	 * Triggered by the InvestigateTexture exec command.
@@ -465,8 +306,6 @@ protected:
 	 */
 	void	InvestigateTexture( const FString& InvestigateTextureName );
 
-	void	DumpTextureInstances( const UPrimitiveComponent* Primitive, UTexture2D* Texture2D );
-
 	/** Next sync, dump texture group stats. */
 	bool	bTriggerDumpTextureGroupStats;
 
@@ -479,17 +318,22 @@ protected:
 	/** Name of a texture to investigate. Can be partial name. */
 	FString	InvestigateTextureName;
 
+	/** Cached from the system settings. */
+	int32 NumStreamedMips[TEXTUREGROUP_MAX];
+
+	float MipBias;
+
 	/** Async work for calculating priorities for all textures. */
-	FAsyncTask<FAsyncTextureStreaming>*	AsyncWork;
+	FAsyncTask<FAsyncTextureStreamingTask>*	AsyncWork;
 
-	/** Textures from dynamic primitives. */
-	FDynamicComponentTextureManager* DynamicComponentTextureManager;
+	/** Textures from dynamic primitives. Owns the data for all levels. */
+	FDynamicComponentTextureManager DynamicComponentManager;
 
 	/** New textures, before they've been added to the thread-safe container. */
-	TArray<UTexture2D*>		PendingStreamingTextures;
+	TArray<UTexture2D*>	PendingStreamingTextures;
 
-	/** New levels, before they've been added to the thread-safe container. */
-	TArray<class ULevel*>	PendingLevels;
+	/** Level data */
+	TArray<FLevelTextureManager> LevelTextureManagers;
 
 	/** Stages [0,N-2] is non-threaded data collection, Stage N-1 is wait-for-AsyncWork-and-finalize. */
 	int32					ProcessingStage;
@@ -505,9 +349,6 @@ protected:
 	/** Extra distance added to the range test, to start streaming the texture before they are actually used. */
 	float					RangePrefetchDistance;
 
-	/** Array of texture streaming objects to use during update. */
-	TArray<FStreamingHandlerTextureBase*> TextureStreamingHandlers;
-
 	/** Amount of memory to leave free in the texture pool. */
 	int64					MemoryMargin;
 
@@ -618,144 +459,3 @@ protected:
 	
 	friend bool TrackTextureEvent( FStreamingTexture* StreamingTexture, UTexture2D* Texture, bool bForceMipLevelsToBeResident, const FStreamingManagerTexture* Manager);
 };
-
-
-
-
-/*-----------------------------------------------------------------------------
-	Texture streaming handler.
------------------------------------------------------------------------------*/
-
-/**
- * Base of texture streaming handler functionality.
- */
-struct FStreamingHandlerTextureBase
-{
-	/** Friendly name identifying the handler for debug purposes. */
-	const TCHAR* HandlerName;
-
-	/** Default constructor. */
-	FStreamingHandlerTextureBase()
-		: HandlerName(TEXT("Base"))
-	{
-	}
-
-	/**
-	 * Returns mip count wanted by this handler for the passed in texture. 
-	 * 
-	 * @param	StreamingManager	Streaming manager
-	 * @param	Streaming Texture	Texture to determine wanted mip count for
-	 * @param	MinDistance			[out] Distance to the closest instance of this texture, in world space units.
-	 * @return	Number of miplevels that should be streamed in or INDEX_NONE if texture isn't handled by this handler.
-	 */
-	virtual FFloatMipLevel GetWantedMips( FStreamingManagerTexture& StreamingManager, FStreamingTexture& StreamingTexture, float& MinDistance ) = 0;
-};
-
-/**
- * Static texture streaming handler. Used to stream textures on static level geometry.
- */
-struct FStreamingHandlerTextureStatic : public FStreamingHandlerTextureBase
-{
-	/** Default constructor. */
-	FStreamingHandlerTextureStatic()
-	{
-		HandlerName = TEXT("Static");
-	}
-
-	/**
-	 * Returns mip count wanted by this handler for the passed in texture. 
-	 * 
-	 * @param	StreamingManager	Streaming manager
-	 * @param	Streaming Texture	Texture to determine wanted mip count for
-	 * @param	MinDistance			[out] Distance to the closest instance of this texture, in world space units.
-	 * @return	Number of miplevels that should be streamed in or INDEX_NONE if texture isn't handled by this handler.
-	 */
-	virtual FFloatMipLevel GetWantedMips( FStreamingManagerTexture& StreamingManager, FStreamingTexture& StreamingTexture, float& MinDistance );
-};
-
-
-/**
- * Dynamic texture streaming handler. Used to stream textures for dynamic primitives.
- */
-struct FStreamingHandlerTextureDynamic : public FStreamingHandlerTextureBase
-{
-	/** Default constructor. */
-	FStreamingHandlerTextureDynamic()
-	{
-		HandlerName = TEXT("Dynamic");
-	}
-
-	/**
-	 * Returns mip count wanted by this handler for the passed in texture. 
-	 * 
-	 * @param	StreamingManager	Streaming manager
-	 * @param	Streaming Texture	Texture to determine wanted mip count for
-	 * @param	MinDistance			[out] Distance to the closest instance of this texture, in world space units.
-	 * @return	Number of miplevels that should be streamed in or INDEX_NONE if texture isn't handled by this handler.
-	 */
-	virtual FFloatMipLevel GetWantedMips( FStreamingManagerTexture& StreamingManager, FStreamingTexture& StreamingTexture, float& MinDistance );
-};
-
-/**
- * Streaming handler that bases its decision on the last render time.
- */
-struct FStreamingHandlerTextureLastRender : public FStreamingHandlerTextureBase
-{
-	/** Default constructor. */
-	FStreamingHandlerTextureLastRender()
-	{
-		HandlerName = TEXT("LastRender");
-	}
-
-	/**
-	 * Returns mip count wanted by this handler for the passed in texture. 
-	 * 
-	 * @param	StreamingManager	Streaming manager
-	 * @param	Streaming Texture	Texture to determine wanted mip count for
-	 * @param	MinDistance			[out] Distance to the closest instance of this texture, in world space units.
-	 * @return	Number of miplevels that should be streamed in or INDEX_NONE if texture isn't handled by this handler.
-	 */
-	virtual FFloatMipLevel GetWantedMips( FStreamingManagerTexture& StreamingManager, FStreamingTexture& StreamingTexture, float& MinDistance );
-};
-
-/**
- *	Streaming handler that bases its decision on the bForceMipStreaming flag in PrimitiveComponent.
- */
-struct FStreamingHandlerTextureLevelForced : public FStreamingHandlerTextureBase
-{
-	/** Default constructor. */
-	FStreamingHandlerTextureLevelForced()
-	{
-		HandlerName = TEXT("LevelForced");
-	}
-
-	/**
-	 * Returns mip count wanted by this handler for the passed in texture. 
-	 * 
-	 * @param	StreamingManager	Streaming manager
-	 * @param	Streaming Texture	Texture to determine wanted mip count for
-	 * @param	MinDistance			[out] Distance to the closest instance of this texture, in world space units.
-	 * @return	Number of miplevels that should be streamed in or INDEX_NONE if texture isn't handled by this handler.
-	 */
-	virtual FFloatMipLevel GetWantedMips( FStreamingManagerTexture& StreamingManager, FStreamingTexture& StreamingTexture, float& MinDistance );
-};
-
-/*-----------------------------------------------------------------------------
-	Texture streaming helper structs.
------------------------------------------------------------------------------*/
-
-struct FSpawnedTextureInstance
-{
-	FSpawnedTextureInstance( UTexture2D* InTexture2D, float InTexelFactor, float InOriginalRadius )
-	:	Texture2D( InTexture2D )
-	,	TexelFactor( InTexelFactor )
-	,	InvOriginalRadius( (InOriginalRadius > 0.0f) ? (1.0f/InOriginalRadius) : 1.0f )
-	{
-	}
-	/** Texture that is used by a dynamic UPrimitiveComponent. */
-	UTexture2D*		Texture2D;
-	/** Texture specific texel scale factor  */
-	float			TexelFactor;
-	/** 1.0f / OriginalBoundingSphereRadius, at the time the TexelFactor was calculated originally. */
-	float			InvOriginalRadius;
-};
\ No newline at end of file
diff --git a/Engine/Source/Runtime/Engine/Private/Streaming/StreamingTexture.cpp b/Engine/Source/Runtime/Engine/Private/Streaming/StreamingTexture.cpp
--- a/Engine/Source/Runtime/Engine/Private/Streaming/StreamingTexture.cpp
+++ b/Engine/Source/Runtime/Engine/Private/Streaming/StreamingTexture.cpp
@@ -19,27 +19,14 @@ FStreamingTexture::FStreamingTexture( UTexture2D* InTexture )
 	STAT( MostResidentMips = InTexture->ResidentMips );
 	LODGroup = (TextureGroup) InTexture->LODGroup;
 
-	bIsHLODTexture = LODGroup == TEXTUREGROUP_HierarchicalLOD;
-
-	// Landscape textures and HLOD textures should not be affected by MipBias
 	//@TODO: Long term this should probably be a property of the texture group instead
-	bCanBeAffectedByMipBias = true;
-	if ((LODGroup == TEXTUREGROUP_Terrain_Heightmap) || (LODGroup == TEXTUREGROUP_Terrain_Weightmap) || bIsHLODTexture)
-	{
-		bCanBeAffectedByMipBias = false;
-	}
-
+	bCanBeAffectedByMipBias = LODGroup != TEXTUREGROUP_Terrain_Heightmap && LODGroup != TEXTUREGROUP_Terrain_Weightmap && LODGroup != TEXTUREGROUP_HierarchicalLOD;
 	NumNonStreamingMips = InTexture->GetNumNonStreamingMips();
-	ForceLoadRefCount = 0;
 	bIsStreamingLightmap = IsStreamingLightmap( Texture );
 	bIsLightmap = bIsStreamingLightmap || LODGroup == TEXTUREGROUP_Lightmap || LODGroup == TEXTUREGROUP_Shadowmap;
-	bUsesStaticHeuristics = false;
-	bUsesDynamicHeuristics = false;
-	bUsesLastRenderHeuristics = false;
-	bUsesForcedHeuristics = false;
-	bUsesOrphanedHeuristics = false;
 	bHasSplitRequest = false;
 	bIsLastSplitRequest = false;
+	bIsVisibleWantedMips = false;
 	BoostFactor = 1.0f;
 	InstanceRemovedTimestamp = -FLT_MAX;
 	LastRenderTimeRefCountTimestamp = -FLT_MAX;
@@ -55,48 +42,58 @@ FStreamingTexture::FStreamingTexture( UTexture2D* InTexture )
 
 void FStreamingTexture::UpdateCachedInfo( )
 {
+	static auto CVarOnlyStreamInTextures = IConsoleManager::Get().FindTConsoleVariableDataInt(TEXT("r.OnlyStreamInTextures"));
+	const bool bOnlyStreamIn = CVarOnlyStreamInTextures->GetValueOnAnyThread() != 0;
+	const int32 HLODStategy = CVarStreamingHLODStrategy.GetValueOnAnyThread();
+
 	check(Texture && Texture->Resource);
 	ResidentMips = Texture->ResidentMips;
 	RequestedMips = Texture->RequestedMips;
 	MinAllowedMips = 1;			//FMath::Min(Texture->ResidentMips, Texture->RequestedMips);
 	MaxAllowedMips = MipCount;	//FMath::Max(Texture->ResidentMips, Texture->RequestedMips);
-	MaxAllowedOptimalMips = MaxAllowedMips;
 	STAT( MostResidentMips = FMath::Max(MostResidentMips, Texture->ResidentMips) );
 	float LastRenderTimeForTexture = Texture->GetLastRenderTimeForStreaming();
 	LastRenderTime = (FApp::GetCurrentTime() > LastRenderTimeForTexture) ? float( FApp::GetCurrentTime() - LastRenderTimeForTexture ) : 0.0f;
-	MinDistance = MAX_STREAMINGDISTANCE;
-	bForceFullyLoad = Texture->ShouldMipLevelsBeForcedResident() || (ForceLoadRefCount > 0);
+
+	bForceFullyLoad = Texture->ShouldMipLevelsBeForcedResident() || (bOnlyStreamIn && LastRenderTime < 300)
+		|| LODGroup == TEXTUREGROUP_Skybox || (LODGroup == TEXTUREGROUP_HierarchicalLOD && HLODStategy == 2);
+
 	TextureLODBias = Texture->GetCachedLODBias();
 	bInFlight = false;
 	bReadyForStreaming = IsStreamingTexture( Texture ) && IsReadyForStreaming( Texture );
 	NumCinematicMipLevels = Texture->bUseCinematicMipLevels ? Texture->NumCinematicMipLevels : 0;
 }
 
-ETextureStreamingType FStreamingTexture::GetStreamingType() const
+/** Set the wanted mips from the async task data */
+void FStreamingTexture::SetPerfectWantedMips(float MaxSize, float MaxSize_VisibleOnly, float MipBias)
 {
-	if ( bUsesForcedHeuristics || bForceFullyLoad )
-	{
-		return StreamType_Forced;
-	}
-	else if ( bUsesDynamicHeuristics )
-	{
-		return StreamType_Dynamic;
-	}
-	else if ( bUsesStaticHeuristics )
-	{
-		return StreamType_Static;
-	}
-	else if ( bUsesOrphanedHeuristics )
+	const float HiddenScale = CVarStreamingHiddenPrimitiveScale.GetValueOnAnyThread();
+	const float MaxResolution = (float)(0x1 << (MaxAllowedMips - 1));
+	const float ExtraScale = GetStreamingScale();
+
+	// Do this now before clamping the resolutions.
+	bAsNeverStream = (MaxSize == FLT_MAX || MaxSize_VisibleOnly == FLT_MAX);
+
+	MaxSize_VisibleOnly = FMath::Clamp<float>(MaxSize_VisibleOnly * ExtraScale, 1.f, MaxResolution);
+	MaxSize = FMath::Clamp<float>(MaxSize * ExtraScale, MaxSize_VisibleOnly, MaxResolution) * HiddenScale;
+
+	int32 WantedMips_VisibleOnly = GetWantedMipsFromSize(MaxSize_VisibleOnly, MipBias);
+	int32 WantedMips_Any = GetWantedMipsFromSize(MaxSize, MipBias);
+
+	// Load visible mips first, then non visible mips.
+	if (WantedMips_Any > WantedMips_VisibleOnly && WantedMips_VisibleOnly <= ResidentMips)
 	{
-		return StreamType_Orphaned;
+		PerfectWantedMips = WantedMips = WantedMips_Any;
+		bIsVisibleWantedMips = false;
 	}
-	else if ( bUsesLastRenderHeuristics )
+	else
 	{
-		return StreamType_LastRenderTime;
+		PerfectWantedMips = WantedMips = WantedMips_VisibleOnly;
+		bIsVisibleWantedMips = true;
 	}
-	return StreamType_Other;
 }
 
+
 bool FStreamingTexture::UpdateMipCount(  FStreamingManagerTexture& Manager, FStreamingContext& Context )
 {
 	if ( !Texture || !bReadyForStreaming )
@@ -155,8 +152,8 @@ bool FStreamingTexture::UpdateMipCount(  FStreamingManagerTexture& Manager, FStr
 				RequestedMips = WantedMips;
 				Texture->RequestedMips = RequestedMips;
 
-				Texture2DResource->BeginUpdateMipCount( bForceFullyLoad );
-				TrackTextureEvent( this, Texture, bForceFullyLoad, &Manager );
+				Texture2DResource->BeginUpdateMipCount( bAsNeverStream );
+				TrackTextureEvent( this, Texture, bAsNeverStream, &Manager );
 				return true;
 			}
 		}
@@ -173,9 +170,9 @@ float FStreamingTexture::CalcRetentionPriority( )
 {
 	const int32 WantedMipsForPriority = GetWantedMipsForPriority();
 
-	bool bMustKeep = LODGroup == TEXTUREGROUP_Terrain_Heightmap || bForceFullyLoad || WantedMipsForPriority <= MinAllowedMips;
+	bool bMustKeep = bAsNeverStream || WantedMipsForPriority <= MinAllowedMips || LODGroup == TEXTUREGROUP_Terrain_Heightmap;
 	bool bIsVisible = !GStreamWithTimeFactor || IsVisible();
-	bool bAlreadyLowRez = /*(PixelPerTexel > 1.5 && Distance > 1) ||*/ bIsHLODTexture || bIsLightmap;
+	bool bAlreadyLowRez = /*(PixelPerTexel > 1.5 && Distance > 1) ||*/ LODGroup == TEXTUREGROUP_HierarchicalLOD || bIsLightmap;
 		
 	// Don't consider dropping mips that give nothing. This is because streaming time has a high cost per texture.
 	// Dropping many textures for not so much memory would affect the time it takes the streamer to recover.
@@ -195,10 +192,10 @@ float FStreamingTexture::CalcLoadOrderPriority()
 {
 	const int32 WantedMipsForPriority = GetWantedMipsForPriority();
 
-	bool bMustLoadFirst = LODGroup == TEXTUREGROUP_Terrain_Heightmap || bForceFullyLoad;
+	bool bMustLoadFirst = bAsNeverStream || LODGroup == TEXTUREGROUP_Terrain_Heightmap;
 	bool bIsVisible = !GStreamWithTimeFactor || IsVisible();
 	bool bFarFromTargetMips = WantedMipsForPriority - ResidentMips > 2;
-	bool bBigOnScreen = WantedMipsForPriority > 9.f || bIsHLODTexture;
+	bool bBigOnScreen = WantedMipsForPriority > 9.f || LODGroup == TEXTUREGROUP_HierarchicalLOD;
 
 	float Priority = 0;
 	if (bMustLoadFirst)		Priority += 1024.f;
@@ -209,6 +206,27 @@ float FStreamingTexture::CalcLoadOrderPriority()
 	return Priority;
 }
 
+int32 FStreamingTexture::GetWantedMipsFromSize(float Size, float Bias) const
+{
+	const bool bUseNewMetrics = CVarStreamingUseNewMetrics.GetValueOnAnyThread() != 0;
+
+	float WantedMipsFloat = 1.0f + FMath::Log2(Size);
+	if (bCanBeAffectedByMipBias)
+	{
+		WantedMipsFloat -= Bias;
+	}
+
+	if (bUseNewMetrics && !bIsLightmap) // Lightmap have approximate logic already so don't apply a bias here.
+	{
+		// Allow a little threshold before going up on mip (so that 513 does not stream a 1024)
+		WantedMipsFloat -= .2f;
+	}
+
+	// Using the new metrics has very conservative distance computation. Rounding to mip here compensates.
+	int32 WantedMipsInt = FMath::CeilToInt(WantedMipsFloat - .01f);
+	return FMath::Clamp<int32>(WantedMipsInt, MinAllowedMips, MaxAllowedMips);
+}
+
 bool FStreamingTexture::IsReadyForStreaming( UTexture2D* Texture )
 {
 	return Texture->IsReadyForStreaming();
