Copying //UE4/Dev-Rendering to //UE4/Dev-Main (Source: //UE4/Dev-Rendering @ 2981742)

diff --git a/Engine/Source/Runtime/Engine/Private/Streaming/StreamingTexture.h b/Engine/Source/Runtime/Engine/Private/Streaming/StreamingTexture.h
--- a/Engine/Source/Runtime/Engine/Private/Streaming/StreamingTexture.h
+++ b/Engine/Source/Runtime/Engine/Private/Streaming/StreamingTexture.h
@@ -20,8 +20,6 @@ struct FStreamingTexture
 	/** Reinitializes the cached variables from UTexture2D. */
 	void UpdateCachedInfo( );
 
-	ETextureStreamingType GetStreamingType() const;
-
 	/**
 	 * Checks whether a UTexture2D is ready for streaming.
 	 *
@@ -57,7 +55,7 @@ struct FStreamingTexture
 	 */
 	FORCEINLINE bool IsVisible() const
 	{
-		return LastRenderTime < 5.f && LastRenderTime != MAX_FLT;
+		return LastRenderTime < 5.f && LastRenderTime != MAX_FLT && bIsVisibleWantedMips;
 	}
 	
 	/**
@@ -67,7 +65,7 @@ struct FStreamingTexture
 	 */
 	FORCEINLINE bool CanDropMips() const
 	{
-		return (LODGroup != TEXTUREGROUP_Terrain_Heightmap) && WantedMips > MinAllowedMips && !bForceFullyLoad;
+		return (LODGroup != TEXTUREGROUP_Terrain_Heightmap) && WantedMips > MinAllowedMips && !bAsNeverStream;
 	}
 
 
@@ -88,6 +86,16 @@ struct FStreamingTexture
 		}
 	}
 
+	FORCEINLINE float GetStreamingScale() const
+	{
+		if (LODGroup == TEXTUREGROUP_Lightmap)
+			return BoostFactor * GLightmapStreamingFactor;
+		else if (LODGroup == TEXTUREGROUP_Shadowmap)
+			return BoostFactor * GShadowmapStreamingFactor;
+		else 
+			return BoostFactor;
+	}
+
 	/**
 	 * Calculates a retention priority value for the textures. Higher value means more important.
 	 * Retention is used determine which texture to drop when out of texture pool memory.
@@ -104,7 +112,9 @@ struct FStreamingTexture
 	 */
 	float CalcLoadOrderPriority();
 
-		
+	int32 GetWantedMipsFromSize(float Size, float Bias) const;
+
+
 	/**
 	 * Not thread-safe: Sets the streaming index in the corresponding UTexture2D.
 	 * @param ArrayIndex	Index into the FStreamingManagerTexture::StreamingTextures array
@@ -114,6 +124,9 @@ struct FStreamingTexture
 		Texture->StreamingIndex = ArrayIndex;
 	}
 
+	/** Set the wanted mips from the async task data */
+	void SetPerfectWantedMips(float MaxSize, float MaxSize_VisibleOnly, float MipBias);
+
 	/** Update texture streaming. Returns whether it did anything */
 	bool UpdateMipCount(  FStreamingManagerTexture& Manager, FStreamingContext& Context );
 
@@ -130,8 +143,6 @@ struct FStreamingTexture
 	int32			WantedMips;
 	/** Legacy: Same as WantedMips, but not clamped by fudge factor. */
 	int32			PerfectWantedMips;
-	/** Same as MaxAllowedMips, but not clamped by the -reducepoolsize command line option. */
-	int32			MaxAllowedOptimalMips;
 #if STATS
 	/** Most number of mip-levels this texture has ever had resident in memory. */
 	int32			MostResidentMips;
@@ -142,8 +153,6 @@ struct FStreamingTexture
 	int32			MaxAllowedMips;
 	/** Cached memory sizes for each possible mipcount. */
 	int32			TextureSizes[MAX_TEXTURE_MIP_COUNT + 1];
-	/** Ref-count for how many ULevels want this texture fully-loaded. */
-	int32			ForceLoadRefCount;
 
 	/** Cached texture group. */
 	TextureGroup	LODGroup;
@@ -156,8 +165,6 @@ struct FStreamingTexture
 
 	/** Last time this texture was rendered, 0-based from GStartTime. */
 	float			LastRenderTime;
-	/** Distance to the closest instance of this texture, in world space units. Defaults to a specific value for non-staticmesh textures. */
-	float			MinDistance;
 	/** If non-zero, the most recent time an instance location was removed for this texture. */
 	double			InstanceRemovedTimestamp;
 
@@ -172,7 +179,9 @@ struct FStreamingTexture
 	 */
 	float			BoostFactor;
 
-	/** Whether the texture should be forcibly fully loaded. */
+	/** Whether the texture should be handled as a never stream. This is the final value and results from the CalcWantedMips analysis*/
+	uint32			bAsNeverStream : 1;
+	/** Whether the texture should be forcibly fully loaded. This comes from the texture only and ignores things like texture group settings and components  */
 	uint32			bForceFullyLoad : 1;
 	/** Whether the texture is ready to be streamed in/out (cached from IsReadyForStreaming()). */
 	uint32			bReadyForStreaming : 1;
@@ -182,23 +191,13 @@ struct FStreamingTexture
 	uint32			bIsStreamingLightmap : 1;
 	/** Whether this is a lightmap or shadowmap. */
 	uint32			bIsLightmap : 1;
-	/** Whether this texture uses StaticTexture heuristics. */
-	uint32			bUsesStaticHeuristics : 1;
-	/** Whether this texture uses DynamicTexture heuristics. */
-	uint32			bUsesDynamicHeuristics : 1;
-	/** Whether this texture uses LastRenderTime heuristics. */
-	uint32			bUsesLastRenderHeuristics : 1;
-	/** Whether this texture uses ForcedIntoMemory heuristics. */
-	uint32			bUsesForcedHeuristics : 1;
-	/** Whether this texture uses the OrphanedTexture heuristics. */
-	uint32			bUsesOrphanedHeuristics : 1;
 	/** Whether this texture has a split request strategy. */
 	uint32			bHasSplitRequest : 1;
 	/** Whether this is the second request to converge to the wanted mip. */
 	uint32			bIsLastSplitRequest : 1;
 	/** Can this texture be affected by mip bias? */
 	uint32			bCanBeAffectedByMipBias : 1;
-	/** Is this texture a HLOD texture? */
-	//@TODO: Temporary: Done because there is currently no texture group for HLOD textures and adding one in a branch isn't ideal; this should be removed / replaced once there is a real HLOD texture group
-	uint32			bIsHLODTexture : 1;
+	/** Is the wanted mip for a visible wanted mips */
+	uint32			bIsVisibleWantedMips : 1;
+
 };
diff --git a/Engine/Source/Runtime/Engine/Private/Streaming/TextureInstanceManager.cpp b/Engine/Source/Runtime/Engine/Private/Streaming/TextureInstanceManager.cpp
--- a/Engine/Source/Runtime/Engine/Private/Streaming/TextureInstanceManager.cpp
+++ b/Engine/Source/Runtime/Engine/Private/Streaming/TextureInstanceManager.cpp
@@ -7,7 +7,7 @@
 #include "EnginePrivate.h"
 #include "TextureInstanceManager.h"
 
-FTextureInstanceManager::FBounds4::FBounds4()
+FTextureInstanceState::FBounds4::FBounds4()
 :	OriginX( 0, 0, 0, 0 )
 ,	OriginY( 0, 0, 0, 0 )
 ,	OriginZ( 0, 0, 0, 0 )
@@ -16,13 +16,13 @@ FTextureInstanceManager::FBounds4::FBounds4()
 ,	ExtentZ( 0, 0, 0, 0 )
 ,	Radius( 0, 0, 0, 0 )
 ,	MinDistanceSq( 0, 0, 0, 0 )
-,	MaxDistanceSq( MAX_FLT, MAX_FLT, MAX_FLT, MAX_FLT )
-,	LastRenderTime(0, 0, 0, 0)
+,	MaxDistanceSq( FLT_MAX, FLT_MAX, FLT_MAX, FLT_MAX )
+,	LastRenderTime( -FLT_MAX, -FLT_MAX, -FLT_MAX, -FLT_MAX)
 {
 }
 
 
-void FTextureInstanceManager::FBounds4::Clear(int32 Index)
+void FTextureInstanceState::FBounds4::Clear(int32 Index)
 {
 	check(Index >= 0 && Index < 4);
 
@@ -34,11 +34,11 @@ void FTextureInstanceManager::FBounds4::Clear(int32 Index)
 	ExtentZ.Component(Index) = 0;
 	Radius.Component(Index) = 0;
 	MinDistanceSq.Component(Index) = 0;
-	MaxDistanceSq.Component(Index) = MAX_FLT;
-	LastRenderTime.Component(Index) = 0;
+	MaxDistanceSq.Component(Index) = FLT_MAX;
+	LastRenderTime.Component(Index) = -FLT_MAX;
 }
 
-void FTextureInstanceManager::FBounds4::Set(int32 Index, const FBoxSphereBounds& Bounds, float InLastRenderTime, float InMinDistance, float InMaxDistance)
+void FTextureInstanceState::FBounds4::Set(int32 Index, const FBoxSphereBounds& Bounds, float InLastRenderTime, float InMinDistance, float InMaxDistance)
 {
 	check(Index >= 0 && Index < 4);
 
@@ -50,11 +50,34 @@ void FTextureInstanceManager::FBounds4::Set(int32 Index, const FBoxSphereBounds&
 	ExtentZ.Component(Index) = Bounds.BoxExtent.Z;
 	Radius.Component(Index) = Bounds.SphereRadius;
 	MinDistanceSq.Component(Index) = InMinDistance * InMinDistance;
-	MaxDistanceSq.Component(Index) = InMaxDistance != MAX_FLT ? (InMaxDistance * InMaxDistance) : MAX_FLT;
+	MaxDistanceSq.Component(Index) = InMaxDistance != FLT_MAX ? (InMaxDistance * InMaxDistance) : FLT_MAX;
 	LastRenderTime.Component(Index) = InLastRenderTime;
 }
 
-FTextureInstanceManager::FElement::FElement()
+void FTextureInstanceState::FBounds4::Update(int32 Index, const FBoxSphereBounds& Bounds, float InLastRenderTime)
+{
+	check(Index >= 0 && Index < 4);
+
+	OriginX.Component(Index) = Bounds.Origin.X;
+	OriginY.Component(Index) = Bounds.Origin.Y;
+	OriginZ.Component(Index) = Bounds.Origin.Z;
+	ExtentX.Component(Index) = Bounds.BoxExtent.X;
+	ExtentY.Component(Index) = Bounds.BoxExtent.Y;
+	ExtentZ.Component(Index) = Bounds.BoxExtent.Z;
+	Radius.Component(Index) = Bounds.SphereRadius;
+	MinDistanceSq.Component(Index) = 0;
+	MaxDistanceSq.Component(Index) = FLT_MAX;
+	LastRenderTime.Component(Index) = InLastRenderTime;
+}
+
+void FTextureInstanceState::FBounds4::Update(int32 Index, float InLastRenderTime)
+{
+	check(Index >= 0 && Index < 4);
+
+	LastRenderTime.Component(Index) = InLastRenderTime;
+}
+
+FTextureInstanceState::FElement::FElement()
 :	Component(nullptr)
 ,	Texture(nullptr)
 ,	BoundsIndex(INDEX_NONE)
@@ -65,9 +88,84 @@ FTextureInstanceManager::FElement::FElement()
 {
 }
 
-int32 FTextureInstanceManager::AddBounds(const FBoxSphereBounds& Bounds, const UPrimitiveComponent* InComponent, float LastRenderTime, float MinDistance, float MaxDistance)
+FTextureInstanceState::FTextureLinkConstIterator::FTextureLinkConstIterator(const FTextureInstanceState& InState, const UTexture2D* InTexture) 
+:	State(InState)
+,	CurrElementIndex(INDEX_NONE)
 {
-	check(bAllowNullComponents || InComponent);
+	const int32* TextureLink = State.TextureMap.Find(InTexture);
+	if (TextureLink)
+	{
+		CurrElementIndex = *TextureLink;
+	}
+}
+
+FBoxSphereBounds FTextureInstanceState::FTextureLinkConstIterator::GetBounds() const
+{ 
+	FBoxSphereBounds Bounds(ForceInitToZero);
+
+	int32 BoundsIndex = State.Elements[CurrElementIndex].BoundsIndex; 
+	if (BoundsIndex != INDEX_NONE)
+	{
+		const FBounds4& TheBounds4 = State.Bounds4[BoundsIndex / 4];
+		int32 Index = BoundsIndex % 4;
+
+		Bounds.Origin.X = TheBounds4.OriginX[Index];
+		Bounds.Origin.Y = TheBounds4.OriginY[Index];
+		Bounds.Origin.Z = TheBounds4.OriginZ[Index];
+
+		Bounds.BoxExtent.X = TheBounds4.ExtentX[Index];
+		Bounds.BoxExtent.Y = TheBounds4.ExtentY[Index];
+		Bounds.BoxExtent.Z = TheBounds4.ExtentZ[Index];
+
+		Bounds.SphereRadius = TheBounds4.Radius[Index];
+	}
+	return Bounds;
+}
+
+void FTextureInstanceState::FTextureLinkConstIterator::OutputToLog(float MaxNormalizedSize, float MaxNormalizedSize_VisibleOnly) const
+{
+	const UPrimitiveComponent* Component = GetComponent();
+	FBoxSphereBounds Bounds = GetBounds();
+	float TexelFactor = GetTexelFactor();
+
+	// Log the component reference.
+	FString ReferenceName(TEXT("Unkown Component"));
+	if (Component)
+	{
+		ReferenceName = Component->GetFullName();
+	}
+	UE_LOG(LogContentStreaming, Log, TEXT("    Reference= %s"), *ReferenceName);
+
+	// Log the wanted mips.
+	if (TexelFactor == FLT_MAX)
+	{
+		UE_LOG(LogContentStreaming, Log, TEXT("    Forcing all mips."), *ReferenceName);
+	}
+	else if (TexelFactor >= 0)
+	{
+		UE_LOG(LogContentStreaming, Log, TEXT("    TexelFactor=%f, MaxSize=%f, MaxVisibleSize=%f"), TexelFactor, TexelFactor * MaxNormalizedSize, TexelFactor * MaxNormalizedSize_VisibleOnly);
+	}
+	else
+	{
+		UE_LOG(LogContentStreaming, Log, TEXT("    Forced Size=%f"), -TexelFactor);
+	}
+
+	// Log the bounds
+	const bool bUseNewMetrics = CVarStreamingUseNewMetrics.GetValueOnAnyThread() != 0;
+	if (bUseNewMetrics)
+	{
+		UE_LOG(LogContentStreaming, Log, TEXT("    Origin=(%s), BoxExtent=(%s)"), *Bounds.Origin.ToString(), *Bounds.BoxExtent.ToString());
+	}
+	else
+	{
+		UE_LOG(LogContentStreaming, Log, TEXT("    Origin=(%s), SphereRadius=%f"),  *Bounds.Origin.ToString(), Bounds.SphereRadius);
+	}
+}
+
+
+int32 FTextureInstanceState::AddBounds(const FBoxSphereBounds& Bounds, const UPrimitiveComponent* InComponent, float LastRenderTime, float MinDistance, float MaxDistance)
+{
+	check(InComponent);
 
 	int BoundsIndex = INDEX_NONE;
 
@@ -86,16 +184,9 @@ int32 FTextureInstanceManager::AddBounds(const FBoxSphereBounds& Bounds, const U
 		Bounds4Components.Push(nullptr);
 
 		// Since each element contains 4 entries, add the 3 unused ones
-		FreeBoundIndices.Push(BoundsIndex + 1);
-		FreeBoundIndices.Push(BoundsIndex + 2);
 		FreeBoundIndices.Push(BoundsIndex + 3);
-
-		// This is just to be compatible with : check(Instances.FirstDirtyComponentIndex == Instances.Bounds4Components.Num());
-		// Whenever new bounds are allocated, they are automatically up to date.
-		if (FirstDirtyComponentIndex == BoundsIndex)
-		{
-			FirstDirtyComponentIndex += 4;
-		}
+		FreeBoundIndices.Push(BoundsIndex + 2);
+		FreeBoundIndices.Push(BoundsIndex + 1);
 	}
 
 	Bounds4[BoundsIndex / 4].Set(BoundsIndex % 4, Bounds, LastRenderTime, MinDistance, MaxDistance);
@@ -104,7 +195,7 @@ int32 FTextureInstanceManager::AddBounds(const FBoxSphereBounds& Bounds, const U
 	return BoundsIndex;
 }
 
-void FTextureInstanceManager::RemoveBounds(int32 BoundsIndex)
+void FTextureInstanceState::RemoveBounds(int32 BoundsIndex)
 {
 	check(BoundsIndex != INDEX_NONE);
 	checkSlow(!FreeBoundIndices.Contains(BoundsIndex));
@@ -124,9 +215,51 @@ void FTextureInstanceManager::RemoveBounds(int32 BoundsIndex)
 	}
 }
 
-void FTextureInstanceManager::AddElement(const UPrimitiveComponent* InComponent, const UTexture2D* InTexture, int InBoundsIndex, float InTexelFactor, int32*& ComponentLink)
+void FTextureInstanceState::AddElement(const UPrimitiveComponent* InComponent, const UTexture2D* InTexture, int InBoundsIndex, float InTexelFactor, int32*& ComponentLink)
 {
-	check(bAllowNullComponents || InComponent);
+	check(InComponent && InTexture);
+	CompiledTextureMap.Empty(); // Data is invalid now!
+
+	int32* TextureLink = TextureMap.Find(InTexture);
+
+	// Since textures are processed per component, if there are already elements for this component-texture,
+	// they will be in the first entries (as we push to head). If such pair use the same bound, merge the texel factors.
+	// The first step is to find if such a duplicate entries exists.
+
+	if (TextureLink)
+	{
+		int32 ElementIndex = *TextureLink;
+		while (ElementIndex != INDEX_NONE)
+		{
+			FElement& TextureElement = Elements[ElementIndex];
+
+			if (TextureElement.Component == InComponent)
+			{
+				if (TextureElement.BoundsIndex == InBoundsIndex)
+				{
+					if (InTexelFactor >= 0 && TextureElement.TexelFactor >= 0)
+					{
+						// Abort inserting a new element, and merge the 2 entries together.
+						TextureElement.TexelFactor = FMath::Max(TextureElement.TexelFactor, InTexelFactor);
+						return;
+					}
+					else if (InTexelFactor < 0 && TextureElement.TexelFactor < 0)
+					{
+						// Negative are forced resolution.
+						TextureElement.TexelFactor = FMath::Min(TextureElement.TexelFactor, InTexelFactor);
+						return;
+					}
+				}
+
+				// Check the next bounds for this component.
+				ElementIndex = TextureElement.NextTextureLink;
+			}
+			else
+			{
+				break;
+			}
+		}
+	}
 
 	int32 ElementIndex = INDEX_NONE;
 	if (FreeElementIndices.Num())
@@ -146,54 +279,38 @@ void FTextureInstanceManager::AddElement(const UPrimitiveComponent* InComponent,
 	Element.BoundsIndex = InBoundsIndex;
 	Element.TexelFactor = InTexelFactor;
 
-	if (InTexture)
+	if (TextureLink)
 	{
-		int32* TextureLink = TextureMap.Find(InTexture);
-		if (TextureLink)
-		{
-			FElement& TextureLinkElement = Elements[*TextureLink];
-			if (TextureLinkElement.Component == InComponent && 
-				TextureLinkElement.BoundsIndex == InBoundsIndex)
-			{
-				// Abort inserting a new element, and merge the 2 entries together.
-				TextureLinkElement.TexelFactor = FMath::Max(TextureLinkElement.TexelFactor, Element.TexelFactor);
-
-				// Clear the element and insert it in the free list.
-				Element = FElement();
-				FreeElementIndices.Push(ElementIndex);
-				return;
-			}
+		FElement& TextureLinkElement = Elements[*TextureLink];
 
-			// The new inserted element as the head element.
-			Element.NextTextureLink = *TextureLink;
-			TextureLinkElement.PrevTextureLink = ElementIndex;
-			*TextureLink = ElementIndex;
-		}
-		else
-		{
-			TextureMap.Add(InTexture, ElementIndex);
-		}
+		// The new inserted element as the head element.
+		Element.NextTextureLink = *TextureLink;
+		TextureLinkElement.PrevTextureLink = ElementIndex;
+		*TextureLink = ElementIndex;
+	}
+	else
+	{
+		TextureMap.Add(InTexture, ElementIndex);
 	}
 
-	if (InComponent)
+	// Simple sanity check to ensure that the component link passed in param is the right one
+	checkSlow(ComponentLink == ComponentMap.Find(InComponent));
+	if (ComponentLink)
 	{
-		checkSlow(ComponentLink == ComponentMap.Find(InComponent));
-		
-		if (ComponentLink)
-		{
-			// The new inserted element as the head element.
-			Element.NextComponentLink = *ComponentLink;
-			*ComponentLink = ElementIndex;
-		}
-		else
-		{
-			ComponentLink = &ComponentMap.Add(InComponent, ElementIndex);
-		}
+		// The new inserted element as the head element.
+		Element.NextComponentLink = *ComponentLink;
+		*ComponentLink = ElementIndex;
+	}
+	else
+	{
+		ComponentLink = &ComponentMap.Add(InComponent, ElementIndex);
 	}
 }
 
-void FTextureInstanceManager::RemoveElement(int32 ElementIndex, int32& NextComponentLink, int32& BoundsIndex, const UTexture2D*& Texture)
+void FTextureInstanceState::RemoveElement(int32 ElementIndex, int32& NextComponentLink, int32& BoundsIndex, const UTexture2D*& Texture)
 {
+	CompiledTextureMap.Empty(); // Data is invalid now!
+
 	FElement& Element = Elements[ElementIndex];
 	NextComponentLink = Element.NextComponentLink; 
 	BoundsIndex = Element.BoundsIndex; 
@@ -226,7 +343,6 @@ void FTextureInstanceManager::RemoveElement(int32 ElementIndex, int32& NextCompo
 		}
 	}
 
-
 	// Clear the element and insert in free list.
 	if (1 + FreeElementIndices.Num() != Elements.Num())
 	{
@@ -240,37 +356,82 @@ void FTextureInstanceManager::RemoveElement(int32 ElementIndex, int32& NextCompo
 	}
 }
 
-void FTextureInstanceManager::AddComponent(const UPrimitiveComponent* Component)
+static bool operator==(const FBoxSphereBounds& A, const FBoxSphereBounds& B)
+{
+	return A.Origin == B.Origin && A.BoxExtent == B.BoxExtent &&  A.SphereRadius == B.SphereRadius;
+}
+
+void FTextureInstanceState::AddComponent(const UPrimitiveComponent* Component, FStreamingTextureLevelContext& LevelContext)
 {
-	FStreamingTextureLevelContext LevelContext;
 	TArray<FStreamingTexturePrimitiveInfo> TextureInstanceInfos;
 	Component->GetStreamingTextureInfoWithNULLRemoval(LevelContext, TextureInstanceInfos);
+	// Texture entries are garantied to be relevant here!
 
-	bool bHasEntries = false;
-
-	int32 BoundsIndex = AddBounds(Component->Bounds, Component, Component->LastRenderTime);
-	int32* ComponentLink = ComponentMap.Find(Component);
-	for (int32 TextureIndex = 0; TextureIndex  < TextureInstanceInfos.Num(); ++TextureIndex)
+	if (TextureInstanceInfos.Num())
 	{
-		const FStreamingTexturePrimitiveInfo& Info = TextureInstanceInfos[TextureIndex];
-		if (Info.TexelFactor > 0.0f && Info.Bounds.SphereRadius > 0.0f && ensure(FMath::IsFinite(Info.TexelFactor)))
+		const UPrimitiveComponent* LODParent = Component->GetLODParentPrimitive();
+
+		// AddElement will handle duplicate texture-bound-component. Here we just have to prevent creating identical bound-component.
+		TArray<int32, TInlineAllocator<12> > BoundsIndices;
+
+		int32* ComponentLink = ComponentMap.Find(Component);
+		for (int32 TextureIndex = 0; TextureIndex  < TextureInstanceInfos.Num(); ++TextureIndex)
 		{
-			const UTexture2D* Texture2D = Cast<UTexture2D>(Info.Texture);
-			if (Texture2D && IsStreamingTexture(Texture2D))
+			const FStreamingTexturePrimitiveInfo& Info = TextureInstanceInfos[TextureIndex];
+
+			int32 BoundsIndex = INDEX_NONE;
+
+			// Find an identical bound, or create a new entry.
 			{
-				AddElement(Component, Texture2D, BoundsIndex, Info.TexelFactor, ComponentLink);
-				bHasEntries = true;
+				for (int32 TestIndex = TextureIndex - 1; TestIndex >= 0; --TestIndex)
+				{
+					const FStreamingTexturePrimitiveInfo& TestInfo = TextureInstanceInfos[TestIndex];
+					if (Info.Bounds == TestInfo.Bounds && BoundsIndices[TestIndex] != INDEX_NONE)
+					{
+						BoundsIndex = BoundsIndices[TestIndex];
+						break;
+					}
+				}
+
+				if (BoundsIndex == INDEX_NONE)
+				{
+					// Take into account the offset between the bounds and the components.
+					float MinDistance = FMath::Min<float>(0, Component->MinDrawDistance - (Info.Bounds.Origin - Component->Bounds.Origin).Size());
+					float MaxDistance = FLT_MAX;
+					if (LODParent)
+					{
+						// Max distance when HLOD becomes visible.
+						MaxDistance = LODParent->MinDrawDistance + (Info.Bounds.Origin - LODParent->Bounds.Origin).Size();
+					}
+					BoundsIndex = AddBounds(Info.Bounds, Component, Component->LastRenderTime, MinDistance, MaxDistance);
+				}
+				BoundsIndices.Push(BoundsIndex);
 			}
+
+			// Handle force mip streaming by over scaling the texel factor. 
+			AddElement(Component, Info.Texture, BoundsIndex, Component->bForceMipStreaming ? FLT_MAX : Info.TexelFactor, ComponentLink);
 		}
 	}
+}
 
-	if (!bHasEntries)
+void FTextureInstanceState::AddComponent(const UPrimitiveComponent* Component)
+{
+	FStreamingTextureLevelContext LevelContext;
+	TArray<FStreamingTexturePrimitiveInfo> TextureInstanceInfos;
+	Component->GetStreamingTextureInfoWithNULLRemoval(LevelContext, TextureInstanceInfos);
+
+	if (TextureInstanceInfos.Num())
 	{
-		RemoveBounds(BoundsIndex);
+		int32 BoundsIndex = AddBounds(Component->Bounds, Component, Component->LastRenderTime);
+		int32* ComponentLink = ComponentMap.Find(Component);
+		for (const FStreamingTexturePrimitiveInfo& Info : TextureInstanceInfos)
+		{
+			AddElement(Component, Info.Texture, BoundsIndex, Info.TexelFactor, ComponentLink);
+		}
 	}
 }
 
-void FTextureInstanceManager::RemoveComponent(const UPrimitiveComponent* Component, FRemovedTextureArray& RemovedTextures)
+void FTextureInstanceState::RemoveComponent(const UPrimitiveComponent* Component, FRemovedTextureArray& RemovedTextures)
 {
 	TArray<int32, TInlineAllocator<12> > RemovedBoundsIndices;
 	int32 ElementIndex = INDEX_NONE;
@@ -292,104 +453,216 @@ void FTextureInstanceManager::RemoveComponent(const UPrimitiveComponent* Compone
 		{
 			RemovedTextures.AddUnique(Texture);
 		}
-
 	};
 
-	for(int32 I = 0; I < RemovedBoundsIndices.Num(); ++I)
+	for (int32 I = 0; I < RemovedBoundsIndices.Num(); ++I)
 	{
 		RemoveBounds(RemovedBoundsIndices[I]);
 	}
 }
 
-void FTextureInstanceManager::UpdateBounds(const UPrimitiveComponent* Component)
-{
+void FTextureInstanceState::RemoveComponentReferences(const UPrimitiveComponent* Component) 
+{ 
+	// Because the async task could be running, we can't change the async view state. 
+	// We limit ourself to clearing the component ptr to avoid invalid access when updating visibility.
+
 	int32* ComponentLink = ComponentMap.Find(Component);
 	if (ComponentLink)
 	{
-		int32 BoundsIndex = Elements[*ComponentLink].BoundsIndex;
+		int32 ElementIndex = *ComponentLink;
+		while (ElementIndex != INDEX_NONE)
+		{
+			FElement& Element = Elements[ElementIndex];
+			if (Element.BoundsIndex != INDEX_NONE)
+			{
+				Bounds4Components[Element.BoundsIndex] = nullptr;
+			}
+			Element.Component = nullptr;
 
-		// For primitives that can update bounds, there is only one bound per component.
-		Bounds4[BoundsIndex / 4].Set(BoundsIndex % 4, Component->Bounds, Component->LastRenderTime);
+			ElementIndex = Element.NextComponentLink;
+		}
 	}
+	ComponentMap.Remove(Component);
 }
 
-void FTextureInstanceManager::IncrementalUpdateBounds(float Percentage)
+void FTextureInstanceState::UpdateBounds(const UPrimitiveComponent* Component)
 {
-	const int32 LastIndex = FMath::Min(Bounds4Components.Num(), FirstDirtyComponentIndex + FMath::CeilToInt((float)Bounds4Components.Num() * Percentage));
-	for (int32 Index = FirstDirtyComponentIndex; Index < LastIndex; ++Index)
+	int32* ComponentLink = ComponentMap.Find(Component);
+	if (ComponentLink)
 	{
-		const UPrimitiveComponent* Component = Bounds4Components[Index];
-		if (Component)
+		int32 ElementIndex = *ComponentLink;
+		while (ElementIndex != INDEX_NONE)
 		{
-			Bounds4[Index / 4].Set(Index % 4, Component->Bounds, Component->LastRenderTime);
+			const FElement& Element = Elements[ElementIndex];
+			if (Element.BoundsIndex != INDEX_NONE)
+			{
+				Bounds4[Element.BoundsIndex / 4].Update(Element.BoundsIndex % 4, Component->Bounds, Component->LastRenderTime);
+			}
+			ElementIndex = Element.NextComponentLink;
 		}
 	}
-	FirstDirtyComponentIndex = LastIndex;
 }
 
-/*
-void FTextureInstanceManager::SanityCheck()
+void FTextureInstanceState::UpdateBounds(int32 BoundIndex)
 {
-	// FreeBoundIndices does not contain duplicate entries
+	const UPrimitiveComponent* Component = Bounds4Components[BoundIndex];
+	if (Component)
+	{
+		Bounds4[BoundIndex / 4].Update(BoundIndex % 4, Component->Bounds, Component->LastRenderTime);
+	}
+}
 
+void FTextureInstanceState::UpdateLastRenderTime(int32 BoundIndex)
+{
+	const UPrimitiveComponent* Component = Bounds4Components[BoundIndex];
+	if (Component)
+	{
+		Bounds4[BoundIndex / 4].Update(BoundIndex % 4, Component->LastRenderTime);
+	}
 }
-*/
 
-void FTextureBoundsVisibility::Set(const FTextureInstanceManager& Instances)
+uint32 FTextureInstanceState::GetAllocatedSize() const
 {
-	QUICK_SCOPE_CYCLE_COUNTER(FTextureBoundsVisibility_Set);
+	int32 CompiledElementsSize = 0;
+	for (TMap<const UTexture2D*, TArray<FCompiledElement> >::TConstIterator It(CompiledTextureMap); It; ++It)
+	{
+		CompiledElementsSize += It.Value().GetAllocatedSize();
+	}
 
-	// Data must be valid at this point.
-	check(Instances.FirstDirtyComponentIndex == Instances.Bounds4Components.Num());
+	return Bounds4.GetAllocatedSize() +
+		Bounds4Components.GetAllocatedSize() +
+		Elements.GetAllocatedSize() +
+		FreeBoundIndices.GetAllocatedSize() +
+		FreeElementIndices.GetAllocatedSize() +
+		TextureMap.GetAllocatedSize() +
+		CompiledTextureMap.GetAllocatedSize() + CompiledElementsSize +
+		ComponentMap.GetAllocatedSize();
+}
 
-	Bounds4 = Instances.Bounds4;
-	Elements = Instances.Elements;
-	TextureMap = Instances.TextureMap;
+void FStaticTextureInstanceManager::IncrementalUpdate(float Percentage)
+{
+	const int32 LastIndex = FMath::Min(State->NumBounds(), DirtyIndex + FMath::CeilToInt((float)State->NumBounds() * Percentage));
+	for (int32 Index = DirtyIndex; Index < LastIndex; ++Index)
+	{
+		State->UpdateLastRenderTime(Index);
+	}
+	DirtyIndex = LastIndex;
 }
 
-void FTextureBoundsVisibility::QuickSet(FTextureInstanceManager& Instances)
+void FStaticTextureInstanceManager::NormalizeLightmapTexelFactor()
 {
-	QUICK_SCOPE_CYCLE_COUNTER(FTextureBoundsVisibility_QuickSet);
+	TArray<float> TexelFactors;
 
-	// Data must be valid at this point.
-	check(Instances.FirstDirtyComponentIndex >= Instances.Bounds4Components.Num());
-	check(!Instances.bAllowNullComponents);
+	for (auto TextureIt = State->GetTextureIterator(); TextureIt; ++TextureIt)
+	{
+		const UTexture2D* Texture = *TextureIt;
 
-	// Set the two arrays to the same size then swap the buffers.
-	const int32 DeltaNum = Instances.Bounds4.Num() - Bounds4.Num();
-	if (DeltaNum > 0) // Bounds4 is too small!
+		if (Texture->LODGroup == TEXTUREGROUP_Lightmap || Texture->LODGroup == TEXTUREGROUP_Shadowmap)
+		{
+			TexelFactors.Reset();
+			for (auto ElementIt = State->GetElementIterator(Texture); ElementIt ; ++ElementIt)
+			{
+				TexelFactors.Push(ElementIt.GetTexelFactor());
+			}
+			TexelFactors.Sort();
+
+			const float MinTexelFactor = TexelFactors[FMath::FloorToInt(TexelFactors.Num() * 0.2f)];
+			const float MaxTexelFactor = TexelFactors[FMath::FloorToInt(TexelFactors.Num() * 0.8f)];
+
+			for (auto ElementIt = State->GetElementIterator(Texture); ElementIt ; ++ElementIt)
+			{
+				ElementIt.ClampTexelFactor(MinTexelFactor, MaxTexelFactor);
+			}
+		}
+	}
+}
+
+void FTextureInstanceState::CompileElements()
+{
+	CompiledTextureMap.Empty();
+
+	// First create an entry for all elements, so that there are no reallocs when inserting each compiled elements.
+	for (TMap<const UTexture2D*, int32>::TConstIterator TextureIt(TextureMap); TextureIt; ++TextureIt)
 	{
-		Bounds4.AddUninitialized(DeltaNum);
+		const UTexture2D* Texture = TextureIt.Key();
+		CompiledTextureMap.Add(Texture);
 	}
-	else if (DeltaNum < 0) // Bounds4 is too big
+
+	// Then fill in each array.
+	for (TMap<const UTexture2D*, TArray<FCompiledElement> >::TIterator It(CompiledTextureMap); It; ++It)
 	{
-		Bounds4.RemoveAt(Instances.Bounds4.Num(), -DeltaNum);
+		const UTexture2D* Texture = It.Key();
+		TArray<FCompiledElement>& CompiledElemements = It.Value();
+
+		int32 CompiledElementCount = 0;
+		for (auto ElementIt = GetElementIterator(Texture); ElementIt; ++ElementIt)
+		{
+			++CompiledElementCount;
+		}
+		CompiledElemements.AddUninitialized(CompiledElementCount);
+
+		CompiledElementCount = 0;
+		for (auto ElementIt = GetElementIterator(Texture); ElementIt; ++ElementIt)
+		{
+			CompiledElemements[CompiledElementCount].BoundsIndex = ElementIt.GetBoundsIndex();
+			CompiledElemements[CompiledElementCount].TexelFactor = ElementIt.GetTexelFactor();
+			++CompiledElementCount;
+		}
+	}
+}
+
+void FDynamicTextureInstanceManager::IncrementalUpdate(float Percentage)
+{
+	const int32 LastIndex = FMath::Min(State->NumBounds(), DirtyIndex + FMath::CeilToInt((float)State->NumBounds() * Percentage));
+	for (int32 Index = DirtyIndex; Index < LastIndex; ++Index)
+	{
+		State->UpdateBounds(Index);
 	}
-	FMemory::Memswap(&Bounds4, &Instances.Bounds4, sizeof(Bounds4));
-	Instances.FirstDirtyComponentIndex = 0; // Expecting an incremental update.
+	DirtyIndex = LastIndex;
+}
 
-	Elements = Instances.Elements;
+TRefCountPtr<const FTextureInstanceState> FDynamicTextureInstanceManager::GetAsyncState()
+{
+	QUICK_SCOPE_CYCLE_COUNTER(FDynamicTextureInstanceManager_GetAsyncState);
+
+	FTextureInstanceState* AsyncState = new FTextureInstanceState();
+
+	AsyncState->Bounds4.AddUninitialized(State->Bounds4.Num());
+	// Pass the bounds to the async state and mark the state bounds as dirty.
+	FMemory::Memswap(&AsyncState->Bounds4, &State->Bounds4, sizeof(AsyncState->Bounds4));
+	DirtyIndex = 0;
+
+	// Copy the map, and keep the lean version
+	AsyncState->TextureMap = State->TextureMap;
+	FMemory::Memswap(&AsyncState->TextureMap , &State->TextureMap, sizeof(AsyncState->TextureMap));
+
+	// Copy the elements
+	AsyncState->Elements = State->Elements;
 
-	// Swap the map and keep the clean version
-	TextureMap = Instances.TextureMap;
-	FMemory::Memswap(&TextureMap , &Instances.TextureMap, sizeof(TextureMap));
+	return TRefCountPtr<const FTextureInstanceState>(AsyncState);
 }
 
-void FTextureBoundsVisibility::ComputeBoundsViewInfos(const TArray<FStreamingViewInfo>& ViewInfos, bool bUseAABB, float MaxEffectiveScreenSize)
+void FTextureInstanceAsyncView::Update_Async(const TArray<FStreamingViewInfo>& ViewInfos, bool bUseAABB, float MaxEffectiveScreenSize)
 {
+	if (!State.IsValid())  return;
+
 	const int32 NumViews = ViewInfos.Num();
 	const VectorRegister One4 = VectorSet(1.f, 1.f, 1.f, 1.f);
+	const int32 NumBounds4 = State->NumBounds4();
 
-	BoundsViewInfo4.Empty(NumViews);
-	for (int32 ViewIndex = 0; ViewIndex < NumViews; ++ViewIndex)
-	{
-		BoundsViewInfo4.AddDefaulted();
-		BoundsViewInfo4[ViewIndex].AddUninitialized(Bounds4.Num());
-	}
+	//@TODO : This should be the value since the last update (something will be visible if seen since last update, no matter how long).
+	// const float LastUpdateTime =  FApp::GetCurrentTime() - GStartTime - 5.f;
+	// const float LastUpdateTime =  World->GetTimeSeconds();
+	const float LastUpdateTime = -1005.f;
+
+	const VectorRegister LastUpdateTime4 = VectorSet(LastUpdateTime, LastUpdateTime, LastUpdateTime, LastUpdateTime);
 
-	for (int32 Bounds4Index = 0; Bounds4Index < Bounds4.Num(); ++Bounds4Index)
+	BoundsViewInfo.Empty(NumBounds4 * 4);
+	BoundsViewInfo.AddUninitialized(NumBounds4 * 4);
+
+	for (int32 Bounds4Index = 0; Bounds4Index < NumBounds4; ++Bounds4Index)
 	{
-		const FBounds4& CurrentBounds4 = Bounds4[Bounds4Index];
+		const FTextureInstanceState::FBounds4& CurrentBounds4 = State->GetBounds4(Bounds4Index);
 
 		// Calculate distance of viewer to bounding sphere.
 		const VectorRegister OriginX = VectorLoadAligned( &CurrentBounds4.OriginX );
@@ -400,14 +673,18 @@ void FTextureBoundsVisibility::ComputeBoundsViewInfos(const TArray<FStreamingVie
 		const VectorRegister ExtentZ = VectorLoadAligned( &CurrentBounds4.ExtentZ );
 		const VectorRegister Radius = VectorLoadAligned( &CurrentBounds4.Radius );
 		const VectorRegister MinDistanceSq = VectorLoadAligned( &CurrentBounds4.MinDistanceSq );
-		const VectorRegister MaxDistanceSq = VectorLoadAligned( &CurrentBounds4.MaxDistanceSq );
+		const VectorRegister MaxDistanceSq = VectorLoadAligned(&CurrentBounds4.MaxDistanceSq);
+		const VectorRegister LastRenderTime = VectorLoadAligned(&CurrentBounds4.LastRenderTime);
+		
+		VectorRegister MaxNormalizedSize = VectorZero();
+		VectorRegister MaxNormalizedSize_VisibleOnly = VectorZero();
 
 		for (int32 ViewIndex = 0; ViewIndex < NumViews; ++ViewIndex)
 		{
 			const FStreamingViewInfo& ViewInfo = ViewInfos[ViewIndex];
 
 			const float EffectiveScreenSize = (MaxEffectiveScreenSize > 0.0f) ? FMath::Min(MaxEffectiveScreenSize, ViewInfo.ScreenSize) : ViewInfo.ScreenSize;
-			const float ScreenSizeFloat = EffectiveScreenSize * ViewInfo.BoostFactor;
+			const float ScreenSizeFloat = EffectiveScreenSize * ViewInfo.BoostFactor * .5f; // Multiply by half since the ratio factors map to half the screen only
 
 			const VectorRegister ScreenSize = VectorLoadFloat1( &ScreenSizeFloat );
 			const VectorRegister ViewOriginX = VectorLoadFloat1( &ViewInfo.ViewOrigin.X );
@@ -457,85 +734,96 @@ void FTextureBoundsVisibility::ComputeBoundsViewInfos(const TArray<FStreamingVie
 			VectorRegister ScreenSizeOverDistance = VectorReciprocalSqrt(DistSqMinusRadiusSq);
 			ScreenSizeOverDistance = VectorMultiply(ScreenSizeOverDistance, ScreenSize);
 
-			// Store results
-			FBoundsViewInfo4& BoundsVieWInfo = BoundsViewInfo4[ViewIndex][Bounds4Index];
-			VectorStoreAligned(ScreenSizeOverDistance, &BoundsVieWInfo.ScreenSizeOverDistance);
-			VectorStoreAligned(DistSqMinusRadiusSq, &BoundsVieWInfo.DistanceSq);
-			VectorStoreAligned(InRangeMask, &BoundsVieWInfo.InRangeMask);
+			MaxNormalizedSize = VectorMax(ScreenSizeOverDistance, MaxNormalizedSize);
+
+			// Now mask to zero if not in range, or not seen recently.
+			ScreenSizeOverDistance = VectorSelect(InRangeMask, ScreenSizeOverDistance, VectorZero());
+			ScreenSizeOverDistance = VectorSelect(VectorCompareGT(LastRenderTime, LastUpdateTime4), ScreenSizeOverDistance, VectorZero());
+
+			MaxNormalizedSize_VisibleOnly = VectorMax(ScreenSizeOverDistance, MaxNormalizedSize_VisibleOnly);
 		}
+
+		// Store results
+		FBoundsViewInfo* BoundsVieWInfo = &BoundsViewInfo[Bounds4Index * 4];
+		for (int32 SubIndex = 0; SubIndex < 4; ++SubIndex)
+		{
+			BoundsVieWInfo[SubIndex].MaxNormalizedSize = VectorGetComponent(MaxNormalizedSize, SubIndex);
+			BoundsVieWInfo[SubIndex].MaxNormalizedSize_VisibleOnly = VectorGetComponent(MaxNormalizedSize_VisibleOnly, SubIndex);
+		}
+	}
+}
+
+void FTextureInstanceAsyncView::ProcessElement(const FBoundsViewInfo& BoundsVieWInfo, float TexelFactor, float& MaxSize, float& MaxSize_VisibleOnly) const
+{
+	if (TexelFactor == FLT_MAX) // If this is a forced load component.
+	{
+		MaxSize = BoundsVieWInfo.MaxNormalizedSize > 0 ? FLT_MAX : MaxSize;
+		MaxSize_VisibleOnly = BoundsVieWInfo.MaxNormalizedSize_VisibleOnly > 0 ? FLT_MAX : MaxSize_VisibleOnly;
+	}
+	else if (TexelFactor >= 0)
+	{
+		MaxSize = FMath::Max(MaxSize, TexelFactor * BoundsVieWInfo.MaxNormalizedSize);
+		MaxSize_VisibleOnly = FMath::Max(MaxSize_VisibleOnly, TexelFactor * BoundsVieWInfo.MaxNormalizedSize_VisibleOnly);
+	}
+	else // Negative texel factors map to fixed resolution. Currently used for lanscape.
+	{
+		MaxSize = FMath::Max(MaxSize, -TexelFactor);
+		MaxSize_VisibleOnly = FMath::Max(MaxSize_VisibleOnly, -TexelFactor);
 	}
 }
 
-bool FTextureBoundsVisibility::GetTexelSize(const UTexture2D* InTexture, float& MaxSize_InRange, float& MinDistanceSq_InRange, float& MaxSize_AnyRange, float& MinDistanceSq_AnyRange) const
+void FTextureInstanceAsyncView::GetTexelSize(const UTexture2D* InTexture, float& MaxSize, float& MaxSize_VisibleOnly, bool bOutputToLog) const
 {
-	MaxSize_InRange = 0;
-	MaxSize_AnyRange = 0;
-	MinDistanceSq_InRange = MAX_FLT;
-	MinDistanceSq_AnyRange = MAX_FLT;
+	// No need to iterate more if texture is already at maximum resolution.
 
-	const int32* TextureLink = TextureMap.Find(InTexture);
-	if (TextureLink)
+	int32 CurrCount = 0;
+
+	if (State.IsValid())
 	{
-		int32 ElementIndex = *TextureLink;
-		while (ElementIndex != INDEX_NONE)
+		// Use the fast path if available, about twice as fast when there are a lot of elements.
+		if (State->HasCompiledElements() && !bOutputToLog)
 		{
-			const FElement& Element = Elements[ElementIndex];
-			const int32 BoundsIndex = Element.BoundsIndex;
-
-			for (int32 ViewIndex = 0; ViewIndex < BoundsViewInfo4.Num(); ++ViewIndex)
+			const TArray<FTextureInstanceState::FCompiledElement>* CompiledElements = State->GetCompiledElements(InTexture);
+			if (CompiledElements)
 			{
-				FBoundsViewInfo4& BoundsVieWInfo = const_cast<FBoundsViewInfo4&>(BoundsViewInfo4[ViewIndex][BoundsIndex / 4]);
-				const float ScreenSizeOverDistance = BoundsVieWInfo.ScreenSizeOverDistance.Component(BoundsIndex % 4);
-				const float DistanceSq = BoundsVieWInfo.DistanceSq.Component(BoundsIndex % 4);
-				const int32 InRangeMask = reinterpret_cast<const int32&>(BoundsVieWInfo.InRangeMask.Component(BoundsIndex % 4));
-
-#if 0
-				//@TODO (for static instances)
-				// Used with TEXTUREGROUP_Terrain_Heightmap
-				// To check Forced LOD...
-				if (Element.TexelFactor < 0)
+				const int32 NumCompiledElements = CompiledElements->Num();
+				const FTextureInstanceState::FCompiledElement* CompiledElementData = CompiledElements->GetData();
+
+				int32 CompiledElementIndex = 0;
+				while (CompiledElementIndex < NumCompiledElements && MaxSize_VisibleOnly < MAX_TEXTURE_SIZE)
 				{
-					VectorRegister MinForcedLODs = VectorLoadAligned( &TextureInstance.TexelFactor );
-					bool AllForcedLOD = !!VectorAllLesserThan(MinForcedLODs, VectorZero());
-					MinForcedLODs = VectorMin( MinForcedLODs, VectorSwizzle(MinForcedLODs, 2, 3, 0, 1) );
-					MinForcedLODs = VectorMin( MinForcedLODs, VectorReplicate(MinForcedLODs, 1) );
-					float MinLODValue;
-					VectorStoreFloat1( MinForcedLODs, &MinLODValue );
-
-					if (MinLODValue <= 0)
-					{
-						WantedMipCount = FMath::Max(WantedMipCount, FFloatMipLevel::FromMipLevel(StreamingTexture.MaxAllowedMips - 13 - FMath::FloorToInt(MinLODValue)));
-						if (WantedMipCount >= FFloatMipLevel::FromMipLevel(StreamingTexture.MaxAllowedMips))
-					{
-						MinDistance = 1.0f;
-						bShouldAbortLoop = true;
-					}
+					const FTextureInstanceState::FCompiledElement& CompiledElement = CompiledElementData[CompiledElementIndex];
+					ProcessElement(BoundsViewInfo[CompiledElement.BoundsIndex], CompiledElement.TexelFactor, MaxSize, MaxSize_VisibleOnly);
+					++CompiledElementIndex;
 				}
-#endif
 
-				float Size = Element.TexelFactor * ScreenSizeOverDistance;
-
-				MaxSize_AnyRange = FMath::Max(Size, MaxSize_AnyRange);
-				MinDistanceSq_AnyRange = FMath::Min(DistanceSq, MinDistanceSq_AnyRange);
-				if (InRangeMask)
+				if (MaxSize_VisibleOnly >= MAX_TEXTURE_SIZE && CompiledElementIndex > 1)
 				{
-					MaxSize_InRange = FMath::Max(Size, MaxSize_InRange);
-					MinDistanceSq_InRange = FMath::Min(DistanceSq, MinDistanceSq_InRange);
-
-					if (MinDistanceSq_InRange <= 1)
-					{
-						return true; // Already at max
-					}
+					// This does not realloc anything but moves the close by element to the first entry, making the next update find it immediately.
+					FTextureInstanceState::FCompiledElement* SwapElementData = const_cast<FTextureInstanceState::FCompiledElement*>(CompiledElementData);
+					Swap<FTextureInstanceState::FCompiledElement>(SwapElementData[0], SwapElementData[CompiledElementIndex - 1]);
+				}
+			}
+		}
+		else
+		{
+			for (auto It = State->GetElementIterator(InTexture); It && (MaxSize_VisibleOnly < MAX_TEXTURE_SIZE || bOutputToLog); ++It)
+			{
+				const FBoundsViewInfo& BoundsVieWInfo = BoundsViewInfo[It.GetBoundsIndex()];
+				ProcessElement(BoundsVieWInfo, It.GetTexelFactor(), MaxSize, MaxSize_VisibleOnly);
+				if (bOutputToLog)
+				{
+					It.OutputToLog(BoundsVieWInfo.MaxNormalizedSize, BoundsVieWInfo.MaxNormalizedSize_VisibleOnly);
 				}
 			}
-
-			ElementIndex = Element.NextTextureLink;
 		}
 	}
-
-	return TextureLink != nullptr;
 }
 
+bool FTextureInstanceAsyncView::HasTextureReferences(const UTexture2D* InTexture) const
+{
+	return State.IsValid() && (bool)State->GetElementIterator(InTexture);
+}
 
 void FDynamicComponentTextureManager::Add(const UPrimitiveComponent* Component, EDynamicPrimitiveType DynamicType, FRemovedTextureArray& RemovedTextures)
 {
@@ -549,7 +837,7 @@ void FDynamicComponentTextureManager::Add(const UPrimitiveComponent* Component,
 	{
 		if (State->bHasTextures) // If we reinsert without having cleared the old data correctly, make sure to clear it now.
 		{
-			TextureInstances.RemoveComponent(Component, RemovedTextures);
+			DynamicInstances.RemoveComponent(Component, RemovedTextures);
 			State->bHasTextures = false;
 		}
 		State->bToDelete = false;
@@ -566,7 +854,7 @@ void FDynamicComponentTextureManager::Remove(const UPrimitiveComponent* Componen
 	{
 		if (State->bHasTextures) // When a component becomes invalid, remove any texture and bound references.
 		{
-			TextureInstances.RemoveComponent(Component, RemovedTextures);
+			DynamicInstances.RemoveComponent(Component, RemovedTextures);
 			State->bHasTextures = false;
 		}
 		State->bToDelete = true;
@@ -582,17 +870,19 @@ void FDynamicComponentTextureManager::Update(const UPrimitiveComponent* Componen
 	if (State && !State->bToDelete && State->bHasTextures)
 	{
 		// If we reinserted without having cleared the old data correctly, make sure to clear it now.
-		TextureInstances.RemoveComponent(Component, RemovedTextures);
+		DynamicInstances.RemoveComponent(Component, RemovedTextures);
 		State->bHasTextures = false;
 		DirtyComponents.Add(Component);
 	}
 }
 
-void FDynamicComponentTextureManager::IncrementalTick(FRemovedTextureArray& RemovedTextures, float Percentage)
+void FDynamicComponentTextureManager::IncrementalUpdate(FRemovedTextureArray& RemovedTextures, float Percentage)
 {
-	QUICK_SCOPE_CYCLE_COUNTER(FDynamicComponentTextureManager_IncrementalTick);
+	QUICK_SCOPE_CYCLE_COUNTER(FDynamicComponentTextureManager_IncrementalUpdate);
 
-	TextureInstances.IncrementalUpdateBounds(Percentage);
+	// This incrementally updates the bounds for all components.
+	// Unrelated to the dirty component logics, which recomputes the texture for each components.
+	DynamicInstances.IncrementalUpdate(Percentage);
 
 	TSet<const UPrimitiveComponent*> NextDirtyComponents; // Components that have a dirty states (to accelerate the update).
 	for (TSet<const UPrimitiveComponent*>::TIterator It(DirtyComponents); It; ++It)
@@ -604,7 +894,7 @@ void FDynamicComponentTextureManager::IncrementalTick(FRemovedTextureArray& Remo
 
 		if (State->bToDelete)
 		{
-			check(!State->bHasTextures);
+			check(!State->bHasTextures); // Textures are removed when the bToDelete flag is set.
 			ComponentStates.Remove(Component);
 		}
 		else
@@ -614,22 +904,99 @@ void FDynamicComponentTextureManager::IncrementalTick(FRemovedTextureArray& Remo
 			{
 				if (State->bHasTextures)
 				{
-					TextureInstances.RemoveComponent(Component, RemovedTextures);
+					DynamicInstances.RemoveComponent(Component, RemovedTextures);
 					State->bHasTextures = false;
 				}
 				NextDirtyComponents.Add(Component);
 			}
 			else if (bIsInScene && !State->bHasTextures)
 			{
-				TextureInstances.AddComponent(Component);
+				DynamicInstances.AddComponent(Component);
 				State->bHasTextures = true;
 			}
 			else
 			{
-				TextureInstances.UpdateBounds(Component);
+				DynamicInstances.UpdateBounds(Component);
 			}
 		}
 	}
 
 	FMemory::Memswap(&DirtyComponents, &NextDirtyComponents, sizeof(DirtyComponents));
 }
+
+uint32 FDynamicComponentTextureManager::GetAllocatedSize() const
+{
+	return ComponentStates.GetAllocatedSize() + DirtyComponents.GetAllocatedSize() + DynamicInstances.GetAllocatedSize();
+}
+
+void FLevelTextureManager::Remove(FDynamicComponentTextureManager& DynamicComponentManager, FRemovedTextureArray& RemovedTextures)
+{ 
+	// Mark all static textures for removal.
+	for (auto It = StaticInstances.GetTextureIterator(); It; ++It)
+	{
+		RemovedTextures.Push(*It);
+	}
+	StaticInstances = FStaticTextureInstanceManager();
+
+	// Mark all dynamic textures for removal.
+	for (const UPrimitiveComponent* Component : DynamicComponents)
+	{
+		DynamicComponentManager.Remove(Component, RemovedTextures);
+	}	
+	DynamicComponents.Empty();
+
+	bToDelete = true; 
+}
+
+void FLevelTextureManager::IncrementalUpdate(FDynamicComponentTextureManager& DynamicComponentManager, FRemovedTextureArray& RemovedTextures, float Percentage, bool bUseDynamicStreaming) 
+{
+	QUICK_SCOPE_CYCLE_COUNTER(FStaticComponentTextureManager_IncrementalUpdate);
+
+	check(!bToDelete);
+
+	StaticInstances.IncrementalUpdate(Percentage);
+
+	if (!bHasTextures)
+	{
+		FStreamingTextureLevelContext LevelContext(Level);
+
+		// Process all components of the level.
+		TArray<UObject*> ObjectsInOuter;
+		GetObjectsWithOuter(Level, ObjectsInOuter);
+
+		for (UObject* LevelObject : ObjectsInOuter)
+		{
+			UPrimitiveComponent* Component = Cast<UPrimitiveComponent>(LevelObject);
+
+			// Skip non primitive components, default objects and non registered objects.
+			if (!Component || Component->IsTemplate(RF_ClassDefaultObject) || !Component->IsRegistered()) 
+				continue;
+
+			const AActor* const Owner = Component->GetOwner();
+			const bool bIsStatic = (!Owner || Component->Mobility == EComponentMobility::Static || Component->Mobility == EComponentMobility::Stationary);
+
+			if (bIsStatic)
+			{
+				StaticInstances.AddComponent(Component, LevelContext);
+			}
+			else if (bUseDynamicStreaming)
+			{
+				DynamicComponentManager.Add(Component, DPT_Level, RemovedTextures);
+
+				// This will be used when the level is removed. Never dereferenced.
+				DynamicComponents.Add(Component);
+			}
+
+		}
+
+		StaticInstances.NormalizeLightmapTexelFactor();
+		StaticInstances.CompiledElements();
+
+		bHasTextures = true;
+	}
+}
+
+uint32 FLevelTextureManager::GetAllocatedSize() const
+{
+	return StaticInstances.GetAllocatedSize() + DynamicComponents.GetAllocatedSize();
+}
diff --git a/Engine/Source/Runtime/Engine/Private/Streaming/TextureInstanceManager.h b/Engine/Source/Runtime/Engine/Private/Streaming/TextureInstanceManager.h
--- a/Engine/Source/Runtime/Engine/Private/Streaming/TextureInstanceManager.h
+++ b/Engine/Source/Runtime/Engine/Private/Streaming/TextureInstanceManager.h
@@ -8,11 +8,13 @@ TextureInstanceManager.h: Definitions of classes used for texture streaming.
 
 #include "TextureStreamingHelpers.h"
 
+#define MAX_TEXTURE_SIZE (float(1 << (MAX_TEXTURE_MIP_COUNT - 1)))
+
 // Main Thread Job Requirement : find all instance of a component and update it's bound.
 // Threaded Job Requirement : get the list of instance texture easily from the list of 
 
 // Can be used either for static primitives or dynamic primitives
-class FTextureInstanceManager
+class FTextureInstanceState : public FRefCountedObject
 {
 public:
 
@@ -43,9 +45,12 @@ public:
 		FVector4 MaxDistanceSq;
 
 		/** Last visibility time for this bound, used for priority */
-		FVector4 LastRenderTime;
+		FVector4 LastRenderTime; //(FApp::GetCurrentTime() - Component->LastRenderTime);
+
 
-		FORCEINLINE_DEBUGGABLE void Set(int32 Index, const FBoxSphereBounds& Bounds, float LastRenderTime, float MinDistance = 0, float MaxDistance = MAX_FLT);
+		FORCEINLINE_DEBUGGABLE void Set(int32 Index, const FBoxSphereBounds& Bounds, float LastRenderTime, float MinDistance, float MaxDistance);
+		FORCEINLINE_DEBUGGABLE void Update(int32 Index, const FBoxSphereBounds& Bounds, float LastRenderTime);
+		FORCEINLINE_DEBUGGABLE void Update(int32 Index, float LastRenderTime);
 
 		// Clears entry between 0 and 4
 		FORCEINLINE void Clear(int32 Index);
@@ -64,89 +69,252 @@ public:
 		int32 PrevTextureLink;	// The previous element which uses the same texture as this Element. The first element referred by TextureMap will have INDEX_NONE.
 		int32 NextTextureLink;	// The next element which uses the same texture as this Element. Last element will have INDEX_NONE
 
-		// Components are always updated as a whole, so individual elements can not be removed.
+		// Components are always updated as a whole, so individual elements can not be removed. Removing the need for PrevComponentLink.
 		int32 NextComponentLink;	// The next element that uses the same component as this Element. The first one is referred by ComponentMap and the last one will have INDEX_NONE.
 	};
 
-public:
 
-	FTextureInstanceManager(bool InAllowNullComponents = false) : bAllowNullComponents(InAllowNullComponents), FirstDirtyComponentIndex(0) {}
+	/**
+ 	 * FCompiledElement is a stripped down version of element and is stored in an array instead of using a linked list.
+     * It is only used when the data is not expected to change and reduce that cache cost of iterating on all elements.
+     **/
+	struct FCompiledElement
+	{
+		int32 BoundsIndex;
+		float TexelFactor;
+	};
 
-	int32 AddBounds(const FBoxSphereBounds& Bounds, const UPrimitiveComponent* Component, float LastRenderTime, float MinDistance = 0 , float MaxDistance = MAX_FLT);
 
-	void AddElement(const UPrimitiveComponent* Component, const UTexture2D* Texture, int BoundsIndex, float TexelFactor, int32*& ComponentLink);
+	// Iterator processing all elements refering to a texture.
+	class FTextureLinkConstIterator
+	{
+	public:
+		FTextureLinkConstIterator(const FTextureInstanceState& InState, const UTexture2D* InTexture);
 
-	// Returns the next elements using the same component.
-	void RemoveElement(int32 ElementIndex, int32& NextComponentLink, int32& BoundsIndex, const UTexture2D*& Texture);
+		FORCEINLINE operator bool() const { return CurrElementIndex != INDEX_NONE; }
+		FORCEINLINE void operator++() { CurrElementIndex = State.Elements[CurrElementIndex].NextTextureLink; }
+
+		void OutputToLog(float MaxNormalizedSize, float MaxNormalizedSize_VisibleOnly) const;
+
+		FORCEINLINE int32 GetBoundsIndex() const { return State.Elements[CurrElementIndex].BoundsIndex; }
+		FORCEINLINE float GetTexelFactor() const { return State.Elements[CurrElementIndex].TexelFactor; }
+
+	protected:
+
+		FORCEINLINE const UPrimitiveComponent* GetComponent() const { return State.Elements[CurrElementIndex].Component; }
+		FBoxSphereBounds GetBounds() const;
+
+		const FTextureInstanceState& State;
+		int32 CurrElementIndex;
+	};
+
+	class FTextureLinkIterator : public FTextureLinkConstIterator
+	{
+	public:
+		FTextureLinkIterator(FTextureInstanceState& InState, const UTexture2D* InTexture) : FTextureLinkConstIterator(InState, InTexture) {}
+
+		FORCEINLINE void ClampTexelFactor(float CMin, float CMax)
+		{ 
+			float& TexelFactor = const_cast<FTextureInstanceState&>(State).Elements[CurrElementIndex].TexelFactor;
+			TexelFactor = FMath::Clamp<float>(TexelFactor, CMin, CMax);
+		}
+	};
+
+	class FTextureIterator
+	{
+	public:
+		FTextureIterator(const FTextureInstanceState& InState) : MapIt(InState.TextureMap) {}
+
+		operator bool() const { return (bool)MapIt; }
+		void operator++() { ++MapIt; }
+
+		const UTexture2D* operator*() const { return MapIt.Key(); }
+
+	private:
+
+		TMap<const UTexture2D*, int32>::TConstIterator MapIt;
+	};
+
+public:
 
 	// Will also remove bounds
+	void AddComponent(const UPrimitiveComponent* Component, FStreamingTextureLevelContext& LevelContext);
 	void AddComponent(const UPrimitiveComponent* Component);
+
 	void RemoveComponent(const UPrimitiveComponent* Component, FRemovedTextureArray& RemovedTextures);
+	void RemoveComponentReferences(const UPrimitiveComponent* Component);
+
+	bool Contains(const UPrimitiveComponent* Component) const { return ComponentMap.Contains(Component); }
 
-	// Update the component bounds, assumed that this would only be used with component using component bounds
 	void UpdateBounds(const UPrimitiveComponent* Component);
 
-	bool Contains(const UPrimitiveComponent* Component) const { return ComponentMap.Contains(Component); }
+	void UpdateBounds(int32 BoundIndex);
+	void UpdateLastRenderTime(int32 BoundIndex);
+
+	FORCEINLINE int32 NumBounds() const { return Bounds4Components.Num(); }
+	FORCEINLINE int32 NumBounds4() const { return Bounds4.Num(); }
 
-	void IncrementalUpdateBounds(float Percentage);
+	FORCEINLINE const FBounds4& GetBounds4(int32 Bounds4Index ) const {  return Bounds4[Bounds4Index]; }
+
+	FORCEINLINE FTextureLinkIterator GetElementIterator(const UTexture2D* InTexture ) {  return FTextureLinkIterator(*this, InTexture); }
+	FORCEINLINE FTextureLinkConstIterator GetElementIterator(const UTexture2D* InTexture ) const {  return FTextureLinkConstIterator(*this, InTexture); }
+	FORCEINLINE FTextureIterator GetTextureIterator( ) const {  return FTextureIterator(*this); }
+
+	uint32 GetAllocatedSize() const;
+
+	// Generate the compiled elements.
+	void CompileElements();
+
+	// Whether or not this state has compiled elements.
+	bool HasCompiledElements() const { return CompiledTextureMap.Num() != 0; }
+
+	// If this has compiled elements, return the array relate to a given texture.
+	const TArray<FCompiledElement>* GetCompiledElements(const UTexture2D* Texture) const { return CompiledTextureMap.Find(Texture); }
 
 private:
 
+	void AddElement(const UPrimitiveComponent* Component, const UTexture2D* Texture, int BoundsIndex, float TexelFactor, int32*& ComponentLink);
+
+	// Returns the next elements using the same component.
+	void RemoveElement(int32 ElementIndex, int32& NextComponentLink, int32& BoundsIndex, const UTexture2D*& Texture);
+
+	int32 AddBounds(const FBoxSphereBounds& Bounds, const UPrimitiveComponent* Component, float LastRenderTime, float MinDistance = 0 , float MaxDistance = FLT_MAX);
 	void RemoveBounds(int32 Index);
 
 private:
 
 	TArray<FBounds4> Bounds4;
+
+	/** 
+	 * Components related to each of the Bounds4 elements. This is stored in another array to allow 
+	 * passing Bounds4 to the threaded task without loosing the bound components, allowing incremental update.
+	 */
+	TArray<const UPrimitiveComponent*> Bounds4Components;
+
 	TArray<FElement> Elements;
 
 	TArray<int32> FreeBoundIndices;
 	TArray<int32> FreeElementIndices;
 
 	TMap<const UTexture2D*, int32> TextureMap;
+	
+	// CompiledTextureMap is used to iterate more quickly on each elements by avoiding the linked list indirections.
+	TMap<const UTexture2D*, TArray<FCompiledElement> > CompiledTextureMap;
+
 	TMap<const UPrimitiveComponent*, int32> ComponentMap;
 
-	bool bAllowNullComponents;
+	friend class FTextureLinkIterator;
+	friend class FTextureIterator;
+	friend class FDynamicTextureInstanceManager;
+};
+
+class FStaticTextureInstanceManager
+{
+public:
+
+	FStaticTextureInstanceManager() : State(new FTextureInstanceState()), DirtyIndex(0) {}
+
+	// Update the component bounds, assumed that this would only be used with component using component bounds
+	void IncrementalUpdate(float Percentage);
+
+	FORCEINLINE void AddComponent(const UPrimitiveComponent* Component, FStreamingTextureLevelContext& LevelContext) { State->AddComponent(Component, LevelContext); }
 
-	// Used to implement iterative bounds update and visibility update.
-	int32 FirstDirtyComponentIndex; // Ranges from 0 to Components.Num()
-	TArray<const UPrimitiveComponent*> Bounds4Components; // The components mapped to Bounds4
+	void RemoveComponentReferences(const UPrimitiveComponent* Component) { State->RemoveComponentReferences(Component); }
 
-	friend class FTextureBoundsVisibility;
+	// Because static data does not change, there is no need to duplicate the states.
+	TRefCountPtr<const FTextureInstanceState> GetAsyncState() 
+	{ 
+		DirtyIndex = 0; // Triggers a refresh.
+		return TRefCountPtr<const FTextureInstanceState>(State.GetReference()); 
+	}
+
+	FORCEINLINE FTextureInstanceState::FTextureIterator GetTextureIterator( ) const {  return State->GetTextureIterator(); }
+
+	void NormalizeLightmapTexelFactor();
+	void CompiledElements() { State->CompileElements(); }
+
+	uint32 GetAllocatedSize() const
+	{
+		return State.IsValid() ? (sizeof(FTextureInstanceState) + State->GetAllocatedSize()) : 0;
+	}
+
+private:
+
+	/** The texture instances. Shared with the async task. */
+	TRefCountPtr<FTextureInstanceState> State;
+
+	/** Ranges from 0 to Bounds4Components.Num(). Used in the incremental update to update visibility. */
+	int32 DirtyIndex;
 };
 
-// Data used to compute visibility
-class FTextureBoundsVisibility
+class FDynamicTextureInstanceManager
 {
 public:
 
-	void Set(const FTextureInstanceManager& Instances);
-	void QuickSet(FTextureInstanceManager& Instances);
+	FDynamicTextureInstanceManager() : State(new FTextureInstanceState()), DirtyIndex(0) {}
+
+	// Update the component bounds, assumed that this would only be used with component using component bounds
+	FORCEINLINE void UpdateBounds(const UPrimitiveComponent* Component) { State->UpdateBounds(Component); }
+	void IncrementalUpdate(float Percentage);
+
+	// Will also remove bounds
+	FORCEINLINE void AddComponent(const UPrimitiveComponent* Component) { State->AddComponent(Component); }
+	FORCEINLINE void RemoveComponent(const UPrimitiveComponent* Component, FRemovedTextureArray& RemovedTextures) { State->RemoveComponent(Component, RemovedTextures); }
 
-	void ComputeBoundsViewInfos(const TArray<FStreamingViewInfo>& ViewInfos, bool bUseAABB, float MaxEffectiveScreenSize);
+	TRefCountPtr<const FTextureInstanceState> GetAsyncState();
 
-	// MaxSize_InRange : Biggest instance size for in range views
-	// MaxSize_AnyRange : Biggest instance size not considering range.
-	bool GetTexelSize(const UTexture2D* InTexture, float& MaxSize_InRange, float& MinDistanceSq_InRange, float& MaxSize_AnyRange, float& MinDistanceSq_AnyRange) const;
+	uint32 GetAllocatedSize() const 
+	{ 
+		return State.IsValid() ? (sizeof(FTextureInstanceState) + State->GetAllocatedSize()) : 0;
+	}
 
 private:
 
-	typedef FTextureInstanceManager::FBounds4 FBounds4;
-	typedef FTextureInstanceManager::FElement FElement;
+	/** The texture instances. Shared with the async task. */
+	TRefCountPtr<FTextureInstanceState> State;
 
-	TArray<FBounds4> Bounds4;
-	TArray<FElement> Elements;
-	TMap<const UTexture2D*, int32> TextureMap;
+	/** Ranges from 0 to Bounds4Components.Num(). Used in the incremental update to update bounds and visibility. */
+	int32 DirtyIndex;
+};
+
+
+// Data used to compute visibility
+class FTextureInstanceAsyncView
+{
+public:
+
+	FTextureInstanceAsyncView() {}
+
+	FTextureInstanceAsyncView(const TRefCountPtr<const FTextureInstanceState>& InState) : State(InState) {}
+
+	void Update_Async(const TArray<FStreamingViewInfo>& ViewInfos, bool bUseAABB, float MaxEffectiveScreenSize);
 
-	struct FBoundsViewInfo4
+	// MaxSize : Biggest texture size for all instances.
+	// MaxSize_VisibleOnly : Biggest texture size for visble instances only.
+	void GetTexelSize(const UTexture2D* InTexture, float& MaxSize, float& MaxSize_VisibleOnly, bool bOutputToLog) const;
+
+	bool HasTextureReferences(const UTexture2D* InTexture) const;
+
+private:
+
+	TRefCountPtr<const FTextureInstanceState> State;
+
+	struct FBoundsViewInfo
 	{
-		FVector4 ScreenSizeOverDistance;
-		FVector4 DistanceSq;
-		FVector4 InRangeMask;
+		/** The biggest normalized size (ScreenSize / Distance) accross all view.*/
+		float MaxNormalizedSize;
+		/*
+		 * The biggest normalized size accross all view for visible instances only.
+		 * Visible instances are the one that are in range and also that have been seen recently.
+		 */
+		float MaxNormalizedSize_VisibleOnly;
 	};
 
-	// Normalized Texel Factors for each bounds and view.
+	// Normalized Texel Factors for each bounds and view. This is the data built by ComputeBoundsViewInfos
 	// @TODO : store data for different views continuously to improve reads.
-	TArray< TArray<FBoundsViewInfo4> > BoundsViewInfo4;
+	TArray<FBoundsViewInfo> BoundsViewInfo;
+
+	FORCEINLINE_DEBUGGABLE void ProcessElement(const FBoundsViewInfo& BoundsVieWInfo, float TexelFactor, float& MaxSize, float& MaxSize_VisibleOnly) const;
 };
 
 
@@ -167,10 +335,11 @@ public:
 
 	void Update(const UPrimitiveComponent* Component, FRemovedTextureArray& RemovedTextures);
 
-	void IncrementalTick(FRemovedTextureArray& RemovedTextures, float Percentage);
+	void IncrementalUpdate(FRemovedTextureArray& RemovedTextures, float Percentage);
+
+	const FDynamicTextureInstanceManager& GetTextureInstances() const { return DynamicInstances; }
 
-	const FTextureInstanceManager& GetTextureInstances() const { return TextureInstances; }
-	FTextureInstanceManager& GetTextureInstances() { return TextureInstances; }
+	FORCEINLINE FTextureInstanceAsyncView GetAsyncView() { return FTextureInstanceAsyncView(DynamicInstances.GetAsyncState()); }
 
 	bool IsDeleted(const UPrimitiveComponent* Component) const 
 	{ 	
@@ -178,6 +347,8 @@ public:
 		return !State || State->bToDelete;
 	}
 
+	uint32 GetAllocatedSize() const;
+
 private:
 
 	struct FComponentState
@@ -201,5 +372,46 @@ private:
 
 	TMap<const UPrimitiveComponent*, FComponentState> ComponentStates;
 	TSet<const UPrimitiveComponent*> DirtyComponents; // Components that have a dirty states (to accelerate the incremental update).
-	FTextureInstanceManager TextureInstances;
+	FDynamicTextureInstanceManager DynamicInstances;
+};
+
+class FLevelTextureManager
+{
+public:
+
+	FLevelTextureManager(ULevel* InLevel) : Level(InLevel), bHasTextures(false), bToDelete(false) {}
+	~FLevelTextureManager() { check(bToDelete); }
+
+
+	ULevel* GetLevel() const { return Level; }
+
+	// Remove the whole level
+	void Remove(FDynamicComponentTextureManager& DynamicManager, FRemovedTextureArray& RemovedTextures);
+
+	// Invalidate a component reference.
+	void RemoveComponentReferences(const UPrimitiveComponent* Component) 
+	{ 
+		StaticInstances.RemoveComponentReferences(Component); 
+		DynamicComponents.RemoveSingleSwap(Component); 
+	}
+
+	const FStaticTextureInstanceManager& GetStaticInstances() const { return StaticInstances; }
+
+	FORCEINLINE FTextureInstanceAsyncView GetAsyncView() { return FTextureInstanceAsyncView(StaticInstances.GetAsyncState()); }
+
+	void IncrementalUpdate(FDynamicComponentTextureManager& DynamicManager, FRemovedTextureArray& RemovedTextures, float Percentage, bool bUseDynamicStreaming);
+
+	uint32 GetAllocatedSize() const;
+
+private:
+
+	ULevel* Level;
+
+	uint32 bHasTextures : 1;
+	uint32 bToDelete : 1;
+
+	FStaticTextureInstanceManager StaticInstances;
+
+	/** The list of dynamic components contained in the level. Used to removed them from the FDynamicComponentTextureManager on ::Remove. */
+	TArray<const UPrimitiveComponent*> DynamicComponents;
 };
diff --git a/Engine/Source/Runtime/Engine/Private/Streaming/TextureStreamingBuild.cpp b/Engine/Source/Runtime/Engine/Private/Streaming/TextureStreamingBuild.cpp
--- a/Engine/Source/Runtime/Engine/Private/Streaming/TextureStreamingBuild.cpp
+++ b/Engine/Source/Runtime/Engine/Private/Streaming/TextureStreamingBuild.cpp
@@ -102,6 +102,7 @@ ENGINE_API void BuildTextureStreamingData(UWorld* InWorld, const FTexCoordScaleM
 
 		Level->bTextureStreamingRotationChanged = false;
 		Level->StreamingTextureGuids.Empty(LevelTextures.Num());
+		Level->Modify();
 
 		// Reset LevelIndex to default for next use and build the level Guid array.
 		for (int32 TextureIndex = 0; TextureIndex < LevelTextures.Num(); ++TextureIndex)
@@ -126,7 +127,7 @@ ENGINE_API void BuildTextureStreamingData(UWorld* InWorld, const FTexCoordScaleM
 	{
 		ULevel* Level = InWorld->GetLevel(LevelIndex);
 		if (!Level) continue;
-		
+
 		for (AActor* Actor : Level->Actors)
 		{
 			if (!Actor) continue;
@@ -185,7 +186,7 @@ static void UnpackRelativeBox(const FVector& InRefOrigin, const FVector& InRefEx
 	OutExtent = .5 * (Max - Min);
 }
 
-void FStreamingTexturePrimitiveInfo::UnPackFrom(UTexture2D* InTexture, const FBoxSphereBounds& RefBounds, const FStreamingTextureBuildInfo& Info, bool bUseRelativeBox)
+void FStreamingTexturePrimitiveInfo::UnPackFrom(UTexture2D* InTexture, float ExtraScale, const FBoxSphereBounds& RefBounds, const FStreamingTextureBuildInfo& Info, bool bUseRelativeBox)
 {
 	check(InTexture->LevelIndex == Info.TextureLevelIndex);
 	Texture = InTexture;
@@ -200,7 +201,7 @@ void FStreamingTexturePrimitiveInfo::UnPackFrom(UTexture2D* InTexture, const FBo
 		Bounds = RefBounds;
 	}
 
-	TexelFactor = Info.TexelFactor;
+	TexelFactor = Info.TexelFactor * ExtraScale;
 }
 
 void FStreamingTextureBuildInfo::PackFrom(TArray<UTexture2D*>& LevelTextures, const FBoxSphereBounds& RefBounds, const FStreamingTexturePrimitiveInfo& Info)
@@ -223,8 +224,8 @@ FStreamingTextureLevelContext::FStreamingTextureLevelContext(const ULevel* InLev
 	bUseRelativeBoxes(InLevel ? !InLevel->bTextureStreamingRotationChanged : false),
 	ComponentTimestamp(0),
 	ComponentBounds(ForceInitToZero),
-	ComponentMeshScale(1.f),
-	ComponentStreamingScale(1.f)
+	ComponentPrecomputedDataScale(1.f),
+	ComponentFallbackScale(1.f)
 {
 	if (InLevel)
 	{
@@ -245,14 +246,14 @@ FStreamingTextureLevelContext::~FStreamingTextureLevelContext()
 	}
 }
 
-void FStreamingTextureLevelContext::BindComponent(const TArray<FStreamingTextureBuildInfo>* BuildData, const FBoxSphereBounds& Bounds, float MeshScale, float StreamingScale)
+void FStreamingTextureLevelContext::BindComponent(const TArray<FStreamingTextureBuildInfo>* BuildData, const FBoxSphereBounds& Bounds, float PrecomputedDataScale, float FallbackScale)
 {
 	// First processed component must use a timestamp > 0  so that default value of 0 is invalid.
 	++ComponentTimestamp;
 
 	ComponentBounds = Bounds;
-	ComponentMeshScale = MeshScale;
-	ComponentStreamingScale = StreamingScale;
+	ComponentPrecomputedDataScale = PrecomputedDataScale;
+	ComponentFallbackScale = FallbackScale;
 	ComponentBuildData = BuildData;
 
 	if (BuildData)
@@ -301,18 +302,16 @@ void FStreamingTextureLevelContext::Process(const TArray<UTexture*>& InTextures,
 			// Check if there is any precomputed data
 			if (BoundState.BuildDataTimestamp == ComponentTimestamp && ComponentBuildData)
 			{
-				StreamingTexture.UnPackFrom(Texture2D, ComponentBounds, (*ComponentBuildData)[BoundState.BuildDataIndex], bUseRelativeBoxes);
+				// The ExtraScale is only applied to the precomputed data as we assume it was already 
+				StreamingTexture.UnPackFrom(Texture2D, ComponentPrecomputedDataScale, ComponentBounds, (*ComponentBuildData)[BoundState.BuildDataIndex], bUseRelativeBoxes);
 			}
 			else // Otherwise add default entry.
 			{
 				StreamingTexture.Texture = Texture2D;
 				StreamingTexture.Bounds = ComponentBounds;
-				StreamingTexture.TexelFactor = ComponentMeshScale;
+				StreamingTexture.TexelFactor = ComponentFallbackScale;
 			}
 
-			// Component scale is not included in the build data, so that changing it does not invalidate the data.
-			StreamingTexture.TexelFactor *= ComponentStreamingScale;
-
 			// Mark this texture as processed for this component.
 			BoundState.Timestamp = ComponentTimestamp;
 		}
diff --git a/Engine/Source/Runtime/Engine/Private/Streaming/TextureStreamingHelpers.cpp b/Engine/Source/Runtime/Engine/Private/Streaming/TextureStreamingHelpers.cpp
--- a/Engine/Source/Runtime/Engine/Private/Streaming/TextureStreamingHelpers.cpp
+++ b/Engine/Source/Runtime/Engine/Private/Streaming/TextureStreamingHelpers.cpp
@@ -8,6 +8,8 @@ TextureStreamingHelpers.cpp: Definitions of classes used for texture streaming.
 #include "TextureStreamingHelpers.h"
 #include "Engine/Texture2D.h"
 #include "GenericPlatformMemoryPoolStats.h"
+#include "Engine/LightMapTexture2D.h"
+#include "Engine/ShadowMapTexture2D.h"
 
 /** Streaming stats */
 
@@ -178,7 +180,6 @@ bool GNeverStreamOutTextures = false;
 /** Accumulated total time spent on dynamic primitives, in seconds. */
 double GStreamingDynamicPrimitivesTime = 0.0;
 
-
 /**
  * Checks whether a UTexture2D is supposed to be streaming.
  * @param Texture	Texture to check
@@ -186,7 +187,14 @@ double GStreamingDynamicPrimitivesTime = 0.0;
  */
 bool IsStreamingTexture( const UTexture2D* Texture2D )
 {
-	return Texture2D && Texture2D->bIsStreamable && Texture2D->NeverStream == false && Texture2D->GetNumMips() > UTexture2D::GetMinTextureResidentMipCount();
+	if (Texture2D && Texture2D->bIsStreamable && !Texture2D->NeverStream && Texture2D->GetNumMips() > UTexture2D::GetMinTextureResidentMipCount())
+	{
+		const UShadowMapTexture2D* ShadowMap2D = Cast<const UShadowMapTexture2D>(Texture2D);
+		const ULightMapTexture2D* Lightmap2D = Cast<const ULightMapTexture2D>(Texture2D);
+
+		return !(Lightmap2D && (Lightmap2D->LightmapFlags & LMF_Streamed) == 0) && !(ShadowMap2D && (ShadowMap2D->ShadowmapFlags & SMF_Streamed) == 0);
+	}
+	return false;
 }
 
 void FStreamingContext::Reset( bool bProcessEverything, UTexture2D* IndividualStreamingTexture, bool bInCollectTextureStats )
diff --git a/Engine/Source/Runtime/Engine/Private/Streaming/TextureStreamingHelpers.h b/Engine/Source/Runtime/Engine/Private/Streaming/TextureStreamingHelpers.h
--- a/Engine/Source/Runtime/Engine/Private/Streaming/TextureStreamingHelpers.h
+++ b/Engine/Source/Runtime/Engine/Private/Streaming/TextureStreamingHelpers.h
@@ -106,7 +106,7 @@ class FTextureBoundsVisibility;
 class FDynamicComponentTextureManager;
 template<typename T>
 class FAsyncTask;
-class FAsyncTextureStreaming;
+class FAsyncTextureStreamingTask;
 
 
 struct FStreamingTexture;
@@ -156,27 +156,6 @@ struct FTextureSortElement
 	int32			NumRequiredResidentMips;
 };
 
-
-enum ETextureStreamingType
-{
-	StreamType_Static,
-	StreamType_Dynamic,
-	StreamType_Forced,
-	StreamType_LastRenderTime,
-	StreamType_Orphaned,
-	StreamType_Other,
-};
-
-static const TCHAR* GStreamTypeNames[] =
-{
-	TEXT("Static"),
-	TEXT("Dynamic"),
-	TEXT("Forced"),
-	TEXT("LastRenderTime"),
-	TEXT("Orphaned"),
-	TEXT("Other"),
-};
-
 struct FTexturePriority
 {
 	FTexturePriority( bool InCanDropMips, float InRetentionPriority, float InLoadOrderPriority, int32 InTextureIndex, const UTexture2D* InTexture)
@@ -204,14 +183,13 @@ struct FTexturePriority
 #if STATS
 struct FTextureStreamingStats
 {
-	FTextureStreamingStats( UTexture2D* InTexture2D, ETextureStreamingType InType, int32 InResidentMips, int32 InWantedMips, int32 InMostResidentMips, int32 InResidentSize, int32 InWantedSize, int32 InMaxSize, int32 InMostResidentSize, float InBoostFactor, float InPriority, int32 InTextureIndex )
+	FTextureStreamingStats( UTexture2D* InTexture2D, int32 InResidentMips, int32 InWantedMips, int32 InMostResidentMips, int32 InResidentSize, int32 InWantedSize, int32 InMaxSize, int32 InMostResidentSize, float InBoostFactor, float InPriority, int32 InTextureIndex )
 	:	TextureName( InTexture2D->GetFullName() )
 	,	SizeX( InTexture2D->GetSizeX() )
 	,	SizeY( InTexture2D->GetSizeY() )
 	,	NumMips( InTexture2D->GetNumMips() )
 	,	LODBias( InTexture2D->GetCachedLODBias() )
 	,	LastRenderTime( FMath::Clamp( InTexture2D->Resource ? (FApp::GetCurrentTime() - InTexture2D->Resource->LastRenderTime) : 1000000.0, 0.0, 1000000.0) )
-	,	StreamType( InType )
 	,	ResidentMips( InResidentMips )
 	,	WantedMips( InWantedMips )
 	,	MostResidentMips( InMostResidentMips )
@@ -236,8 +214,6 @@ struct FTextureStreamingStats
 	int32		LODBias;
 	/** How many seconds since it was last rendered: FApp::GetCurrentTime() - UTexture2D::Resource->LastRenderTime */
 	float		LastRenderTime;
-	/** What streaming heuristics were primarily used for this texture. */
-	ETextureStreamingType	StreamType;
 	/** Number of resident mip levels. */
 	int32		ResidentMips;
 	/** Number of wanted mip levels. */
diff --git a/Engine/Source/Runtime/Engine/Private/Tests/AutomationCommon.cpp b/Engine/Source/Runtime/Engine/Private/Tests/AutomationCommon.cpp
--- a/Engine/Source/Runtime/Engine/Private/Tests/AutomationCommon.cpp
+++ b/Engine/Source/Runtime/Engine/Private/Tests/AutomationCommon.cpp
@@ -194,21 +194,22 @@ bool FEngineWaitLatentCommand::Update()
 	return false;
 }
 
+ENGINE_API uint32 GStreamAllResourcesStillInFlight = -1;
 bool FStreamAllResourcesLatentCommand::Update()
 {
-	const float CommandStartTime = FPlatformTime::Seconds();
+	float LocalStartTime = FPlatformTime::Seconds();
 
-	int32 StillInFlight = IStreamingManager::Get().StreamAllResources(Duration);
+	GStreamAllResourcesStillInFlight = IStreamingManager::Get().StreamAllResources(Duration);
 
 	float Time = FPlatformTime::Seconds();
 
-	if(StillInFlight)
+	if(GStreamAllResourcesStillInFlight)
 	{
-		UE_LOG(LogEngineAutomationLatentCommand, Warning, TEXT("StreamAllResources() waited for %.2fs but %d resources are still in flight."), Time - CommandStartTime, StillInFlight);
+		UE_LOG(LogEngineAutomationLatentCommand, Warning, TEXT("StreamAllResources() waited for %.2fs but %d resources are still in flight."), Time - LocalStartTime, GStreamAllResourcesStillInFlight);
 	}
 	else
 	{
-		UE_LOG(LogEngineAutomationLatentCommand, Log, TEXT("StreamAllResources() waited for %.2fs (max duration: %.2f)."), Time - CommandStartTime, Duration);
+		UE_LOG(LogEngineAutomationLatentCommand, Log, TEXT("StreamAllResources() waited for %.2fs (max duration: %.2f)."), Time - LocalStartTime, Duration);
 	}
 
 	return true;
diff --git a/Engine/Source/Runtime/Engine/Private/Texture2D.cpp b/Engine/Source/Runtime/Engine/Private/Texture2D.cpp
--- a/Engine/Source/Runtime/Engine/Private/Texture2D.cpp
+++ b/Engine/Source/Runtime/Engine/Private/Texture2D.cpp
@@ -325,7 +325,7 @@ float UTexture2D::GetAverageBrightness(bool bIgnoreTrueBlack, bool bUseGrayscale
 
 void UTexture2D::LinkStreaming()
 {
-	if (!IsTemplate() && IStreamingManager::Get().IsTextureStreamingEnabled())
+	if (!IsTemplate() && IStreamingManager::Get().IsTextureStreamingEnabled() && IsStreamingTexture(this))
 	{
 		IStreamingManager::Get().GetTextureStreamingManager().AddStreamingTexture(this);
 	}
diff --git a/Engine/Source/Runtime/Engine/Private/TileRendering.cpp b/Engine/Source/Runtime/Engine/Private/TileRendering.cpp
--- a/Engine/Source/Runtime/Engine/Private/TileRendering.cpp
+++ b/Engine/Source/Runtime/Engine/Private/TileRendering.cpp
@@ -328,27 +328,28 @@ bool FCanvasTileRendererItem::Render_GameThread(const FCanvas* Canvas)
 		DrawTileCommand,
 		FDrawTileParameters, Parameters, DrawTileParameters,
 		{
-		SCOPED_DRAW_EVENT(RHICmdList, CanvasDrawTile);
-		for (int32 TileIdx = 0; TileIdx < Parameters.RenderData->Tiles.Num(); TileIdx++)
-		{
-			const FRenderData::FTileInst& Tile = Parameters.RenderData->Tiles[TileIdx];
-			FTileRenderer::DrawTile(
-				RHICmdList,
-				*Parameters.View,
-				Parameters.RenderData->MaterialRenderProxy,
-				Parameters.bNeedsToSwitchVerticalAxis,
-				Tile.X, Tile.Y, Tile.SizeX, Tile.SizeY,
-				Tile.U, Tile.V, Tile.SizeU, Tile.SizeV,
-				Parameters.bIsHitTesting, Tile.HitProxyId,
-				Tile.InColor);
-		}
-
-		delete Parameters.View->Family;
-		delete Parameters.View;
-		if (Parameters.AllowedCanvasModes & FCanvas::Allow_DeleteOnRender)
-		{
-			delete Parameters.RenderData;
-		}
+			SCOPED_DRAW_EVENTF(RHICmdList, CanvasDrawTile, *Parameters.RenderData->MaterialRenderProxy->GetMaterial(GMaxRHIFeatureLevel)->GetFriendlyName());
+			
+			for (int32 TileIdx = 0; TileIdx < Parameters.RenderData->Tiles.Num(); TileIdx++)
+			{
+				const FRenderData::FTileInst& Tile = Parameters.RenderData->Tiles[TileIdx];
+				FTileRenderer::DrawTile(
+					RHICmdList,
+					*Parameters.View,
+					Parameters.RenderData->MaterialRenderProxy,
+					Parameters.bNeedsToSwitchVerticalAxis,
+					Tile.X, Tile.Y, Tile.SizeX, Tile.SizeY,
+					Tile.U, Tile.V, Tile.SizeU, Tile.SizeV,
+					Parameters.bIsHitTesting, Tile.HitProxyId,
+					Tile.InColor);
+			}
+
+			delete Parameters.View->Family;
+			delete Parameters.View;
+			if (Parameters.AllowedCanvasModes & FCanvas::Allow_DeleteOnRender)
+			{
+				delete Parameters.RenderData;
+			}
 		});
 	if (Canvas->GetAllowedModes() & FCanvas::Allow_DeleteOnRender)
 	{
diff --git a/Engine/Source/Runtime/Engine/Private/UnrealEngine.cpp b/Engine/Source/Runtime/Engine/Private/UnrealEngine.cpp
--- a/Engine/Source/Runtime/Engine/Private/UnrealEngine.cpp
+++ b/Engine/Source/Runtime/Engine/Private/UnrealEngine.cpp
@@ -244,7 +244,6 @@ FCachedSystemScalabilityCVars::FCachedSystemScalabilityCVars()
 	, MaxShadowResolution(-1)
 	, ViewDistanceScale(-1)
 	, ViewDistanceScaleSquared(-1)
-	, SimpleDynamicLighting(-1)
 	, MaxAnisotropy(-1)
 {
 
@@ -282,11 +281,6 @@ void ScalabilityCVarsSinkCallback()
 		LocalScalabilityCVars.MaterialQualityLevel = (EMaterialQualityLevel::Type)FMath::Clamp(MaterialQualityLevelVar->GetValueOnGameThread(), 0, (int32)EMaterialQualityLevel::Num - 1);
 	}
 
-	{
-		static IConsoleVariable* CVar = IConsoleManager::Get().FindConsoleVariable(TEXT("r.SimpleDynamicLighting"));
-		LocalScalabilityCVars.SimpleDynamicLighting = FMath::Clamp(CVar->GetInt(), 0, 1);
-	}
-
 	LocalScalabilityCVars.bInitialized = true;
 
 	if (!GCachedScalabilityCVars.bInitialized)
@@ -299,12 +293,6 @@ void ScalabilityCVarsSinkCallback()
 		bool bRecreateRenderstate = false;
 		bool bCacheResourceShaders = false;
 
-		if (LocalScalabilityCVars.DetailMode != GCachedScalabilityCVars.DetailMode ||
-			LocalScalabilityCVars.SimpleDynamicLighting != GCachedScalabilityCVars.SimpleDynamicLighting)
-		{
-			bRecreateRenderstate = true;
-		}
-
 		if (LocalScalabilityCVars.MaterialQualityLevel != GCachedScalabilityCVars.MaterialQualityLevel)
 		{
 			bCacheResourceShaders = true;
@@ -3638,7 +3626,7 @@ bool UEngine::HandleListTexturesCommand( const TCHAR* Cmd, FOutputDevice& Ar )
 	TArray<FSortedTexture> SortedTextures;
 	for( TObjectIterator<UTexture2D> It; It; ++It )
 	{
-		UTexture2D*		Texture				= *It;
+		UTexture2D*			Texture				= *It;
 		int32				LODGroup			= Texture->LODGroup;
 		int32				LODBias				= Texture->GetCachedLODBias();
 		int32				NumMips				= Texture->GetNumMips();	
@@ -3651,7 +3639,7 @@ bool UEngine::HandleListTexturesCommand( const TCHAR* Cmd, FOutputDevice& Ar )
 		int32				DroppedMips			= Texture->GetNumMips() - Texture->ResidentMips;
 		int32				CurSizeX			= Texture->GetSizeX() >> DroppedMips;
 		int32				CurSizeY			= Texture->GetSizeY() >> DroppedMips;
-		bool			bIsStreamingTexture		= IStreamingManager::Get().IsTextureStreamingEnabled() ? IStreamingManager::Get().GetTextureStreamingManager().IsManagedStreamingTexture( Texture ) : false;
+		bool			bIsStreamingTexture		= Texture->GetStreamingIndex() != INDEX_NONE;
 		int32				MaxSize				= Texture->CalcTextureMemorySizeEnum( TMC_AllMips );
 		int32				CurrentSize			= Texture->CalcTextureMemorySizeEnum( TMC_ResidentMips );
 		int32				UsageCount			= TextureToUsageMap.FindRef( Texture );
diff --git a/Engine/Source/Runtime/Engine/Private/UserInterface/Canvas.cpp b/Engine/Source/Runtime/Engine/Private/UserInterface/Canvas.cpp
--- a/Engine/Source/Runtime/Engine/Private/UserInterface/Canvas.cpp
+++ b/Engine/Source/Runtime/Engine/Private/UserInterface/Canvas.cpp
@@ -219,6 +219,7 @@ FCanvas::FCanvas(FRenderTarget* InRenderTarget, FHitProxyConsumer* InHitProxyCon
 ,	CurrentDeltaWorldTime(0)
 ,	FeatureLevel(InFeatureLevel)
 ,	StereoDepth(150)
+,	DrawMode(CDM_DeferDrawing)
 {
 	Construct();
 
@@ -230,7 +231,7 @@ FCanvas::FCanvas(FRenderTarget* InRenderTarget, FHitProxyConsumer* InHitProxyCon
 	}
 }
 
-FCanvas::FCanvas(FRenderTarget* InRenderTarget,FHitProxyConsumer* InHitProxyConsumer, float InRealTime, float InWorldTime, float InWorldDeltaTime, ERHIFeatureLevel::Type InFeatureLevel)
+FCanvas::FCanvas(FRenderTarget* InRenderTarget,FHitProxyConsumer* InHitProxyConsumer, float InRealTime, float InWorldTime, float InWorldDeltaTime, ERHIFeatureLevel::Type InFeatureLevel, ECanvasDrawMode InDrawMode)
 :	ViewRect(0,0,0,0)
 ,	RenderTarget(InRenderTarget)
 ,	HitProxyConsumer(InHitProxyConsumer)
@@ -241,6 +242,7 @@ FCanvas::FCanvas(FRenderTarget* InRenderTarget,FHitProxyConsumer* InHitProxyCons
 ,	CurrentDeltaWorldTime(InWorldDeltaTime)
 ,	FeatureLevel(InFeatureLevel)
 ,	StereoDepth(150)
+,	DrawMode(InDrawMode)
 {
 	Construct();
 }
@@ -757,8 +759,6 @@ void FCanvas::Flush_GameThread(bool bForce)
 		CanvasFlushSetupCommand,
 		FCanvasFlushParameters,Parameters,FlushParameters,
 	{
-		SCOPED_DRAW_EVENT(RHICmdList, CanvasFlush);
-
 		// Set the RHI render target.
 		::SetRenderTarget(RHICmdList, Parameters.CanvasRenderTarget->GetRenderTargetTexture(), FTextureRHIRef(), true);
 		// disable depth test & writes
@@ -1792,6 +1792,11 @@ void FCanvas::DrawItem(FCanvasItem& Item)
 	{
 		Item.Draw(this);
 	}
+
+	if (DrawMode == CDM_ImmediateDrawing)
+	{
+		Flush_GameThread();
+	}
 }
 
 void FCanvas::DrawItem(FCanvasItem& Item, const FVector2D& InPosition)
@@ -1813,6 +1818,11 @@ void FCanvas::DrawItem(FCanvasItem& Item, const FVector2D& InPosition)
 	{
 		Item.Draw(this , InPosition);
 	}
+
+	if (DrawMode == CDM_ImmediateDrawing)
+	{
+		Flush_GameThread();
+	}
 }
 
 void FCanvas::DrawItem(FCanvasItem& Item, float X, float Y)
@@ -1834,6 +1844,11 @@ void FCanvas::DrawItem(FCanvasItem& Item, float X, float Y)
 	{
 		Item.Draw(this, X, Y);
 	}
+
+	if (DrawMode == CDM_ImmediateDrawing)
+	{
+		Flush_GameThread();
+	}
 }
 
 void UCanvas::SetView(FSceneView* InView)
diff --git a/Engine/Source/Runtime/Engine/Private/World.cpp b/Engine/Source/Runtime/Engine/Private/World.cpp
--- a/Engine/Source/Runtime/Engine/Private/World.cpp
+++ b/Engine/Source/Runtime/Engine/Private/World.cpp
@@ -2646,17 +2646,6 @@ void UWorld::TriggerStreamingDataRebuild()
 
 void UWorld::ConditionallyBuildStreamingData()
 {
-#if !(UE_BUILD_SHIPPING || UE_BUILD_TEST)
-	// This code is to trigger an update to the a streaming data update when the console variable is changed.
-	static int32 LastMetricsUsed = CVarStreamingUseNewMetrics.GetValueOnGameThread();
-	if (CVarStreamingUseNewMetrics.GetValueOnGameThread() != LastMetricsUsed)
-	{
-		LastMetricsUsed = CVarStreamingUseNewMetrics.GetValueOnGameThread();
-		bStreamingDataDirty = true;
-		BuildStreamingDataTimer = FPlatformTime::Seconds() - 1.0;
-	}
-#endif
-
 	if ( bStreamingDataDirty && FPlatformTime::Seconds() > BuildStreamingDataTimer )
 	{
 		bStreamingDataDirty = false;
diff --git a/Engine/Source/Runtime/Engine/Public/AudioStreaming.h b/Engine/Source/Runtime/Engine/Public/AudioStreaming.h
--- a/Engine/Source/Runtime/Engine/Public/AudioStreaming.h
+++ b/Engine/Source/Runtime/Engine/Public/AudioStreaming.h
@@ -191,7 +191,7 @@ struct FAudioStreamingManager : public IAudioStreamingManager
 	virtual void CancelForcedResources() override;
 	virtual void NotifyLevelChange() override;
 	virtual void SetDisregardWorldResourcesForFrames( int32 NumFrames ) override;
-	virtual void AddPreparedLevel( class ULevel* Level ) override;
+	virtual void AddLevel( class ULevel* Level ) override;
 	virtual void RemoveLevel( class ULevel* Level ) override;
 	// End IStreamingManager interface
 
diff --git a/Engine/Source/Runtime/Engine/Public/CanvasTypes.h b/Engine/Source/Runtime/Engine/Public/CanvasTypes.h
--- a/Engine/Source/Runtime/Engine/Public/CanvasTypes.h
+++ b/Engine/Source/Runtime/Engine/Public/CanvasTypes.h
@@ -218,6 +218,12 @@ public:
 		Allow_DeleteOnRender= 1<<1
 	};
 
+	enum ECanvasDrawMode
+	{
+		CDM_DeferDrawing,
+		CDM_ImmediateDrawing
+	};
+
 	/** 
 	* Constructor.
 	*/
@@ -226,7 +232,7 @@ public:
 	/** 
 	* Constructor. For situations where a world is not available, but time information is
 	*/
-	ENGINE_API FCanvas(FRenderTarget* InRenderTarget, FHitProxyConsumer* InHitProxyConsumer, float InRealTime, float InWorldTime, float InWorldDeltaTime, ERHIFeatureLevel::Type InFeatureLevel);
+	ENGINE_API FCanvas(FRenderTarget* InRenderTarget, FHitProxyConsumer* InHitProxyConsumer, float InRealTime, float InWorldTime, float InWorldDeltaTime, ERHIFeatureLevel::Type InFeatureLevel, ECanvasDrawMode DrawMode = CDM_DeferDrawing);
 
 	/** 
 	* Destructor.
@@ -599,6 +605,8 @@ private:
 	FMatrix CachedOrthoProjection[2];
 	int32 CachedRTWidth, CachedRTHeight, CachedDrawDepth;
 
+	ECanvasDrawMode DrawMode;
+
 	bool GetOrthoProjectionMatrices(float InDrawDepth, FMatrix OutOrthoProjection[2]);
 
 	/** 
diff --git a/Engine/Source/Runtime/Engine/Public/ContentStreaming.h b/Engine/Source/Runtime/Engine/Public/ContentStreaming.h
--- a/Engine/Source/Runtime/Engine/Public/ContentStreaming.h
+++ b/Engine/Source/Runtime/Engine/Public/ContentStreaming.h
@@ -187,7 +187,7 @@ struct IStreamingManager
 	}
 
 	/** Adds a ULevel that has already prepared StreamingData to the streaming manager. */
-	virtual void AddPreparedLevel( class ULevel* Level ) = 0;
+	virtual void AddLevel( class ULevel* Level ) = 0;
 
 	/** Removes a ULevel from the streaming manager. */
 	virtual void RemoveLevel( class ULevel* Level ) = 0;
@@ -386,9 +386,6 @@ struct ITextureStreamingManager : public IStreamingManager
 	/** Removes a texture from the streaming manager. */
 	virtual void RemoveStreamingTexture(UTexture2D* Texture) = 0;
 
-	/** Returns true if this is a streaming texture that is managed by the streaming manager. */
-	virtual bool IsManagedStreamingTexture(const UTexture2D* Texture2D) = 0;
-
 	virtual int64 GetMemoryOverBudget() const = 0;
 
 	/** Max required textures ever seen in bytes. */
@@ -569,14 +566,11 @@ struct FStreamingManagerCollection : public IStreamingManager
 	virtual bool Exec( UWorld* InWorld, const TCHAR* Cmd, FOutputDevice& Ar ) override;
 
 	/** Adds a ULevel to the streaming manager. */
-	virtual void AddLevel( class ULevel* Level );
+	virtual void AddLevel( class ULevel* Level ) override;
 
 	/** Removes a ULevel from the streaming manager. */
 	virtual void RemoveLevel( class ULevel* Level ) override;
 
-	/** Adds a ULevel that has already prepared StreamingData to the streaming manager. */
-	virtual void AddPreparedLevel( class ULevel* Level ) override;
-
 	/** Called when an actor is spawned. */
 	virtual void NotifyActorSpawned( AActor* Actor ) override;
 
diff --git a/Engine/Source/Runtime/Engine/Public/DynamicMeshBuilder.h b/Engine/Source/Runtime/Engine/Public/DynamicMeshBuilder.h
--- a/Engine/Source/Runtime/Engine/Public/DynamicMeshBuilder.h
+++ b/Engine/Source/Runtime/Engine/Public/DynamicMeshBuilder.h
@@ -48,7 +48,10 @@ struct FDynamicMeshVertex
 
 	FVector GetTangentY()
 	{
-		return (FVector(TangentZ) ^ FVector(TangentX)) * ((float)TangentZ.Vector.W / 127.5f - 1.0f);
+		FVector TanX = TangentX;
+		FVector TanZ = TangentZ;
+
+		return (TanZ ^ TanX) * ((float)TangentZ.Vector.W / 127.5f - 1.0f);
 	};
 
 	FVector Position;
diff --git a/Engine/Source/Runtime/Engine/Public/MaterialShared.h b/Engine/Source/Runtime/Engine/Public/MaterialShared.h
--- a/Engine/Source/Runtime/Engine/Public/MaterialShared.h
+++ b/Engine/Source/Runtime/Engine/Public/MaterialShared.h
@@ -98,6 +98,7 @@ enum EMaterialCommonBasis
 	MCB_World,
 	MCB_View,
 	MCB_Camera,
+	MCB_MeshParticle,
 	MCB_MAX,
 };
 
@@ -1001,7 +1002,9 @@ public:
 	virtual bool UseLmDirectionality() const { return true; }
 	virtual bool IsMasked() const = 0;
 	virtual bool IsDitherMasked() const { return false; }
+	virtual bool AllowNegativeEmissiveColor() const { return false; }
 	virtual enum EBlendMode GetBlendMode() const = 0;
+	ENGINE_API virtual enum ERefractionMode GetRefractionMode() const;
 	virtual enum EMaterialShadingModel GetShadingModel() const = 0;
 	virtual enum ETranslucencyLightingMode GetTranslucencyLightingMode() const { return TLM_VolumetricNonDirectional; };
 	virtual float GetOpacityMaskClipValue() const = 0;
@@ -1616,6 +1619,7 @@ public:
 	ENGINE_API virtual bool IsNonmetal() const override;
 	ENGINE_API virtual bool UseLmDirectionality() const override;
 	ENGINE_API virtual enum EBlendMode GetBlendMode() const override;
+	ENGINE_API virtual enum ERefractionMode GetRefractionMode() const override;
 	ENGINE_API virtual uint32 GetDecalBlendMode() const override;
 	ENGINE_API virtual uint32 GetMaterialDecalResponse() const override;
 	ENGINE_API virtual bool HasNormalConnected() const override;
@@ -1635,6 +1639,7 @@ public:
 	ENGINE_API virtual float GetTranslucentShadowStartOffset() const override;
 	ENGINE_API virtual bool IsMasked() const override;
 	ENGINE_API virtual bool IsDitherMasked() const override;
+	ENGINE_API virtual bool AllowNegativeEmissiveColor() const override;
 	ENGINE_API virtual FString GetFriendlyName() const override;
 	ENGINE_API virtual bool RequiresSynchronousCompilation() const override;
 	ENGINE_API virtual bool IsDefaultMaterial() const override;
diff --git a/Engine/Source/Runtime/Engine/Public/ParameterCollection.h b/Engine/Source/Runtime/Engine/Public/ParameterCollection.h
--- a/Engine/Source/Runtime/Engine/Public/ParameterCollection.h
+++ b/Engine/Source/Runtime/Engine/Public/ParameterCollection.h
@@ -51,4 +51,7 @@ private:
 	FRHIUniformBufferLayout UniformBufferLayout;
 
 	void UpdateContents(const FGuid& InId, const TArray<FVector4>& Data);
-};
\ No newline at end of file
+};
+
+// Default instance resources used when rendering a material using a parameter collection but there's no FScene present to get a FMaterialParameterCollectionInstanceResource
+extern ENGINE_API TMap<FGuid, FMaterialParameterCollectionInstanceResource*> GDefaultMaterialParameterCollectionInstances;
\ No newline at end of file
diff --git a/Engine/Source/Runtime/Engine/Public/ParticleEmitterInstances.h b/Engine/Source/Runtime/Engine/Public/ParticleEmitterInstances.h
--- a/Engine/Source/Runtime/Engine/Public/ParticleEmitterInstances.h
+++ b/Engine/Source/Runtime/Engine/Public/ParticleEmitterInstances.h
@@ -276,6 +276,10 @@ public:
 	float EmitterTime;
 	/** The amount of time simulated in the previous time step. */
 	float LastDeltaTime;
+	/** how long did the last tick take? */
+#if WITH_EDITOR
+	float LastTickDurationMs;
+#endif
 	/** The previous location of the instance.							*/
 	FVector OldLocation;
 	/** The bounding box for the particles.								*/
@@ -649,6 +653,15 @@ public:
 	virtual void SetBeamTargetTangent(FVector NewTangentPoint,int32 TargetIndex) {};
 	virtual void SetBeamTargetStrength(float NewTargetStrength,int32 TargetIndex) {};
 
+	//Beam get interface
+	virtual bool GetBeamEndPoint(FVector& OutEndPoint) const { return false; }
+	virtual bool GetBeamSourcePoint(int32 SourceIndex, FVector& OutSourcePoint) const { return false; }
+	virtual bool GetBeamSourceTangent(int32 SourceIndex, FVector& OutTangentPoint) const { return false; }
+	virtual bool GetBeamSourceStrength(int32 SourceIndex, float& OutSourceStrength) const { return false; }
+	virtual bool GetBeamTargetPoint(int32 TargetIndex, FVector& OutTargetPoint) const { return false; }
+	virtual bool GetBeamTargetTangent(int32 TargetIndex, FVector& OutTangentPoint) const { return false; }
+	virtual bool GetBeamTargetStrength(int32 TargetIndex, float& OutTargetStrength) const { return false; }
+	
 	// Called on world origin changes
 	virtual void ApplyWorldOffset(FVector InOffset, bool bWorldShift);
 		
@@ -994,6 +1007,13 @@ struct FParticleBeam2EmitterInstance : public FParticleEmitterInstance
 	virtual void SetBeamTargetStrength(float NewTargetStrength,int32 TargetIndex) override;
 	virtual void ApplyWorldOffset(FVector InOffset, bool bWorldShift) override;
 	
+	virtual bool GetBeamEndPoint(FVector& OutEndPoint) const override;
+	virtual bool GetBeamSourcePoint(int32 SourceIndex, FVector& OutSourcePoint) const override;
+	virtual bool GetBeamSourceTangent(int32 SourceIndex, FVector& OutTangentPoint) const override;
+	virtual bool GetBeamSourceStrength(int32 SourceIndex, float& OutSourceStrength) const override;
+	virtual bool GetBeamTargetPoint(int32 TargetIndex, FVector& OutTargetPoint) const override;
+	virtual bool GetBeamTargetTangent(int32 TargetIndex, FVector& OutTangentPoint) const override;
+	virtual bool GetBeamTargetStrength(int32 TargetIndex, float& OutTargetStrength) const override;
 	//
 	virtual void InitParameters(UParticleEmitter* InTemplate, UParticleSystemComponent* InComponent) override;
 	virtual void Init() override;
diff --git a/Engine/Source/Runtime/Engine/Public/PrimitiveSceneProxy.h b/Engine/Source/Runtime/Engine/Public/PrimitiveSceneProxy.h
--- a/Engine/Source/Runtime/Engine/Public/PrimitiveSceneProxy.h
+++ b/Engine/Source/Runtime/Engine/Public/PrimitiveSceneProxy.h
@@ -411,7 +411,7 @@ public:
 	inline bool CastsFarShadow() const { return bCastFarShadow; }
 	inline bool LightAsIfStatic() const { return bLightAsIfStatic; }
 	inline bool LightAttachmentsAsGroup() const { return bLightAttachmentsAsGroup; }
-	inline bool UseSingleSampleShadowFromStationaryLights() const { return bSingleSampleShadowFromStationaryLights; }
+	ENGINE_API bool UseSingleSampleShadowFromStationaryLights() const;
 	inline bool StaticElementsAlwaysUseProxyPrimitiveUniformBuffer() const { return bStaticElementsAlwaysUseProxyPrimitiveUniformBuffer; }
 	inline bool ShouldUseAsOccluder() const { return bUseAsOccluder; }
 	inline bool AllowApproximateOcclusion() const { return bAllowApproximateOcclusion; }
@@ -521,12 +521,13 @@ public:
 
 	/**
 	 * Get the temp data used in the texture streaming build. Used to debug accuracy.
-	 * @param OutDistanceMultiplier	The component streaming distance multiplier.
+	 * @param OutComponentExtraScale The component streaming resolution extra scale.
+	 * @param OutMeshExtraScale	The mesh streaming resolution extra scale.
 	 * @param LODIndex			LOD index of the query. INDEX_NONE for a value for all LODs.
 	 * @param ElementIndex		Element index of the query. INDEX_NONE for a value for all elements.
 	 * @return					The section data built in the texture streaming build.
 	 */
-	virtual const FStreamingSectionBuildInfo* GetStreamingSectionData(float& OutDistanceMultiplier, int32 LODIndex, int32 ElementIndex) const { return nullptr; }
+	virtual const FStreamingSectionBuildInfo* GetStreamingSectionData(float& OutComponentExtraScale, float& OutMeshExtraScale, int32 LODIndex, int32 ElementIndex) const { return nullptr; }
 
 	/**
 	* Get the lightmap resolution for this primitive. Used in VMI_LightmapDensity.
diff --git a/Engine/Source/Runtime/Engine/Public/SceneInterface.h b/Engine/Source/Runtime/Engine/Public/SceneInterface.h
--- a/Engine/Source/Runtime/Engine/Public/SceneInterface.h
+++ b/Engine/Source/Runtime/Engine/Public/SceneInterface.h
@@ -147,9 +147,6 @@ public:
 	 */
 	virtual void UpdateLightColorAndBrightness(ULightComponent* Light) = 0;
 
-	/** Updates the scene's dynamic skylight. */
-	virtual void UpdateDynamicSkyLight(const FLinearColor& UpperColor, const FLinearColor& LowerColor) {}
-
 	/** Sets the precomputed visibility handler for the scene, or NULL to clear the current one. */
 	virtual void SetPrecomputedVisibility(const class FPrecomputedVisibilityHandler* PrecomputedVisibilityHandler) {}
 
diff --git a/Engine/Source/Runtime/Engine/Public/SceneView.h b/Engine/Source/Runtime/Engine/Public/SceneView.h
--- a/Engine/Source/Runtime/Engine/Public/SceneView.h
+++ b/Engine/Source/Runtime/Engine/Public/SceneView.h
@@ -478,8 +478,6 @@ BEGIN_UNIFORM_BUFFER_STRUCT_WITH_CONSTRUCTOR(FFrameUniformShaderParameters, ENGI
 	DECLARE_UNIFORM_BUFFER_STRUCT_MEMBER_EX(FVector4, DirectionalLightShadowSize, EShaderPrecisionModifier::Half)
 	DECLARE_UNIFORM_BUFFER_STRUCT_MEMBER_ARRAY(FMatrix, DirectionalLightScreenToShadow, [MAX_FORWARD_SHADOWCASCADES])
 	DECLARE_UNIFORM_BUFFER_STRUCT_MEMBER_EX(FVector4, DirectionalLightShadowDistances, EShaderPrecisionModifier::Half)
-	DECLARE_UNIFORM_BUFFER_STRUCT_MEMBER_EX(FLinearColor, UpperSkyColor, EShaderPrecisionModifier::Half)
-	DECLARE_UNIFORM_BUFFER_STRUCT_MEMBER_EX(FLinearColor, LowerSkyColor, EShaderPrecisionModifier::Half)
 	DECLARE_UNIFORM_BUFFER_STRUCT_MEMBER_ARRAY(FVector4, TranslucencyLightingVolumeMin, [TVC_MAX])
 	DECLARE_UNIFORM_BUFFER_STRUCT_MEMBER_ARRAY(FVector4, TranslucencyLightingVolumeInvSize, [TVC_MAX])
 	DECLARE_UNIFORM_BUFFER_STRUCT_MEMBER(FVector4, TemporalAAParams)
@@ -933,13 +931,17 @@ public:
 // for r.DisplayInternals (allows for easy passing down data from main to render thread)
 struct FDisplayInternalsData
 {
+	//
 	int32 DisplayInternalsCVarValue;
 	// current time Matinee location (in seconds) of the single playing playing actor, -1 if none is playing, -2 if multiple are playing
 	float MatineeTime;
+	// -1 if not set, from IStreamingManager::Get().StreamAllResources(Duration) in FStreamAllResourcesLatentCommand
+	uint32 NumPendingStreamingRequests;
 
 	FDisplayInternalsData()
 		: DisplayInternalsCVarValue(0)
 		, MatineeTime(-1.0f)
+		, NumPendingStreamingRequests(-1)
 	{
 		check(!IsValid());
 	}
