Fix issue with curves coming from DDC being incorrect (recreated from 4.12 branch CL 2973609)

diff --git a/Engine/Source/Runtime/Engine/Private/Animation/AnimCompressionDerivedData.h b/Engine/Source/Runtime/Engine/Private/Animation/AnimCompressionDerivedData.h
--- a/Engine/Source/Runtime/Engine/Private/Animation/AnimCompressionDerivedData.h
+++ b/Engine/Source/Runtime/Engine/Private/Animation/AnimCompressionDerivedData.h
@@ -43,7 +43,7 @@ public:
 		// This is a version string that mimics the old versioning scheme. If you
 		// want to bump this version, generate a new guid using VS->Tools->Create GUID and
 		// return it here. Ex.
-		return TEXT("363A1E830FCA4738BA7FD1A2704A9F40");
+		return TEXT("A9288501F8C6460B80E9E07AD56FB59C");
 	}
 
 	virtual FString GetPluginSpecificCacheKeySuffix() const override;
diff --git a/Engine/Source/Runtime/Engine/Private/Animation/AnimSequence.cpp b/Engine/Source/Runtime/Engine/Private/Animation/AnimSequence.cpp
--- a/Engine/Source/Runtime/Engine/Private/Animation/AnimSequence.cpp
+++ b/Engine/Source/Runtime/Engine/Private/Animation/AnimSequence.cpp
@@ -2264,6 +2264,10 @@ void UAnimSequence::BakeOutAdditiveIntoRawData()
 	TrackToSkeletonMapTable = MoveTemp(NewTrackToSkeletonMapTable);
 	RawCurveData = NewCurveTracks;
 
+	const FSmartNameMapping* Mapping = GetSkeleton()->GetSmartNameContainer(USkeleton::AnimCurveMappingName);
+	check(Mapping); // Should always exist
+	RawCurveData.UpdateLastObservedNames(Mapping);
+
 #if 0 //Validate baked data
 	for (FRawAnimSequenceTrack& RawTrack : RawAnimationData)
 	{
