Fix UT link errors in shipping due to being non-monolithic.

diff --git a/Engine/Source/Runtime/Engine/Public/ShadowMap.h b/Engine/Source/Runtime/Engine/Public/ShadowMap.h
--- a/Engine/Source/Runtime/Engine/Public/ShadowMap.h
+++ b/Engine/Source/Runtime/Engine/Public/ShadowMap.h
@@ -331,6 +331,7 @@ public:
 	virtual FShadowMap2D* GetShadowMap2D() { return this; }
 	virtual const FShadowMap2D* GetShadowMap2D() const { return this; }
 
+#if WITH_EDITOR
 	/** Call to enable/disable status update of LightMap encoding */
 	static void SetStatusUpdate(bool bInEnable)
 	{
@@ -341,6 +342,7 @@ public:
 	{
 		return bUpdateStatus;
 	}
+#endif
 
 protected:
 
@@ -359,8 +361,10 @@ protected:
 	/** Stores the inverse of the penumbra size, normalized.  Stores 1 to interpret the shadowmap as a shadow factor directly, instead of as a distance field. */
 	FVector4 InvUniformPenumbraSize;
 
+#if WITH_EDITOR
 	/** If true, update the status when encoding light maps */
 	static bool bUpdateStatus;
+#endif
 };
 
 
diff --git a/Engine/Source/Runtime/Engine/Public/Tests/AutomationCommon.h b/Engine/Source/Runtime/Engine/Public/Tests/AutomationCommon.h
--- a/Engine/Source/Runtime/Engine/Public/Tests/AutomationCommon.h
+++ b/Engine/Source/Runtime/Engine/Public/Tests/AutomationCommon.h
@@ -6,6 +6,8 @@
 #include "AutomationTest.h"
 #include "Delegate.h"
 
+#if (WITH_DEV_AUTOMATION_TESTS || WITH_PERF_AUTOMATION_TESTS)
+
 ///////////////////////////////////////////////////////////////////////
 // Common Latent commands which are used across test type. I.e. Engine, Network, etc...
 
@@ -176,3 +178,5 @@ DEFINE_LATENT_AUTOMATION_COMMAND_ONE_PARAMETER(FMatineePerformanceCaptureCommand
 * Latent command to run an exec command that also requires a UWorld.
 */
 DEFINE_LATENT_AUTOMATION_COMMAND_ONE_PARAMETER(FExecWorldStringLatentCommand, FString, ExecCommand);
+
+#endif
