Put back one automation function which is still referenced without enabled test defines.

diff --git a/Engine/Source/Runtime/Engine/Public/Tests/AutomationCommon.h b/Engine/Source/Runtime/Engine/Public/Tests/AutomationCommon.h
--- a/Engine/Source/Runtime/Engine/Public/Tests/AutomationCommon.h
+++ b/Engine/Source/Runtime/Engine/Public/Tests/AutomationCommon.h
@@ -16,6 +16,8 @@ ENGINE_API DECLARE_LOG_CATEGORY_EXTERN(LogEngineAutomationTests, Log, All);
 
 DECLARE_MULTICAST_DELEGATE_TwoParams(FOnEditorAutomationMapLoad, const FString&, FString*);
 
+#endif
+
 /** Common automation functions */
 namespace AutomationCommon
 {
@@ -68,6 +70,8 @@ namespace AutomationCommon
 		return HardwareDetailsString;
 	}
 
+#if (WITH_DEV_AUTOMATION_TESTS || WITH_PERF_AUTOMATION_TESTS)
+
 	/** Gets a path used for automation testing (PNG sent to the AutomationTest folder) */
 	static void GetScreenshotPath(const FString& TestName, FString& OutScreenshotName, const bool bIncludeHardwareDetails)
 	{
@@ -88,8 +92,12 @@ namespace AutomationCommon
 	{
 		return OnEditorAutomationMapLoad;
 	}
+
+#endif
 }
 
+#if (WITH_DEV_AUTOMATION_TESTS || WITH_PERF_AUTOMATION_TESTS)
+
 /**
  * Parameters to the Latent Automation command FTakeEditorScreenshotCommand
  */
