Fix when apply curve to nullptr component because the component disappeared.

diff --git a/Engine/Source/Runtime/Engine/Private/Animation/AnimInstance.cpp b/Engine/Source/Runtime/Engine/Private/Animation/AnimInstance.cpp
--- a/Engine/Source/Runtime/Engine/Private/Animation/AnimInstance.cpp
+++ b/Engine/Source/Runtime/Engine/Private/Animation/AnimInstance.cpp
@@ -1144,7 +1144,8 @@ void UAnimInstance::UpdateCurves(const FBlendedHeapCurve& InCurve)
 	}
 
 	// update curves to component
-	Proxy.UpdateCurvesToComponents();
+	// Proxy might not have SkeletalMeshComponent
+	Proxy.UpdateCurvesToComponents(GetOwningComponent());
 }
 
 bool UAnimInstance::HasMorphTargetCurves() const
diff --git a/Engine/Source/Runtime/Engine/Private/Animation/AnimInstanceProxy.cpp b/Engine/Source/Runtime/Engine/Private/Animation/AnimInstanceProxy.cpp
--- a/Engine/Source/Runtime/Engine/Private/Animation/AnimInstanceProxy.cpp
+++ b/Engine/Source/Runtime/Engine/Private/Animation/AnimInstanceProxy.cpp
@@ -743,14 +743,12 @@ void FAnimInstanceProxy::EvaluateAnimation(FPoseContext& Output)
 
 void FAnimInstanceProxy::UpdateCurvesToComponents(USkeletalMeshComponent* Component)
 {
+	// AnimInstanceProxy has SkeletalMeshComponent but it's not safe to be used everywhere
+	// so please give us Component to update
 	if (Component)
 	{
 		Component->ApplyAnimationCurvesToComponent(&MaterialParameterCurves, &MorphTargetCurves);
 	}
-	else
-	{
-		SkeletalMeshComponent->ApplyAnimationCurvesToComponent(&MaterialParameterCurves, &MorphTargetCurves);
-	}
 }
 
 // for now disable becauase it will not work with single node instance
diff --git a/Engine/Source/Runtime/Engine/Public/Animation/AnimInstanceProxy.h b/Engine/Source/Runtime/Engine/Public/Animation/AnimInstanceProxy.h
--- a/Engine/Source/Runtime/Engine/Public/Animation/AnimInstanceProxy.h
+++ b/Engine/Source/Runtime/Engine/Public/Animation/AnimInstanceProxy.h
@@ -366,7 +366,7 @@ protected:
 	void RecalcRequiredBones(USkeletalMeshComponent* Component, UObject* Asset);
 
 	/** Update the material parameters of the supplied component from this instance */
-	void UpdateCurvesToComponents(USkeletalMeshComponent* Component=nullptr);
+	void UpdateCurvesToComponents(USkeletalMeshComponent* Component);
 
 	/** Get Currently active montage evaluation state.
 		Note that there might be multiple Active at the same time. This will only return the first active one it finds. **/
