Copying //UE4/Dev-Rendering to //UE4/Dev-Main (Source: //UE4/Dev-Rendering @ 2998063)

diff --git a/Engine/Source/Runtime/Engine/Private/Components/SceneCaptureComponent.cpp b/Engine/Source/Runtime/Engine/Private/Components/SceneCaptureComponent.cpp
--- a/Engine/Source/Runtime/Engine/Private/Components/SceneCaptureComponent.cpp
+++ b/Engine/Source/Runtime/Engine/Private/Components/SceneCaptureComponent.cpp
@@ -16,6 +16,9 @@
 #include "Components/PlanarReflectionComponent.h"
 #include "PlanarReflectionSceneProxy.h"
 #include "Components/BoxComponent.h"
+#include "MessageLog.h"
+
+#define LOCTEXT_NAMESPACE "SceneCaptureComponent"
 
 ASceneCapture::ASceneCapture(const FObjectInitializer& ObjectInitializer)
 	: Super(ObjectInitializer)
@@ -147,7 +150,7 @@ void ASceneCaptureCube::PostEditMove(bool bFinished)
 	Super::PostEditMove(bFinished);
 	if(bFinished)
 	{
-		CaptureComponentCube->UpdateContent();
+		CaptureComponentCube->CaptureSceneDeferred();
 	}
 }
 #endif
@@ -296,6 +299,7 @@ USceneCaptureComponent2D::USceneCaptureComponent2D(const FObjectInitializer& Obj
 	: Super(ObjectInitializer)
 {
 	FOVAngle = 90.0f;
+	OrthoWidth = 512;
 	bAutoActivate = true;
 	PrimaryComponentTick.bCanEverTick = true;
 	PrimaryComponentTick.TickGroup = TG_DuringPhysics;
@@ -315,13 +319,13 @@ void USceneCaptureComponent2D::OnRegister()
 #if WITH_EDITOR
 	// Update content on register to have at least one frames worth of good data.
 	// Without updating here this component would not work in a blueprint construction script which recreates the component after each move in the editor
-	UpdateContent();
+	CaptureSceneDeferred();
 #endif
 }
 
 void USceneCaptureComponent2D::SendRenderTransform_Concurrent()
 {	
-	UpdateContent();
+	CaptureSceneDeferred();
 
 	Super::SendRenderTransform_Concurrent();
 }
@@ -338,7 +342,7 @@ void USceneCaptureComponent2D::TickComponent(float DeltaTime, enum ELevelTick Ti
 
 static TMultiMap<TWeakObjectPtr<UWorld>, TWeakObjectPtr<USceneCaptureComponent2D>> SceneCapturesToUpdateMap;
 
-void USceneCaptureComponent2D::UpdateContent()
+void USceneCaptureComponent2D::CaptureSceneDeferred()
 {
 	UWorld* World = GetWorld();
 	if (World && World->Scene && IsVisible())
@@ -351,11 +355,28 @@ void USceneCaptureComponent2D::UpdateContent()
 	}	
 }
 
+void USceneCaptureComponent2D::CaptureScene()
+{
+	UWorld* World = GetWorld();
+	if (World && World->Scene && IsVisible())
+	{
+		// We must push any deferred render state recreations before causing any rendering to happen, to make sure that deleted resource references are updated
+		World->SendAllEndOfFrameUpdates();
+		World->Scene->UpdateSceneCaptureContents(this);
+	}	
+
+	if (bCaptureEveryFrame)
+	{
+		FMessageLog("Blueprint").Warning(LOCTEXT("CaptureScene", "CaptureScene: Scene capture with bCaptureEveryFrame enabled was told to update - major inefficiency."));
+	}
+}
+
 void USceneCaptureComponent2D::UpdateDeferredCaptures( FSceneInterface* Scene )
 {
 	UWorld* World = Scene->GetWorld();
 	if( World && SceneCapturesToUpdateMap.Num() > 0 )
 	{
+		// We must push any deferred render state recreations before causing any rendering to happen, to make sure that deleted resource references are updated
 		World->SendAllEndOfFrameUpdates();
 		// Only update the scene captures assoicated with the current scene.
 		// Updating others not associated with the scene would cause invalid data to be rendered into the target
@@ -376,12 +397,32 @@ void USceneCaptureComponent2D::UpdateDeferredCaptures( FSceneInterface* Scene )
 }
 
 #if WITH_EDITOR
+
+bool USceneCaptureComponent2D::CanEditChange(const UProperty* InProperty) const
+{
+	if (InProperty)
+	{
+		FString PropertyName = InProperty->GetName();
+
+		if (PropertyName == GET_MEMBER_NAME_STRING_CHECKED(USceneCaptureComponent2D, FOVAngle))
+		{
+			return ProjectionType == ECameraProjectionMode::Perspective;
+		}
+		else if (PropertyName == GET_MEMBER_NAME_STRING_CHECKED(USceneCaptureComponent2D, OrthoWidth))
+		{
+			return ProjectionType == ECameraProjectionMode::Orthographic;
+		}
+	}
+
+	return true;
+}
+
 void USceneCaptureComponent2D::PostEditChangeProperty(FPropertyChangedEvent& PropertyChangedEvent)
 {
 	// AActor::PostEditChange will ForceUpdateComponents()
 	Super::PostEditChangeProperty(PropertyChangedEvent);
 
-	UpdateContent();
+	CaptureSceneDeferred();
 }
 #endif // WITH_EDITOR
 
@@ -523,9 +564,6 @@ UPlanarReflectionComponent::UPlanarReflectionComponent(const FObjectInitializer&
 
 	ShowFlags.SetLightShafts(0);
 
-	// This is disabled because the math needs to be updated to compute fog from the reflection plane, not the reflected camera position (which is underwater)
-	ShowFlags.SetFog(0);
-
 	NextPlanarReflectionId++;
 	PlanarReflectionId = NextPlanarReflectionId;
 }
@@ -640,13 +678,13 @@ void USceneCaptureComponentCube::OnRegister()
 #if WITH_EDITOR
 	// Update content on register to have at least one frames worth of good data.
 	// Without updating here this component would not work in a blueprint construction script which recreates the component after each move in the editor
-	UpdateContent();
+	CaptureSceneDeferred();
 #endif
 }
 
 void USceneCaptureComponentCube::SendRenderTransform_Concurrent()
 {	
-	UpdateContent();
+	CaptureSceneDeferred();
 
 	Super::SendRenderTransform_Concurrent();
 }
@@ -663,7 +701,7 @@ void USceneCaptureComponentCube::TickComponent(float DeltaTime, enum ELevelTick
 
 static TMultiMap<TWeakObjectPtr<UWorld>, TWeakObjectPtr<USceneCaptureComponentCube> > CubedSceneCapturesToUpdateMap;
 
-void USceneCaptureComponentCube::UpdateContent()
+void USceneCaptureComponentCube::CaptureSceneDeferred()
 {
 	UWorld* World = GetWorld();
 	if (World && World->Scene && IsVisible())
@@ -676,12 +714,29 @@ void USceneCaptureComponentCube::UpdateContent()
 	}	
 }
 
+void USceneCaptureComponentCube::CaptureScene()
+{
+	UWorld* World = GetWorld();
+	if (World && World->Scene && IsVisible())
+	{
+		// We must push any deferred render state recreations before causing any rendering to happen, to make sure that deleted resource references are updated
+		World->SendAllEndOfFrameUpdates();
+		World->Scene->UpdateSceneCaptureContents(this);
+	}	
+
+	if (bCaptureEveryFrame)
+	{
+		FMessageLog("Blueprint").Warning(LOCTEXT("CaptureScene", "CaptureScene: Scene capture with bCaptureEveryFrame enabled was told to update - major inefficiency."));
+	}
+}
+
 void USceneCaptureComponentCube::UpdateDeferredCaptures( FSceneInterface* Scene )
 {
 	UWorld* World = Scene->GetWorld();
 	
 	if( World && CubedSceneCapturesToUpdateMap.Num() > 0 )
 	{
+		// We must push any deferred render state recreations before causing any rendering to happen, to make sure that deleted resource references are updated
 		World->SendAllEndOfFrameUpdates();
 		// Only update the scene captures associated with the current scene.
 		// Updating others not associated with the scene would cause invalid data to be rendered into the target
@@ -708,7 +763,7 @@ void USceneCaptureComponentCube::PostEditChangeProperty(FPropertyChangedEvent& P
 	// AActor::PostEditChange will ForceUpdateComponents()
 	Super::PostEditChangeProperty(PropertyChangedEvent);
 
-	UpdateContent();
+	CaptureSceneDeferred();
 }
 #endif // WITH_EDITOR
 
@@ -722,3 +777,5 @@ UDrawFrustumComponent* ASceneCapture2D::GetDrawFrustum() const { return DrawFrus
 USceneCaptureComponentCube* ASceneCaptureCube::GetCaptureComponentCube() const { return CaptureComponentCube; }
 /** Returns DrawFrustum subobject **/
 UDrawFrustumComponent* ASceneCaptureCube::GetDrawFrustum() const { return DrawFrustum; }
+
+#undef LOCTEXT_NAMESPACE
\ No newline at end of file
diff --git a/Engine/Source/Runtime/Engine/Private/Components/SkyLightComponent.cpp b/Engine/Source/Runtime/Engine/Private/Components/SkyLightComponent.cpp
--- a/Engine/Source/Runtime/Engine/Private/Components/SkyLightComponent.cpp
+++ b/Engine/Source/Runtime/Engine/Private/Components/SkyLightComponent.cpp
@@ -176,7 +176,8 @@ void USkyLightComponent::SetCaptureIsDirty()
 void USkyLightComponent::SanatizeCubemapSize()
 {
 	static const int32 MaxCubemapResolution = 1024;
-	CubemapResolution = FMath::Min(int32(FMath::RoundUpToPowerOfTwo(CubemapResolution)), MaxCubemapResolution);
+	static const int32 MinCubemapResolution = 64;
+	CubemapResolution = FMath::Clamp(int32(FMath::RoundUpToPowerOfTwo(CubemapResolution)), MinCubemapResolution, MaxCubemapResolution);
 }
 
 void USkyLightComponent::SetBlendDestinationCaptureIsDirty()
diff --git a/Engine/Source/Runtime/Engine/Private/DebugViewModeMaterialProxy.cpp b/Engine/Source/Runtime/Engine/Private/DebugViewModeMaterialProxy.cpp
--- a/Engine/Source/Runtime/Engine/Private/DebugViewModeMaterialProxy.cpp
+++ b/Engine/Source/Runtime/Engine/Private/DebugViewModeMaterialProxy.cpp
@@ -184,4 +184,12 @@ float FDebugViewModeMaterialProxy::GetOpacityMaskClipValue() const
 	return Material ? Material->GetOpacityMaskClipValue() : .5f;
 }
 
+void FDebugViewModeMaterialProxy::GatherCustomOutputExpressions(TArray<class UMaterialExpressionCustomOutput*>& OutCustomOutputs) const
+{
+	if (Material)
+	{
+		Material->GetAllCustomOutputExpressions(OutCustomOutputs);
+	}
+}
+
 #endif
\ No newline at end of file
diff --git a/Engine/Source/Runtime/Engine/Private/DebugViewModeMaterialProxy.h b/Engine/Source/Runtime/Engine/Private/DebugViewModeMaterialProxy.h
--- a/Engine/Source/Runtime/Engine/Private/DebugViewModeMaterialProxy.h
+++ b/Engine/Source/Runtime/Engine/Private/DebugViewModeMaterialProxy.h
@@ -90,6 +90,7 @@ public:
 	virtual enum EBlendMode GetBlendMode() const override;
 	virtual enum EMaterialShadingModel GetShadingModel() const override;
 	virtual float GetOpacityMaskClipValue() const override;
+	virtual void GatherCustomOutputExpressions(TArray<class UMaterialExpressionCustomOutput*>& OutCustomOutputs) const override;
 
 	virtual EMaterialShaderMapUsage::Type GetMaterialShaderMapUsage() const { return Usage; }
 
diff --git a/Engine/Source/Runtime/Engine/Private/HierarchicalInstancedStaticMesh.cpp b/Engine/Source/Runtime/Engine/Private/HierarchicalInstancedStaticMesh.cpp
--- a/Engine/Source/Runtime/Engine/Private/HierarchicalInstancedStaticMesh.cpp
+++ b/Engine/Source/Runtime/Engine/Private/HierarchicalInstancedStaticMesh.cpp
@@ -1481,7 +1481,11 @@ void FHierarchicalStaticMeshSceneProxy::GetDynamicMeshElements(const TArray<cons
 					const FFoliageOcclusionResults* OldResults = OcclusionResults.Find(ViewId);
 					if (OldResults && 
 						OldResults->FrameNumberRenderThread == GFrameNumberRenderThread &&
-						1 + LastOcclusionNode - FirstOcclusionNode == OldResults->NumResults
+						1 + LastOcclusionNode - FirstOcclusionNode == OldResults->NumResults &&
+						// OcclusionResultsArray[Params.OcclusionResultsStart + Index - Params.FirstOcclusionNode]
+
+						OldResults->Results->IsValidIndex(OldResults->ResultsStart) &&
+						OldResults->Results->IsValidIndex(OldResults->ResultsStart + LastOcclusionNode - FirstOcclusionNode)
 						)
 					{
 						InstanceParams.FirstOcclusionNode = FirstOcclusionNode;
diff --git a/Engine/Source/Runtime/Engine/Private/KismetRenderingLibrary.cpp b/Engine/Source/Runtime/Engine/Private/KismetRenderingLibrary.cpp
--- a/Engine/Source/Runtime/Engine/Private/KismetRenderingLibrary.cpp
+++ b/Engine/Source/Runtime/Engine/Private/KismetRenderingLibrary.cpp
@@ -16,8 +16,6 @@ UKismetRenderingLibrary::UKismetRenderingLibrary(const FObjectInitializer& Objec
 {
 }
 
-UCanvas* GCanvasForRenderingToTarget = NULL;
-
 void UKismetRenderingLibrary::DrawMaterialToRenderTarget(UObject* WorldContextObject, UTextureRenderTarget2D* TextureRenderTarget, UMaterialInterface* Material)
 {
 	check(WorldContextObject);
@@ -28,13 +26,7 @@ void UKismetRenderingLibrary::DrawMaterialToRenderTarget(UObject* WorldContextOb
 		&& World 
 		&& Material)
 	{
-		if (!GCanvasForRenderingToTarget)
-		{
-			GCanvasForRenderingToTarget = NewObject<UCanvas>(GetTransientPackage(), NAME_None);
-			GCanvasForRenderingToTarget->AddToRoot();
-		}
-
-		UCanvas* Canvas = GCanvasForRenderingToTarget;
+		UCanvas* Canvas = World->GetCanvasForRenderingToTarget();
 
 		Canvas->Init(TextureRenderTarget->SizeX, TextureRenderTarget->SizeY, nullptr);
 		Canvas->Update();
@@ -42,9 +34,7 @@ void UKismetRenderingLibrary::DrawMaterialToRenderTarget(UObject* WorldContextOb
 		FCanvas RenderCanvas(
 			TextureRenderTarget->GameThread_GetRenderTargetResource(), 
 			nullptr, 
-			FApp::GetCurrentTime() - GStartTime, 
-			FApp::GetDeltaTime(), 
-			FApp::GetCurrentTime() - GStartTime, 
+			World,
 			World->FeatureLevel);
 
 		Canvas->Canvas = &RenderCanvas;
@@ -103,13 +93,8 @@ void UKismetRenderingLibrary::BeginDrawCanvasToRenderTarget(UObject* WorldContex
 	{
 		Context.RenderTarget = TextureRenderTarget;
 
-		if (!GCanvasForRenderingToTarget)
-		{
-			GCanvasForRenderingToTarget = NewObject<UCanvas>(GetTransientPackage(), NAME_None);
-			GCanvasForRenderingToTarget->AddToRoot();
-		}
+		Canvas = World->GetCanvasForRenderingToTarget();
 
-		Canvas = GCanvasForRenderingToTarget;
 		Size = FVector2D(TextureRenderTarget->SizeX, TextureRenderTarget->SizeY);
 
 		Canvas->Init(TextureRenderTarget->SizeX, TextureRenderTarget->SizeY, nullptr);
@@ -118,9 +103,7 @@ void UKismetRenderingLibrary::BeginDrawCanvasToRenderTarget(UObject* WorldContex
 		Canvas->Canvas = new FCanvas(
 			TextureRenderTarget->GameThread_GetRenderTargetResource(), 
 			nullptr, 
-			FApp::GetCurrentTime() - GStartTime, 
-			FApp::GetDeltaTime(), 
-			FApp::GetCurrentTime() - GStartTime, 
+			World,
 			World->FeatureLevel, 
 			// Draw immediately so that interleaved SetVectorParameter (etc) function calls work as expected
 			FCanvas::CDM_ImmediateDrawing);
@@ -152,11 +135,13 @@ void UKismetRenderingLibrary::EndDrawCanvasToRenderTarget(UObject* WorldContextO
 
 	if (World)
 	{
-		if (GCanvasForRenderingToTarget && GCanvasForRenderingToTarget->Canvas)
+		UCanvas* WorldCanvas = World->GetCanvasForRenderingToTarget();
+
+		if (WorldCanvas->Canvas)
 		{
-			GCanvasForRenderingToTarget->Canvas->Flush_GameThread();
-			delete GCanvasForRenderingToTarget->Canvas;
-			GCanvasForRenderingToTarget->Canvas = nullptr;
+			WorldCanvas->Canvas->Flush_GameThread();
+			delete WorldCanvas->Canvas;
+			WorldCanvas->Canvas = nullptr;
 		}
 		
 		if (Context.RenderTarget)
diff --git a/Engine/Source/Runtime/Engine/Private/Materials/HLSLMaterialTranslator.h b/Engine/Source/Runtime/Engine/Private/Materials/HLSLMaterialTranslator.h
--- a/Engine/Source/Runtime/Engine/Private/Materials/HLSLMaterialTranslator.h
+++ b/Engine/Source/Runtime/Engine/Private/Materials/HLSLMaterialTranslator.h
@@ -476,6 +476,16 @@ public:
 				Errorf(TEXT("Only unlit materials can output negative emissive color."));
 			}
 
+			static IConsoleVariable* CDBufferVar = IConsoleManager::Get().FindConsoleVariable(TEXT("r.DBuffer"));
+			bool bDBufferAllowed = CDBufferVar ? CDBufferVar->GetInt() != 0 : false;
+			bool bDBufferBlendMode = IsDBufferDecalBlendMode((EDecalBlendMode)Material->GetDecalBlendMode());
+
+			if (bDBufferBlendMode && !bDBufferAllowed)
+			{
+				// Error feedback for when the decal would not be displayed due to project settings
+				Errorf(TEXT("DBuffer decal blend modes are only supported when the 'DBuffer Decals' Rendering Project setting is enabled."));
+			}
+
 			if (MaterialCompilationOutput.bNeedsSceneTextures)
 			{
 				if (Domain != MD_DeferredDecal && Domain != MD_PostProcess)
diff --git a/Engine/Source/Runtime/Engine/Private/Materials/MaterialExpressions.cpp b/Engine/Source/Runtime/Engine/Private/Materials/MaterialExpressions.cpp
--- a/Engine/Source/Runtime/Engine/Private/Materials/MaterialExpressions.cpp
+++ b/Engine/Source/Runtime/Engine/Private/Materials/MaterialExpressions.cpp
@@ -6135,7 +6135,7 @@ void UMaterialExpressionWorldPosition::GetCaption(TArray<FString>& OutCaptions)
 	{
 	case WPT_Default:
 		{
-			OutCaptions.Add(NSLOCTEXT("MaterialExpressions", "WorldPositonText", "Absolute World Position (Including Material Offsets)").ToString());
+			OutCaptions.Add(NSLOCTEXT("MaterialExpressions", "WorldPositonText", "Absolute World Position").ToString());
 			break;
 		}
 
@@ -6147,7 +6147,7 @@ void UMaterialExpressionWorldPosition::GetCaption(TArray<FString>& OutCaptions)
 
 	case WPT_CameraRelative:
 		{
-			OutCaptions.Add(NSLOCTEXT("MaterialExpressions", "CamRelativeWorldPositonText", "Camera Relative World Position (Including Material Offsets)").ToString());
+			OutCaptions.Add(NSLOCTEXT("MaterialExpressions", "CamRelativeWorldPositonText", "Camera Relative World Position").ToString());
 			break;
 		}
 
diff --git a/Engine/Source/Runtime/Engine/Private/Particles/ParticleGpuSimulation.cpp b/Engine/Source/Runtime/Engine/Private/Particles/ParticleGpuSimulation.cpp
--- a/Engine/Source/Runtime/Engine/Private/Particles/ParticleGpuSimulation.cpp
+++ b/Engine/Source/Runtime/Engine/Private/Particles/ParticleGpuSimulation.cpp
@@ -83,6 +83,8 @@ static TAutoConsoleVariable<float> CVarGPUParticleFixDeltaSeconds(TEXT("r.GPUPar
 static TAutoConsoleVariable<float> CVarGPUParticleFixTolerance(TEXT("r.GPUParticle.FixTolerance"),.1f,TEXT("Delta second tolerance before switching to a fix delta seconds."));
 static TAutoConsoleVariable<int32> CVarGPUParticleMaxNumIterations(TEXT("r.GPUParticle.MaxNumIterations"),3,TEXT("Max number of iteration when using a fix delta seconds."));
 
+static TAutoConsoleVariable<int32> CVarSimulateGPUParticles(TEXT("r.GPUParticle.Simulate"), 1, TEXT("Enable or disable GPU particle simulation"));
+
 /*-----------------------------------------------------------------------------
 	Allocators used to manage GPU particle resources.
 -----------------------------------------------------------------------------*/
@@ -1382,6 +1384,11 @@ void ExecuteSimulationCommands(
 	FTexture2DRHIParamRef GBufferATexture,
 	bool bUseFixDT)
 {
+	if (!CVarSimulateGPUParticles.GetValueOnAnyThread())
+	{
+		return;
+	}
+
 	SCOPED_DRAW_EVENT(RHICmdList, ParticleSimulation);
 
 	const float FixDeltaSeconds = CVarGPUParticleFixDeltaSeconds.GetValueOnRenderThread();
@@ -1487,6 +1494,11 @@ void ExecuteSimulationCommands(
  */
 void ClearTiles(FRHICommandList& RHICmdList, ERHIFeatureLevel::Type FeatureLevel, const TArray<uint32>& Tiles)
 {
+	if (!CVarSimulateGPUParticles.GetValueOnAnyThread())
+	{
+		return;
+	}
+
 	SCOPED_DRAW_EVENT(RHICmdList, ClearTiles);
 
 	const int32 MaxTilesPerDrawCallUnaligned = GParticleScratchVertexBufferSize / sizeof(FVector2D);
@@ -1739,7 +1751,7 @@ TGlobalResource<FParticleInjectionVertexDeclaration> GParticleInjectionVertexDec
  */
 void InjectNewParticles(FRHICommandList& RHICmdList, ERHIFeatureLevel::Type FeatureLevel, const TArray<FNewParticle>& NewParticles)
 {
-	if (GIsRenderingThreadSuspended)
+	if (GIsRenderingThreadSuspended || !CVarSimulateGPUParticles.GetValueOnAnyThread())
 	{
 		return;
 	}
diff --git a/Engine/Source/Runtime/Engine/Private/Rendering/SubDSurface.cpp b/Engine/Source/Runtime/Engine/Private/Rendering/SubDSurface.cpp
new file mode 100644
--- /dev/null
+++ b/Engine/Source/Runtime/Engine/Private/Rendering/SubDSurface.cpp
@@ -0,0 +1,105 @@
+// Copyright 1998-2016 Epic Games, Inc. All Rights Reserved.
+
+#include "EnginePrivate.h"
+#include "SeparableSSS.h"
+#include "EnginePrivate.h"
+#include "RendererInterface.h"
+#include "Engine/SubDSurface.h"
+
+DEFINE_LOG_CATEGORY_STATIC(LogSubDSurface, Log, All);
+
+// ------------------------------------------------------
+
+UVertexAttributeStream::UVertexAttributeStream(const FObjectInitializer& ObjectInitializer)
+	: Super(ObjectInitializer)
+	, AttributeType(VAST_unknown)
+{
+}
+
+FVector* UVertexAttributeStream::CreateFVectorUninitialized(uint32 Count)
+{
+	AttributeType = EVertexAttributeStreamType::VAST_float3;
+
+	// can be optimized
+	Data.Empty();
+	Data.AddUninitialized(sizeof(FVector) * Count);
+
+	return (FVector*)Data.GetData();
+}
+
+FVector2D* UVertexAttributeStream::CreateFVector2DUninitialized(uint32 Count)
+{
+	AttributeType = EVertexAttributeStreamType::VAST_float2;
+
+	// can be optimized
+	Data.Empty();
+	Data.AddUninitialized(sizeof(FVector2D) * Count);
+
+	return (FVector2D*)Data.GetData();
+}
+
+FVector* UVertexAttributeStream::GetFVector(uint32& OutCount)
+{
+	if(AttributeType == EVertexAttributeStreamType::VAST_float3)
+	{
+		OutCount = Data.Num() / sizeof(FVector);
+		return (FVector*)Data.GetData();
+	}
+
+	OutCount = 0;
+	return 0;
+}
+
+FVector2D* UVertexAttributeStream::GetFVector2D(uint32& OutCount)
+{
+	if(AttributeType == EVertexAttributeStreamType::VAST_float2)
+	{
+		OutCount = Data.Num() / sizeof(FVector2D);
+		return (FVector2D*)Data.GetData();
+	}
+
+	OutCount = 0;
+	return 0;
+}
+
+
+uint32 UVertexAttributeStream::GetTypeSize() const
+{
+	switch(AttributeType)
+	{
+		case EVertexAttributeStreamType::VAST_float:
+			return sizeof(float);
+
+		case EVertexAttributeStreamType::VAST_float2:
+			return sizeof(FVector2D);
+
+		case EVertexAttributeStreamType::VAST_float3:
+			return sizeof(FVector);
+
+		case EVertexAttributeStreamType::VAST_float4:
+			return sizeof(FVector4);
+	}
+
+	check(0);
+	return 0;
+}
+
+uint32 UVertexAttributeStream::Num() const
+{
+	if(EVertexAttributeStreamType::VAST_unknown)
+	{
+		return 0;
+	}
+
+	return Data.Num() / GetTypeSize();
+}
+
+// ------------------------------------------------------
+
+USubDSurface::USubDSurface(const FObjectInitializer& ObjectInitializer)
+	: Super(ObjectInitializer)
+{
+}
+
+
+
diff --git a/Engine/Source/Runtime/Engine/Private/RenderingSettings.cpp b/Engine/Source/Runtime/Engine/Private/RenderingSettings.cpp
--- a/Engine/Source/Runtime/Engine/Private/RenderingSettings.cpp
+++ b/Engine/Source/Runtime/Engine/Private/RenderingSettings.cpp
@@ -79,5 +79,6 @@ void URendererSettings::PostEditChangeProperty(FPropertyChangedEvent& PropertyCh
 void URendererSettings::SanatizeReflectionCaptureResolution()
 {
 	static const int32 MaxReflectionCaptureResolution = 1024;
-	ReflectionCaptureResolution = FMath::Min(int32(FMath::RoundUpToPowerOfTwo(ReflectionCaptureResolution)), MaxReflectionCaptureResolution);
+	static const int32 MinReflectionCaptureResolution = 64;
+	ReflectionCaptureResolution = FMath::Clamp(int32(FMath::RoundUpToPowerOfTwo(ReflectionCaptureResolution)), MinReflectionCaptureResolution, MaxReflectionCaptureResolution);
 }
diff --git a/Engine/Source/Runtime/Engine/Private/SceneView.cpp b/Engine/Source/Runtime/Engine/Private/SceneView.cpp
--- a/Engine/Source/Runtime/Engine/Private/SceneView.cpp
+++ b/Engine/Source/Runtime/Engine/Private/SceneView.cpp
@@ -337,9 +337,9 @@ TLinkedList<FSceneViewStateReference*>*& FSceneViewStateReference::GetSceneViewS
  * @param InvertZ - projection calc is affected by inverted device Z
  * @return vector containing the ratios needed to convert from device Z to world Z
  */
-FVector4 CreateInvDeviceZToWorldZTransform(FMatrix const & ProjMatrix)
+FVector4 CreateInvDeviceZToWorldZTransform(const FMatrix& ProjMatrix)
 {
-	// The depth projection comes from the the following projection matrix:
+	// The perspective depth projection comes from the the following projection matrix:
 	//
 	// | 1  0  0  0 |
 	// | 0  1  0  0 |
@@ -367,18 +367,53 @@ FVector4 CreateInvDeviceZToWorldZTransform(FMatrix const & ProjMatrix)
 		DepthAdd = 0.00000001f;
 	}
 
-	float SubtractValue = DepthMul / DepthAdd;
+	// perspective
+	// SceneDepth = 1.0f / (DeviceZ / ProjMatrix.M[3][2] - ProjMatrix.M[2][2] / ProjMatrix.M[3][2])
 
-	// Subtract a tiny number to avoid divide by 0 errors in the shader when a very far distance is decided from the depth buffer.
-	// This fixes fog not being applied to the black background in the editor.
-	SubtractValue -= 0.00000001f;
+	// ortho
+	// SceneDepth = DeviceZ / ProjMatrix.M[2][2] - ProjMatrix.M[3][2] / ProjMatrix.M[2][2];
 
-	return FVector4(
-		0.0f,			// Unused
-		0.0f,			// Unused
-		1.f / DepthAdd,	
-		SubtractValue
-		);
+	// combined equation in shader to handle either
+	// SceneDepth = DeviceZ * View.InvDeviceZToWorldZTransform[0] + View.InvDeviceZToWorldZTransform[1] + 1.0f / (DeviceZ * View.InvDeviceZToWorldZTransform[2] - View.InvDeviceZToWorldZTransform[3]);
+
+	// therefore perspective needs
+	// View.InvDeviceZToWorldZTransform[0] = 0.0f
+	// View.InvDeviceZToWorldZTransform[1] = 0.0f
+	// View.InvDeviceZToWorldZTransform[2] = 1.0f / ProjMatrix.M[3][2]
+	// View.InvDeviceZToWorldZTransform[3] = ProjMatrix.M[2][2] / ProjMatrix.M[3][2]
+
+	// and ortho needs
+	// View.InvDeviceZToWorldZTransform[0] = 1.0f / ProjMatrix.M[2][2]
+	// View.InvDeviceZToWorldZTransform[1] = -ProjMatrix.M[3][2] / ProjMatrix.M[2][2] + 1.0f
+	// View.InvDeviceZToWorldZTransform[2] = 0.0f
+	// View.InvDeviceZToWorldZTransform[3] = 1.0f
+
+	bool bIsPerspectiveProjection = ProjMatrix.M[3][3] < 1.0f;
+
+	if (bIsPerspectiveProjection)
+	{
+		float SubtractValue = DepthMul / DepthAdd;
+
+		// Subtract a tiny number to avoid divide by 0 errors in the shader when a very far distance is decided from the depth buffer.
+		// This fixes fog not being applied to the black background in the editor.
+		SubtractValue -= 0.00000001f;
+
+		return FVector4(
+			0.0f,			
+			0.0f,			
+			1.0f / DepthAdd,	
+			SubtractValue
+			);
+	}
+	else
+	{
+		return FVector4(
+			1.0f / ProjMatrix.M[2][2],
+			-ProjMatrix.M[3][2] / ProjMatrix.M[2][2] + 1.0f,
+			0.0f,
+			1.0f
+			);
+	}
 }
 
 FSceneView::FSceneView(const FSceneViewInitOptions& InitOptions)
@@ -1862,6 +1897,7 @@ FSceneViewFamily::FSceneViewFamily( const ConstructionValues& CVS )
 	bRealtimeUpdate(CVS.bRealtimeUpdate),
 	bDeferClear(CVS.bDeferClear),
 	bResolveScene(CVS.bResolveScene),
+	SceneCaptureSource(SCS_FinalColorLDR),
 	bWorldIsPaused(false),
 	GammaCorrection(CVS.GammaCorrection)
 {
diff --git a/Engine/Source/Runtime/Engine/Private/ShowFlags.cpp b/Engine/Source/Runtime/Engine/Private/ShowFlags.cpp
--- a/Engine/Source/Runtime/Engine/Private/ShowFlags.cpp
+++ b/Engine/Source/Runtime/Engine/Private/ShowFlags.cpp
@@ -369,6 +369,14 @@ void EngineShowFlagOverride(EShowFlagInitMode ShowFlagInitMode, EViewModeIndex V
 		}
 	}
 
+	{
+		static const auto ICVar = IConsoleManager::Get().FindTConsoleVariableDataInt(TEXT("r.ShadowQuality"));
+		if(ICVar->GetValueOnGameThread() <= 0)
+		{
+			EngineShowFlags.DynamicShadows = 0;
+		}
+	}
+
 	// some view modes want some features off or on (no state)
 	{
 		if( ViewModeIndex == VMI_BrushWireframe ||
diff --git a/Engine/Source/Runtime/Engine/Private/SkeletalMesh.cpp b/Engine/Source/Runtime/Engine/Private/SkeletalMesh.cpp
--- a/Engine/Source/Runtime/Engine/Private/SkeletalMesh.cpp
+++ b/Engine/Source/Runtime/Engine/Private/SkeletalMesh.cpp
@@ -5062,7 +5062,7 @@ void FSkeletalMeshSceneProxy::GetShadowShapes(TArray<FCapsuleShape>& CapsuleShap
 
 		NewCapsule.Center = ReferenceToWorld.TransformPosition(CapsuleData.Value.Center);
 		NewCapsule.Radius = CapsuleData.Value.Radius * MaxScale;
-		NewCapsule.Orientation = ReferenceToWorld.TransformVector(CapsuleData.Value.Orientation);
+		NewCapsule.Orientation = ReferenceToWorld.TransformVector(CapsuleData.Value.Orientation).GetSafeNormal();
 		NewCapsule.Length = CapsuleData.Value.Length * MaxScale;
 	}
 }
diff --git a/Engine/Source/Runtime/Engine/Private/Streaming/AsyncTextureStreaming.cpp b/Engine/Source/Runtime/Engine/Private/Streaming/AsyncTextureStreaming.cpp
--- a/Engine/Source/Runtime/Engine/Private/Streaming/AsyncTextureStreaming.cpp
+++ b/Engine/Source/Runtime/Engine/Private/Streaming/AsyncTextureStreaming.cpp
@@ -7,9 +7,10 @@ AsyncTextureStreaming.cpp: Definitions of classes used for texture streaming asy
 #include "EnginePrivate.h"
 #include "AsyncTextureStreaming.h"
 
-void FAsyncTextureStreamingData::Init(TArray<FStreamingViewInfo> InViewInfos, TArray<FLevelTextureManager>& LevelTextureManagers, FDynamicComponentTextureManager& DynamicComponentManager)
+void FAsyncTextureStreamingData::Init(TArray<FStreamingViewInfo> InViewInfos, float InLastUpdateTime, TArray<FLevelTextureManager>& LevelTextureManagers, FDynamicComponentTextureManager& DynamicComponentManager)
 {
 	ViewInfos = InViewInfos;
+	LastUpdateTime = InLastUpdateTime;
 
 	DynamicInstancesView = DynamicComponentManager.GetAsyncView();
 
@@ -27,9 +28,9 @@ void FAsyncTextureStreamingData::Update_Async()
 
 	for (FTextureInstanceAsyncView& StaticInstancesView : StaticInstancesViews)
 	{
-		StaticInstancesView.Update_Async(ViewInfos, bUseNewMetrics, MaxEffectiveScreenSize);
+		StaticInstancesView.Update_Async(ViewInfos, LastUpdateTime, bUseNewMetrics, MaxEffectiveScreenSize);
 	}
-	DynamicInstancesView.Update_Async(ViewInfos, bUseNewMetrics, MaxEffectiveScreenSize);
+	DynamicInstancesView.Update_Async(ViewInfos, LastUpdateTime, bUseNewMetrics, MaxEffectiveScreenSize);
 }
 
 void FAsyncTextureStreamingTask::ClearRemovedTextureReferences()
@@ -122,7 +123,7 @@ void FAsyncTextureStreamingTask::DoWork()
 			}
 		}
 
-		if (StreamingTexture.IsVisible() && StreamingTexture.GetWantedMipsForPriority() - StreamingTexture.ResidentMips > 1)
+		if (StreamingTexture.IsVisible() && StreamingTexture.GetWantedMipsForPriority() > StreamingTexture.ResidentMips)
 		{
 			++ThreadStats.NumVisibleTexturesWithLowResolutions;
 		}
@@ -212,4 +213,4 @@ void FAsyncTextureStreamingTask::DoWork()
 		PrioTexIndicesSortedByLoadOrder[Index] = Index;
 	}
 	PrioTexIndicesSortedByLoadOrder.Sort(FCompareTexturePriorityByLoadOrder(PrioritizedTextures));
-}
\ No newline at end of file
+}
diff --git a/Engine/Source/Runtime/Engine/Private/Streaming/AsyncTextureStreaming.h b/Engine/Source/Runtime/Engine/Private/Streaming/AsyncTextureStreaming.h
--- a/Engine/Source/Runtime/Engine/Private/Streaming/AsyncTextureStreaming.h
+++ b/Engine/Source/Runtime/Engine/Private/Streaming/AsyncTextureStreaming.h
@@ -8,6 +8,27 @@ AsyncTextureStreaming.h: Definitions of classes used for texture streaming async
 
 #include "StreamingManagerTexture.h"
 
+struct FTexturePriority
+{
+	FTexturePriority( bool InCanDropMips, int32 InRetentionPriority, int32 InLoadOrderPriority, int32 InTextureIndex, const UTexture2D* InTexture)
+	:	bCanDropMips( InCanDropMips )
+	,	RetentionPriority( InRetentionPriority )
+	,	LoadOrderPriority( InLoadOrderPriority )
+	,	TextureIndex( InTextureIndex )
+	,	Texture( InTexture )
+	{
+	}
+	/** True if we allows this texture to drop mips to fit in budget. */
+	bool bCanDropMips;
+	/** Texture retention priority, higher value means it should be kept in memory. */
+	int32 RetentionPriority;
+	/** Texture load order priority, higher value means it should load first. */
+	int32 LoadOrderPriority;
+	/** Index into the FStreamingManagerTexture::StreamingTextures array. */
+	int32 TextureIndex;
+	/** The texture to stream. Only used for validation. */
+	const UTexture2D* Texture;
+};
 
 /** Thread-safe helper struct for streaming information. */
 class FAsyncTextureStreamingData
@@ -15,7 +36,7 @@ class FAsyncTextureStreamingData
 public:
 
 	/** Set the data but do as little as possible, called from the game thread. */
-	void Init(TArray<FStreamingViewInfo> InViewInfos, TArray<FLevelTextureManager>& LevelTextureManagers, FDynamicComponentTextureManager& DynamicComponentManager);
+	void Init(TArray<FStreamingViewInfo> InViewInfos, float InWorldTime, TArray<FLevelTextureManager>& LevelTextureManagers, FDynamicComponentTextureManager& DynamicComponentManager);
 
 	/** Update everything internally so to allow calls to CalcWantedMips */
 	void Update_Async();
@@ -35,6 +56,8 @@ private:
 
 	/** Cached from each ULevel. */
 	TArray<FTextureInstanceAsyncView> StaticInstancesViews;
+
+	float LastUpdateTime;
 };
 
 
@@ -49,17 +72,17 @@ class FAsyncTextureStreamingTask : public FNonAbandonableTask
 public:
 	FAsyncTextureStreamingTask( FStreamingManagerTexture* InStreamingManager )
 	:	StreamingManager( *InStreamingManager )
-	,	ThreadContext( false, NULL, false )
+	,	ThreadContext( false, NULL )
 	,	bAbort( false )
 	{
-		Reset(false);
+		Reset();
 	}
 
 	/** Resets the state to start a new async job. */
-	void Reset( bool bCollectTextureStats )
+	void Reset()
 	{
 		bAbort = false;
-		ThreadContext.Reset( false, NULL, bCollectTextureStats );
+		ThreadContext.Reset( false, NULL );
 		ThreadStats.Reset();
 		PrioritizedTextures.Empty( StreamingManager.StreamingTextures.Num() );
 		PrioTexIndicesSortedByLoadOrder.Empty( StreamingManager.StreamingTextures.Num() );
diff --git a/Engine/Source/Runtime/Engine/Private/Streaming/StreamingManagerTexture.cpp b/Engine/Source/Runtime/Engine/Private/Streaming/StreamingManagerTexture.cpp
--- a/Engine/Source/Runtime/Engine/Private/Streaming/StreamingManagerTexture.cpp
+++ b/Engine/Source/Runtime/Engine/Private/Streaming/StreamingManagerTexture.cpp
@@ -49,7 +49,6 @@ FStreamingManagerTexture::FStreamingManagerTexture()
 :	CurrentUpdateStreamingTextureIndex(0)
 ,	bTriggerDumpTextureGroupStats( false )
 ,	bDetailedDumpTextureGroupStats( false )
-,	bTriggerInvestigateTexture( false )
 ,	MipBias(0)
 ,	AsyncWork( nullptr )
 ,	ProcessingStage( 0 )
@@ -84,9 +83,8 @@ FStreamingManagerTexture::FStreamingManagerTexture()
 ,	OriginalTexturePoolSize(0)
 ,	PreviousPoolSizeTimestamp(0.0)
 ,	PreviousPoolSizeSetting(-1)
-,	bCollectTextureStats(false)
-,	bReportTextureStats(false)
 ,	bPauseTextureStreaming(false)
+,	LastUpdateTime(GIsEditor ? -FLT_MAX : 0) // In editor, visibility is not taken into consideration.
 {
 	// Read settings from ini file.
 	int32 TempInt;
@@ -420,7 +418,23 @@ void FStreamingManagerTexture::UpdateThreadData()
 	IncrementalUpdate(1.f);
 	MipBias = FMath::Max(CVarStreamingMipBias.GetValueOnGameThread(), 0.0f);
 	// Update the thread data.
-	AsyncWork->GetTask().StreamingData.Init(CurrentViewInfos, LevelTextureManagers, DynamicComponentManager);
+
+	AsyncWork->GetTask().StreamingData.Init(CurrentViewInfos, GStreamWithTimeFactor ? LastUpdateTime : -FLT_MAX, LevelTextureManagers, DynamicComponentManager);
+
+	if (!GIsEditor)
+	{
+		float WorldTime = 0;
+		for (int32 LevelIndex = 0; LevelIndex < LevelTextureManagers.Num() && !WorldTime; ++LevelIndex)
+		{
+			WorldTime = LevelTextureManagers[LevelIndex].GetWorldTime();
+		}
+
+		// Update last update time only if there is a reasonable threshold to define visibility.
+		if (WorldTime - LastUpdateTime  > .5f)
+		{
+			LastUpdateTime = WorldTime;
+		}
+	}
 }
 
 /**
@@ -889,7 +903,7 @@ void FStreamingManagerTexture::StreamTextures( bool bProcessEverything )
 	const int64 MaxTempMemoryAllowed = CVarStreamingMaxTempMemoryAllowed.GetValueOnAnyThread() * 1024 * 1024;
 
 	// Setup a context for this run.
-	FStreamingContext Context( bProcessEverything, IndividualStreamingTexture, bCollectTextureStats );
+	FStreamingContext Context( bProcessEverything, IndividualStreamingTexture );
 
 	AsyncWork->GetTask().ClearRemovedTextureReferences();
 
@@ -901,86 +915,6 @@ void FStreamingManagerTexture::StreamTextures( bool bProcessEverything )
 
 	FIOSystem::Get().FlushLog();
 
-#if STATS
-	// Did we collect texture stats? Triggered by the ListStreamingTextures exec command.
-	if ( Context.TextureStats.Num() > 0 )
-	{
-		// Reinitialize each time
-		TextureStatsReport.Empty();
-
-		// Sort textures by cost.
-		struct FCompareFTextureStreamingStats
-		{
-			FORCEINLINE bool operator()( const FTextureStreamingStats& A, const FTextureStreamingStats& B ) const
-			{
-				if ( A.Priority > B.Priority )
-				{
-					return true;
-				}
-				else if ( A.Priority == B.Priority )
-				{
-					return ( A.TextureIndex < B.TextureIndex );
-				}
-				else
-				{
-					return false;
-				}
-			}
-		};
-		Context.TextureStats.Sort( FCompareFTextureStreamingStats() );
-		int64 TotalCurrentSize	= 0;
-		int64 TotalWantedSize	= 0;
-		int64 TotalMaxSize		= 0;
-		TextureStatsReport.Add( FString( TEXT(",Priority,Current,Wanted,Max,Largest Resident,Current Size (KB),Wanted Size (KB),Max Size (KB),Largest Resident Size (KB),Streaming Type,Last Rendered,BoostFactor,Name") ) );
-		for( int32 TextureIndex=0; TextureIndex<Context.TextureStats.Num(); TextureIndex++ )
-		{
-			const FTextureStreamingStats& TextureStat = Context.TextureStats[TextureIndex];
-			int32 LODBias			= TextureStat.LODBias;
-			int32 CurrDroppedMips	= TextureStat.NumMips - TextureStat.ResidentMips;
-			int32 WantedDroppedMips	= TextureStat.NumMips - TextureStat.WantedMips;
-			int32 MostDroppedMips	= TextureStat.NumMips - TextureStat.MostResidentMips;
-			TextureStatsReport.Add( FString::Printf( TEXT(",%.3f,%ix%i,%ix%i,%ix%i,%ix%i,%i,%i,%i,%i,%3f sec,%.1f,%s"),
-				TextureStat.Priority,
-				FMath::Max(TextureStat.SizeX >> CurrDroppedMips, 1),
-				FMath::Max(TextureStat.SizeY >> CurrDroppedMips, 1),
-				FMath::Max(TextureStat.SizeX >> WantedDroppedMips, 1),
-				FMath::Max(TextureStat.SizeY >> WantedDroppedMips, 1),
-				FMath::Max(TextureStat.SizeX >> LODBias, 1),
-				FMath::Max(TextureStat.SizeY >> LODBias, 1),
-				FMath::Max(TextureStat.SizeX >> MostDroppedMips, 1),
-				FMath::Max(TextureStat.SizeY >> MostDroppedMips, 1),
-				TextureStat.ResidentSize/1024,
-				TextureStat.WantedSize/1024,
-				TextureStat.MaxSize/1024,
-				TextureStat.MostResidentSize/1024,
-				TextureStat.LastRenderTime,
-				TextureStat.BoostFactor,
-				*TextureStat.TextureName
-				) );
-			TotalCurrentSize	+= TextureStat.ResidentSize;
-			TotalWantedSize		+= TextureStat.WantedSize;
-			TotalMaxSize		+= TextureStat.MaxSize;
-		}
-		TextureStatsReport.Add( FString::Printf( TEXT("Total size: Current= %d  Wanted=%d  Max= %d"), TotalCurrentSize/1024, TotalWantedSize/1024, TotalMaxSize/1024 ) );
-		Context.TextureStats.Empty();
-
-		if( bReportTextureStats )
-		{
-			if( TextureStatsReport.Num() > 0 )
-			{
-				UE_LOG(LogContentStreaming, Log,  TEXT("Listing collected stats for all streaming textures") );
-				for( int32 ReportIndex = 0; ReportIndex < TextureStatsReport.Num(); ReportIndex++ )
-				{
-					UE_LOG(LogContentStreaming, Log, TEXT("%s"), *TextureStatsReport[ReportIndex]);
-				}
-				TextureStatsReport.Empty();
-			}
-			bReportTextureStats = false;
-		}
-		bCollectTextureStats = false;
-	}
-#endif
-
 	FTextureMemoryStats Stats;
 	RHIGetTextureMemoryStats(Stats);
 	ThreadStats.PendingStreamInSize += (Stats.PendingMemoryAdjustment > 0) ? Stats.PendingMemoryAdjustment : 0;
@@ -1137,11 +1071,10 @@ void FStreamingManagerTexture::DropWantedMips( FStreamingContext& Context )
 		for (int32 TexPrioIndex = 0; TexPrioIndex < PrioritizedTextures.Num() && Context.AvailableLater < 0; ++TexPrioIndex)
 		{
 			const FTexturePriority& TexturePriority = PrioritizedTextures[ TexPrioIndex ];
-			const int32 TextureIndex = TexturePriority.TextureIndex;
 
 			// This will also account for invalid entries.
 			if (TexturePriority.TextureIndex == INDEX_NONE || !TexturePriority.bCanDropMips) continue;
-			FStreamingTexture& StreamingTexture = StreamingTextures[TextureIndex];
+			FStreamingTexture& StreamingTexture = StreamingTextures[TexturePriority.TextureIndex];
 
 			// Allow to drop all mips until min mips
 			if (StreamingTexture.WantedMips <= StreamingTexture.MinAllowedMips) continue;
@@ -1158,7 +1091,7 @@ void FStreamingManagerTexture::DropWantedMips( FStreamingContext& Context )
 				bMoreMemoryCouldBeFreed = true;
 			}
 		}
-	bIsDroppingFirstMip = false;
+		bIsDroppingFirstMip = false;
 	}
 }
 
@@ -1175,11 +1108,10 @@ void FStreamingManagerTexture::DropForcedMips( FStreamingContext& Context )
 		for (int32 TexPrioIndex = 0; TexPrioIndex < PrioritizedTextures.Num() && Context.AvailableLater < 0; ++TexPrioIndex)
 		{
 			const FTexturePriority& TexturePriority = PrioritizedTextures[ TexPrioIndex ];
-			const int32 TextureIndex = TexturePriority.TextureIndex;
 
 			// This will also account for invalid entries.
 			if (TexturePriority.TextureIndex == INDEX_NONE || TexturePriority.bCanDropMips) continue;
-			FStreamingTexture& StreamingTexture = StreamingTextures[TextureIndex];
+			FStreamingTexture& StreamingTexture = StreamingTextures[TexturePriority.TextureIndex];
 
 			// Allow to drop all mips until non resident mips.
 			if (StreamingTexture.WantedMips <= StreamingTexture.NumNonStreamingMips) continue;
@@ -1335,10 +1267,6 @@ void FStreamingManagerTexture::UpdateResourceStreaming( float DeltaTime, bool bP
 		{
 			DumpTextureGroupStats( bDetailedDumpTextureGroupStats );
 		}
-		if ( bTriggerInvestigateTexture )
-		{
-			InvestigateTexture( InvestigateTextureName );
-		}
 	}
 	else
 	{
@@ -1350,7 +1278,7 @@ void FStreamingManagerTexture::UpdateResourceStreaming( float DeltaTime, bool bP
 	if ( ProcessingStage < NumDataCollectionStages )
 	{
 		// Setup a context for this run. Note that we only (potentially) collect texture stats from the AsyncWork.
-		FStreamingContext Context( bProcessEverything, IndividualStreamingTexture, false );
+		FStreamingContext Context( bProcessEverything, IndividualStreamingTexture );
 		{
 			QUICK_SCOPE_CYCLE_COUNTER(STAT_FStreamingManagerTexture_UpdateResourceStreaming_UpdateStreamingTextures);
 			UpdateStreamingTextures( Context, ProcessingStage, NumDataCollectionStages );
@@ -1371,7 +1299,7 @@ void FStreamingManagerTexture::UpdateResourceStreaming( float DeltaTime, bool bP
 			AsyncWork->EnsureCompletion();
 		}
 
-		AsyncWork->GetTask().Reset(bCollectTextureStats);
+		AsyncWork->GetTask().Reset();
 		if ( NumTextureProcessingStages > 1 )
 		{
 			QUICK_SCOPE_CYCLE_COUNTER(STAT_FStreamingManagerTexture_UpdateResourceStreaming_StartBackgroundTask);
@@ -1558,14 +1486,6 @@ void FStreamingManagerTexture::UpdateFrameStats( FStreamingContext& Context, FSt
 #if STATS
 	int64 PotentialSize = StreamingTexture.GetSize(StreamingTexture.MaxAllowedMips);
 
-	if ( Context.bCollectTextureStats )
-	{
-		FString TextureName = StreamingTexture.Texture->GetFullName();
-		if ( CollectTextureStatsName.Len() == 0 || TextureName.Contains( CollectTextureStatsName) )
-		{
-			new (Context.TextureStats) FTextureStreamingStats( StreamingTexture.Texture, StreamingTexture.ResidentMips, PerfectWantedMips, StreamingTexture.MostResidentMips, ResidentSize, PerfectWantedSize, PotentialSize, MostResidentSize, StreamingTexture.BoostFactor, StreamingTexture.CalcLoadOrderPriority(), TextureIndex );
-		}
-	}
 	Context.ThisFrameNumStreamingTextures++;
 	Context.ThisFrameTotalStreamingTexturesMaxSize += PotentialSize;
 	Context.ThisFrameTotalLightmapMemorySize += StreamingTexture.bIsLightmap ? ResidentSize : 0;
@@ -1575,39 +1495,30 @@ void FStreamingManagerTexture::UpdateFrameStats( FStreamingContext& Context, FSt
 #endif
 }
 
-void FStreamingManagerTexture::LogHeuristic(bool bOutputToLog, bool bUseHeuristic, const TCHAR*& Heuristic, const TCHAR* NewHeuristic)
-{
-#if !UE_BUILD_SHIPPING
-	if (bUseHeuristic)
-	{
-		if (bOutputToLog)
-		{
-			UE_LOG(LogContentStreaming, Log,  TEXT("  %s Reference."), NewHeuristic);
-		}
-		Heuristic = NewHeuristic;
-	}
-#endif
-}
-
 /**
  * Calculates the number of mip-levels we would like to have in memory for a texture.
  */
-const TCHAR* FStreamingManagerTexture::CalcWantedMips(const FAsyncTextureStreamingData& StreamingData, FStreamingTexture& StreamingTexture, bool bGetHeuristic, bool bOutputToLog)
+void FStreamingManagerTexture::CalcWantedMips(const FAsyncTextureStreamingData& StreamingData, FStreamingTexture& StreamingTexture, bool bGetHeuristic, bool bOutputToLog)
 {
-	const TCHAR* Heuristic = nullptr; 
+#if UE_BUILD_SHIPPING
+	bOutputToLog = false;
+#endif
+
+	// Whether streaming scale should be ignored (used when selecting specific resolutions).
+	bool bIgnoreStreamingScale = false;
 
 	float MaxSize = 0;
 	float MaxSize_VisibleOnly = 0;
 
 	if ( StreamingTexture.bForceFullyLoad || StreamingTexture.MinAllowedMips == StreamingTexture.MaxAllowedMips )
 	{
-		LogHeuristic(bOutputToLog, bGetHeuristic, Heuristic, TEXT("Forced"));
+		if (bOutputToLog) UE_LOG(LogContentStreaming, Log,  TEXT("  Forced FullyLoad"));
+
 		MaxSize_VisibleOnly = FLT_MAX;
 	}
 	else
 	{
-		LogHeuristic(bOutputToLog, bGetHeuristic && StreamingData.GetDynamicInstancesView().HasTextureReferences(StreamingTexture.Texture), Heuristic, TEXT("Dynamic"));
-		StreamingData.GetDynamicInstancesView().GetTexelSize(StreamingTexture.Texture, MaxSize, MaxSize_VisibleOnly, bOutputToLog);
+		StreamingData.GetDynamicInstancesView().GetTexelSize(StreamingTexture.Texture, MaxSize, MaxSize_VisibleOnly, bOutputToLog ? TEXT("Dynamic") : nullptr);
 
 		for (const FTextureInstanceAsyncView& StaticInstancesView : StreamingData.GetStaticInstancesViews())
 		{
@@ -1617,8 +1528,7 @@ const TCHAR* FStreamingManagerTexture::CalcWantedMips(const FAsyncTextureStreami
 				break;
 			}
 
-			LogHeuristic(bOutputToLog, bGetHeuristic && StaticInstancesView.HasTextureReferences(StreamingTexture.Texture), Heuristic, TEXT("Static"));
-			StaticInstancesView.GetTexelSize(StreamingTexture.Texture, MaxSize, MaxSize_VisibleOnly, bOutputToLog);
+			StaticInstancesView.GetTexelSize(StreamingTexture.Texture, MaxSize, MaxSize_VisibleOnly, bOutputToLog ? TEXT("Static") : nullptr);
 		}
 
 		/**
@@ -1632,31 +1542,38 @@ const TCHAR* FStreamingManagerTexture::CalcWantedMips(const FAsyncTextureStreami
 		const float TimeSinceRemoved = (float)(FApp::GetCurrentTime() - StreamingTexture.InstanceRemovedTimestamp);
 		if (MaxSize == 0 && MaxSize_VisibleOnly == 0 && TimeSinceRemoved < 91.f && StreamingTexture.LastRenderTime > TimeSinceRemoved - 5.f)
 		{
-			LogHeuristic(bOutputToLog, true, Heuristic, TEXT("Orphaned"));
-			MaxSize = (float)(0x1 << (FMath::Min<int32>(StreamingTexture.MaxAllowedMips - 1, StreamingTexture.ResidentMips) - 1));
-			if (StreamingTexture.LastRenderTime < 5)
+			if (bOutputToLog) UE_LOG(LogContentStreaming, Log,  TEXT("  Orphaned"));
+
+			// Keep the current mip, unless the space is required.
+			MaxSize = (float)(0x1 << (StreamingTexture.ResidentMips - 1)); // Keep all mips.
+
+			if( StreamingTexture.LastRenderTime < 5.0f)
 			{
 				MaxSize_VisibleOnly = MaxSize;
 			}
+
+			bIgnoreStreamingScale = true;
 		}
 
 		const bool bUseLastRenderTime = (StreamingTexture.LastRenderTimeRefCount > 0 || FApp::GetCurrentTime() - StreamingTexture.LastRenderTimeRefCountTimestamp < 91.0);
-		if (((MaxSize == 0 && MaxSize_VisibleOnly == 0) || bUseLastRenderTime) && GStreamWithTimeFactor)
+		if (((MaxSize == 0 && MaxSize_VisibleOnly == 0) || bUseLastRenderTime) && GStreamWithTimeFactor && StreamingTexture.LastRenderTime < 90.0f)
 		{
-			LogHeuristic(bOutputToLog, bGetHeuristic, Heuristic, TEXT("LastRenderTime"));
+			if (bOutputToLog) UE_LOG(LogContentStreaming, Log,  TEXT("  Using Last RenderTime"));
+
+			// Get all mips. Using MaxSize=MaxTextureSize, and not FLT_MAX, will drop one mip if the texture is not used.
+			MaxSize = (float)(0x1 << (StreamingTexture.MaxAllowedMips - 1));
+
 			if( StreamingTexture.LastRenderTime < 5.0f)
 			{
-				MaxSize_VisibleOnly = FLT_MAX;
-			}
-			else if( StreamingTexture.LastRenderTime < 90.0f)
-			{
-				MaxSize = FLT_MAX;
+				MaxSize_VisibleOnly = MaxSize;
 			}
+
+			// Applying streaming scales like GLightmapStreamingFactor would break the mip selection here.
+			bIgnoreStreamingScale = true;
 		}
 	}
 
-	StreamingTexture.SetPerfectWantedMips(MaxSize, MaxSize_VisibleOnly, MipBias);
-	return Heuristic;
+	StreamingTexture.SetPerfectWantedMips(MaxSize, MaxSize_VisibleOnly, MipBias, bIgnoreStreamingScale);
 }
 
 /**
@@ -1781,49 +1698,50 @@ void FStreamingManagerTexture::UpdateLatencyStats( UTexture2D* Texture, int32 Wa
 
 bool FStreamingManagerTexture::HandleListStreamingTexturesCommand( const TCHAR* Cmd, FOutputDevice& Ar )
 {
-	// Collect and report stats
-	CollectTextureStatsName = FParse::Token(Cmd, 0);
-	bCollectTextureStats = true;
-	bReportTextureStats = true;
-	return true;
-}
+	// Make sure the async task is idle (also implies Update_Async is finished and that the distances are valid).
+	AsyncWork->EnsureCompletion();
 
-bool FStreamingManagerTexture::HandleListStreamingTexturesCollectCommand( const TCHAR* Cmd, FOutputDevice& Ar )
-{
-	// Collect stats, but disable automatic reporting
-	CollectTextureStatsName = FParse::Token(Cmd, 0);
-	bCollectTextureStats = true;
-	bReportTextureStats = false;
-	return true;
-}
+	// Sort texture by names so that the state can be compared between runs.
+	TMap<FString, int32> SortedTextures;
+	for ( int32 TextureIndex=0; TextureIndex < StreamingTextures.Num(); ++TextureIndex )
+	{
+		const FStreamingTexture& StreamingTexture = StreamingTextures[TextureIndex];
+		if (!StreamingTexture.Texture) continue;
 
-bool FStreamingManagerTexture::HandleListStreamingTexturesReportReadyCommand( const TCHAR* Cmd, FOutputDevice& Ar )
-{
-	if( bCollectTextureStats )
+		SortedTextures.Add(StreamingTexture.Texture->GetFullName(), TextureIndex);
+	}
+
+	SortedTextures.KeySort(TLess<FString>());
+
+	for (TMap<FString, int32>::TConstIterator It(SortedTextures); It; ++It)
 	{
-		if( TextureStatsReport.Num() > 0 )
+		const FStreamingTexture& StreamingTexture = StreamingTextures[It.Value()];
+		const UTexture2D* Texture2D = StreamingTexture.Texture;
+		UE_LOG(LogContentStreaming, Log,  TEXT("Texture [%d] : %s"), It.Value(), *Texture2D->GetFullName() );
+
+		int32 CurrentMipIndex = FMath::Max(Texture2D->GetNumMips() - StreamingTexture.ResidentMips, 0);
+		int32 WantedMipIndex = FMath::Max(Texture2D->GetNumMips() - StreamingTexture.PerfectWantedMips, 0);
+		int32 MaxMipIndex = FMath::Max(Texture2D->GetNumMips() - StreamingTexture.MaxAllowedMips, 0);
+		const TIndirectArray<struct FTexture2DMipMap>& Mips = Texture2D->PlatformData->Mips;
+
+		if (StreamingTexture.LastRenderTime != MAX_FLT)
 		{
-			return true;
+			UE_LOG(LogContentStreaming, Log,  TEXT("    Current=%dx%d Wanted=%dx%d Max=%dx%d LastRenderTime=%.3f Group=%s"), 
+				Mips[CurrentMipIndex].SizeX, Mips[CurrentMipIndex].SizeY, 
+				Mips[WantedMipIndex].SizeX, Mips[WantedMipIndex].SizeY, 
+				Mips[MaxMipIndex].SizeX, Mips[MaxMipIndex].SizeY, 
+				StreamingTexture.LastRenderTime,
+				UTexture::GetTextureGroupString(StreamingTexture.LODGroup));
 		}
-		return false;
-	}
-	return true;
-}
-bool FStreamingManagerTexture::HandleListStreamingTexturesReportCommand( const TCHAR* Cmd, FOutputDevice& Ar )
-{
-	// TextureStatsReport is assumed to have been populated already via ListStreamingTexturesCollect
-	if( TextureStatsReport.Num() > 0 )
-	{
-		Ar.Logf( TEXT("Listing collected stats for all streaming textures") );
-		for( int32 ReportIndex = 0; ReportIndex < TextureStatsReport.Num(); ReportIndex++ )
+		else
 		{
-			Ar.Logf(*(TextureStatsReport[ReportIndex]));
+			UE_LOG(LogContentStreaming, Log,  TEXT("    Current=%dx%d Wanted=%dx%d Max=%dx%d Group=%s"), 
+				Mips[CurrentMipIndex].SizeX, Mips[CurrentMipIndex].SizeY, 
+				Mips[WantedMipIndex].SizeX, Mips[WantedMipIndex].SizeY, 
+				Mips[MaxMipIndex].SizeX, Mips[MaxMipIndex].SizeY, 
+				UTexture::GetTextureGroupString(StreamingTexture.LODGroup));
 		}
-		TextureStatsReport.Empty();
-	}
-	else
-	{
-		Ar.Logf( TEXT("No stats have been collected for streaming textures. Use ListStreamingTexturesCollect to do so.") );
+
 	}
 	return true;
 }
@@ -2078,13 +1996,76 @@ bool FStreamingManagerTexture::HandleTextureGroupsCommand( const TCHAR* Cmd, FOu
 	return true;
 }
 
-bool FStreamingManagerTexture::HandleInvestigateTextureCommand( const TCHAR* Cmd, FOutputDevice& Ar )
+bool FStreamingManagerTexture::HandleInvestigateTextureCommand( const TCHAR* Cmd, FOutputDevice& Ar, UWorld* InWorld )
 {
-	FString TextureName(FParse::Token(Cmd, 0));
-	if ( TextureName.Len() )
+	FString InvestigateTextureName(FParse::Token(Cmd, 0));
+	if ( InvestigateTextureName.Len() )
 	{
-		bTriggerInvestigateTexture = true;
-		InvestigateTextureName = TextureName;
+		// Make sure the async task is idle (also implies Update_Async is finished and that the distances are valid).
+		AsyncWork->EnsureCompletion();
+		FAsyncTextureStreamingData& StreamingData = AsyncWork->GetTask().StreamingData;
+		StreamingData.Update_Async();
+
+		for ( int32 TextureIndex=0; TextureIndex < StreamingTextures.Num(); ++TextureIndex )
+		{
+			FStreamingTexture& StreamingTexture = StreamingTextures[TextureIndex];
+			FString TextureName = StreamingTexture.Texture->GetFullName();
+			if (TextureName.Contains(InvestigateTextureName))
+			{
+				UTexture2D* Texture2D = StreamingTexture.Texture;
+				int32 CurrentMipIndex = FMath::Max(Texture2D->GetNumMips() - StreamingTexture.ResidentMips, 0);
+				int32 WantedMipIndex = FMath::Max(Texture2D->GetNumMips() - StreamingTexture.PerfectWantedMips, 0);
+				int32 MaxMipIndex = FMath::Max(Texture2D->GetNumMips() - StreamingTexture.MaxAllowedMips, 0);
+
+				UE_LOG(LogContentStreaming, Log,  TEXT("Texture: %s"), *TextureName );
+				UE_LOG(LogContentStreaming, Log,  TEXT("  Texture group:   %s"), UTexture::GetTextureGroupString(StreamingTexture.LODGroup) );
+
+				if ( StreamingTexture.bAsNeverStream )
+				{
+					UE_LOG(LogContentStreaming, Log,  TEXT("  Force all mips: bForceFullyLoad") );
+				}
+				else if ( Texture2D->bGlobalForceMipLevelsToBeResident )
+				{
+					UE_LOG(LogContentStreaming, Log,  TEXT("  Force all mips:  bGlobalForceMipLevelsToBeResident") );
+				}
+				else if ( Texture2D->bForceMiplevelsToBeResident )
+				{
+					UE_LOG(LogContentStreaming, Log,  TEXT("  Force all mips:  bForceMiplevelsToBeResident") );
+				}
+				else if ( Texture2D->ShouldMipLevelsBeForcedResident() )
+				{
+					float CurrentTime = float(FPlatformTime::Seconds() - GStartTime);
+					float TimeLeft = CurrentTime - Texture2D->ForceMipLevelsToBeResidentTimestamp;
+					UE_LOG(LogContentStreaming, Log,  TEXT("  Force all mips:  %.1f seconds left"), FMath::Max(TimeLeft,0.0f) );
+				}
+				else if ( StreamingTexture.MipCount == 1 )
+				{
+					UE_LOG(LogContentStreaming, Log,  TEXT("  Force all mips:  No mip-maps") );
+				}
+				UE_LOG(LogContentStreaming, Log,  TEXT("  Current size:    %dx%d"), Texture2D->PlatformData->Mips[CurrentMipIndex].SizeX, Texture2D->PlatformData->Mips[CurrentMipIndex].SizeY );
+				UE_LOG(LogContentStreaming, Log,  TEXT("  Wanted size:     %dx%d"), Texture2D->PlatformData->Mips[WantedMipIndex].SizeX, Texture2D->PlatformData->Mips[WantedMipIndex].SizeY );
+				UE_LOG(LogContentStreaming, Log,  TEXT("  Max size:        %dx%d"), Texture2D->PlatformData->Mips[MaxMipIndex].SizeX, Texture2D->PlatformData->Mips[MaxMipIndex].SizeY );
+				UE_LOG(LogContentStreaming, Log,  TEXT("  LoadOrder Priority: %d"), StreamingTexture.CalcLoadOrderPriority() );
+				UE_LOG(LogContentStreaming, Log,  TEXT("  Retention Priority: %d"), StreamingTexture.CalcRetentionPriority() );
+				UE_LOG(LogContentStreaming, Log,  TEXT("  Boost factor:    %.1f"), StreamingTexture.BoostFactor );
+				UE_LOG(LogContentStreaming, Log,  TEXT("  Allowed mips:    %d-%d"), StreamingTexture.MinAllowedMips, StreamingTexture.MaxAllowedMips );
+				UE_LOG(LogContentStreaming, Log,  TEXT("  Global mip bias: %.1f"), MipBias );
+
+				if (InWorld && !GIsEditor)
+				{
+					UE_LOG(LogContentStreaming, Log,  TEXT("  Time: World=%.3f LastUpdate=%.3f "), InWorld->GetTimeSeconds(), LastUpdateTime);
+				}
+
+				for( int32 ViewIndex=0; ViewIndex < StreamingData.GetViewInfos().Num(); ViewIndex++ )
+				{
+					// Calculate distance of viewer to bounding sphere.
+					const FStreamingViewInfo& ViewInfo = StreamingData.GetViewInfos()[ViewIndex];
+					UE_LOG(LogContentStreaming, Log,  TEXT("  View%d: Position=(%s) ScreenSize=%f Boost=%f"), ViewIndex, *ViewInfo.ViewOrigin.ToString(), ViewInfo.ScreenSize, ViewInfo.BoostFactor);
+				}
+
+				CalcWantedMips(StreamingData, StreamingTexture, true, true);
+			}
+		}
 	}
 	else
 	{
@@ -2113,18 +2094,6 @@ bool FStreamingManagerTexture::Exec( UWorld* InWorld, const TCHAR* Cmd, FOutputD
 	{
 		return HandleListStreamingTexturesCommand( Cmd, Ar );
 	}
-	else if (FParse::Command(&Cmd,TEXT("ListStreamingTexturesCollect")))
-	{
-		return HandleListStreamingTexturesCollectCommand( Cmd, Ar );
-	}
-	else if (FParse::Command(&Cmd,TEXT("ListStreamingTexturesReportReady")))
-	{
-		return HandleListStreamingTexturesReportReadyCommand( Cmd, Ar );
-	}
-	else if (FParse::Command(&Cmd,TEXT("ListStreamingTexturesReport")))
-	{
-		return HandleListStreamingTexturesReportCommand( Cmd, Ar );
-	}
 #endif
 #if !UE_BUILD_SHIPPING
 	if (FParse::Command(&Cmd, TEXT("ResetMaxEverRequiredTextures")))
@@ -2181,7 +2150,7 @@ bool FStreamingManagerTexture::Exec( UWorld* InWorld, const TCHAR* Cmd, FOutputD
 	}
 	else if (FParse::Command(&Cmd,TEXT("InvestigateTexture")))
 	{
-		return HandleInvestigateTextureCommand( Cmd, Ar );
+		return HandleInvestigateTextureCommand( Cmd, Ar, InWorld );
 	}
 #endif // !UE_BUILD_SHIPPING
 
@@ -2390,81 +2359,3 @@ void FStreamingManagerTexture::DumpTextureGroupStats( bool bDetailedStats )
 //	UE_LOG(LogContentStreaming, Log,  TEXT("%34s: NumTextures=%4d, Current=%7.1f KB"), TEXT("Non-streaming non-pool textures"), NumNonStreamingTextures, NonStreamingSize/1024.0f );
 #endif // !UE_BUILD_SHIPPING
 }
-
-/**
- * Prints out detailed information about streaming textures that has a name that contains the given string.
- * Triggered by the InvestigateTexture exec command.
- *
- * @param InvestigateTextureName	Partial name to match textures against
- */
-void FStreamingManagerTexture::InvestigateTexture( const FString& InInvestigateTextureName )
-{
-	bTriggerInvestigateTexture = false;
-#if !UE_BUILD_SHIPPING
-	for ( int32 TextureIndex=0; TextureIndex < StreamingTextures.Num(); ++TextureIndex )
-	{
-		FStreamingTexture& StreamingTexture = StreamingTextures[TextureIndex];
-		FString TextureName = StreamingTexture.Texture->GetFullName();
-		if (TextureName.Contains(InInvestigateTextureName))
-		{
-			UTexture2D* Texture2D = StreamingTexture.Texture;
-			int32 CurrentMipIndex = FMath::Max(Texture2D->GetNumMips() - StreamingTexture.ResidentMips, 0);
-			int32 WantedMipIndex = FMath::Max(Texture2D->GetNumMips() - StreamingTexture.WantedMips, 0);
-			int32 MaxMipIndex = FMath::Max(Texture2D->GetNumMips() - StreamingTexture.MaxAllowedMips, 0);
-			UE_LOG(LogContentStreaming, Log,  TEXT("Texture: %s"), *TextureName );
-			UE_LOG(LogContentStreaming, Log,  TEXT("  Texture group:   %s"), UTexture::GetTextureGroupString(StreamingTexture.LODGroup) );
-
-			if ( StreamingTexture.bAsNeverStream )
-			{
-				UE_LOG(LogContentStreaming, Log,  TEXT("  Force all mips: bForceFullyLoad") );
-			}
-			else if ( Texture2D->bGlobalForceMipLevelsToBeResident )
-			{
-				UE_LOG(LogContentStreaming, Log,  TEXT("  Force all mips:  bGlobalForceMipLevelsToBeResident") );
-			}
-			else if ( Texture2D->bForceMiplevelsToBeResident )
-			{
-				UE_LOG(LogContentStreaming, Log,  TEXT("  Force all mips:  bForceMiplevelsToBeResident") );
-			}
-			else if ( Texture2D->ShouldMipLevelsBeForcedResident() )
-			{
-				float CurrentTime = float(FPlatformTime::Seconds() - GStartTime);
-				float TimeLeft = CurrentTime - Texture2D->ForceMipLevelsToBeResidentTimestamp;
-				UE_LOG(LogContentStreaming, Log,  TEXT("  Force all mips:  %.1f seconds left"), FMath::Max(TimeLeft,0.0f) );
-			}
-			else if ( StreamingTexture.MipCount == 1 )
-			{
-				UE_LOG(LogContentStreaming, Log,  TEXT("  Force all mips:  No mip-maps") );
-			}
-			UE_LOG(LogContentStreaming, Log,  TEXT("  Current size:    %dx%d"), Texture2D->PlatformData->Mips[CurrentMipIndex].SizeX, Texture2D->PlatformData->Mips[CurrentMipIndex].SizeY );
-			UE_LOG(LogContentStreaming, Log,  TEXT("  Wanted size:     %dx%d"), Texture2D->PlatformData->Mips[WantedMipIndex].SizeX, Texture2D->PlatformData->Mips[WantedMipIndex].SizeY );
-			UE_LOG(LogContentStreaming, Log,  TEXT("  Max size:        %dx%d"), Texture2D->PlatformData->Mips[MaxMipIndex].SizeX, Texture2D->PlatformData->Mips[MaxMipIndex].SizeY );
-			UE_LOG(LogContentStreaming, Log,  TEXT("  LoadOrder Priority: %.3f"), StreamingTexture.CalcLoadOrderPriority() );
-			UE_LOG(LogContentStreaming, Log,  TEXT("  Retention Priority: %.3f"), StreamingTexture.CalcRetentionPriority() );
-			UE_LOG(LogContentStreaming, Log,  TEXT("  Boost factor:    %.1f"), StreamingTexture.BoostFactor );
-			UE_LOG(LogContentStreaming, Log,  TEXT("  Allowed mips:    %d-%d"), StreamingTexture.MinAllowedMips, StreamingTexture.MaxAllowedMips );
-			UE_LOG(LogContentStreaming, Log,  TEXT("  Global mip bias: %.1f"), MipBias );
-
-			// Make sure the async task is idle (also implies Update_Async is finished and that the distances are valid).
-			AsyncWork->EnsureCompletion();
-			FAsyncTextureStreamingData& StreamingData = AsyncWork->GetTask().StreamingData;
-			StreamingData.Update_Async();
-
-			float MaxSize = 0;
-			float MaxSize_VisibleOnly = 0;
-
-			for( int32 ViewIndex=0; ViewIndex < StreamingData.GetViewInfos().Num(); ViewIndex++ )
-			{
-				// Calculate distance of viewer to bounding sphere.
-				const FStreamingViewInfo& ViewInfo = StreamingData.GetViewInfos()[ViewIndex];
-				UE_LOG(LogContentStreaming, Log,  TEXT("  View%d: Position=(%d,%d,%d) ScreenSize=%f Boost=%f"), ViewIndex,
-					int32(ViewInfo.ViewOrigin.X), int32(ViewInfo.ViewOrigin.Y), int32(ViewInfo.ViewOrigin.Z),
-					ViewInfo.ScreenSize, ViewInfo.BoostFactor);
-			}
-
-			CalcWantedMips(StreamingData, StreamingTexture, true, true);
-		}
-	}
-#endif // !UE_BUILD_SHIPPING
-}
-
diff --git a/Engine/Source/Runtime/Engine/Private/Streaming/StreamingManagerTexture.h b/Engine/Source/Runtime/Engine/Private/Streaming/StreamingManagerTexture.h
--- a/Engine/Source/Runtime/Engine/Private/Streaming/StreamingManagerTexture.h
+++ b/Engine/Source/Runtime/Engine/Private/Streaming/StreamingManagerTexture.h
@@ -93,9 +93,6 @@ struct FStreamingManagerTexture : public ITextureStreamingManager
 #endif // STATS_FAST
 #if STATS
 	bool HandleListStreamingTexturesCommand( const TCHAR* Cmd, FOutputDevice& Ar );
-	bool HandleListStreamingTexturesCollectCommand( const TCHAR* Cmd, FOutputDevice& Ar );
-	bool HandleListStreamingTexturesReportReadyCommand( const TCHAR* Cmd, FOutputDevice& Ar );
-	bool HandleListStreamingTexturesReportCommand( const TCHAR* Cmd, FOutputDevice& Ar );
 #endif // STATS
 #if !UE_BUILD_SHIPPING
 	bool HandleResetMaxEverRequiredTexturesCommand(const TCHAR* Cmd, FOutputDevice& Ar);
@@ -111,7 +108,7 @@ struct FStreamingManagerTexture : public ITextureStreamingManager
 	bool HandlePauseTextureStreamingCommand( const TCHAR* Cmd, FOutputDevice& Ar );
 	bool HandleStreamingManagerMemoryCommand( const TCHAR* Cmd, FOutputDevice& Ar, UWorld* InWorld );
 	bool HandleTextureGroupsCommand( const TCHAR* Cmd, FOutputDevice& Ar );
-	bool HandleInvestigateTextureCommand( const TCHAR* Cmd, FOutputDevice& Ar );
+	bool HandleInvestigateTextureCommand( const TCHAR* Cmd, FOutputDevice& Ar, UWorld* InWorld );
 #endif // !UE_BUILD_SHIPPING
 	/** Adds a new texture to the streaming manager. */
 	virtual void AddStreamingTexture( UTexture2D* Texture ) override;
@@ -220,11 +217,8 @@ protected:
 		/** Calculates the minimum and maximum number of mip-levels for a streaming texture. */
 		void CalcMinMaxMips( FStreamingTexture& StreamingTexture );
 
-
-		static FORCEINLINE void LogHeuristic(bool bOutputToLog, bool bUseHeuristic, const TCHAR*& Heuristic, const TCHAR* NewHeuristic);
-
 		/** Calculates the number of mip-levels we would like to have in memory for a texture. */
-		const TCHAR* CalcWantedMips( const FAsyncTextureStreamingData& StreamingData, FStreamingTexture& StreamingTexture, bool bGetHeuristic, bool bOutputToLog );
+		void CalcWantedMips( const FAsyncTextureStreamingData& StreamingData, FStreamingTexture& StreamingTexture, bool bGetHeuristic, bool bOutputToLog );
 
 		/** Updates this frame's STATs by one texture. */
 		void UpdateFrameStats( FStreamingContext& Context, FStreamingTexture& StreamingTexture, int32 TextureIndex );
@@ -298,26 +292,12 @@ protected:
 
 	void	IncrementalUpdate( float Percentage );
 
-	/**
-	 * Prints out detailed information about streaming textures that has a name that contains the given string.
-	 * Triggered by the InvestigateTexture exec command.
-	 *
-	 * @param InvestigateTextureName	Partial name to match textures against
-	 */
-	void	InvestigateTexture( const FString& InvestigateTextureName );
-
 	/** Next sync, dump texture group stats. */
 	bool	bTriggerDumpTextureGroupStats;
 
 	/** Whether to the dumped texture group stats should contain extra information. */
 	bool	bDetailedDumpTextureGroupStats;
 
-	/** Next sync, dump all information we have about a certain texture. */
-	bool	bTriggerInvestigateTexture;
-
-	/** Name of a texture to investigate. Can be partial name. */
-	FString	InvestigateTextureName;
-
 	/** Cached from the system settings. */
 	int32 NumStreamedMips[TEXTUREGROUP_MAX];
 
@@ -409,17 +389,12 @@ protected:
 	/** PoolSize CVar setting the last time we adjusted the pool size. */
 	int32 PreviousPoolSizeSetting;
 
-	/** Whether to collect, and optionally report, texture stats for the next run. */
-	bool   bCollectTextureStats;
-	bool   bReportTextureStats;
-	TArray<FString> TextureStatsReport;
-
-	/** Optional string to match against the texture names when collecting stats. */
-	FString CollectTextureStatsName;
-
 	/** Whether texture streaming is paused or not. When paused, it won't stream any textures in or out. */
 	bool bPauseTextureStreaming;
 
+	/** Last time all data were fully updated. Instances are considered visible if they were rendered between that last time and the current time. */
+	float LastUpdateTime;
+
 #if STATS
 	/**
 	 * Ring buffer containing all latency samples we keep track of, as measured in seconds from the time the streaming system
diff --git a/Engine/Source/Runtime/Engine/Private/Streaming/StreamingTexture.cpp b/Engine/Source/Runtime/Engine/Private/Streaming/StreamingTexture.cpp
--- a/Engine/Source/Runtime/Engine/Private/Streaming/StreamingTexture.cpp
+++ b/Engine/Source/Runtime/Engine/Private/Streaming/StreamingTexture.cpp
@@ -19,8 +19,6 @@ FStreamingTexture::FStreamingTexture( UTexture2D* InTexture )
 	STAT( MostResidentMips = InTexture->ResidentMips );
 	LODGroup = (TextureGroup) InTexture->LODGroup;
 
-	//@TODO: Long term this should probably be a property of the texture group instead
-	bCanBeAffectedByMipBias = LODGroup != TEXTUREGROUP_Terrain_Heightmap && LODGroup != TEXTUREGROUP_Terrain_Weightmap && LODGroup != TEXTUREGROUP_HierarchicalLOD;
 	NumNonStreamingMips = InTexture->GetNumNonStreamingMips();
 	bIsStreamingLightmap = IsStreamingLightmap( Texture );
 	bIsLightmap = bIsStreamingLightmap || LODGroup == TEXTUREGROUP_Lightmap || LODGroup == TEXTUREGROUP_Shadowmap;
@@ -64,21 +62,38 @@ void FStreamingTexture::UpdateCachedInfo( )
 	NumCinematicMipLevels = Texture->bUseCinematicMipLevels ? Texture->NumCinematicMipLevels : 0;
 }
 
+float FStreamingTexture::GetStreamingScale(float GlobalBias) const
+{
+	if (LODGroup == TEXTUREGROUP_Terrain_Heightmap || LODGroup == TEXTUREGROUP_Terrain_Weightmap) return BoostFactor;
+	if (LODGroup == TEXTUREGROUP_Lightmap)	return BoostFactor * GLightmapStreamingFactor;
+	if (LODGroup == TEXTUREGROUP_Shadowmap)	return BoostFactor * GShadowmapStreamingFactor;
+
+	// When using the new metrics, the bounds are conservative and we can relax the required resolution.
+	const bool bUseNewMetrics = CVarStreamingUseNewMetrics.GetValueOnAnyThread() != 0;
+	const float GlobalScale = FMath::Exp2(-GlobalBias) * (bUseNewMetrics ? 0.875f : 1.f);
+
+	return BoostFactor * GlobalScale;
+}
+
+int32 FStreamingTexture::GetWantedMipsFromSize(float Size) const
+{
+	float WantedMipsFloat = 1.0f + FMath::Log2(FMath::Max(1.f, Size));
+	// Using the new metrics has very conservative distance computation. Rounding to mip here compensates.
+	int32 WantedMipsInt = FMath::CeilToInt(WantedMipsFloat);
+	return FMath::Clamp<int32>(WantedMipsInt, MinAllowedMips, MaxAllowedMips);
+}
+
 /** Set the wanted mips from the async task data */
-void FStreamingTexture::SetPerfectWantedMips(float MaxSize, float MaxSize_VisibleOnly, float MipBias)
+void FStreamingTexture::SetPerfectWantedMips(float MaxSize, float MaxSize_VisibleOnly, float MipBias, bool bIgnoreStreamingScale)
 {
 	const float HiddenScale = CVarStreamingHiddenPrimitiveScale.GetValueOnAnyThread();
-	const float MaxResolution = (float)(0x1 << (MaxAllowedMips - 1));
-	const float ExtraScale = GetStreamingScale();
+	const float ExtraScale = bIgnoreStreamingScale ? 1.f : GetStreamingScale(MipBias);
 
 	// Do this now before clamping the resolutions.
 	bAsNeverStream = (MaxSize == FLT_MAX || MaxSize_VisibleOnly == FLT_MAX);
 
-	MaxSize_VisibleOnly = FMath::Clamp<float>(MaxSize_VisibleOnly * ExtraScale, 1.f, MaxResolution);
-	MaxSize = FMath::Clamp<float>(MaxSize * ExtraScale, MaxSize_VisibleOnly, MaxResolution) * HiddenScale;
-
-	int32 WantedMips_VisibleOnly = GetWantedMipsFromSize(MaxSize_VisibleOnly, MipBias);
-	int32 WantedMips_Any = GetWantedMipsFromSize(MaxSize, MipBias);
+	int32 WantedMips_VisibleOnly = GetWantedMipsFromSize(MaxSize_VisibleOnly * ExtraScale);
+	int32 WantedMips_Any = GetWantedMipsFromSize(MaxSize * HiddenScale * ExtraScale);
 
 	// Load visible mips first, then non visible mips.
 	if (WantedMips_Any > WantedMips_VisibleOnly && WantedMips_VisibleOnly <= ResidentMips)
@@ -166,11 +181,11 @@ bool FStreamingTexture::UpdateMipCount(  FStreamingManagerTexture& Manager, FStr
 	return false;
 }
 
-float FStreamingTexture::CalcRetentionPriority( )
+int32 FStreamingTexture::CalcRetentionPriority( )
 {
 	const int32 WantedMipsForPriority = GetWantedMipsForPriority();
 
-	bool bMustKeep = bAsNeverStream || WantedMipsForPriority <= MinAllowedMips || LODGroup == TEXTUREGROUP_Terrain_Heightmap;
+	bool bMustKeep = bAsNeverStream;
 	bool bIsVisible = !GStreamWithTimeFactor || IsVisible();
 	bool bAlreadyLowRez = /*(PixelPerTexel > 1.5 && Distance > 1) ||*/ LODGroup == TEXTUREGROUP_HierarchicalLOD || bIsLightmap;
 		
@@ -179,16 +194,15 @@ float FStreamingTexture::CalcRetentionPriority( )
 	// We also don't want to prioritize by size, since bigger wanted mips are expected to have move visual impact.
 	bool bGivesOffMemory = GetSize(WantedMipsForPriority) - GetSize(WantedMipsForPriority - 1) >= 256 * 1024; 
 
-	float Priority = 0;
-	if (bMustKeep)			Priority += 1024.f;
-	if (bAlreadyLowRez)		Priority += 512.f;
-	if (!bGivesOffMemory)	Priority += 256.f;
-	if (bIsVisible)			Priority += 128.f;
-
+	int32 Priority = 0;
+	if (bMustKeep)			Priority += 1024;
+	if (bIsVisible)			Priority += 512;
+	if (!bGivesOffMemory)	Priority += 256;
+	if (bAlreadyLowRez)		Priority += 128;
 	return Priority;
 }
 
-float FStreamingTexture::CalcLoadOrderPriority()
+int32 FStreamingTexture::CalcLoadOrderPriority()
 {
 	const int32 WantedMipsForPriority = GetWantedMipsForPriority();
 
@@ -197,36 +211,15 @@ float FStreamingTexture::CalcLoadOrderPriority()
 	bool bFarFromTargetMips = WantedMipsForPriority - ResidentMips > 2;
 	bool bBigOnScreen = WantedMipsForPriority > 9.f || LODGroup == TEXTUREGROUP_HierarchicalLOD;
 
-	float Priority = 0;
-	if (bMustLoadFirst)		Priority += 1024.f;
-	if (!bIsLastSplitRequest)	Priority += 512.f; 
-	if (bIsVisible)				Priority += 256.f;
-	if (bFarFromTargetMips)		Priority += 128.f;
-	if (bBigOnScreen)			Priority += 64.f;
+	int32 Priority = 0;
+	if (bIsVisible)				Priority += 1024;
+	if (!bIsLastSplitRequest)	Priority += 512; 
+	if (bFarFromTargetMips)		Priority += 256;
+	if (bMustLoadFirst)			Priority += 128;
+	if (bBigOnScreen)			Priority += 64;
 	return Priority;
 }
 
-int32 FStreamingTexture::GetWantedMipsFromSize(float Size, float Bias) const
-{
-	const bool bUseNewMetrics = CVarStreamingUseNewMetrics.GetValueOnAnyThread() != 0;
-
-	float WantedMipsFloat = 1.0f + FMath::Log2(Size);
-	if (bCanBeAffectedByMipBias)
-	{
-		WantedMipsFloat -= Bias;
-	}
-
-	if (bUseNewMetrics && !bIsLightmap) // Lightmap have approximate logic already so don't apply a bias here.
-	{
-		// Allow a little threshold before going up on mip (so that 513 does not stream a 1024)
-		WantedMipsFloat -= .2f;
-	}
-
-	// Using the new metrics has very conservative distance computation. Rounding to mip here compensates.
-	int32 WantedMipsInt = FMath::CeilToInt(WantedMipsFloat - .01f);
-	return FMath::Clamp<int32>(WantedMipsInt, MinAllowedMips, MaxAllowedMips);
-}
-
 bool FStreamingTexture::IsReadyForStreaming( UTexture2D* Texture )
 {
 	return Texture->IsReadyForStreaming();
diff --git a/Engine/Source/Runtime/Engine/Private/Streaming/StreamingTexture.h b/Engine/Source/Runtime/Engine/Private/Streaming/StreamingTexture.h
--- a/Engine/Source/Runtime/Engine/Private/Streaming/StreamingTexture.h
+++ b/Engine/Source/Runtime/Engine/Private/Streaming/StreamingTexture.h
@@ -65,7 +65,7 @@ struct FStreamingTexture
 	 */
 	FORCEINLINE bool CanDropMips() const
 	{
-		return (LODGroup != TEXTUREGROUP_Terrain_Heightmap) && WantedMips > MinAllowedMips && !bAsNeverStream;
+		return LODGroup != TEXTUREGROUP_Terrain_Heightmap && LODGroup != TEXTUREGROUP_Terrain_Weightmap;
 	}
 
 
@@ -78,23 +78,15 @@ struct FStreamingTexture
 	{
 		if (bHasSplitRequest && !bIsLastSplitRequest)
 		{
-			return WantedMips + 1;
+			return PerfectWantedMips + 1;
 		}
 		else
 		{
-			return WantedMips;
+			return PerfectWantedMips;
 		}
 	}
 
-	FORCEINLINE float GetStreamingScale() const
-	{
-		if (LODGroup == TEXTUREGROUP_Lightmap)
-			return BoostFactor * GLightmapStreamingFactor;
-		else if (LODGroup == TEXTUREGROUP_Shadowmap)
-			return BoostFactor * GShadowmapStreamingFactor;
-		else 
-			return BoostFactor;
-	}
+	FORCEINLINE float GetStreamingScale(float GlobalBias) const;
 
 	/**
 	 * Calculates a retention priority value for the textures. Higher value means more important.
@@ -103,16 +95,16 @@ struct FStreamingTexture
 	 * Not doing so could make the streamer go into a loop where is never stops dropping and loading different textures when out of budget.
 	 * @return		Priority value
 	 */
-	float CalcRetentionPriority( );
+	int32 CalcRetentionPriority( );
 
 	/**
 	 * Calculates a load order priority value for the texture. Higher value means more important.
 	 * Load Order can depend on the current state of resident mips, because it will not affect the streamer stability.
 	 * @return		Priority value
 	 */
-	float CalcLoadOrderPriority();
+	int32 CalcLoadOrderPriority();
 
-	int32 GetWantedMipsFromSize(float Size, float Bias) const;
+	FORCEINLINE int32 GetWantedMipsFromSize(float Size) const;
 
 
 	/**
@@ -125,7 +117,7 @@ struct FStreamingTexture
 	}
 
 	/** Set the wanted mips from the async task data */
-	void SetPerfectWantedMips(float MaxSize, float MaxSize_VisibleOnly, float MipBias);
+	void SetPerfectWantedMips(float MaxSize, float MaxSize_VisibleOnly, float MipBias, bool bIgnoreStreamingScale);
 
 	/** Update texture streaming. Returns whether it did anything */
 	bool UpdateMipCount(  FStreamingManagerTexture& Manager, FStreamingContext& Context );
@@ -195,8 +187,6 @@ struct FStreamingTexture
 	uint32			bHasSplitRequest : 1;
 	/** Whether this is the second request to converge to the wanted mip. */
 	uint32			bIsLastSplitRequest : 1;
-	/** Can this texture be affected by mip bias? */
-	uint32			bCanBeAffectedByMipBias : 1;
 	/** Is the wanted mip for a visible wanted mips */
 	uint32			bIsVisibleWantedMips : 1;
 
diff --git a/Engine/Source/Runtime/Engine/Private/Streaming/TextureInstanceManager.cpp b/Engine/Source/Runtime/Engine/Private/Streaming/TextureInstanceManager.cpp
--- a/Engine/Source/Runtime/Engine/Private/Streaming/TextureInstanceManager.cpp
+++ b/Engine/Source/Runtime/Engine/Private/Streaming/TextureInstanceManager.cpp
@@ -122,44 +122,82 @@ FBoxSphereBounds FTextureInstanceState::FTextureLinkConstIterator::GetBounds() c
 	return Bounds;
 }
 
-void FTextureInstanceState::FTextureLinkConstIterator::OutputToLog(float MaxNormalizedSize, float MaxNormalizedSize_VisibleOnly) const
+void FTextureInstanceState::FTextureLinkConstIterator::OutputToLog(float MaxNormalizedSize, float MaxNormalizedSize_VisibleOnly, const TCHAR* Prefix) const
 {
+#if !UE_BUILD_SHIPPING
 	const UPrimitiveComponent* Component = GetComponent();
 	FBoxSphereBounds Bounds = GetBounds();
 	float TexelFactor = GetTexelFactor();
 
 	// Log the component reference.
-	FString ReferenceName(TEXT("Unkown Component"));
 	if (Component)
 	{
-		ReferenceName = Component->GetFullName();
+		UE_LOG(LogContentStreaming, Log, TEXT("  %sReference= %s"), Prefix, *Component->GetFullName());
 	}
-	UE_LOG(LogContentStreaming, Log, TEXT("    Reference= %s"), *ReferenceName);
-
+	else
+	{
+		UE_LOG(LogContentStreaming, Log, TEXT("  %sReference"), Prefix);
+	}
+	
 	// Log the wanted mips.
 	if (TexelFactor == FLT_MAX)
 	{
-		UE_LOG(LogContentStreaming, Log, TEXT("    Forcing all mips."), *ReferenceName);
+		UE_LOG(LogContentStreaming, Log, TEXT("    Forced FullyLoad"));
 	}
 	else if (TexelFactor >= 0)
 	{
-		UE_LOG(LogContentStreaming, Log, TEXT("    TexelFactor=%f, MaxSize=%f, MaxVisibleSize=%f"), TexelFactor, TexelFactor * MaxNormalizedSize, TexelFactor * MaxNormalizedSize_VisibleOnly);
+		if (GIsEditor)
+		{
+			UE_LOG(LogContentStreaming, Log, TEXT("    Size=%f"), TexelFactor * FMath::Max(MaxNormalizedSize, MaxNormalizedSize_VisibleOnly));
+		}
+		else if (MaxNormalizedSize_VisibleOnly > 0)
+		{
+			UE_LOG(LogContentStreaming, Log, TEXT("    OnScreenSize=%f"), TexelFactor * MaxNormalizedSize_VisibleOnly);
+		}
+		else
+		{
+			int32 BoundIndex = GetBoundsIndex();
+			if (!GIsEditor && State.Bounds4.IsValidIndex(BoundIndex / 4))
+			{
+				float LastRenderTime = State.Bounds4[BoundIndex / 4].LastRenderTime[BoundIndex % 4];
+				UE_LOG(LogContentStreaming, Log, TEXT("    OffScreenSize=%f, LastRenderTime= %.3f"), TexelFactor * MaxNormalizedSize, LastRenderTime);
+			}
+			else
+			{
+				UE_LOG(LogContentStreaming, Log, TEXT("    OffScreenSize=%f"), TexelFactor * MaxNormalizedSize);
+			}
+		}
 	}
 	else
 	{
-		UE_LOG(LogContentStreaming, Log, TEXT("    Forced Size=%f"), -TexelFactor);
+		UE_LOG(LogContentStreaming, Log, TEXT("    ForcedSize=%f"), -TexelFactor);
 	}
 
 	// Log the bounds
 	const bool bUseNewMetrics = CVarStreamingUseNewMetrics.GetValueOnAnyThread() != 0;
 	if (bUseNewMetrics)
 	{
-		UE_LOG(LogContentStreaming, Log, TEXT("    Origin=(%s), BoxExtent=(%s)"), *Bounds.Origin.ToString(), *Bounds.BoxExtent.ToString());
+		if (TexelFactor >= 0 && TexelFactor < FLT_MAX)
+		{
+			UE_LOG(LogContentStreaming, Log, TEXT("    Origin=(%s), BoxExtent=(%s), TexelSize=%f"), *Bounds.Origin.ToString(), *Bounds.BoxExtent.ToString(), TexelFactor);
+		}
+		else
+		{
+			UE_LOG(LogContentStreaming, Log, TEXT("    Origin=(%s), BoxExtent=(%s)"), *Bounds.Origin.ToString(), *Bounds.BoxExtent.ToString());
+		}
 	}
 	else
 	{
-		UE_LOG(LogContentStreaming, Log, TEXT("    Origin=(%s), SphereRadius=%f"),  *Bounds.Origin.ToString(), Bounds.SphereRadius);
+		if (TexelFactor >= 0 && TexelFactor < FLT_MAX)
+		{
+			UE_LOG(LogContentStreaming, Log, TEXT("    Origin=(%s), SphereRadius=%f, TexelSize=%f"),  *Bounds.Origin.ToString(), Bounds.SphereRadius, TexelFactor);
+		}
+		else
+		{
+			UE_LOG(LogContentStreaming, Log, TEXT("    Origin=(%s), SphereRadius=%f"),  *Bounds.Origin.ToString(), Bounds.SphereRadius);
+		}
 	}
+#endif // !UE_BUILD_SHIPPING
 }
 
 
@@ -395,15 +433,19 @@ void FTextureInstanceState::AddComponent(const UPrimitiveComponent* Component, F
 
 				if (BoundsIndex == INDEX_NONE)
 				{
-					// Take into account the offset between the bounds and the components.
-					float MinDistance = FMath::Min<float>(0, Component->MinDrawDistance - (Info.Bounds.Origin - Component->Bounds.Origin).Size());
+					// In the engine, the MinDistance is computed from the component bound center to the viewpoint.
+					// The streaming computes the distance as the distance from viewpoint to the edge of the texture bound box.
+					// The implementation also handles MinDistance by bounding the distance to it so that if the viewpoint gets closer the screen size will stop increasing at some point.
+					// The fact that the primitive will disappear is not so relevant as this will be handled by the visibility logic, normally streaming one less mip than requested.
+					// The important mather is to control the requested mip by limiting the distance, since at close up, the distance becomes very small and all mips are streamer (even after the 1 mip bias).
+					float MinDistance = FMath::Max<float>(0, Component->MinDrawDistance - (Info.Bounds.Origin - Component->Bounds.Origin).Size() - Info.Bounds.SphereRadius);
 					float MaxDistance = FLT_MAX;
 					if (LODParent)
 					{
 						// Max distance when HLOD becomes visible.
 						MaxDistance = LODParent->MinDrawDistance + (Info.Bounds.Origin - LODParent->Bounds.Origin).Size();
 					}
-					BoundsIndex = AddBounds(Info.Bounds, Component, Component->LastRenderTime, MinDistance, MaxDistance);
+					BoundsIndex = AddBounds(Info.Bounds, Component, Component->LastRenderTimeOnScreen, MinDistance, MaxDistance);
 				}
 				BoundsIndices.Push(BoundsIndex);
 			}
@@ -422,7 +464,7 @@ void FTextureInstanceState::AddComponent(const UPrimitiveComponent* Component)
 
 	if (TextureInstanceInfos.Num())
 	{
-		int32 BoundsIndex = AddBounds(Component->Bounds, Component, Component->LastRenderTime);
+		int32 BoundsIndex = AddBounds(Component->Bounds, Component, Component->LastRenderTimeOnScreen);
 		int32* ComponentLink = ComponentMap.Find(Component);
 		for (const FStreamingTexturePrimitiveInfo& Info : TextureInstanceInfos)
 		{
@@ -496,7 +538,7 @@ void FTextureInstanceState::UpdateBounds(const UPrimitiveComponent* Component)
 			const FElement& Element = Elements[ElementIndex];
 			if (Element.BoundsIndex != INDEX_NONE)
 			{
-				Bounds4[Element.BoundsIndex / 4].Update(Element.BoundsIndex % 4, Component->Bounds, Component->LastRenderTime);
+				Bounds4[Element.BoundsIndex / 4].Update(Element.BoundsIndex % 4, Component->Bounds, Component->LastRenderTimeOnScreen);
 			}
 			ElementIndex = Element.NextComponentLink;
 		}
@@ -508,7 +550,7 @@ void FTextureInstanceState::UpdateBounds(int32 BoundIndex)
 	const UPrimitiveComponent* Component = Bounds4Components[BoundIndex];
 	if (Component)
 	{
-		Bounds4[BoundIndex / 4].Update(BoundIndex % 4, Component->Bounds, Component->LastRenderTime);
+		Bounds4[BoundIndex / 4].Update(BoundIndex % 4, Component->Bounds, Component->LastRenderTimeOnScreen);
 	}
 }
 
@@ -517,7 +559,7 @@ void FTextureInstanceState::UpdateLastRenderTime(int32 BoundIndex)
 	const UPrimitiveComponent* Component = Bounds4Components[BoundIndex];
 	if (Component)
 	{
-		Bounds4[BoundIndex / 4].Update(BoundIndex % 4, Component->LastRenderTime);
+		Bounds4[BoundIndex / 4].Update(BoundIndex % 4, Component->LastRenderTimeOnScreen);
 	}
 }
 
@@ -642,7 +684,7 @@ TRefCountPtr<const FTextureInstanceState> FDynamicTextureInstanceManager::GetAsy
 	return TRefCountPtr<const FTextureInstanceState>(AsyncState);
 }
 
-void FTextureInstanceAsyncView::Update_Async(const TArray<FStreamingViewInfo>& ViewInfos, bool bUseAABB, float MaxEffectiveScreenSize)
+void FTextureInstanceAsyncView::Update_Async(const TArray<FStreamingViewInfo>& ViewInfos, float LastUpdateTime, bool bUseAABB, float MaxEffectiveScreenSize)
 {
 	if (!State.IsValid())  return;
 
@@ -650,11 +692,6 @@ void FTextureInstanceAsyncView::Update_Async(const TArray<FStreamingViewInfo>& V
 	const VectorRegister One4 = VectorSet(1.f, 1.f, 1.f, 1.f);
 	const int32 NumBounds4 = State->NumBounds4();
 
-	//@TODO : This should be the value since the last update (something will be visible if seen since last update, no matter how long).
-	// const float LastUpdateTime =  FApp::GetCurrentTime() - GStartTime - 5.f;
-	// const float LastUpdateTime =  World->GetTimeSeconds();
-	const float LastUpdateTime = -1005.f;
-
 	const VectorRegister LastUpdateTime4 = VectorSet(LastUpdateTime, LastUpdateTime, LastUpdateTime, LastUpdateTime);
 
 	BoundsViewInfo.Empty(NumBounds4 * 4);
@@ -691,23 +728,11 @@ void FTextureInstanceAsyncView::Update_Async(const TArray<FStreamingViewInfo>& V
 			const VectorRegister ViewOriginY = VectorLoadFloat1( &ViewInfo.ViewOrigin.Y );
 			const VectorRegister ViewOriginZ = VectorLoadFloat1( &ViewInfo.ViewOrigin.Z );
 
-			//const float DistSq = FVector::DistSquared( ViewInfo.ViewOrigin, TextureInstance.BoundingSphere.Center );
-			VectorRegister Temp = VectorSubtract( ViewOriginX, OriginX );
-			VectorRegister DistSq = VectorMultiply( Temp, Temp );
-			Temp = VectorSubtract( ViewOriginY, OriginY );
-			DistSq = VectorMultiplyAdd( Temp, Temp, DistSq );
-			Temp = VectorSubtract( ViewOriginZ, OriginZ );
-			DistSq = VectorMultiplyAdd( Temp, Temp, DistSq );
-
-			VectorRegister ClampedDistSq = VectorMax( MinDistanceSq, DistSq );
-			ClampedDistSq = VectorMin( MaxDistanceSq, ClampedDistSq );
-			const VectorRegister InRangeMask = VectorCompareEQ(DistSq, ClampedDistSq);
-
 			VectorRegister DistSqMinusRadiusSq = VectorZero();
 			if (bUseAABB)
 			{
 				// In this case DistSqMinusRadiusSq will contain the distance to the box^2
-				Temp = VectorSubtract( ViewOriginX, OriginX );
+				VectorRegister Temp = VectorSubtract( ViewOriginX, OriginX );
 				Temp = VectorAbs( Temp );
 				VectorRegister BoxRef = VectorMin( Temp, ExtentX );
 				Temp = VectorSubtract( Temp, BoxRef );
@@ -727,11 +752,25 @@ void FTextureInstanceAsyncView::Update_Async(const TArray<FStreamingViewInfo>& V
 			}
 			else
 			{
+				VectorRegister Temp = VectorSubtract( ViewOriginX, OriginX );
+				VectorRegister DistSq = VectorMultiply( Temp, Temp );
+				Temp = VectorSubtract( ViewOriginY, OriginY );
+				DistSq = VectorMultiplyAdd( Temp, Temp, DistSq );
+				Temp = VectorSubtract( ViewOriginZ, OriginZ );
+				DistSq = VectorMultiplyAdd( Temp, Temp, DistSq );
+
 				DistSqMinusRadiusSq = VectorMultiply( Radius, Radius );
 				DistSqMinusRadiusSq = VectorSubtract( DistSq, DistSqMinusRadiusSq );
 			}
-			DistSqMinusRadiusSq = VectorMax(DistSqMinusRadiusSq, One4); // Prevents / 0
-			VectorRegister ScreenSizeOverDistance = VectorReciprocalSqrt(DistSqMinusRadiusSq);
+
+			// Clamp Squared distance between range.
+			VectorRegister ClampedDistSq = VectorMax( MinDistanceSq, DistSqMinusRadiusSq );
+			ClampedDistSq = VectorMin( MaxDistanceSq, ClampedDistSq );
+			const VectorRegister InRangeMask = VectorCompareEQ(DistSqMinusRadiusSq, ClampedDistSq); // If the clamp dist is equal, then it was in range.
+
+
+			ClampedDistSq = VectorMax(ClampedDistSq, One4); // Prevents / 0
+			VectorRegister ScreenSizeOverDistance = VectorReciprocalSqrt(ClampedDistSq);
 			ScreenSizeOverDistance = VectorMultiply(ScreenSizeOverDistance, ScreenSize);
 
 			MaxNormalizedSize = VectorMax(ScreenSizeOverDistance, MaxNormalizedSize);
@@ -772,7 +811,7 @@ void FTextureInstanceAsyncView::ProcessElement(const FBoundsViewInfo& BoundsVieW
 	}
 }
 
-void FTextureInstanceAsyncView::GetTexelSize(const UTexture2D* InTexture, float& MaxSize, float& MaxSize_VisibleOnly, bool bOutputToLog) const
+void FTextureInstanceAsyncView::GetTexelSize(const UTexture2D* InTexture, float& MaxSize, float& MaxSize_VisibleOnly, const TCHAR* LogPrefix) const
 {
 	// No need to iterate more if texture is already at maximum resolution.
 
@@ -781,7 +820,7 @@ void FTextureInstanceAsyncView::GetTexelSize(const UTexture2D* InTexture, float&
 	if (State.IsValid())
 	{
 		// Use the fast path if available, about twice as fast when there are a lot of elements.
-		if (State->HasCompiledElements() && !bOutputToLog)
+		if (State->HasCompiledElements() && !LogPrefix)
 		{
 			const TArray<FTextureInstanceState::FCompiledElement>* CompiledElements = State->GetCompiledElements(InTexture);
 			if (CompiledElements)
@@ -807,13 +846,13 @@ void FTextureInstanceAsyncView::GetTexelSize(const UTexture2D* InTexture, float&
 		}
 		else
 		{
-			for (auto It = State->GetElementIterator(InTexture); It && (MaxSize_VisibleOnly < MAX_TEXTURE_SIZE || bOutputToLog); ++It)
+			for (auto It = State->GetElementIterator(InTexture); It && (MaxSize_VisibleOnly < MAX_TEXTURE_SIZE || LogPrefix); ++It)
 			{
 				const FBoundsViewInfo& BoundsVieWInfo = BoundsViewInfo[It.GetBoundsIndex()];
 				ProcessElement(BoundsVieWInfo, It.GetTexelFactor(), MaxSize, MaxSize_VisibleOnly);
-				if (bOutputToLog)
+				if (LogPrefix)
 				{
-					It.OutputToLog(BoundsVieWInfo.MaxNormalizedSize, BoundsVieWInfo.MaxNormalizedSize_VisibleOnly);
+					It.OutputToLog(BoundsVieWInfo.MaxNormalizedSize, BoundsVieWInfo.MaxNormalizedSize_VisibleOnly, LogPrefix);
 				}
 			}
 		}
@@ -948,6 +987,19 @@ void FLevelTextureManager::Remove(FDynamicComponentTextureManager& DynamicCompon
 	bToDelete = true; 
 }
 
+float FLevelTextureManager::GetWorldTime() const
+{
+	if (!GIsEditor && Level)
+	{
+		UWorld* World = Level->GetWorld();
+		if (World)
+		{
+			return World->GetTimeSeconds();
+		}
+	}
+	return 0;
+}
+
 void FLevelTextureManager::IncrementalUpdate(FDynamicComponentTextureManager& DynamicComponentManager, FRemovedTextureArray& RemovedTextures, float Percentage, bool bUseDynamicStreaming) 
 {
 	QUICK_SCOPE_CYCLE_COUNTER(FStaticComponentTextureManager_IncrementalUpdate);
diff --git a/Engine/Source/Runtime/Engine/Private/Streaming/TextureInstanceManager.h b/Engine/Source/Runtime/Engine/Private/Streaming/TextureInstanceManager.h
--- a/Engine/Source/Runtime/Engine/Private/Streaming/TextureInstanceManager.h
+++ b/Engine/Source/Runtime/Engine/Private/Streaming/TextureInstanceManager.h
@@ -94,7 +94,7 @@ public:
 		FORCEINLINE operator bool() const { return CurrElementIndex != INDEX_NONE; }
 		FORCEINLINE void operator++() { CurrElementIndex = State.Elements[CurrElementIndex].NextTextureLink; }
 
-		void OutputToLog(float MaxNormalizedSize, float MaxNormalizedSize_VisibleOnly) const;
+		void OutputToLog(float MaxNormalizedSize, float MaxNormalizedSize_VisibleOnly, const TCHAR* Prefix) const;
 
 		FORCEINLINE int32 GetBoundsIndex() const { return State.Elements[CurrElementIndex].BoundsIndex; }
 		FORCEINLINE float GetTexelFactor() const { return State.Elements[CurrElementIndex].TexelFactor; }
@@ -287,11 +287,11 @@ public:
 
 	FTextureInstanceAsyncView(const TRefCountPtr<const FTextureInstanceState>& InState) : State(InState) {}
 
-	void Update_Async(const TArray<FStreamingViewInfo>& ViewInfos, bool bUseAABB, float MaxEffectiveScreenSize);
+	void Update_Async(const TArray<FStreamingViewInfo>& ViewInfos, float LastUpdateTime, bool bUseAABB, float MaxEffectiveScreenSize);
 
 	// MaxSize : Biggest texture size for all instances.
 	// MaxSize_VisibleOnly : Biggest texture size for visble instances only.
-	void GetTexelSize(const UTexture2D* InTexture, float& MaxSize, float& MaxSize_VisibleOnly, bool bOutputToLog) const;
+	void GetTexelSize(const UTexture2D* InTexture, float& MaxSize, float& MaxSize_VisibleOnly, const TCHAR* LogPrefix) const;
 
 	bool HasTextureReferences(const UTexture2D* InTexture) const;
 
@@ -397,6 +397,8 @@ public:
 
 	const FStaticTextureInstanceManager& GetStaticInstances() const { return StaticInstances; }
 
+	float GetWorldTime() const;
+
 	FORCEINLINE FTextureInstanceAsyncView GetAsyncView() { return FTextureInstanceAsyncView(StaticInstances.GetAsyncState()); }
 
 	void IncrementalUpdate(FDynamicComponentTextureManager& DynamicManager, FRemovedTextureArray& RemovedTextures, float Percentage, bool bUseDynamicStreaming);
diff --git a/Engine/Source/Runtime/Engine/Private/Streaming/TextureStreamingHelpers.cpp b/Engine/Source/Runtime/Engine/Private/Streaming/TextureStreamingHelpers.cpp
--- a/Engine/Source/Runtime/Engine/Private/Streaming/TextureStreamingHelpers.cpp
+++ b/Engine/Source/Runtime/Engine/Private/Streaming/TextureStreamingHelpers.cpp
@@ -197,9 +197,8 @@ bool IsStreamingTexture( const UTexture2D* Texture2D )
 	return false;
 }
 
-void FStreamingContext::Reset( bool bProcessEverything, UTexture2D* IndividualStreamingTexture, bool bInCollectTextureStats )
+void FStreamingContext::Reset( bool bProcessEverything, UTexture2D* IndividualStreamingTexture )
 {
-	bCollectTextureStats							= bInCollectTextureStats;
 	ThisFrameTotalRequestSize						= 0;
 	ThisFrameTotalLightmapRequestSize				= 0;
 	ThisFrameNumStreamingTextures					= 0;
@@ -225,8 +224,6 @@ void FStreamingContext::Reset( bool bProcessEverything, UTexture2D* IndividualSt
 	AvailableLater = 0;
 	AvailableTempMemory = 0;
 
-	STAT( TextureStats.Empty() );
-
 	AllocatedMemorySize	= INDEX_NONE;
 	PendingMemoryAdjustment = INDEX_NONE;
 		
@@ -278,6 +275,4 @@ void FStreamingContext::AddStats( const FStreamingContext& Other )
 	ThisFrameTotalDynamicTextureHeuristicSize		+= Other.ThisFrameTotalDynamicTextureHeuristicSize;
 	ThisFrameTotalLastRenderHeuristicSize			+= Other.ThisFrameTotalLastRenderHeuristicSize;
 	ThisFrameTotalForcedHeuristicSize				+= Other.ThisFrameTotalForcedHeuristicSize;
-	bCollectTextureStats							= bCollectTextureStats || Other.bCollectTextureStats;
-	STAT( TextureStats.Append( Other.TextureStats ) );
 }
diff --git a/Engine/Source/Runtime/Engine/Private/Streaming/TextureStreamingHelpers.h b/Engine/Source/Runtime/Engine/Private/Streaming/TextureStreamingHelpers.h
--- a/Engine/Source/Runtime/Engine/Private/Streaming/TextureStreamingHelpers.h
+++ b/Engine/Source/Runtime/Engine/Private/Streaming/TextureStreamingHelpers.h
@@ -156,104 +156,22 @@ struct FTextureSortElement
 	int32			NumRequiredResidentMips;
 };
 
-struct FTexturePriority
-{
-	FTexturePriority( bool InCanDropMips, float InRetentionPriority, float InLoadOrderPriority, int32 InTextureIndex, const UTexture2D* InTexture)
-	:	bCanDropMips( InCanDropMips )
-	,	RetentionPriority( InRetentionPriority )
-	,	LoadOrderPriority( InLoadOrderPriority )
-	,	TextureIndex( InTextureIndex )
-	,	Texture( InTexture )
-	{
-	}
-	/** True if we allows this texture to drop mips to fit in budget. */
-	bool bCanDropMips;
-	/** Texture retention priority, higher value means it should be kept in memory. */
-	float RetentionPriority;
-	/** Texture load order priority, higher value means it should load first. */
-	float LoadOrderPriority;
-	/** Index into the FStreamingManagerTexture::StreamingTextures array. */
-	int32 TextureIndex;
-	/** The texture to stream. Only used for validation. */
-	const UTexture2D* Texture;
-};
-
-
-
-#if STATS
-struct FTextureStreamingStats
-{
-	FTextureStreamingStats( UTexture2D* InTexture2D, int32 InResidentMips, int32 InWantedMips, int32 InMostResidentMips, int32 InResidentSize, int32 InWantedSize, int32 InMaxSize, int32 InMostResidentSize, float InBoostFactor, float InPriority, int32 InTextureIndex )
-	:	TextureName( InTexture2D->GetFullName() )
-	,	SizeX( InTexture2D->GetSizeX() )
-	,	SizeY( InTexture2D->GetSizeY() )
-	,	NumMips( InTexture2D->GetNumMips() )
-	,	LODBias( InTexture2D->GetCachedLODBias() )
-	,	LastRenderTime( FMath::Clamp( InTexture2D->Resource ? (FApp::GetCurrentTime() - InTexture2D->Resource->LastRenderTime) : 1000000.0, 0.0, 1000000.0) )
-	,	ResidentMips( InResidentMips )
-	,	WantedMips( InWantedMips )
-	,	MostResidentMips( InMostResidentMips )
-	,	ResidentSize( InResidentSize )
-	,	WantedSize( InWantedSize )
-	,	MaxSize( InMaxSize )
-	,	MostResidentSize( InMostResidentSize )
-	,	BoostFactor( InBoostFactor )
-	,	Priority( InPriority )
-	,	TextureIndex( InTextureIndex )
-	{
-	}
-	/** Mirror of UTexture2D::GetName() */
-	FString		TextureName;
-	/** Mirror of UTexture2D::SizeX */
-	int32		SizeX;
-	/** Mirror of UTexture2D::SizeY */
-	int32		SizeY;
-	/** Mirror of UTexture2D::Mips.Num() */
-	int32		NumMips;
-	/** Mirror of UTexture2D::GetCachedLODBias() */
-	int32		LODBias;
-	/** How many seconds since it was last rendered: FApp::GetCurrentTime() - UTexture2D::Resource->LastRenderTime */
-	float		LastRenderTime;
-	/** Number of resident mip levels. */
-	int32		ResidentMips;
-	/** Number of wanted mip levels. */
-	int32		WantedMips;
-	/** Most number of mip-levels this texture has ever had resident in memory. */
-	int32		MostResidentMips;
-	/** Number of bytes currently in memory. */
-	int32		ResidentSize;
-	/** Number of bytes we want in memory. */
-	int32		WantedSize;
-	/** Number of bytes we could potentially stream in. */
-	int32		MaxSize;
-	/** Most number of bytes this texture has ever had resident in memory. */
-	int32		MostResidentSize;
-	/** Temporary boost of the streaming distance factor. */
-	float		BoostFactor;
-	/** Texture priority */
-	float		Priority;
-	/** Index into the FStreamingManagerTexture::StreamingTextures array. */
-	int32		TextureIndex;
-};
-#endif
-
 /**
  * Helper struct for temporary information for one frame of processing texture streaming.
  */
 struct FStreamingContext
 {
-	FStreamingContext( bool bProcessEverything, UTexture2D* IndividualStreamingTexture, bool bInCollectTextureStats )
+	FStreamingContext( bool bProcessEverything, UTexture2D* IndividualStreamingTexture )
 	{
-		Reset( bProcessEverything, IndividualStreamingTexture, bInCollectTextureStats );
+		Reset( bProcessEverything, IndividualStreamingTexture );
 	}
 
 	/**
 	 * Initializes all variables for the one frame.
 	 * @param bProcessEverything			If true, process all resources with no throttling limits
 	 * @param IndividualStreamingTexture	A specific texture to be fully processed this frame, or NULL
-	 * @param bInCollectTextureStats		Whether to fill in the TextureStats array this frame
 	 */
-	void Reset( bool bProcessEverything, UTexture2D* IndividualStreamingTexture, bool bInCollectTextureStats );
+	void Reset( bool bProcessEverything, UTexture2D* IndividualStreamingTexture );
 
 	/**
 	 * Adds in the stats from another context.
@@ -268,12 +186,6 @@ struct FStreamingContext
 	int64 AllocatedMemorySize;
 	/** Pending Adjustments to allocated texture memory, due to async reallocations, etc. */
 	int64 PendingMemoryAdjustment;
-	/** Whether to fill in TextureStats this frame. */
-	bool bCollectTextureStats;
-#if STATS
-	/** Stats for all textures. */
-	TArray<FTextureStreamingStats>	TextureStats;
-#endif
 
 	// Stats for this frame.
 	uint64 ThisFrameTotalRequestSize;
