Copying //UE4/Dev-Rendering to //UE4/Dev-Main (Source: //UE4/Dev-Rendering @ 2998063)

diff --git a/Engine/Source/Runtime/Engine/Private/SubDSurfaceComponent.cpp b/Engine/Source/Runtime/Engine/Private/SubDSurfaceComponent.cpp
new file mode 100644
--- /dev/null
+++ b/Engine/Source/Runtime/Engine/Private/SubDSurfaceComponent.cpp
@@ -0,0 +1,996 @@
+// Copyright 1998-2014 Epic Games, Inc. All Rights Reserved.
+
+#include "EnginePrivate.h"
+#include "Components/SubDSurfaceComponent.h"
+#include "DynamicMeshBuilder.h"
+#include "Engine/SubDSurface.h"
+#include "Engine/Font.h"
+#include "Engine/SubDSurfaceActor.h"
+#include "LocalVertexFactory.h"
+//#include "RawMesh.h"	//  fatal error C1083: Cannot open include file: 'RawMesh.h': No such file or directory  (when RawMesh is not part of PrivateDependencyModuleNames in Engine.Build.cs)
+
+// 0/1 - currently disabled as the .lib files are not yet compiled form VS2013
+#define USE_OPENSUBDIV 0
+
+#if USE_OPENSUBDIV
+
+// Disable warning C4191: 'type cast' : unsafe conversion
+#pragma warning(push)
+#pragma warning(disable:4191)
+
+//	#include "osd/cpuComputeContext.h"
+//	#include "osd/cpuComputeController.h"
+	#include "osd/cpuVertexBuffer.h"
+	#include <far/topologyRefinerFactory.h>
+	#include <far/topologyDescriptor.h>
+	#include <far/primvarRefiner.h>
+
+// from RefineAdaptive?
+	#include <far/patchTableFactory.h>
+	#include <far/patchMap.h>
+	#include <far/ptexIndices.h>
+
+	#include <far/stencilTableFactory.h>
+
+// for non manifold
+	#include <hbr/mesh.h>
+	#include <hbr/catmark.h>
+
+// Restore warning C4191.
+#pragma warning(pop)
+
+using namespace OpenSubdiv;
+
+#endif // USE_OPENSUBDIV
+
+uint32 USubDSurface::GetVertexCount() const
+{
+	if(!VertexAttributeStreams.Num())
+	{
+		return 0;
+	}
+
+	// assumes are all the same size
+	return VertexAttributeStreams[0]->Num();
+}
+
+UVertexAttributeStream* USubDSurface::CreateVertexAttributeStream(FName InUsage)
+{
+	// ensure the stream wasn't there before
+	check(!FindStreamByUsage(InUsage));
+
+	uint32 Num = VertexAttributeStreams.Num();
+
+//	VertexAttributeStreams* Ret = new UVertexAttributeStream();
+//	VertexAttributeStreams.Add(Ret);
+	
+	UVertexAttributeStream* VertexAttributeStream = NewObject<UVertexAttributeStream>(this); 
+
+	VertexAttributeStream->Usage = InUsage;
+
+	VertexAttributeStreams.Add(VertexAttributeStream);
+//	VertexAttributeStreams.Add(new UVertexAttributeStream(InUsage)); todo
+
+	return VertexAttributeStreams[Num];
+}
+
+UVertexAttributeStream* USubDSurface::FindStreamByUsage(FName InUsage)
+{
+	for(UVertexAttributeStream* it : VertexAttributeStreams)
+	{
+		if(it->Usage == InUsage)
+		{
+			return it;
+		}
+	}
+
+	return 0;
+}
+
+// -------------------------------------------------------------------
+
+ASubDSurfaceActor::ASubDSurfaceActor(const FObjectInitializer& ObjectInitializer)
+	: Super(ObjectInitializer)
+{
+	DisplayMeshComponent = ObjectInitializer.CreateDefaultSubobject<UStaticMeshComponent>(this, TEXT("DisplayMeshComponent"));
+	RootComponent = DisplayMeshComponent;
+
+	SubDSurface = ObjectInitializer.CreateDefaultSubobject<USubDSurfaceComponent>(this, TEXT("NewSubDSurfaceComponent"));
+	SubDSurface->SetDisplayMeshComponent(DisplayMeshComponent);
+
+	SubDSurface->AttachToComponent(DisplayMeshComponent, FAttachmentTransformRules::KeepRelativeTransform);
+//	bool USceneComponent::AttachToComponent(USceneComponent* Parent, const FAttachmentTransformRules& AttachmentRules, FName SocketName)
+//	SubDSurface->AttachParent = DisplayMeshComponent;
+
+//	ProcMeshComponent = ObjectInitializer.CreateDefaultSubobject<UProceduralMeshComponent>(this, TEXT("NewProceduralMeshComponent"));
+
+#if WITH_EDITORONLY_DATA
+	SpriteComponent = ObjectInitializer.CreateEditorOnlyDefaultSubobject<UBillboardComponent>(this, TEXT("Sprite"));
+
+	if (!IsRunningCommandlet() && (SpriteComponent != nullptr))
+	{
+		// Structure to hold one-time initialization
+		struct FConstructorStatics
+		{
+			ConstructorHelpers::FObjectFinderOptional<UTexture2D> SubDSurfaceTexture;
+			FConstructorStatics()
+				: SubDSurfaceTexture(TEXT("/Engine/EditorResources/S_Terrain.S_Terrain"))
+			{
+			}
+		};
+		static FConstructorStatics ConstructorStatics;
+
+		SpriteComponent->Sprite = ConstructorStatics.SubDSurfaceTexture.Get();
+		SpriteComponent->RelativeScale3D = FVector(0.5f, 0.5f, 0.5f);
+		
+		SpriteComponent->AttachToComponent(SubDSurface, FAttachmentTransformRules::KeepRelativeTransform);
+		SpriteComponent->bIsScreenSizeScaled = true;
+		SpriteComponent->bAbsoluteScale = true;
+		SpriteComponent->bReceivesDecals = false;
+	}
+#endif
+}
+
+
+// for OpenSubDiv to interpolate positions
+struct FVertexPosition
+{
+	// Minimal required interface ----------------------
+	FVertexPosition() { }
+	
+	FVertexPosition(int /*i*/) { }
+
+/*
+	FVertexPosition(FVertexPosition const & src)
+	{
+		_position[0] = src._position[0];
+		_position[1] = src._position[1];
+		_position[1] = src._position[1];
+	}
+*/
+
+	void Clear( void * =0 )
+	{
+		point = FVector(0, 0, 0);
+	}
+
+	void AddWithWeight(FVertexPosition const & src, float weight)
+	{
+		point += weight * src.point;
+	}
+
+	void AddVaryingWithWeight(FVertexPosition const &, float) { }
+
+	FVector point;
+};
+
+// for OpenSubDiv to interpolate UV
+struct FVarVertexUV
+{
+    // Minimal required interface ----------------------
+    void Clear()
+	{
+        UV = FVector2D(0.0f, 0.0f);
+    }
+
+    void AddWithWeight(FVarVertexUV const & src, float weight)
+	{
+        UV += weight * src.UV;
+    }
+
+    // Basic 'uv' layout channel
+    FVector2D UV;
+};
+
+// for OpenSubDiv to compute the tangent space (similar to here: http://graphics.pixar.com/opensubdiv/docs/far_tutorial_6.html)
+struct FLimitFrame
+{
+	void Clear( void * = 0 )
+	{
+		point = FVector(0, 0, 0);
+        deriv1 = FVector(0, 0, 0);
+        deriv2 = FVector(0, 0, 0);
+    }
+
+    void AddWithWeight(FVertexPosition const & src, float weight, float d1Weight, float d2Weight)
+	{
+        point += weight * src.point;
+        deriv1 += d1Weight * src.point;
+        deriv2 += d2Weight * src.point;
+    }
+
+	FVector point;
+    FVector deriv1;
+	FVector deriv2;
+};
+
+// ------------------------------------------------------
+
+USubDSurfaceComponent::USubDSurfaceComponent(const FObjectInitializer& ObjectInitializer)
+	: Super(ObjectInitializer)
+{
+	// Structure to hold one-time initialization
+/*
+	struct FConstructorStatics
+	{
+		ConstructorHelpers::FObjectFinderOptional<UMaterial> TextMaterial;
+		FConstructorStatics()
+			: TextMaterial(TEXT("/Engine/EngineMaterials/DefaultTextMaterialOpaque"))
+		{
+		}
+	};
+	static FConstructorStatics ConstructorStatics;
+
+	TextMaterial = ConstructorStatics.TextMaterial.Get();
+*/
+
+	SetCollisionProfileName(UCollisionProfile::NoCollision_ProfileName);
+	bGenerateOverlapEvents = false;
+}
+
+
+#if WITH_EDITOR
+void USubDSurfaceComponent::PostEditChangeProperty(FPropertyChangedEvent& PropertyChangedEvent)
+{
+	UProperty* PropertyThatChanged = PropertyChangedEvent.MemberProperty;
+	const FString PropertyName = PropertyThatChanged ? PropertyThatChanged->GetName() : TEXT("");
+
+	if( PropertyName == GET_MEMBER_NAME_STRING_CHECKED(USubDSurfaceComponent, DebugLevel) || 
+		PropertyName == GET_MEMBER_NAME_STRING_CHECKED(USubDSurfaceComponent, Mesh))
+	{
+		SetMesh(Mesh);
+	}
+
+	Super::PostEditChangeProperty(PropertyChangedEvent);
+}
+#endif // WITH_EDITOR
+
+
+bool USubDSurfaceComponent::SetMesh(USubDSurface* NewMesh)
+{
+	// Do nothing if we are already using the supplied static mesh
+/*
+	// good for blueprint code, not good for editor update
+	if(NewMesh == Mesh)
+	{
+		return false;
+	}
+*/
+
+/*
+	// Don't allow changing static meshes if "static" and registered
+	AActor* Owner = GetOwner();
+	if(Mobility == EComponentMobility::Static && IsRegistered() && Owner != NULL)
+	{
+		FMessageLog("PIE").Warning(FText::Format(LOCTEXT("SetMeshOnStatic", "Calling SetStaticMesh on '{0}' but Mobility is Static."), 
+			FText::FromString(GetPathName(this))));
+		return false;
+	}
+*/
+	double StartTime = FPlatformTime::Seconds();
+
+	Mesh = NewMesh;
+
+	// Need to send this to render thread at some point
+//	MarkRenderStateDirty();
+
+	RecreateMeshData();
+
+	// Update physics representation right away
+	RecreatePhysicsState();
+
+	// Notify the streaming system. Don't use Update(), because this may be the first time the mesh has been set
+	// and the component may have to be added to the streaming system for the first time.
+	IStreamingManager::Get().NotifyPrimitiveAttached( this, DPT_Spawned );
+
+	// Since we have new mesh, we need to update bounds
+	UpdateBounds();
+
+	UE_LOG(LogEngine, Log, TEXT("SubDSurface update: %5.2f seconds"), FPlatformTime::Seconds() - StartTime);
+
+	return true;
+}
+
+#if USE_OPENSUBDIV
+typedef OpenSubdiv::HbrMesh<FVertexPosition>      Hmesh;
+typedef OpenSubdiv::HbrFace<FVertexPosition>      Hface;
+typedef OpenSubdiv::HbrVertex<FVertexPosition>    Hvertex;
+typedef OpenSubdiv::HbrHalfedge<FVertexPosition>  Hhalfedge;
+typedef Far::TopologyDescriptor Descriptor;
+
+
+void TestNonManiFold(USubDSurface* Mesh, const Descriptor& desc)
+{
+	FVector* PosData;
+	{
+		UVertexAttributeStream* Stream = Mesh->FindStreamByUsage(FName(TEXT("Position")));
+
+		check(Stream);
+
+		uint32 VertexCount = 0;
+		PosData = Stream->GetFVector(VertexCount);
+	}
+
+
+	OpenSubdiv::HbrCatmarkSubdivision<FVertexPosition> * catmark = new OpenSubdiv::HbrCatmarkSubdivision<FVertexPosition>();
+
+    Hmesh * hmesh = new Hmesh(catmark);
+
+    for (int32 i = 0, count = Mesh->GetVertexCount(); i < count; ++i)
+	{
+		FVertexPosition v;
+
+		v.point = PosData[i];
+
+        hmesh->NewVertex(i, v);
+    }
+
+    // Create the topology
+    const int32* fv = desc.vertIndicesPerFace;
+
+    for (int32 i = 0; i < desc.numFaces; ++i)
+	{
+        int nv = desc.numVertsPerFace[i];
+
+        bool bValid = true;
+
+        for(int j = 0; j < nv; j++)
+		{
+
+            Hvertex const * origin      = hmesh->GetVertex(fv[j]),
+                          * destination = hmesh->GetVertex(fv[(j+1)%nv]);
+
+			check(destination);
+
+            Hhalfedge const * opposite = destination->GetEdge(origin);
+
+            // Make sure that the vertices exist in the mesh
+            if (origin==NULL || destination==NULL)
+			{
+                printf(" An edge was specified that connected a nonexistent vertex\n");
+                bValid = false;
+                break;
+            }
+
+            // Check for a degenerate edge
+            if (origin == destination)
+			{
+                printf(" An edge was specified that connected a vertex to itself\n");
+                bValid = false;
+                break;
+            }
+
+            // Check that no more than 2 faces are adjacent to the edge
+            if (opposite && opposite->GetOpposite() )
+			{
+                printf(" A non-manifold edge incident to more than 2 faces was found\n");
+                bValid = false;
+                break;
+            }
+
+            // Check that the edge is unique and oriented properly
+            if (origin->GetEdge(destination))
+			{
+                printf(" An edge connecting two vertices was specified more than once."
+                       " It's likely that an incident face was flipped\n");
+                bValid = false;
+                break;
+            }
+        }
+
+        if (bValid)
+		{
+            hmesh->NewFace(nv, fv, 0);
+        }
+		else
+		{
+            printf(" Skipped face %d\n", i);
+        }
+
+        fv += nv;
+    }
+
+    hmesh->SetInterpolateBoundaryMethod(Hmesh::k_InterpolateBoundaryEdgeOnly);
+
+    hmesh->Finish();
+
+    printf("Created a fan with %d faces and %d vertices.\n", hmesh->GetNumFaces(), hmesh->GetNumVertices());
+
+    delete hmesh;
+    delete catmark;
+}
+#endif // USE_OPENSUBDIV
+
+
+// pyramid geometry from catmark_pyramid_crease0.h
+static int const g_nverts = 5;
+static float const g_verts[24] = { 0.0f,  0.0f, 2.0f,
+                                   0.0f, -2.0f, 0.0f,
+                                   2.0f,  0.0f, 0.0f,
+                                   0.0f,  2.0f, 0.0f,
+                                  -2.0f,  0.0f, 0.0f, };
+
+
+static int const g_vertsperface[5] = { 3, 3, 3, 3, 4 };
+
+static int const g_nfaces = 5;
+static int const g_faceverts[16] = { 0, 1, 2,
+                                     0, 2, 3,
+                                     0, 3, 4,
+                                     0, 4, 1,
+                                     4, 3, 2, 1 };
+
+static int const g_ncreases = 4;
+static int const g_creaseverts[8] = { 4, 3, 3, 2, 2, 1, 1, 4 };
+static float const g_creaseweights[4] = { 3.0f, 3.0f, 3.0f, 3.0f };
+
+
+// compute the 3x3 tangent space (3 vectors) from 2 given delta position vectors and the UV delta change on those vectors
+// The interface is more general to make this usable for more than triangles.
+// for a triangle you would have to compute the input like this:
+//   InDeltaPos[0]=InPos[1]-InPos[0], InDeltaPos[1]=InPos[2]-InPos[0]
+//   InDeltaUV[0]=InUV[1]-InUV[0], InDeltaUV[1]=InUV[2]-InUV[0]
+// degenerated cases are not handled yet
+void ComputeTangents(FVector OutTangentXYZ[3], FVector InDeltaPos[2], FVector2D InDeltaUV[2])
+{
+	// originally from: Lengyel, Eric. “Computing Tangent Space Basis Vectors for an Arbitrary Mesh”. Terathon Software 3D Graphics Library, 2001. http://www.terathon.com/code/tangent.html
+
+	float x1 = InDeltaPos[0].X;
+    float x2 = InDeltaPos[1].X;
+    float y1 = InDeltaPos[0].Y;
+    float y2 = InDeltaPos[1].Y;
+    float z1 = InDeltaPos[0].Z;
+    float z2 = InDeltaPos[1].Z;
+        
+    float s1 = InDeltaUV[0].X;
+    float s2 = InDeltaUV[1].X;
+    float t1 = InDeltaUV[0].Y;
+    float t2 = InDeltaUV[1].Y;
+
+	float deter = s1 * t2 - s2 * t1;
+        
+	// no need as we normalize anyway
+//    float r = 1.0f / deter;
+    float r = 1.0f;
+
+    OutTangentXYZ[0] = FVector(t2 * x1 - t1 * x2, t2 * y1 - t1 * y2, t2 * z1 - t1 * z2) * r;
+    OutTangentXYZ[1] = FVector(s1 * x2 - s2 * x1, s1 * y2 - s2 * y1, s1 * z2 - s2 * z1) * r;
+
+	OutTangentXYZ[0] = OutTangentXYZ[0].GetSafeNormal();
+	OutTangentXYZ[1] = OutTangentXYZ[1].GetSafeNormal();
+
+    OutTangentXYZ[2] = (OutTangentXYZ[0] ^ OutTangentXYZ[1]).GetSafeNormal();
+
+	if(deter < 0)
+	{
+		// mirrored UV changes the handiness of the matrix
+		OutTangentXYZ[2] = -OutTangentXYZ[2];
+	}
+}
+
+
+void USubDSurfaceComponent::RecreateMeshData()
+{
+	if(!Mesh)
+	{
+		if(DisplayMeshComponent)
+		{
+			DisplayMeshComponent->SetStaticMesh(0);
+			DisplayMeshComponent->MarkRenderStateDirty();
+		}
+		return;
+	}
+
+//	bool bUV = false;			/// todo expose
+	bool bUV = true;			/// todo expose
+
+	double StartTime = FPlatformTime::Seconds();
+
+	// number of n-gons
+	uint32 FaceCount = Mesh->VertexCountPerFace.Num();
+	
+	check(FaceCount);
+
+#if USE_OPENSUBDIV
+	// ----
+	Sdc::SchemeType type = OpenSubdiv::Sdc::SCHEME_CATMARK;
+
+	Sdc::Options options;
+	options.SetVtxBoundaryInterpolation(Sdc::Options::VTX_BOUNDARY_EDGE_ONLY);
+//    options.SetFVarLinearInterpolation(Sdc::Options::FVAR_LINEAR_NONE);				// affects how UVs are interpolated
+
+	Descriptor desc;
+	desc.numVertices  = Mesh->GetVertexCount();
+	desc.numFaces     = FaceCount;
+	desc.numVertsPerFace = (const int32*)&Mesh->VertexCountPerFace[0];
+	desc.vertIndicesPerFace = (const Far::Index *)Mesh->IndicesPerFace.GetData();
+
+/*
+	// hack
+    Descriptor desc;
+    desc.numVertices = g_nverts;
+    desc.numFaces = g_nfaces;
+    desc.numVertsPerFace = g_vertsperface;
+    desc.vertIndicesPerFace = g_faceverts;
+	bUV = false;
+	// hack end
+*/
+/*	// debug
+
+	TArray<uint32> TestValency;
+	TestValency.AddZeroed(desc.numVertices);
+	{
+		int32* IndexCountBuffer = (int32*)desc.numVertsPerFace;
+		int32* IndexBuffer = (int32*)desc.vertIndicesPerFace;
+		for (uint32 Face = 0; Face < (uint32)desc.numFaces; ++Face)
+		{
+			uint32 CornerCount = *IndexCountBuffer++;
+			for (uint32 Corner = 0; Corner < CornerCount; ++Corner)
+			{
+				uint32 VertexIndex = *IndexBuffer++;
+
+				TestValency[VertexIndex]++;
+			}
+		}
+	}
+*/
+
+	TArray<uint32> PerCornerIndexBuffer;
+    Descriptor::FVarChannel channels[1];
+    int32 channelUV = 0;
+
+	if(bUV)
+	{
+		PerCornerIndexBuffer.Empty(Mesh->IndicesPerFace.Num());
+
+		uint32 UniqueIndex = 0;
+
+		for(uint32 Face = 0; Face < FaceCount; ++Face)
+		{
+			uint32 Corners = Mesh->VertexCountPerFace[Face];
+
+			for(uint32 i = 0; i < Corners; ++i)
+			{
+				PerCornerIndexBuffer.Add(UniqueIndex++);
+			}
+		}
+
+		// Add the channel topology to the main descriptor
+
+		// Create a face-varying channel descriptor
+		channels[channelUV].numValues = PerCornerIndexBuffer.Num();			// for now we map the UV by the face
+		channels[channelUV].valueIndices = (Far::Index const *)PerCornerIndexBuffer.GetData();
+
+		desc.numFVarChannels = 1;
+		desc.fvarChannels = channels;
+	}
+
+	TestNonManiFold(Mesh, desc);
+
+	int32 maxlevel = FMath::Clamp((int32)DebugLevel, 1, 6);
+
+	// hack
+//	maxlevel = 3;
+
+	// Instantiate a FarTopologyRefiner from the descriptor
+	Far::TopologyRefiner * refiner = Far::TopologyRefinerFactory<Descriptor>::Create(desc,
+                                         Far::TopologyRefinerFactory<Descriptor>::Options(type, options));
+
+	{
+		// Uniformly refine the topology up to 'maxlevel'
+		Far::TopologyRefiner::UniformOptions refineOptions(maxlevel);
+
+		// note: fullTopologyInLastLevel must be true to work with face-varying data
+		refineOptions.fullTopologyInLastLevel = true;
+
+		// simpler but we don't get good normals
+		refiner->RefineUniform(refineOptions);
+	}
+
+	// cage mesh vertex count
+    int CoarseVerts = desc.numVertices;
+	// most refined level vertex count
+	int FineVerts   = refiner->GetLevel(maxlevel).GetNumVertices();
+	// vertex count of all levels [cage mesh, intermediate levels, most refined level]
+    int TotalVerts  = refiner->GetNumVerticesTotal();
+//    int TempVerts = TotalVerts - CoarseVerts - FineVerts;
+
+	// intermediate position for SubDiv process, does not have the limit positions and not the last level
+	// UVs for [cage mesh, intermediate levels, but not most refined level]
+	TArray<FVertexPosition> TempPosBuffer;
+
+//	TempPosBuffer.AddUninitialized(refiner->GetNumVerticesTotal());
+	TempPosBuffer.AddUninitialized(TotalVerts - FineVerts);
+
+	FVertexPosition* verts = (FVertexPosition*)TempPosBuffer.GetData();
+
+	// setup position
+	{
+		FVector* PosData = 0;
+		{
+			UVertexAttributeStream* Stream = Mesh->FindStreamByUsage(FName(TEXT("Position")));
+
+			check(Stream);
+
+			uint32 VertexCount = 0;
+			PosData = Stream->GetFVector(VertexCount);
+
+			check(VertexCount == desc.numVertices);
+		}
+
+		// Initialize coarse mesh positions
+		for (uint32 i = 0, Count = (uint32)desc.numVertices; i < Count; ++i)
+		{
+			FVector Pos = PosData[i];
+
+			verts[i].point = Pos;
+		}
+	}
+
+	// Interpolate vertex primvar data
+    Far::PrimvarRefiner primvarRefiner(*refiner);
+
+	// position
+	FVertexPosition* src = verts;
+
+	// interpolate position [not the cage mesh, intermediate levels, most refined level]
+	for(int32 level = 1; level < maxlevel; ++level)
+	{
+		FVertexPosition* dst = src + refiner->GetLevel(level - 1).GetNumVertices();
+
+		primvarRefiner.Interpolate(level, src, dst);
+
+		src = dst;
+	}
+
+	// SubD Position for [most refined level]
+	TArray<FVertexPosition> finePosBuffer;
+	finePosBuffer.AddUninitialized(FineVerts);
+
+	// Limit Position for [most refined level]
+	TArray<FVertexPosition> fineLimitPos;
+	fineLimitPos.AddUninitialized(FineVerts);
+
+	// Patch Derivative in X for [most refined level]
+    TArray<FVertexPosition> fineDu;
+	fineDu.AddUninitialized(FineVerts);
+
+	// Patch Derivative in Y for [most refined level]
+    TArray<FVertexPosition> fineDv;
+	fineDv.AddUninitialized(FineVerts);
+
+	// test
+	// Limit Position for [most refined level]
+	TArray<FVarVertexUV> fineLimitUV;
+	TArray<FVarVertexUV> fineLimitUVDu;
+	TArray<FVarVertexUV> fineLimitUVDv;
+	fineLimitUV.AddUninitialized(refiner->GetLevel(maxlevel).GetNumFVarValues());
+	fineLimitUVDu.AddUninitialized(refiner->GetLevel(maxlevel).GetNumFVarValues());
+	fineLimitUVDv.AddUninitialized(refiner->GetLevel(maxlevel).GetNumFVarValues());
+
+	// interpolate SubD position [most refined level]
+    primvarRefiner.Interpolate(maxlevel, src, finePosBuffer);
+	// interpolate limit position [most refined level]
+    primvarRefiner.Limit(finePosBuffer, fineLimitPos, fineDu, fineDv);
+
+
+	// UVs for [cage mesh, intermediate levels, most refined level]
+	TArray<FVarVertexUV> fvBufferUV;
+	if(bUV)
+	{
+		fvBufferUV.AddUninitialized(refiner->GetNumFVarValuesTotal(channelUV));
+		FVarVertexUV* fvVertsUV = fvBufferUV.GetData();
+
+		uint32 CageUVCount = 0;
+
+		// setup UV
+		{
+			FVector2D* UVData = 0;
+			{
+				UVertexAttributeStream* Stream = Mesh->FindStreamByUsage(FName(TEXT("UV0")));
+
+				check(Stream);
+
+				UVData = Stream->GetFVector2D(CageUVCount);
+			}
+
+			for(uint32 i = 0; i < CageUVCount; ++i)
+			{
+				fvVertsUV[i].UV = UVData[i];
+			}
+		}
+			
+		{
+			FVarVertexUV* srcUV = &fvVertsUV[0];
+
+			uint32 Level0NumVertices = refiner->GetLevel(0).GetNumFVarValues();
+
+			check(CageUVCount == Level0NumVertices);
+			
+			// interpolate UV [not the cage mesh, intermediate levels, most refined level]
+			for (int32 level = 1; level <= maxlevel; ++level)
+			{
+				FVarVertexUV* dstUV = srcUV + refiner->GetLevel(level - 1).GetNumFVarValues();
+
+				// FaceVarying: see http://graphics.pixar.com/opensubdiv/docs/subdivision_surfaces.html#face-varying-interpolation-rules
+				primvarRefiner.InterpolateFaceVarying(level, srcUV, dstUV);
+
+				srcUV = dstUV;
+			}
+		}
+
+		auto fineLimitUVPtr =&fineLimitUV[0];
+
+		// test
+		// interpolate limit position [most refined level]
+//		primvarRefiner.Limit(&fvBufferUV[refiner->GetNumFVarValuesTotal(channelUV) - refiner->GetLevel(maxlevel).GetNumFVarValues()], fineLimitUV, fineLimitUVDu, fineLimitUVDv);
+		primvarRefiner.LimitFaceVarying(&fvBufferUV[refiner->GetNumFVarValuesTotal(channelUV) - refiner->GetLevel(maxlevel).GetNumFVarValues()], fineLimitUVPtr, channelUV);
+	}
+
+	
+    // Create a Far::PtexIndices to help find indices of ptex faces.
+
+    Far::PtexIndices ptexIndices(*refiner);
+	FRawMesh RawMesh;
+
+	// Output of the highest level refined -----------
+    {
+        Far::TopologyLevel const & refLastLevel = refiner->GetLevel(maxlevel);
+
+        int32 nverts = refLastLevel.GetNumVertices();
+        int32 nfaces = refLastLevel.GetNumFaces();
+
+		uint32 TriangleCount = PerCornerIndexBuffer.Num();
+
+		// reserve to avoid memory reallocations
+		RawMesh.VertexPositions.Empty(nverts);
+		RawMesh.WedgeIndices.Empty(TriangleCount * 3);
+		RawMesh.FaceMaterialIndices.Empty(TriangleCount);
+		RawMesh.FaceSmoothingMasks.Empty(TriangleCount);
+		
+		int32 firstOfLastUvs = 0;
+
+		if(bUV)
+		{
+			int32 nuvs = refLastLevel.GetNumFVarValues(channelUV);
+			RawMesh.WedgeTexCoords[0].Empty(TriangleCount * 3);		// our FRawMesh wants it like this
+			firstOfLastUvs = refiner->GetNumFVarValuesTotal(channelUV) - nuvs;
+		}
+
+		int32 firstOfLastVerts = refiner->GetNumVerticesTotal() - nverts;
+
+		for(int32 vert = 0; vert < nverts; ++vert)
+		{
+			FVector pos = fineLimitPos[vert].point;
+
+			RawMesh.VertexPositions.Add(pos);
+		}
+	
+		{
+			uint32 StartIndexForHighestLevel = 0; //- refiner->GetNumVertices(maxlevel);
+			FVarVertexUV* fvVertsUV = fvBufferUV.GetData();
+
+			for (int32 face = 0; face < nfaces; ++face)
+			{
+				Far::ConstIndexArray fverts = refLastLevel.GetFaceVertices(face);
+
+				// all refined Catmark faces should be quads
+				uint32 CornerCount = fverts.size();
+
+				// one n-gon
+				for(uint32 Corner = 0; Corner < CornerCount; ++Corner)
+				{
+					// one triangle
+					RawMesh.FaceMaterialIndices.Add(0);
+					RawMesh.FaceSmoothingMasks.Add(0xffffffff);
+
+					uint32 CornerIds[] = { 0, (Corner + 1) % CornerCount, (Corner + 2) % CornerCount };
+
+					// convert to triangles by spawning 3 vertices
+					for(uint32 i = 0 ; i < 3; ++i)
+					{
+						uint32 LocalCornerId = CornerIds[i];
+
+						RawMesh.WedgeIndices.Add(fverts[LocalCornerId] + StartIndexForHighestLevel);
+					}
+
+					if(!bUV)
+					{
+						for(uint32 i = 0 ; i < 3; ++i)
+						{
+							RawMesh.WedgeTexCoords[0].Add(FVector2D(0,0));
+						}
+					}
+				}
+			}
+			if (bUV)
+			{
+				for(int32 face = 0; face < nfaces; ++face)
+				{
+					Far::ConstIndexArray fverts = refLastLevel.GetFaceVertices(face);
+					Far::ConstIndexArray fuvs = refLastLevel.GetFaceFVarValues(face, channelUV);
+
+					// all refined Catmark faces should be quads
+					uint32 CornerCount = fverts.size();
+
+					// one n-gon
+					for (uint32 Corner = 0; Corner < CornerCount; ++Corner)
+					{
+						uint32 CornerIds[] = { 0, (Corner + 1) % CornerCount, (Corner + 2) % CornerCount };
+						
+						// convert to triangles by spawning 3 vertices
+						for(uint32 i = 0 ; i < 3; ++i)
+						{
+							uint32 LocalCornerId = CornerIds[i];
+							uint32 UVIndex = fuvs[LocalCornerId];
+							uint32 VertexIndex = fverts[LocalCornerId];
+
+//							RawMesh.WedgeTexCoords[0].Add(fvVertsUV[firstOfLastUvs + UVIndex].UV);
+							RawMesh.WedgeTexCoords[0].Add(fineLimitUV[UVIndex].UV);
+
+							FVector TangentXYZ[3];
+							FVector DeltaPos[2] = { fineDu[VertexIndex].point, fineDv[VertexIndex].point };
+							FVector2D DeltaUV[2] = { fineLimitUVDu[UVIndex].UV, fineLimitUVDv[UVIndex].UV };
+							
+							ComputeTangents(TangentXYZ, DeltaPos, DeltaUV);
+
+/*
+							FVector TangentX = (fineDu[VertexIndex].point).GetSafeNormal();
+							FVector TangentY = (fineDv[VertexIndex].point).GetSafeNormal();
+							FVector TangentZ = (TangentY ^ TangentX).GetSafeNormal();
+
+							RawMesh.WedgeTangentX.Add(TangentX);
+							RawMesh.WedgeTangentY.Add(TangentY);
+							RawMesh.WedgeTangentZ.Add(TangentZ);
+*/
+							RawMesh.WedgeTangentX.Add(TangentXYZ[0]);
+							RawMesh.WedgeTangentY.Add(TangentXYZ[1]);
+							RawMesh.WedgeTangentZ.Add(TangentXYZ[2]);
+						}
+					}
+				}
+/*
+
+				for (int32 vert = 0; vert < nverts; ++vert)
+				{
+					FVector TangentX = (fineDu[vert].point).GetSafeNormal();
+					FVector TangentY = (fineDv[vert].point).GetSafeNormal();
+					FVector TangentZ = (TangentX ^ TangentY).GetSafeNormal();
+
+					RawMesh.WedgeTangentX.Add(TangentX);
+					RawMesh.WedgeTangentY.Add(TangentY);
+					RawMesh.WedgeTangentZ.Add(TangentZ);
+				}*/
+			}
+		}
+    }
+
+	check(RawMesh.IsValid());
+
+	UE_LOG(LogEngine, Log, TEXT("SubDSurface OpenSubDiv: %5.2f seconds"), FPlatformTime::Seconds() - StartTime);
+
+	UStaticMesh* GeneratedMesh = 0;
+
+	// see UStaticMesh* FAbcImporter::ImportSingleAsStaticMesh(const int32 MeshIndex, UObject* InParent, EObjectFlags Flags)
+
+	// Setup static mesh instance
+	GeneratedMesh = NewObject<UStaticMesh>(this);
+
+	// Generate a new lighting GUID (so its unique)
+	GeneratedMesh->LightingGuid = FGuid::NewGuid();
+
+	// Set it to use textured lightmaps. Note that Build Lighting will do the error-checking (texcoord index exists for all LODs, etc).
+	GeneratedMesh->LightMapResolution = 64;
+	GeneratedMesh->LightMapCoordinateIndex = 1;
+
+	// Material setup, since there isn't much material information in the Alembic file, 
+	UMaterial* DefaultMaterial = UMaterial::GetDefaultMaterial(MD_Surface);
+	check(DefaultMaterial);
+
+	// Material list
+//	GeneratedMesh->Materials.Empty();
+	// If there were FaceSets available in the Alembic file use the number of unique face sets as num material entries, otherwise default to one material for the whole mesh
+//	uint32 NumFaceSets = GetNumFaceSetsFromSchema(Schema);
+//	uint32 NumMaterials = (NumFaceSets) ? NumFaceSets : 1;
+//	for (uint32 MaterialIndex = 0; MaterialIndex < NumMaterials; ++MaterialIndex )
+	{
+		GeneratedMesh->Materials.Add(DefaultMaterial);
+	}	
+
+	// Add the first LOD ( QQQ change this )	
+	new(GeneratedMesh->SourceModels) FStaticMeshSourceModel();
+
+	// Get the first LOD for filling it up with geometry (QQ)
+	FStaticMeshSourceModel& SrcModel = GeneratedMesh->SourceModels[0];
+	// Set build settings for the static mesh (QQ want this from import options menu)
+	SrcModel.BuildSettings.bRecomputeNormals = false;
+	SrcModel.BuildSettings.bRecomputeTangents = true;
+	// Generate Lightmaps uvs (no support for importing right now)
+	
+//	SrcModel.BuildSettings.bGenerateLightmapUVs = true;		// reconsider later
+	SrcModel.BuildSettings.bGenerateLightmapUVs = false;
+	SrcModel.BuildSettings.bBuildAdjacencyBuffer = false;	// reconsider later
+	SrcModel.BuildSettings.bUseMikkTSpace = false;			// MikkT is much slower
+	SrcModel.BuildSettings.bRemoveDegenerates = false;		// makes FindOverlappingCorners a tiny bit faster
+
+	SrcModel.BuildSettings.bRecomputeNormals = false;
+	SrcModel.BuildSettings.bRecomputeTangents = false;
+
+	// Set lightmap UV index to 1 since we currently only import one set of UVs
+	SrcModel.BuildSettings.DstLightmapIndex = 1;
+
+	// Store the raw mesh within the RawMeshBulkData
+	SrcModel.RawMeshBulkData->SaveRawMesh(RawMesh);
+	
+	// Build the static mesh (using the build setting etc.) this generates correct tangents using the extracting smoothing group along with the imported Normals data
+	GeneratedMesh->Build(false);
+	GeneratedMesh->MarkPackageDirty();
+
+	// QQ No collision generation for now
+	GeneratedMesh->CreateBodySetup();
+
+	check(DisplayMeshComponent);
+
+	DisplayMeshComponent->SetStaticMesh(GeneratedMesh);
+	DisplayMeshComponent->MarkRenderStateDirty();
+#endif // USE_OPENSUBDIV
+}
+
+void USubDSurfaceComponent::GetUsedMaterials( TArray<UMaterialInterface*>& OutMaterials ) const
+{
+/* todo
+	if( StaticMesh && StaticMesh->RenderData )
+	{
+		for (int32 LODIndex = 0; LODIndex < StaticMesh->RenderData->LODResources.Num(); LODIndex++)
+		{
+			FStaticMeshLODResources& LODResources = StaticMesh->RenderData->LODResources[LODIndex];
+			for (int32 SectionIndex = 0; SectionIndex < LODResources.Sections.Num(); SectionIndex++)
+			{
+				// Get the material for each element at the current lod index
+				OutMaterials.AddUnique(GetMaterial(LODResources.Sections[SectionIndex].MaterialIndex));
+			}
+		}
+	}
+*/
+}
+
+int32 USubDSurfaceComponent::GetNumMaterials() const
+{
+	// todo
+	return 0;
+}
+
+void USubDSurfaceComponent::SetMaterial(int32 ElementIndex, UMaterialInterface* InMaterial)
+{
+	// todo
+}
+
+UMaterialInterface* USubDSurfaceComponent::GetMaterial(int32 ElementIndex) const
+{
+	// todo
+	return NULL;
+}
+
+bool USubDSurfaceComponent::ShouldRecreateProxyOnUpdateTransform() const
+{
+	return true;
+}
+
+FBoxSphereBounds USubDSurfaceComponent::CalcBounds(const FTransform& LocalToWorld) const
+{
+	// todo
+
+	FBox LocalBox(FVector(-100, -100, -100), FVector(100, 100, 100));
+
+	FBoxSphereBounds Ret(LocalBox.TransformBy(LocalToWorld));
+
+	Ret.BoxExtent *= BoundsScale;
+	Ret.SphereRadius *= BoundsScale;
+
+	return Ret;
+}
+
+void USubDSurfaceComponent::SetDisplayMeshComponent(UStaticMeshComponent* InDisplayMeshComponent)
+{
+	DisplayMeshComponent = InDisplayMeshComponent;
+}
diff --git a/Engine/Source/Runtime/Engine/Private/TileRendering.cpp b/Engine/Source/Runtime/Engine/Private/TileRendering.cpp
--- a/Engine/Source/Runtime/Engine/Private/TileRendering.cpp
+++ b/Engine/Source/Runtime/Engine/Private/TileRendering.cpp
@@ -287,7 +287,7 @@ bool FCanvasTileRendererItem::Render_GameThread(const FCanvas* Canvas)
 	const FRenderTarget* CanvasRenderTarget = Canvas->GetRenderTarget();
 	FSceneViewFamily* ViewFamily = new FSceneViewFamily(FSceneViewFamily::ConstructionValues(
 		CanvasRenderTarget,
-		NULL,
+		Canvas->GetScene(),
 		FEngineShowFlags(ESFIM_Game))
 		.SetWorldTimes(CurrentWorldTime, DeltaWorldTime, CurrentRealTime)
 		.SetGammaCorrection(CanvasRenderTarget->GetDisplayGamma()));
diff --git a/Engine/Source/Runtime/Engine/Private/TriangleRendering.cpp b/Engine/Source/Runtime/Engine/Private/TriangleRendering.cpp
--- a/Engine/Source/Runtime/Engine/Private/TriangleRendering.cpp
+++ b/Engine/Source/Runtime/Engine/Private/TriangleRendering.cpp
@@ -287,7 +287,7 @@ bool FCanvasTriangleRendererItem::Render_GameThread(const FCanvas* Canvas)
 	const FRenderTarget* CanvasRenderTarget = Canvas->GetRenderTarget();
 	FSceneViewFamily* ViewFamily = new FSceneViewFamily(FSceneViewFamily::ConstructionValues(
 		CanvasRenderTarget,
-		nullptr,
+		Canvas->GetScene(),
 		FEngineShowFlags(ESFIM_Game))
 		.SetWorldTimes(CurrentWorldTime, DeltaWorldTime, CurrentRealTime)
 		.SetGammaCorrection(CanvasRenderTarget->GetDisplayGamma()));
diff --git a/Engine/Source/Runtime/Engine/Private/UserInterface/Canvas.cpp b/Engine/Source/Runtime/Engine/Private/UserInterface/Canvas.cpp
--- a/Engine/Source/Runtime/Engine/Private/UserInterface/Canvas.cpp
+++ b/Engine/Source/Runtime/Engine/Private/UserInterface/Canvas.cpp
@@ -208,10 +208,11 @@ int32 FCanvasWordWrapper::FindEndOfLastWholeGraphemeCluster(const int32 InStartI
 	return BreakIndex;
 }
 
-FCanvas::FCanvas(FRenderTarget* InRenderTarget, FHitProxyConsumer* InHitProxyConsumer, UWorld* InWorld, ERHIFeatureLevel::Type InFeatureLevel)
+FCanvas::FCanvas(FRenderTarget* InRenderTarget, FHitProxyConsumer* InHitProxyConsumer, UWorld* InWorld, ERHIFeatureLevel::Type InFeatureLevel, ECanvasDrawMode InDrawMode)
 :	ViewRect(0,0,0,0)
 ,	RenderTarget(InRenderTarget)
 ,	HitProxyConsumer(InHitProxyConsumer)
+,	Scene(InWorld ? InWorld->Scene : NULL)
 ,	AllowedModes(0xFFFFFFFF)
 ,	bRenderTargetDirty(false)
 ,	CurrentRealTime(0)
@@ -219,7 +220,7 @@ FCanvas::FCanvas(FRenderTarget* InRenderTarget, FHitProxyConsumer* InHitProxyCon
 ,	CurrentDeltaWorldTime(0)
 ,	FeatureLevel(InFeatureLevel)
 ,	StereoDepth(150)
-,	DrawMode(CDM_DeferDrawing)
+,	DrawMode(InDrawMode)
 {
 	Construct();
 
@@ -231,10 +232,11 @@ FCanvas::FCanvas(FRenderTarget* InRenderTarget, FHitProxyConsumer* InHitProxyCon
 	}
 }
 
-FCanvas::FCanvas(FRenderTarget* InRenderTarget,FHitProxyConsumer* InHitProxyConsumer, float InRealTime, float InWorldTime, float InWorldDeltaTime, ERHIFeatureLevel::Type InFeatureLevel, ECanvasDrawMode InDrawMode)
+FCanvas::FCanvas(FRenderTarget* InRenderTarget,FHitProxyConsumer* InHitProxyConsumer, float InRealTime, float InWorldTime, float InWorldDeltaTime, ERHIFeatureLevel::Type InFeatureLevel)
 :	ViewRect(0,0,0,0)
 ,	RenderTarget(InRenderTarget)
 ,	HitProxyConsumer(InHitProxyConsumer)
+,	Scene(NULL)
 ,	AllowedModes(0xFFFFFFFF)
 ,	bRenderTargetDirty(false)
 ,	CurrentRealTime(InRealTime)
@@ -242,7 +244,7 @@ FCanvas::FCanvas(FRenderTarget* InRenderTarget,FHitProxyConsumer* InHitProxyCons
 ,	CurrentDeltaWorldTime(InWorldDeltaTime)
 ,	FeatureLevel(InFeatureLevel)
 ,	StereoDepth(150)
-,	DrawMode(InDrawMode)
+,	DrawMode(CDM_DeferDrawing)
 {
 	Construct();
 }
diff --git a/Engine/Source/Runtime/Engine/Private/World.cpp b/Engine/Source/Runtime/Engine/Private/World.cpp
--- a/Engine/Source/Runtime/Engine/Private/World.cpp
+++ b/Engine/Source/Runtime/Engine/Private/World.cpp
@@ -823,6 +823,16 @@ void UWorld::UpdateParameterCollectionInstances(bool bUpdateInstanceUniformBuffe
 	}
 }
 
+UCanvas* UWorld::GetCanvasForRenderingToTarget()
+{
+	if (!CanvasForRenderingToTarget)
+	{
+		CanvasForRenderingToTarget = NewObject<UCanvas>(GetTransientPackage(), NAME_None);
+	}
+
+	return CanvasForRenderingToTarget;
+}
+
 UAISystemBase* UWorld::CreateAISystem()
 {
 	// create navigation system for editor and server targets, but remove it from game clients
diff --git a/Engine/Source/Runtime/Engine/Public/CanvasTypes.h b/Engine/Source/Runtime/Engine/Public/CanvasTypes.h
--- a/Engine/Source/Runtime/Engine/Public/CanvasTypes.h
+++ b/Engine/Source/Runtime/Engine/Public/CanvasTypes.h
@@ -227,12 +227,12 @@ public:
 	/** 
 	* Constructor.
 	*/
-	ENGINE_API FCanvas(FRenderTarget* InRenderTarget,FHitProxyConsumer* InHitProxyConsumer, UWorld* InWorld, ERHIFeatureLevel::Type InFeatureLevel);
+	ENGINE_API FCanvas(FRenderTarget* InRenderTarget,FHitProxyConsumer* InHitProxyConsumer, UWorld* InWorld, ERHIFeatureLevel::Type InFeatureLevel, ECanvasDrawMode DrawMode = CDM_DeferDrawing);
 
 	/** 
 	* Constructor. For situations where a world is not available, but time information is
 	*/
-	ENGINE_API FCanvas(FRenderTarget* InRenderTarget, FHitProxyConsumer* InHitProxyConsumer, float InRealTime, float InWorldTime, float InWorldDeltaTime, ERHIFeatureLevel::Type InFeatureLevel, ECanvasDrawMode DrawMode = CDM_DeferDrawing);
+	ENGINE_API FCanvas(FRenderTarget* InRenderTarget, FHitProxyConsumer* InHitProxyConsumer, float InRealTime, float InWorldTime, float InWorldDeltaTime, ERHIFeatureLevel::Type InFeatureLevel);
 
 	/** 
 	* Destructor.
@@ -416,6 +416,8 @@ public:
 	FHitProxyConsumer* GetHitProxyConsumer() const { return HitProxyConsumer; }
 	bool IsHitTesting() const { return HitProxyConsumer != NULL; }
 	
+	FSceneInterface* GetScene() const { return Scene; }
+
 	/**
 	* Push sort key onto the stack. Rendering is done with the current sort key stack entry.
 	*
@@ -578,6 +580,8 @@ private:
 	FHitProxyConsumer* HitProxyConsumer;
 	/** Current hit proxy object */
 	TRefCountPtr<HHitProxy> CurrentHitProxy;
+	/* Optional scene for rendering. */
+	FSceneInterface* Scene;
 	/** Toggles for various canvas rendering functionality **/
 	uint32 AllowedModes;
 	/** true if the render target has been rendered to since last calling SetRenderTarget() */
diff --git a/Engine/Source/Runtime/Engine/Public/Engine.h b/Engine/Source/Runtime/Engine/Public/Engine.h
--- a/Engine/Source/Runtime/Engine/Public/Engine.h
+++ b/Engine/Source/Runtime/Engine/Public/Engine.h
@@ -49,6 +49,7 @@
 #include "Curves/CurveFloat.h"
 #include "GameFramework/ForceFeedbackEffect.h"
 #include "Engine/SubsurfaceProfile.h"
+#include "Engine/SubDSurface.h"
 #include "Engine/DebugCameraController.h"
 #include "Engine/DecalActor.h"
 #include "PhysicsEngine/DestructibleActor.h"
diff --git a/Engine/Source/Runtime/Engine/Public/FogRendering.h b/Engine/Source/Runtime/Engine/Public/FogRendering.h
--- a/Engine/Source/Runtime/Engine/Public/FogRendering.h
+++ b/Engine/Source/Runtime/Engine/Public/FogRendering.h
@@ -23,8 +23,9 @@ public:
 		const FViewInfo* ViewInfo = static_cast<const FViewInfo*>(View);
 		SetShaderValue(RHICmdList, Shader, ExponentialFogParameters, ViewInfo->ExponentialFogParameters);
 		SetShaderValue(RHICmdList, Shader, ExponentialFogColorParameter, FVector4(ViewInfo->ExponentialFogColor, 1.0f - ViewInfo->FogMaxOpacity));
+		SetShaderValue(RHICmdList, Shader, ExponentialFogParameters3, ViewInfo->ExponentialFogParameters3);
 		SetShaderValue(RHICmdList, Shader, InscatteringLightDirection, FVector4(ViewInfo->InscatteringLightDirection, ViewInfo->bUseDirectionalInscattering ? 1 : 0));
-		SetShaderValue(RHICmdList, Shader, DirectionalInscatteringColor, FVector4(FVector(ViewInfo->DirectionalInscatteringColor), ViewInfo->DirectionalInscatteringExponent));
+		SetShaderValue(RHICmdList, Shader, DirectionalInscatteringColor, FVector4(FVector(ViewInfo->DirectionalInscatteringColor), FMath::Clamp(ViewInfo->DirectionalInscatteringExponent, 0.000001f, 1000.0f)));
 		SetShaderValue(RHICmdList, Shader, DirectionalInscatteringStartDistance, ViewInfo->DirectionalInscatteringStartDistance);
 	}
 
@@ -33,6 +34,7 @@ public:
 
 	FShaderParameter ExponentialFogParameters;
 	FShaderParameter ExponentialFogColorParameter;
+	FShaderParameter ExponentialFogParameters3;
 	FShaderParameter InscatteringLightDirection;
 	FShaderParameter DirectionalInscatteringColor;
 	FShaderParameter DirectionalInscatteringStartDistance;
diff --git a/Engine/Source/Runtime/Engine/Public/SceneManagement.h b/Engine/Source/Runtime/Engine/Public/SceneManagement.h
--- a/Engine/Source/Runtime/Engine/Public/SceneManagement.h
+++ b/Engine/Source/Runtime/Engine/Public/SceneManagement.h
@@ -951,6 +951,7 @@ public:
 	inline float GetIndirectLightingScale() const { return IndirectLightingScale; }
 	inline FGuid GetLightGuid() const { return LightGuid; }
 	inline float GetShadowSharpen() const { return ShadowSharpen; }
+	inline float GetContactShadowLength() const { return ContactShadowLength; }
 	inline FVector GetLightFunctionScale() const { return LightFunctionScale; }
 	inline float GetLightFunctionFadeDistance() const { return LightFunctionFadeDistance; }
 	inline float GetLightFunctionDisabledBrightness() const { return LightFunctionDisabledBrightness; }
@@ -1023,6 +1024,9 @@ protected:
 	/** Sharpen shadow filtering */
 	float ShadowSharpen;
 
+	/** Length of screen space ray trace for sharp contact shadows. */
+	float ContactShadowLength;
+
 	/** Min roughness */
 	float MinRoughness;
 
diff --git a/Engine/Source/Runtime/Engine/Public/SceneView.h b/Engine/Source/Runtime/Engine/Public/SceneView.h
--- a/Engine/Source/Runtime/Engine/Public/SceneView.h
+++ b/Engine/Source/Runtime/Engine/Public/SceneView.h
@@ -519,6 +519,8 @@ BEGIN_UNIFORM_BUFFER_STRUCT_WITH_CONSTRUCTOR(FFrameUniformShaderParameters, ENGI
 	DECLARE_UNIFORM_BUFFER_STRUCT_MEMBER(float, MobilePreviewMode)
 	DECLARE_UNIFORM_BUFFER_STRUCT_MEMBER(float, HMDEyePaddingOffset)
 	DECLARE_UNIFORM_BUFFER_STRUCT_MEMBER_EX(float, ReflectionCubemapMaxMip, EShaderPrecisionModifier::Half)
+	DECLARE_UNIFORM_BUFFER_STRUCT_MEMBER(float, ShowDecalsMask)
+	DECLARE_UNIFORM_BUFFER_STRUCT_MEMBER(uint32, DistanceFieldAOSpecularOcclusionMode)
 
 	DECLARE_UNIFORM_BUFFER_STRUCT_MEMBER_ARRAY(FVector4, GlobalVolumeCenterAndExtent_UB, [GMaxGlobalDistanceFieldClipmaps])
 	DECLARE_UNIFORM_BUFFER_STRUCT_MEMBER_ARRAY(FVector4, GlobalVolumeWorldToUVAddAndMul_UB, [GMaxGlobalDistanceFieldClipmaps])
@@ -1093,6 +1095,12 @@ public:
 	/** if true then results of scene rendering are copied/resolved to the RenderTarget. */
 	bool bResolveScene;
 
+	/** 
+	 * Which component of the scene rendering should be output to the final render target.
+	 * If SCS_FinalColorLDR this indicates do nothing.
+	 */
+	ESceneCaptureSource SceneCaptureSource;
+
 	/**
 	 * GetWorld->IsPaused() && !Simulate
 	 * Simulate is excluded as the camera can move which invalidates motionblur
diff --git a/Engine/Source/Runtime/Engine/Public/ShowFlagsValues.inl b/Engine/Source/Runtime/Engine/Public/ShowFlagsValues.inl
--- a/Engine/Source/Runtime/Engine/Public/ShowFlagsValues.inl
+++ b/Engine/Source/Runtime/Engine/Public/ShowFlagsValues.inl
@@ -273,7 +273,7 @@ SHOWFLAG_ALWAYS_ACCESSIBLE(Paper2DSprites, SFG_Advanced, LOCTEXT("Paper2DSprites
 /** Visualization of distance field AO */
 SHOWFLAG_FIXED_IN_SHIPPING(0, VisualizeDistanceFieldAO, SFG_Visualize, LOCTEXT("VisualizeDistanceFieldAOSF", "Distance Field Ambient Occlusion"))
 /** Visualization of distance field GI */
-SHOWFLAG_FIXED_IN_SHIPPING(0, VisualizeDistanceFieldGI, SFG_Visualize, LOCTEXT("VisualizeDistanceFieldGISF", "Distance Field Global Illumination"))
+SHOWFLAG_FIXED_IN_SHIPPING(0, VisualizeDistanceFieldGI, SFG_Hidden, LOCTEXT("VisualizeDistanceFieldGISF", "Distance Field Global Illumination"))
 /** Mesh Distance fields */
 SHOWFLAG_FIXED_IN_SHIPPING(0, VisualizeMeshDistanceFields, SFG_Visualize, LOCTEXT("MeshDistanceFieldsSF", "Mesh DistanceFields"))
 /** Screen space AO, for now SHOWFLAG_ALWAYS_ACCESSIBLE because r.GBuffer need that */
@@ -281,7 +281,7 @@ SHOWFLAG_ALWAYS_ACCESSIBLE(ScreenSpaceAO, SFG_LightingFeatures, LOCTEXT("ScreenS
 /** Distance field AO, for now SHOWFLAG_ALWAYS_ACCESSIBLE because it's exposed in SceneCapture */
 SHOWFLAG_ALWAYS_ACCESSIBLE(DistanceFieldAO, SFG_LightingFeatures, LOCTEXT("DistanceFieldAOSF", "Distance Field Ambient Occlusion"))
 /** Distance field GI */
-SHOWFLAG_FIXED_IN_SHIPPING(1, DistanceFieldGI, SFG_LightingFeatures, LOCTEXT("DistanceFieldGISF", "Distance Field Global Illumination"))
+SHOWFLAG_FIXED_IN_SHIPPING(1, DistanceFieldGI, SFG_Hidden, LOCTEXT("DistanceFieldGISF", "Distance Field Global Illumination"))
 /** Visualize screen space reflections, for developer (by default off): */
 SHOWFLAG_FIXED_IN_SHIPPING(0, VisualizeSSR, SFG_Visualize, LOCTEXT("VisualizeSSR", "Screen Space Reflections"))
 /** Visualize the Shading Models, mostly or debugging and profiling */
diff --git a/Engine/Source/Runtime/OpenGLDrv/Private/OpenGLDevice.cpp b/Engine/Source/Runtime/OpenGLDrv/Private/OpenGLDevice.cpp
--- a/Engine/Source/Runtime/OpenGLDrv/Private/OpenGLDevice.cpp
+++ b/Engine/Source/Runtime/OpenGLDrv/Private/OpenGLDevice.cpp
@@ -119,10 +119,18 @@ bool IsUniformBufferBound( GLuint Buffer )
 extern void BeginFrame_UniformBufferPoolCleanup();
 extern void BeginFrame_VertexBufferCleanup();
 
-FOpenGLContextState& FOpenGLDynamicRHI::GetContextStateForCurrentContext()
+FOpenGLContextState& FOpenGLDynamicRHI::GetContextStateForCurrentContext(bool bAssertIfInvalid)
 {
 	int32 ContextType = (int32)PlatformOpenGLCurrentContext(PlatformDevice);
-	check(ContextType >= 0);
+	if (bAssertIfInvalid)
+	{
+		check(ContextType >= 0);
+	}
+	else if (ContextType < 0)
+	{
+		return InvalidContextState;
+	}
+
 	if (ContextType == CONTEXT_Rendering)
 	{
 		return RenderingContextState;
diff --git a/Engine/Source/Runtime/OpenGLDrv/Private/OpenGLRenderTarget.cpp b/Engine/Source/Runtime/OpenGLDrv/Private/OpenGLRenderTarget.cpp
--- a/Engine/Source/Runtime/OpenGLDrv/Private/OpenGLRenderTarget.cpp
+++ b/Engine/Source/Runtime/OpenGLDrv/Private/OpenGLRenderTarget.cpp
@@ -766,14 +766,17 @@ void FOpenGLDynamicRHI::RHIReadSurfaceData(FTextureRHIParamRef TextureRHI,FIntRe
 	TArray<uint8> Temp;
 
 	FOpenGLContextState& ContextState = GetContextStateForCurrentContext();
-	ReadSurfaceDataRaw(ContextState, TextureRHI, Rect, Temp, InFlags);
+	OutData.Empty();
+	if (&ContextState != &InvalidContextState)
+	{
+		ReadSurfaceDataRaw(ContextState, TextureRHI, Rect, Temp, InFlags);
 
-	uint32 Size = Rect.Width() * Rect.Height(); 
+		uint32 Size = Rect.Width() * Rect.Height();
 
-	OutData.Empty();
-	OutData.AddUninitialized(Size);
+		OutData.AddUninitialized(Size);
 
-	FMemory::Memcpy(OutData.GetData(), Temp.GetData(), Size * sizeof(FColor));
+		FMemory::Memcpy(OutData.GetData(), Temp.GetData(), Size * sizeof(FColor));
+	}
 }
 
 void FOpenGLDynamicRHI::RHIMapStagingSurface(FTextureRHIParamRef TextureRHI,void*& OutData,int32& OutWidth,int32& OutHeight)
diff --git a/Engine/Source/Runtime/RHI/Public/RHIDefinitions.h b/Engine/Source/Runtime/RHI/Public/RHIDefinitions.h
--- a/Engine/Source/Runtime/RHI/Public/RHIDefinitions.h
+++ b/Engine/Source/Runtime/RHI/Public/RHIDefinitions.h
@@ -791,7 +791,7 @@ inline bool RHISupportsGeometryShaders(const EShaderPlatform Platform)
 
 inline bool RHIHasTiledGPU(const EShaderPlatform Platform)
 {
-	return (Platform == SP_METAL_MRT) || Platform == SP_METAL || Platform == SP_OPENGL_ES2_IOS || Platform == SP_OPENGL_ES2_ANDROID || Platform == SP_OPENGL_ES2_ANDROID;
+	return (Platform == SP_METAL_MRT) || Platform == SP_METAL || Platform == SP_OPENGL_ES2_IOS || Platform == SP_OPENGL_ES2_ANDROID;
 }
 
 inline bool RHISupportsVertexShaderLayer(const EShaderPlatform Platform)
diff --git a/Engine/Source/Runtime/Renderer/Private/BasePassRendering.cpp b/Engine/Source/Runtime/Renderer/Private/BasePassRendering.cpp
--- a/Engine/Source/Runtime/Renderer/Private/BasePassRendering.cpp
+++ b/Engine/Source/Runtime/Renderer/Private/BasePassRendering.cpp
@@ -48,7 +48,7 @@ bool GVisualizeMipLevels = false;
 
 #define IMPLEMENT_BASEPASS_PIXELSHADER_TYPE(LightMapPolicyType,LightMapPolicyName,bEnableSkyLight,SkyLightName) \
 	typedef TBasePassPS<LightMapPolicyType, bEnableSkyLight> TBasePassPS##LightMapPolicyName##SkyLightName; \
-	IMPLEMENT_MATERIAL_SHADER_TYPE(template<>,TBasePassPS##LightMapPolicyName##SkyLightName,TEXT("BasePassPixelShader"),TEXT("Main"),SF_Pixel);
+	IMPLEMENT_MATERIAL_SHADER_TYPE(template<>,TBasePassPS##LightMapPolicyName##SkyLightName,TEXT("BasePassPixelShader"),TEXT("MainPS"),SF_Pixel);
 
 // Implement a pixel shader type for skylights and one without, and one vertex shader that will be shared between them
 #define IMPLEMENT_BASEPASS_LIGHTMAPPED_SHADER_TYPE(LightMapPolicyType,LightMapPolicyName) \
@@ -444,7 +444,7 @@ public:
 			/* bInEnableReceiveDecalOutput = */ Scene != nullptr
 			);
 		RHICmdList.BuildAndSetLocalBoundShaderState(DrawingPolicy.GetBoundShaderStateInput(View.GetFeatureLevel()));
-		DrawingPolicy.SetSharedState(RHICmdList, &View, typename TBasePassDrawingPolicy<LightMapPolicyType>::ContextDataType(Parameters.bIsInstancedStereo, false));
+		DrawingPolicy.SetSharedState(RHICmdList, &View, typename TBasePassDrawingPolicy<LightMapPolicyType>::ContextDataType(Parameters.bIsInstancedStereo, false), 1.0f);
 		const FMeshDrawingRenderState DrawRenderState(DitheredLODTransitionAlpha);
 
 		for( int32 BatchElementIndex = 0, Num = Parameters.Mesh.Elements.Num(); BatchElementIndex < Num; BatchElementIndex++ )
diff --git a/Engine/Source/Runtime/Renderer/Private/BasePassRendering.h b/Engine/Source/Runtime/Renderer/Private/BasePassRendering.h
--- a/Engine/Source/Runtime/Renderer/Private/BasePassRendering.h
+++ b/Engine/Source/Runtime/Renderer/Private/BasePassRendering.h
@@ -511,8 +511,6 @@ public:
 			const int32 VelocityIndex = (CVar && CVar->GetValueOnAnyThread() != 0) ? 6 : 5;
 			OutEnvironment.SetRenderTargetOutputFormat(VelocityIndex, PF_G16R16);
 		}
-
-		OutEnvironment.SetDefine(TEXT("TRANSLUCENCY_RENDERING"), 1);
 	}
 
 	/** Initialization constructor. */
@@ -524,7 +522,7 @@ public:
 		TranslucentLightingParameters.Bind(Initializer.ParameterMap);
 		EditorCompositeParams.Bind(Initializer.ParameterMap);
 		LightGrid.Bind(Initializer.ParameterMap,TEXT("LightGrid"));
-		ScreenTextureUVScale.Bind(Initializer.ParameterMap, TEXT("ScreenPositionUVScale"));
+		DownsampleFactorFromSceneBufferSize.Bind(Initializer.ParameterMap, TEXT("DownsampleFactorFromSceneBufferSize"));
 	}
 	TBasePassPixelShaderPolicyParamType() {}
 
@@ -536,7 +534,7 @@ public:
 		EBlendMode BlendMode, 
 		bool bEnableEditorPrimitveDepthTest,
 		ESceneRenderTargetsMode::Type TextureMode,
-		float ScreenTextureScaleFactor = 1.0f)
+		float DownsampleFactorFromSceneBufferSizeValue)
 	{
 		const FPixelShaderRHIParamRef ShaderRHI = GetPixelShader();
 
@@ -548,7 +546,7 @@ public:
 
 			if (IsTranslucentBlendMode(BlendMode))
 			{
-				SetShaderValue(RHICmdList, ShaderRHI, ScreenTextureUVScale, FVector(ScreenTextureScaleFactor));
+				SetShaderValue(RHICmdList, ShaderRHI, DownsampleFactorFromSceneBufferSize, DownsampleFactorFromSceneBufferSizeValue);
 				TranslucentLightingParameters.Set(RHICmdList, ShaderRHI, View);
 
 				// Experimental dynamic forward lighting for translucency. Can be the base for opaque forward lighting which will allow more lighting models or rendering without a GBuffer
@@ -575,7 +573,7 @@ public:
 		Ar << TranslucentLightingParameters;
  		Ar << EditorCompositeParams;
 		Ar << LightGrid;
-		Ar << ScreenTextureUVScale;
+		Ar << DownsampleFactorFromSceneBufferSize;
 		return bShaderHasOutdatedParameters;
 	}
 
@@ -584,7 +582,7 @@ private:
 	FTranslucentLightingParameters TranslucentLightingParameters;
 	FEditorCompositingParameters EditorCompositeParams;
 	FShaderResourceParameter LightGrid;
-	FShaderParameter ScreenTextureUVScale;
+	FShaderParameter DownsampleFactorFromSceneBufferSize;
 };
 
 /**
@@ -806,7 +804,7 @@ public:
 		DRAWING_POLICY_MATCH_END
 	}
 
-	void SetSharedState(FRHICommandList& RHICmdList, const FViewInfo* View, const ContextDataType PolicyContext, float ScreenTextureScaleFactor = 1.0f) const
+	void SetSharedState(FRHICommandList& RHICmdList, const FViewInfo* View, const ContextDataType PolicyContext, float DownsampleFactorFromSceneBufferSize = 1.0f) const
 	{
 		// If the current debug view shader modes are allowed, different VS/DS/HS must be used (with only SV_POSITION as PS interpolant).
 		if (View->Family->UseDebugViewVSDSHS())
@@ -843,7 +841,7 @@ public:
 		}
 		else
 		{
-			PixelShader->SetParameters(RHICmdList, MaterialRenderProxy, *MaterialResource, View, BlendMode, bEnableEditorPrimitiveDepthTest, SceneTextureMode, ScreenTextureScaleFactor);
+			PixelShader->SetParameters(RHICmdList, MaterialRenderProxy, *MaterialResource, View, BlendMode, bEnableEditorPrimitiveDepthTest, SceneTextureMode, DownsampleFactorFromSceneBufferSize);
 
 			switch(BlendMode)
 			{
@@ -855,13 +853,14 @@ public:
 				// Masked materials are rendered together in the base pass, where the blend state is set at a higher level
 				break;
 			case BLEND_Translucent:
-				// Alpha channel is only needed for SeparateTranslucency, before this was preserving the alpha channel but we no longer store depth in the alpha channel so it's no problem
-
+				// Note: alpha channel used by separate translucency, storing how much of the background should be added when doing the final composite
+				// The Alpha channel is also used by non-separate translucency when rendering to scene captures, which store the final opacity
 				RHICmdList.SetBlendState( TStaticBlendState<CW_RGBA, BO_Add, BF_SourceAlpha, BF_InverseSourceAlpha, BO_Add, BF_Zero, BF_InverseSourceAlpha>::GetRHI());
 				break;
 			case BLEND_Additive:
 				// Add to the existing scene color
-				// Alpha channel is only needed for SeparateTranslucency, before this was preserving the alpha channel but we no longer store depth in the alpha channel so it's no problem
+				// Note: alpha channel used by separate translucency, storing how much of the background should be added when doing the final composite
+				// The Alpha channel is also used by non-separate translucency when rendering to scene captures, which store the final opacity
 				RHICmdList.SetBlendState( TStaticBlendState<CW_RGBA, BO_Add, BF_One, BF_One, BO_Add, BF_Zero, BF_InverseSourceAlpha>::GetRHI());
 				break;
 			case BLEND_Modulate:
diff --git a/Engine/Source/Runtime/Renderer/Private/CustomDepthRendering.h b/Engine/Source/Runtime/Renderer/Private/CustomDepthRendering.h
--- a/Engine/Source/Runtime/Renderer/Private/CustomDepthRendering.h
+++ b/Engine/Source/Runtime/Renderer/Private/CustomDepthRendering.h
@@ -20,15 +20,6 @@ public:
 	*/
 	bool DrawPrims(FRHICommandListImmediate& RHICmdList, const class FViewInfo& View, bool bWriteCustomStencilValues);
 
-	/**
-	* Add a new primitive to the list of prims
-	* @param PrimitiveSceneProxy - primitive info to add.
-	* @param ViewInfo - used to transform bounds to view space
-	*/
-	void AddScenePrimitive(FPrimitiveSceneProxy* PrimitiveSceneProxy)
-	{
-		Prims.Add(PrimitiveSceneProxy);
-	}
 	/**
 	* Adds a new primitives to the list of distortion prims
 	* @param PrimitiveSceneProxies - primitive info to add.
diff --git a/Engine/Source/Runtime/Renderer/Private/DeferredShadingRenderer.cpp b/Engine/Source/Runtime/Renderer/Private/DeferredShadingRenderer.cpp
--- a/Engine/Source/Runtime/Renderer/Private/DeferredShadingRenderer.cpp
+++ b/Engine/Source/Runtime/Renderer/Private/DeferredShadingRenderer.cpp
@@ -35,7 +35,7 @@ TAutoConsoleVariable<int32> CVarEarlyZPass(
 	TEXT("  1: only if not masked, and only if large on the screen\n")
 	TEXT("  2: all opaque (including masked)\n")
 	TEXT("  x: use built in heuristic (default is 3)"),
-	ECVF_ReadOnly);
+	ECVF_Scalability);
 
 int32 GEarlyZPassMovable = 0;
 
@@ -45,7 +45,7 @@ static FAutoConsoleVariableRef CVarEarlyZPassMovable(
 	GEarlyZPassMovable,
 	TEXT("Whether to render movable objects into the depth only pass.  Movable objects are typically not good occluders so this defaults to off.\n")
 	TEXT("Note: also look at r.EarlyZPass"),
-	ECVF_RenderThreadSafe | ECVF_ReadOnly
+	ECVF_RenderThreadSafe | ECVF_Scalability
 	);
 
 TAutoConsoleVariable<int32> CVarCustomDepthOrder(
@@ -125,6 +125,7 @@ DECLARE_CYCLE_STAT(TEXT("OcclusionSubmittedFence Wait"), STAT_OcclusionSubmitted
 FDeferredShadingSceneRenderer::FDeferredShadingSceneRenderer(const FSceneViewFamily* InViewFamily,FHitProxyConsumer* HitProxyConsumer)
 	: FSceneRenderer(InViewFamily, HitProxyConsumer)
 	, EarlyZPassMode(DDM_NonMaskedOnly)
+	, bEarlyZPassMovable(false)
 	, TranslucentSelfShadowLayout(0, 0, 0, 0)
 	, CachedTranslucentSelfShadowLightId(INDEX_NONE)
 {
@@ -140,8 +141,7 @@ FDeferredShadingSceneRenderer::FDeferredShadingSceneRenderer(const FSceneViewFam
 
 	// developer override, good for profiling, can be useful as project setting
 	{
-		static const auto ICVar = IConsoleManager::Get().FindTConsoleVariableDataInt(TEXT("r.EarlyZPass"));
-		const int32 CVarValue = ICVar->GetValueOnGameThread();
+		const int32 CVarValue = CVarEarlyZPass.GetValueOnGameThread();
 
 		switch(CVarValue)
 		{
@@ -158,12 +158,32 @@ FDeferredShadingSceneRenderer::FDeferredShadingSceneRenderer(const FSceneViewFam
 		EarlyZPassMode = DDM_AllOpaque;
 	}
 
+	// Warning: EarlyZPass logic is mirrored in FStaticMesh::AddToDrawLists
+
+	bEarlyZPassMovable = GEarlyZPassMovable != 0;
+
+	static IConsoleVariable* CDBufferVar = IConsoleManager::Get().FindConsoleVariable(TEXT("r.DBuffer"));
+	bool bDBufferAllowed = CDBufferVar ? CDBufferVar->GetInt() != 0 : false;
+
+	if (bDBufferAllowed)
+	{
+		// DBuffer decals force a full prepass
+		EarlyZPassMode = DDM_AllOccluders;
+		bEarlyZPassMovable = true;
+	}
+
 	// When possible use a stencil optimization for dithering in pre-pass
 	bDitheredLODTransitionsUseStencil = (EarlyZPassMode == DDM_AllOccluders);
 }
 
 extern FGlobalBoundShaderState GClearMRTBoundShaderState[8];
 
+float GetSceneColorClearAlpha()
+{
+	// Scene color alpha is used during scene captures and planar reflections.  1 indicates background should be shown, 0 indicates foreground is fully present.
+	return 1.0f;
+}
+
 /** 
 * Clears view where Z is still at the maximum value (ie no geometry rendered)
 */
@@ -174,9 +194,11 @@ void FDeferredShadingSceneRenderer::ClearGBufferAtMaxZ(FRHICommandList& RHICmdLi
 
 	// Clear the G Buffer render targets
 	const bool bClearBlack = Views[0].Family->EngineShowFlags.ShaderComplexity || Views[0].Family->EngineShowFlags.StationaryLightOverlap;
+	const float ClearAlpha = GetSceneColorClearAlpha();
+	const FLinearColor ClearColor = bClearBlack ? FLinearColor(0, 0, 0, ClearAlpha) : FLinearColor(Views[0].BackgroundColor.R, Views[0].BackgroundColor.G, Views[0].BackgroundColor.B, ClearAlpha);
 	// Same clear color from RHIClearMRT
 	FLinearColor ClearColors[MaxSimultaneousRenderTargets] = 
-		{bClearBlack ? FLinearColor(0,0,0,0) : Views[0].BackgroundColor, FLinearColor(0.5f,0.5f,0.5f,0), FLinearColor(0,0,0,1), FLinearColor(0,0,0,0), FLinearColor(0,1,1,1), FLinearColor(1,1,1,1), FLinearColor::Transparent, FLinearColor::Transparent};
+		{ClearColor, FLinearColor(0.5f,0.5f,0.5f,0), FLinearColor(0,0,0,1), FLinearColor(0,0,0,0), FLinearColor(0,1,1,1), FLinearColor(1,1,1,1), FLinearColor::Transparent, FLinearColor::Transparent};
 
 	uint32 NumActiveRenderTargets = FSceneRenderTargets::Get(RHICmdList).GetNumGBufferTargets();
 	
@@ -232,7 +254,7 @@ void FDeferredShadingSceneRenderer::ClearGBufferAtMaxZ(FRHICommandList& RHICmdLi
 		// Render quad
 		static const FVector4 ClearQuadVertices[4] = 
 		{
-			FVector4( -1.0f,	  1.0f, (float)ERHIZBuffer::FarPlane, 1.0f),
+			FVector4( -1.0f,  1.0f, (float)ERHIZBuffer::FarPlane, 1.0f),
 			FVector4(  1.0f,  1.0f, (float)ERHIZBuffer::FarPlane, 1.0f ),
 			FVector4( -1.0f, -1.0f, (float)ERHIZBuffer::FarPlane, 1.0f ),
 			FVector4(  1.0f, -1.0f, (float)ERHIZBuffer::FarPlane, 1.0f )
@@ -699,7 +721,7 @@ FGraphEventRef FDeferredShadingSceneRenderer::TranslucencyTimestampQuerySubmitte
 static FORCEINLINE bool NeedsPrePass(const FDeferredShadingSceneRenderer* Renderer)
 {
 	return (RHIHasTiledGPU(Renderer->ViewFamily.GetShaderPlatform()) == false) && 
-		(Renderer->EarlyZPassMode != DDM_None || GEarlyZPassMovable != 0);
+		(Renderer->EarlyZPassMode != DDM_None || Renderer->bEarlyZPassMovable != 0);
 }
 
 /**
@@ -721,7 +743,8 @@ static void SetAndClearViewGBuffer(FRHICommandListImmediate& RHICmdList, FViewIn
 	ERenderTargetLoadAction DepthLoadAction = bClearDepth ? ERenderTargetLoadAction::EClear : ERenderTargetLoadAction::ELoad;
 
 	const bool bClearBlack = View.Family->EngineShowFlags.ShaderComplexity || View.Family->EngineShowFlags.StationaryLightOverlap;
-	const FLinearColor ClearColor = (bClearBlack ? FLinearColor(0, 0, 0, 0) : View.BackgroundColor);
+	const float ClearAlpha = GetSceneColorClearAlpha();
+	const FLinearColor ClearColor = bClearBlack ? FLinearColor(0, 0, 0, ClearAlpha) : FLinearColor(View.BackgroundColor.R, View.BackgroundColor.G, View.BackgroundColor.B, ClearAlpha);
 
 	// clearing the GBuffer
 	FSceneRenderTargets::Get(RHICmdList).BeginRenderingGBuffer(RHICmdList, ERenderTargetLoadAction::EClear, DepthLoadAction, View.Family->EngineShowFlags.ShaderComplexity, ClearColor);
@@ -1223,8 +1246,7 @@ void FDeferredShadingSceneRenderer::Render(FRHICommandListImmediate& RHICmdList)
 
 	// Pre-lighting composition lighting stage
 	// e.g. deferred decals, SSAO
-	if (FeatureLevel >= ERHIFeatureLevel::SM4
-		&& bUseGBuffer)
+	if (FeatureLevel >= ERHIFeatureLevel::SM4)
 	{
 		SCOPE_CYCLE_COUNTER(STAT_FDeferredShadingSceneRenderer_AfterBasePass);
 
@@ -1456,6 +1478,8 @@ void FDeferredShadingSceneRenderer::Render(FRHICommandListImmediate& RHICmdList)
 	// Resolve the scene color for post processing.
 	SceneContext.ResolveSceneColor(RHICmdList, FResolveRect(0, 0, ViewFamily.FamilySizeX, ViewFamily.FamilySizeY));
 
+	CopySceneCaptureComponentToTarget(RHICmdList);
+
 	// Finish rendering for each view.
 	if (ViewFamily.bResolveScene)
 	{
@@ -1520,7 +1544,7 @@ bool FDeferredShadingSceneRenderer::RenderPrePassViewDynamic(FRHICommandList& RH
 				// Only render primitives marked as occluders
 				bShouldUseAsOccluder = PrimitiveSceneProxy->ShouldUseAsOccluder()
 					// Only render static objects unless movable are requested
-					&& (!PrimitiveSceneProxy->IsMovable() || GEarlyZPassMovable)
+					&& (!PrimitiveSceneProxy->IsMovable() || bEarlyZPassMovable)
 					&& (FMath::Square(PrimitiveSceneProxy->GetBounds().SphereRadius) > GMinScreenRadiusForDepthPrepass * GMinScreenRadiusForDepthPrepass * LODFactorDistanceSquared);
 			}
 
@@ -1848,7 +1872,10 @@ bool FDeferredShadingSceneRenderer::PreRenderPrePass(FRHICommandListImmediate& R
 bool FDeferredShadingSceneRenderer::RenderPrePass(FRHICommandListImmediate& RHICmdList, TFunctionRef<void()> AfterTasksAreStarted)
 {
 	bool bDepthWasCleared = false;
-	SCOPED_DRAW_EVENT(RHICmdList, PrePass);
+
+	static IConsoleVariable* CDBufferVar = IConsoleManager::Get().FindConsoleVariable(TEXT("r.DBuffer"));
+	bool bDBufferAllowed = CDBufferVar ? CDBufferVar->GetInt() != 0 : false;
+	SCOPED_DRAW_EVENTF(RHICmdList, PrePass, TEXT("PrePass %s %s"), GetDepthDrawingModeString(EarlyZPassMode), (bDBufferAllowed ? TEXT("(Forced by DBuffer)") : TEXT("")));
 	SCOPE_CYCLE_COUNTER(STAT_DepthDrawTime);
 
 	bool bDirty = false;
diff --git a/Engine/Source/Runtime/Renderer/Private/DeferredShadingRenderer.h b/Engine/Source/Runtime/Renderer/Private/DeferredShadingRenderer.h
--- a/Engine/Source/Runtime/Renderer/Private/DeferredShadingRenderer.h
+++ b/Engine/Source/Runtime/Renderer/Private/DeferredShadingRenderer.h
@@ -24,6 +24,7 @@ public:
 
 	/** Defines which objects we want to render in the EarlyZPass. */
 	EDepthDrawingMode EarlyZPassMode;
+	bool bEarlyZPassMovable;
 	bool bDitheredLODTransitionsUseStencil;
 
 	/** 
@@ -418,5 +419,7 @@ private:
 
 	void DrawAllTranslucencyPasses(FRHICommandListImmediate& RHICmdList, const FViewInfo& View, ETranslucencyPass::Type TranslucenyPassType);
 
+	void CopySceneCaptureComponentToTarget(FRHICommandListImmediate& RHICmdList);
+
 	friend class FTranslucentPrimSet;
 };
diff --git a/Engine/Source/Runtime/Renderer/Private/DepthRendering.cpp b/Engine/Source/Runtime/Renderer/Private/DepthRendering.cpp
--- a/Engine/Source/Runtime/Renderer/Private/DepthRendering.cpp
+++ b/Engine/Source/Runtime/Renderer/Private/DepthRendering.cpp
@@ -7,6 +7,25 @@
 #include "RendererPrivate.h"
 #include "ScenePrivate.h"
 
+const TCHAR* GetDepthDrawingModeString(EDepthDrawingMode Mode)
+{
+	switch (Mode)
+	{
+	case DDM_None:
+		return TEXT("DDM_None");
+	case DDM_NonMaskedOnly:
+		return TEXT("DDM_NonMaskedOnly");
+	case DDM_AllOccluders:
+		return TEXT("DDM_AllOccluders");
+	case DDM_AllOpaque:
+		return TEXT("DDM_AllOpaque");
+	default:
+		check(0);
+	}
+
+	return TEXT("");
+}
+
 /**
  * A vertex shader for rendering the depth of a mesh.
  */
diff --git a/Engine/Source/Runtime/Renderer/Private/DepthRendering.h b/Engine/Source/Runtime/Renderer/Private/DepthRendering.h
--- a/Engine/Source/Runtime/Renderer/Private/DepthRendering.h
+++ b/Engine/Source/Runtime/Renderer/Private/DepthRendering.h
@@ -21,6 +21,8 @@ enum EDepthDrawingMode
 	//		 Please keep in sync if changed to avoid compile environment mismatch.
 };
 
+extern const TCHAR* GetDepthDrawingModeString(EDepthDrawingMode Mode);
+
 template<bool>
 class TDepthOnlyVS;
 
diff --git a/Engine/Source/Runtime/Renderer/Private/DistanceFieldScreenGridLighting.cpp b/Engine/Source/Runtime/Renderer/Private/DistanceFieldScreenGridLighting.cpp
--- a/Engine/Source/Runtime/Renderer/Private/DistanceFieldScreenGridLighting.cpp
+++ b/Engine/Source/Runtime/Renderer/Private/DistanceFieldScreenGridLighting.cpp
@@ -793,8 +793,8 @@ void FDeferredShadingSceneRenderer::RenderDistanceFieldAOScreenGrid(
 	{
 		SCOPED_DRAW_EVENT(RHICmdList, ConeTraceObjects);
 
-		const uint32 GroupSizeX = FMath::DivideAndRoundDown(View.ViewRect.Size().X / GAODownsampleFactor, GDistanceFieldAOTileSizeX / GConeTraceDownsampleFactor);
-		const uint32 GroupSizeY = FMath::DivideAndRoundDown(View.ViewRect.Size().Y / GAODownsampleFactor, GDistanceFieldAOTileSizeY / GConeTraceDownsampleFactor);
+		const uint32 GroupSizeX = FMath::DivideAndRoundUp(View.ViewRect.Size().X / GAODownsampleFactor, GDistanceFieldAOTileSizeX / GConeTraceDownsampleFactor);
+		const uint32 GroupSizeY = FMath::DivideAndRoundUp(View.ViewRect.Size().Y / GAODownsampleFactor, GDistanceFieldAOTileSizeY / GConeTraceDownsampleFactor);
 		
 		if (bUseGlobalDistanceField)
 		{
@@ -846,8 +846,8 @@ void FDeferredShadingSceneRenderer::RenderDistanceFieldAOScreenGrid(
 	{
 		SCOPED_DRAW_EVENT(RHICmdList, ConeTraceGlobal);
 
-		const uint32 GroupSizeX = FMath::DivideAndRoundDown(View.ViewRect.Size().X / GAODownsampleFactor, GConeTraceGlobalDFTileSize);
-		const uint32 GroupSizeY = FMath::DivideAndRoundDown(View.ViewRect.Size().Y / GAODownsampleFactor, GConeTraceGlobalDFTileSize);
+		const uint32 GroupSizeX = FMath::DivideAndRoundUp(View.ViewRect.Size().X / GAODownsampleFactor, GConeTraceGlobalDFTileSize);
+		const uint32 GroupSizeY = FMath::DivideAndRoundUp(View.ViewRect.Size().Y / GAODownsampleFactor, GConeTraceGlobalDFTileSize);
 
 		check(View.GlobalDistanceFieldInfo.Clipmaps.Num() > 0);
 
diff --git a/Engine/Source/Runtime/Renderer/Private/DistanceFieldSpecularOcclusion.cpp b/Engine/Source/Runtime/Renderer/Private/DistanceFieldSpecularOcclusion.cpp
--- a/Engine/Source/Runtime/Renderer/Private/DistanceFieldSpecularOcclusion.cpp
+++ b/Engine/Source/Runtime/Renderer/Private/DistanceFieldSpecularOcclusion.cpp
@@ -22,16 +22,6 @@
 #include "BasePassRendering.h"
 #include "HeightfieldLighting.h"
 
-// Also have to set USE_CONE_TRACED_SPECULAR_OCCLUSION to 1
-// Prototype feature
-int32 GDistanceFieldSpecularOcclusion = 0;
-FAutoConsoleVariableRef CVarDistanceFieldSpecularOcclusion(
-	TEXT("r.DistanceFieldSpecularOcclusion"),
-	GDistanceFieldSpecularOcclusion,
-	TEXT("Whether the distance field specular occlusion feature is allowed, which is used to shadow reflection captures and skylight reflections more accurately."),
-	ECVF_Scalability | ECVF_RenderThreadSafe
-	);
-
 int32 GSpecularOcclusionDownsampleFactor = 1;
 
 FIntPoint GetBufferSizeForSpecularOcclusion()
@@ -189,7 +179,9 @@ void FDeferredShadingSceneRenderer::RenderDistanceFieldSpecularOcclusion(
 	const TRefCountPtr<IPooledRenderTarget>& DistanceFieldNormal, 
 	TRefCountPtr<IPooledRenderTarget>& OutSpecularOcclusion)
 {
-	if (GDistanceFieldSpecularOcclusion)
+	extern int32 GDistanceFieldAOSpecularOcclusionMode;
+
+	if (GDistanceFieldAOSpecularOcclusionMode == 2)
 	{
 		const FIntPoint ConeTraceBufferSize = GetBufferSizeForSpecularOcclusion();
 
diff --git a/Engine/Source/Runtime/Renderer/Private/DistanceFieldSurfaceCacheLighting.cpp b/Engine/Source/Runtime/Renderer/Private/DistanceFieldSurfaceCacheLighting.cpp
--- a/Engine/Source/Runtime/Renderer/Private/DistanceFieldSurfaceCacheLighting.cpp
+++ b/Engine/Source/Runtime/Renderer/Private/DistanceFieldSurfaceCacheLighting.cpp
@@ -39,6 +39,17 @@ FAutoConsoleVariableRef CVarDistanceFieldAOApplyToStaticIndirect(
 	ECVF_Scalability | ECVF_RenderThreadSafe
 	);
 
+int32 GDistanceFieldAOSpecularOcclusionMode = 1;
+FAutoConsoleVariableRef CVarDistanceFieldAOSpecularOcclusionMode(
+	TEXT("r.AOSpecularOcclusionMode"),
+	GDistanceFieldAOSpecularOcclusionMode,
+	TEXT("Determines how specular should be occluded by DFAO\n")
+	TEXT("0: Apply non-directional AO to specular.\n")
+	TEXT("1: (default) Intersect the reflection cone with the unoccluded cone produced by DFAO.  This gives more accurate occlusion than 0, but can bring out DFAO sampling artifacts.\n")
+	TEXT("2: (experimental) Cone trace through distance fields along the reflection vector.  Costs about the same as DFAO again because more cone tracing is done, but produces more accurate occlusion."),
+	ECVF_Scalability | ECVF_RenderThreadSafe
+	);
+
 bool IsDistanceFieldGIAllowed(const FViewInfo& View)
 {
 	return DoesPlatformSupportDistanceFieldGI(View.GetShaderPlatform())
diff --git a/Engine/Source/Runtime/Renderer/Private/DistortionRendering.h b/Engine/Source/Runtime/Renderer/Private/DistortionRendering.h
--- a/Engine/Source/Runtime/Renderer/Private/DistortionRendering.h
+++ b/Engine/Source/Runtime/Renderer/Private/DistortionRendering.h
@@ -21,14 +21,6 @@ public:
 	*/
 	bool DrawAccumulatedOffsets(FRHICommandListImmediate& RHICmdList, const class FViewInfo& View, bool bInitializeOffsets);
 
-	/**
-	* Add a new primitive to the list of distortion prims
-	* @param PrimitiveSceneProxy - primitive info to add.
-	*/
-	void AddScenePrimitive(FPrimitiveSceneProxy* PrimitiveSceneProxy)
-	{
-		Prims.Add(PrimitiveSceneProxy);
-	}
 	/**
 	* Adds a new primitives to the list of distortion prims
 	* @param PrimitiveSceneProxies - primitive info to add.
diff --git a/Engine/Source/Runtime/Renderer/Private/FogRendering.cpp b/Engine/Source/Runtime/Renderer/Private/FogRendering.cpp
--- a/Engine/Source/Runtime/Renderer/Private/FogRendering.cpp
+++ b/Engine/Source/Runtime/Renderer/Private/FogRendering.cpp
@@ -38,8 +38,9 @@ static TAutoConsoleVariable<int32> CVarFog(
 /** Binds the parameters. */
 void FExponentialHeightFogShaderParameters::Bind(const FShaderParameterMap& ParameterMap)
 {
-	ExponentialFogParameters.Bind(ParameterMap,TEXT("SharedFogParameter0"));
-	ExponentialFogColorParameter.Bind(ParameterMap,TEXT("SharedFogParameter1"));
+	ExponentialFogParameters.Bind(ParameterMap,TEXT("ExponentialFogParameters"));
+	ExponentialFogColorParameter.Bind(ParameterMap,TEXT("ExponentialFogColorParameter"));
+	ExponentialFogParameters3.Bind(ParameterMap,TEXT("ExponentialFogParameters3"));
 	InscatteringLightDirection.Bind(ParameterMap,TEXT("InscatteringLightDirection"));
 	DirectionalInscatteringColor.Bind(ParameterMap,TEXT("DirectionalInscatteringColor"));
 	DirectionalInscatteringStartDistance.Bind(ParameterMap,TEXT("DirectionalInscatteringStartDistance"));
@@ -50,6 +51,7 @@ FArchive& operator<<(FArchive& Ar,FExponentialHeightFogShaderParameters& Paramet
 {
 	Ar << Parameters.ExponentialFogParameters;
 	Ar << Parameters.ExponentialFogColorParameter;
+	Ar << Parameters.ExponentialFogParameters3;
 	Ar << Parameters.InscatteringLightDirection;
 	Ar << Parameters.DirectionalInscatteringColor;
 	Ar << Parameters.DirectionalInscatteringStartDistance;
@@ -249,6 +251,7 @@ void FSceneRenderer::InitFogConstants()
 				View.ExponentialFogParameters = FVector4(CollapsedFogParameter, FogInfo.FogHeightFalloff, CosTerminatorAngle, FogInfo.StartDistance);
 				View.ExponentialFogColor = FVector(FogInfo.FogColor.R, FogInfo.FogColor.G, FogInfo.FogColor.B);
 				View.FogMaxOpacity = FogInfo.FogMaxOpacity;
+				View.ExponentialFogParameters3 = FVector2D(FogInfo.FogDensity, FogInfo.FogHeight);
 
 				View.DirectionalInscatteringExponent = FogInfo.DirectionalInscatteringExponent;
 				View.DirectionalInscatteringStartDistance = FogInfo.DirectionalInscatteringStartDistance;
diff --git a/Engine/Source/Runtime/Renderer/Private/ForwardShadingRenderer.cpp b/Engine/Source/Runtime/Renderer/Private/ForwardShadingRenderer.cpp
--- a/Engine/Source/Runtime/Renderer/Private/ForwardShadingRenderer.cpp
+++ b/Engine/Source/Runtime/Renderer/Private/ForwardShadingRenderer.cpp
@@ -42,7 +42,7 @@ void FForwardShadingSceneRenderer::InitViews(FRHICommandListImmediate& RHICmdLis
 	ComputeViewVisibility(RHICmdList);
 	PostVisibilityFrameSetup(ILCTaskData);
 
-	bool bDynamicShadows = ViewFamily.EngineShowFlags.DynamicShadows && GetShadowQuality() > 0;
+	bool bDynamicShadows = ViewFamily.EngineShowFlags.DynamicShadows;
 
 	if (bDynamicShadows && !IsSimpleForwardShadingEnabled(GetFeatureLevelShaderPlatform(FeatureLevel)))
 	{
@@ -336,7 +336,7 @@ void FForwardShadingSceneRenderer::ConditionalResolveSceneDepth(FRHICommandListI
 		{
 			// Only these features require depth texture
 			bool bDecals = ViewFamily.EngineShowFlags.Decals && Scene->Decals.Num();
-			bool bModulatedShadows = ViewFamily.EngineShowFlags.DynamicShadows && GetShadowQuality() > 0 && bModulatedShadowsInUse;
+			bool bModulatedShadows = ViewFamily.EngineShowFlags.DynamicShadows && bModulatedShadowsInUse;
 
 			if (bDecals || bModulatedShadows)
 			{
diff --git a/Engine/Source/Runtime/Renderer/Private/LightRendering.cpp b/Engine/Source/Runtime/Renderer/Private/LightRendering.cpp
--- a/Engine/Source/Runtime/Renderer/Private/LightRendering.cpp
+++ b/Engine/Source/Runtime/Renderer/Private/LightRendering.cpp
@@ -333,7 +333,7 @@ void FDeferredShadingSceneRenderer::RenderLights(FRHICommandListImmediate& RHICm
 	TArray<FSortedLightSceneInfo, SceneRenderingAllocator> SortedLights;
 	SortedLights.Empty(Scene->Lights.Num());
 
-	bool bDynamicShadows = ViewFamily.EngineShowFlags.DynamicShadows && GetShadowQuality() > 0;
+	bool bDynamicShadows = ViewFamily.EngineShowFlags.DynamicShadows;
 	
 	// Build a list of visible lights.
 	for (TSparseArray<FLightSceneInfoCompact>::TConstIterator LightIt(Scene->Lights); LightIt; ++LightIt)
diff --git a/Engine/Source/Runtime/Renderer/Private/LightRendering.h b/Engine/Source/Runtime/Renderer/Private/LightRendering.h
--- a/Engine/Source/Runtime/Renderer/Private/LightRendering.h
+++ b/Engine/Source/Runtime/Renderer/Private/LightRendering.h
@@ -20,9 +20,10 @@ BEGIN_UNIFORM_BUFFER_STRUCT(FDeferredLightUniformStruct,)
 	DECLARE_UNIFORM_BUFFER_STRUCT_MEMBER(float,SourceRadius)
 	DECLARE_UNIFORM_BUFFER_STRUCT_MEMBER(float,SourceLength)
 	DECLARE_UNIFORM_BUFFER_STRUCT_MEMBER(float,MinRoughness)
+	DECLARE_UNIFORM_BUFFER_STRUCT_MEMBER(float,ContactShadowLength)
 	DECLARE_UNIFORM_BUFFER_STRUCT_MEMBER(FVector2D,DistanceFadeMAD)
 	DECLARE_UNIFORM_BUFFER_STRUCT_MEMBER(FVector4,ShadowMapChannelMask)
-	DECLARE_UNIFORM_BUFFER_STRUCT_MEMBER(uint32,bShadowed)
+	DECLARE_UNIFORM_BUFFER_STRUCT_MEMBER(uint32,ShadowedBits)
 	DECLARE_UNIFORM_BUFFER_STRUCT_MEMBER(uint32,LightingChannelMask)
 END_UNIFORM_BUFFER_STRUCT(FDeferredLightUniformStruct)
 
@@ -80,7 +81,10 @@ void SetDeferredLightParameters(
 
 	const bool bDynamicShadows = View.Family->EngineShowFlags.DynamicShadows && GetShadowQuality() > 0;
 	const bool bHasLightFunction = LightSceneInfo->Proxy->GetLightFunctionMaterial() != NULL;
-	DeferredLightUniformsValue.bShadowed = ((LightSceneInfo->Proxy->CastsDynamicShadow() || LightSceneInfo->Proxy->CastsStaticShadow()) && bDynamicShadows) || bHasLightFunction;
+	DeferredLightUniformsValue.ShadowedBits  = LightSceneInfo->Proxy->CastsStaticShadow() || bHasLightFunction ? 1 : 0;
+	DeferredLightUniformsValue.ShadowedBits |= LightSceneInfo->Proxy->CastsDynamicShadow() && View.Family->EngineShowFlags.DynamicShadows ? 3 : 0;
+
+	DeferredLightUniformsValue.ContactShadowLength = LightSceneInfo->Proxy->GetContactShadowLength();
 
 	if( LightSceneInfo->Proxy->IsInverseSquared() )
 	{
diff --git a/Engine/Source/Runtime/Renderer/Private/MaterialTexCoordScalesRendering.cpp b/Engine/Source/Runtime/Renderer/Private/MaterialTexCoordScalesRendering.cpp
--- a/Engine/Source/Runtime/Renderer/Private/MaterialTexCoordScalesRendering.cpp
+++ b/Engine/Source/Runtime/Renderer/Private/MaterialTexCoordScalesRendering.cpp
@@ -36,7 +36,7 @@ void FMaterialTexCoordScalePS::SetParameters(
 		SetShaderValue(RHICmdList, FMeshMaterialShader::GetPixelShader(), AccuracyColorsParameter, FLinearColor::Black, ColorIndex);
 	}
 
-	FMeshMaterialShader::SetParameters(RHICmdList, FMeshMaterialShader::GetPixelShader(), MaterialRenderProxy, Material, View, ESceneRenderTargetsMode::DontSet);
+	FMeshMaterialShader::SetParameters(RHICmdList, FMeshMaterialShader::GetPixelShader(), MaterialRenderProxy, Material, View, ESceneRenderTargetsMode::SetTextures);
 }
 
 void FMaterialTexCoordScalePS::SetMesh(
diff --git a/Engine/Source/Runtime/Renderer/Private/PlanarReflectionRendering.cpp b/Engine/Source/Runtime/Renderer/Private/PlanarReflectionRendering.cpp
--- a/Engine/Source/Runtime/Renderer/Private/PlanarReflectionRendering.cpp
+++ b/Engine/Source/Runtime/Renderer/Private/PlanarReflectionRendering.cpp
@@ -145,6 +145,8 @@ void PrefilterPlanarReflection(FRHICommandListImmediate& RHICmdList, FViewInfo&
 	}
 }
 
+extern float GetSceneColorClearAlpha();
+
 static void UpdatePlanarReflectionContents_RenderThread(
 	FRHICommandListImmediate& RHICmdList, 
 	FSceneRenderer* MainSceneRenderer, 
@@ -210,6 +212,7 @@ static void UpdatePlanarReflectionContents_RenderThread(
 				SetRenderTarget(RHICmdList, Target->GetRenderTargetTexture(), NULL, true);
 
 				// Note: relying on GBuffer SceneColor alpha being cleared to 1 in the main scene rendering
+				check(GetSceneColorClearAlpha() == 1.0f);
 				RHICmdList.Clear(true, FLinearColor(0, 0, 0, 1), false, (float)ERHIZBuffer::FarPlane, false, 0, FIntRect());
 
 				// Render the scene normally
@@ -233,7 +236,7 @@ static void UpdatePlanarReflectionContents_RenderThread(
 	}
 }
 
-extern void BuildProjectionMatrix(FIntPoint RenderTargetSize, float FOV, FMatrix& ProjectionMatrix);
+extern void BuildProjectionMatrix(FIntPoint RenderTargetSize, ECameraProjectionMode::Type ProjectionType, float FOV, float OrthoWidth, FMatrix& ProjectionMatrix);
 
 extern FSceneRenderer* CreateSceneRendererForSceneCapture(
 	FScene* Scene,
@@ -309,7 +312,7 @@ void FScene::UpdatePlanarReflectionContents(UPlanarReflectionComponent* CaptureC
 		float FOV = FMath::Atan(1.0f / ParentView.ViewMatrices.ProjMatrix.M[0][0]);
 
 		FMatrix ProjectionMatrix;
-		BuildProjectionMatrix(DesiredPlanarReflectionTextureSize, FOV + CaptureComponent->ExtraFOV * (float)PI / 180.0f, ProjectionMatrix);
+		BuildProjectionMatrix(DesiredPlanarReflectionTextureSize, ECameraProjectionMode::Perspective, FOV + CaptureComponent->ExtraFOV * (float)PI / 180.0f, 1.0f, ProjectionMatrix);
 
 		FSceneRenderer* SceneRenderer = CreateSceneRendererForSceneCapture(this, CaptureComponent, CaptureComponent->RenderTarget, DesiredPlanarReflectionTextureSize, ViewRotationMatrix, ViewLocation, ProjectionMatrix, CaptureComponent->MaxViewDistanceOverride, true, true, &PostProcessSettings, 1.0f);
 
diff --git a/Engine/Source/Runtime/Renderer/Private/PostProcess/PostProcessCombineLUTs.cpp b/Engine/Source/Runtime/Renderer/Private/PostProcess/PostProcessCombineLUTs.cpp
--- a/Engine/Source/Runtime/Renderer/Private/PostProcess/PostProcessCombineLUTs.cpp
+++ b/Engine/Source/Runtime/Renderer/Private/PostProcess/PostProcessCombineLUTs.cpp
@@ -126,7 +126,11 @@ class FLUTBlenderPS : public FGlobalShader
 	DECLARE_SHADER_TYPE(FLUTBlenderPS,Global);
 public:
 
-	static bool ShouldCache(EShaderPlatform Platform) { return IsFeatureLevelSupported(Platform, ERHIFeatureLevel::ES2); }
+	static bool ShouldCache(EShaderPlatform Platform)
+	{
+		//#todo-rco: Glslang issue
+		return IsFeatureLevelSupported(Platform, ERHIFeatureLevel::ES2) && !IsVulkanPlatform(Platform);
+	}
 
 	FLUTBlenderPS(const ShaderMetaType::CompiledShaderInitializerType& Initializer)
 		: FGlobalShader(Initializer) 
