Copying //UE4/Orion-Staging to //UE4/Main (Source: //Orion/Dev-General @ 2997507)

diff --git a/Engine/Source/Runtime/Engine/Classes/Engine/NetConnection.h b/Engine/Source/Runtime/Engine/Classes/Engine/NetConnection.h
--- a/Engine/Source/Runtime/Engine/Classes/Engine/NetConnection.h
+++ b/Engine/Source/Runtime/Engine/Classes/Engine/NetConnection.h
@@ -678,6 +678,8 @@ public:
 	/** Removes a channel from the ticking list directly */
 	void StopTickingChannel(UChannel* Channel) { ChannelsToTick.Remove(Channel); }
 
+	FORCEINLINE FHistogram GetNetHistogram() const { return NetConnectionHistogram; }
+
 protected:
 
 	void CleanupDormantActorState();
@@ -691,6 +693,9 @@ private:
 	 */
 	UPROPERTY()
 	TArray<UChannel*> ChannelsToTick;
+
+	/** Histogram of the received packet time */
+	FHistogram NetConnectionHistogram;
 };
 
 
diff --git a/Engine/Source/Runtime/Engine/Classes/Engine/Texture.h b/Engine/Source/Runtime/Engine/Classes/Engine/Texture.h
--- a/Engine/Source/Runtime/Engine/Classes/Engine/Texture.h
+++ b/Engine/Source/Runtime/Engine/Classes/Engine/Texture.h
@@ -337,6 +337,7 @@ struct FTexturePlatformData
 	bool AreDerivedMipsAvailable() const;
 #endif
 
+	int32 GetNumNonStreamingMips() const;
 };
 
 UCLASS(abstract, MinimalAPI, BlueprintType)
@@ -705,7 +706,7 @@ public:
 	ENGINE_API virtual void Serialize(FArchive& Ar) override;
 	ENGINE_API virtual void PostInitProperties() override;
 	ENGINE_API virtual void PostLoad() override;
-	ENGINE_API virtual void PreSave() override;
+	ENGINE_API virtual void PreSave(const class ITargetPlatform* TargetPlatform) override;
 	ENGINE_API virtual void BeginDestroy() override;
 	ENGINE_API virtual bool IsReadyForFinishDestroy() override;
 	ENGINE_API virtual void FinishDestroy() override;
diff --git a/Engine/Source/Runtime/Engine/Classes/GameFramework/Character.h b/Engine/Source/Runtime/Engine/Classes/GameFramework/Character.h
--- a/Engine/Source/Runtime/Engine/Classes/GameFramework/Character.h
+++ b/Engine/Source/Runtime/Engine/Classes/GameFramework/Character.h
@@ -343,6 +343,9 @@ public:
 	UPROPERTY(BlueprintReadOnly, replicatedUsing=OnRep_IsCrouched, Category=Character)
 	uint32 bIsCrouched:1;
 
+	UPROPERTY( replicated )
+	uint32 bReplayHasRootMotionSources:1;
+
 	/** Handle Crouching replicated from server */
 	UFUNCTION()
 	virtual void OnRep_IsCrouched();
diff --git a/Engine/Source/Runtime/Engine/Classes/GameFramework/PlayerInput.h b/Engine/Source/Runtime/Engine/Classes/GameFramework/PlayerInput.h
--- a/Engine/Source/Runtime/Engine/Classes/GameFramework/PlayerInput.h
+++ b/Engine/Source/Runtime/Engine/Classes/GameFramework/PlayerInput.h
@@ -29,38 +29,38 @@ struct FKeyBind
 
 	/** Whether the control key needs to be held when the key event occurs */
 	UPROPERTY(config)
-	uint32 Control:1;
+	uint8 Control:1;
 
 	/** Whether the shift key needs to be held when the key event occurs */
 	UPROPERTY(config)
-	uint32 Shift:1;
+	uint8 Shift:1;
 
 	/** Whether the alt key needs to be held when the key event occurs */
 	UPROPERTY(config)
-	uint32 Alt:1;
+	uint8 Alt:1;
 
 	/** Whether the command key needs to be held when the key event occurs */
 	UPROPERTY(config)
-	uint32 Cmd:1;
+	uint8 Cmd:1;
 
 	/** Whether the control key must not be held when the key event occurs */
 	UPROPERTY(config)
-	uint32 bIgnoreCtrl:1;
+	uint8 bIgnoreCtrl:1;
 
 	/** Whether the shift key must not be held when the key event occurs */
 	UPROPERTY(config)
-	uint32 bIgnoreShift:1;
+	uint8 bIgnoreShift:1;
 
 	/** Whether the alt key must not be held when the key event occurs */
 	UPROPERTY(config)
-	uint32 bIgnoreAlt:1;
+	uint8 bIgnoreAlt:1;
 
 	/** Whether the command key must not be held when the key event occurs */
 	UPROPERTY(config)
-	uint32 bIgnoreCmd:1;
+	uint8 bIgnoreCmd:1;
 
 	UPROPERTY(transient)
-	uint32 bDisabled : 1;
+	uint8 bDisabled : 1;
 
 	FKeyBind()
 	{
@@ -87,7 +87,7 @@ struct FInputAxisProperties
 
 	/** Inverts reported values for this axis */
 	UPROPERTY(EditAnywhere, Category="Input")
-	uint32 bInvert:1;
+	uint8 bInvert:1;
 
 	FInputAxisProperties()
 		: DeadZone(0.2f)
@@ -133,19 +133,19 @@ struct FInputActionKeyMapping
 
 	/** true if one of the Shift keys must be down when the KeyEvent is received to be acknowledged */
 	UPROPERTY(EditAnywhere, BlueprintReadWrite, Category="Input")
-	uint32 bShift:1;
+	uint8 bShift:1;
 
 	/** true if one of the Ctrl keys must be down when the KeyEvent is received to be acknowledged */
 	UPROPERTY(EditAnywhere, BlueprintReadWrite, Category="Input")
-	uint32 bCtrl:1;
+	uint8 bCtrl:1;
 
 	/** true if one of the Alt keys must be down when the KeyEvent is received to be acknowledged */
 	UPROPERTY(EditAnywhere, BlueprintReadWrite, Category="Input")
-	uint32 bAlt:1;
+	uint8 bAlt:1;
 
 	/** true if one of the Cmd keys must be down when the KeyEvent is received to be acknowledged */
 	UPROPERTY(EditAnywhere, BlueprintReadWrite, Category="Input")
-	uint32 bCmd:1;
+	uint8 bCmd:1;
 
 	bool operator==(const FInputActionKeyMapping& Other) const
 	{
@@ -373,7 +373,7 @@ private:
 		TArray<FInputAxisKeyMapping> KeyMappings;
 
 		/** Whether this axis should invert its outputs */
-		uint32 bInverted:1;
+		uint8 bInverted:1;
 
 		FAxisKeyDetails()
 			: bInverted(false)
@@ -393,6 +393,8 @@ private:
 	/** The current game view of each key */
 	TMap<FKey,FKeyState> KeyStateMap;
 
+	uint8 bKeyMapsBuilt:1;
+
 public:
 	
 	//~ Begin UObject Interface
@@ -564,7 +566,15 @@ private:
 	float DetermineAxisValue(const FInputAxisBinding& AxisBinding, const bool bGamePaused, TArray<FKey>& KeysToConsume);
 
 	/** Utility function to ensure the key mapping cache maps are built */
-	void ConditionalBuildKeyMappings();
+	FORCEINLINE void ConditionalBuildKeyMappings()
+	{
+		if (!bKeyMapsBuilt)
+		{
+			ConditionalBuildKeyMappings_Internal();
+		}
+	}
+
+	void ConditionalBuildKeyMappings_Internal();
 
 	/** Set the Key consumed for the frame so that subsequent input components will not be notified they were pressed */
 	void ConsumeKey(FKey Key);
@@ -592,6 +602,9 @@ private:
 	/** Axis Mappings defined by engine systems that cannot be remapped by users */
 	static TArray<FInputAxisKeyMapping> EngineDefinedAxisMappings;
 
+	// Temporary array used as part of input processing
+	TArray<uint32> EventIndices;
+
 	/** A counter used to track the order in which events occurred since the last time the input stack was processed */
 	uint32 EventCount;
 
diff --git a/Engine/Source/Runtime/Engine/Classes/GameFramework/RootMotionSource.h b/Engine/Source/Runtime/Engine/Classes/GameFramework/RootMotionSource.h
--- a/Engine/Source/Runtime/Engine/Classes/GameFramework/RootMotionSource.h
+++ b/Engine/Source/Runtime/Engine/Classes/GameFramework/RootMotionSource.h
@@ -432,6 +432,12 @@ struct ENGINE_API FRootMotionSource_RadialForce : public FRootMotionSource
 	UPROPERTY()
 	UCurveFloat* StrengthOverTime;
 
+	UPROPERTY()
+	bool bUseFixedWorldDirection;
+
+	UPROPERTY()
+	FRotator FixedWorldDirection;
+
 	virtual FRootMotionSource* Clone() const override;
 
 	virtual bool Matches(const FRootMotionSource* Other) const override;
diff --git a/Engine/Source/Runtime/Engine/Classes/Materials/Material.h b/Engine/Source/Runtime/Engine/Classes/Materials/Material.h
--- a/Engine/Source/Runtime/Engine/Classes/Materials/Material.h
+++ b/Engine/Source/Runtime/Engine/Classes/Materials/Material.h
@@ -862,7 +862,7 @@ public:
 	//~ End UMaterialInterface Interface.
 
 	//~ Begin UObject Interface
-	ENGINE_API virtual void PreSave() override;
+	ENGINE_API virtual void PreSave(const class ITargetPlatform* TargetPlatform) override;
 	ENGINE_API virtual void PostInitProperties() override;
 	ENGINE_API virtual void Serialize(FArchive& Ar) override;
 	ENGINE_API virtual void PostDuplicate(bool bDuplicateForPIE) override;
diff --git a/Engine/Source/Runtime/Engine/Private/AI/Navigation/NavLinkProxy.cpp b/Engine/Source/Runtime/Engine/Private/AI/Navigation/NavLinkProxy.cpp
--- a/Engine/Source/Runtime/Engine/Private/AI/Navigation/NavLinkProxy.cpp
+++ b/Engine/Source/Runtime/Engine/Private/AI/Navigation/NavLinkProxy.cpp
@@ -55,10 +55,13 @@ ANavLinkProxy::ANavLinkProxy(const FObjectInitializer& ObjectInitializer) : Supe
 	}
 #endif
 
-	SmartLinkComp = CreateDefaultSubobject<UNavLinkCustomComponent>(TEXT("SmartLinkComp"));
-	SmartLinkComp->SetNavigationRelevancy(false);
-	SmartLinkComp->SetMoveReachedLink(this, &ANavLinkProxy::NotifySmartLinkReached);
-	bSmartLinkIsRelevant = false;
+	if (HasAnyFlags(RF_ClassDefaultObject) == false)
+	{
+		SmartLinkComp = CreateDefaultSubobject<UNavLinkCustomComponent>(TEXT("SmartLinkComp"));
+		SmartLinkComp->SetNavigationRelevancy(false);
+		SmartLinkComp->SetMoveReachedLink(this, &ANavLinkProxy::NotifySmartLinkReached);
+		bSmartLinkIsRelevant = false;
+	}
 
 	PointLinks.Add(FNavigationLink());
 	SetActorEnableCollision(false);
diff --git a/Engine/Source/Runtime/Engine/Private/AI/Navigation/NavigationPath.cpp b/Engine/Source/Runtime/Engine/Private/AI/Navigation/NavigationPath.cpp
--- a/Engine/Source/Runtime/Engine/Private/AI/Navigation/NavigationPath.cpp
+++ b/Engine/Source/Runtime/Engine/Private/AI/Navigation/NavigationPath.cpp
@@ -103,14 +103,17 @@ void FNavigationPath::SetGoalActorObservation(const AActor& ActorToObserve, floa
 	}
 
 	// register for path observing only if we weren't registered already
-	const bool RegisterForPathUpdates = GoalActor.IsValid() == false;	
+	const bool RegisterForPathUpdates = (GoalActor.IsValid() == false);
 	GoalActor = &ActorToObserve;
 	checkSlow(GoalActor.IsValid());
 	GoalActorAsNavAgent = Cast<INavAgentInterface>(&ActorToObserve);
 	GoalActorLocationTetherDistanceSq = FMath::Square(TetherDistance);
 	UpdateLastRepathGoalLocation();
 
-	NavigationDataUsed->RegisterObservedPath(AsShared());
+	if (RegisterForPathUpdates)
+	{
+		NavigationDataUsed->RegisterObservedPath(AsShared());
+	}
 }
 
 void FNavigationPath::SetSourceActor(const AActor& InSourceActor)
diff --git a/Engine/Source/Runtime/Engine/Private/AI/Navigation/NavigationSystem.cpp b/Engine/Source/Runtime/Engine/Private/AI/Navigation/NavigationSystem.cpp
--- a/Engine/Source/Runtime/Engine/Private/AI/Navigation/NavigationSystem.cpp
+++ b/Engine/Source/Runtime/Engine/Private/AI/Navigation/NavigationSystem.cpp
@@ -262,6 +262,7 @@ FNavigationSystemExec UNavigationSystem::ExecHandler;
 UNavigationSystem::FOnNavigationDirty UNavigationSystem::NavigationDirtyEvent;
 
 bool UNavigationSystem::bUpdateNavOctreeOnComponentChange = true;
+bool UNavigationSystem::bStaticRuntimeNavigation = false;
 //----------------------------------------------------------------------//
 // life cycle stuff                                                                
 //----------------------------------------------------------------------//
@@ -326,6 +327,11 @@ UNavigationSystem::~UNavigationSystem()
 #endif // WITH_EDITOR
 }
 
+void UNavigationSystem::ConfigureAsStatic()
+{
+	bStaticRuntimeNavigation = true;
+}
+
 void UNavigationSystem::DoInitialSetup()
 {
 	if (bInitialSetupHasBeenPerformed)
@@ -2473,6 +2479,10 @@ const FNavigationRelevantData* UNavigationSystem::GetDataForObject(const UObject
 
 void UNavigationSystem::UpdateActorInNavOctree(AActor& Actor)
 {
+	if (IsNavigationSystemStatic())
+	{
+		return;
+	}
 	SCOPE_CYCLE_COUNTER(STAT_DebugNavOctree);
 
 	INavRelevantInterface* NavElement = Cast<INavRelevantInterface>(&Actor);
@@ -2725,6 +2735,11 @@ bool UNavigationSystem::UpdateNavOctreeElementBounds(UActorComponent* Comp, cons
 
 void UNavigationSystem::OnComponentRegistered(UActorComponent* Comp)
 {
+	if (IsNavigationSystemStatic())
+	{
+		return;
+	}
+
 	SCOPE_CYCLE_COUNTER(STAT_DebugNavOctree);
 	INavRelevantInterface* NavInterface = Cast<INavRelevantInterface>(Comp);
 	if (NavInterface)
@@ -2743,6 +2758,11 @@ void UNavigationSystem::OnComponentRegistered(UActorComponent* Comp)
 
 void UNavigationSystem::OnComponentUnregistered(UActorComponent* Comp)
 {
+	if (IsNavigationSystemStatic())
+	{
+		return;
+	}
+
 	SCOPE_CYCLE_COUNTER(STAT_DebugNavOctree);
 	INavRelevantInterface* NavInterface = Cast<INavRelevantInterface>(Comp);
 	if (NavInterface)
@@ -2763,6 +2783,11 @@ void UNavigationSystem::OnComponentUnregistered(UActorComponent* Comp)
 
 void UNavigationSystem::OnActorRegistered(AActor* Actor)
 {
+	if (IsNavigationSystemStatic())
+	{
+		return;
+	}
+
 	SCOPE_CYCLE_COUNTER(STAT_DebugNavOctree);
 	INavRelevantInterface* NavInterface = Cast<INavRelevantInterface>(Actor);
 	if (NavInterface)
@@ -2777,6 +2802,11 @@ void UNavigationSystem::OnActorRegistered(AActor* Actor)
 
 void UNavigationSystem::OnActorUnregistered(AActor* Actor)
 {
+	if (IsNavigationSystemStatic())
+	{
+		return;
+	}
+
 	SCOPE_CYCLE_COUNTER(STAT_DebugNavOctree);
 	INavRelevantInterface* NavInterface = Cast<INavRelevantInterface>(Actor);
 	if (NavInterface)
@@ -2814,6 +2844,12 @@ void UNavigationSystem::ReleaseInitialBuildingLock()
 
 void UNavigationSystem::InitializeLevelCollisions()
 {
+	if (IsNavigationSystemStatic())
+	{
+		bInitialLevelsAdded = true;
+		return;
+	}
+
 	UWorld* World = GetWorld();
 	if (!bInitialLevelsAdded && UNavigationSystem::GetCurrent(World) == this)
 	{
@@ -2867,7 +2903,7 @@ void UNavigationSystem::OnEditorModeChanged(FEdMode* Mode, bool IsEntering)
 
 void UNavigationSystem::OnNavigationBoundsUpdated(ANavMeshBoundsVolume* NavVolume)
 {
-	if (NavVolume == NULL)
+	if (NavVolume == nullptr || IsNavigationSystemStatic())
 	{
 		return;
 	}
@@ -2884,7 +2920,7 @@ void UNavigationSystem::OnNavigationBoundsUpdated(ANavMeshBoundsVolume* NavVolum
 
 void UNavigationSystem::OnNavigationBoundsAdded(ANavMeshBoundsVolume* NavVolume)
 {
-	if (NavVolume == NULL)
+	if (NavVolume == nullptr || IsNavigationSystemStatic())
 	{
 		return;
 	}
@@ -2901,7 +2937,7 @@ void UNavigationSystem::OnNavigationBoundsAdded(ANavMeshBoundsVolume* NavVolume)
 
 void UNavigationSystem::OnNavigationBoundsRemoved(ANavMeshBoundsVolume* NavVolume)
 {
-	if (NavVolume == NULL)
+	if (NavVolume == nullptr || IsNavigationSystemStatic())
 	{
 		return;
 	}
@@ -3335,7 +3371,7 @@ int32 UNavigationSystem::GetNumRunningBuildTasks() const
 
 void UNavigationSystem::OnLevelAddedToWorld(ULevel* InLevel, UWorld* InWorld)
 {
-	if (InWorld == GetWorld())
+	if ((IsNavigationSystemStatic() == false) && (InWorld == GetWorld()))
 	{
 		AddLevelCollisionToOctree(InLevel);
 
@@ -3354,7 +3390,7 @@ void UNavigationSystem::OnLevelAddedToWorld(ULevel* InLevel, UWorld* InWorld)
 
 void UNavigationSystem::OnLevelRemovedFromWorld(ULevel* InLevel, UWorld* InWorld)
 {
-	if (InWorld == GetWorld())
+	if ((IsNavigationSystemStatic() == false) && (InWorld == GetWorld()))
 	{
 		RemoveLevelCollisionFromOctree(InLevel);
 
@@ -3715,7 +3751,7 @@ bool UNavigationSystem::IsNavigationBeingBuilt(UObject* WorldContextObject)
 //----------------------------------------------------------------------//
 bool UNavigationSystem::ShouldGeneratorRun(const FNavDataGenerator* Generator) const
 {
-	if (Generator != NULL)
+	if (Generator != NULL && (IsNavigationSystemStatic() == false))
 	{
 		for (int32 NavDataIndex = 0; NavDataIndex < NavDataSet.Num(); ++NavDataIndex)
 		{
diff --git a/Engine/Source/Runtime/Engine/Private/Animation/AnimSequence.cpp b/Engine/Source/Runtime/Engine/Private/Animation/AnimSequence.cpp
--- a/Engine/Source/Runtime/Engine/Private/Animation/AnimSequence.cpp
+++ b/Engine/Source/Runtime/Engine/Private/Animation/AnimSequence.cpp
@@ -306,7 +306,7 @@ void UAnimSequence::SortSyncMarkers()
 	RefreshSyncMarkerDataFromAuthored();
 }
 
-void UAnimSequence::PreSave()
+void UAnimSequence::PreSave(const class ITargetPlatform* TargetPlatform)
 {
 #if WITH_EDITOR
 	// we have to bake it if it's not baked
@@ -327,7 +327,7 @@ void UAnimSequence::PreSave()
 	}
 #endif
 
-	Super::PreSave();
+	Super::PreSave(TargetPlatform);
 }
 
 void UAnimSequence::PostLoad()
diff --git a/Engine/Source/Runtime/Engine/Private/Animation/AnimationRuntime.cpp b/Engine/Source/Runtime/Engine/Private/Animation/AnimationRuntime.cpp
--- a/Engine/Source/Runtime/Engine/Private/Animation/AnimationRuntime.cpp
+++ b/Engine/Source/Runtime/Engine/Private/Animation/AnimationRuntime.cpp
@@ -782,6 +782,11 @@ void FAnimationRuntime::ConvertPoseToMeshSpace(const TArray<FTransform>& LocalTr
 	}
 }
 
+struct FEnsureParentsPresentScratchArea : public TThreadSingleton<FEnsureParentsPresentScratchArea>
+{
+	TArray<bool> BoneExists;
+};
+
 /** 
  *	Utility for taking an array of bone indices and ensuring that all parents are present 
  *	(ie. all bones between those in the array and the root are present). 
@@ -792,6 +797,11 @@ void FAnimationRuntime::EnsureParentsPresent(TArray<FBoneIndexType>& BoneIndices
 	const int32 NumBones = SkelMesh->RefSkeleton.GetNum();
 	// Iterate through existing array.
 	int32 i=0;
+
+	TArray<bool>& BoneExists = FEnsureParentsPresentScratchArea::Get().BoneExists;
+	BoneExists.Reset();
+	BoneExists.SetNumZeroed(NumBones);
+
 	while( i<BoneIndices.Num() )
 	{
 		const int32 BoneIndex = BoneIndices[i];
@@ -810,15 +820,18 @@ void FAnimationRuntime::EnsureParentsPresent(TArray<FBoneIndexType>& BoneIndices
 				continue;
 			}
 #endif
+			BoneExists[BoneIndex] = true;
+
 			const int32 ParentIndex = SkelMesh->RefSkeleton.GetParentIndex(BoneIndex);
 
 			// If we do not have this parent in the array, we add it in this location, and leave 'i' where it is.
 			// This can happen if somebody removes bones in the physics asset, then it will try add back in, and in the process, 
 			// parent can be missing
-			if( !BoneIndices.Contains(ParentIndex) )
+			if (!BoneExists[ParentIndex])
 			{
 				BoneIndices.InsertUninitialized(i);
 				BoneIndices[i] = ParentIndex;
+				BoneExists[ParentIndex] = true;
 			}
 			// If parent was in array, just move on.
 			else
@@ -828,6 +841,7 @@ void FAnimationRuntime::EnsureParentsPresent(TArray<FBoneIndexType>& BoneIndices
 		}
 		else
 		{
+			BoneExists[0] = true;
 			i++;
 		}
 	}
diff --git a/Engine/Source/Runtime/Engine/Private/AnimationAsset.cpp b/Engine/Source/Runtime/Engine/Private/AnimationAsset.cpp
--- a/Engine/Source/Runtime/Engine/Private/AnimationAsset.cpp
+++ b/Engine/Source/Runtime/Engine/Private/AnimationAsset.cpp
@@ -384,6 +384,11 @@ void FBoneContainer::InitializeTo(const TArray<FBoneIndexType>& InRequiredBoneIn
 	Initialize();
 }
 
+struct FBoneContainerScratchArea : public TThreadSingleton<FBoneContainerScratchArea>
+{
+	TArray<int32> MeshIndexToCompactPoseIndex;
+};
+
 void FBoneContainer::Initialize()
 {
 	RefSkeleton = NULL;
@@ -429,7 +434,7 @@ void FBoneContainer::Initialize()
 	}
 
 	// Clear remapping table
-	SkeletonToPoseBoneIndexArray.Empty();
+	SkeletonToPoseBoneIndexArray.Reset();
 
 	// Cache our mapping tables
 	// Here we create look up tables between our target asset and its USkeleton's refpose.
@@ -447,38 +452,56 @@ void FBoneContainer::Initialize()
 
 	//Set up compact pose data
 	int32 NumReqBones = BoneIndicesArray.Num();
-	CompactPoseParentBones.Empty(NumReqBones);
-	
-	CompactPoseRefPoseBones.Empty(NumReqBones);
+	CompactPoseParentBones.Reset(NumReqBones);
+
+	CompactPoseRefPoseBones.Reset(NumReqBones);
 	CompactPoseRefPoseBones.AddUninitialized(NumReqBones);
-	
-	CompactPoseToSkeletonIndex.Empty(NumReqBones);
+
+	CompactPoseToSkeletonIndex.Reset(NumReqBones);
 	CompactPoseToSkeletonIndex.AddUninitialized(NumReqBones);
-	
-	SkeletonToCompactPose.Empty(SkeletonToPoseBoneIndexArray.Num());
+
+	SkeletonToCompactPose.Reset(SkeletonToPoseBoneIndexArray.Num());
 
 	const TArray<FTransform>& RefPoseArray = RefSkeleton->GetRefBonePose();
+	TArray<int32>& MeshIndexToCompactPoseIndex = FBoneContainerScratchArea::Get().MeshIndexToCompactPoseIndex;
+	MeshIndexToCompactPoseIndex.Reset(PoseToSkeletonBoneIndexArray.Num());
+	MeshIndexToCompactPoseIndex.AddUninitialized(PoseToSkeletonBoneIndexArray.Num());
 
+	for (int32& Item : MeshIndexToCompactPoseIndex)
+	{
+		Item = -1;
+	}
+		
 	for (int32 CompactBoneIndex = 0; CompactBoneIndex < NumReqBones; ++CompactBoneIndex)
 	{
 		FBoneIndexType MeshPoseIndex = BoneIndicesArray[CompactBoneIndex];
+		MeshIndexToCompactPoseIndex[MeshPoseIndex] = CompactBoneIndex;
 
 		//Parent Bone
 		const int32 ParentIndex = GetParentBoneIndex(MeshPoseIndex);
-		const int32 CompactParentIndex = ParentIndex == INDEX_NONE ? INDEX_NONE : BoneIndicesArray.IndexOfByKey(ParentIndex);
+		const int32 CompactParentIndex = ParentIndex == INDEX_NONE ? INDEX_NONE : MeshIndexToCompactPoseIndex[ParentIndex];
 
 		CompactPoseParentBones.Add(FCompactPoseBoneIndex(CompactParentIndex));
+	}
 
-		//Ref Pose
+	//Ref Pose
+	for (int32 CompactBoneIndex = 0; CompactBoneIndex < NumReqBones; ++CompactBoneIndex)
+	{
+		FBoneIndexType MeshPoseIndex = BoneIndicesArray[CompactBoneIndex];
 		CompactPoseRefPoseBones[CompactBoneIndex] = RefPoseArray[MeshPoseIndex];
+	}
 
+	for (int32 CompactBoneIndex = 0; CompactBoneIndex < NumReqBones; ++CompactBoneIndex)
+	{
+		FBoneIndexType MeshPoseIndex = BoneIndicesArray[CompactBoneIndex];
 		CompactPoseToSkeletonIndex[CompactBoneIndex] = PoseToSkeletonBoneIndexArray[MeshPoseIndex];
 	}
 
+
 	for (int32 SkeletonBoneIndex = 0; SkeletonBoneIndex < SkeletonToPoseBoneIndexArray.Num(); ++SkeletonBoneIndex)
 	{
 		int32 PoseBoneIndex = SkeletonToPoseBoneIndexArray[SkeletonBoneIndex];
-		int32 CompactIndex  = (PoseBoneIndex != INDEX_NONE) ? (BoneIndicesArray.IndexOfByKey(PoseBoneIndex)) : INDEX_NONE;
+		int32 CompactIndex = (PoseBoneIndex != INDEX_NONE) ? MeshIndexToCompactPoseIndex[PoseBoneIndex] : INDEX_NONE;
 		SkeletonToCompactPose.Add(FCompactPoseBoneIndex(CompactIndex));
 	}
 }
diff --git a/Engine/Source/Runtime/Engine/Private/Blueprint.cpp b/Engine/Source/Runtime/Engine/Private/Blueprint.cpp
--- a/Engine/Source/Runtime/Engine/Private/Blueprint.cpp
+++ b/Engine/Source/Runtime/Engine/Private/Blueprint.cpp
@@ -310,9 +310,9 @@ UBlueprint::UBlueprint(const FObjectInitializer& ObjectInitializer)
 }
 
 #if WITH_EDITORONLY_DATA
-void UBlueprint::PreSave()
+void UBlueprint::PreSave(const class ITargetPlatform* TargetPlatform)
 {
-	Super::PreSave();
+	Super::PreSave(TargetPlatform);
 
 	// Clear all upgrade notes, the user has saved and should not see them anymore
 	UpgradeNotesLog.Reset();
diff --git a/Engine/Source/Runtime/Engine/Private/Camera/CameraAnim.cpp b/Engine/Source/Runtime/Engine/Private/Camera/CameraAnim.cpp
--- a/Engine/Source/Runtime/Engine/Private/Camera/CameraAnim.cpp
+++ b/Engine/Source/Runtime/Engine/Private/Camera/CameraAnim.cpp
@@ -90,12 +90,12 @@ FBox UCameraAnim::GetAABB(FVector const& BaseLoc, FRotator const& BaseRot, float
 }
 
 
-void UCameraAnim::PreSave()
+void UCameraAnim::PreSave(const class ITargetPlatform* TargetPlatform)
 {
 #if WITH_EDITORONLY_DATA
 	CalcLocalAABB();
 #endif // WITH_EDITORONLY_DATA
-	Super::PreSave();
+	Super::PreSave(TargetPlatform);
 }
 
 void UCameraAnim::PostLoad()
diff --git a/Engine/Source/Runtime/Engine/Private/Character.cpp b/Engine/Source/Runtime/Engine/Private/Character.cpp
--- a/Engine/Source/Runtime/Engine/Private/Character.cpp
+++ b/Engine/Source/Runtime/Engine/Private/Character.cpp
@@ -1351,6 +1351,8 @@ void ACharacter::PreReplication( IRepChangedPropertyTracker & ChangedPropertyTra
 
 		ChangedPropertyTracker.SetExternalData( Writer.GetData(), Writer.GetNumBits() );
 	}
+
+	bReplayHasRootMotionSources = CharacterMovement ? CharacterMovement->HasRootMotionSources() : false;
 }
 
 void ACharacter::GetLifetimeReplicatedProps( TArray< FLifetimeProperty > & OutLifetimeProps ) const
@@ -1363,6 +1365,7 @@ void ACharacter::GetLifetimeReplicatedProps( TArray< FLifetimeProperty > & OutLi
 	DOREPLIFETIME_CONDITION( ACharacter, ReplicatedMovementMode,			COND_SimulatedOnly );
 	DOREPLIFETIME_CONDITION( ACharacter, bIsCrouched,						COND_SimulatedOnly );
 	DOREPLIFETIME_CONDITION( ACharacter, AnimRootMotionTranslationScale,	COND_SimulatedOnly );
+	DOREPLIFETIME_CONDITION( ACharacter, bReplayHasRootMotionSources,		COND_ReplayOnly );
 
 	// Change the condition of the replicated movement property to not replicate in replays since we handle this specifically via saving this out in external replay data
 	DOREPLIFETIME_CHANGE_CONDITION( AActor, ReplicatedMovement,				COND_SimulatedOrPhysicsNoReplay );
diff --git a/Engine/Source/Runtime/Engine/Private/ChartCreation.cpp b/Engine/Source/Runtime/Engine/Private/ChartCreation.cpp
--- a/Engine/Source/Runtime/Engine/Private/ChartCreation.cpp
+++ b/Engine/Source/Runtime/Engine/Private/ChartCreation.cpp
@@ -708,6 +708,17 @@ void UEngine::TickFPSChart( float DeltaSeconds )
 		{
 			UE_LOG(LogChartCreation, Warning, TEXT("Idle time for this frame (%f) is larger than delta between FPSChart ticks (%f)"), ThisFrameIdleTime, DeltaSeconds);
 		}
+
+#if USE_SERVER_PERF_COUNTERS
+		if (LIKELY(PerfCounters))
+		{
+			FHistogram* NoSleepHistogram = PerfCounters->PerformanceHistograms().Find(IPerfCounters::Histograms::FrameTimeWithoutSleep);
+			if (LIKELY(NoSleepHistogram))
+			{
+				NoSleepHistogram->AddMeasurement(DeltaSeconds * 1000.0);
+			}
+		}
+#endif // USE_SERVER_PERF_COUNTERS
 	}
 
 	// now gather some stats on what this frame was bound by (game, render, gpu)
@@ -1019,6 +1030,7 @@ void UEngine::StartFPSChart(const FString& Label, bool bRecordPerFrameTimes)
 
 		PerfCounters->PerformanceHistograms().Add(IPerfCounters::Histograms::FrameTime, NewHitchTrackingHistogram);
 		PerfCounters->PerformanceHistograms().Add(IPerfCounters::Histograms::FrameTimePeriodic, NewHitchTrackingHistogram);		
+		PerfCounters->PerformanceHistograms().Add(IPerfCounters::Histograms::FrameTimeWithoutSleep, NewHitchTrackingHistogram);
 		PerfCounters->PerformanceHistograms().Add(IPerfCounters::Histograms::ServerReplicateActorsTime, NewHitchTrackingHistogram);
 		PerfCounters->PerformanceHistograms().Add(IPerfCounters::Histograms::SleepTime, NewHitchTrackingHistogram);
 	}
@@ -1032,15 +1044,6 @@ void UEngine::StopFPSChart()
 	GFPSChartStopTime = FPlatformTime::Seconds();
 	GEnableDataCapture = false;
 
-#if USE_SERVER_PERF_COUNTERS
-	IPerfCounters* PerfCounters = IPerfCountersModule::Get().GetPerformanceCounters();
-	if (LIKELY(PerfCounters))
-	{
-		// remove existing histograms so the code doesn't keep updating them. If you need to use them, do that before calling StopFPSChart() (or copy)
-		PerfCounters->PerformanceHistograms().Reset();
-	}
-#endif // USE_SERVER_PERF_COUNTERS
-
 	UE_LOG(LogChartCreation, Log, TEXT("Stopped creating FPS charts at %f seconds"), GFPSChartStopTime);
 }
 
diff --git a/Engine/Source/Runtime/Engine/Private/Collision/CollisionProfile.cpp b/Engine/Source/Runtime/Engine/Private/Collision/CollisionProfile.cpp
--- a/Engine/Source/Runtime/Engine/Private/Collision/CollisionProfile.cpp
+++ b/Engine/Source/Runtime/Engine/Private/Collision/CollisionProfile.cpp
@@ -61,6 +61,14 @@ UCollisionProfile* UCollisionProfile::Get()
 	return CollisionProfile;
 }
 
+void UCollisionProfile::PostReloadConfig(UProperty* PropertyThatWasLoaded)
+{
+	if (HasAnyFlags(RF_ClassDefaultObject))
+	{
+		LoadProfileConfig();
+	}
+}
+
 void UCollisionProfile::GetProfileNames(TArray<TSharedPtr<FName>>& OutNameList)
 {
 	const UCollisionProfile* CollisionProfile = UCollisionProfile::Get();
diff --git a/Engine/Source/Runtime/Engine/Private/Components/CharacterMovementComponent.cpp b/Engine/Source/Runtime/Engine/Private/Components/CharacterMovementComponent.cpp
--- a/Engine/Source/Runtime/Engine/Private/Components/CharacterMovementComponent.cpp
+++ b/Engine/Source/Runtime/Engine/Private/Components/CharacterMovementComponent.cpp
@@ -1209,6 +1209,8 @@ void UCharacterMovementComponent::SimulatedTick(float DeltaSeconds)
 		bJustTeleported = false;
 
 		CharacterOwner->RootMotionRepMoves.Empty();
+		CurrentRootMotion.Clear();
+		CharacterOwner->SavedRootMotion.Clear();
 
 		// Note: we do not call the Super implementation, that runs prediction.
 		// We do still need to call these though
@@ -8725,6 +8727,11 @@ void UCharacterMovementComponent::TickCharacterPose(float DeltaTime)
 
 bool UCharacterMovementComponent::HasRootMotionSources() const
 {
+	if ( NetworkSmoothingMode == ENetworkSmoothingMode::Replay && CharacterOwner && CharacterOwner->Role == ROLE_SimulatedProxy && CharacterOwner->bReplayHasRootMotionSources )
+	{
+		return true;
+	}
+
 	return CurrentRootMotion.HasActiveRootMotionSources() || (CharacterOwner && CharacterOwner->IsPlayingRootMotion() && CharacterOwner->GetMesh());
 }
 
diff --git a/Engine/Source/Runtime/Engine/Private/Components/InputComponent.cpp b/Engine/Source/Runtime/Engine/Private/Components/InputComponent.cpp
--- a/Engine/Source/Runtime/Engine/Private/Components/InputComponent.cpp
+++ b/Engine/Source/Runtime/Engine/Private/Components/InputComponent.cpp
@@ -134,19 +134,6 @@ void UInputComponent::ClearActionBindings( )
 	ActionBindings.Reset();
 }
 
-
-FInputActionBinding& UInputComponent::GetActionBinding( const int32 BindingIndex )
-{
-	return ActionBindings[BindingIndex];
-}
-
-
-int32 UInputComponent::GetNumActionBindings( ) const
-{
-	return ActionBindings.Num();
-}
-
-
 void UInputComponent::RemoveActionBinding( const int32 BindingIndex )
 {
 	if (BindingIndex >= 0 && BindingIndex < ActionBindings.Num())
diff --git a/Engine/Source/Runtime/Engine/Private/Components/InterpToMovementComponent.cpp b/Engine/Source/Runtime/Engine/Private/Components/InterpToMovementComponent.cpp
--- a/Engine/Source/Runtime/Engine/Private/Components/InterpToMovementComponent.cpp
+++ b/Engine/Source/Runtime/Engine/Private/Components/InterpToMovementComponent.cpp
@@ -259,7 +259,10 @@ FVector UInterpToMovementComponent::ComputeMoveDelta(float InTime) const
 	}
 
 	FVector CurrentPosition = UpdatedComponent->GetComponentLocation();
-	MoveDelta = NewPosition - CurrentPosition;
+	if (CurrentPosition != NewPosition)
+	{
+		MoveDelta = NewPosition - CurrentPosition;
+	}
 	return MoveDelta;
 }
 
@@ -463,14 +466,23 @@ void UInterpToMovementComponent::UpdateControlPoints(bool InForceUpdate)
 					TotalDistance += ControlPoints[ControlPoint].DistanceToNext;
 					CurrentPos = NextPosition;
 				}
+				else
+				{
+					ControlPoints[ControlPoint].DistanceToNext = 0.0f;
+					ControlPoints[ControlPoint].Percentage = 1.0f;
+					ControlPoints[ControlPoint].StartTime = 1.0f;
+				}
 			}
 			float Percent = 0.0f;
 			// Use the distance to determine what % of time to spend going from each point
 			for (int32 ControlPoint = 0; ControlPoint < ControlPoints.Num(); ControlPoint++)
 			{
 				ControlPoints[ControlPoint].StartTime = Percent;
-				ControlPoints[ControlPoint].Percentage = ControlPoints[ControlPoint].DistanceToNext / TotalDistance;
-				Percent += ControlPoints[ControlPoint].Percentage;
+				if(ControlPoints[ControlPoint].DistanceToNext != 0.0f)
+				{
+					ControlPoints[ControlPoint].Percentage = ControlPoints[ControlPoint].DistanceToNext / TotalDistance;
+					Percent += ControlPoints[ControlPoint].Percentage;
+				}
 			}
 		}
 	}
diff --git a/Engine/Source/Runtime/Engine/Private/Components/MeshComponent.cpp b/Engine/Source/Runtime/Engine/Private/Components/MeshComponent.cpp
--- a/Engine/Source/Runtime/Engine/Private/Components/MeshComponent.cpp
+++ b/Engine/Source/Runtime/Engine/Private/Components/MeshComponent.cpp
@@ -195,7 +195,7 @@ void UMeshComponent::SetScalarParameterValueOnMaterials(const FName ParameterNam
 	}
 	else
 	{
-		UE_LOG(LogMaterialParameter, Log, TEXT("%s material parameter hasn't found on the component %s"), *ParameterName.ToString(), *GetName());
+		UE_LOG(LogMaterialParameter, Log, TEXT("%s material parameter hasn't found on the component %s"), *ParameterName.ToString(), *GetPathName());
 	}
 
 	// CreateAndSetMaterialInstanceDynamic should not flag the cached data as dirty, since only the type of material changes not the contents (UMaterialInstanceDynamic)
diff --git a/Engine/Source/Runtime/Engine/Private/Components/ReflectionCaptureComponent.cpp b/Engine/Source/Runtime/Engine/Private/Components/ReflectionCaptureComponent.cpp
--- a/Engine/Source/Runtime/Engine/Private/Components/ReflectionCaptureComponent.cpp
+++ b/Engine/Source/Runtime/Engine/Private/Components/ReflectionCaptureComponent.cpp
@@ -1156,9 +1156,9 @@ void UReflectionCaptureComponent::PostLoad()
 	}
 }
 
-void UReflectionCaptureComponent::PreSave() 
+void UReflectionCaptureComponent::PreSave(const class ITargetPlatform* TargetPlatform)
 {
-	Super::PreSave();
+	Super::PreSave(TargetPlatform);
 
 	// This is done on save of the package, because this capture data can only be generated by the renderer
 	// So we must make efforts to ensure that it is generated in the editor, because it can't be generated during cooking when missing
diff --git a/Engine/Source/Runtime/Engine/Private/Components/SkeletalMeshComponent.cpp b/Engine/Source/Runtime/Engine/Private/Components/SkeletalMeshComponent.cpp
--- a/Engine/Source/Runtime/Engine/Private/Components/SkeletalMeshComponent.cpp
+++ b/Engine/Source/Runtime/Engine/Private/Components/SkeletalMeshComponent.cpp
@@ -1033,21 +1033,14 @@ void USkeletalMeshComponent::RecalcRequiredBones(int32 LODIndex)
 	// Ensure that we have a complete hierarchy down to those bones.
 	FAnimationRuntime::EnsureParentsPresent(RequiredBones, SkeletalMesh);
 
-	FillSpaceBasesRequiredBones.Empty(RequiredBones.Num() + NeededBonesForFillSpaceBases.Num());
+	FillSpaceBasesRequiredBones.Reset(RequiredBones.Num() + NeededBonesForFillSpaceBases.Num());
 	FillSpaceBasesRequiredBones = RequiredBones;
 	
 	NeededBonesForFillSpaceBases.Sort();
 	MergeInBoneIndexArrays(FillSpaceBasesRequiredBones, NeededBonesForFillSpaceBases);
 	FAnimationRuntime::EnsureParentsPresent(FillSpaceBasesRequiredBones, SkeletalMesh);
 
-	// Sanitise bones that we aren't going to be updating
-	for (int32 BoneIndex = 0; BoneIndex < LocalAtoms.Num(); ++BoneIndex)
-	{
-		if (!RequiredBones.Contains(BoneIndex))
-		{
-			LocalAtoms[BoneIndex] = SkeletalMesh->RefSkeleton.GetRefBonePose()[BoneIndex];
-		}
-	}
+	LocalAtoms = SkeletalMesh->RefSkeleton.GetRefBonePose();
 
 	// make sure animation requiredBone to mark as dirty
 	if (AnimScriptInstance)
diff --git a/Engine/Source/Runtime/Engine/Private/Components/StaticMeshComponent.cpp b/Engine/Source/Runtime/Engine/Private/Components/StaticMeshComponent.cpp
--- a/Engine/Source/Runtime/Engine/Private/Components/StaticMeshComponent.cpp
+++ b/Engine/Source/Runtime/Engine/Private/Components/StaticMeshComponent.cpp
@@ -325,9 +325,9 @@ bool UStaticMeshComponent::AreNativePropertiesIdenticalTo( UObject* Other ) cons
 	return bNativePropertiesAreIdentical;
 }
 
-void UStaticMeshComponent::PreSave()
+void UStaticMeshComponent::PreSave(const class ITargetPlatform* TargetPlatform)
 {
-	Super::PreSave();
+	Super::PreSave(TargetPlatform);
 #if WITH_EDITORONLY_DATA
 	CachePaintedDataIfNecessary();
 #endif // WITH_EDITORONLY_DATA
diff --git a/Engine/Source/Runtime/Engine/Private/DemoNetDriver.cpp b/Engine/Source/Runtime/Engine/Private/DemoNetDriver.cpp
--- a/Engine/Source/Runtime/Engine/Private/DemoNetDriver.cpp
+++ b/Engine/Source/Runtime/Engine/Private/DemoNetDriver.cpp
@@ -602,7 +602,7 @@ bool UDemoNetDriver::InitConnectInternal( FString& Error )
 
 		// Set up the pending net game so that the engine can call LoadMap on the next tick.
 		NewPendingNetGame->DemoNetDriver = this;
-		NewPendingNetGame->URL = DemoURL;
+		NewPendingNetGame->URL = LocalDemoURL;
 		NewPendingNetGame->bSuccessfullyConnected = true;
 
 		WorldContext->PendingNetGame = NewPendingNetGame;
@@ -1451,9 +1451,9 @@ void UDemoNetDriver::TickDemoRecord( float DeltaSeconds )
 
 			TSharedPtr<FRepChangedPropertyTracker> PropertyTracker = FindOrCreateRepChangedPropertyTracker( Actor );
 
-			if (!bDidReplicateActor)
+			if ( !GuidCache->NetGUIDLookup.Contains( Actor ) )
 			{
-				// Clear external data if the actor didn't replicate for some reason
+				// Clear external data if the actor has never replicated yet (and doesn't have a net guid)
 				PropertyTracker->ExternalData.Empty();
 				PropertyTracker->ExternalDataNumBits = 0;
 			}
diff --git a/Engine/Source/Runtime/Engine/Private/EdGraph/EdGraphNode.cpp b/Engine/Source/Runtime/Engine/Private/EdGraph/EdGraphNode.cpp
--- a/Engine/Source/Runtime/Engine/Private/EdGraph/EdGraphNode.cpp
+++ b/Engine/Source/Runtime/Engine/Private/EdGraph/EdGraphNode.cpp
@@ -287,9 +287,9 @@ void UEdGraphNode::Serialize(FArchive& Ar)
 	}
 }
 
-void UEdGraphNode::PreSave()
+void UEdGraphNode::PreSave(const class ITargetPlatform* TargetPlatform)
 {
-	Super::PreSave();
+	Super::PreSave(TargetPlatform);
 
 #if WITH_EDITORONLY_DATA
 	if (!NodeUpgradeMessage.IsEmpty())
diff --git a/Engine/Source/Runtime/Engine/Private/GameFramework/RootMotionSource.cpp b/Engine/Source/Runtime/Engine/Private/GameFramework/RootMotionSource.cpp
--- a/Engine/Source/Runtime/Engine/Private/GameFramework/RootMotionSource.cpp
+++ b/Engine/Source/Runtime/Engine/Private/GameFramework/RootMotionSource.cpp
@@ -402,6 +402,8 @@ FRootMotionSource_RadialForce::FRootMotionSource_RadialForce()
 	, bNoZForce(false)
 	, StrengthDistanceFalloff(nullptr)
 	, StrengthOverTime(nullptr)
+	, bUseFixedWorldDirection(false)
+	, FixedWorldDirection(EForceInit::ForceInitToZero)
 {
 }
 
@@ -423,12 +425,14 @@ bool FRootMotionSource_RadialForce::Matches(const FRootMotionSource* Other) cons
 
 	return bIsPush == OtherCast->bIsPush &&
 		bNoZForce == OtherCast->bNoZForce &&
+		bUseFixedWorldDirection == OtherCast->bUseFixedWorldDirection &&
 		StrengthDistanceFalloff == OtherCast->StrengthDistanceFalloff &&
 		StrengthOverTime == OtherCast->StrengthOverTime &&
 		(LocationActor == OtherCast->LocationActor ||
 		FVector::PointsAreNear(Location, OtherCast->Location, 1.0f)) &&
 		FMath::IsNearlyEqual(Radius, OtherCast->Radius, SMALL_NUMBER) &&
-		FMath::IsNearlyEqual(Strength, OtherCast->Strength, SMALL_NUMBER);
+		FMath::IsNearlyEqual(Strength, OtherCast->Strength, SMALL_NUMBER) &&
+		FixedWorldDirection.Equals(OtherCast->FixedWorldDirection, 3.0f);
 }
 
 bool FRootMotionSource_RadialForce::MatchesAndHasSameState(const FRootMotionSource* Other) const
@@ -488,11 +492,18 @@ void FRootMotionSource_RadialForce::PrepareRootMotion
 			CurrentStrength = Strength * FMath::Clamp(AdditiveStrengthFactor, 0.f, 1.f);
 		}
 
-		Force = (ForceLocation - CharacterLocation).GetSafeNormal() * CurrentStrength;
-
-		if (bIsPush)
+		if (bUseFixedWorldDirection)
+		{
+			Force = FixedWorldDirection.Vector() * CurrentStrength;
+		}
+		else
 		{
-			Force *= -1.f;
+			Force = (ForceLocation - CharacterLocation).GetSafeNormal() * CurrentStrength;
+			
+			if (bIsPush)
+			{
+				Force *= -1.f;
+			}
 		}
 	}
 
@@ -534,6 +545,8 @@ bool FRootMotionSource_RadialForce::NetSerialize(FArchive& Ar, UPackageMap* Map,
 	Ar << bNoZForce;
 	Ar << StrengthDistanceFalloff;
 	Ar << StrengthOverTime;
+	Ar << bUseFixedWorldDirection;
+	Ar << FixedWorldDirection;
 
 	bOutSuccess = true;
 	return true;
diff --git a/Engine/Source/Runtime/Engine/Private/GameMode.cpp b/Engine/Source/Runtime/Engine/Private/GameMode.cpp
--- a/Engine/Source/Runtime/Engine/Private/GameMode.cpp
+++ b/Engine/Source/Runtime/Engine/Private/GameMode.cpp
@@ -726,6 +726,8 @@ void AGameMode::SetMatchState(FName NewState)
 		return;
 	}
 
+	UE_LOG(LogGameMode, Display, TEXT("Match State Changed from %s to %s"), *MatchState.ToString(), *NewState.ToString());
+
 	MatchState = NewState;
 
 	OnMatchStateSet();
@@ -772,6 +774,7 @@ void AGameMode::Tick(float DeltaSeconds)
 		// Check to see if we should start the match
 		if (ReadyToStartMatch())
 		{
+			UE_LOG(LogGameMode, Log, TEXT("GameMode returned ReadyToStartMatch"));
 			StartMatch();
 		}
 	}
@@ -780,6 +783,7 @@ void AGameMode::Tick(float DeltaSeconds)
 		// Check to see if we should start the match
 		if (ReadyToEndMatch())
 		{
+			UE_LOG(LogGameMode, Log, TEXT("GameMode returned ReadyToEndMatch"));
 			EndMatch();
 		}
 	}
diff --git a/Engine/Source/Runtime/Engine/Private/GameState.cpp b/Engine/Source/Runtime/Engine/Private/GameState.cpp
--- a/Engine/Source/Runtime/Engine/Private/GameState.cpp
+++ b/Engine/Source/Runtime/Engine/Private/GameState.cpp
@@ -213,6 +213,8 @@ void AGameState::SetMatchState(FName NewState)
 {
 	if (Role == ROLE_Authority)
 	{
+		UE_LOG(LogGameState, Log, TEXT("Match State Changed from %s to %s"), *MatchState.ToString(), *NewState.ToString());
+
 		MatchState = NewState;
 
 		// Call the onrep to make sure the callbacks happen
diff --git a/Engine/Source/Runtime/Engine/Private/Level.cpp b/Engine/Source/Runtime/Engine/Private/Level.cpp
--- a/Engine/Source/Runtime/Engine/Private/Level.cpp
+++ b/Engine/Source/Runtime/Engine/Private/Level.cpp
@@ -483,9 +483,9 @@ void ULevel::ValidateLightGUIDs()
 }
 
 
-void ULevel::PreSave()
+void ULevel::PreSave(const class ITargetPlatform* TargetPlatform)
 {
-	Super::PreSave();
+	Super::PreSave(TargetPlatform);
 
 #if WITH_EDITOR
 	if( !IsTemplate() )
diff --git a/Engine/Source/Runtime/Engine/Private/Materials/Material.cpp b/Engine/Source/Runtime/Engine/Private/Materials/Material.cpp
--- a/Engine/Source/Runtime/Engine/Private/Materials/Material.cpp
+++ b/Engine/Source/Runtime/Engine/Private/Materials/Material.cpp
@@ -634,9 +634,9 @@ UMaterial::UMaterial(const FObjectInitializer& ObjectInitializer)
 #endif //WITH_EDITORONLY_DATA
 }
 
-void UMaterial::PreSave()
+void UMaterial::PreSave(const class ITargetPlatform* TargetPlatform)
 {
-	Super::PreSave();
+	Super::PreSave(TargetPlatform);
 #if WITH_EDITOR
 	GMaterialsWithDirtyUsageFlags.RemoveAnnotation(this);
 #endif
diff --git a/Engine/Source/Runtime/Engine/Private/NetConnection.cpp b/Engine/Source/Runtime/Engine/Private/NetConnection.cpp
--- a/Engine/Source/Runtime/Engine/Private/NetConnection.cpp
+++ b/Engine/Source/Runtime/Engine/Private/NetConnection.cpp
@@ -119,6 +119,8 @@ void UNetConnection::InitBase(UNetDriver* InDriver,class FSocket* InSocket, cons
 	LastRecvAckTime			= Driver->Time;
 	ConnectTime				= Driver->Time;
 
+	NetConnectionHistogram.InitHitchTracking();
+
 	// Current state
 	State = InState;
 	// Copy the URL
@@ -769,6 +771,10 @@ void UNetConnection::ReceivedPacket( FBitReader& Reader )
 
 	ValidateSendBuffer();
 
+	//Record the packet time to the histogram
+	double LastPacketTimeDiffInMs = (FPlatformTime::Seconds() - LastReceiveRealtime) * 1000.0;
+	NetConnectionHistogram.AddMeasurement(LastPacketTimeDiffInMs);
+
 	// Update receive time to avoid timeout.
 	LastReceiveTime		= Driver->Time;
 	LastReceiveRealtime = FPlatformTime::Seconds();
@@ -1922,7 +1928,7 @@ void UNetConnection::CleanupDormantActorState()
 	DormantReplicatorMap.Empty();
 }
 
-void UNetConnection::FlushDormancy( class AActor* Actor )
+void UNetConnection::FlushDormancy(class AActor* Actor)
 {
 	UE_LOG( LogNetDormancy, Verbose, TEXT( "FlushDormancy: %s. Connection: %s" ), *Actor->GetName(), *GetName() );
 	
diff --git a/Engine/Source/Runtime/Engine/Private/Particles/ParticleBeam2EmitterInstance.cpp b/Engine/Source/Runtime/Engine/Private/Particles/ParticleBeam2EmitterInstance.cpp
--- a/Engine/Source/Runtime/Engine/Private/Particles/ParticleBeam2EmitterInstance.cpp
+++ b/Engine/Source/Runtime/Engine/Private/Particles/ParticleBeam2EmitterInstance.cpp
@@ -1121,18 +1121,17 @@ void FParticleBeam2EmitterInstance::DetermineVertexAndTriangleCount()
 		int32 LocalTriangles = 0;
 		if (BeamData->TriangleCount > 0)
 		{
+			if (VerticesToRender > 0)
+			{
+				LocalTriangles += 4;//Degenerate tris linking from previous beam.
+			}
+
 			// Stored triangle count is per sheet...
 			LocalTriangles	+= BeamData->TriangleCount * Sheets;
 			VerticesToRender += (BeamData->TriangleCount + 2) * Sheets;
 			// 4 Degenerates Per Sheet (except for last one)
 			LocalTriangles	+= (Sheets - 1) * 4;
 			EmitterTrianglesToRender += LocalTriangles;
-			// Multiple beams?
-			if (i < (ActiveParticles - 1))
-			{
-				// 4 Degenerates Per Beam (except for last one)
-				EmitterTrianglesToRender	+= 4;
-			}
 		}
 	}
 
diff --git a/Engine/Source/Runtime/Engine/Private/Particles/ParticleComponents.cpp b/Engine/Source/Runtime/Engine/Private/Particles/ParticleComponents.cpp
--- a/Engine/Source/Runtime/Engine/Private/Particles/ParticleComponents.cpp
+++ b/Engine/Source/Runtime/Engine/Private/Particles/ParticleComponents.cpp
@@ -1701,10 +1701,9 @@ void UParticleEmitter::CacheEmitterModuleInfo()
 
 float UParticleEmitter::GetQualityLevelSpawnRateMult()
 {
-	static Scalability::FEffectsQualityCVarAccessor EffectsQuality;
-
+	int32 EffectsQuality = Scalability::GetEffectsQualityDirect(true);
 	float Q = 1;
-	float Level = (1 - EffectsQuality.GetGameThreadValue());
+	float Level = (1 - EffectsQuality);
 	for (int i = 0; i < Level + 1; i++)
 	{
 		Q = Q*QualityLevelSpawnRateScale;
@@ -2131,9 +2130,9 @@ void UParticleSystem::PostEditChangeProperty(FPropertyChangedEvent& PropertyChan
 }
 #endif // WITH_EDITOR
 
-void UParticleSystem::PreSave()
+void UParticleSystem::PreSave(const class ITargetPlatform* TargetPlatform)
 {
-	Super::PreSave();
+	Super::PreSave(TargetPlatform);
 #if WITH_EDITORONLY_DATA
 	// Ensure that soloing is undone...
 	int32 NumEmitters = FMath::Min(Emitters.Num(),SoloTracking.Num());
diff --git a/Engine/Source/Runtime/Engine/Private/Particles/ParticleSystemRender.cpp b/Engine/Source/Runtime/Engine/Private/Particles/ParticleSystemRender.cpp
--- a/Engine/Source/Runtime/Engine/Private/Particles/ParticleSystemRender.cpp
+++ b/Engine/Source/Runtime/Engine/Private/Particles/ParticleSystemRender.cpp
@@ -3127,15 +3127,19 @@ void FDynamicBeam2EmitterData::GetIndexAllocInfo(int32& OutNumIndices, int32& Ou
 		}
 		else
 		{
-			if (TempIndexCount == 0)
+			if (Triangles > 0)
 			{
-				TempIndexCount = 2;
-			}
-			TempIndexCount += Triangles * Source.Sheets;
-			TempIndexCount += 4 * (Source.Sheets - 1);	// Degenerate indices between sheets
-			if ((ii + 1) < Source.TrianglesPerSheet.Num())
-			{
-				TempIndexCount += 4;	// Degenerate indices between beams
+				if (TempIndexCount == 0)
+				{
+					TempIndexCount = 2;     // First Beam
+				}
+				else
+				{
+					TempIndexCount += 4;	// Degenerate indices between beams
+				}
+
+				TempIndexCount += Triangles * Source.Sheets;
+				TempIndexCount += 4 * (Source.Sheets - 1);	// Degenerate indices between sheets
 			}
 		}
 	}
@@ -3164,11 +3168,20 @@ static int32 CreateDynamicBeam2EmitterIndices(TIndexType* OutIndex, const FDynam
 			continue;
 		}
 
-		if (Beam == 0)
+		if (VertexIndex == 0)//First Beam
 		{
 			*(OutIndex++) = VertexIndex++;	// SheetIndex + 0
 			*(OutIndex++) = VertexIndex++;	// SheetIndex + 1
 		}
+		else//Degenerate tris between beams
+		{
+			*(OutIndex++) = VertexIndex - 1;	// Last vertex of the previous sheet
+			*(OutIndex++) = VertexIndex;		// First vertex of the next sheet
+			*(OutIndex++) = VertexIndex++;		// First vertex of the next sheet
+			*(OutIndex++) = VertexIndex++;		// Second vertex of the next sheet
+
+			TrianglesToRender += 4;
+		}
 
 		for (int32 SheetIndex = 0; SheetIndex < Source.Sheets; SheetIndex++)
 		{
@@ -3192,15 +3205,6 @@ static int32 CreateDynamicBeam2EmitterIndices(TIndexType* OutIndex, const FDynam
 				TrianglesToRender += 4;
 			}
 		}
-		if ((Beam + 1) < Source.ActiveParticleCount)
-		{
-			*(OutIndex++) = VertexIndex - 1;	// Last vertex of the previous sheet
-			*(OutIndex++) = VertexIndex;		// First vertex of the next sheet
-			*(OutIndex++) = VertexIndex++;		// First vertex of the next sheet
-			*(OutIndex++) = VertexIndex++;		// Second vertex of the next sheet
-
-			TrianglesToRender += 4;
-		}
 	}
 
 	return TrianglesToRender;
diff --git a/Engine/Source/Runtime/Engine/Private/Particles/ParticleTrail2EmitterInstance.cpp b/Engine/Source/Runtime/Engine/Private/Particles/ParticleTrail2EmitterInstance.cpp
--- a/Engine/Source/Runtime/Engine/Private/Particles/ParticleTrail2EmitterInstance.cpp
+++ b/Engine/Source/Runtime/Engine/Private/Particles/ParticleTrail2EmitterInstance.cpp
@@ -3241,10 +3241,9 @@ float FParticleAnimTrailEmitterInstance::Spawn(float DeltaTime)
 	// Figure out spawn rate for this tick.
 	if (bProcessSpawnRate)
 	{
-		static Scalability::FEffectsQualityCVarAccessor EffectsQuality;
-
+		int32 EffectsQuality = Scalability::GetEffectsQualityDirect(true);
 		float RateScale = LODLevel->SpawnModule->RateScale.GetValue(EmitterTime, Component) * LODLevel->SpawnModule->GetGlobalRateScale();
-		float QualityMult = 0.25f * (1 << EffectsQuality.GetGameThreadValue());
+		float QualityMult = 0.25f * (1 << EffectsQuality);
 		RateScale *= SpriteTemplate->QualityLevelSpawnRateScale*QualityMult;
 		SpawnRate += LODLevel->SpawnModule->Rate.GetValue(EmitterTime, Component) * FMath::Clamp<float>(RateScale, 0.0f, RateScale);
 	}
diff --git a/Engine/Source/Runtime/Engine/Private/Performance/EnginePerformanceTargets.cpp b/Engine/Source/Runtime/Engine/Private/Performance/EnginePerformanceTargets.cpp
--- a/Engine/Source/Runtime/Engine/Private/Performance/EnginePerformanceTargets.cpp
+++ b/Engine/Source/Runtime/Engine/Private/Performance/EnginePerformanceTargets.cpp
@@ -11,7 +11,8 @@ static TAutoConsoleVariable<float> GUnacceptableFrameTimeThresholdCVar(
 	TEXT("t.UnacceptableFrameTimeThreshold"),
 	50.0f,
 	TEXT("The frame time theshold for what is considered completely unacceptable (in ms); values above this will be drawn as red\n")
-	TEXT(" default: 50.0 ms"));
+	TEXT(" default: 50.0 ms"),
+	ECVF_Scalability);
 
 // The target threshold for frame time in miliseconds (*cosmetic only* and used for fps/stat display, should not be used in scalability code)
 // Values below this will be green, values between this and the unacceptable threshold will be yellow, and values above that will be red
@@ -19,7 +20,8 @@ TAutoConsoleVariable<float> GTargetFrameTimeThresholdCVar(
 	TEXT("t.TargetFrameTimeThreshold"),
 	33.9f,
 	TEXT("The target frame time (in ms); values below this will be drawn in green, values above will be yellow or red depending on the severity\n")
-	TEXT(" default: 33.9 ms"));
+	TEXT(" default: 33.9 ms"),
+	ECVF_Scalability);
 
 // The threshold that would be considered so bad that it would cause a hitch in gameplay
 // (*cosmetic only* for reporting purposes such as FPS charts, should not be used in scalability code)
@@ -27,21 +29,24 @@ FAutoConsoleVariableRef GHitchFrameTimeThresholdCVar(
 	TEXT("t.HitchFrameTimeThreshold"),
 	GHitchThresholdMS,
 	TEXT("Definition of a hitchy frame (in ms)\n")
-	TEXT(" default: 60.0 ms"));
+	TEXT(" default: 60.0 ms"),
+	ECVF_Scalability);
 
 // Minimum time passed before we'll record a new hitch
 TAutoConsoleVariable<float> GHitchDeadTimeWindowCVar(
 	TEXT("t.HitchDeadTimeWindow"),
 	200.0f,
 	TEXT("Minimum time passed before we'll record a new hitch (in ms)\n")
-	TEXT(" default: 200.0 ms"));
+	TEXT(" default: 200.0 ms"),
+	ECVF_Scalability);
 
 // For the current frame to be considered a hitch, it must have run at least this many times slower than the previous frame
 TAutoConsoleVariable<float> GHitchVersusNonHitchMinimumRatioCVar(
 	TEXT("t.HitchVersusNonHitchRatio"),
 	1.5f,
 	TEXT("For the current frame to be considered a hitch, it must have run at least this many times slower than the previous frame.\nThe actual ratio is clamped to be between this and t.HitchFrameTimeThreshold/t.TargetFrameTimeThreshold\n")
-	TEXT(" default: 1.5"));
+	TEXT(" default: 1.5"),
+	ECVF_Scalability);
 
 //////////////////////////////////////////////////////////////////////
 // FEnginePerformanceTargets
diff --git a/Engine/Source/Runtime/Engine/Private/Scalability.cpp b/Engine/Source/Runtime/Engine/Private/Scalability.cpp
--- a/Engine/Source/Runtime/Engine/Private/Scalability.cpp
+++ b/Engine/Source/Runtime/Engine/Private/Scalability.cpp
@@ -500,15 +500,6 @@ void SetQualityLevels(const FQualityLevels& QualityLevels)
 	CVarFoliageQuality.AsVariable()->Set(NewFoliageQuality, ECVF_SetByCode);
 }
 
-/** Helper accessor objects */
-FResolutionQualityCVarAccessor::FResolutionQualityCVarAccessor() : FCVarAccessorFloat(CVarResolutionQuality.AsVariable()){}
-FViewDistanceQualityCVarAccessor::FViewDistanceQualityCVarAccessor() : FCVarAccessorInt(CVarViewDistanceQuality.AsVariable()){}
-FAntiAliasingQualityCVarAccessor::FAntiAliasingQualityCVarAccessor() : FCVarAccessorInt(CVarAntiAliasingQuality.AsVariable()){}
-FShadowQualityCVarAccessor::FShadowQualityCVarAccessor() : FCVarAccessorInt(CVarShadowQuality.AsVariable()){}
-FPostProcessQualityCVarAccessor::FPostProcessQualityCVarAccessor() : FCVarAccessorInt(CVarPostProcessQuality.AsVariable()){}
-FTextureQualityCVarAccessor::FTextureQualityCVarAccessor() : FCVarAccessorInt(CVarTextureQuality.AsVariable()){}
-FEffectsQualityCVarAccessor::FEffectsQualityCVarAccessor() : FCVarAccessorInt(CVarEffectsQuality.AsVariable()){}
-
 FQualityLevels GetQualityLevels()
 {
 	FQualityLevels Ret;
@@ -526,6 +517,18 @@ FQualityLevels GetQualityLevels()
 	return Ret;
 }
 
+int32 GetEffectsQualityDirect(bool bGameThread)
+{
+	if (bGameThread)
+	{
+		return CVarEffectsQuality.GetValueOnAnyThread(true);
+	}
+	else
+	{
+		return CVarEffectsQuality.GetValueOnRenderThread();
+	}
+}
+
 void FQualityLevels::SetBenchmarkFallback()
 {
 	ResolutionQuality = 100.0f;
diff --git a/Engine/Source/Runtime/Engine/Private/ShaderCompiler/ShaderCompiler.cpp b/Engine/Source/Runtime/Engine/Private/ShaderCompiler/ShaderCompiler.cpp
--- a/Engine/Source/Runtime/Engine/Private/ShaderCompiler/ShaderCompiler.cpp
+++ b/Engine/Source/Runtime/Engine/Private/ShaderCompiler/ShaderCompiler.cpp
@@ -77,6 +77,16 @@ static TAutoConsoleVariable<int32> CVarShaderFastMath(
 	TEXT("Whether to use fast-math optimisations in shaders."),
 	ECVF_ReadOnly);
 
+static TAutoConsoleVariable<int32> CVarD3DRemoveUnusedInterpolators(
+	TEXT("r.D3D.RemoveUnusedInterpolators"),
+	0,
+	TEXT("Enables removing unused interpolators mode when compiling pipelines for D3D.\n")
+	TEXT(" -1: Do not actually remove, but make the app think it did (for debugging)\n")
+	TEXT(" 0: Disable (default)\n")
+	TEXT(" 1: Enable removing unused"),
+	ECVF_ReadOnly
+	);
+
 #if !(UE_BUILD_SHIPPING || UE_BUILD_TEST)
 static TAutoConsoleVariable<FString> CVarD3DCompilerPath(TEXT("r.D3DCompilerPath"),
 	TEXT(""),	// default
@@ -2341,6 +2351,23 @@ void GlobalBeginCompileShader(
 		}
 	}
 
+	{
+		static const auto CVar = IConsoleManager::Get().FindConsoleVariable(TEXT("r.Shaders.FastMath"));
+		if (CVar && CVar->GetInt() == 0)
+		{
+			Input.Environment.CompilerFlags.Add(CFLAG_NoFastMath);
+		}
+	}
+
+	if (IsD3DPlatform((EShaderPlatform)Target.Platform, false))
+	{
+		static const auto CVar = IConsoleManager::Get().FindConsoleVariable(TEXT("r.D3D.RemoveUnusedInterpolators"));
+		if (CVar && CVar->GetInt() != 0)
+		{
+			Input.Environment.CompilerFlags.Add(CFLAG_ForceRemoveUnusedInterpolators);
+		}
+	}
+
 	Input.Environment.SetDefine(TEXT("HAS_INVERTED_Z_BUFFER"), (bool)ERHIZBuffer::IsInverted);
 
 	{
diff --git a/Engine/Source/Runtime/Engine/Private/SkeletalMesh.cpp b/Engine/Source/Runtime/Engine/Private/SkeletalMesh.cpp
--- a/Engine/Source/Runtime/Engine/Private/SkeletalMesh.cpp
+++ b/Engine/Source/Runtime/Engine/Private/SkeletalMesh.cpp
@@ -2201,7 +2201,15 @@ static void GetStreamingTextureFactorForLOD(FStaticLODModel& LODModel, TArray<fl
 {
 	int32 NumTotalTriangles = LODModel.GetTotalFaces();
 
-	TArray<float> TexelRatios[MAX_TEXCOORDS];
+	struct FTriangleInfo
+	{
+		FTriangleInfo(float InAera, float InTexelRatio) : Aera(InAera), TexelRatio(InTexelRatio) {}
+		float Aera;
+		float TexelRatio;
+	};
+
+
+	TArray<FTriangleInfo> TexelRatios[MAX_TEXCOORDS];
 	float MaxTexelRatio = 0.0f;
 	for(int32 UVIndex = 0;UVIndex < MAX_TEXCOORDS;UVIndex++)
 	{
@@ -2230,52 +2238,82 @@ static void GetStreamingTextureFactorForLOD(FStaticLODModel& LODModel, TArray<fl
 				const FVector Pos0 = LODModel.VertexBufferGPUSkin.GetVertexPositionFast<bExtraBoneInfluencesT>(Index0);
 				const FVector Pos1 = LODModel.VertexBufferGPUSkin.GetVertexPositionFast<bExtraBoneInfluencesT>(Index1);
 				const FVector Pos2 = LODModel.VertexBufferGPUSkin.GetVertexPositionFast<bExtraBoneInfluencesT>(Index2);
-				float L1 = (Pos0 - Pos1).Size();
-				float L2 = (Pos0 - Pos2).Size();
 
-				int32 NumUVs = LODModel.NumTexCoords;
-				for(int32 UVIndex = 0;UVIndex < FMath::Min(NumUVs,(int32)MAX_TEXCOORDS);UVIndex++)
-				{
-					const FVector2D UV0 = LODModel.VertexBufferGPUSkin.GetVertexUVFast<bExtraBoneInfluencesT>(Index0, UVIndex);
-					const FVector2D UV1 = LODModel.VertexBufferGPUSkin.GetVertexUVFast<bExtraBoneInfluencesT>(Index1, UVIndex);
-					const FVector2D UV2 = LODModel.VertexBufferGPUSkin.GetVertexUVFast<bExtraBoneInfluencesT>(Index2, UVIndex);
+				FVector P01 = Pos1 - Pos0;
+				FVector P02 = Pos2 - Pos0;
+
+				float Aera = FVector::CrossProduct(P01, P02).Size();
 
-					float T1 = (UV0 - UV1).Size();
-					float T2 = (UV0 - UV2).Size();
+				if (Aera > SMALL_NUMBER)
+				{
+					float L1 = (Pos0 - Pos1).Size();
+					float L2 = (Pos0 - Pos2).Size();
 
-					if( FMath::Abs(T1 * T2) > FMath::Square(SMALL_NUMBER) )
+					int32 NumUVs = LODModel.NumTexCoords;
+					for(int32 UVIndex = 0;UVIndex < FMath::Min(NumUVs,(int32)MAX_TEXCOORDS);UVIndex++)
 					{
-						const float TexelRatio = FMath::Max( L1 / T1, L2 / T2 );
-						TexelRatios[UVIndex].Add( TexelRatio );
+						const FVector2D UV0 = LODModel.VertexBufferGPUSkin.GetVertexUVFast<bExtraBoneInfluencesT>(Index0, UVIndex);
+						const FVector2D UV1 = LODModel.VertexBufferGPUSkin.GetVertexUVFast<bExtraBoneInfluencesT>(Index1, UVIndex);
+						const FVector2D UV2 = LODModel.VertexBufferGPUSkin.GetVertexUVFast<bExtraBoneInfluencesT>(Index2, UVIndex);
 
-						// Update max texel ratio
-						if( TexelRatio > MaxTexelRatio )
+						float T1 = (UV0 - UV1).Size();
+						float T2 = (UV0 - UV2).Size();
+
+						if( FMath::Abs(T1 * T2) > FMath::Square(SMALL_NUMBER) )
 						{
-							MaxTexelRatio = TexelRatio;
+							const float TexelRatio = FMath::Max( L1 / T1, L2 / T2 );
+							TexelRatios[UVIndex].Add( FTriangleInfo(Aera, TexelRatio) );
+
+							// Update max texel ratio
+							if( TexelRatio > MaxTexelRatio )
+							{
+								MaxTexelRatio = TexelRatio;
+							}
 						}
 					}
 				}
 			}
+		}
+		else
+		{
+			UE_LOG(LogSkeletalMesh,Warning,TEXT("GetStreamingTextureFactor called but section %d has no indices."),SectionIndex);
+		}
+	}
+
+	for(int32 UVIndex = 0;UVIndex < MAX_TEXCOORDS;UVIndex++)
+	{
+		TArray<FTriangleInfo>& TriangleInfos = TexelRatios[UVIndex];
+		float WeightedTexelFactorSum = 0;
+		float AreaSum = 0;
+
+		if( TriangleInfos.Num() )
+		{
+			struct FCompareTexelRatio
+			{
+				FORCEINLINE bool operator()(FTriangleInfo const& A, FTriangleInfo const& B) const { return A.TexelRatio < B.TexelRatio; }
+			};
+
+			TriangleInfos.Sort( FCompareTexelRatio() );
+
+			// Disregard upper 10% of texel ratios.
+			// This is to ignore backfacing surfaces or other non-visible surfaces that tend to map a small number of texels to a large surface.
+			int32 Threshold = FMath::FloorToInt(.10f * (float)TriangleInfos.Num());
+			for (int32 Index = Threshold; Index < TriangleInfos.Num() - Threshold; ++Index)
+			{
+				WeightedTexelFactorSum += TriangleInfos[Index].TexelRatio * TriangleInfos[Index].Aera;
+				AreaSum += TriangleInfos[Index].Aera;
+			}
 
-			for(int32 UVIndex = 0;UVIndex < MAX_TEXCOORDS;UVIndex++)
+			if (AreaSum != 0)
 			{
-				if( TexelRatios[UVIndex].Num() )
+				CachedStreamingTextureFactors[UVIndex] = WeightedTexelFactorSum / AreaSum;
+
+				if ( UVIndex == 0 )
 				{
-					// Disregard upper 75% of texel ratios.
-					// This is to ignore backfacing surfaces or other non-visible surfaces that tend to map a small number of texels to a large surface.
-					TexelRatios[UVIndex].Sort( TGreater<float>() );
-					float TexelRatio = TexelRatios[UVIndex][ FMath::TruncToInt(TexelRatios[UVIndex].Num() * 0.75f) ];
-					if ( UVIndex == 0 )
-					{
-						TexelRatio *= InStreamingDistanceMultiplier;
-					}
-					CachedStreamingTextureFactors[UVIndex] = TexelRatio;
+					CachedStreamingTextureFactors[UVIndex] *= InStreamingDistanceMultiplier;
 				}
 			}
-		}
-		else
-		{
-			UE_LOG(LogSkeletalMesh,Warning,TEXT("GetStreamingTextureFactor called but section %d has no indices."),SectionIndex);
+
 		}
 	}
 }
@@ -2482,11 +2520,6 @@ void USkeletalMesh::PostEditChangeProperty(FPropertyChangedEvent& PropertyChange
 		}
 	}
 	
-	if (!bSkipRestartRenderState)
-	{
-		RestartRenderState();
-	}
-
 	if ( GIsEditor && PropertyThatChanged && PropertyThatChanged->GetName() == TEXT("StreamingDistanceMultiplier") )
 	{
 		// Allow recalculating the texture factor.
@@ -2495,6 +2528,11 @@ void USkeletalMesh::PostEditChangeProperty(FPropertyChangedEvent& PropertyChange
 		GEngine->TriggerStreamingDataRebuild();
 	}
 
+	if (!bSkipRestartRenderState)
+	{
+		RestartRenderState();
+	}
+
 	if( GIsEditor &&
 		PropertyThatChanged &&
 		PropertyThatChanged->GetFName() == FName(TEXT("PhysicsAsset")) )
@@ -2821,12 +2859,12 @@ void USkeletalMesh::RestartRenderState()
 	RecreateRenderState_Internal(this);
 }
 
-void USkeletalMesh::PreSave()
+void USkeletalMesh::PreSave(const class ITargetPlatform* TargetPlatform)
 {
 	// check the parent index of the root bone is invalid
 	check((RefSkeleton.GetNum() == 0) || (RefSkeleton.GetRefBoneInfo()[0].ParentIndex == INDEX_NONE));
 
-	Super::PreSave();
+	Super::PreSave(TargetPlatform);
 	// Make sure streaming texture factors have been cached. Calling GetStreamingTextureFactor
 	// with editor-only data available will calculate it if it has not been cached.
 	GetStreamingTextureFactor(0);
@@ -4503,6 +4541,10 @@ FSkeletalMeshSceneProxy::FSkeletalMeshSceneProxy(const USkinnedMeshComponent* Co
 		,	bCanHighlightSelectedSections(Component->bCanHighlightSelectedSections)
 		,	MaterialRelevance(Component->GetMaterialRelevance(GetScene().GetFeatureLevel()))
 		,	bMaterialsNeedMorphUsage_GameThread(false)
+#if WITH_EDITORONLY_DATA
+		,	StreamingDistanceMultiplier(FMath::Max(0.0f, Component->StreamingDistanceMultiplier))
+		,	StreamingTexelFactor(Component->SkeletalMesh->GetStreamingTextureFactor(0) * Component->ComponentToWorld.GetMaximumAxisScale())
+#endif
 {
 	check(MeshObject);
 	check(SkelMeshResource);
@@ -5213,6 +5255,24 @@ void FSkeletalMeshSceneProxy::UpdateMorphMaterialUsage_GameThread(bool bNeedsMor
 	}
 }
 
+#if WITH_EDITORONLY_DATA
+const FStreamingSectionBuildInfo* FSkeletalMeshSceneProxy::GetStreamingSectionData(float& OutComponentExtraScale, float& OutMeshExtraScale, int32 LODIndex, int32 ElementIndex) const
+{
+	static FStreamingSectionBuildInfo Data;
+	
+	OutMeshExtraScale = 1.f; // No extra scale as this scale is already taken into account in StreamingTexelFactor
+	OutComponentExtraScale = StreamingDistanceMultiplier;
+
+	Data.BoxOrigin = GetBounds().Origin;
+	Data.BoxExtent = GetBounds().BoxExtent;
+	for (int32 TexCoordIndex = 0; TexCoordIndex < FMaterialTexCoordBuildInfo::MAX_NUM_TEX_COORD; ++TexCoordIndex)
+	{
+		Data.TexelFactors[TexCoordIndex] = StreamingTexelFactor;
+	}
+	return &Data;
+}
+#endif
+
 USkinnedMeshComponent::USkinnedMeshComponent(const FObjectInitializer& ObjectInitializer)
 	: Super(ObjectInitializer)
 	, AnimUpdateRateParams(NULL)
diff --git a/Engine/Source/Runtime/Engine/Private/Streaming/TextureStreamingHelpers.cpp b/Engine/Source/Runtime/Engine/Private/Streaming/TextureStreamingHelpers.cpp
--- a/Engine/Source/Runtime/Engine/Private/Streaming/TextureStreamingHelpers.cpp
+++ b/Engine/Source/Runtime/Engine/Private/Streaming/TextureStreamingHelpers.cpp
@@ -187,12 +187,9 @@ double GStreamingDynamicPrimitivesTime = 0.0;
  */
 bool IsStreamingTexture( const UTexture2D* Texture2D )
 {
-	if (Texture2D && Texture2D->bIsStreamable && !Texture2D->NeverStream && Texture2D->GetNumMips() > UTexture2D::GetMinTextureResidentMipCount())
+	if (Texture2D && Texture2D->bIsStreamable && !Texture2D->NeverStream && Texture2D->GetNumMips() > Texture2D->GetNumNonStreamingMips())
 	{
-		const UShadowMapTexture2D* ShadowMap2D = Cast<const UShadowMapTexture2D>(Texture2D);
-		const ULightMapTexture2D* Lightmap2D = Cast<const ULightMapTexture2D>(Texture2D);
-
-		return !(Lightmap2D && (Lightmap2D->LightmapFlags & LMF_Streamed) == 0) && !(ShadowMap2D && (ShadowMap2D->ShadowmapFlags & SMF_Streamed) == 0);
+		return true;
 	}
 	return false;
 }
diff --git a/Engine/Source/Runtime/Engine/Private/Texture.cpp b/Engine/Source/Runtime/Engine/Private/Texture.cpp
--- a/Engine/Source/Runtime/Engine/Private/Texture.cpp
+++ b/Engine/Source/Runtime/Engine/Private/Texture.cpp
@@ -371,11 +371,11 @@ void UTexture::FinishDestroy()
 #endif
 }
 
-void UTexture::PreSave()
+void UTexture::PreSave(const class ITargetPlatform* TargetPlatform)
 {
 	PreSaveEvent.Broadcast(this);
 
-	Super::PreSave();
+	Super::PreSave(TargetPlatform);
 
 #if WITH_EDITOR
 	if (DeferCompression)
@@ -385,7 +385,8 @@ void UTexture::PreSave()
 		UpdateResource();
 	}
 
-	if( !GEngine->IsAutosaving() )
+	bool bIsCooking = TargetPlatform != nullptr;
+	if (!GEngine->IsAutosaving() && (!bIsCooking))
 	{
 		GWarn->StatusUpdate(0, 0, FText::Format(NSLOCTEXT("UnrealEd", "SavingPackage_CompressingSourceArt", "Compressing source art for texture:  {0}"), FText::FromString(GetName())));
 		Source.Compress();
diff --git a/Engine/Source/Runtime/Engine/Private/Texture2D.cpp b/Engine/Source/Runtime/Engine/Private/Texture2D.cpp
--- a/Engine/Source/Runtime/Engine/Private/Texture2D.cpp
+++ b/Engine/Source/Runtime/Engine/Private/Texture2D.cpp
@@ -365,9 +365,9 @@ void UTexture2D::PostLoad()
 	Super::PostLoad();
 }
 
-void UTexture2D::PreSave()
+void UTexture2D::PreSave(const class ITargetPlatform* TargetPlatform)
 {
-	Super::PreSave();
+	Super::PreSave(TargetPlatform);
 #if WITH_EDITOR
 	if( bTemporarilyDisableStreaming )
 	{
@@ -724,25 +724,20 @@ bool UTexture2D::HasAlphaChannel() const
 
 int32 UTexture2D::GetNumNonStreamingMips() const
 {
-	// Take in to account the mip tail.
-	int32 MipCount = GetNumMips();
-	int32 NumNonStreamingMips = FMath::Max(0, MipCount - GetMipTailBaseIndex());
+	int32 NumNonStreamingMips = 0;
 
-	// Take in to account the min resident limit.
-	NumNonStreamingMips = FMath::Max(NumNonStreamingMips, UTexture2D::GetMinTextureResidentMipCount());
-	NumNonStreamingMips = FMath::Min(NumNonStreamingMips, MipCount);
-
-	// Take in to account restrictions due to block size.
-	if (PlatformData && PlatformData->Mips.Num() > 0)
+	if (PlatformData)
 	{
-		EPixelFormat PixelFormat = PlatformData->PixelFormat;
-		int32 BlockSizeX = GPixelFormats[PixelFormat].BlockSizeX;
-		int32 BlockSizeY = GPixelFormats[PixelFormat].BlockSizeY;
-		if (BlockSizeX > 1 || BlockSizeY > 1)
-		{
-			NumNonStreamingMips = FMath::Max<int32>(NumNonStreamingMips, MipCount - FPlatformMath::FloorLog2(PlatformData->Mips[0].SizeX / BlockSizeX));
-			NumNonStreamingMips = FMath::Max<int32>(NumNonStreamingMips, MipCount - FPlatformMath::FloorLog2(PlatformData->Mips[0].SizeY / BlockSizeY));
-		}
+		NumNonStreamingMips = PlatformData->GetNumNonStreamingMips();
+	}
+	else
+	{
+		int32 MipCount = GetNumMips();
+		NumNonStreamingMips = FMath::Max(0, MipCount - GetMipTailBaseIndex());
+
+		// Take in to account the min resident limit.
+		NumNonStreamingMips = FMath::Max(NumNonStreamingMips, UTexture2D::GetMinTextureResidentMipCount());
+		NumNonStreamingMips = FMath::Min(NumNonStreamingMips, MipCount);
 	}
 
 	return NumNonStreamingMips;
diff --git a/Engine/Source/Runtime/Engine/Private/TextureDerivedData.cpp b/Engine/Source/Runtime/Engine/Private/TextureDerivedData.cpp
--- a/Engine/Source/Runtime/Engine/Private/TextureDerivedData.cpp
+++ b/Engine/Source/Runtime/Engine/Private/TextureDerivedData.cpp
@@ -1146,6 +1146,7 @@ bool FTexturePlatformData::TryLoadMips(int32 FirstMipToLoad, void** OutMipData)
 			if (OutMipData)
 			{
 				OutMipData[MipIndex - FirstMipToLoad] = FMemory::Malloc(Mip.BulkData.GetBulkDataSize());
+				checkSlow(!Mip.BulkData.GetFilename().EndsWith(TEXT(".ubulk"))); // We want to make sure that any non-streamed mips are coming from the texture asset file, and not from an external bulk file
 				Mip.BulkData.GetCopy(&OutMipData[MipIndex - FirstMipToLoad]);
 			}
 			NumMipsCached++;
@@ -1197,6 +1198,49 @@ bool FTexturePlatformData::TryLoadMips(int32 FirstMipToLoad, void** OutMipData)
 	return true;
 }
 
+int32 FTexturePlatformData::GetNumNonStreamingMips() const
+{
+	if (FPlatformProperties::RequiresCookedData())
+	{
+		// We're on a cooked platform so we should only be streaming mips that were not inlined in the texture by thecooker.
+		int32 NumNonStreamingMips = Mips.Num();
+
+		for (const FTexture2DMipMap& Mip : Mips)
+		{
+			uint32 BulkDataFlags = Mip.BulkData.GetBulkDataFlags();
+			if (BulkDataFlags & BULKDATA_PayloadInSeperateFile || BulkDataFlags & BULKDATA_PayloadAtEndOfFile)
+			{
+				--NumNonStreamingMips;
+			}
+			else
+			{
+				break;
+			}
+		}
+
+		return NumNonStreamingMips;
+	}
+	else
+	{
+		check(Mips.Num() > 0);
+		int32 MipCount = Mips.Num();
+		int32 NumNonStreamingMips = 1;
+
+		// Take in to account the min resident limit.
+		NumNonStreamingMips = FMath::Max(NumNonStreamingMips, UTexture2D::GetMinTextureResidentMipCount());
+		NumNonStreamingMips = FMath::Min(NumNonStreamingMips, MipCount);
+		int32 BlockSizeX = GPixelFormats[PixelFormat].BlockSizeX;
+		int32 BlockSizeY = GPixelFormats[PixelFormat].BlockSizeY;
+		if (BlockSizeX > 1 || BlockSizeY > 1)
+		{
+			NumNonStreamingMips = FMath::Max<int32>(NumNonStreamingMips, MipCount - FPlatformMath::FloorLog2(Mips[0].SizeX / BlockSizeX));
+			NumNonStreamingMips = FMath::Max<int32>(NumNonStreamingMips, MipCount - FPlatformMath::FloorLog2(Mips[0].SizeY / BlockSizeY));
+		}
+
+		return NumNonStreamingMips;
+	}
+}
+
 #if WITH_EDITOR
 bool FTexturePlatformData::AreDerivedMipsAvailable() const
 {
@@ -1275,8 +1319,13 @@ static void SerializePlatformData(
 	// Force resident mips inline
 	if (bCooked && Ar.IsSaving())
 	{
-		int32 MinMipToInline = bStreamable ? NumMips - UTexture2D::GetMinTextureResidentMipCount() : 0;
-		MinMipToInline = FMath::Max(MinMipToInline, 0);
+		int32 MinMipToInline = 0;
+			
+		if (bStreamable)
+		{
+			MinMipToInline = FMath::Max(0, NumMips - PlatformData->GetNumNonStreamingMips());
+		}
+
 		for (int32 MipIndex = MinMipToInline; MipIndex < NumMips; ++MipIndex)
 		{
 			PlatformData->Mips[MipIndex + FirstMipToSerialize].BulkData.SetBulkDataFlags(BULKDATA_ForceInlinePayload | BULKDATA_SingleUse);
diff --git a/Engine/Source/Runtime/Engine/Private/TimerManager.cpp b/Engine/Source/Runtime/Engine/Private/TimerManager.cpp
--- a/Engine/Source/Runtime/Engine/Private/TimerManager.cpp
+++ b/Engine/Source/Runtime/Engine/Private/TimerManager.cpp
@@ -9,6 +9,67 @@
 DECLARE_CYCLE_STAT(TEXT("SetTimer"), STAT_SetTimer, STATGROUP_Engine);
 DECLARE_CYCLE_STAT(TEXT("ClearTimer"), STAT_ClearTimer, STATGROUP_Engine);
 
+namespace
+{
+	void DescribeFTImerDataSafely(const FTimerData& Data)
+	{
+		UE_LOG(LogEngine, Log, TEXT("TimerData %p : bLoop=%s, bRequiresDelegate=%s, Status=%d, Rate=%f, ExpireTime=%f"),
+			&Data,
+			Data.bLoop ? TEXT("true") : TEXT("false"),
+			Data.bRequiresDelegate ? TEXT("true") : TEXT("false"),
+			static_cast<int32>(Data.Status),
+			Data.Rate,
+			Data.ExpireTime
+			)
+	}
+}
+
+FTimerManager::FTimerManager()
+	: InternalTime(0.0)
+	, LastTickedFrame(static_cast<uint64>(-1))
+	, LastAssignedHandle(0)
+{
+	if (IsRunningDedicatedServer())
+	{
+		FCoreDelegates::OnHandleSystemError.AddRaw(this, &FTimerManager::OnCrash);
+	}
+}
+
+FTimerManager::~FTimerManager()
+{
+	if (IsRunningDedicatedServer())
+	{
+		FCoreDelegates::OnHandleSystemError.RemoveAll(this);
+	}
+}
+
+void FTimerManager::OnCrash()
+{
+	UE_LOG(LogEngine, Warning, TEXT("TimerManager %p on crashing delegate called, dumping extra information"), this);
+
+	UE_LOG(LogEngine, Log, TEXT("------- %d Active Timers -------"), ActiveTimerHeap.Num());
+	for(const FTimerData& Data : ActiveTimerHeap)
+	{
+		DescribeFTImerDataSafely(Data);
+	}
+
+	UE_LOG(LogEngine, Log, TEXT("------- %d Paused Timers -------"), PausedTimerList.Num());
+	for (const FTimerData& Data : PausedTimerList)
+	{
+		DescribeFTImerDataSafely(Data);
+	}
+
+	UE_LOG(LogEngine, Log, TEXT("------- %d Pending Timers -------"), PendingTimerList.Num());
+	for (const FTimerData& Data : PendingTimerList)
+	{
+		DescribeFTImerDataSafely(Data);
+	}
+
+	UE_LOG(LogEngine, Log, TEXT("------- %d Total Timers -------"), PendingTimerList.Num() + PausedTimerList.Num() + ActiveTimerHeap.Num());
+
+	UE_LOG(LogEngine, Warning, TEXT("TimerManager %p dump ended"), this);
+}
+
 void FTimerHandle::MakeValid()
 {
 	if (GWorld)
@@ -579,4 +640,4 @@ FAutoConsoleCommandWithWorld ListTimersConsoleCommand(
 	TEXT("ListTimers"),
 	TEXT(""),
 	FConsoleCommandWithWorldDelegate::CreateStatic(OnListTimers)
-	);
\ No newline at end of file
+	);
diff --git a/Engine/Source/Runtime/Engine/Private/UnrealEngine.cpp b/Engine/Source/Runtime/Engine/Private/UnrealEngine.cpp
--- a/Engine/Source/Runtime/Engine/Private/UnrealEngine.cpp
+++ b/Engine/Source/Runtime/Engine/Private/UnrealEngine.cpp
@@ -5898,7 +5898,7 @@ bool UEngine::PerformError(const TCHAR* Cmd, FOutputDevice& Ar)
 #if !UE_BUILD_SHIPPING
 	if (FParse::Command(&Cmd, TEXT("RENDERCRASH")))
 	{
-		ENQUEUE_UNIQUE_RENDER_COMMAND(CauseRenderThreadCrash, { UE_LOG(LogEngine, Warning, TEXT("Printed warning to log.")); SetImageIntegrityStatus(-1); UE_LOG(LogEngine, Fatal, TEXT("Crashing the renderthread at your request")); });
+		ENQUEUE_UNIQUE_RENDER_COMMAND(CauseRenderThreadCrash, { UE_LOG(LogEngine, Warning, TEXT("Printed warning to log.")); SetCrashType(ECrashType::Debug); UE_LOG(LogEngine, Fatal, TEXT("Crashing the renderthread at your request")); });
 		return true;
 	}
 	if (FParse::Command(&Cmd, TEXT("RENDERCHECK")))
@@ -5908,7 +5908,7 @@ bool UEngine::PerformError(const TCHAR* Cmd, FOutputDevice& Ar)
 			static void Check()
 			{				
 				UE_LOG(LogEngine, Warning, TEXT("Printed warning to log."));
-				SetImageIntegrityStatus(-1);
+				SetCrashType(ECrashType::Debug);
 				check(!"Crashing the renderthread via check(0) at your request");
 			}
 		};
@@ -5917,7 +5917,7 @@ bool UEngine::PerformError(const TCHAR* Cmd, FOutputDevice& Ar)
 	}
 	if (FParse::Command(&Cmd, TEXT("RENDERGPF")))
 	{
-		ENQUEUE_UNIQUE_RENDER_COMMAND(CauseRenderThreadCrash, { UE_LOG(LogEngine, Warning, TEXT("Printed warning to log.")); SetImageIntegrityStatus(-1); *(int32 *)3 = 123; });
+		ENQUEUE_UNIQUE_RENDER_COMMAND(CauseRenderThreadCrash, { UE_LOG(LogEngine, Warning, TEXT("Printed warning to log.")); SetCrashType(ECrashType::Debug); *(int32 *)3 = 123; });
 		return true;
 	}
 	if (FParse::Command(&Cmd, TEXT("RENDERFATAL")))
@@ -5925,7 +5925,7 @@ bool UEngine::PerformError(const TCHAR* Cmd, FOutputDevice& Ar)
 		ENQUEUE_UNIQUE_RENDER_COMMAND(CauseRenderThreadCrash,
 		{
 			UE_LOG(LogEngine, Warning, TEXT("Printed warning to log."));
-			SetImageIntegrityStatus(-1);
+			SetCrashType(ECrashType::Debug);
 			LowLevelFatalError(TEXT("FError::LowLevelFatal test"));
 		});
 		return true;
@@ -5949,7 +5949,7 @@ bool UEngine::PerformError(const TCHAR* Cmd, FOutputDevice& Ar)
 			static void Crash(ENamedThreads::Type, const  FGraphEventRef&)
 			{
 				UE_LOG(LogEngine, Warning, TEXT("Printed warning to log."));
-				SetImageIntegrityStatus(-1);
+				SetCrashType(ECrashType::Debug);
 				UE_LOG(LogEngine, Fatal, TEXT("Crashing the worker thread at your request"));
 			}
 		};
@@ -5974,7 +5974,7 @@ bool UEngine::PerformError(const TCHAR* Cmd, FOutputDevice& Ar)
 			static void Check(ENamedThreads::Type, const FGraphEventRef&)
 			{
 				UE_LOG(LogEngine, Warning, TEXT("Printed warning to log."));
-				SetImageIntegrityStatus(-1);
+				SetCrashType(ECrashType::Debug);
 				check(!"Crashing a worker thread via check(0) at your request");
 			}
 		};
@@ -5999,7 +5999,7 @@ bool UEngine::PerformError(const TCHAR* Cmd, FOutputDevice& Ar)
 			static void GPF(ENamedThreads::Type, const FGraphEventRef&)
 			{
 				UE_LOG(LogEngine, Warning, TEXT("Printed warning to log."));
-				SetImageIntegrityStatus(-1);
+				SetCrashType(ECrashType::Debug);
 				*(int32 *)3 = 123;
 			}
 		};
@@ -6047,7 +6047,7 @@ bool UEngine::PerformError(const TCHAR* Cmd, FOutputDevice& Ar)
 			static void Fatal(ENamedThreads::Type, const FGraphEventRef&)
 			{
 				UE_LOG(LogEngine, Warning, TEXT("Printed warning to log."));
-				SetImageIntegrityStatus(-1);
+				SetCrashType(ECrashType::Debug);
 				LowLevelFatalError(TEXT("FError::LowLevelFatal test"));
 			}
 		};
@@ -6067,14 +6067,14 @@ bool UEngine::PerformError(const TCHAR* Cmd, FOutputDevice& Ar)
 	else if (FParse::Command(&Cmd, TEXT("CRASH")))
 	{
 		UE_LOG(LogEngine, Warning, TEXT("Printed warning to log."));
-		SetImageIntegrityStatus(-1);
+		SetCrashType(ECrashType::Debug);
 		UE_LOG(LogEngine, Fatal, TEXT("%s"), TEXT("Crashing the gamethread at your request"));
 		return true;
 	}
 	else if (FParse::Command(&Cmd, TEXT("CHECK")))
 	{
 		UE_LOG(LogEngine, Warning, TEXT("Printed warning to log."));
-		SetImageIntegrityStatus(-1);
+		SetCrashType(ECrashType::Debug);
 		check(!"Crashing the game thread via check(0) at your request");
 		return true;
 	}
@@ -6082,7 +6082,7 @@ bool UEngine::PerformError(const TCHAR* Cmd, FOutputDevice& Ar)
 	{
 		UE_LOG(LogEngine, Warning, TEXT("Printed warning to log."));
 		Ar.Log(TEXT("Crashing with voluntary GPF"));
-		SetImageIntegrityStatus(-1);
+		SetCrashType(ECrashType::Debug);
 		// changed to 3 from NULL because clang noticed writing to NULL and warned about it
 		*(int32 *)3 = 123;
 		return true;
@@ -6106,7 +6106,7 @@ bool UEngine::PerformError(const TCHAR* Cmd, FOutputDevice& Ar)
 	else if (FParse::Command(&Cmd, TEXT("FATAL")))
 	{
 		UE_LOG(LogEngine, Warning, TEXT("Printed warning to log."));
-		SetImageIntegrityStatus(-1);
+		SetCrashType(ECrashType::Debug);
 		LowLevelFatalError(TEXT("FError::LowLevelFatal test"));
 		return true;
 	}
@@ -6114,13 +6114,13 @@ bool UEngine::PerformError(const TCHAR* Cmd, FOutputDevice& Ar)
 	{
 		// stack overflow test - this case should be caught by /GS (Buffer Overflow Check) compile option
 		ANSICHAR SrcBuffer[] = "12345678901234567890123456789012345678901234567890";
-		SetImageIntegrityStatus(-1);
+		SetCrashType(ECrashType::Debug);
 		BufferOverflowFunction(ARRAY_COUNT(SrcBuffer), SrcBuffer);
 		return true;
 	}
 	else if (FParse::Command(&Cmd, TEXT("CRTINVALID")))
 	{
-		SetImageIntegrityStatus(-1);
+		SetCrashType(ECrashType::Debug);
 		FString::Printf(NULL);
 		return true;
 	}
@@ -6143,7 +6143,7 @@ bool UEngine::PerformError(const TCHAR* Cmd, FOutputDevice& Ar)
 	{
 		Ar.Logf(TEXT("Recursing to create a very deep callstack."));
 		GLog->Flush();
-		SetImageIntegrityStatus(-1);
+		SetCrashType(ECrashType::Debug);
 		InfiniteRecursionFunction(1);
 		Ar.Logf(TEXT("You will never see this log line."));
 		return true;
@@ -6155,7 +6155,7 @@ bool UEngine::PerformError(const TCHAR* Cmd, FOutputDevice& Ar)
 		{
 			static void InfiniteRecursion(ENamedThreads::Type, const FGraphEventRef&)
 			{
-				SetImageIntegrityStatus(-1);
+				SetCrashType(ECrashType::Debug);
 				InfiniteRecursionFunction(1);
 			}
 		};
@@ -6175,7 +6175,7 @@ bool UEngine::PerformError(const TCHAR* Cmd, FOutputDevice& Ar)
 	else if (FParse::Command(&Cmd, TEXT("EATMEM")))
 	{
 		Ar.Log(TEXT("Eating up all available memory"));
-		SetImageIntegrityStatus(-1);
+		SetCrashType(ECrashType::Debug);
 		while (1)
 		{
 			void* Eat = FMemory::Malloc(65536);
@@ -6192,7 +6192,7 @@ bool UEngine::PerformError(const TCHAR* Cmd, FOutputDevice& Ar)
 	else if (FParse::Command(&Cmd, TEXT("STACKOVERFLOW")))
 	{
 		Ar.Log(TEXT("Infinite recursion to cause stack overflow"));
-		SetImageIntegrityStatus(-1);
+		SetCrashType(ECrashType::Debug);
 		StackOverflowFunction(nullptr);
 		return true;
 	}
@@ -6203,7 +6203,7 @@ bool UEngine::PerformError(const TCHAR* Cmd, FOutputDevice& Ar)
 		{
 			static void StackOverflow(ENamedThreads::Type, const FGraphEventRef&)
 			{
-				SetImageIntegrityStatus(-1);
+				SetCrashType(ECrashType::Debug);
 				StackOverflowFunction(nullptr);
 			}
 		};
@@ -6223,7 +6223,7 @@ bool UEngine::PerformError(const TCHAR* Cmd, FOutputDevice& Ar)
 	else if (FParse::Command(&Cmd, TEXT("SOFTLOCK")))
 	{
 		Ar.Log(TEXT("Hanging the current thread"));
-		SetImageIntegrityStatus(-1);
+		SetCrashType(ECrashType::Debug);
 		while (1)
 		{
 			FPlatformProcess::Sleep(1.0f);
@@ -6233,7 +6233,7 @@ bool UEngine::PerformError(const TCHAR* Cmd, FOutputDevice& Ar)
 	else if (FParse::Command(&Cmd, TEXT("INFINITELOOP")))
 	{
 		Ar.Log(TEXT("Hanging the current thread (CPU-intensive)"));
-		SetImageIntegrityStatus(-1);
+		SetCrashType(ECrashType::Debug);
 		for(;;)
 		{
 		}
diff --git a/Engine/Source/Runtime/Engine/Private/UnrealNetwork.cpp b/Engine/Source/Runtime/Engine/Private/UnrealNetwork.cpp
--- a/Engine/Source/Runtime/Engine/Private/UnrealNetwork.cpp
+++ b/Engine/Source/Runtime/Engine/Private/UnrealNetwork.cpp
@@ -79,10 +79,6 @@ bool FNetworkVersion::IsNetworkCompatible( const uint32 LocalNetworkVersion, con
 
 FNetworkReplayVersion FNetworkVersion::GetReplayVersion()
 {
-	// Temp until Version.h is unlocked :O)
-#ifndef ENGINE_REPLAY_VERSION
-	#define ENGINE_REPLAY_VERSION BUILT_FROM_CHANGELIST
-#endif
 	return FNetworkReplayVersion(FApp::GetGameName(), 0, ENGINE_REPLAY_VERSION);
 }
 
diff --git a/Engine/Source/Runtime/Engine/Private/UserInterface/PlayerInput.cpp b/Engine/Source/Runtime/Engine/Private/UserInterface/PlayerInput.cpp
--- a/Engine/Source/Runtime/Engine/Private/UserInterface/PlayerInput.cpp
+++ b/Engine/Source/Runtime/Engine/Private/UserInterface/PlayerInput.cpp
@@ -560,6 +560,7 @@ void UPlayerInput::AddActionMapping(const FInputActionKeyMapping& KeyMapping)
 {
 	ActionMappings.AddUnique(KeyMapping);
 	ActionKeyMap.Reset();
+	bKeyMapsBuilt = false;
 }
 
 void UPlayerInput::RemoveActionMapping(const FInputActionKeyMapping& KeyMapping)
@@ -570,6 +571,7 @@ void UPlayerInput::RemoveActionMapping(const FInputActionKeyMapping& KeyMapping)
 		{
 			ActionMappings.RemoveAtSwap(ActionIndex);
 			ActionKeyMap.Reset();
+			bKeyMapsBuilt = false;
 			// we don't break because the mapping may have been in the array twice
 		}
 	}
@@ -579,6 +581,7 @@ void UPlayerInput::AddAxisMapping(const FInputAxisKeyMapping& KeyMapping)
 {
 	AxisMappings.AddUnique(KeyMapping);
 	AxisKeyMap.Reset();
+	bKeyMapsBuilt = false;
 }
 
 void UPlayerInput::RemoveAxisMapping(const FInputAxisKeyMapping& InKeyMapping)
@@ -591,6 +594,7 @@ void UPlayerInput::RemoveAxisMapping(const FInputAxisKeyMapping& InKeyMapping)
 		{
 			AxisMappings.RemoveAtSwap(AxisIndex);
 			AxisKeyMap.Reset();
+			bKeyMapsBuilt = false;
 			// we don't break because the mapping may have been in the array twice
 		}
 	}
@@ -602,6 +606,7 @@ void UPlayerInput::AddEngineDefinedActionMapping(const FInputActionKeyMapping& A
 	for (TObjectIterator<UPlayerInput> It; It; ++It)
 	{
 		It->ActionKeyMap.Reset();
+		It->bKeyMapsBuilt = false;
 	}
 }
 
@@ -611,6 +616,7 @@ void UPlayerInput::AddEngineDefinedAxisMapping(const FInputAxisKeyMapping& AxisM
 	for (TObjectIterator<UPlayerInput> It; It; ++It)
 	{
 		It->AxisKeyMap.Reset();
+		It->bKeyMapsBuilt = false;
 	}
 }
 
@@ -626,9 +632,10 @@ void UPlayerInput::ForceRebuildingKeyMaps(const bool bRestoreDefaults)
 	ActionKeyMap.Reset();
 	AxisKeyMap.Reset();
 	AxisProperties.Reset();
+	bKeyMapsBuilt = false;
 }
 
-void UPlayerInput::ConditionalBuildKeyMappings()
+void UPlayerInput::ConditionalBuildKeyMappings_Internal()
 {
 	if (ActionKeyMap.Num() == 0)
 	{
@@ -647,6 +654,7 @@ void UPlayerInput::ConditionalBuildKeyMappings()
 		ActionMappingsUtility.Build(ActionMappings, ActionKeyMap);
 		ActionMappingsUtility.Build(EngineDefinedActionMappings, ActionKeyMap);
 	}
+
 	if (AxisKeyMap.Num() == 0)
 	{
 		struct
@@ -687,13 +695,16 @@ void UPlayerInput::ConditionalBuildKeyMappings()
 			}
 		}
 	}
+
+	bKeyMapsBuilt = true;
 }
 
 void UPlayerInput::GetChordsForKeyMapping(const FInputActionKeyMapping& KeyMapping, const FInputActionBinding& ActionBinding, const bool bGamePaused, TArray<FDelegateDispatchDetails>& FoundChords, TArray<FKey>& KeysToConsume)
 {
-	TArray<uint32> EventIndices;
 	bool bConsumeInput = false;
 
+	checkSlow(!EventIndices.Num());
+
 	// test modifier conditions and ignore the event if they failed
 	if (	(KeyMapping.bAlt == false || IsAltPressed())
 		&&	(KeyMapping.bCtrl == false || IsCtrlPressed())
@@ -745,6 +756,8 @@ void UPlayerInput::GetChordsForKeyMapping(const FInputActionKeyMapping& KeyMappi
 	{
 		KeysToConsume.AddUnique(KeyMapping.Key);
 	}
+
+	EventIndices.Reset();
 }
 
 void UPlayerInput::GetChordsForAction(const FInputActionBinding& ActionBinding, const bool bGamePaused, TArray<FDelegateDispatchDetails>& FoundChords, TArray<FKey>& KeysToConsume)
@@ -794,7 +807,7 @@ void UPlayerInput::GetChordForKey(const FInputKeyBinding& KeyBinding, const bool
 	}
 	else if ( !IsKeyConsumed(KeyBinding.Chord.Key) )
 	{
-		TArray<uint32> EventIndices;
+		checkSlow(!EventIndices.Num());
 
 		// test modifier conditions and ignore the event if they failed
 		if (	(KeyBinding.Chord.bAlt == false || IsAltPressed())
@@ -839,6 +852,8 @@ void UPlayerInput::GetChordForKey(const FInputKeyBinding& KeyBinding, const bool
 				}
 				bConsumeInput = true;
 			}
+
+			EventIndices.Reset();
 		}
 	}
 	if (KeyBinding.bConsumeInput && (bConsumeInput || !(KeyBinding.Chord.bAlt || KeyBinding.Chord.bCtrl || KeyBinding.Chord.bShift || KeyBinding.Chord.bCmd || KeyBinding.KeyEvent == EInputEvent::IE_DoubleClick)))
@@ -990,7 +1005,7 @@ void UPlayerInput::ProcessInputStack(const TArray<UInputComponent*>& InputCompon
 
 
 	// must be called non-recursively and on the game thread
-	check(IsInGameThread() && !AxisDelegates.Num() && !VectorAxisDelegates.Num() && !NonAxisDelegates.Num() && !KeysToConsume.Num() && !FoundChords.Num());
+	check(IsInGameThread() && !AxisDelegates.Num() && !VectorAxisDelegates.Num() && !NonAxisDelegates.Num() && !KeysToConsume.Num() && !FoundChords.Num() && !EventIndices.Num());
 
 	struct FDelegateDispatchDetailsSorter
 	{
@@ -1008,7 +1023,7 @@ void UPlayerInput::ProcessInputStack(const TArray<UInputComponent*>& InputCompon
 		UInputComponent* const IC = InputComponentStack[StackIndex];
 		if (IC)
 		{
-			check(!KeysToConsume.Num() && !FoundChords.Num())
+			check(!KeysToConsume.Num() && !FoundChords.Num() && !EventIndices.Num())
 
 			for (int32 ActionIndex=0; ActionIndex<IC->GetNumActionBindings(); ++ActionIndex)
 			{
@@ -1059,7 +1074,6 @@ void UPlayerInput::ProcessInputStack(const TArray<UInputComponent*>& InputCompon
 			{
 				FInputTouchBinding& TB = IC->TouchBindings[TouchBindingIndex];
 
-				TArray<uint32> EventIndices;
 				for (ETouchIndex::Type TouchIndex = ETouchIndex::Touch1; TouchIndex <= ETouchIndex::Touch10; TouchIndex = ETouchIndex::Type(TouchIndex + 1))
 				{
 					const FKey& TouchKey = EKeys::TouchKeys[TouchIndex];
@@ -1085,13 +1099,13 @@ void UPlayerInput::ProcessInputStack(const TArray<UInputComponent*>& InputCompon
 						}
 					}
 				}
+
+				EventIndices.Reset();
 			}
 
 			// look for any gestures that happened
 			for (FInputGestureBinding& GB : IC->GestureBindings)
 			{
-				TArray<uint32> EventIndices;
-
 				// treat gestures as fire-and-forget, so by convention we assume if they happen, it was a "pressed" event
 				if (!IsKeyConsumed(GB.GestureKey) && KeyEventOccurred(GB.GestureKey, IE_Pressed, EventIndices))
 				{
@@ -1108,6 +1122,8 @@ void UPlayerInput::ProcessInputStack(const TArray<UInputComponent*>& InputCompon
 						}
 					}
 				}
+
+				EventIndices.Reset();
 			}
 
 			// Run though game axis bindings and accumulate axis values
@@ -1656,12 +1672,12 @@ void UPlayerInput::ConsumeKey(FKey Key)
 	}
 }
 
-bool UPlayerInput::KeyEventOccurred(FKey Key, EInputEvent Event, TArray<uint32>& EventIndices) const
+bool UPlayerInput::KeyEventOccurred(FKey Key, EInputEvent Event, TArray<uint32>& InEventIndices) const
 {
 	FKeyState const* const KeyState = KeyStateMap.Find(Key);
 	if (KeyState && KeyState->EventCounts[Event].Num() > 0)
 	{
-		EventIndices = KeyState->EventCounts[Event];
+		InEventIndices = KeyState->EventCounts[Event];
 		return true;
 	}
 
diff --git a/Engine/Source/Runtime/Engine/Public/Scalability.h b/Engine/Source/Runtime/Engine/Public/Scalability.h
--- a/Engine/Source/Runtime/Engine/Public/Scalability.h
+++ b/Engine/Source/Runtime/Engine/Public/Scalability.h
@@ -70,6 +70,12 @@ namespace Scalability
 	/** This is the only suggested way to get the current state - don't get CVars directly */
 	ENGINE_API FQualityLevels GetQualityLevels();
 
+	/** Gets the effects quality directly for the passed thread.
+	*
+	* @param bGameThread	If true, the game thread value for the CVar is returned, otherwise the render thread value is returned. Useful when accessing the CVar from a game task.
+	*/
+	int32 GetEffectsQualityDirect(bool bGameThread);
+
 	/**  */
 	ENGINE_API void InitScalabilitySystem();
 
@@ -100,35 +106,4 @@ namespace Scalability
 
 	/** Maximum single axis scale for render resolution */
 	static const float MaxResolutionScale = 100.0f;
-
-
-	// A set of accessor helper objects that allow fast and thread safe access to the scalability CVars values without giving direct access.
-	struct ENGINE_API FResolutionQualityCVarAccessor : public FCVarAccessorFloat
-	{
-		FResolutionQualityCVarAccessor();
-	};
-	struct ENGINE_API FViewDistanceQualityCVarAccessor : public FCVarAccessorInt
-	{
-		FViewDistanceQualityCVarAccessor();
-	};
-	struct ENGINE_API FAntiAliasingQualityCVarAccessor : public FCVarAccessorInt
-	{
-		FAntiAliasingQualityCVarAccessor();
-	};
-	struct ENGINE_API FShadowQualityCVarAccessor : public FCVarAccessorInt
-	{
-		FShadowQualityCVarAccessor();
-	};
-	struct ENGINE_API FPostProcessQualityCVarAccessor : public FCVarAccessorInt
-	{
-		FPostProcessQualityCVarAccessor();
-	};
-	struct ENGINE_API FTextureQualityCVarAccessor : public FCVarAccessorInt
-	{
-		FTextureQualityCVarAccessor();
-	};
-	struct ENGINE_API FEffectsQualityCVarAccessor : public FCVarAccessorInt
-	{
-		FEffectsQualityCVarAccessor();
-	};
 }
diff --git a/Engine/Source/Runtime/Engine/Public/SkeletalMeshTypes.h b/Engine/Source/Runtime/Engine/Public/SkeletalMeshTypes.h
--- a/Engine/Source/Runtime/Engine/Public/SkeletalMeshTypes.h
+++ b/Engine/Source/Runtime/Engine/Public/SkeletalMeshTypes.h
@@ -1906,6 +1906,11 @@ public:
 	*/
 	void UpdateMorphMaterialUsage_GameThread(bool bNeedsMorphUsage);
 
+
+#if WITH_EDITORONLY_DATA
+	virtual const FStreamingSectionBuildInfo* GetStreamingSectionData(float& OutComponentExtraScale, float& OutMeshExtraScale, int32 LODIndex, int32 ElementIndex) const override;
+#endif
+
 	friend class FSkeletalMeshSectionIter;
 
 protected:
@@ -1969,6 +1974,13 @@ protected:
 	TSet<UMaterialInterface*> MaterialsInUse_GameThread;
 	bool bMaterialsNeedMorphUsage_GameThread;
 	
+#if WITH_EDITORONLY_DATA
+	/** The component streaming distance multiplier */
+	float StreamingDistanceMultiplier;
+	/** The mesh streaming texel factor (fallback) */
+	float StreamingTexelFactor;
+#endif
+
 	void GetDynamicElementsSection(const TArray<const FSceneView*>& Views, const FSceneViewFamily& ViewFamily, uint32 VisibilityMap, 
 		const FStaticLODModel& LODModel, const int32 LODIndex, const FSkelMeshSection& Section, const FSkelMeshChunk& Chunk, 
 		const FSectionElementInfo& SectionElementInfo, const FTwoVectors& CustomLeftRightVectors, bool bInSelectable, FMeshElementCollector& Collector ) const;
diff --git a/Engine/Source/Runtime/Engine/Public/TimerManager.h b/Engine/Source/Runtime/Engine/Public/TimerManager.h
--- a/Engine/Source/Runtime/Engine/Public/TimerManager.h
+++ b/Engine/Source/Runtime/Engine/Public/TimerManager.h
@@ -153,12 +153,13 @@ public:
 	// ----------------------------------
 	// Timer API
 
-	FTimerManager()
-		: InternalTime(0.0)
-		, LastTickedFrame(static_cast<uint64>(-1))
-		, LastAssignedHandle(0)
-	{}
+	FTimerManager();
+	virtual ~FTimerManager();
 
+	/**
+	 * Called from crash handler to provide more debug information.
+	 */
+	virtual void OnCrash();
 
 	/**
 	 * Sets a timer to call the given native function at a set interval.  If a timer is already set
diff --git a/Engine/Source/Runtime/GameplayAbilities/Private/Abilities/GameplayAbility.cpp b/Engine/Source/Runtime/GameplayAbilities/Private/Abilities/GameplayAbility.cpp
--- a/Engine/Source/Runtime/GameplayAbilities/Private/Abilities/GameplayAbility.cpp
+++ b/Engine/Source/Runtime/GameplayAbilities/Private/Abilities/GameplayAbility.cpp
@@ -1257,6 +1257,13 @@ void UGameplayAbility::K2_ExecuteGameplayCueWithParams(FGameplayTag GameplayCueT
 void UGameplayAbility::K2_AddGameplayCue(FGameplayTag GameplayCueTag, FGameplayEffectContextHandle Context, bool bRemoveOnAbilityEnd)
 {
 	check(CurrentActorInfo);
+
+	// Make default context if nothing is passed in
+	if (Context.IsValid() == false)
+	{
+		Context = MakeEffectContext(CurrentSpecHandle, CurrentActorInfo);
+	}
+
 	Context.SetAbility(this);
 
 	CurrentActorInfo->AbilitySystemComponent->AddGameplayCue(GameplayCueTag, Context);
diff --git a/Engine/Source/Runtime/GameplayAbilities/Private/Abilities/Tasks/AbilityTask_ApplyRootMotionRadialForce.cpp b/Engine/Source/Runtime/GameplayAbilities/Private/Abilities/Tasks/AbilityTask_ApplyRootMotionRadialForce.cpp
--- a/Engine/Source/Runtime/GameplayAbilities/Private/Abilities/Tasks/AbilityTask_ApplyRootMotionRadialForce.cpp
+++ b/Engine/Source/Runtime/GameplayAbilities/Private/Abilities/Tasks/AbilityTask_ApplyRootMotionRadialForce.cpp
@@ -15,9 +15,10 @@ UAbilityTask_ApplyRootMotionRadialForce::UAbilityTask_ApplyRootMotionRadialForce
 	MovementComponent = nullptr;
 	StrengthDistanceFalloff = nullptr;
 	StrengthOverTime = nullptr;
+	bUseFixedWorldDirection = false;
 }
 
-UAbilityTask_ApplyRootMotionRadialForce* UAbilityTask_ApplyRootMotionRadialForce::ApplyRootMotionRadialForce(UObject* WorldContextObject, FName TaskInstanceName, FVector Location, AActor* LocationActor, float Strength, float Duration, float Radius, bool bIsPush, bool bIsAdditive, bool bNoZForce, UCurveFloat* StrengthDistanceFalloff, UCurveFloat* StrengthOverTime)
+UAbilityTask_ApplyRootMotionRadialForce* UAbilityTask_ApplyRootMotionRadialForce::ApplyRootMotionRadialForce(UObject* WorldContextObject, FName TaskInstanceName, FVector Location, AActor* LocationActor, float Strength, float Duration, float Radius, bool bIsPush, bool bIsAdditive, bool bNoZForce, UCurveFloat* StrengthDistanceFalloff, UCurveFloat* StrengthOverTime, bool bUseFixedWorldDirection, FRotator FixedWorldDirection)
 {
 	auto MyTask = NewAbilityTask<UAbilityTask_ApplyRootMotionRadialForce>(WorldContextObject, TaskInstanceName);
 
@@ -32,6 +33,8 @@ UAbilityTask_ApplyRootMotionRadialForce* UAbilityTask_ApplyRootMotionRadialForce
 	MyTask->bNoZForce = bNoZForce;
 	MyTask->StrengthDistanceFalloff = StrengthDistanceFalloff;
 	MyTask->StrengthOverTime = StrengthOverTime;
+	MyTask->bUseFixedWorldDirection = bUseFixedWorldDirection;
+	MyTask->FixedWorldDirection = FixedWorldDirection;
 	MyTask->SharedInitAndApply();
 
 	return MyTask;
@@ -72,6 +75,8 @@ void UAbilityTask_ApplyRootMotionRadialForce::SharedInitAndApply()
 			RadialForce->bNoZForce = bNoZForce;
 			RadialForce->StrengthDistanceFalloff = StrengthDistanceFalloff;
 			RadialForce->StrengthOverTime = StrengthOverTime;
+			RadialForce->bUseFixedWorldDirection = bUseFixedWorldDirection;
+			RadialForce->FixedWorldDirection = FixedWorldDirection;
 			RootMotionSourceID = MovementComponent->ApplyRootMotionSource(RadialForce);
 
 			if (Ability)
@@ -134,6 +139,8 @@ void UAbilityTask_ApplyRootMotionRadialForce::GetLifetimeReplicatedProps(TArray<
 	DOREPLIFETIME(UAbilityTask_ApplyRootMotionRadialForce, bNoZForce);
 	DOREPLIFETIME(UAbilityTask_ApplyRootMotionRadialForce, StrengthDistanceFalloff);
 	DOREPLIFETIME(UAbilityTask_ApplyRootMotionRadialForce, StrengthOverTime);
+	DOREPLIFETIME(UAbilityTask_ApplyRootMotionRadialForce, bUseFixedWorldDirection);
+	DOREPLIFETIME(UAbilityTask_ApplyRootMotionRadialForce, FixedWorldDirection);
 }
 
 void UAbilityTask_ApplyRootMotionRadialForce::PreDestroyFromReplication()
diff --git a/Engine/Source/Runtime/GameplayAbilities/Private/AbilitySystemBlueprintLibrary.cpp b/Engine/Source/Runtime/GameplayAbilities/Private/AbilitySystemBlueprintLibrary.cpp
--- a/Engine/Source/Runtime/GameplayAbilities/Private/AbilitySystemBlueprintLibrary.cpp
+++ b/Engine/Source/Runtime/GameplayAbilities/Private/AbilitySystemBlueprintLibrary.cpp
@@ -417,6 +417,11 @@ bool UAbilitySystemBlueprintLibrary::EffectContextHasHitResult(FGameplayEffectCo
 	return EffectContext.GetHitResult() != NULL;
 }
 
+void UAbilitySystemBlueprintLibrary::EffectContextAddHitResult(FGameplayEffectContextHandle EffectContext, FHitResult HitResult, bool bReset)
+{
+	EffectContext.AddHitResult(HitResult, bReset);
+}
+
 AActor*	UAbilitySystemBlueprintLibrary::EffectContextGetInstigatorActor(FGameplayEffectContextHandle EffectContext)
 {
 	return EffectContext.GetInstigator();
