Disable warning in commandlets

diff --git a/Engine/Source/Runtime/Engine/Private/LevelActor.cpp b/Engine/Source/Runtime/Engine/Private/LevelActor.cpp
--- a/Engine/Source/Runtime/Engine/Private/LevelActor.cpp
+++ b/Engine/Source/Runtime/Engine/Private/LevelActor.cpp
@@ -620,9 +620,12 @@ bool UWorld::DestroyActor( AActor* ThisActor, bool bNetForce, bool bShouldModify
 	}
 	else
 	{
-		// Only worlds in the middle of seamless travel should have no context, and in that case, we shouldn't be destroying actors on them until
-		// they have become the current world (i.e. CopyWorldData has been called)
-		UE_LOG(LogSpawn, Warning, TEXT("UWorld::DestroyActor: World has no context! World: %s, Actor: %s"), *GetName(), *ThisActor->GetPathName());
+		if (!IsRunningCommandlet())
+		{
+			// Only worlds in the middle of seamless travel should have no context, and in that case, we shouldn't be destroying actors on them until
+			// they have become the current world (i.e. CopyWorldData has been called)
+			UE_LOG(LogSpawn, Warning, TEXT("UWorld::DestroyActor: World has no context! World: %s, Actor: %s"), *GetName(), *ThisActor->GetPathName());
+		}
 	}
 
 	if ( DemoNetDriver )
