Copying //UE4/Dev-Framework to Dev-Main (//UE4/Dev-Main)

diff --git a/Engine/Source/Runtime/Engine/Private/AI/Navigation/NavigationSystem.cpp b/Engine/Source/Runtime/Engine/Private/AI/Navigation/NavigationSystem.cpp
--- a/Engine/Source/Runtime/Engine/Private/AI/Navigation/NavigationSystem.cpp
+++ b/Engine/Source/Runtime/Engine/Private/AI/Navigation/NavigationSystem.cpp
@@ -3645,24 +3645,6 @@ FVector UNavigationSystem::ProjectPointToNavigation(UObject* WorldContextObject,
 	return ProjectedPoint.Location;
 }
 
-FVector UNavigationSystem::GetRandomPoint(UObject* WorldContextObject, ANavigationData* NavData, TSubclassOf<UNavigationQueryFilter> FilterClass)
-{
-	FNavLocation RandomPoint;
-
-	UWorld* World = GEngine->GetWorldFromContextObject( WorldContextObject );
-	UNavigationSystem* NavSys = UNavigationSystem::GetCurrent(World);
-	if (NavSys)
-	{
-		ANavigationData* UseNavData = NavData ? NavData : NavSys->GetMainNavData(FNavigationSystem::DontCreate);
-		if (UseNavData)
-		{
-			NavSys->GetRandomPoint(RandomPoint, UseNavData, UNavigationQueryFilter::GetQueryFilter(*UseNavData, WorldContextObject, FilterClass));
-		}
-	}
-
-	return RandomPoint.Location;
-}
-
 FVector UNavigationSystem::GetRandomReachablePointInRadius(UObject* WorldContext, const FVector& Origin, float Radius, ANavigationData* NavData, TSubclassOf<UNavigationQueryFilter> FilterClass)
 {
 	FNavLocation RandomPoint;
@@ -4043,32 +4025,6 @@ void UNavigationSystem::UnregisterNavigationInvoker(AActor* Invoker)
 //----------------------------------------------------------------------//
 // DEPRECATED
 //----------------------------------------------------------------------//
-void UNavigationSystem::RegisterCustomLink(INavLinkCustomInterface* CustomLink)
-{
-	if (CustomLink)
-	{
-		RegisterCustomLink(*CustomLink);
-	}
-}
-
-void UNavigationSystem::UnregisterCustomLink(INavLinkCustomInterface* CustomLink)
-{
-	if (CustomLink)
-	{
-		UnregisterCustomLink(*CustomLink);
-	}
-}
-
-FVector UNavigationSystem::GetRandomPointInRadius(UObject* WorldContextObject, const FVector& Origin, float Radius, ANavigationData* NavData, TSubclassOf<UNavigationQueryFilter> FilterClass)
-{
-	return UNavigationSystem::GetRandomReachablePointInRadius(WorldContextObject, Origin, Radius, NavData, FilterClass);
-}
-
-bool UNavigationSystem::GetRandomPointInRadius(const FVector& Origin, float Radius, FNavLocation& ResultLocation, ANavigationData* NavData, FSharedConstNavQueryFilter QueryFilter) const
-{
-	return GetRandomReachablePointInRadius(Origin, Radius, ResultLocation, NavData, QueryFilter);
-}
-
 void UNavigationSystem::UpdateNavOctree(AActor* Actor)
 {
 	if (Actor)
diff --git a/Engine/Source/Runtime/Engine/Private/Actor.cpp b/Engine/Source/Runtime/Engine/Private/Actor.cpp
--- a/Engine/Source/Runtime/Engine/Private/Actor.cpp
+++ b/Engine/Source/Runtime/Engine/Private/Actor.cpp
@@ -1742,7 +1742,7 @@ void AActor::ForceNetUpdate()
 
 void AActor::SetNetDormancy(ENetDormancy NewDormancy)
 {
-	if (GetNetMode() == NM_Client)
+	if (IsNetMode(NM_Client))
 	{
 		return;
 	}
@@ -1767,7 +1767,7 @@ void AActor::SetNetDormancy(ENetDormancy NewDormancy)
 /** Removes the actor from the NetDriver's dormancy list: forcing at least one more update. */
 void AActor::FlushNetDormancy()
 {
-	if (GetNetMode() == NM_Client || NetDormancy <= DORM_Awake)
+	if (IsNetMode(NM_Client) || NetDormancy <= DORM_Awake)
 	{
 		return;
 	}
@@ -2765,7 +2765,7 @@ void AActor::PostActorConstruction()
 	}
 
 	// If this is dynamically spawned replicted actor, defer calls to BeginPlay and UpdateOverlaps until replicated properties are deserialized
-	bool deferBeginPlayAndUpdateOverlaps = (bExchangedRoles && RemoteRole == ROLE_Authority);
+	const bool bDeferBeginPlayAndUpdateOverlaps = (bExchangedRoles && RemoteRole == ROLE_Authority);
 
 	if (bActorsInitialized)
 	{
@@ -2830,7 +2830,14 @@ void AActor::PostActorConstruction()
 					UE_LOG(LogActor, Fatal, TEXT("%s failed to route PostInitializeComponents.  Please call Super::PostInitializeComponents() in your <className>::PostInitializeComponents() function. "), *GetFullName());
 				}
 
-				if (World->HasBegunPlay() && !deferBeginPlayAndUpdateOverlaps)
+				bool bRunBeginPlay = !bDeferBeginPlayAndUpdateOverlaps && World->HasBegunPlay();
+				if (bRunBeginPlay && IsChildActor())
+				{
+					// Child Actors cannot run begin play until their parent has run
+					bRunBeginPlay = GetParentComponent()->GetOwner()->HasActorBegunPlay();
+				}
+
+				if (bRunBeginPlay)
 				{
 					SCOPE_CYCLE_COUNTER(STAT_ActorBeginPlay);
 					BeginPlay();
@@ -2851,7 +2858,7 @@ void AActor::PostActorConstruction()
 	if (!IsPendingKill())
 	{
 		// Components are all there and we've begun play, init overlapping state
-		if (!deferBeginPlayAndUpdateOverlaps)
+		if (!bDeferBeginPlayAndUpdateOverlaps)
 		{
 			UpdateOverlaps();
 		}
@@ -3500,21 +3507,23 @@ UNetDriver* GetNetDriver_Internal(UWorld* World, FName NetDriverName)
 	return GEngine->FindNamedNetDriver(World, NetDriverName);
 }
 
-ENetMode AActor::GetNetMode() const
+ENetMode AActor::InternalGetNetMode() const
 {
+	const bool bIsClientOnly = IsRunningClientOnly();
+
 	UWorld* World = GetWorld();
 	UNetDriver* NetDriver = GetNetDriver_Internal(World, NetDriverName);
 	if (NetDriver != nullptr)
 	{
-		return NetDriver->GetNetMode();
+		return bIsClientOnly ? NM_Client : NetDriver->GetNetMode();
 	}
 
 	if (World != nullptr && World->DemoNetDriver != nullptr)
 	{
-		return World->DemoNetDriver->GetNetMode();
+		return bIsClientOnly ? NM_Client : World->DemoNetDriver->GetNetMode();
 	}
 
-	return (IsRunningDedicatedServer() ? NM_DedicatedServer : NM_Standalone);
+	return NM_Standalone;
 }
 
 UNetDriver* AActor::GetNetDriver() const
@@ -4156,7 +4165,7 @@ bool AActor::ActorLineTraceSingle(struct FHitResult& OutHit, const FVector& Star
 float AActor::ActorGetDistanceToCollision(const FVector& Point, ECollisionChannel TraceChannel, FVector& ClosestPointOnCollision, UPrimitiveComponent** OutPrimitiveComponent) const
 {
 	ClosestPointOnCollision = Point;
-	float ClosestPointDistance = -1.f;
+	float ClosestPointDistanceSqr = -1.f;
 
 	TInlineComponentArray<UPrimitiveComponent*> Components;
 	GetComponents(Components);
@@ -4168,17 +4177,17 @@ float AActor::ActorGetDistanceToCollision(const FVector& Point, ECollisionChanne
 			&& (Primitive->GetCollisionResponseToChannel(TraceChannel) == ECollisionResponse::ECR_Block) )
 		{
 			FVector ClosestPoint;
-			const float Distance = Primitive->GetDistanceToCollision(Point, ClosestPoint);
+			float DistanceSqr = -1.f;
 
-			if (Distance < 0.f)
+			if (!Primitive->GetSquaredDistanceToCollision(Point, DistanceSqr, ClosestPoint))
 			{
 				// Invalid result, impossible to be better than ClosestPointDistance
 				continue;
 			}
 
-			if( (ClosestPointDistance < 0.f) || (Distance < ClosestPointDistance) )
+			if( (ClosestPointDistanceSqr < 0.f) || (DistanceSqr < ClosestPointDistanceSqr) )
 			{
-				ClosestPointDistance = Distance;
+				ClosestPointDistanceSqr = DistanceSqr;
 				ClosestPointOnCollision = ClosestPoint;
 				if( OutPrimitiveComponent )
 				{
@@ -4186,7 +4195,7 @@ float AActor::ActorGetDistanceToCollision(const FVector& Point, ECollisionChanne
 				}
 
 				// If we're inside collision, we're not going to find anything better, so abort search we've got our best find.
-				if( Distance <= KINDA_SMALL_NUMBER )
+				if( DistanceSqr <= KINDA_SMALL_NUMBER )
 				{
 					break;
 				}
@@ -4194,7 +4203,7 @@ float AActor::ActorGetDistanceToCollision(const FVector& Point, ECollisionChanne
 		}
 	}
 
-	return ClosestPointDistance;
+	return (ClosestPointDistanceSqr > 0.f ? FMath::Sqrt(ClosestPointDistanceSqr) : ClosestPointDistanceSqr);
 }
 
 
diff --git a/Engine/Source/Runtime/Engine/Private/Animation/AnimBlueprintGeneratedClass.cpp b/Engine/Source/Runtime/Engine/Private/Animation/AnimBlueprintGeneratedClass.cpp
--- a/Engine/Source/Runtime/Engine/Private/Animation/AnimBlueprintGeneratedClass.cpp
+++ b/Engine/Source/Runtime/Engine/Private/Animation/AnimBlueprintGeneratedClass.cpp
@@ -240,6 +240,15 @@ void UAnimBlueprintGeneratedClass::Link(FArchive& Ar, bool bRelinkExistingProper
 			ensure(bValidRootIndex);
 		}
 	}
+
+	if(RootClass != this)
+	{
+		if(OrderedSavedPoseIndices.Num() != RootClass->OrderedSavedPoseIndices.Num() || OrderedSavedPoseIndices != RootClass->OrderedSavedPoseIndices)
+		{
+			// Derived and our parent has a new ordered pose order, copy over.
+			OrderedSavedPoseIndices = RootClass->OrderedSavedPoseIndices;
+		}
+	}
 }
 
 void UAnimBlueprintGeneratedClass::PurgeClass(bool bRecompilingOnLoad)
diff --git a/Engine/Source/Runtime/Engine/Private/Animation/AnimCompressionDerivedData.h b/Engine/Source/Runtime/Engine/Private/Animation/AnimCompressionDerivedData.h
--- a/Engine/Source/Runtime/Engine/Private/Animation/AnimCompressionDerivedData.h
+++ b/Engine/Source/Runtime/Engine/Private/Animation/AnimCompressionDerivedData.h
@@ -43,7 +43,7 @@ public:
 		// This is a version string that mimics the old versioning scheme. If you
 		// want to bump this version, generate a new guid using VS->Tools->Create GUID and
 		// return it here. Ex.
-		return TEXT("A9288501F8C6460B80E9E07AD56FB59C");
+		return TEXT("D6830EA6B283485FA15ACC36BBCFA79A");
 	}
 
 	virtual FString GetPluginSpecificCacheKeySuffix() const override;
diff --git a/Engine/Source/Runtime/Engine/Private/Animation/AnimCurveTypes.cpp b/Engine/Source/Runtime/Engine/Private/Animation/AnimCurveTypes.cpp
--- a/Engine/Source/Runtime/Engine/Private/Animation/AnimCurveTypes.cpp
+++ b/Engine/Source/Runtime/Engine/Private/Animation/AnimCurveTypes.cpp
@@ -4,13 +4,37 @@
 #include "AnimCurveTypes.h"
 #include "Animation/AnimInstance.h"
 #include "AnimationRuntime.h"
-
+#include "FrameworkObjectVersion.h"
 
 DECLARE_CYCLE_STAT(TEXT("AnimSeq EvalCurveData"), STAT_AnimSeq_EvalCurveData, STATGROUP_Anim);
 
 /////////////////////////////////////////////////////
 // FFloatCurve
 
+void FAnimCurveBase::PostSerialize(FArchive& Ar)
+{
+	FSmartNameMapping::UID CurveUid = FSmartNameMapping::MaxUID;
+
+	if ((Ar.IsLoading()))
+	{
+		Ar.UsingCustomVersion(FFrameworkObjectVersion::GUID);
+		if (Ar.CustomVer(FFrameworkObjectVersion::GUID) < FFrameworkObjectVersion::SmartNameRefactor)
+		{
+			if (Ar.UE4Ver() >= VER_UE4_SKELETON_ADD_SMARTNAMES)
+			{
+				Ar << CurveUid;
+
+				Name.UID = CurveUid;
+				Name.DisplayName = LastObservedName_DEPRECATED;
+			}
+			else
+			{
+				Name.DisplayName = LastObservedName_DEPRECATED;
+			}
+		}
+	}
+}
+
 void FAnimCurveBase::SetCurveTypeFlag(EAnimCurveFlags InFlag, bool bValue)
 {
 	if (bValue)
@@ -165,7 +189,7 @@ void FRawCurveTracks::EvaluateCurveData( FBlendedCurve& Curves, float CurrentTim
 	for(auto CurveIter = FloatCurves.CreateConstIterator(); CurveIter; ++CurveIter)
 	{
 		const FFloatCurve& Curve = *CurveIter;
-		Curves.Set(Curve.CurveUid, Curve.Evaluate(CurrentTime), Curve.GetCurveTypeFlags());
+		Curves.Set(Curve.Name.UID, Curve.Evaluate(CurrentTime), Curve.GetCurveTypeFlags());
 	}
 }
 
@@ -188,18 +212,12 @@ void FRawCurveTracks::EvaluateTransformCurveData(USkeleton * Skeleton, TMap<FNam
 			continue;
 		}
 
-		const FSmartNameMapping* NameMapping = Skeleton->GetSmartNameContainer(USkeleton::AnimTrackCurveMappingName);
-
 		// Add or retrieve curve
-		FName CurveName;
+		FName CurveName = Curve.Name.DisplayName;
 		
-		// make sure it was added
-		if (ensure (NameMapping->GetName(Curve.CurveUid, CurveName)))
-		{
-			// note we're not checking Curve.GetCurveTypeFlags() yet
-			FTransform & Value = OutCurves.FindOrAdd(CurveName);
-			Value = Curve.Evaluate(CurrentTime, BlendWeight);
-		}
+		// note we're not checking Curve.GetCurveTypeFlags() yet
+		FTransform & Value = OutCurves.FindOrAdd(CurveName);
+		Value = Curve.Evaluate(CurrentTime, BlendWeight);
 	}
 }
 #endif
@@ -219,19 +237,19 @@ FAnimCurveBase * FRawCurveTracks::GetCurveData(USkeleton::AnimCurveUID Uid, ESup
 	}
 }
 
-bool FRawCurveTracks::DeleteCurveData(USkeleton::AnimCurveUID Uid, ESupportedCurveType SupportedCurveType /*= FloatType*/)
+bool FRawCurveTracks::DeleteCurveData(const FSmartName& CurveToDelete, ESupportedCurveType SupportedCurveType /*= FloatType*/)
 {
 	switch(SupportedCurveType)
 	{
 #if WITH_EDITOR
 	case VectorType:
-		return DeleteCurveDataImpl<FVectorCurve>(VectorCurves, Uid);
+		return DeleteCurveDataImpl<FVectorCurve>(VectorCurves, CurveToDelete);
 	case TransformType:
-		return DeleteCurveDataImpl<FTransformCurve>(TransformCurves, Uid);
+		return DeleteCurveDataImpl<FTransformCurve>(TransformCurves, CurveToDelete);
 #endif // WITH_EDITOR
 	case FloatType:
 	default:
-		return DeleteCurveDataImpl<FFloatCurve>(FloatCurves, Uid);
+		return DeleteCurveDataImpl<FFloatCurve>(FloatCurves, CurveToDelete);
 	}
 }
 
@@ -255,13 +273,13 @@ void FRawCurveTracks::DeleteAllCurveData(ESupportedCurveType SupportedCurveType
 }
 
 #if WITH_EDITOR
-void FRawCurveTracks::AddFloatCurveKey(const USkeleton::AnimCurveUID Uid, int32 CurveFlags, float Time, float Value)
+void FRawCurveTracks::AddFloatCurveKey(const FSmartName& NewCurve, int32 CurveFlags, float Time, float Value)
 {
-	FFloatCurve* FloatCurve = GetCurveDataImpl<FFloatCurve>(FloatCurves, Uid);
+	FFloatCurve* FloatCurve = GetCurveDataImpl<FFloatCurve>(FloatCurves, NewCurve.UID);
 	if (FloatCurve == nullptr)
 	{
-		AddCurveData(Uid, CurveFlags, FloatType);
-		FloatCurve = GetCurveDataImpl<FFloatCurve>(FloatCurves, Uid);
+		AddCurveData(NewCurve, CurveFlags, FloatType);
+		FloatCurve = GetCurveDataImpl<FFloatCurve>(FloatCurves, NewCurve.UID);
 	}
 
 	if (FloatCurve->GetCurveTypeFlags() != CurveFlags)
@@ -282,19 +300,19 @@ void FRawCurveTracks::RemoveRedundantKeys()
 }
 #endif
 
-bool FRawCurveTracks::AddCurveData(USkeleton::AnimCurveUID Uid, int32 CurveFlags /*= ACF_DefaultCurve*/, ESupportedCurveType SupportedCurveType /*= FloatType*/)
+bool FRawCurveTracks::AddCurveData(const FSmartName& NewCurve, int32 CurveFlags /*= ACF_DefaultCurve*/, ESupportedCurveType SupportedCurveType /*= FloatType*/)
 {
 	switch(SupportedCurveType)
 	{
 #if WITH_EDITOR
 	case VectorType:
-		return AddCurveDataImpl<FVectorCurve>(VectorCurves, Uid, CurveFlags);
+		return AddCurveDataImpl<FVectorCurve>(VectorCurves, NewCurve, CurveFlags);
 	case TransformType:
-		return AddCurveDataImpl<FTransformCurve>(TransformCurves, Uid, CurveFlags);
+		return AddCurveDataImpl<FTransformCurve>(TransformCurves, NewCurve, CurveFlags);
 #endif // WITH_EDITOR
 	case FloatType:
 	default:
-		return AddCurveDataImpl<FFloatCurve>(FloatCurves, Uid, CurveFlags);
+		return AddCurveDataImpl<FFloatCurve>(FloatCurves, NewCurve, CurveFlags);
 	}
 }
 
@@ -318,15 +336,12 @@ void FRawCurveTracks::Resize(float TotalLength, bool bInsert/* whether insert or
 #endif
 }
 
-void FRawCurveTracks::Serialize(FArchive& Ar)
+void FRawCurveTracks::PostSerialize(FArchive& Ar)
 {
 	// @TODO: If we're about to serialize vector curve, add here
-	if(Ar.UE4Ver() >= VER_UE4_SKELETON_ADD_SMARTNAMES)
+	for(FFloatCurve& Curve : FloatCurves)
 	{
-		for(FFloatCurve& Curve : FloatCurves)
-		{
-			Curve.Serialize(Ar);
-		}
+		Curve.PostSerialize(Ar);
 	}
 #if WITH_EDITORONLY_DATA
 	if( !Ar.IsCooking() )
@@ -335,19 +350,15 @@ void FRawCurveTracks::Serialize(FArchive& Ar)
 		{
 			for( FTransformCurve& Curve : TransformCurves )
 			{
-				Curve.Serialize( Ar );
+				Curve.PostSerialize( Ar );
 			}
 
 		}
 	}
 #endif // WITH_EDITORONLY_DATA
-	if (Ar.IsLoading())
-	{
-		SortFloatCurvesByUID();
-	}
 }
 
-void FRawCurveTracks::UpdateLastObservedNames(const FSmartNameMapping* NameMapping, ESupportedCurveType SupportedCurveType /*= FloatType*/)
+void FRawCurveTracks::RefreshName(const FSmartNameMapping* NameMapping, ESupportedCurveType SupportedCurveType /*= FloatType*/)
 {
 	switch(SupportedCurveType)
 	{
@@ -365,19 +376,19 @@ void FRawCurveTracks::UpdateLastObservedNames(const FSmartNameMapping* NameMappi
 	}
 }
 
-bool FRawCurveTracks::DuplicateCurveData(USkeleton::AnimCurveUID ToCopyUid, USkeleton::AnimCurveUID NewUid, ESupportedCurveType SupportedCurveType /*= FloatType*/)
+bool FRawCurveTracks::DuplicateCurveData(const FSmartName& CurveToCopy, const FSmartName& NewCurve, ESupportedCurveType SupportedCurveType /*= FloatType*/)
 {
 	switch(SupportedCurveType)
 	{
 #if WITH_EDITOR
 	case VectorType:
-		return DuplicateCurveDataImpl<FVectorCurve>(VectorCurves, ToCopyUid, NewUid);
+		return DuplicateCurveDataImpl<FVectorCurve>(VectorCurves, CurveToCopy, NewCurve);
 	case TransformType:
-		return DuplicateCurveDataImpl<FTransformCurve>(TransformCurves, ToCopyUid, NewUid);
+		return DuplicateCurveDataImpl<FTransformCurve>(TransformCurves, CurveToCopy, NewCurve);
 #endif // WITH_EDITOR
 	case FloatType:
 	default:
-		return DuplicateCurveDataImpl<FFloatCurve>(FloatCurves, ToCopyUid, NewUid);
+		return DuplicateCurveDataImpl<FFloatCurve>(FloatCurves, CurveToCopy, NewCurve);
 	}
 }
 
@@ -396,7 +407,7 @@ DataType * FRawCurveTracks::GetCurveDataImpl(TArray<DataType> & Curves, USkeleto
 {
 	for(DataType& Curve : Curves)
 	{
-		if(Curve.CurveUid == Uid)
+		if(Curve.Name.UID == Uid)
 		{
 			return &Curve;
 		}
@@ -406,11 +417,11 @@ DataType * FRawCurveTracks::GetCurveDataImpl(TArray<DataType> & Curves, USkeleto
 }
 
 template <typename DataType>
-bool FRawCurveTracks::DeleteCurveDataImpl(TArray<DataType> & Curves, USkeleton::AnimCurveUID Uid)
+bool FRawCurveTracks::DeleteCurveDataImpl(TArray<DataType> & Curves, const FSmartName& CurveToDelete)
 {
 	for(int32 Idx = 0; Idx < Curves.Num(); ++Idx)
 	{
-		if(Curves[Idx].CurveUid == Uid)
+		if(Curves[Idx].Name.UID == CurveToDelete.UID)
 		{
 			Curves.RemoveAt(Idx);
 			return true;
@@ -421,11 +432,11 @@ bool FRawCurveTracks::DeleteCurveDataImpl(TArray<DataType> & Curves, USkeleton::
 }
 
 template <typename DataType>
-bool FRawCurveTracks::AddCurveDataImpl(TArray<DataType> & Curves, USkeleton::AnimCurveUID Uid, int32 CurveFlags)
+bool FRawCurveTracks::AddCurveDataImpl(TArray<DataType> & Curves, const FSmartName& NewCurve, int32 CurveFlags)
 {
-	if(GetCurveDataImpl<DataType>(Curves, Uid) == NULL)
+	if(GetCurveDataImpl<DataType>(Curves, NewCurve.UID) == NULL)
 	{
-		Curves.Add(DataType(Uid, CurveFlags));
+		Curves.Add(DataType(NewCurve, CurveFlags));
 		return true;
 	}
 	return false;
@@ -438,19 +449,19 @@ void FRawCurveTracks::UpdateLastObservedNamesImpl(TArray<DataType> & Curves, con
 	{
 		for(DataType& Curve : Curves)
 		{
-			NameMapping->GetName(Curve.CurveUid, Curve.LastObservedName);
+			NameMapping->GetName(Curve.Name.UID, Curve.Name.DisplayName);
 		}
 	}
 }
 
 template <typename DataType>
-bool FRawCurveTracks::DuplicateCurveDataImpl(TArray<DataType> & Curves, USkeleton::AnimCurveUID ToCopyUid, USkeleton::AnimCurveUID NewUid)
+bool FRawCurveTracks::DuplicateCurveDataImpl(TArray<DataType> & Curves, const FSmartName& CurveToCopy, const FSmartName& NewCurve)
 {
-	DataType* ExistingCurve = GetCurveDataImpl<DataType>(Curves, ToCopyUid);
-	if(ExistingCurve && GetCurveDataImpl<DataType>(Curves, NewUid) == NULL)
+	DataType* ExistingCurve = GetCurveDataImpl<DataType>(Curves, CurveToCopy.UID);
+	if(ExistingCurve && GetCurveDataImpl<DataType>(Curves, NewCurve.UID) == NULL)
 	{
 		// Add the curve to the track and set its data to the existing curve
-		Curves.Add(DataType(NewUid, ExistingCurve->GetCurveTypeFlags()));
+		Curves.Add(DataType(NewCurve, ExistingCurve->GetCurveTypeFlags()));
 		Curves.Last().CopyCurve(*ExistingCurve);
 
 		return true;
@@ -460,9 +471,9 @@ bool FRawCurveTracks::DuplicateCurveDataImpl(TArray<DataType> & Curves, USkeleto
 
 FArchive& operator<<(FArchive& Ar, FRawCurveTracks& D)
 {
-	FRawCurveTracks::StaticStruct()->SerializeTaggedProperties(Ar, (uint8*)&D, FRawCurveTracks::StaticStruct(), NULL);
-	D.Serialize(Ar);
-
+	UScriptStruct* StaticStruct = FRawCurveTracks::StaticStruct();
+	StaticStruct->SerializeTaggedProperties(Ar, (uint8*)&D, StaticStruct, nullptr);
+	// do not call custom serialize that relies on version number. The Archive version doesn't exists on this. 
 	return Ar;
 }
 
@@ -474,7 +485,7 @@ void FAnimCurveParam::Initialize(USkeleton* Skeleton)
 	// Initialize for curve UID
 	if (Name != NAME_None)
 	{
-		Skeleton->AddSmartNameAndModify(USkeleton::AnimCurveMappingName, Name, UID);
+		UID = Skeleton->GetUIDByName(USkeleton::AnimCurveMappingName, Name);
 	}
 	else
 	{
diff --git a/Engine/Source/Runtime/Engine/Private/Animation/AnimInstance.cpp b/Engine/Source/Runtime/Engine/Private/Animation/AnimInstance.cpp
--- a/Engine/Source/Runtime/Engine/Private/Animation/AnimInstance.cpp
+++ b/Engine/Source/Runtime/Engine/Private/Animation/AnimInstance.cpp
@@ -257,7 +257,7 @@ void UAnimInstance::UninitializeAnimation()
 	}
 
 	ActiveAnimNotifyState.Reset();
-	EventCurves.Reset();
+	ResetAnimationCurves();
 	MaterialParamatersToClear.Reset();
 	NotifyQueue.Reset(SkelMeshComp);
 }
@@ -911,38 +911,48 @@ void UAnimInstance::DisplayDebug(class UCanvas* Canvas, const FDebugDisplayInfo&
 		{
 			FIndenter CurveIndent(Indent);
 
-			Heading = FString::Printf(TEXT("Morph Curves: %i"), Proxy.GetMorphTargetCurves().Num());
+			Heading = FString::Printf(TEXT("Morph Curves: %i"), AnimationCurves[(uint8)EAnimCurveType::MorphTargetCurve].Num());
 			DisplayDebugManager.DrawString(Heading, Indent);
 
 			DisplayDebugManager.SetLinearDrawColor(TextWhite);
 
 			{
 				FIndenter MorphCurveIndent(Indent);
-				OutputCurveMap(Proxy.GetMorphTargetCurves(), Canvas, DisplayDebugManager, Indent);
+				OutputCurveMap(AnimationCurves[(uint8)EAnimCurveType::MorphTargetCurve], Canvas, DisplayDebugManager, Indent);
 			}
 
 			DisplayDebugManager.SetLinearDrawColor(TextYellow);
 
-			Heading = FString::Printf(TEXT("Material Curves: %i"), Proxy.GetMaterialParameterCurves().Num());
+			Heading = FString::Printf(TEXT("Material Curves: %i"), AnimationCurves[(uint8)EAnimCurveType::MaterialCurve].Num());
 			DisplayDebugManager.DrawString(Heading, Indent);
 
 			DisplayDebugManager.SetLinearDrawColor(TextWhite);
 
 			{
 				FIndenter MaterialCurveIndent(Indent);
-				OutputCurveMap(Proxy.GetMaterialParameterCurves(), Canvas, DisplayDebugManager, Indent);
+				OutputCurveMap(AnimationCurves[(uint8)EAnimCurveType::MaterialCurve], Canvas, DisplayDebugManager, Indent);
 			}
 
 			DisplayDebugManager.SetLinearDrawColor(TextYellow);
 
-			Heading = FString::Printf(TEXT("Event Curves: %i"), EventCurves.Num());
+			Heading = FString::Printf(TEXT("Event Curves: %i"), AnimationCurves[(uint8)EAnimCurveType::EventCurve].Num());
 			DisplayDebugManager.DrawString(Heading, Indent);
 
 			DisplayDebugManager.SetLinearDrawColor(TextWhite);
 
 			{
 				FIndenter EventCurveIndent(Indent);
-				OutputCurveMap(EventCurves, Canvas, DisplayDebugManager, Indent);
+				OutputCurveMap(AnimationCurves[(uint8)EAnimCurveType::EventCurve], Canvas, DisplayDebugManager, Indent);
+			}
+
+			Heading = FString::Printf(TEXT("Pose Curves: %i"), AnimationCurves[(uint8)EAnimCurveType::PoseCurve].Num());
+			DisplayDebugManager.DrawString(Heading, Indent);
+
+			DisplayDebugManager.SetLinearDrawColor(TextWhite);
+
+			{
+				FIndenter EventCurveIndent(Indent);
+				OutputCurveMap(AnimationCurves[(uint8)EAnimCurveType::PoseCurve], Canvas, DisplayDebugManager, Indent);
 			}
 		}
 	}
@@ -1034,45 +1044,58 @@ void UAnimInstance::AddCurveValue(const FName& CurveName, float Value, int32 Cur
 	//CurveValues.Add(CurveName, Value);
 	if (CurveTypeFlags & ACF_TriggerEvent)
 	{
-		float *CurveValPtr = EventCurves.Find(CurveName);
+		float *CurveValPtr = AnimationCurves[(uint8)EAnimCurveType::EventCurve].Find(CurveName);
 		if ( CurveValPtr )
 		{
 			// sum up, in the future we might normalize, but for now this just sums up
 			// this won't work well if all of them have full weight - i.e. additive 
-			*CurveValPtr += Value;
+			*CurveValPtr = Value;
 		}
 		else
 		{
-			EventCurves.Add(CurveName, Value);
+			AnimationCurves[(uint8)EAnimCurveType::EventCurve].Add(CurveName, Value);
 		}
 	}
 
 	if (CurveTypeFlags & ACF_DrivesMorphTarget)
 	{
-		float *CurveValPtr = Proxy.GetMorphTargetCurves().Find(CurveName);
+		float *CurveValPtr = AnimationCurves[(uint8)EAnimCurveType::MorphTargetCurve].Find(CurveName);
 		if ( CurveValPtr )
 		{
 			// sum up, in the future we might normalize, but for now this just sums up
 			// this won't work well if all of them have full weight - i.e. additive 
-			*CurveValPtr += Value;
+			*CurveValPtr = Value;
 		}
 		else
 		{
-			Proxy.GetMorphTargetCurves().Add(CurveName, Value);
+			AnimationCurves[(uint8)EAnimCurveType::MorphTargetCurve].Add(CurveName, Value);
 		}
 	}
 
 	if (CurveTypeFlags & ACF_DrivesMaterial)
 	{
 		MaterialParamatersToClear.RemoveSwap(CurveName);
-		float* CurveValPtr = Proxy.GetMaterialParameterCurves().Find(CurveName);
+		float* CurveValPtr = AnimationCurves[(uint8)EAnimCurveType::MaterialCurve].Find(CurveName);
 		if( CurveValPtr)
 		{
-			*CurveValPtr += Value;
+			*CurveValPtr = Value;
 		}
 		else
 		{
-			Proxy.GetMaterialParameterCurves().Add(CurveName, Value);
+			AnimationCurves[(uint8)EAnimCurveType::MaterialCurve].Add(CurveName, Value);
+		}
+	}
+
+	if (CurveTypeFlags & ACF_DrivesPose)
+	{
+		float* CurveValPtr = AnimationCurves[(uint8)EAnimCurveType::PoseCurve].Find(CurveName);
+		if (CurveValPtr)
+		{
+			*CurveValPtr = Value;
+		}
+		else
+		{
+			AnimationCurves[(uint8)EAnimCurveType::PoseCurve].Add(CurveName, Value);
 		}
 	}
 }
@@ -1104,10 +1127,51 @@ void UAnimInstance::TickSyncGroupWriteIndex()
 	GetProxyOnGameThread<FAnimInstanceProxy>().TickSyncGroupWriteIndex();
 }
 
-void UAnimInstance::RefreshCurves(USkeletalMeshComponent* Component)
+void UAnimInstance::UpdateCurvesToComponents(USkeletalMeshComponent* Component /*= nullptr*/)
 {
 	// update curves to component
-	GetProxyOnGameThread<FAnimInstanceProxy>().UpdateCurvesToComponents(Component);
+	if (Component)
+	{
+		Component->ApplyAnimationCurvesToComponent(&AnimationCurves[(uint8)EAnimCurveType::MaterialCurve], &AnimationCurves[(uint8)EAnimCurveType::MorphTargetCurve]);
+	}
+}
+
+void UAnimInstance::GetAnimationCurveList(int32 CurveFlags, TMap<FName, float>& OutCurveList) const
+{
+	OutCurveList.Reset();
+
+	if (CurveFlags & ACF_TriggerEvent)
+	{
+		OutCurveList.Append(AnimationCurves[(uint8)EAnimCurveType::EventCurve]);
+	}
+
+	if (CurveFlags & ACF_DrivesMorphTarget)
+	{
+		OutCurveList.Append(AnimationCurves[(uint8)EAnimCurveType::MorphTargetCurve]);
+	}
+
+	if (CurveFlags & ACF_DrivesMaterial)
+	{
+		OutCurveList.Append(AnimationCurves[(uint8)EAnimCurveType::MaterialCurve]);
+	}
+
+	if (CurveFlags & ACF_DrivesPose)
+	{
+		OutCurveList.Append(AnimationCurves[(uint8)EAnimCurveType::PoseCurve]);
+	}
+}
+
+void UAnimInstance::RefreshCurves(USkeletalMeshComponent* Component)
+{
+	UpdateCurvesToComponents(Component);
+}
+
+void UAnimInstance::ResetAnimationCurves()
+{
+	for (uint8 Index = 0; Index < (uint8)EAnimCurveType::MaxAnimCurveType; ++Index)
+	{
+		AnimationCurves[Index].Reset();
+	}
 }
 
 void UAnimInstance::UpdateCurves(const FBlendedHeapCurve& InCurve)
@@ -1118,7 +1182,7 @@ void UAnimInstance::UpdateCurves(const FBlendedHeapCurve& InCurve)
 
 	//Track material params we set last time round so we can clear them if they aren't set again.
 	MaterialParamatersToClear.Reset();
-	for(auto Iter = Proxy.GetMaterialParameterCurves().CreateConstIterator(); Iter; ++Iter)
+	for(auto Iter = AnimationCurves[(uint8)EAnimCurveType::MaterialCurve].CreateConstIterator(); Iter; ++Iter)
 	{
 		if(Iter.Value() > 0.0f)
 		{
@@ -1126,9 +1190,7 @@ void UAnimInstance::UpdateCurves(const FBlendedHeapCurve& InCurve)
 		}
 	}
 
-	EventCurves.Reset();
-	Proxy.GetMorphTargetCurves().Reset();
-	Proxy.GetMaterialParameterCurves().Reset();
+	ResetAnimationCurves();
 
 	if (InCurve.UIDList != nullptr)
 	{
@@ -1148,19 +1210,24 @@ void UAnimInstance::UpdateCurves(const FBlendedHeapCurve& InCurve)
 		AddCurveValue(ParamsToClearCopy[i], 0.0f, ACF_DrivesMaterial);
 	}
 
+#if WITH_EDITOR
+	// if we're supporting this in-game, this code has to change to work with UID
+	for (auto& AddAnimCurveDelegate : OnAddAnimationCurves)
+	{
+		if (AddAnimCurveDelegate.IsBound())
+		{
+			AddAnimCurveDelegate.Execute(this);
+		}
+	}
+
+#endif
 	// update curves to component
-	// Proxy might not have SkeletalMeshComponent
-	Proxy.UpdateCurvesToComponents(GetOwningComponent());
+	UpdateCurvesToComponents(GetOwningComponent());
 }
 
 bool UAnimInstance::HasMorphTargetCurves() const
 {
-	return GetProxyOnGameThread<FAnimInstanceProxy>().HasMorphTargetCurves();
-}
-
-TMap<FName, float>& UAnimInstance::GetMorphTargetCurves()
-{
-	return GetProxyOnGameThread<FAnimInstanceProxy>().GetMorphTargetCurves();
+	return AnimationCurves[(uint8)EAnimCurveType::MorphTargetCurve].Num() > 0;
 }
 
 void UAnimInstance::TriggerAnimNotifies(float DeltaSeconds)
@@ -1313,7 +1380,7 @@ float UAnimInstance::GetSlotMontageGlobalWeight(FName SlotNodeName) const
 
 float UAnimInstance::GetCurveValue(FName CurveName)
 {
-	float* Value = EventCurves.Find(CurveName);
+	float* Value = AnimationCurves[(uint8)EAnimCurveType::EventCurve].Find(CurveName);
 	if (Value)
 	{
 		return *Value;
@@ -2698,4 +2765,28 @@ void UAnimInstance::QueueRootMotionBlend(const FTransform& RootTransform, const
 	RootMotionBlendQueue.Add(FQueuedRootMotionBlend(RootTransform, SlotName, Weight));
 }
 
+#if WITH_EDITOR
+void UAnimInstance::AddDelegate_AddCustomAnimationCurve(FOnAddCustomAnimationCurves& InOnAddCustomAnimationCurves)
+{
+	if (InOnAddCustomAnimationCurves.IsBound())
+	{
+		OnAddAnimationCurves.Add(InOnAddCustomAnimationCurves);
+	}
+}
+
+void UAnimInstance::RemoveDelegate_AddCustomAnimationCurve(FOnAddCustomAnimationCurves& InOnAddCustomAnimationCurves)
+{
+	for (int32 DelegateId = 0; DelegateId < OnAddAnimationCurves.Num(); ++DelegateId)
+	{
+		if (InOnAddCustomAnimationCurves.GetHandle() == OnAddAnimationCurves[DelegateId].GetHandle())
+		{
+			InOnAddCustomAnimationCurves.Unbind();
+			OnAddAnimationCurves.RemoveAt(DelegateId);
+			break;
+		}
+	}
+}
+
+#endif // WITH_EDITOR
+
 #undef LOCTEXT_NAMESPACE 
diff --git a/Engine/Source/Runtime/Engine/Private/Animation/AnimInstanceProxy.cpp b/Engine/Source/Runtime/Engine/Private/Animation/AnimInstanceProxy.cpp
--- a/Engine/Source/Runtime/Engine/Private/Animation/AnimInstanceProxy.cpp
+++ b/Engine/Source/Runtime/Engine/Private/Animation/AnimInstanceProxy.cpp
@@ -1,6 +1,7 @@
 // Copyright 1998-2016 Epic Games, Inc. All Rights Reserved.
 
 #include "EnginePrivate.h"
+#include "Animation/PoseAsset.h"
 #include "AnimInstanceProxy.h"
 #include "Animation/AnimStats.h"
 #include "Animation/AnimMontage.h"
@@ -10,6 +11,7 @@
 #include "Animation/AnimNode_AssetPlayerBase.h"
 #include "Animation/AnimNode_StateMachine.h"
 #include "Animation/AnimNode_TransitionResult.h"
+#include "Animation/AnimNode_SaveCachedPose.h"
 
 #define DO_ANIMSTAT_PROCESSING(StatName) DEFINE_STAT(STAT_ ## StatName)
 #include "AnimMTStats.h"
@@ -31,6 +33,12 @@ void FAnimInstanceProxy::UpdateAnimationNode(float DeltaSeconds)
 	{
 		UpdateCounter.Increment();
 		RootNode->Update(FAnimationUpdateContext(this, DeltaSeconds));
+
+		// We've updated the graph, now update the fractured saved pose sections
+		for(FAnimNode_SaveCachedPose* PoseNode : SavedPoseQueue)
+		{
+			PoseNode->PostGraphUpdate();
+		}
 	}
 }
 
@@ -55,6 +63,17 @@ void FAnimInstanceProxy::Initialize(UAnimInstance* InAnimInstance)
 			RootNode = nullptr;
 		}
 
+		// Initialise the pose node list
+		const TArray<int32>& PoseNodeIndices = AnimClassInterface->GetOrderedSavedPoseNodeIndices();
+		const TArray<UStructProperty*>& AnimNodeProperties = AnimClassInterface->GetAnimNodeProperties();
+		SavedPoseQueue.Empty(PoseNodeIndices.Num());
+		for(int32 Idx : PoseNodeIndices)
+		{
+			int32 ActualPropertyIdx = AnimNodeProperties.Num() - 1 - Idx;
+			FAnimNode_SaveCachedPose* ActualPoseNode = AnimNodeProperties[ActualPropertyIdx]->ContainerPtrToValuePtr<FAnimNode_SaveCachedPose>(InAnimInstance);
+			SavedPoseQueue.Add(ActualPoseNode);
+		}
+
 		// if no mesh, use Blueprint Skeleton
 		if (Skeleton == nullptr)
 		{
@@ -154,8 +173,6 @@ void FAnimInstanceProxy::InitializeRootNode()
 
 void FAnimInstanceProxy::Uninitialize(UAnimInstance* InAnimInstance)
 {
-	MorphTargetCurves.Reset();
-	MaterialParameterCurves.Reset();
 }
 
 void FAnimInstanceProxy::PreUpdate(UAnimInstance* InAnimInstance, float DeltaSeconds)
@@ -276,6 +293,13 @@ void FAnimInstanceProxy::MakeBlendSpaceTickRecord(FAnimTickRecord& TickRecord, c
 	TickRecord.bLooping = bLooping;
 }
 
+/** Helper function: make a tick record for a pose asset*/
+void FAnimInstanceProxy::MakePoseAssetTickRecord(FAnimTickRecord& TickRecord, class UPoseAsset* PoseAsset, float FinalBlendWeight) const
+{
+	TickRecord.SourceAsset = PoseAsset;
+	TickRecord.EffectiveBlendWeight = FinalBlendWeight;
+}
+
 void FAnimInstanceProxy::SequenceAdvanceImmediate(UAnimSequenceBase* Sequence, bool bLooping, float PlayRate, float DeltaSeconds, float& CurrentTime, FMarkerTickRecord& MarkerTickRecord)
 {
 	FAnimTickRecord TickRecord;
@@ -741,16 +765,6 @@ void FAnimInstanceProxy::EvaluateAnimation(FPoseContext& Output)
 	}
 }
 
-void FAnimInstanceProxy::UpdateCurvesToComponents(USkeletalMeshComponent* Component)
-{
-	// AnimInstanceProxy has SkeletalMeshComponent but it's not safe to be used everywhere
-	// so please give us Component to update
-	if (Component)
-	{
-		Component->ApplyAnimationCurvesToComponent(&MaterialParameterCurves, &MorphTargetCurves);
-	}
-}
-
 // for now disable becauase it will not work with single node instance
 #if (UE_BUILD_SHIPPING || UE_BUILD_TEST)
 #define DEBUG_MONTAGEINSTANCE_WEIGHT 0
@@ -987,10 +1001,17 @@ const FMontageEvaluationState* FAnimInstanceProxy::GetActiveMontageEvaluationSta
 
 void FAnimInstanceProxy::GatherDebugData(FNodeDebugData& DebugData)
 {
+	// Gather debug data for Root Node
 	if(RootNode != nullptr)
 	{
 		 RootNode->GatherDebugData(DebugData); 
 	}
+
+	// Gather debug data for Cached Poses.
+	for (FAnimNode_SaveCachedPose* PoseNode : SavedPoseQueue)
+	{
+		PoseNode->GatherDebugData(DebugData);
+	}
 }
 
 float FAnimInstanceProxy::GetInstanceAssetPlayerLength(int32 AssetPlayerIndex)
diff --git a/Engine/Source/Runtime/Engine/Private/Animation/AnimMontage.cpp b/Engine/Source/Runtime/Engine/Private/Animation/AnimMontage.cpp
--- a/Engine/Source/Runtime/Engine/Private/Animation/AnimMontage.cpp
+++ b/Engine/Source/Runtime/Engine/Private/Animation/AnimMontage.cpp
@@ -1174,6 +1174,44 @@ bool UAnimMontage::ContainRecursive(TArray<UAnimCompositeBase*>& CurrentAccumula
 // MontageInstance
 /////////////////////////////////////////////////////////////////////////////////////////////
 
+FAnimMontageInstance::FAnimMontageInstance()
+	: Montage(NULL)
+	, bPlaying(false)
+	, DefaultBlendTimeMultiplier(1.0f)
+	, bDidUseMarkerSyncThisTick(false)
+	, AnimInstance(NULL)
+	, InstanceID(INDEX_NONE)
+	, Position(0.f)
+	, PlayRate(1.f)
+	, bInterrupted(false)
+	, PreviousWeight(0.f)
+	, DeltaMoved(0.f)
+	, PreviousPosition(0.f)
+	, SyncGroupIndex(INDEX_NONE)
+	, MontageSyncLeader(NULL)
+	, MontageSyncUpdateFrameCounter(INDEX_NONE)
+{
+}
+
+FAnimMontageInstance::FAnimMontageInstance(UAnimInstance * InAnimInstance)
+	: Montage(NULL)
+	, bPlaying(false)
+	, DefaultBlendTimeMultiplier(1.0f)
+	, bDidUseMarkerSyncThisTick(false)
+	, AnimInstance(InAnimInstance)
+	, InstanceID(INDEX_NONE)
+	, Position(0.f)
+	, PlayRate(1.f)
+	, bInterrupted(false)
+	, PreviousWeight(0.f)
+	, DeltaMoved(0.f)
+	, PreviousPosition(0.f)
+	, SyncGroupIndex(INDEX_NONE)
+	, MontageSyncLeader(NULL)
+	, MontageSyncUpdateFrameCounter(INDEX_NONE)
+{
+}
+
 void FAnimMontageInstance::Play(float InPlayRate)
 {
 	bPlaying = true;
diff --git a/Engine/Source/Runtime/Engine/Private/Animation/AnimNodeBase.cpp b/Engine/Source/Runtime/Engine/Private/Animation/AnimNodeBase.cpp
--- a/Engine/Source/Runtime/Engine/Private/Animation/AnimNodeBase.cpp
+++ b/Engine/Source/Runtime/Engine/Private/Animation/AnimNodeBase.cpp
@@ -358,28 +358,11 @@ FNodeDebugData& FNodeDebugData::BranchFlow(float BranchWeight, FString InNodeDes
 	return NodeChain.Last().ChildNodeChain.Last();
 }
 
-FNodeDebugData* FNodeDebugData::GetCachePoseDebugData(FName InCachePoseName)
+FNodeDebugData* FNodeDebugData::GetCachePoseDebugData(float GlobalWeight)
 {
 	check(RootNodePtr);
 
-	for (int32 NodeIndex = 0; NodeIndex < RootNodePtr->SaveCachePoseNodes.Num(); NodeIndex++)
-	{
-		FNodeDebugData& CachePoseDebugData = RootNodePtr->SaveCachePoseNodes[NodeIndex];
- 		if (CachePoseDebugData.CachePoseName == InCachePoseName)
-		{
-			if (CachePoseDebugData.AbsoluteWeight >= AbsoluteWeight)
-			{
-				return nullptr;
-			}	
-
-			CachePoseDebugData.AbsoluteWeight = AbsoluteWeight;
-			CachePoseDebugData.NodeChain.Empty(ANIM_NODE_DEBUG_MAX_CHAIN);
-			return &CachePoseDebugData;
-		}
-	}
-
-	RootNodePtr->SaveCachePoseNodes.Add( FNodeDebugData(AnimInstance, AbsoluteWeight, FString(), RootNodePtr) );
-	RootNodePtr->SaveCachePoseNodes.Last().CachePoseName = InCachePoseName;
+	RootNodePtr->SaveCachePoseNodes.Add( FNodeDebugData(AnimInstance, GlobalWeight, FString(), RootNodePtr) );
 	RootNodePtr->SaveCachePoseNodes.Last().NodeChain.Reserve(ANIM_NODE_DEBUG_MAX_CHAIN);
 	return &RootNodePtr->SaveCachePoseNodes.Last();
 }
diff --git a/Engine/Source/Runtime/Engine/Private/Animation/AnimNode_SaveCachedPose.cpp b/Engine/Source/Runtime/Engine/Private/Animation/AnimNode_SaveCachedPose.cpp
new file mode 100644
--- /dev/null
+++ b/Engine/Source/Runtime/Engine/Private/Animation/AnimNode_SaveCachedPose.cpp
@@ -0,0 +1,100 @@
+// Copyright 1998-2016 Epic Games, Inc. All Rights Reserved.
+
+#include "EnginePrivate.h"
+#include "Animation/AnimNode_SaveCachedPose.h"
+#include "Animation/AnimInstanceProxy.h"
+
+/////////////////////////////////////////////////////
+// FAnimNode_SaveCachedPose
+
+FAnimNode_SaveCachedPose::FAnimNode_SaveCachedPose()
+{
+}
+
+void FAnimNode_SaveCachedPose::Initialize(const FAnimationInitializeContext& Context)
+{
+	// StateMachines cause reinitialization on state changes.
+	// we only want to let them through if we're not relevant as to not create a pop.
+	if (!InitializationCounter.IsSynchronizedWith(Context.AnimInstanceProxy->GetInitializationCounter())
+		|| ((UpdateCounter.Get() != INDEX_NONE) && !UpdateCounter.WasSynchronizedInTheLastFrame(Context.AnimInstanceProxy->GetUpdateCounter())))
+	{
+		InitializationCounter.SynchronizeWith(Context.AnimInstanceProxy->GetInitializationCounter());
+
+		FAnimNode_Base::Initialize(Context);
+
+		// Initialize the subgraph
+		Pose.Initialize(Context);
+	}
+}
+
+void FAnimNode_SaveCachedPose::CacheBones(const FAnimationCacheBonesContext& Context)
+{
+	if (!CachedBonesCounter.IsSynchronizedWith(Context.AnimInstanceProxy->GetCachedBonesCounter()))
+	{
+		CachedBonesCounter.SynchronizeWith(Context.AnimInstanceProxy->GetCachedBonesCounter());
+
+		// Cache bones in the subgraph
+		Pose.CacheBones(Context);
+	}
+}
+
+void FAnimNode_SaveCachedPose::Update(const FAnimationUpdateContext& Context)
+{
+	// Store this context for the post update
+	CachedUpdateContexts.Add(Context);
+}
+
+void FAnimNode_SaveCachedPose::Evaluate(FPoseContext& Output)
+{
+	if (!EvaluationCounter.IsSynchronizedWith(Output.AnimInstanceProxy->GetEvaluationCounter()))
+	{
+		EvaluationCounter.SynchronizeWith(Output.AnimInstanceProxy->GetEvaluationCounter());
+
+		FPoseContext CachingContext(Output);
+		Pose.Evaluate(CachingContext);
+		CachedPose.CopyBonesFrom(CachingContext.Pose);
+		CachedCurve.CopyFrom(CachingContext.Curve);
+	}
+
+	// Return the cached result
+	Output.Pose.CopyBonesFrom(CachedPose);
+	Output.Curve.CopyFrom(CachedCurve);
+}
+
+void FAnimNode_SaveCachedPose::GatherDebugData(FNodeDebugData& DebugData)
+{
+	FString DebugLine = DebugData.GetNodeName(this);
+	DebugLine += FString::Printf(TEXT("%s:"), *CachePoseName.ToString());
+
+	if (FNodeDebugData* SaveCachePoseDebugDataPtr = DebugData.GetCachePoseDebugData(GlobalWeight))
+	{
+		SaveCachePoseDebugDataPtr->AddDebugItem(DebugLine);
+		Pose.GatherDebugData(*SaveCachePoseDebugDataPtr);
+	}
+}
+
+void FAnimNode_SaveCachedPose::PostGraphUpdate()
+{
+	GlobalWeight = 0.f;
+
+	// Update GlobalWeight based on highest weight calling us.
+	const int32 NumContexts = CachedUpdateContexts.Num();
+	if(NumContexts > 0)
+	{
+		GlobalWeight = CachedUpdateContexts[0].GetFinalBlendWeight();
+		int32 MaxWeightIdx = 0;
+		for(int32 CurrIdx = 1; CurrIdx < NumContexts; ++CurrIdx)
+		{
+			const float BlendWeight = CachedUpdateContexts[CurrIdx].GetFinalBlendWeight();
+			if(BlendWeight > GlobalWeight)
+			{
+				GlobalWeight = BlendWeight;
+				MaxWeightIdx = CurrIdx;
+			}
+		}
+
+		Pose.Update(CachedUpdateContexts[MaxWeightIdx]);
+	}
+
+	CachedUpdateContexts.Reset();
+}
diff --git a/Engine/Source/Runtime/Engine/Private/Animation/AnimNode_StateMachine.cpp b/Engine/Source/Runtime/Engine/Private/Animation/AnimNode_StateMachine.cpp
--- a/Engine/Source/Runtime/Engine/Private/Animation/AnimNode_StateMachine.cpp
+++ b/Engine/Source/Runtime/Engine/Private/Animation/AnimNode_StateMachine.cpp
@@ -842,6 +842,9 @@ void FAnimNode_StateMachine::EvaluateTransitionCustomBlend(FPoseContext& Output,
 		{
 			Output.Pose[BoneIndex] = StatePoseResult.Pose[BoneIndex];
 		}
+
+		// Copy curve over also, replacing current.
+		Output.Curve.CopyFrom(StatePoseResult.Curve);
 	}
 }
 
diff --git a/Engine/Source/Runtime/Engine/Private/Animation/AnimNode_UseCachedPose.cpp b/Engine/Source/Runtime/Engine/Private/Animation/AnimNode_UseCachedPose.cpp
--- a/Engine/Source/Runtime/Engine/Private/Animation/AnimNode_UseCachedPose.cpp
+++ b/Engine/Source/Runtime/Engine/Private/Animation/AnimNode_UseCachedPose.cpp
@@ -39,5 +39,6 @@ void FAnimNode_UseCachedPose::GatherDebugData(FNodeDebugData& DebugData)
 
 	DebugData.AddDebugItem(DebugLine, true);
 
-	LinkToCachingNode.GatherDebugData(DebugData);
+	// we explicitly do not forward this call to the SaveCachePose node here.
+	// It is handled in FAnimInstanceProxy::GatherDebugData
 }
\ No newline at end of file
diff --git a/Engine/Source/Runtime/Engine/Private/Animation/AnimSequence.cpp b/Engine/Source/Runtime/Engine/Private/Animation/AnimSequence.cpp
--- a/Engine/Source/Runtime/Engine/Private/Animation/AnimSequence.cpp
+++ b/Engine/Source/Runtime/Engine/Private/Animation/AnimSequence.cpp
@@ -489,23 +489,9 @@ void UAnimSequence::PostLoad()
 	if (CurrentSkeleton)
 	{
 		// Get the name mapping object for curves
-		const FSmartNameMapping* NameMapping = CurrentSkeleton->GetSmartNameContainer(USkeleton::AnimTrackCurveMappingName);
-
-		TArray<FTransformCurve*> UnlinkedCurves;
 		for(FTransformCurve& Curve : RawCurveData.TransformCurves)
 		{
-			if(!NameMapping->Exists(Curve.LastObservedName))
-			{
-				// The skeleton doesn't know our name. Use the last observed name that was saved with the
-				// curve to create a new name. This can happen if a user saves an animation but not a skeleton
-				// either when updating the assets or editing the curves within.
-				UnlinkedCurves.Add(&Curve);
-			}
-		}
-
-		for(FTransformCurve* Curve : UnlinkedCurves)
-		{
-			CurrentSkeleton->AddSmartNameAndModify(USkeleton::AnimTrackCurveMappingName, Curve->LastObservedName, Curve->CurveUid);
+			CurrentSkeleton->VerifySmartName(USkeleton::AnimTrackCurveMappingName, Curve.Name);
 		}
 	}
 
@@ -513,7 +499,7 @@ void UAnimSequence::PostLoad()
 	// actual anim flags 
 	for (FFloatCurve& Curve : RawCurveData.FloatCurves)
 	{
-		if (FAnimCurveBase* CompressedCurve = CompressedCurveData.GetCurveData(Curve.CurveUid))
+		if (FAnimCurveBase* CompressedCurve = CompressedCurveData.GetCurveData(Curve.Name.UID))
 		{
 			CompressedCurve->SetCurveTypeFlags(Curve.GetCurveTypeFlags());
 		}
@@ -729,6 +715,33 @@ void UAnimSequence::ExtractBoneTransform(const TArray<struct FRawAnimSequenceTra
 	ExtractBoneTransform(InRawAnimationData[TrackIndex], OutAtom, Time);
 }
 
+void UAnimSequence::ExtractBoneTransform(const struct FRawAnimSequenceTrack& RawTrack, FTransform& OutAtom, int32 KeyIndex) const
+{
+	// Bail out (with rather wacky data) if data is empty for some reason.
+	if (RawTrack.PosKeys.Num() == 0 || RawTrack.RotKeys.Num() == 0)
+	{
+		UE_LOG(LogAnimation, Log, TEXT("UAnimSequence::GetBoneTransform : No anim data in AnimSequence!"));
+		OutAtom.SetIdentity();
+		return;
+	}
+
+	const int32 PosKeyIndex = FMath::Min(KeyIndex, RawTrack.PosKeys.Num() - 1);
+	const int32 RotKeyIndex = FMath::Min(KeyIndex, RawTrack.RotKeys.Num() - 1);
+	static const FVector DefaultScale3D = FVector(1.f);
+
+	OutAtom.SetTranslation(RawTrack.PosKeys[PosKeyIndex]);
+	OutAtom.SetRotation(RawTrack.RotKeys[RotKeyIndex]);
+	if (RawTrack.ScaleKeys.Num() > 0)
+	{
+		const int32 ScaleKeyIndex = FMath::Min(KeyIndex, RawTrack.ScaleKeys.Num() - 1);
+		OutAtom.SetScale3D(RawTrack.ScaleKeys[ScaleKeyIndex]);
+	}
+	else
+	{
+		OutAtom.SetScale3D(DefaultScale3D);
+	}
+}
+
 void UAnimSequence::ExtractBoneTransform(const struct FRawAnimSequenceTrack& RawTrack, FTransform& OutAtom, float Time) const
 {
 	// Bail out (with rather wacky data) if data is empty for some reason.
@@ -2123,7 +2136,7 @@ void UAnimSequence::UpdateSHAWithCurves(FSHA1& Sha, const FRawCurveTracks& InRaw
 {
 	for (const FFloatCurve& Curve : InRawCurveData.FloatCurves)
 	{
-		UpdateWithData(Sha, Curve.CurveUid);
+		UpdateWithData(Sha, Curve.Name.UID);
 		UpdateWithData(Sha, Curve.FloatCurve.DefaultValue);
 		UpdateSHAWithArray(Sha, Curve.FloatCurve.Keys);
 		UpdateWithData(Sha, Curve.FloatCurve.PreInfinityExtrap);
@@ -2220,12 +2233,14 @@ void UAnimSequence::BakeOutAdditiveIntoRawData()
 	FCompactPose Pose;
 	Pose.SetBoneContainer(&RequiredBones);
 	FAnimExtractContext ExtractContext;
+	USkeleton* MySkeleton = GetSkeleton();
+	check(MySkeleton);
 
 	for (int Frame = 0; Frame < NumFrames; ++Frame)
 	{
 		// Initialise curve data from Skeleton
 		FBlendedCurve Curve;
-		Curve.InitFrom(&(GetSkeleton()->GetCachedAnimCurveMappingNameUids()));
+		Curve.InitFrom(&(MySkeleton->GetCachedAnimCurveMappingNameUids()));
 
 		//Grab pose for this frame
 		const float CurrentFrameTime = Frame * IntervalTime;
@@ -2245,7 +2260,6 @@ void UAnimSequence::BakeOutAdditiveIntoRawData()
 		}
 
 		//Write out curve data for this frame
-		
 		const TArray<FSmartNameMapping::UID>& UIDList = *Curve.UIDList;
 		for (int32 CurveIndex = 0; CurveIndex < UIDList.Num(); ++CurveIndex)
 		{
@@ -2254,7 +2268,10 @@ void UAnimSequence::BakeOutAdditiveIntoRawData()
 			FFloatCurve* RawCurve = GetFloatCurve(NewCurveTracks, CurveUID);
 			if (!RawCurve && CurveEL.Value > 0.f) //Only make a new curve if we are going to give it data
 			{
-				NewCurveTracks.AddCurveData(CurveUID, CurveEL.Flags, FRawCurveTracks::FloatType);
+				FSmartName NewCurveName;
+				// if we don't have name, there is something wrong here. 
+				ensureAlways(MySkeleton->GetSmartNameByUID(USkeleton::AnimCurveMappingName, CurveUID, NewCurveName));
+				NewCurveTracks.AddCurveData(NewCurveName, CurveEL.Flags, FRawCurveTracks::FloatType);
 				RawCurve = GetFloatCurve(NewCurveTracks, CurveUID);
 			}
 
@@ -2287,7 +2304,7 @@ void UAnimSequence::BakeOutAdditiveIntoRawData()
 
 	const FSmartNameMapping* Mapping = GetSkeleton()->GetSmartNameContainer(USkeleton::AnimCurveMappingName);
 	check(Mapping); // Should always exist
-	RawCurveData.UpdateLastObservedNames(Mapping);
+	RawCurveData.RefreshName(Mapping);
 
 #if 0 //Validate baked data
 	for (FRawAnimSequenceTrack& RawTrack : RawAnimationData)
@@ -3201,8 +3218,11 @@ void UAnimSequence::RemoveNaNTracks()
 		}
 	}
 
-	MarkRawDataAsModified();
-	OnRawDataChanged();
+	if(bRecompress)
+	{
+		MarkRawDataAsModified();
+		OnRawDataChanged();
+	}
 }
 
 
@@ -3896,7 +3916,7 @@ void UAnimSequence::BakeTrackCurvesToRawAnimation()
 		if(!NameMapping)
 		{
 			// if no name mapping is found but curve exists, we should verify curve namex
-			VerifyCurveNames(CurSkeleton, USkeleton::AnimTrackCurveMappingName, RawCurveData.TransformCurves);
+			VerifyCurveNames<FTransformCurve>(CurSkeleton, USkeleton::AnimTrackCurveMappingName, RawCurveData.TransformCurves);
 			NameMapping = CurSkeleton->GetSmartNameContainer(USkeleton::AnimTrackCurveMappingName);
 		}
 		
@@ -3918,7 +3938,7 @@ void UAnimSequence::BakeTrackCurvesToRawAnimation()
 			FName BoneName;
 
 			if(Curve.GetCurveTypeFlag(ACF_Disabled)== false &&
-				ensure(NameMapping->GetName(Curve.CurveUid, BoneName)))
+				ensureAlways(NameMapping->GetName(Curve.Name.UID, BoneName)))
 			{
 				int32 TrackIndex = AnimationTrackNames.Find(BoneName);
 
@@ -3980,7 +4000,7 @@ void UAnimSequence::BakeTrackCurvesToRawAnimation()
 				for(int32 KeyIndex=0; KeyIndex < NumFrames; ++KeyIndex)
 				{
 					// now evaluate
-					FTransformCurve* TransformCurve = static_cast<FTransformCurve*>(RawCurveData.GetCurveData(Curve.CurveUid, FRawCurveTracks::TransformType));
+					FTransformCurve* TransformCurve = static_cast<FTransformCurve*>(RawCurveData.GetCurveData(Curve.Name.UID, FRawCurveTracks::TransformType));
 
 					if(ensure(TransformCurve))
 					{
@@ -4035,16 +4055,14 @@ void UAnimSequence::AddKeyToSequence(float Time, const FName& BoneName, const FT
 	USkeleton * CurrentSkeleton = GetSkeleton();
 	check (CurrentSkeleton);
 
-	const FSmartNameMapping* NameMapping = CurrentSkeleton->GetOrAddSmartNameContainer(USkeleton::AnimTrackCurveMappingName);
-	
-	FSmartNameMapping::UID NewUID;
-	CurrentSkeleton->AddSmartNameAndModify(USkeleton::AnimTrackCurveMappingName, CurveName, NewUID);
+	FSmartName NewCurveName;
+	CurrentSkeleton->AddSmartNameAndModify(USkeleton::AnimTrackCurveMappingName, CurveName, NewCurveName);
 
 	// add curve - this won't add duplicate curve
-	RawCurveData.AddCurveData(NewUID, ACF_DriveTrack | ACF_Editable, FRawCurveTracks::TransformType);
+	RawCurveData.AddCurveData(NewCurveName, ACF_DriveTrack | ACF_Editable, FRawCurveTracks::TransformType);
 
 	//Add this curve
-	FTransformCurve* TransformCurve = static_cast<FTransformCurve*>(RawCurveData.GetCurveData(NewUID, FRawCurveTracks::TransformType));
+	FTransformCurve* TransformCurve = static_cast<FTransformCurve*>(RawCurveData.GetCurveData(NewCurveName.UID, FRawCurveTracks::TransformType));
 	check(TransformCurve);
 
 	TransformCurve->UpdateOrAddKey(AdditiveTransform, Time);
diff --git a/Engine/Source/Runtime/Engine/Private/Animation/AnimSequenceBase.cpp b/Engine/Source/Runtime/Engine/Private/Animation/AnimSequenceBase.cpp
--- a/Engine/Source/Runtime/Engine/Private/Animation/AnimSequenceBase.cpp
+++ b/Engine/Source/Runtime/Engine/Private/Animation/AnimSequenceBase.cpp
@@ -19,39 +19,6 @@ UAnimSequenceBase::UAnimSequenceBase(const FObjectInitializer& ObjectInitializer
 {
 }
 
-template <typename DataType>
-void UAnimSequenceBase::VerifyCurveNames(USkeleton* InSkeleton, const FName& NameContainer, TArray<DataType>& CurveList)
-{
-	// since this is verify function that makes sure it exists after loaded
-	// we should add it if it doesn't exist
-	const FSmartNameMapping* NameMapping = InSkeleton->GetOrAddSmartNameContainer(NameContainer);
-
-	TArray<DataType*> UnlinkedCurves;
-	for(DataType& Curve : CurveList)
-	{
-		const FSmartNameMapping::UID* UID = NameMapping->FindUID(Curve.LastObservedName);
-		if(!UID)
-		{
-			// The skeleton doesn't know our name. Use the last observed name that was saved with the
-			// curve to create a new name. This can happen if a user saves an animation but not a skeleton
-			// either when updating the assets or editing the curves within.
-			UnlinkedCurves.Add(&Curve);
-		}
-		else if (Curve.CurveUid != *UID)// we verify if UID is correct
-		{
-			// if UID doesn't match, this is suspicious
-			// because same name but different UID is not idea
-			// so we'll fix up UID
-			Curve.CurveUid = *UID;
-		}
-	}
-
-	for(DataType* Curve : UnlinkedCurves)
-	{
-		InSkeleton->AddSmartNameAndModify(NameContainer, Curve->LastObservedName, Curve->CurveUid);
-	}
-}
-
 void UAnimSequenceBase::PostLoad()
 {
 	Super::PostLoad();
@@ -87,8 +54,7 @@ void UAnimSequenceBase::PostLoad()
 		{
 			for(FFloatCurve& Curve : RawCurveData.FloatCurves)
 			{
-				// Add the names of the curves into the smartname mapping and store off the curve uid which will be saved next time the sequence saves.
-				MySkeleton->AddSmartNameAndModify(USkeleton::AnimCurveMappingName, Curve.LastObservedName, Curve.CurveUid);
+				MySkeleton->VerifySmartName(USkeleton::AnimCurveMappingName, Curve.Name);
 			}
 		}
 		else
@@ -99,6 +65,8 @@ void UAnimSequenceBase::PostLoad()
 #if WITH_EDITOR
 		VerifyCurveNames<FTransformCurve>(MySkeleton, USkeleton::AnimTrackCurveMappingName, RawCurveData.TransformCurves);
 #endif
+
+		RawCurveData.SortFloatCurvesByUID();
 	}
 }
 
@@ -439,7 +407,7 @@ void UAnimSequenceBase::GetAssetRegistryTags(TArray<FAssetRegistryTag>& OutTags)
 
 	for(const FFloatCurve& Curve : RawCurveData.FloatCurves)
 	{
-		CurveIdList += FString::Printf(TEXT("%u%s"), Curve.CurveUid, *USkeleton::CurveTagDelimiter);
+		CurveIdList += FString::Printf(TEXT("%u%s"), Curve.Name.UID, *USkeleton::CurveTagDelimiter);
 	}
 	OutTags.Add(FAssetRegistryTag(USkeleton::CurveTag, CurveIdList, FAssetRegistryTag::TT_Hidden));
 }
@@ -490,38 +458,10 @@ void UAnimSequenceBase::EvaluateCurveData(FBlendedCurve& OutCurve, float Current
 
 void UAnimSequenceBase::Serialize(FArchive& Ar)
 {
-	if(Ar.ArIsSaving && Ar.UE4Ver() >= VER_UE4_SKELETON_ADD_SMARTNAMES)
-	{
-		if(USkeleton* MySkeleton = GetSkeleton())
-		{
-			const FSmartNameMapping* Mapping = MySkeleton->GetSmartNameContainer(USkeleton::AnimCurveMappingName);
-			check(Mapping); // Should always exist
-			RawCurveData.UpdateLastObservedNames(Mapping);
-		}
-	}
-
-#if WITH_EDITORONLY_DATA
-	if(Ar.ArIsSaving && Ar.UE4Ver() >= VER_UE4_ANIMATION_ADD_TRACKCURVES)
-	{
-		if(USkeleton* MySkeleton = GetSkeleton())
-		{
-			// we don't add track curve container unless it has been edited. 
-			if ( RawCurveData.TransformCurves.Num() > 0 )
-			{
-				const FSmartNameMapping* Mapping = MySkeleton->GetSmartNameContainer(USkeleton::AnimTrackCurveMappingName);
-				// this name might not exists because it's only available if you edit animation
-				if (Mapping)
-				{
-					RawCurveData.UpdateLastObservedNames(Mapping, FRawCurveTracks::TransformType);
-				}
-			}
-		}
-	}
-#endif //WITH_EDITORONLY_DATA
-
 	Super::Serialize(Ar);
 
-	RawCurveData.Serialize(Ar);
+	// fix up version issue and so on
+	RawCurveData.PostSerialize(Ar);
 }
 
 void UAnimSequenceBase::OnAssetPlayerTickedInternal(FAnimAssetTickContext &Context, const float PreviousTime, const float MoveDelta, const FAnimTickRecord &Instance, class UAnimInstance* InAnimInstance) const
diff --git a/Engine/Source/Runtime/Engine/Private/Animation/AnimSingleNodeInstance.cpp b/Engine/Source/Runtime/Engine/Private/Animation/AnimSingleNodeInstance.cpp
--- a/Engine/Source/Runtime/Engine/Private/Animation/AnimSingleNodeInstance.cpp
+++ b/Engine/Source/Runtime/Engine/Private/Animation/AnimSingleNodeInstance.cpp
@@ -91,6 +91,11 @@ void UAnimSingleNodeInstance::SetAnimationAsset(class UAnimationAsset* NewAsset,
 	}
 }
 
+void UAnimSingleNodeInstance::SetPreviewPoseOverride(const FName& PoseName, float Value)
+{
+	GetProxyOnGameThread<FAnimSingleNodeInstanceProxy>().SetPreviewPoseOverride(PoseName, Value);
+}
+
 void UAnimSingleNodeInstance::SetMontageLoop(UAnimMontage* Montage, bool bIsLooping, FName StartingSection)
 {
 	check (Montage);
@@ -405,4 +410,4 @@ FVector UAnimSingleNodeInstance::GetFilterLastOutput()
 	}
 
 	return FVector::ZeroVector;
-}
\ No newline at end of file
+}
diff --git a/Engine/Source/Runtime/Engine/Private/Animation/AnimSingleNodeInstanceProxy.cpp b/Engine/Source/Runtime/Engine/Private/Animation/AnimSingleNodeInstanceProxy.cpp
--- a/Engine/Source/Runtime/Engine/Private/Animation/AnimSingleNodeInstanceProxy.cpp
+++ b/Engine/Source/Runtime/Engine/Private/Animation/AnimSingleNodeInstanceProxy.cpp
@@ -5,6 +5,7 @@
 #include "Animation/AnimComposite.h"
 #include "Animation/AnimMontage.h"
 #include "Animation/BlendSpace.h"
+#include "Animation/PoseAsset.h"
 #include "Animation/AnimSingleNodeInstance.h"
 
 void FAnimSingleNodeInstanceProxy::Initialize(UAnimInstance* InAnimInstance)
@@ -45,7 +46,7 @@ bool FAnimSingleNodeInstanceProxy::Evaluate(FPoseContext& Output)
 
 				if (bCanProcessAdditiveAnimationsLocal)
 				{
-				Sequence->GetAdditiveBasePose(Output.Pose, Output.Curve, ExtractionContext);
+					Sequence->GetAdditiveBasePose(Output.Pose, Output.Curve, ExtractionContext);
 				}
 				else
 				{
@@ -136,6 +137,44 @@ bool FAnimSingleNodeInstanceProxy::Evaluate(FPoseContext& Output)
 				SlotEvaluatePose(Montage->SlotAnimTracks[0].SlotName, SourcePose, SourceCurve, WeightInfo.SourceWeight, Output.Pose, Output.Curve, WeightInfo.SlotNodeWeight, WeightInfo.TotalNodeWeight);
 			}
 		}
+		else if (UPoseAsset* PoseAsset = Cast<UPoseAsset>(CurrentAsset))
+		{
+#if WITH_EDITORONLY_DATA
+			USkeleton* MySkeleton = PoseAsset->GetSkeleton();
+ 			for (auto Iter=PreviewPoseOverride.CreateConstIterator(); Iter; ++Iter)
+ 			{
+ 				const FName& Name = Iter.Key();
+				const float Value = Iter.Value();
+				
+ 				FSmartNameMapping::UID CurveUID = MySkeleton->GetUIDByName(USkeleton::AnimCurveMappingName, Name);
+ 
+ 				if (CurveUID != FSmartNameMapping::MaxUID)
+ 				{
+					Output.Curve.Set(CurveUID, Value, ACF_DrivesPose);
+ 				}
+ 			}
+#endif 
+			if (PoseAsset->IsValidAdditive())
+			{
+				FAnimExtractContext ExtractContext;
+
+				PoseAsset->GetBaseAnimationPose(Output.Pose, Output.Curve, ExtractContext);
+				PoseAsset->GetAnimationPose(Output.Pose, Output.Curve, ExtractContext);
+			}
+			else
+			{
+				// @todo : what do we do for ingame? I think in game, we should just allow to load anim sequence with curves
+				PoseAsset->GetAnimationPose(Output.Pose, Output.Curve, FAnimExtractContext());
+			}
+		}
+
+#if WITH_EDITORONLY_DATA
+		// if it has preview pose asset, we have to handle that after we do all animation
+		if (const UPoseAsset* PoseAsset = CurrentAsset->PreviewPoseAsset)
+		{
+			PoseAsset->GetAnimationPose(Output.Pose, Output.Curve, FAnimExtractContext());
+		}
+#endif // WITH_EDITORONLY_DATA
 	}
 
 	return true;
@@ -216,6 +255,11 @@ void FAnimSingleNodeInstanceProxy::UpdateAnimationNode(float DeltaSeconds)
 			PreviewBasePose = Montage->PreviewBasePose;
 #endif
 		}
+		else if (UPoseAsset* PoseAsset = Cast<UPoseAsset>(CurrentAsset))
+		{
+			FAnimTickRecord& TickRecord = CreateUninitializedTickRecord(INDEX_NONE, /*out*/ SyncGroup);
+			MakePoseAssetTickRecord(TickRecord, PoseAsset, 1.f);
+		}
 	}
 
 #if WITH_EDITORONLY_DATA
@@ -248,6 +292,22 @@ void FAnimSingleNodeInstanceProxy::PostUpdate(UAnimInstance* InAnimInstance) con
 	}
 }
 
+void FAnimSingleNodeInstanceProxy::PreUpdate(UAnimInstance* InAnimInstance, float DeltaSeconds) 
+{
+	FAnimInstanceProxy::PreUpdate(InAnimInstance, DeltaSeconds);
+#if WITH_EDITOR
+	// @fixme only do this in pose asset
+	// copy data to PreviewPoseOverride
+	TMap<FName, float> PoseCurveList;
+
+	InAnimInstance->GetAnimationCurveList(ACF_DrivesPose, PoseCurveList);
+
+	if (PoseCurveList.Num() > 0)
+	{
+		PreviewPoseOverride.Append(PoseCurveList);
+	}
+#endif // WITH_EDITOR
+}
 void FAnimSingleNodeInstanceProxy::InitializeObjects(UAnimInstance* InAnimInstance)
 {
 	FAnimInstanceProxy::InitializeObjects(InAnimInstance);
@@ -263,6 +323,32 @@ void FAnimSingleNodeInstanceProxy::ClearObjects()
 	CurrentAsset = nullptr;
 }
 
+void FAnimSingleNodeInstanceProxy::SetPreviewPoseOverride(const FName& PoseName, float Value)
+{
+	float *CurveValPtr = PreviewPoseOverride.Find(PoseName);
+	bool bShouldAddToList = FPlatformMath::Abs(Value) > ZERO_ANIMWEIGHT_THRESH;
+	if (bShouldAddToList)
+	{
+		if (CurveValPtr)
+		{
+			// sum up, in the future we might normalize, but for now this just sums up
+			// this won't work well if all of them have full weight - i.e. additive 
+			*CurveValPtr = Value;
+		}
+		else
+		{
+			PreviewPoseOverride.Add(PoseName, Value);
+		}
+	}
+	// if less than ZERO_ANIMWEIGHT_THRESH
+	// no reason to keep them on the list
+	else
+	{
+		// remove if found
+		PreviewPoseOverride.Remove(PoseName);
+	}
+}
+
 void FAnimSingleNodeInstanceProxy::InternalBlendSpaceEvaluatePose(class UBlendSpaceBase* BlendSpace, TArray<FBlendSampleData>& BlendSampleDataCache, FPoseContext& OutContext)
 {
 	if (BlendSpace->IsValidAdditive())
@@ -310,10 +396,11 @@ void FAnimSingleNodeInstanceProxy::SetAnimationAsset(class UAnimationAsset* NewA
 
 #if WITH_EDITORONLY_DATA
 	PreviewPoseCurrentTime = 0.0f;
+	PreviewPoseOverride.Reset();
 #endif
 
-	UBlendSpaceBase * BlendSpace = Cast<UBlendSpaceBase>(NewAsset);
-	if (BlendSpace)
+	
+	if (UBlendSpaceBase* BlendSpace = Cast<UBlendSpaceBase>(NewAsset))
 	{
 		BlendSpace->InitializeFilter(&BlendFilter);
 	}
diff --git a/Engine/Source/Runtime/Engine/Private/Animation/AnimationRuntime.cpp b/Engine/Source/Runtime/Engine/Private/Animation/AnimationRuntime.cpp
--- a/Engine/Source/Runtime/Engine/Private/Animation/AnimationRuntime.cpp
+++ b/Engine/Source/Runtime/Engine/Private/Animation/AnimationRuntime.cpp
@@ -27,38 +27,6 @@ DECLARE_CYCLE_STAT(TEXT("AccumulateMeshSpaceRotAdditiveToLocalPose"), STAT_Accum
 DECLARE_CYCLE_STAT(TEXT("BlendPosesPerBoneFilter"), STAT_BlendPosesPerBoneFilter, STATGROUP_Anim);
 
 
-/////////////////////////////////////////////////////////
-// Templated Transform Blend Functionality
-
-namespace ETransformBlendMode
-{
-	enum Type
-	{
-		Overwrite,
-		Accumulate
-	};
-}
-
-template<int32>
-void BlendTransform(const FTransform& Source, FTransform& Dest, const float BlendWeight)
-{
-	check(false); /// should never call this
-}
-
-template<>
-void BlendTransform<ETransformBlendMode::Overwrite>(const FTransform& Source, FTransform& Dest, const float BlendWeight)
-{
-	const ScalarRegister VBlendWeight(BlendWeight);
-	Dest = Source * VBlendWeight;
-}
-
-template<>
-void BlendTransform<ETransformBlendMode::Accumulate>(const FTransform& Source, FTransform& Dest, const float BlendWeight)
-{
-	const ScalarRegister VBlendWeight(BlendWeight);
-	Dest.AccumulateWithShortestRotation(Source, VBlendWeight);
-}
-
 //////////////////////////////////////////////////////////////////////////
 
 void FAnimationRuntime::NormalizeRotations(const FBoneContainer& RequiredBones, /*inout*/ FTransformArrayA2& Atoms)
@@ -110,19 +78,6 @@ FORCEINLINE void BlendPose(const FCompactPose& SourcePose, FCompactPose& ResultP
 	}
 }
 
-FORCEINLINE void BlendCurves(const TFixedSizeArrayView<FBlendedCurve>& SourceCurves, const TFixedSizeArrayView<float>& SourceWeights, FBlendedCurve& OutCurve)
-{
-	if (SourceCurves.Num() > 0)
-	{
-		OutCurve.Override(SourceCurves[0], SourceWeights[0]);
-
-		for (int32 CurveIndex = 1; CurveIndex<SourceCurves.Num(); ++CurveIndex)
-		{
-			OutCurve.Accumulate(SourceCurves[CurveIndex], SourceWeights[CurveIndex]);
-		}
-	}
-}
-
 FORCEINLINE void BlendCurves(const TFixedSizeArrayView<FBlendedCurve>& SourceCurves, const TFixedSizeArrayView<float>& SourceWeights, const TFixedSizeArrayView<int32>& SourceWeightsIndices, FBlendedCurve& OutCurve)
 {
 	if (SourceCurves.Num() > 0)
@@ -487,6 +442,14 @@ void FAnimationRuntime::CombineWithAdditiveAnimations(int32 NumAdditivePoses, co
 	}
 }
 
+void FAnimationRuntime::ConvertTransformToAdditive(FTransform& TargetTransform, const FTransform& BaseTransform)
+{
+	TargetTransform.SetRotation(TargetTransform.GetRotation() * BaseTransform.GetRotation().Inverse());
+	TargetTransform.SetTranslation(TargetTransform.GetTranslation() - BaseTransform.GetTranslation());
+	TargetTransform.SetScale3D(TargetTransform.GetScale3D() * BaseTransform.GetSafeScaleReciprocal(BaseTransform.GetScale3D()));
+	TargetTransform.NormalizeRotation();
+}
+
 void FAnimationRuntime::ConvertPoseToAdditive(FCompactPose& TargetPose, const FCompactPose& BasePose)
 {
 	for (FCompactPoseBoneIndex BoneIndex : BasePose.ForEachBoneIndex())
@@ -494,10 +457,7 @@ void FAnimationRuntime::ConvertPoseToAdditive(FCompactPose& TargetPose, const FC
 		FTransform& TargetTransform = TargetPose[BoneIndex];
 		const FTransform& BaseTransform = BasePose[BoneIndex];
 
-		TargetTransform.SetRotation(TargetTransform.GetRotation() * BaseTransform.GetRotation().Inverse());
-		TargetTransform.SetTranslation(TargetTransform.GetTranslation() - BaseTransform.GetTranslation());
-		TargetTransform.SetScale3D(TargetTransform.GetScale3D() * BaseTransform.GetSafeScaleReciprocal(BaseTransform.GetScale3D()));
-		TargetTransform.NormalizeRotation();
+		ConvertTransformToAdditive(TargetTransform, BaseTransform);
 	}
 }
 
@@ -530,6 +490,7 @@ void FAnimationRuntime::ConvertMeshRotationPoseToLocalSpace(FCompactPose& Pose)
 		Pose[BoneIndex].SetRotation(LocalSpaceRotation);
 	}
 }
+
 void FAnimationRuntime::AccumulateAdditivePose(FCompactPose& BasePose, const FCompactPose& AdditivePose, FBlendedCurve& BaseCurve, const FBlendedCurve& AdditiveCurve, float Weight, enum EAdditiveAnimationType AdditiveType)
 {
 	if (AdditiveType == AAT_RotationOffsetMeshSpace)
@@ -1314,59 +1275,50 @@ bool FAnimationRuntime::ContainsNaN(TArray<FBoneIndexType>& RequiredBoneIndices,
 #endif
 
 #if WITH_EDITOR
-void FAnimationRuntime::FillUpSpaceBasesRefPose(const USkeleton* Skeleton, TArray<FTransform> &SpaceBaseRefPose)
+void FAnimationRuntime::FillUpSpaceBases(const FReferenceSkeleton& RefSkeleton, const TArray<FTransform> &LocalAtoms, TArray<FTransform> &SpaceBases)
 {
-	check(Skeleton);
-
-	const TArray<FTransform>& ReferencePose = Skeleton->GetReferenceSkeleton().GetRefBonePose();
-	SpaceBaseRefPose.Empty(ReferencePose.Num());
-	SpaceBaseRefPose.AddUninitialized(ReferencePose.Num());
+	SpaceBases.Empty(RefSkeleton.GetNum());
+	SpaceBases.AddUninitialized(RefSkeleton.GetNum());
 
 	// initialize to identity since some of them don't have tracks
-	for(int Index=0; Index <SpaceBaseRefPose.Num(); ++Index)
+	for (int Index = 0; Index < SpaceBases.Num(); ++Index)
 	{
-		int32 ParentIndex = Skeleton->GetReferenceSkeleton().GetParentIndex(Index);
-		if(ParentIndex != INDEX_NONE)
+		int32 ParentIndex = RefSkeleton.GetParentIndex(Index);
+		if (ParentIndex != INDEX_NONE)
 		{
-			SpaceBaseRefPose[Index] = ReferencePose[Index] * SpaceBaseRefPose[ParentIndex];
+			SpaceBases[Index] = LocalAtoms[Index] * SpaceBases[ParentIndex];
 		}
 		else
 		{
-			SpaceBaseRefPose[Index] = ReferencePose[Index];
+			SpaceBases[Index] = LocalAtoms[Index];
 		}
 	}
 }
 
-void FAnimationRuntime::FillUpSpaceBasesRetargetBasePose(const USkeleton* Skeleton, TArray<FTransform> &SpaceBaseRefPose)
+void FAnimationRuntime::FillUpSpaceBasesRefPose(const USkeleton* Skeleton, TArray<FTransform> &SpaceBaseRefPose)
+{
+	check(Skeleton);
+
+	const FReferenceSkeleton& RefSkeleton = Skeleton->GetReferenceSkeleton();
+	const TArray<FTransform>& ReferencePose = RefSkeleton.GetRefBonePose();
+	FillUpSpaceBases(RefSkeleton, ReferencePose, SpaceBaseRefPose);
+}
+
+void FAnimationRuntime::FillUpSpaceBasesRetargetBasePose(const USkeleton* Skeleton, TArray<FTransform> &SpaceBases)
 {
 	check(Skeleton);
 
 	// @Todo fixme: this has to get preview mesh instead of skeleton
-	
 	const USkeletalMesh* PreviewMesh = Skeleton->GetPreviewMesh();
 	if (PreviewMesh)
 	{
 		const TArray<FTransform>& ReferencePose = PreviewMesh->RetargetBasePose;
-		SpaceBaseRefPose.Empty(ReferencePose.Num());
-		SpaceBaseRefPose.AddUninitialized(ReferencePose.Num());
-
-		// initialize to identity since some of them don't have tracks
-		for(int Index=0; Index <SpaceBaseRefPose.Num(); ++Index)
-		{
-			int32 ParentIndex = PreviewMesh->RefSkeleton.GetParentIndex(Index);
-			if(ParentIndex != INDEX_NONE)
-			{
-				SpaceBaseRefPose[Index] = ReferencePose[Index] * SpaceBaseRefPose[ParentIndex];
-			}
-			else
-			{
-				SpaceBaseRefPose[Index] = ReferencePose[Index];
-			}
-		}
+		const FReferenceSkeleton& RefSkeleton = PreviewMesh->RefSkeleton;
+		FillUpSpaceBases(RefSkeleton, ReferencePose, SpaceBases);
 	}
 	else
 	{
-		FAnimationRuntime::FillUpSpaceBasesRefPose(Skeleton, SpaceBaseRefPose);
+		FAnimationRuntime::FillUpSpaceBasesRefPose(Skeleton, SpaceBases);
 	}
 }
 #endif // WITH_EDITOR
@@ -1421,6 +1373,62 @@ void FAnimationRuntime::AppendActiveMorphTargets(const USkeletalMesh* InSkeletal
 	}
 }
 
+int32 FAnimationRuntime::GetStringDistance(const FString& First, const FString& Second) 
+{
+	// Finds the distance between strings, where the distance is the number of operations we would need
+	// to perform on First to match Second.
+	// Operations are: Adding a character, Removing a character, changing a character.
+
+	const int32 FirstLength = First.Len();
+	const int32 SecondLength = Second.Len();
+
+	// Already matching
+	if (First == Second)
+	{
+		return 0;
+	}
+
+	// No first string, so we need to add SecondLength characters to match
+	if (FirstLength == 0)
+	{
+		return SecondLength;
+	}
+
+	// No Second string, so we need to add FirstLength characters to match
+	if (SecondLength == 0)
+	{
+		return FirstLength;
+	}
+
+	TArray<int32> PrevRow;
+	TArray<int32> NextRow;
+	PrevRow.AddZeroed(SecondLength + 1);
+	NextRow.AddZeroed(SecondLength + 1);
+
+	// Initialise prev row to num characters we need to remove from Second
+	for (int32 I = 0; I < PrevRow.Num(); ++I)
+	{
+		PrevRow[I] = I;
+	}
+
+	for (int32 I = 0; I < FirstLength; ++I)
+	{
+		// Calculate current row
+		NextRow[0] = I + 1;
+
+		for (int32 J = 0; J < SecondLength; ++J)
+		{
+			int32 Indicator = (First[I] == Second[J]) ? 0 : 1;
+			NextRow[J + 1] = FMath::Min3(NextRow[J] + 1, PrevRow[J + 1] + 1, PrevRow[J] + Indicator);
+		}
+
+		// Copy back
+		PrevRow = NextRow;
+	}
+
+	return NextRow[SecondLength];
+}
+
 /////////////////////////////////////////////////////////////////////////////////////////
 // FA2CSPose
 /////////////////////////////////////////////////////////////////////////////////////////
diff --git a/Engine/Source/Runtime/Engine/Private/Animation/PoseAsset.cpp b/Engine/Source/Runtime/Engine/Private/Animation/PoseAsset.cpp
new file mode 100644
--- /dev/null
+++ b/Engine/Source/Runtime/Engine/Private/Animation/PoseAsset.cpp
@@ -0,0 +1,915 @@
+// Copyright 1998-2016 Epic Games, Inc. All Rights Reserved.
+
+#include "EnginePrivate.h"
+#include "AnimationUtils.h"
+#include "AnimationRuntime.h"
+#include "Animation/PoseAsset.h"
+
+
+//////////////////////////////////////////////////////////////////////////////////////////////////////////////////
+//
+// FPoseDataContainer
+//
+//////////////////////////////////////////////////////////////////////////////////////////////////////////////////
+
+void FPoseDataContainer::Reset()
+{
+	// clear everything
+	PoseNames.Reset();
+	Poses.Reset();
+	Tracks.Reset();
+	TrackMap.Reset();
+	Curves.Reset();
+}
+
+void FPoseDataContainer::AddOrUpdatePose(const FSmartName& InPoseName, const TArray<FTransform>& InlocalSpacePose, const TArray<float>& InCurveData)
+{
+	// make sure the transform is correct size
+	if (ensureAlways(InlocalSpacePose.Num() == Tracks.Num()))
+	{
+		// find or add pose data
+		FPoseData* PoseDataPtr = FindOrAddPoseData(InPoseName);
+		// now add pose
+		PoseDataPtr->LocalSpacePose = InlocalSpacePose;
+		PoseDataPtr->CurveData = InCurveData;
+	}
+
+	// for now we only supports same tracks
+}
+
+bool FPoseDataContainer::InsertTrack(const FName& InTrackName, USkeleton* InSkeleton, FName& InRetargetSourceName)
+{
+	check(InSkeleton);
+
+	// make sure the transform is correct size
+	if (Tracks.Contains(InTrackName))
+	{
+		int32 SkeletonIndex = InSkeleton->GetReferenceSkeleton().FindBoneIndex(InTrackName);
+		int32 TrackIndex = INDEX_NONE;
+		if (SkeletonIndex != INDEX_NONE)
+		{
+			Tracks.Add(InTrackName);
+			TrackMap.Add(InTrackName, SkeletonIndex);
+			TrackIndex = Tracks.Num() - 1;
+
+			// now insert default refpose
+			const FTransform DefaultPose = GetDefaultTransform(SkeletonIndex, InSkeleton, InRetargetSourceName);
+
+			for (auto& PoseData : Poses)
+			{
+				ensureAlways(PoseData.LocalSpacePose.Num() == TrackIndex);
+
+				PoseData.LocalSpacePose.Add(DefaultPose);
+
+				// make sure they always match
+				ensureAlways(PoseData.LocalSpacePose.Num() == Tracks.Num());
+			}
+
+			return true;
+		}
+
+		return false;
+	}
+
+	return false;
+}
+
+void FPoseDataContainer::GetPoseCurve(const FPoseData* PoseData, FBlendedCurve& OutCurve) const
+{
+	if (PoseData)
+	{
+		const TArray<float>& CurveValues = PoseData->CurveData;
+		checkSlow(CurveValues.Num() == Curves.Num());
+
+		// extract curve - not optimized, can use optimization
+		for (int32 CurveIndex = 0; CurveIndex < Curves.Num(); ++CurveIndex)
+		{
+			const FAnimCurveBase& Curve = Curves[CurveIndex];
+			OutCurve.Set(Curve.Name.UID, CurveValues[CurveIndex], Curve.GetCurveTypeFlags());
+		}
+	}
+}
+// remove track if all poses has identity key
+// we don't want to save transform per track (good for compression) because that will slow down blending
+void FPoseDataContainer::Shrink(USkeleton* InSkeleton, FName& InRetargetSourceName)
+{
+}
+
+FTransform FPoseDataContainer::GetDefaultTransform(const FName& InTrackName, USkeleton* InSkeleton, const FName& InRetargetSourceName) const
+{
+	int32 SkeletonIndex = InSkeleton->GetReferenceSkeleton().FindBoneIndex(InTrackName);
+	if (SkeletonIndex != INDEX_NONE)
+	{
+		return GetDefaultTransform(SkeletonIndex, InSkeleton, InRetargetSourceName);
+	}
+
+	return FTransform::Identity;
+}
+
+FTransform FPoseDataContainer::GetDefaultTransform(int32 SkeletonIndex, USkeleton* InSkeleton, const FName& InRetargetSourceName) const
+{
+	// now insert default refpose
+	const TArray<FTransform>& RefPose = InSkeleton->GetRefLocalPoses(InRetargetSourceName);
+
+	if (RefPose.IsValidIndex(SkeletonIndex))
+	{
+		return RefPose[SkeletonIndex];
+	}
+
+	return FTransform::Identity;
+}
+
+bool FPoseDataContainer::FillUpDefaultPose(FPoseData* PoseData, USkeleton* InSkeleton, FName& InRetargetSourceName)
+{
+	if (PoseData)
+	{
+		int32 TrackIndex = 0;
+		const TArray<FTransform>& RefPose = InSkeleton->GetRefLocalPoses(InRetargetSourceName);
+		for (TPair<FName, int32>& TrackIter : TrackMap)
+		{
+			int32 SkeletonIndex = TrackIter.Value;
+			PoseData->LocalSpacePose[TrackIndex] = RefPose[SkeletonIndex];
+
+			++TrackIndex;
+		}
+
+		return true;
+	}
+
+	return false;
+}
+
+bool FPoseDataContainer::FillUpDefaultPose(const FSmartName& InPoseName, USkeleton* InSkeleton, FName& InRetargetSourceName)
+{
+	const TArray<FTransform>& RefPose = InSkeleton->GetRefLocalPoses(InRetargetSourceName);
+	FPoseData* PoseData = FindPoseData(InPoseName);
+
+	if (PoseData)
+	{
+		return FillUpDefaultPose(PoseData, InSkeleton, InRetargetSourceName);
+	}
+
+	return false;
+}
+
+FPoseData* FPoseDataContainer::FindPoseData(FSmartName PoseName)
+{
+	int32 PoseIndex = PoseNames.Find(PoseName);
+	if (PoseIndex != INDEX_NONE)
+	{
+		return &Poses[PoseIndex];
+	}
+
+	return nullptr;
+}
+
+FPoseData* FPoseDataContainer::FindOrAddPoseData(FSmartName PoseName)
+{
+	int32 PoseIndex = PoseNames.Find(PoseName);
+	if (PoseIndex == INDEX_NONE)
+	{
+		PoseIndex = PoseNames.Add(PoseName);
+		check(PoseIndex == Poses.AddZeroed(1));
+	}
+
+	return &Poses[PoseIndex];
+}
+
+void FPoseDataContainer::RenamePose(FSmartName OldPoseName, FSmartName NewPoseName)
+{
+	int32 PoseIndex = PoseNames.Find(OldPoseName);
+	if (PoseIndex != INDEX_NONE)
+	{
+		PoseNames[PoseIndex] = NewPoseName;
+	}
+}
+
+bool FPoseDataContainer::DeletePose(FSmartName PoseName)
+{
+	int32 PoseIndex = PoseNames.Find(PoseName);
+	if (PoseIndex != INDEX_NONE)
+	{
+		PoseNames.RemoveAt(PoseIndex);
+		Poses.RemoveAt(PoseIndex);
+		return true;
+	}
+
+	return false;
+}
+
+bool FPoseDataContainer::DeleteCurve(FSmartName CurveName)
+{
+	for (int32 CurveIndex = 0; CurveIndex < Curves.Num(); ++CurveIndex)
+	{
+		if (Curves[CurveIndex].Name == CurveName)
+		{
+			Curves.RemoveAt(CurveIndex);
+
+			// delete this index from all poses
+			for (int32 PoseIndex = 0; PoseIndex < Poses.Num(); ++PoseIndex)
+			{
+				Poses[PoseIndex].CurveData.RemoveAt(CurveIndex);
+			}
+
+			return true;
+		}
+	}
+
+	return false;
+}
+void FPoseDataContainer::ConvertToFullPose(int32 InBasePoseIndex, const TArray<FTransform>& InBasePose, const TArray<float>& InBaseCurve)
+{
+	check(InBaseCurve.Num() == Curves.Num());
+
+	for (int32 PoseIndex = 0; PoseIndex < Poses.Num(); ++PoseIndex)
+	{
+		// if this pose is not base pose
+		if (PoseIndex != InBasePoseIndex)
+		{
+			// should it be move? Why? I need that buffer still
+			TArray<FTransform> AdditivePose = Poses[PoseIndex].LocalSpacePose;
+			FPoseData& PoseData = Poses[PoseIndex];
+
+			check(AdditivePose.Num() == InBasePose.Num());
+			for (int32 BoneIndex = 0; BoneIndex < AdditivePose.Num(); ++BoneIndex)
+			{
+				PoseData.LocalSpacePose[BoneIndex] = InBasePose[BoneIndex];
+				PoseData.LocalSpacePose[BoneIndex].Accumulate(AdditivePose[BoneIndex]);
+				PoseData.LocalSpacePose[BoneIndex].NormalizeRotation();
+			}
+
+			int32 CurveNum = Curves.Num();
+			checkSlow(CurveNum == PoseData.CurveData.Num());
+			for (int32 CurveIndex = 0; CurveIndex < CurveNum; ++CurveIndex)
+			{
+				PoseData.CurveData[CurveIndex] = InBaseCurve[CurveIndex] + PoseData.CurveData[CurveIndex];
+			}
+		}
+	}
+}
+
+void FPoseDataContainer::ConvertToAdditivePose(int32 InBasePoseIndex, const TArray<FTransform>& InBasePose, const TArray<float>& InBaseCurve)
+{
+	check(InBaseCurve.Num() == Curves.Num());
+	for (int32 PoseIndex = 0; PoseIndex < Poses.Num(); ++PoseIndex)
+	{
+		// if this pose is not base pose
+		if (PoseIndex != InBasePoseIndex)
+		{
+			FPoseData& PoseData = Poses[PoseIndex];
+			check(PoseData.LocalSpacePose.Num() == InBasePose.Num());
+			for (int32 BoneIndex = 0; BoneIndex < InBasePose.Num(); ++BoneIndex)
+			{
+				FAnimationRuntime::ConvertTransformToAdditive(PoseData.LocalSpacePose[BoneIndex], InBasePose[BoneIndex]);
+				PoseData.LocalSpacePose[BoneIndex].NormalizeRotation();
+			}
+
+			int32 CurveNum = Curves.Num();
+			checkSlow(CurveNum == PoseData.CurveData.Num());
+			for (int32 CurveIndex = 0; CurveIndex < CurveNum; ++CurveIndex)
+			{
+				PoseData.CurveData[CurveIndex] = PoseData.CurveData[CurveIndex] - InBaseCurve[CurveIndex];
+			}
+		}
+	}
+}
+
+/////////////////////////////////////////////////////
+// UPoseAsset
+/////////////////////////////////////////////////////
+UPoseAsset::UPoseAsset(const FObjectInitializer& ObjectInitializer)
+	: Super(ObjectInitializer)
+	, NumPoses(0)
+	, bNormalizeWeight(true)
+	, bAdditivePose(false)
+	, BasePoseIndex(-1)
+{
+}
+
+void UPoseAsset::GetBaseAnimationPose(struct FCompactPose& OutPose, FBlendedCurve& InOutCurve, const FAnimExtractContext& ExtractionContext) const
+{
+	OutPose.ResetToRefPose();
+
+	if (bAdditivePose && PoseContainer.Poses.IsValidIndex(BasePoseIndex))
+	{
+		const FBoneContainer& RequiredBones = OutPose.GetBoneContainer();
+		USkeleton* MySkeleton = GetSkeleton();
+
+		// this contains compact bone pose list that this pose cares
+		TArray<FCompactPoseBoneIndex> CompactBonePoseList;
+
+		const int32 TrackNum = PoseContainer.TrackMap.Num();
+
+		for (const TPair<FName, int32>& TrackPair : PoseContainer.TrackMap)
+		{
+			const int32 SkeletonBoneIndex = TrackPair.Value;
+			const FCompactPoseBoneIndex PoseBoneIndex = RequiredBones.GetCompactPoseIndexFromSkeletonIndex(SkeletonBoneIndex);
+			// we add even if it's invalid because we want it to match with track index
+			CompactBonePoseList.Add(PoseBoneIndex);
+		}
+
+		const TArray<FTransform>& PoseTransform = PoseContainer.Poses[BasePoseIndex].LocalSpacePose;
+
+		for (int32 TrackIndex = 0; TrackIndex < TrackNum; ++TrackIndex)
+		{
+			if (CompactBonePoseList[TrackIndex] != INDEX_NONE)
+			{
+				OutPose[CompactBonePoseList[TrackIndex]] = PoseTransform[TrackIndex];
+			}
+		}
+
+		PoseContainer.GetPoseCurve(&PoseContainer.Poses[BasePoseIndex], InOutCurve);
+	}
+}
+
+void UPoseAsset::GetAnimationPose(struct FCompactPose& OutPose, FBlendedCurve& OutCurve, const FAnimExtractContext& ExtractionContext) const
+{
+	const FBoneContainer& RequiredBones = OutPose.GetBoneContainer();
+	USkeleton* MySkeleton = GetSkeleton();
+
+	// this contains compact bone pose list that this pose cares
+	TArray<FCompactPoseBoneIndex> CompactBonePoseList;
+
+	const int32 TrackNum = PoseContainer.TrackMap.Num();
+
+	for (const TPair<FName, int32>& TrackPair : PoseContainer.TrackMap)
+	{
+		const int32 SkeletonBoneIndex = TrackPair.Value;
+		const FCompactPoseBoneIndex PoseBoneIndex = RequiredBones.GetCompactPoseIndexFromSkeletonIndex(SkeletonBoneIndex);
+		// we add even if it's invalid because we want it to match with track index
+		CompactBonePoseList.Add(PoseBoneIndex);
+	}
+
+	TMap<const FPoseData*, float> IndexToWeightMap;
+
+	float TotalWeight = 0.f;
+
+	check(PoseContainer.IsValid());
+
+	for (int32 PoseIndex = 0; PoseIndex < PoseContainer.Poses.Num(); ++PoseIndex)
+	{
+		const FSmartName& PoseName = PoseContainer.PoseNames[PoseIndex];
+		const FPoseData& PoseData = PoseContainer.Poses[PoseIndex];
+
+		float Value = OutCurve.Get(PoseName.UID);
+
+		// we only add to the list if it's not additive Or if it's additive, we don't want to add base pose index
+		// and has weight
+		if ((!bAdditivePose || PoseIndex != BasePoseIndex) && FAnimationRuntime::HasWeight(Value))
+		{
+			IndexToWeightMap.Add(&PoseData, Value);
+			TotalWeight += Value;
+		}
+	}
+
+	// need a copy of pose since we don't know if this can be additive
+	FCompactPose CurrentPose;
+	FBlendedCurve CurrentCurve;
+	CurrentPose.SetBoneContainer(&OutPose.GetBoneContainer());
+	CurrentCurve.InitFrom(OutCurve);
+
+	if (bAdditivePose)
+	{
+		CurrentPose.ResetToIdentity();
+	}
+	else
+	{
+		CurrentPose.ResetToRefPose();
+	}
+
+	const int32 TotalNumberOfValidPoses = IndexToWeightMap.Num();
+	if (TotalNumberOfValidPoses > 0)
+	{
+		TArray<FTransform> BlendedBoneTransform;
+		bool bFirstPose = true;
+		TArray<FBlendedCurve> PoseCurves;
+		TArray<float>	CurveWeights;
+
+		//if full pose, we'll have to normalize by weight
+		if (bNormalizeWeight && TotalWeight != 1.f)
+		{
+			for (TPair<const FPoseData*, float>& WeightPair : IndexToWeightMap)
+			{
+				WeightPair.Value /= TotalWeight;
+			}
+		}
+
+		PoseCurves.AddDefaulted(TotalNumberOfValidPoses);
+		CurveWeights.AddUninitialized(TotalNumberOfValidPoses);
+		int32 PoseIndex = 0;
+		for (const TPair<const FPoseData*, float>& ActivePosePair : IndexToWeightMap)
+		{
+			const FPoseData* Pose = ActivePosePair.Key;
+			const float Weight = ActivePosePair.Value;
+
+			if (bFirstPose)
+			{
+				BlendedBoneTransform.AddUninitialized(Pose->LocalSpacePose.Num());
+				for (int32 TrackIndex = 0; TrackIndex < Pose->LocalSpacePose.Num(); ++TrackIndex)
+				{
+					BlendTransform<ETransformBlendMode::Overwrite>(Pose->LocalSpacePose[TrackIndex], BlendedBoneTransform[TrackIndex], Weight);
+				}
+				bFirstPose = false;
+			}
+			else
+			{
+				for (int32 TrackIndex = 0; TrackIndex < Pose->LocalSpacePose.Num(); ++TrackIndex)
+				{
+					BlendTransform<ETransformBlendMode::Accumulate>(Pose->LocalSpacePose[TrackIndex], BlendedBoneTransform[TrackIndex], Weight);
+				}
+			}
+
+			CurveWeights[PoseIndex] = Weight;
+			PoseCurves[PoseIndex].InitFrom(OutCurve);
+			PoseContainer.GetPoseCurve(Pose, PoseCurves[PoseIndex]);
+			++PoseIndex;
+		}
+
+		// blend curves
+		BlendCurves(PoseCurves, CurveWeights, CurrentCurve);
+
+		for (int32 TrackIndex = 0; TrackIndex < TrackNum; ++TrackIndex)
+		{
+			if (CompactBonePoseList[TrackIndex] != INDEX_NONE)
+			{
+				CurrentPose[CompactBonePoseList[TrackIndex]] = BlendedBoneTransform[TrackIndex];
+				CurrentPose.NormalizeRotations();
+			}
+		}
+
+		if (bAdditivePose)
+		{
+			FAnimationRuntime::AccumulateAdditivePose(OutPose, CurrentPose, OutCurve, CurrentCurve, TotalWeight, AAT_LocalSpaceBase);
+			OutPose.NormalizeRotations();
+		}
+		else
+		{
+			OutPose = CurrentPose;
+			OutCurve = CurrentCurve;
+		}
+	}
+}
+
+void UPoseAsset::PostLoad()
+{
+	Super::PostLoad();
+
+	// fix curve names
+	USkeleton* MySkeleton = GetSkeleton();
+	if (MySkeleton)
+	{
+		MySkeleton->VerifySmartNames(USkeleton::AnimCurveMappingName, PoseContainer.PoseNames);
+
+		for (auto& Curve : PoseContainer.Curves)
+		{
+			MySkeleton->VerifySmartName(USkeleton::AnimCurveMappingName, Curve.Name);
+		}
+	}
+
+	// I  have to fix pose names
+	RecacheTrackmap();
+}
+
+void UPoseAsset::Serialize(FArchive& Ar)
+{
+	Super::Serialize(Ar);
+}
+
+const TArray<FSmartName> UPoseAsset::GetPoseNames() const
+{
+	return PoseContainer.PoseNames;
+}
+
+const TArray<FSmartName> UPoseAsset::GetCurveNames() const
+{
+	TArray<FSmartName> CurveNames;
+	for (int32 CurveIndex = 0; CurveIndex < PoseContainer.Curves.Num(); ++CurveIndex)
+	{
+		CurveNames.Add(PoseContainer.Curves[CurveIndex].Name);
+	}
+
+	return CurveNames;
+}
+
+
+bool UPoseAsset::ContainsPose(const FName& InPoseName) const
+{
+	for (const auto& PoseName : PoseContainer.PoseNames)
+	{
+		if (PoseName.DisplayName == InPoseName)
+		{
+			return true;
+		}
+	}
+
+	return false;
+}
+
+void UPoseAsset::AddOrUpdatePose(const FSmartName& PoseName, USkeletalMeshComponent* MeshComponent)
+{
+	USkeleton* MySkeleton = GetSkeleton();
+	if (MySkeleton && MeshComponent && MeshComponent->SkeletalMesh)
+	{
+		TArray<FName> TrackNames;
+		// note this ignores root motion
+		TArray<FTransform> BoneTransform = MeshComponent->GetSpaceBases();
+		const FReferenceSkeleton& RefSkeleton = MeshComponent->SkeletalMesh->RefSkeleton;
+		for (int32 BoneIndex = 0; BoneIndex < RefSkeleton.GetNum(); ++BoneIndex)
+		{
+			TrackNames.Add(RefSkeleton.GetBoneName(BoneIndex));
+		}
+
+		// convert to local space
+		for (int32 BoneIndex = BoneTransform.Num() - 1; BoneIndex >= 0; --BoneIndex)
+		{
+			const int32 ParentIndex = RefSkeleton.GetParentIndex(BoneIndex);
+			if (ParentIndex != INDEX_NONE)
+			{
+				BoneTransform[BoneIndex] = BoneTransform[BoneIndex].GetRelativeTransform(BoneTransform[ParentIndex]);
+			}
+		}
+
+		AddOrUpdatePose(PoseName, TrackNames, BoneTransform);
+	}
+}
+
+void UPoseAsset::AddOrUpdatePose(const FSmartName& PoseName, TArray<FName> TrackNames, TArray<FTransform>& LocalTransform)
+{
+	USkeleton* MySkeleton = GetSkeleton();
+	if (MySkeleton)
+	{
+		// first combine track, we want to make sure all poses contains tracks with this
+		CombineTracks(TrackNames);
+
+		FPoseData* PoseData = PoseContainer.FindOrAddPoseData(PoseName);
+		// now copy all transform back to it. 
+		check(PoseData);
+		PoseContainer.FillUpDefaultPose(PoseData, MySkeleton, RetargetSource);
+
+		const FReferenceSkeleton& RefSkeleton = MySkeleton->GetReferenceSkeleton();
+		for (int32 Index = 0; Index < TrackNames.Num(); ++Index)
+		{
+			// now get poseData track index
+			const FName& TrackName = TrackNames[Index];
+			int32 SkeletonIndex = RefSkeleton.FindBoneIndex(TrackName);
+			int32 InternalTrackIndex = PoseContainer.Tracks.Find(TrackName);
+			PoseData->LocalSpacePose[InternalTrackIndex] = LocalTransform[Index];
+		}
+
+		NumPoses = PoseContainer.GetNumPoses();
+	}
+}
+
+void UPoseAsset::CombineTracks(const TArray<FName>& NewTracks)
+{
+	USkeleton* MySkeleton = GetSkeleton();
+	if (MySkeleton)
+	{
+		for (const auto& NewTrack : NewTracks)
+		{
+			if (PoseContainer.Tracks.Contains(NewTrack) == false)
+			{
+				// if we don't have it, then we'll have to add this track and then 
+				// right now it doesn't have to be in the hierarchy
+				// @todo: it is probably best to keep the hierarchy of the skeleton, so in the future, we might like to sort this by track after
+				PoseContainer.InsertTrack(NewTrack, MySkeleton, RetargetSource);
+			}
+		}
+	}
+}
+
+FName GetUniquePoseName(USkeleton* Skeleton)
+{
+	check(Skeleton);
+	int32 NameIndex = 0;
+
+	FSmartNameMapping::UID NewUID;
+	FName NewName;
+
+	do
+	{
+		NewName = FName(*FString::Printf(TEXT("Pose_%d"), NameIndex++));
+		NewUID = Skeleton->GetUIDByName(USkeleton::AnimCurveMappingName, NewName);
+	} while (NewUID != FSmartNameMapping::MaxUID);
+
+	// if found, 
+
+	return NewName;
+}
+
+#if WITH_EDITOR
+void UPoseAsset::CreatePoseFromAnimation(class UAnimSequence* AnimSequence)
+{
+	if (AnimSequence)
+	{
+		USkeleton* TargetSkeleton = AnimSequence->GetSkeleton();
+		if (TargetSkeleton)
+		{
+			SetSkeleton(TargetSkeleton);
+
+			PoseContainer.Reset();
+
+			NumPoses = AnimSequence->GetNumberOfFrames();
+
+			// make sure we have more than one pose
+			if (NumPoses > 0)
+			{
+				// stack allocator for extracting curve
+				FMemMark Mark(FMemStack::Get());
+
+				const int32 NumTracks = AnimSequence->AnimationTrackNames.Num();
+
+				// set up track data - @todo: add revaliation code when checked
+				for (int32 TrackIndex = 0; TrackIndex < NumTracks; ++TrackIndex)
+				{
+					const FName& TrackName = AnimSequence->AnimationTrackNames[TrackIndex];
+					PoseContainer.Tracks.Add(TrackName);
+				}
+
+				// now create pose transform
+				TArray<FTransform> NewPose;
+				
+				NewPose.Reset(NumTracks);
+				NewPose.AddUninitialized(NumTracks);
+
+				// @Todo fill up curve data
+				TArray<float> CurveData;
+				float IntervalBetweenKeys = AnimSequence->SequenceLength / NumPoses;
+
+				// add curves - only float curves
+				const int32 TotalFloatCurveCount = AnimSequence->RawCurveData.FloatCurves.Num();
+				
+				// have to construct own UIDList;
+				TArray<FSmartNameMapping::UID> UIDList;
+
+				if (TotalFloatCurveCount > 0)
+				{
+					for (const FFloatCurve& Curve : AnimSequence->RawCurveData.FloatCurves)
+					{
+						PoseContainer.Curves.Add(FAnimCurveBase(Curve.Name, Curve.GetCurveTypeFlags()));
+						UIDList.Add(Curve.Name.UID);
+					}
+				}
+
+				CurveData.AddZeroed(UIDList.Num());
+				// add to skeleton UID, so that it knows the curve data
+				for (int32 PoseIndex = 0; PoseIndex < NumPoses; ++PoseIndex)
+				{
+					FName PoseName = GetUniquePoseName(TargetSkeleton);
+					// now get rawanimationdata, and each key is converted to new pose
+					for (int32 TrackIndex = 0; TrackIndex < NumTracks; ++TrackIndex)
+					{
+						const auto& RawTrack = AnimSequence->RawAnimationData[TrackIndex];
+						AnimSequence->ExtractBoneTransform(RawTrack, NewPose[TrackIndex], PoseIndex);
+					}
+
+					// add to skeleton UID, so that it knows the curve data
+					FSmartName NewPoseName;
+					TargetSkeleton->AddSmartNameAndModify(USkeleton::AnimCurveMappingName, PoseName, NewPoseName);
+
+					if (TotalFloatCurveCount > 0)
+					{
+						// get curve data
+						// have to do iterate over time
+						// support curve
+						FBlendedCurve SourceCurve;
+						SourceCurve.InitFrom(&UIDList);
+						AnimSequence->EvaluateCurveData(SourceCurve, PoseIndex*IntervalBetweenKeys, true);
+
+						// copy back to CurveData
+						for (int32 CurveIndex = 0; CurveIndex < CurveData.Num(); ++CurveIndex)
+						{
+							CurveData[CurveIndex] = SourceCurve.Get(UIDList[CurveIndex]);
+						}
+
+						check(CurveData.Num() == PoseContainer.Curves.Num());
+					}
+				
+					// add new pose
+					PoseContainer.AddOrUpdatePose(NewPoseName, NewPose, CurveData);
+				}
+
+
+				PoseContainer.Shrink(TargetSkeleton, RetargetSource);
+				RecacheTrackmap();
+			}
+		}
+	}
+}
+
+#endif // WITH_EDITOR
+
+bool UPoseAsset::ModifyPoseName(FName OldPoseName, FName NewPoseName, const FSmartNameMapping::UID* NewUID)
+{
+	USkeleton* MySkeleton = GetSkeleton();
+
+	if (ContainsPose(NewPoseName))
+	{
+		// already exists, return 
+		return false;
+	}
+
+	FSmartName OldPoseSmartName;
+	ensureAlways(MySkeleton->GetSmartNameByName(USkeleton::AnimCurveMappingName, OldPoseName, OldPoseSmartName));
+
+	FPoseData* PoseData = PoseContainer.FindPoseData(OldPoseSmartName);
+	if (PoseData)
+	{
+		FSmartName NewPoseSmartName;
+		if (NewUID)
+		{
+			MySkeleton->GetSmartNameByUID(USkeleton::AnimCurveMappingName, *NewUID, NewPoseSmartName);
+		}
+		else
+		{
+			// we're renaming current one
+			MySkeleton->RenameSmartName(USkeleton::AnimCurveMappingName, OldPoseSmartName.UID, NewPoseName);
+			NewPoseSmartName = OldPoseSmartName;
+			NewPoseSmartName.DisplayName = NewPoseName;
+		}
+
+		PoseContainer.RenamePose(OldPoseSmartName, NewPoseSmartName);
+		OnPoseListChanged.Broadcast();
+
+		return true;
+	}
+
+	return false;
+}
+
+int32 UPoseAsset::DeletePoses(TArray<FName> PoseNamesToDelete)
+{
+	int32 ItemsDeleted = 0;
+
+	USkeleton* MySkeleton = GetSkeleton();
+
+	for (const auto& PoseName : PoseNamesToDelete)
+	{
+		FSmartName PoseSmartName;
+		if (MySkeleton->GetSmartNameByName(USkeleton::AnimCurveMappingName, PoseName, PoseSmartName) 
+			&& 	PoseContainer.DeletePose(PoseSmartName))
+		{
+			++ItemsDeleted;
+		}
+	}
+
+	PoseContainer.Shrink(GetSkeleton(), RetargetSource);
+
+	OnPoseListChanged.Broadcast();
+
+	return ItemsDeleted;
+}
+
+int32 UPoseAsset::DeleteCurves(TArray<FName> CurveNamesToDelete)
+{
+	int32 ItemsDeleted = 0;
+
+	USkeleton* MySkeleton = GetSkeleton();
+
+	for (const auto& CurveName : CurveNamesToDelete)
+	{
+		FSmartName CurveSmartName;
+		if (MySkeleton->GetSmartNameByName(USkeleton::AnimCurveMappingName, CurveName, CurveSmartName))
+		{
+			PoseContainer.DeleteCurve(CurveSmartName);
+			++ItemsDeleted;
+		}
+	}
+
+	OnPoseListChanged.Broadcast();
+
+	return ItemsDeleted;
+}
+bool UPoseAsset::ConvertToFullPose()
+{
+	if (ensureAlways(bAdditivePose))
+	{
+		TArray<FTransform>	BasePose;
+		TArray<float>		BaseCurves;
+		GetBasePoseTransform(BasePose, BaseCurves);
+
+		PoseContainer.ConvertToFullPose(BasePoseIndex, BasePose, BaseCurves);
+
+		bAdditivePose = false;
+
+		return true;
+	}
+
+	return false;
+}
+
+bool UPoseAsset::ConvertToAdditivePose(int32 NewBasePoseIndex)
+{
+	if (ensureAlways(!bAdditivePose))
+	{
+		// make sure it's valid
+		check(NewBasePoseIndex == -1 || PoseContainer.Poses.IsValidIndex(NewBasePoseIndex));
+
+		BasePoseIndex = NewBasePoseIndex;
+
+		TArray<FTransform> BasePose;
+		TArray<float>		BaseCurves;
+		GetBasePoseTransform(BasePose, BaseCurves);
+
+		PoseContainer.ConvertToAdditivePose(BasePoseIndex, BasePose, BaseCurves);
+
+		bAdditivePose = true;
+
+		return true;
+	}
+
+	return false;
+}
+
+bool UPoseAsset::ConvertSpace(bool bNewAdditivePose, int32 NewBasePoseInde)
+{
+	// first convert to full pose first
+	if (bAdditivePose)
+	{
+		// convert to full pose
+		if (!ConvertToFullPose())
+		{
+			// issue with converting full pose
+			return false;
+		}
+	}
+
+	// now we have full pose
+	if (bNewAdditivePose)
+	{
+		ConvertToAdditivePose(NewBasePoseInde);
+	}
+
+	return true;
+}
+
+const int32 UPoseAsset::GetPoseIndexByName(const FName& InBasePoseName) const
+{
+	for (int32 PoseIndex = 0; PoseIndex < PoseContainer.PoseNames.Num(); ++PoseIndex)
+	{
+		if (PoseContainer.PoseNames[PoseIndex].DisplayName == InBasePoseName)
+		{
+			return PoseIndex;
+		}
+	}
+
+	return INDEX_NONE;
+}
+
+void UPoseAsset::RecacheTrackmap()
+{
+	USkeleton* MySkeleton = GetSkeleton();
+	PoseContainer.TrackMap.Reset();
+
+	if (MySkeleton)
+	{
+		const FReferenceSkeleton& RefSkeleton = MySkeleton->GetReferenceSkeleton();
+
+		// set up track data - @todo: add revaliation code when checked
+		for (int32 TrackIndex = 0; TrackIndex < PoseContainer.Tracks.Num(); ++TrackIndex)
+		{
+			const FName& TrackName = PoseContainer.Tracks[TrackIndex];
+			const int32 SkeletonTrackIndex = RefSkeleton.FindBoneIndex(TrackName);
+			PoseContainer.TrackMap.Add(TrackName, SkeletonTrackIndex);
+		}
+	}
+}
+
+#if WITH_EDITOR
+void UPoseAsset::RemapTracksToNewSkeleton(USkeleton* NewSkeleton, bool bConvertSpaces)
+{
+	Super::RemapTracksToNewSkeleton(NewSkeleton, bConvertSpaces);
+	// @todo: should allow deleted tracks and so on
+	RecacheTrackmap();
+}
+#endif // WITH_EDITOR
+
+bool UPoseAsset::GetBasePoseTransform(TArray<FTransform>& OutBasePose, TArray<float>& OutCurve) const
+{
+	int32 TotalNumTrack = PoseContainer.Tracks.Num();
+	OutBasePose.Reset(TotalNumTrack);
+
+	if (BasePoseIndex == -1)
+	{
+		OutBasePose.AddUninitialized(TotalNumTrack);
+		for (int32 TrackIndex = 0; TrackIndex < PoseContainer.Tracks.Num(); ++TrackIndex)
+		{
+			const FName& TrackName = PoseContainer.Tracks[TrackIndex];
+			OutBasePose[TrackIndex] = PoseContainer.GetDefaultTransform(TrackName, GetSkeleton(), RetargetSource);
+		}
+
+		// add zero curves
+		OutCurve.AddZeroed(PoseContainer.Curves.Num());
+		check(OutBasePose.Num() == TotalNumTrack);
+		return true;
+	}
+	else if (PoseContainer.Poses.IsValidIndex(BasePoseIndex))
+	{
+		OutBasePose = PoseContainer.Poses[BasePoseIndex].LocalSpacePose;
+		OutCurve = PoseContainer.Poses[BasePoseIndex].CurveData;
+		check(OutBasePose.Num() == TotalNumTrack);
+		return true;
+	}
+
+	return false;
+}
\ No newline at end of file
