Copying //UE4/Dev-Framework to Dev-Main (//UE4/Dev-Main)

diff --git a/Engine/Source/Runtime/Engine/Private/Animation/Skeleton.cpp b/Engine/Source/Runtime/Engine/Private/Animation/Skeleton.cpp
--- a/Engine/Source/Runtime/Engine/Private/Animation/Skeleton.cpp
+++ b/Engine/Source/Runtime/Engine/Private/Animation/Skeleton.cpp
@@ -1065,13 +1065,14 @@ void USkeleton::RenameSlotName(const FName& OldName, const FName& NewName)
 }
 
 
-bool USkeleton::AddSmartNameAndModify(FName ContainerName, FName NewName, FSmartNameMapping::UID& NewUid)
+
+bool USkeleton::AddSmartNameAndModify(FName ContainerName, FName NewDisplayName, FSmartName& NewName)
 {
 	bool Successful = false;
-	FSmartNameMapping* RequestedMapping = SmartNames.GetContainerInternal(ContainerName);
+	FSmartNameMapping* RequestedMapping = GetOrAddSmartNameContainer(ContainerName);
 	if (RequestedMapping)
 	{
-		if (RequestedMapping->AddOrFindName(NewName, NewUid))
+		if (RequestedMapping->FindOrAddSmartName(NewDisplayName, NewName))
 		{
 			Modify(true);
 			Successful = true;
@@ -1080,7 +1081,6 @@ bool USkeleton::AddSmartNameAndModify(FName ContainerName, FName NewName, FSmart
 	}
 	return Successful;
 }
-
 bool USkeleton::RenameSmartnameAndModify(FName ContainerName, FSmartNameMapping::UID Uid, FName NewName)
 {
 	bool Successful = false;
@@ -1104,7 +1104,11 @@ void USkeleton::RemoveSmartnameAndModify(FName ContainerName, FSmartNameMapping:
 	FSmartNameMapping* RequestedMapping = SmartNames.GetContainerInternal(ContainerName);
 	if (RequestedMapping)
 	{
-		
+		Modify();
+		if (RequestedMapping->Remove(Uid))
+		{
+			CacheAnimCurveMappingNameUids();
+		}
 	}
 }
 
@@ -1123,28 +1127,150 @@ void USkeleton::RemoveSmartnamesAndModify(FName ContainerName, const TArray<FSma
 			}
 		}
 
-		if (bModified)
+if (bModified)
+{
+	Modify();
+	CacheAnimCurveMappingNameUids();
+}
+	}
+}
+
+bool USkeleton::GetSmartNameByUID(FName ContainerName, FSmartNameMapping::UID UID, FSmartName& OutSmartName)
+{
+	const FSmartNameMapping* RequestedMapping = SmartNames.GetContainerInternal(ContainerName);
+	if (RequestedMapping)
+	{
+		return (RequestedMapping->FindSmartNameByUID(UID, OutSmartName));
+	}
+
+	return false;
+}
+
+bool USkeleton::GetSmartNameByName(FName ContainerName, FName InName, FSmartName& OutSmartName)
+{
+	const FSmartNameMapping* RequestedMapping = SmartNames.GetContainerInternal(ContainerName);
+	if (RequestedMapping)
+	{
+		return (RequestedMapping->FindSmartName(InName, OutSmartName));
+	}
+
+	return false;
+}
+
+FSmartNameMapping::UID USkeleton::GetUIDByName(FName ContainerName, FName Name)
+{
+	const FSmartNameMapping* RequestedMapping = SmartNames.GetContainerInternal(ContainerName);
+	if (RequestedMapping)
+	{
+		if (const FSmartNameMapping::UID* UID = RequestedMapping->FindUID(Name))
 		{
-			Modify();
-			CacheAnimCurveMappingNameUids();
-		}		
+			if (UID)
+			{
+				return *UID;
+			}
+		}
+	}
+
+	return FSmartNameMapping::MaxUID;
+}
+
+// this does very simple thing.
+// @todo: this for now prioritize FName because that is main issue right now
+// @todo: @fixme: this has to be fixed when we have GUID
+void USkeleton::VerifySmartName(const FName& ContainerName, FSmartName& InOutSmartName)
+{
+	VerifySmartNameInternal(ContainerName, InOutSmartName);
+	if (ContainerName == USkeleton::AnimCurveMappingName)
+	{
+		CacheAnimCurveMappingNameUids();
+	}
+}
+
+
+bool USkeleton::FillSmartNameByDisplayName(FSmartNameMapping* Mapping, const FName& DisplayName, FSmartName& OutSmartName)
+{
+	FSmartName SkeletonName;
+	if (Mapping->FindSmartName(DisplayName, SkeletonName))
+	{
+		OutSmartName.DisplayName = DisplayName;
+
+		// if same guid, this is same
+		if (SkeletonName.Guid == OutSmartName.Guid)
+		{
+			OutSmartName.UID = SkeletonName.UID;
+			return true;
+		}
+		// else, take Skeleton Guid, we don't allow same name with different Guid
+		else
+		{
+			OutSmartName.Guid = SkeletonName.Guid;
+			OutSmartName.UID = SkeletonName.UID;
+			return true;
+		}
+	}
+
+	return false;
+}
+//@todo: this does prioritize name over UID since UID doesn't work well
+// until we have GUID, we'll prioritize name over UID
+bool USkeleton::VerifySmartNameInternal(const FName&  ContainerName, FSmartName& InOutSmartName)
+{
+	FSmartNameMapping* Mapping = GetOrAddSmartNameContainer(ContainerName);
+	if (Mapping != nullptr)
+	{
+		// make a copy just in case we change by accident
+		FName DisplayName = InOutSmartName.DisplayName;
+
+		// if I find the name, fill up the data
+		if (FillSmartNameByDisplayName(Mapping, DisplayName, InOutSmartName) == false)
+		{
+			// look for same guid, the name might have been changed
+			if (Mapping->GetNameByGuid(InOutSmartName.Guid, DisplayName))
+			{
+				ensureAlways(FillSmartNameByDisplayName(Mapping, DisplayName, InOutSmartName));
+			}
+			else
+			{
+				// this is only case where we add new one
+				Mapping->FindOrAddSmartName(InOutSmartName.DisplayName, InOutSmartName);
+				return true;
+			}
+		}
 	}
+
+	return false;
 }
 
-const FSmartNameMapping* USkeleton::GetOrAddSmartNameContainer(FName ContainerName)
+void USkeleton::VerifySmartNames(const FName&  ContainerName, TArray<FSmartName>& InOutSmartNames)
 {
-	const FSmartNameMapping* Mapping = SmartNames.GetContainerInternal(ContainerName);
+	bool bRefreshCache = false;
+
+	for(auto& SmartName: InOutSmartNames)
+	{
+		VerifySmartNameInternal(ContainerName, SmartName);
+	}
+
+	if (ContainerName == USkeleton::AnimCurveMappingName)
+	{
+		CacheAnimCurveMappingNameUids();
+	}
+}
+
+FSmartNameMapping* USkeleton::GetOrAddSmartNameContainer(const FName& ContainerName)
+{
+	FSmartNameMapping* Mapping = SmartNames.GetContainerInternal(ContainerName);
 	if (Mapping == nullptr)
 	{
 		Modify();
+		CacheAnimCurveMappingNameUids();
 		SmartNames.AddContainer(ContainerName);
-		Mapping = SmartNames.GetContainer(ContainerName);		
+		Mapping = SmartNames.GetContainerInternal(ContainerName);		
 	}
 
 	return Mapping;
 }
 
-const FSmartNameMapping* USkeleton::GetSmartNameContainer(FName ContainerName) const
+const FSmartNameMapping* USkeleton::GetSmartNameContainer(const FName& ContainerName) const
 {
 	return SmartNames.GetContainer(ContainerName);
 }
@@ -1155,9 +1281,23 @@ void USkeleton::RegenerateGuid()
 	check(Guid.IsValid());
 }
 
+bool USkeleton::RenameSmartName(FName ContainerName, const FSmartNameMapping::UID& Uid, FName NewName)
+{
+	FSmartNameMapping* Mapping = SmartNames.GetContainerInternal(ContainerName);
+	if (Mapping->Exists(Uid))
+	{
+		Modify();
+		CacheAnimCurveMappingNameUids();
+		Mapping->Rename(Uid, NewName);
+		return true;
+	}
+
+	return false;
+}
+
 void USkeleton::CacheAnimCurveMappingNameUids()
 {
-	const FSmartNameMapping* Mapping = GetSmartNameContainer(USkeleton::AnimCurveMappingName);
+	const FSmartNameMapping* Mapping = SmartNames.GetContainerInternal(USkeleton::AnimCurveMappingName);
 	if (Mapping != nullptr)
 	{
 		Mapping->FillUidArray(CachedAnimCurveMappingNameUids);
diff --git a/Engine/Source/Runtime/Engine/Private/Animation/SmartName.cpp b/Engine/Source/Runtime/Engine/Private/Animation/SmartName.cpp
--- a/Engine/Source/Runtime/Engine/Private/Animation/SmartName.cpp
+++ b/Engine/Source/Runtime/Engine/Private/Animation/SmartName.cpp
@@ -2,46 +2,45 @@
 
 #include "EnginePrivate.h"
 #include "Animation/SmartName.h"
-
+#include "FrameworkObjectVersion.h"
+////////////////////////////////////////////////////////////////////////
+//
+// FSmartNameMapping
+//
+///////////////////////////////////////////////////////////////////////
 FSmartNameMapping::FSmartNameMapping()
 : NextUid(0)
 {}
 
-bool FSmartNameMapping::AddOrFindName(FName Name, UID& OutUid)
+bool FSmartNameMapping::AddOrFindName(FName Name, UID& OutUid, FGuid& OutGuid)
 {
 	check(Name.IsValid());
 
 	// Check for UID overflow
-	if(NextUid == MaxUID && FreeList.Num() == 0)
+	const UID* ExistingUid = UidMap.FindKey(Name);
+	const FGuid* ExistingGuid = GuidMap.Find(Name);
+
+	// make sure they both exists and same 
+	check(!!ExistingUid == !!ExistingGuid);
+
+	if(ExistingUid)
 	{
-		// No UIDs left
-		UE_LOG(LogAnimation, Error, TEXT("No more UIDs left in skeleton smartname container, consider changing UID type to longer type."));
+		// Already present in the list
+		OutUid = *ExistingUid;
+		OutGuid = *ExistingGuid;
 		return false;
 	}
-	else
-	{
-		const UID* ExistingUid = UidMap.FindKey(Name);
 
-		if(ExistingUid)
-		{
-			// Already present in the list
-			OutUid = *ExistingUid;
-			return false;
-		}
+	// make sure we didn't reach till end
+	check(NextUid != MaxUID);
 
-		if(FreeList.Num() > 0)
-		{
-			// We've got some UIDs that are unused < NextUid. Use them first
-			OutUid = FreeList.Pop();
-		}
-		else
-		{
-			OutUid = NextUid++;
-		}
-		UidMap.Add(OutUid, Name);
+	OutUid = NextUid;
+	OutGuid = FGuid::NewGuid();
+	UidMap.Add(OutUid, Name);
+	GuidMap.Add(Name, OutGuid);
 
-		return true;
-	}
+	++NextUid;
+	return true;
 }
 
 bool FSmartNameMapping::GetName(const UID& Uid, FName& OutName) const
@@ -55,11 +54,32 @@ bool FSmartNameMapping::GetName(const UID& Uid, FName& OutName) const
 	return false;
 }
 
+bool FSmartNameMapping::GetNameByGuid(const FGuid& Guid, FName& OutName) const
+{
+	const FName* FoundName = GuidMap.FindKey(Guid);
+	if (FoundName)
+	{
+		OutName = *FoundName;
+		return true;
+	}
+
+	return false;
+}
+
 bool FSmartNameMapping::Rename(const UID& Uid, FName NewName)
 {
 	FName* ExistingName = UidMap.Find(Uid);
 	if(ExistingName)
 	{
+		FGuid* Guid = GuidMap.Find(*ExistingName);
+		check(Guid);
+
+		// re add new value
+		GuidMap.Remove(*ExistingName);
+		GuidMap.Add(NewName, *Guid);
+
+		check(GuidMap.Num() == UidMap.Num());
+
 		*ExistingName = NewName;
 		return true;
 	}
@@ -68,59 +88,59 @@ bool FSmartNameMapping::Rename(const UID& Uid, FName NewName)
 
 bool FSmartNameMapping::Remove(const UID& Uid)
 {
-	bool bRemoved = UidMap.Remove(Uid) != 0;
-
-	if(bRemoved)
+	FName* ExistingName = UidMap.Find(Uid);
+	if (ExistingName)
 	{
-		// The provided UID has been cleared. If it is at the end of the mapping decrement next id,
-		// otherwise add to the freelist
-		if(Uid == NextUid - 1)
-		{
-			--NextUid;
-		}
-		else
-		{
-			FreeList.Add(Uid);
-		}
+		FGuid* Guid = GuidMap.Find(*ExistingName);
+		check(Guid);
+
+		// re add new value
+		GuidMap.Remove(*ExistingName);
+		UidMap.Remove(Uid);
+
+		check(GuidMap.Num() == UidMap.Num());
+
+		return true;
 	}
 
-	return bRemoved;
+	return false;
 }
 
 void FSmartNameMapping::Serialize(FArchive& Ar)
 {
-	if(Ar.UE4Ver() >= VER_UE4_SKELETON_ADD_SMARTNAMES)
+	Ar.UsingCustomVersion(FFrameworkObjectVersion::GUID);
+	if (Ar.CustomVer(FFrameworkObjectVersion::GUID) >= FFrameworkObjectVersion::SmartNameRefactor)
 	{
-		Ar << NextUid;
-		Ar << UidMap;
+		Ar << GuidMap;
 
-		if(Ar.ArIsLoading)
+		if (Ar.ArIsLoading)
 		{
-			// Sort the Uid Map
-			UidMap.KeySort([](const UID& A, const UID& B)
+			NextUid = 0;
+			//fill up uidmap
+			UidMap.Empty(GuidMap.Num());
+			for (const TPair<FName, FGuid>& GuidPair : GuidMap)
 			{
-				return A < B;
-			});
-
-			// Build freelist as necessary
-			UID ExpectedID = 0;
-			UID Last = 0;
-			for(TPair<UID, FName>& Pair : UidMap)
-			{
-				UID CurrentUid = Pair.Key;
-				while(CurrentUid != ExpectedID)
-				{
-					FreeList.Add(ExpectedID++);
-				}
-
-				Last = CurrentUid;
-				++ExpectedID;
+				UidMap.Add(NextUid, GuidPair.Key);
+				++NextUid;
 			}
+		}
+	}
+	else if(Ar.UE4Ver() >= VER_UE4_SKELETON_ADD_SMARTNAMES)
+	{
+		Ar << NextUid;
+		Ar << UidMap;
 
-			// In the case of multipled queued deletions previously,
-			// fix up the next UID to match the last used key we found.
-			NextUid = Last + 1;
+		TMap<UID, FName> NewUidMap;
+		// convert to GuidMap
+		NextUid = 0;
+		GuidMap.Empty(UidMap.Num());
+		for (TPair<UID, FName>& UidPair : UidMap)
+		{
+			GuidMap.Add(UidPair.Value, FGuid::NewGuid());
+			NewUidMap.Add(NextUid++, UidPair.Value);
 		}
+
+		UidMap = NewUidMap;
 	}
 }
 
@@ -161,6 +181,47 @@ FArchive& operator<<(FArchive& Ar, FSmartNameMapping& Elem)
 	return Ar;
 }
 
+bool FSmartNameMapping::FindOrAddSmartName(FName Name, FSmartName& OutName)
+{
+	FSmartNameMapping::UID NewUID;
+	FGuid NewGuid;
+	bool bNewlyAdded = AddOrFindName(Name, NewUID, NewGuid);
+	OutName = FSmartName(Name, NewUID, NewGuid);
+
+	return bNewlyAdded;
+}
+
+bool FSmartNameMapping::FindSmartName(FName Name, FSmartName& OutName) const
+{
+	const FSmartNameMapping::UID* ExistingUID = FindUID(Name);
+	if (ExistingUID)
+	{
+		const FGuid* ExistingGuid = GuidMap.Find(Name);
+		OutName = FSmartName(Name, *ExistingUID, *ExistingGuid);
+		return true;
+	}
+
+	return false;
+}
+
+bool FSmartNameMapping::FindSmartNameByUID(FSmartNameMapping::UID UID, FSmartName& OutName) const
+{
+	FName ExistingName;
+	if (GetName(UID, ExistingName))
+	{
+		OutName.DisplayName = ExistingName;
+		OutName.Guid = *GuidMap.Find(ExistingName);
+		OutName.UID = UID;
+		return true;
+	}
+
+	return false;
+}
+////////////////////////////////////////////////////////////////////////
+//
+// FSmartNameContainer
+//
+///////////////////////////////////////////////////////////////////////
 void FSmartNameContainer::AddContainer(FName NewContainerName)
 {
 	if(!NameMappings.Find(NewContainerName))
@@ -179,7 +240,21 @@ void FSmartNameContainer::Serialize(FArchive& Ar)
 	Ar << NameMappings;
 }
 
-FSmartNameMapping* FSmartNameContainer::GetContainerInternal(FName ContainerName)
+FSmartNameMapping* FSmartNameContainer::GetContainerInternal(const FName& ContainerName)
 {
 	return NameMappings.Find(ContainerName);
+}
+
+////////////////////////////////////////////////////////////////////////
+//
+// FSmartName
+//
+///////////////////////////////////////////////////////////////////////
+bool FSmartName::Serialize(FArchive& Ar)
+{
+	Ar << DisplayName;
+	Ar << UID;
+	Ar << Guid;
+
+	return true;
 }
\ No newline at end of file
diff --git a/Engine/Source/Runtime/Engine/Private/AnimationAsset.cpp b/Engine/Source/Runtime/Engine/Private/AnimationAsset.cpp
--- a/Engine/Source/Runtime/Engine/Private/AnimationAsset.cpp
+++ b/Engine/Source/Runtime/Engine/Private/AnimationAsset.cpp
@@ -207,6 +207,11 @@ void UAnimationAsset::TickAssetPlayerInstance(FAnimTickRecord& Instance, class U
 }
 
 #if WITH_EDITOR
+void UAnimationAsset::RemapTracksToNewSkeleton(USkeleton* NewSkeleton, bool bConvertSpaces)
+{
+	SetSkeleton(NewSkeleton);
+}
+
 bool UAnimationAsset::ReplaceSkeleton(USkeleton* NewSkeleton, bool bConvertSpaces/*=false*/)
 {
 	// if it's not same 
@@ -223,16 +228,10 @@ bool UAnimationAsset::ReplaceSkeleton(USkeleton* NewSkeleton, bool bConvertSpace
 		{
 			for (auto Iter = AnimAssetsToReplace.CreateIterator(); Iter; ++Iter)
 			{
-				UAnimSequence* AnimSeq = Cast<UAnimSequence>(*Iter);
-				if (AnimSeq && AnimSeq->Skeleton != NewSkeleton)
-				{
-					AnimSeq->RemapTracksToNewSkeleton(NewSkeleton, bConvertSpaces);
-				}
-				else
-				{
-					// if not anim sequence, at least replace skeleton
-					(*Iter)->SetSkeleton(NewSkeleton);
-				}
+				UAnimationAsset* IterAsset = (*Iter);
+				// these two are different functions for now
+				// technically if you have implementation for Remap, it will also set skeleton 
+				IterAsset->RemapTracksToNewSkeleton(NewSkeleton, bConvertSpaces);
 			}
 		}
 
diff --git a/Engine/Source/Runtime/Engine/Private/AudioDevice.cpp b/Engine/Source/Runtime/Engine/Private/AudioDevice.cpp
--- a/Engine/Source/Runtime/Engine/Private/AudioDevice.cpp
+++ b/Engine/Source/Runtime/Engine/Private/AudioDevice.cpp
@@ -2958,7 +2958,11 @@ UAudioComponent* FAudioDevice::CreateComponent(USoundBase* Sound, UWorld* World,
 		{
 			AudioDevice = GEngine->GetMainAudioDevice();
 		}
-		check(AudioDevice != nullptr);
+		
+		if (AudioDevice == nullptr)
+		{
+			return nullptr;
+		}
 
 		// Avoid creating component if we're trying to play a sound on an already destroyed actor.
 		if (Actor && Actor->IsPendingKill())
diff --git a/Engine/Source/Runtime/Engine/Private/Character.cpp b/Engine/Source/Runtime/Engine/Private/Character.cpp
--- a/Engine/Source/Runtime/Engine/Private/Character.cpp
+++ b/Engine/Source/Runtime/Engine/Private/Character.cpp
@@ -511,9 +511,12 @@ namespace MovementBaseUtility
 			}
 
 			// Fall back to physics velocity.
-			if (BaseVelocity.IsZero() && MovementBase->GetBodyInstance())
+			if (BaseVelocity.IsZero())
 			{
-				BaseVelocity = MovementBase->GetBodyInstance()->GetUnrealWorldVelocity();
+				if (FBodyInstance* BaseBodyInstance = MovementBase->GetBodyInstance())
+				{
+					BaseVelocity = BaseBodyInstance->GetUnrealWorldVelocity();
+				}
 			}
 		}
 		
diff --git a/Engine/Source/Runtime/Engine/Private/Collision/CollisionConversions.cpp b/Engine/Source/Runtime/Engine/Private/Collision/CollisionConversions.cpp
--- a/Engine/Source/Runtime/Engine/Private/Collision/CollisionConversions.cpp
+++ b/Engine/Source/Runtime/Engine/Private/Collision/CollisionConversions.cpp
@@ -41,10 +41,10 @@ extern TAutoConsoleVariable<int32> CVarShowInitialOverlaps;
 // Sentinel for invalid query results.
 static const PxQueryHit InvalidQueryHit;
 
-static bool IsInvalidFaceIndex(PxU32 faceIndex)
+FORCEINLINE_DEBUGGABLE bool IsInvalidFaceIndex(PxU32 faceIndex)
 {
-	checkf(InvalidQueryHit.faceIndex == 0xFFFFffff, TEXT("Engine code needs fixing: PhysX invalid face index sentinel has changed or is not part of default PxQueryHit!"));
-	return (faceIndex == InvalidQueryHit.faceIndex);
+	checkfSlow(InvalidQueryHit.faceIndex == 0xFFFFffff, TEXT("Engine code needs fixing: PhysX invalid face index sentinel has changed or is not part of default PxQueryHit!"));
+	return (faceIndex == 0xFFFFffff);
 }
 
 // Forward declare, I don't want to move the entire function right now or we lose change history.
diff --git a/Engine/Source/Runtime/Engine/Private/Components/ActorComponent.cpp b/Engine/Source/Runtime/Engine/Private/Components/ActorComponent.cpp
--- a/Engine/Source/Runtime/Engine/Private/Components/ActorComponent.cpp
+++ b/Engine/Source/Runtime/Engine/Private/Components/ActorComponent.cpp
@@ -319,7 +319,7 @@ bool UActorComponent::ComponentHasTag(FName Tag) const
 }
 
 
-ENetMode UActorComponent::GetNetMode() const
+ENetMode UActorComponent::InternalGetNetMode() const
 {
 	AActor* MyOwner = GetOwner();
 	return MyOwner ? MyOwner->GetNetMode() : NM_Standalone;
@@ -1516,11 +1516,6 @@ ENetRole UActorComponent::GetOwnerRole() const
 	return (MyOwner ? MyOwner->Role.GetValue() : ROLE_None);
 }
 
-bool UActorComponent::IsNetSimulating() const
-{
-	return GetIsReplicated() && GetOwnerRole() != ROLE_Authority;
-}
-
 void UActorComponent::GetLifetimeReplicatedProps( TArray< FLifetimeProperty > & OutLifetimeProps ) const
 {
 	UBlueprintGeneratedClass* BPClass = Cast<UBlueprintGeneratedClass>(GetClass());
diff --git a/Engine/Source/Runtime/Engine/Private/Components/CharacterMovementComponent.cpp b/Engine/Source/Runtime/Engine/Private/Components/CharacterMovementComponent.cpp
--- a/Engine/Source/Runtime/Engine/Private/Components/CharacterMovementComponent.cpp
+++ b/Engine/Source/Runtime/Engine/Private/Components/CharacterMovementComponent.cpp
@@ -17,6 +17,7 @@
 
 // @todo this is here only due to circular dependency to AIModule. To be removed
 #include "Navigation/PathFollowingComponent.h"
+#include "AI/Navigation/RecastNavMesh.h"
 #include "AI/Navigation/AvoidanceManager.h"
 #include "Components/CapsuleComponent.h"
 #include "Components/BrushComponent.h"
@@ -25,6 +26,7 @@
 #include "Engine/DemoNetDriver.h"
 #include "Engine/NetworkObjectList.h"
 
+#include "HAL/IConsoleManager.h"
 #include "PerfCountersHelpers.h"
 
 DEFINE_LOG_CATEGORY_STATIC(LogCharacterMovement, Log, All);
@@ -86,11 +88,11 @@ namespace CharacterMovementComponentStatics
 }
 
 // CVars
-namespace CharacterCVars
+namespace CharacterMovementCVars
 {
 	// Listen server smoothing
 	static int32 NetEnableListenServerSmoothing = 1;
-	static FAutoConsoleVariableRef CVarNetEnableListenServerSmoothing(
+	FAutoConsoleVariableRef CVarNetEnableListenServerSmoothing(
 		TEXT("p.NetEnableListenServerSmoothing"),
 		NetEnableListenServerSmoothing,
 		TEXT("Whether to enable move combining on the client to reduce bandwidth by combining similar moves.\n")
@@ -104,88 +106,105 @@ namespace CharacterCVars
 	static float StuckWarningPeriod = 1.f;
 #endif
 
-	static FAutoConsoleVariableRef CVarStuckWarningPeriod(
+	FAutoConsoleVariableRef CVarStuckWarningPeriod(
 		TEXT("p.CharacterStuckWarningPeriod"),
 		StuckWarningPeriod,
 		TEXT("How often (in seconds) we are allowed to log a message about being stuck in geometry.\n")
 		TEXT("<0: Disable, >=0: Enable and log this often, in seconds."),
 		ECVF_Default);
-}
 
+	static int32 NetEnableMoveCombining = 1;
+	FAutoConsoleVariableRef CVarNetEnableMoveCombining(
+		TEXT("p.NetEnableMoveCombining"),
+		NetEnableMoveCombining,
+		TEXT("Whether to enable move combining on the client to reduce bandwidth by combining similar moves.\n")
+		TEXT("0: Disable, 1: Enable"),
+		ECVF_Default);
 
-static TAutoConsoleVariable<int32> CVarNetEnableMoveCombining(
-	TEXT("p.NetEnableMoveCombining"),
-	1,
-	TEXT("Whether to enable move combining on the client to reduce bandwidth by combining similar moves.\n")
-	TEXT("0: Disable, 1: Enable"),
-	ECVF_Default);
+	static float NetProxyShrinkRadius = 0.01f;
+	FAutoConsoleVariableRef CVarNetProxyShrinkRadius(
+		TEXT("p.NetProxyShrinkRadius"),
+		NetProxyShrinkRadius,
+		TEXT("Shrink simulated proxy capsule radius by this amount, to account for network rounding that may cause encroachment.\n")
+		TEXT("Changing this value at runtime may require the proxy to re-join for correct behavior.\n")
+		TEXT("<= 0: disabled, > 0: shrink by this amount."),
+		ECVF_Default);
 
-static TAutoConsoleVariable<float> CVarNetProxyShrinkRadius(
-	TEXT("p.NetProxyShrinkRadius"),
-	0.01f,
-	TEXT("Shrink simulated proxy capsule radius by this amount, to account for network rounding that may cause encroachment.\n")
-	TEXT("Changing this value at runtime may require the proxy to re-join for correct behavior.\n")
-	TEXT("<= 0: disabled, > 0: shrink by this amount."),
-	ECVF_Default);
+	static float NetProxyShrinkHalfHeight = 0.01f;
+	FAutoConsoleVariableRef CVarNetProxyShrinkHalfHeight(
+		TEXT("p.NetProxyShrinkHalfHeight"),
+		NetProxyShrinkHalfHeight,
+		TEXT("Shrink simulated proxy capsule half height by this amount, to account for network rounding that may cause encroachment.\n")
+		TEXT("Changing this value at runtime may require the proxy to re-join for correct behavior.\n")
+		TEXT("<= 0: disabled, > 0: shrink by this amount."),
+		ECVF_Default);
 
-static TAutoConsoleVariable<float> CVarNetProxyShrinkHalfHeight(
-	TEXT("p.NetProxyShrinkHalfHeight"),
-	0.01f,
-	TEXT("Shrink simulated proxy capsule half height by this amount, to account for network rounding that may cause encroachment.\n")
-	TEXT("Changing this value at runtime may require the proxy to re-join for correct behavior.\n")
-	TEXT("<= 0: disabled, > 0: shrink by this amount."),
-	ECVF_Default);
+	static int32 ReplayUseInterpolation = 1;
+	FAutoConsoleVariableRef CVarReplayUseInterpolation(
+		TEXT( "p.ReplayUseInterpolation" ),
+		ReplayUseInterpolation,
+		TEXT( "" ),
+		ECVF_Default);
 
 #if !UE_BUILD_SHIPPING
-static TAutoConsoleVariable<int32> CVarNetShowCorrections(
-	TEXT("p.NetShowCorrections"),
-	0,
-	TEXT("Whether to draw client position corrections (red is incorrect, green is corrected).\n")
-	TEXT("0: Disable, 1: Enable"),
-	ECVF_Cheat);
-
-static TAutoConsoleVariable<float> CVarNetCorrectionLifetime(
-	TEXT("p.NetCorrectionLifetime"),
-	4.f,
-	TEXT("How long a visualized network correction persists.\n")
-	TEXT("Time in seconds each visualized network correction persists."),
-	ECVF_Cheat);
+
+	static int32 NetShowCorrections = 0;
+	FAutoConsoleVariableRef CVarNetShowCorrections(
+		TEXT("p.NetShowCorrections"),
+		NetShowCorrections,
+		TEXT("Whether to draw client position corrections (red is incorrect, green is corrected).\n")
+		TEXT("0: Disable, 1: Enable"),
+		ECVF_Cheat);
+
+	static float NetCorrectionLifetime = 4.f;
+	FAutoConsoleVariableRef CVarNetCorrectionLifetime(
+		TEXT("p.NetCorrectionLifetime"),
+		NetCorrectionLifetime,
+		TEXT("How long a visualized network correction persists.\n")
+		TEXT("Time in seconds each visualized network correction persists."),
+		ECVF_Cheat);
+
 #endif // !UI_BUILD_SHIPPING
 
+
 #if !(UE_BUILD_SHIPPING || UE_BUILD_TEST)
-static TAutoConsoleVariable<float> CVarNetForceClientAdjustmentPercent(
-	TEXT("p.NetForceClientAdjustmentPercent"),
-	0.f,
-	TEXT("Percent of ServerCheckClientError checks to return true regardless of actual error.\n")
-	TEXT("Useful for testing client correction code.\n")
-	TEXT("<=0: Disable, 0.05: 5% of checks will return failed, 1.0: Always send client adjustments"),
-	ECVF_Cheat);
-
-static TAutoConsoleVariable<int32> CVarVisualizeMovement(
-	TEXT("p.VisualizeMovement"),
-	0,
-	TEXT("Whether to draw in-world debug information for character movement.\n")
-	TEXT("0: Disable, 1: Enable"),
-	ECVF_Cheat);
-
-static TAutoConsoleVariable<int32> CVarNetVisualizeSimulatedCorrections(
-	TEXT( "p.NetVisualizeSimulatedCorrections" ),
-	0,
-	TEXT( "" )
-	TEXT( "0: Disable, 1: Enable" ),
-	ECVF_Cheat );
-
-static TAutoConsoleVariable<int32> CVarDebugTimeDiscrepancy(
-	TEXT( "p.DebugTimeDiscrepancy" ),
-	0,
-	TEXT( "Whether to log detailed Movement Time Discrepancy values for testing" )
-	TEXT( "0: Disable, 1: Enable Detection logging, 2: Enable Detection and Resolution logging" ),
-	ECVF_Cheat );
 
+	static float NetForceClientAdjustmentPercent = 0.f;
+	FAutoConsoleVariableRef CVarNetForceClientAdjustmentPercent(
+		TEXT("p.NetForceClientAdjustmentPercent"),
+		NetForceClientAdjustmentPercent,
+		TEXT("Percent of ServerCheckClientError checks to return true regardless of actual error.\n")
+		TEXT("Useful for testing client correction code.\n")
+		TEXT("<=0: Disable, 0.05: 5% of checks will return failed, 1.0: Always send client adjustments"),
+		ECVF_Cheat);
+
+	static int32 VisualizeMovement = 0;
+	FAutoConsoleVariableRef CVarVisualizeMovement(
+		TEXT("p.VisualizeMovement"),
+		VisualizeMovement,
+		TEXT("Whether to draw in-world debug information for character movement.\n")
+		TEXT("0: Disable, 1: Enable"),
+		ECVF_Cheat);
+
+	static int32 NetVisualizeSimulatedCorrections = 0;
+	FAutoConsoleVariableRef CVarNetVisualizeSimulatedCorrections(
+		TEXT("p.NetVisualizeSimulatedCorrections"),
+		NetVisualizeSimulatedCorrections,
+		TEXT("")
+		TEXT("0: Disable, 1: Enable"),
+		ECVF_Cheat);
+
+	static int32 DebugTimeDiscrepancy = 0;
+	FAutoConsoleVariableRef CVarDebugTimeDiscrepancy(
+		TEXT("p.DebugTimeDiscrepancy"),
+		DebugTimeDiscrepancy,
+		TEXT("Whether to log detailed Movement Time Discrepancy values for testing")
+		TEXT("0: Disable, 1: Enable Detection logging, 2: Enable Detection and Resolution logging"),
+		ECVF_Cheat);
 
 #endif // !(UE_BUILD_SHIPPING || UE_BUILD_TEST)
+}
 
-static TAutoConsoleVariable<int32> CVarReplayUseInterpolation( TEXT( "p.ReplayUseInterpolation" ), 1, TEXT( "" ) );
 
 void FFindFloorResult::SetFromSweep(const FHitResult& InHit, const float InSweepFloorDist, const bool bIsWalkableFloor)
 {
@@ -442,7 +461,7 @@ void UCharacterMovementComponent::OnRegister()
 	const bool IsReplay = (MyWorld && MyWorld->DemoNetDriver && MyWorld->DemoNetDriver->ServerConnection);
 	if (IsReplay)
 	{
-		if (CVarReplayUseInterpolation.GetValueOnGameThread() == 1)
+		if (CharacterMovementCVars::ReplayUseInterpolation == 1)
 		{
 			NetworkSmoothingMode = ENetworkSmoothingMode::Replay;
 		}
@@ -714,7 +733,7 @@ float UCharacterMovementComponent::GetNetworkSafeRandomAngleDegrees() const
 {
 	float Angle = FMath::SRand() * 360.f;
 
-	if (GetNetMode() > NM_Standalone)
+	if (!IsNetMode(NM_Standalone))
 	{
 		// Networked game
 		// Get a timestamp that is relatively close between client and server (within ping).
@@ -1039,7 +1058,7 @@ void UCharacterMovementComponent::TickComponent(float DeltaTime, enum ELevelTick
 	if (bIsSimulatingPhysics)
 	{
 		// Update camera to ensure client gets updates even when physics move him far away from point where simulation started
-		if (CharacterOwner->Role == ROLE_AutonomousProxy && GetNetMode() == NM_Client)
+		if (CharacterOwner->Role == ROLE_AutonomousProxy && IsNetMode(NM_Client))
 		{
 			APlayerController* PC = Cast<APlayerController>(CharacterOwner->GetController());
 			APlayerCameraManager* PlayerCameraManager = (PC ? PC->PlayerCameraManager : NULL);
@@ -1059,7 +1078,7 @@ void UCharacterMovementComponent::TickComponent(float DeltaTime, enum ELevelTick
 		SCOPE_CYCLE_COUNTER(STAT_CharacterMovementNonSimulated);
 
 		// If we are a client we might have received an update from the server.
-		const bool bIsClient = (CharacterOwner->Role == ROLE_AutonomousProxy && GetNetMode() == NM_Client);
+		const bool bIsClient = (CharacterOwner->Role == ROLE_AutonomousProxy && IsNetMode(NM_Client));
 		if (bIsClient)
 		{
 			ClientUpdatePositionAfterServerUpdate();
@@ -1098,7 +1117,7 @@ void UCharacterMovementComponent::TickComponent(float DeltaTime, enum ELevelTick
 			MaybeSaveBaseLocation();
 
 			// Smooth on listen server for local view of remote clients. We may receive updates at a rate different than our own tick rate.
-			if (CharacterCVars::NetEnableListenServerSmoothing && !bNetworkSmoothingComplete && GetNetMode() == NM_ListenServer)
+			if (CharacterMovementCVars::NetEnableListenServerSmoothing && !bNetworkSmoothingComplete && IsNetMode(NM_ListenServer))
 			{
 				SmoothClientPosition(DeltaTime);
 			}
@@ -1126,7 +1145,7 @@ void UCharacterMovementComponent::TickComponent(float DeltaTime, enum ELevelTick
 	}
 
 #if !(UE_BUILD_SHIPPING || UE_BUILD_TEST)
-	const bool bVisualizeMovement = CVarVisualizeMovement.GetValueOnGameThread() > 0;
+	const bool bVisualizeMovement = CharacterMovementCVars::VisualizeMovement > 0;
 	if (bVisualizeMovement)
 	{
 		VisualizeMovement();
@@ -1152,8 +1171,8 @@ void UCharacterMovementComponent::AdjustProxyCapsuleSize()
 	{
 		bShrinkProxyCapsule = false;
 
-		float ShrinkRadius = FMath::Max(0.f, CVarNetProxyShrinkRadius.GetValueOnGameThread());
-		float ShrinkHalfHeight = FMath::Max(0.f, CVarNetProxyShrinkHalfHeight.GetValueOnGameThread());
+		float ShrinkRadius = FMath::Max(0.f, CharacterMovementCVars::NetProxyShrinkRadius);
+		float ShrinkHalfHeight = FMath::Max(0.f, CharacterMovementCVars::NetProxyShrinkHalfHeight);
 
 		if (ShrinkRadius == 0.f && ShrinkHalfHeight == 0.f)
 		{
@@ -1249,6 +1268,7 @@ void UCharacterMovementComponent::SimulatedTick(float DeltaSeconds)
 			const FVector OldLocation = UpdatedComponent->GetComponentLocation();
 			SimulateRootMotion(DeltaSeconds, RootMotionParams.RootMotionTransform);
 
+#if !(UE_BUILD_SHIPPING)
 			// debug
 			if (false)
 			{
@@ -1261,6 +1281,7 @@ void UCharacterMovementComponent::SimulatedTick(float DeltaSeconds)
 				UE_LOG(LogRootMotion, Log,  TEXT("UCharacterMovementComponent::SimulatedTick DeltaMovement Translation: %s, Rotation: %s, MovementBase: %s"),
 					*(NewLocation - OldLocation).ToCompactString(), *(NewRotation - OldRotation).GetNormalized().ToCompactString(), *GetNameSafe(CharacterOwner->GetMovementBase()) );
 			}
+#endif // !(UE_BUILD_SHIPPING)
 		}
 
 		// then, once our position is up to date with our animation, 
@@ -1970,6 +1991,7 @@ void UCharacterMovementComponent::PerformMovement(float DeltaSeconds)
 				MoveUpdatedComponent(FVector::ZeroVector, NewActorRotationQuat, true);
 			}
 
+#if !(UE_BUILD_SHIPPING)
 			// debug
 			if (false)
 			{
@@ -1992,6 +2014,7 @@ void UCharacterMovementComponent::PerformMovement(float DeltaSeconds)
 				UE_LOG(LogRootMotion, Warning,  TEXT("PerformMovement Resulting DeltaError Translation: %s, Rotation: %s"),
 					*(ResultingLocation - OldLocation - RMTranslation).ToCompactString(), *(ResultingRotation - OldActorRotation - RMRotation).GetNormalized().ToCompactString() );
 			}
+#endif // !(UE_BUILD_SHIPPING)
 
 			// Root Motion has been used, clear
 			RootMotionParams.Clear();
@@ -2032,7 +2055,7 @@ void UCharacterMovementComponent::PerformMovement(float DeltaSeconds)
 		}
 	}
 
-	if (bHasAuthority && UpdatedComponent && GetNetMode() < NM_Client)
+	if (bHasAuthority && UpdatedComponent && !IsNetMode(NM_Client))
 	{
 		const bool bLocationChanged = (UpdatedComponent->GetComponentLocation() != LastUpdateLocation);
 		const bool bRotationChanged = (UpdatedComponent->GetComponentQuat() != LastUpdateRotation);
@@ -4006,11 +4029,11 @@ FVector UCharacterMovementComponent::ComputeGroundMovementDelta(const FVector& D
 
 void UCharacterMovementComponent::OnCharacterStuckInGeometry(const FHitResult* Hit)
 {
-	if (CharacterCVars::StuckWarningPeriod >= 0)
+	if (CharacterMovementCVars::StuckWarningPeriod >= 0)
 	{
 		UWorld* MyWorld = GetWorld();
 		const float RealTimeSeconds = MyWorld->GetRealTimeSeconds();
-		if ((RealTimeSeconds - LastStuckWarningTime) >= CharacterCVars::StuckWarningPeriod)
+		if ((RealTimeSeconds - LastStuckWarningTime) >= CharacterMovementCVars::StuckWarningPeriod)
 		{
 			LastStuckWarningTime = RealTimeSeconds;
 			if (Hit == nullptr)
@@ -6435,7 +6458,7 @@ void UCharacterMovementComponent::SmoothCorrection(const FVector& OldLocation, c
 			MyWorld->GetTimeSeconds(), ClientData->SmoothingServerTimeStamp, ClientData->SmoothingClientTimeStamp, ClientData->LastCorrectionDelta, *GetNameSafe(CharacterOwner));
 
 #if !(UE_BUILD_SHIPPING || UE_BUILD_TEST)
-		if ( CVarNetVisualizeSimulatedCorrections.GetValueOnGameThread() >= 2 )
+		if ( CharacterMovementCVars::NetVisualizeSimulatedCorrections >= 2 )
 		{
 			const float Radius		= 4.0f;
 			const bool	bPersist	= false;
@@ -6585,7 +6608,7 @@ void UCharacterMovementComponent::SmoothClientPosition_Interpolate(float DeltaSe
 
 #if !(UE_BUILD_SHIPPING || UE_BUILD_TEST)
 			// Show lerp percent
-			if ( CVarNetVisualizeSimulatedCorrections.GetValueOnGameThread() >= 1 )
+			if ( CharacterMovementCVars::NetVisualizeSimulatedCorrections >= 1 )
 			{
 				const FColor DebugColor = FColor::White;
 				const FVector DebugLocation = CharacterOwner->GetMesh()->GetComponentLocation() + FVector( 0.f, 0.f, 300.0f ) - CharacterOwner->GetBaseTranslationOffset();
@@ -6709,7 +6732,7 @@ void UCharacterMovementComponent::SmoothClientPosition_Interpolate(float DeltaSe
 					ClientData->MeshRotationOffset = Rotation;
 
 #if !(UE_BUILD_SHIPPING || UE_BUILD_TEST)
-					if ( CVarNetVisualizeSimulatedCorrections.GetValueOnGameThread() >= 1 )
+					if ( CharacterMovementCVars::NetVisualizeSimulatedCorrections >= 1 )
 					{
 						const float Radius		= 4.0f;
 						const int32	Sides		= 8;
@@ -6771,7 +6794,7 @@ void UCharacterMovementComponent::SmoothClientPosition_Interpolate(float DeltaSe
 
 #if !(UE_BUILD_SHIPPING || UE_BUILD_TEST)
 			// Show future samples
-			if ( CVarNetVisualizeSimulatedCorrections.GetValueOnGameThread() >= 1 )
+			if ( CharacterMovementCVars::NetVisualizeSimulatedCorrections >= 1 )
 			{
 				const float Radius		= 4.0f;
 				const int32	Sides		= 8;
@@ -6798,12 +6821,12 @@ void UCharacterMovementComponent::SmoothClientPosition_Interpolate(float DeltaSe
 
 						DrawCircle( GetWorld(), Location, FVector( 1, 0, 0 ), FVector( 0, 1, 0 ), FColor( 255, 0, 0, 255 ), Radius, Sides, false, 0.0f );
 						
-						if ( CVarNetVisualizeSimulatedCorrections.GetValueOnGameThread() >= 2 )
+						if ( CharacterMovementCVars::NetVisualizeSimulatedCorrections >= 2 )
 						{
 							DrawDebugDirectionalArrow( GetWorld(), Location, Location + ClientData->ReplaySamples[i].Velocity, 20.0f, FColor( 255, 0, 0, 255 ) );
 						}
 
-						if ( CVarNetVisualizeSimulatedCorrections.GetValueOnGameThread() >= 3 )
+						if ( CharacterMovementCVars::NetVisualizeSimulatedCorrections >= 3 )
 						{
 							DrawDebugDirectionalArrow( GetWorld(), Location, Location + ClientData->ReplaySamples[i].Acceleration, 20.0f, FColor( 255, 255, 255, 255 ) );
 						}
@@ -6828,7 +6851,7 @@ void UCharacterMovementComponent::SmoothClientPosition_Interpolate(float DeltaSe
 		//UE_LOG(LogCharacterNetSmoothing, VeryVerbose, TEXT("SmoothClientPosition_Interpolate %s: Translation: %s Rotation: %s"),
 		//	*GetNameSafe(CharacterOwner), *ClientData->MeshTranslationOffset.ToString(), *ClientData->MeshRotationOffset.ToString());
 
-		if ( CVarNetVisualizeSimulatedCorrections.GetValueOnGameThread() >= 1 && NetworkSmoothingMode != ENetworkSmoothingMode::Replay )
+		if ( CharacterMovementCVars::NetVisualizeSimulatedCorrections >= 1 && NetworkSmoothingMode != ENetworkSmoothingMode::Replay )
 		{
 			const FVector DebugLocation = CharacterOwner->GetMesh()->GetComponentLocation() + FVector( 0.f, 0.f, 300.0f ) - CharacterOwner->GetBaseTranslationOffset();
 			DrawDebugBox( GetWorld(), DebugLocation, FVector( 45, 45, 45 ), CharacterOwner->GetMesh()->GetComponentQuat(), FColor( 0, 255, 0 ) );
@@ -6836,7 +6859,7 @@ void UCharacterMovementComponent::SmoothClientPosition_Interpolate(float DeltaSe
 			//DrawDebugCoordinateSystem( GetWorld(), UpdatedComponent->GetComponentLocation() + FVector( 0, 0, 300.0f ), UpdatedComponent->GetComponentRotation(), 45.0f );
 			//DrawDebugBox( GetWorld(), UpdatedComponent->GetComponentLocation() + FVector( 0, 0, 300.0f ), FVector( 45, 45, 45 ), UpdatedComponent->GetComponentQuat(), FColor( 0, 255, 0 ) );
 
-			if ( CVarNetVisualizeSimulatedCorrections.GetValueOnGameThread() >= 3 )
+			if ( CharacterMovementCVars::NetVisualizeSimulatedCorrections >= 3 )
 			{
 				ClientData->SimulatedDebugDrawTime += DeltaSeconds;
 
@@ -6924,7 +6947,7 @@ bool UCharacterMovementComponent::ClientUpdatePositionAfterServerUpdate()
 	if (bIgnoreClientMovementErrorChecksAndCorrection)
 	{
 #if !UE_BUILD_SHIPPING
-		if (CVarNetShowCorrections.GetValueOnGameThread() != 0)
+		if (CharacterMovementCVars::NetShowCorrections != 0)
 		{
 			UE_LOG(LogNetPlayerMovement, Warning, TEXT("*** Client: %s is set to ignore error checks and corrections with %d saved moves in queue."), *GetNameSafe(CharacterOwner), ClientData->SavedMoves.Num());
 		}
@@ -7241,7 +7264,7 @@ void UCharacterMovementComponent::ReplicateMoveToServer(float DeltaTime, const F
 	{
 		ClientData->SavedMoves.Push(NewMove);
 
-		const bool bCanDelayMove = (CVarNetEnableMoveCombining.GetValueOnGameThread() != 0) && CanDelaySendingMove(NewMove);
+		const bool bCanDelayMove = (CharacterMovementCVars::NetEnableMoveCombining != 0) && CanDelaySendingMove(NewMove);
 
 		if (bCanDelayMove && ClientData->PendingMove.IsValid() == false)
 		{
@@ -7555,7 +7578,7 @@ void UCharacterMovementComponent::ProcessClientTimeStampForTimeDiscrepancy(float
 
 #if !(UE_BUILD_SHIPPING || UE_BUILD_TEST)
 		// Per-frame spew of time discrepancy-related values - useful for investigating state of time discrepancy tracking
-		if ( CVarDebugTimeDiscrepancy.GetValueOnGameThread() > 0 )
+		if (CharacterMovementCVars::DebugTimeDiscrepancy > 0)
 		{
 			UE_LOG(LogNetPlayerMovement, Warning, TEXT("TimeDiscrepancyDetection: ClientError: %f, TimeDiscrepancy: %f, LifetimeRawTimeDiscrepancy: %f (Lifetime %f), Resolving: %d, ClientDelta: %f, ServerDelta: %f, ClientTimeStamp: %f"), 
 				ClientError, ServerData.TimeDiscrepancy, ServerData.LifetimeRawTimeDiscrepancy, WorldTimeSeconds - ServerData.WorldCreationTime, ServerData.bResolvingTimeDiscrepancy, ClientDelta, ServerDelta, ClientTimeStamp);
@@ -7669,7 +7692,7 @@ void UCharacterMovementComponent::ProcessClientTimeStampForTimeDiscrepancy(float
 
 #if !(UE_BUILD_SHIPPING || UE_BUILD_TEST)
 		// Per-frame spew of time discrepancy resolution related values - useful for investigating state of time discrepancy tracking
-		if ( CVarDebugTimeDiscrepancy.GetValueOnGameThread() > 1 )
+		if ( CharacterMovementCVars::DebugTimeDiscrepancy > 1 )
 		{
 			UE_LOG(LogNetPlayerMovement, Warning, TEXT("TimeDiscrepancyResolution: DeltaOverride: %f, TimeToPayBack: %f, BaseDelta: %f, ServerDeltaSinceLastMovementUpdate: %f, TimeDiscrepancyAccumulatedClientDeltasSinceLastServerTick: %f"), 
 				ServerData.TimeDiscrepancyResolutionMoveDeltaOverride, TimeToPayBack, BaseDeltaTime, ServerDeltaSinceLastMovementUpdate, ServerData.TimeDiscrepancyAccumulatedClientDeltasSinceLastServerTick);
@@ -7853,13 +7876,13 @@ void UCharacterMovementComponent::ServerMoveHandleClientError(float ClientTimeSt
 
 
 #if !UE_BUILD_SHIPPING
-		if (CVarNetShowCorrections.GetValueOnGameThread() != 0)
+		if (CharacterMovementCVars::NetShowCorrections != 0)
 		{
 			const FVector LocDiff = UpdatedComponent->GetComponentLocation() - ClientLoc;
 			const FString BaseString = MovementBase ? MovementBase->GetPathName(MovementBase->GetOutermost()) : TEXT("None");
 			UE_LOG(LogNetPlayerMovement, Warning, TEXT("*** Server: Error for %s at Time=%.3f is %3.3f LocDiff(%s) ClientLoc(%s) ServerLoc(%s) Base: %s Bone: %s Accel(%s) Velocity(%s)"),
 				*GetNameSafe(CharacterOwner), ClientTimeStamp, LocDiff.Size(), *LocDiff.ToString(), *ClientLoc.ToString(), *UpdatedComponent->GetComponentLocation().ToString(), *BaseString, *ServerData->PendingAdjustment.NewBaseBoneName.ToString(), *Accel.ToString(), *Velocity.ToString());
-			const float DebugLifetime = CVarNetCorrectionLifetime.GetValueOnGameThread();
+			const float DebugLifetime = CharacterMovementCVars::NetCorrectionLifetime;
 			DrawDebugCapsule(GetWorld(), UpdatedComponent->GetComponentLocation(), CharacterOwner->GetSimpleCollisionHalfHeight(), CharacterOwner->GetSimpleCollisionRadius(), FQuat::Identity, FColor(100, 255, 100), true, DebugLifetime);
 			DrawDebugCapsule(GetWorld(), ClientLoc                    , CharacterOwner->GetSimpleCollisionHalfHeight(), CharacterOwner->GetSimpleCollisionRadius(), FQuat::Identity, FColor(255, 100, 100), true, DebugLifetime);
 		}
@@ -7920,9 +7943,9 @@ bool UCharacterMovementComponent::ServerCheckClientError(float ClientTimeStamp,
 			return true;
 		}
 #if !(UE_BUILD_SHIPPING || UE_BUILD_TEST)
-		if (CVarNetForceClientAdjustmentPercent.GetValueOnGameThread() > SMALL_NUMBER)
+		if (CharacterMovementCVars::NetForceClientAdjustmentPercent > SMALL_NUMBER)
 		{
-			if (FMath::SRand() < CVarNetForceClientAdjustmentPercent.GetValueOnGameThread())
+			if (FMath::SRand() < CharacterMovementCVars::NetForceClientAdjustmentPercent)
 			{
 				UE_LOG(LogNetPlayerMovement, VeryVerbose, TEXT("** ServerCheckClientError forced by p.NetForceClientAdjustmentPercent"));
 				return true;
@@ -7933,7 +7956,7 @@ bool UCharacterMovementComponent::ServerCheckClientError(float ClientTimeStamp,
 	else
 	{
 #if !UE_BUILD_SHIPPING
-		if (CVarNetShowCorrections.GetValueOnGameThread() != 0)
+		if (CharacterMovementCVars::NetShowCorrections != 0)
 		{
 			UE_LOG(LogNetPlayerMovement, Warning, TEXT("*** Server: %s is set to ignore error checks and corrections."), *GetNameSafe(CharacterOwner));
 		}
@@ -8014,9 +8037,9 @@ void UCharacterMovementComponent::MoveAutonomous
 	if (CharacterOwner && UpdatedComponent)
 	{
 		// Smooth local view of remote clients on listen servers
-		if (CharacterCVars::NetEnableListenServerSmoothing &&
+		if (CharacterMovementCVars::NetEnableListenServerSmoothing &&
 			CharacterOwner->GetRemoteRole() == ROLE_AutonomousProxy &&
-			GetNetMode() == NM_ListenServer)
+			IsNetMode(NM_ListenServer))
 		{
 			SmoothCorrection(OldLocation, OldRotation, UpdatedComponent->GetComponentLocation(), UpdatedComponent->GetComponentQuat());
 		}
@@ -8218,13 +8241,13 @@ void UCharacterMovementComponent::ClientAdjustPosition_Implementation
 	}
 
 #if !UE_BUILD_SHIPPING
-	if (CVarNetShowCorrections.GetValueOnGameThread() != 0)
+	if (CharacterMovementCVars::NetShowCorrections != 0)
 	{
 		const FVector LocDiff = UpdatedComponent->GetComponentLocation() - NewLocation;
 		const FString NewBaseString = NewBase ? NewBase->GetPathName(NewBase->GetOutermost()) : TEXT("None");
 		UE_LOG(LogNetPlayerMovement, Warning, TEXT("*** Client: Error for %s at Time=%.3f is %3.3f LocDiff(%s) ClientLoc(%s) ServerLoc(%s) NewBase: %s NewBone: %s ClientVel(%s) ServerVel(%s) SavedMoves %d"),
 			*GetNameSafe(CharacterOwner), TimeStamp, LocDiff.Size(), *LocDiff.ToString(), *UpdatedComponent->GetComponentLocation().ToString(), *NewLocation.ToString(), *NewBaseString, *NewBaseBoneName.ToString(), *Velocity.ToString(), *NewVelocity.ToString(), ClientData->SavedMoves.Num());
-		const float DebugLifetime = CVarNetCorrectionLifetime.GetValueOnGameThread();
+		const float DebugLifetime = CharacterMovementCVars::NetCorrectionLifetime;
 		DrawDebugCapsule(GetWorld(), UpdatedComponent->GetComponentLocation()	, CharacterOwner->GetSimpleCollisionHalfHeight(), CharacterOwner->GetSimpleCollisionRadius(), FQuat::Identity, FColor(255, 100, 100), true, DebugLifetime);
 		DrawDebugCapsule(GetWorld(), NewLocation						, CharacterOwner->GetSimpleCollisionHalfHeight(), CharacterOwner->GetSimpleCollisionRadius(), FQuat::Identity, FColor(100, 255, 100), true, DebugLifetime);
 	}
@@ -8748,6 +8771,7 @@ void UCharacterMovementComponent::TickCharacterPose(float DeltaTime)
 			RootMotionParams.Accumulate(RootMotion);
 		}
 
+#if !(UE_BUILD_SHIPPING)
 		// Debugging
 		{
 			FAnimMontageInstance* RootMotionMontageInstance = CharacterOwner->GetRootMotionAnimMontageInstance();
@@ -8760,6 +8784,7 @@ void UCharacterMovementComponent::TickCharacterPose(float DeltaTime)
 				, *RootMotionParams.RootMotionTransform.GetTranslation().ToCompactString()
 				);
 		}
+#endif // !(UE_BUILD_SHIPPING)
 	}
 }
 
@@ -8795,7 +8820,7 @@ uint16 UCharacterMovementComponent::ApplyRootMotionSource(FRootMotionSource* Sou
 						SourcePtr->StartTime = ClientData->CurrentTimeStamp;
 					}
 				}
-				else if (CharacterOwner->Role == ROLE_Authority && GetNetMode() < NM_Client)
+				else if (CharacterOwner->Role == ROLE_Authority && !IsNetMode(NM_Client))
 				{
 					// Authority defaults to current client time stamp, meaning it'll start next tick if not corrected
 					FNetworkPredictionData_Server_Character* ServerData = GetPredictionData_Server_Character();
diff --git a/Engine/Source/Runtime/Engine/Private/Components/ChildActorComponent.cpp b/Engine/Source/Runtime/Engine/Private/Components/ChildActorComponent.cpp
--- a/Engine/Source/Runtime/Engine/Private/Components/ChildActorComponent.cpp
+++ b/Engine/Source/Runtime/Engine/Private/Components/ChildActorComponent.cpp
@@ -9,6 +9,7 @@ DEFINE_LOG_CATEGORY_STATIC(LogChildActorComponent, Warning, All);
 UChildActorComponent::UChildActorComponent(const FObjectInitializer& ObjectInitializer)
 	: Super(ObjectInitializer)
 {
+	bWantsBeginPlay = true;
 }
 
 void UChildActorComponent::OnRegister()
@@ -433,3 +434,13 @@ void UChildActorComponent::DestroyChildActor()
 		ChildActor = nullptr;
 	}
 }
+
+void UChildActorComponent::BeginPlay()
+{
+	Super::BeginPlay();
+
+	if (ChildActor && !ChildActor->HasActorBegunPlay())
+	{
+		ChildActor->BeginPlay();
+	}
+}
\ No newline at end of file
diff --git a/Engine/Source/Runtime/Engine/Private/Components/MovementComponent.cpp b/Engine/Source/Runtime/Engine/Private/Components/MovementComponent.cpp
--- a/Engine/Source/Runtime/Engine/Private/Components/MovementComponent.cpp
+++ b/Engine/Source/Runtime/Engine/Private/Components/MovementComponent.cpp
@@ -286,7 +286,7 @@ bool UMovementComponent::ShouldSkipUpdate(float DeltaTime) const
 
 	if (bUpdateOnlyIfRendered)
 	{
-		if (GetNetMode() == NM_DedicatedServer)
+		if (IsNetMode(NM_DedicatedServer))
 		{
 			// Dedicated servers never render
 			return true;
diff --git a/Engine/Source/Runtime/Engine/Private/Components/SceneComponent.cpp b/Engine/Source/Runtime/Engine/Private/Components/SceneComponent.cpp
--- a/Engine/Source/Runtime/Engine/Private/Components/SceneComponent.cpp
+++ b/Engine/Source/Runtime/Engine/Private/Components/SceneComponent.cpp
@@ -456,7 +456,7 @@ void USceneComponent::OnRegister()
 	if (bVisualizeComponent && SpriteComponent == nullptr && GetOwner() && !GetWorld()->IsGameWorld() )
 	{
 		// Create a new billboard component to serve as a visualization of the actor until there is another primitive component
-		SpriteComponent = NewObject<UBillboardComponent>(GetOwner(), NAME_None, RF_Transactional | RF_TextExportTransient);
+		SpriteComponent = NewObject<UBillboardComponent>(GetOwner(), NAME_None, RF_Transactional | RF_Transient);
 
 		SpriteComponent->Sprite = LoadObject<UTexture2D>(nullptr, TEXT("/Engine/EditorResources/EmptyActor.EmptyActor"));
 		SpriteComponent->RelativeScale3D = FVector(0.5f, 0.5f, 0.5f);
@@ -3125,7 +3125,7 @@ void USceneComponent::UpdateNavigationData()
 	if (UNavigationSystem::ShouldUpdateNavOctreeOnComponentChange() && IsRegistered())
 	{
 		UWorld* MyWorld = GetWorld();
-		if ((MyWorld == nullptr) || !MyWorld->IsGameWorld() || MyWorld->GetNetMode() < ENetMode::NM_Client)
+		if ((MyWorld == nullptr) || !MyWorld->IsGameWorld() || !MyWorld->IsNetMode(ENetMode::NM_Client))
 		{
 			// use propagated component's transform update in editor OR server game with additional navsys check
 			UNavigationSystem::UpdateComponentInNavOctree(*this);
@@ -3140,7 +3140,7 @@ void USceneComponent::PostUpdateNavigationData()
 	if (UNavigationSystem::ShouldUpdateNavOctreeOnComponentChange() && IsRegistered())
 	{
 		UWorld* MyWorld = GetWorld();
-		if ((MyWorld == nullptr) || !MyWorld->IsGameWorld() || MyWorld->GetNetMode() < ENetMode::NM_Client)
+		if ((MyWorld == nullptr) || !MyWorld->IsGameWorld() || !MyWorld->IsNetMode(ENetMode::NM_Client))
 		{
 			// use propagated component's transform update in editor OR server game with additional navsys check
 			UNavigationSystem::UpdateNavOctreeAfterMove(this);
diff --git a/Engine/Source/Runtime/Engine/Private/Components/SkeletalMeshComponent.cpp b/Engine/Source/Runtime/Engine/Private/Components/SkeletalMeshComponent.cpp
--- a/Engine/Source/Runtime/Engine/Private/Components/SkeletalMeshComponent.cpp
+++ b/Engine/Source/Runtime/Engine/Private/Components/SkeletalMeshComponent.cpp
@@ -257,7 +257,7 @@ void USkeletalMeshComponent::RegisterClothTick(bool bRegister)
 
 bool USkeletalMeshComponent::ShouldRunPostPhysicsTick() const
 {
-	return	(bEnablePhysicsOnDedicatedServer || GetNetMode() != NM_DedicatedServer) && // Early out with we are on a dedicated server and not running physics
+	return	(bEnablePhysicsOnDedicatedServer || !IsNetMode(NM_DedicatedServer)) && // Early out if we are on a dedicated server and not running physics.
 			(IsSimulatingPhysics() || ShouldBlendPhysicsBones());
 }
 
@@ -270,7 +270,7 @@ bool USkeletalMeshComponent::ShouldRunClothTick() const
 {
 #if WITH_APEX_CLOTHING
 	bool bShouldRunCloth = ClothingActors.Num() > 0 && SkeletalMesh && SkeletalMesh->ClothingAssets.Num() > 0
-								&& GetNetMode() != NM_DedicatedServer; // Cloth never needs to run on dedicated server
+								&& !IsNetMode(NM_DedicatedServer); // Cloth never needs to run on dedicated server
 
 	//If we are eligible to run cloth we should check if any of the clothing actors will actually simulate at this LOD
 	if(bShouldRunCloth)
@@ -1949,6 +1949,12 @@ class UAnimSingleNodeInstance* USkeletalMeshComponent::GetSingleNodeInstance() c
 	return Cast<class UAnimSingleNodeInstance>(AnimScriptInstance);
 }
 
+bool USkeletalMeshComponent::PoseTickedThisFrame() const 
+{ 
+	return LastPoseTickTime == GetWorld()->TimeSeconds; 
+}
+
+
 FTransform USkeletalMeshComponent::ConvertLocalRootMotionToWorld(const FTransform& InTransform)
 {
 	// Make sure component to world is up to date
@@ -2036,31 +2042,11 @@ bool USkeletalMeshComponent::ComponentIsTouchingSelectionBox(const FBox& InSelBB
 		check(SkelMeshResource);
 		check(SkelMeshResource->LODModels.Num() > 0);
 
-		// Transform hard and soft verts into world space. Note that this assumes skeletal mesh is in reference pose...
+		// Transform verts into world space. Note that this assumes skeletal mesh is in reference pose...
 		const FStaticLODModel& LODModel = SkelMeshResource->LODModels[0];
-		for (const auto& Chunk : LODModel.Chunks)
+		for (const auto& Section : LODModel.Sections)
 		{
-			for (const auto& Vertex : Chunk.RigidVertices)
-			{
-				const FVector Location = ComponentToWorld.TransformPosition(Vertex.Position);
-				const bool bLocationIntersected = FMath::PointBoxIntersection(Location, InSelBBox);
-
-				// If the selection box doesn't have to encompass the entire component and a skeletal mesh vertex has intersected with
-				// the selection box, this component is being touched by the selection box
-				if (!bMustEncompassEntireComponent && bLocationIntersected)
-				{
-					return true;
-				}
-
-				// If the selection box has to encompass the entire component and a skeletal mesh vertex didn't intersect with the selection
-				// box, this component does not qualify
-				else if (bMustEncompassEntireComponent && !bLocationIntersected)
-				{
-					return false;
-				}
-			}
-
-			for (const auto& Vertex : Chunk.SoftVertices)
+			for (const auto& Vertex : Section.SoftVertices)
 			{
 				const FVector Location = ComponentToWorld.TransformPosition(Vertex.Position);
 				const bool bLocationIntersected = FMath::PointBoxIntersection(Location, InSelBBox);
@@ -2085,8 +2071,8 @@ bool USkeletalMeshComponent::ComponentIsTouchingSelectionBox(const FBox& InSelBB
 		// is consider touching
 		if (bMustEncompassEntireComponent)
 		{
-		return true;
-	}
+			return true;
+		}
 	}
 
 	return false;
@@ -2100,31 +2086,11 @@ bool USkeletalMeshComponent::ComponentIsTouchingSelectionFrustum(const FConvexVo
 		check(SkelMeshResource);
 		check(SkelMeshResource->LODModels.Num() > 0);
 
-		// Transform hard and soft verts into world space. Note that this assumes skeletal mesh is in reference pose...
+		// Transform verts into world space. Note that this assumes skeletal mesh is in reference pose...
 		const FStaticLODModel& LODModel = SkelMeshResource->LODModels[0];
-		for (const auto& Chunk : LODModel.Chunks)
+		for (const auto& Section : LODModel.Sections)
 		{
-			for (const auto& Vertex : Chunk.RigidVertices)
-			{
-				const FVector Location = ComponentToWorld.TransformPosition(Vertex.Position);
-				const bool bLocationIntersected = InFrustum.IntersectSphere(Location, 0.0f);
-
-				// If the selection box doesn't have to encompass the entire component and a skeletal mesh vertex has intersected with
-				// the selection box, this component is being touched by the selection box
-				if (!bMustEncompassEntireComponent && bLocationIntersected)
-				{
-					return true;
-				}
-
-				// If the selection box has to encompass the entire component and a skeletal mesh vertex didn't intersect with the selection
-				// box, this component does not qualify
-				else if (bMustEncompassEntireComponent && !bLocationIntersected)
-				{
-					return false;
-				}
-			}
-
-			for (const auto& Vertex : Chunk.SoftVertices)
+			for (const auto& Vertex : Section.SoftVertices)
 			{
 				const FVector Location = ComponentToWorld.TransformPosition(Vertex.Position);
 				const bool bLocationIntersected = InFrustum.IntersectSphere(Location, 0.0f);
diff --git a/Engine/Source/Runtime/Engine/Private/Components/SkinnedMeshComponent.cpp b/Engine/Source/Runtime/Engine/Private/Components/SkinnedMeshComponent.cpp
--- a/Engine/Source/Runtime/Engine/Private/Components/SkinnedMeshComponent.cpp
+++ b/Engine/Source/Runtime/Engine/Private/Components/SkinnedMeshComponent.cpp
@@ -989,15 +989,11 @@ void USkinnedMeshComponent::SetSkeletalMesh(USkeletalMesh* InSkelMesh, bool bRei
 			UpdateMasterBoneMap();
 			UpdateLODStatus();
 			InvalidateCachedBounds();
+			// clear morphtarget cache
+			ActiveMorphTargets.Empty();			
 		}
 	}
 	
-	// TODO: (LH) find better way to call this 
-	// Right now SkeletalMeshComponent(child) needs to call this after animation update
-	// But some of functions of SkinnedMeshComponent has to be called before
-	// So it's hard to split both separate, but do we really need to call this when skeletalmesh is set
-	// Think about this
-	// RefreshBoneTransforms(); 
 	
 	// Notify the streaming system. Don't use Update(), because this may be the first time the mesh has been set
 	// and the component may have to be added to the streaming system for the first time.
@@ -1631,7 +1627,7 @@ void USkinnedMeshComponent::GetUsedMaterials( TArray<UMaterialInterface*>& OutMa
 }
 
 template <bool bExtraBoneInfluencesT, bool bCachedMatrices>
-FORCEINLINE FVector USkinnedMeshComponent::GetTypedSkinnedVertexPosition(const FSkelMeshChunk& Chunk, const FSkeletalMeshVertexBuffer& VertexBufferGPUSkin, int32 VertIndex, bool bSoftVertex, const TArray<FMatrix> & RefToLocals) const
+FORCEINLINE FVector USkinnedMeshComponent::GetTypedSkinnedVertexPosition(const FSkelMeshSection& Section, const FSkeletalMeshVertexBuffer& VertexBufferGPUSkin, int32 VertIndex, const TArray<FMatrix> & RefToLocals) const
 {
 	FVector SkinnedPos(0,0,0);
 
@@ -1639,60 +1635,34 @@ FORCEINLINE FVector USkinnedMeshComponent::GetTypedSkinnedVertexPosition(const F
 	const USkinnedMeshComponent* BaseComponent = MasterPoseComponentInst ? MasterPoseComponentInst : this;
 
 	// Do soft skinning for this vertex.
-	if(bSoftVertex)
-	{
-		const TGPUSkinVertexBase<bExtraBoneInfluencesT>* SrcSoftVertex = VertexBufferGPUSkin.GetVertexPtr<bExtraBoneInfluencesT>(Chunk.GetSoftVertexBufferIndex()+VertIndex);
+	const TGPUSkinVertexBase<bExtraBoneInfluencesT>* SrcSoftVertex = VertexBufferGPUSkin.GetVertexPtr<bExtraBoneInfluencesT>(Section.GetVertexBufferIndex()+VertIndex);
 
 #if !PLATFORM_LITTLE_ENDIAN
-		// uint8[] elements in LOD.VertexBufferGPUSkin have been swapped for VET_UBYTE4 vertex stream use
-		for(int32 InfluenceIndex = MAX_INFLUENCES-1;InfluenceIndex >=  MAX_INFLUENCES-Chunk.MaxBoneInfluences;InfluenceIndex--)
+	// uint8[] elements in LOD.VertexBufferGPUSkin have been swapped for VET_UBYTE4 vertex stream use
+	for(int32 InfluenceIndex = MAX_INFLUENCES-1;InfluenceIndex >=  MAX_INFLUENCES- Section.MaxBoneInfluences;InfluenceIndex--)
 #else
-		for(int32 InfluenceIndex = 0;InfluenceIndex < Chunk.MaxBoneInfluences;InfluenceIndex++)
+	for(int32 InfluenceIndex = 0;InfluenceIndex < Section.MaxBoneInfluences;InfluenceIndex++)
 #endif
-		{
-			int32 BoneIndex = Chunk.BoneMap[SrcSoftVertex->InfluenceBones[InfluenceIndex]];
-			if(MasterPoseComponentInst)
-			{		
-				check(MasterBoneMap.Num() == SkeletalMesh->RefSkeleton.GetNum());
-				BoneIndex = MasterBoneMap[BoneIndex];
-			}
-
-			const float	Weight = (float)SrcSoftVertex->InfluenceWeights[InfluenceIndex] / 255.0f;
-			{
-				if (bCachedMatrices)
-				{
-					const FMatrix& RefToLocal = RefToLocals[BoneIndex];
-					SkinnedPos += RefToLocal.TransformPosition(VertexBufferGPUSkin.GetVertexPositionFast(SrcSoftVertex)) * Weight;
-				}
-				else
-				{
-					const FMatrix RefToLocal = SkeletalMesh->RefBasesInvMatrix[BoneIndex] * BaseComponent->GetSpaceBases()[BoneIndex].ToMatrixWithScale();
-					SkinnedPos += RefToLocal.TransformPosition(VertexBufferGPUSkin.GetVertexPositionFast(SrcSoftVertex)) * Weight;
-				}
-			}
-		}
-	}
-	// Do rigid (one-influence) skinning for this vertex.
-	else
 	{
-		const TGPUSkinVertexBase<bExtraBoneInfluencesT>* SrcRigidVertex = VertexBufferGPUSkin.GetVertexPtr<bExtraBoneInfluencesT>(Chunk.GetRigidVertexBufferIndex() + VertIndex);
-		const int32 RigidInfluenceIndex = SkinningTools::GetRigidInfluenceIndex();
-		int32 BoneIndex = Chunk.BoneMap[SrcRigidVertex->InfluenceBones[RigidInfluenceIndex]];
+		int32 BoneIndex = Section.BoneMap[SrcSoftVertex->InfluenceBones[InfluenceIndex]];
 		if(MasterPoseComponentInst)
-		{
+		{		
 			check(MasterBoneMap.Num() == SkeletalMesh->RefSkeleton.GetNum());
 			BoneIndex = MasterBoneMap[BoneIndex];
 		}
 
-		if (bCachedMatrices)
-		{
-			const FMatrix& RefToLocal = RefToLocals[BoneIndex];
-			SkinnedPos = RefToLocal.TransformPosition(VertexBufferGPUSkin.GetVertexPositionFast(SrcRigidVertex));
-		}
-		else
+		const float	Weight = (float)SrcSoftVertex->InfluenceWeights[InfluenceIndex] / 255.0f;
 		{
-			const FMatrix RefToLocal = SkeletalMesh->RefBasesInvMatrix[BoneIndex] * BaseComponent->GetSpaceBases()[BoneIndex].ToMatrixWithScale();
-			SkinnedPos = RefToLocal.TransformPosition(VertexBufferGPUSkin.GetVertexPositionFast(SrcRigidVertex));
+			if (bCachedMatrices)
+			{
+				const FMatrix& RefToLocal = RefToLocals[BoneIndex];
+				SkinnedPos += RefToLocal.TransformPosition(VertexBufferGPUSkin.GetVertexPositionFast(SrcSoftVertex)) * Weight;
+			}
+			else
+			{
+				const FMatrix RefToLocal = SkeletalMesh->RefBasesInvMatrix[BoneIndex] * BaseComponent->GetSpaceBases()[BoneIndex].ToMatrixWithScale();
+				SkinnedPos += RefToLocal.TransformPosition(VertexBufferGPUSkin.GetVertexPositionFast(SrcSoftVertex)) * Weight;
+			}
 		}
 	}
 
@@ -1712,21 +1682,18 @@ FVector USkinnedMeshComponent::GetSkinnedVertexPosition(int32 VertexIndex) const
 	FStaticLODModel& Model = MeshObject->GetSkeletalMeshResource().LODModels[0];
 
 	//cache RefToLocal matrices
-	int32 ChunkIndex;
+	int32 SectionIndex;
 	int32 VertIndex;
-	bool bSoftVertex;
 	bool bHasExtraBoneInfluences;
-	Model.GetChunkAndSkinType(VertexIndex, ChunkIndex, VertIndex, bSoftVertex, bHasExtraBoneInfluences);
+	Model.GetSectionFromVertexIndex(VertexIndex, SectionIndex, VertIndex, bHasExtraBoneInfluences);
 
 	//update positions
-	check(ChunkIndex < Model.Chunks.Num());
-	const FSkelMeshChunk& Chunk = Model.Chunks[ChunkIndex];
-			//rigid
+	check(SectionIndex < Model.Sections.Num());
+	const FSkelMeshSection& Section = Model.Sections[SectionIndex];
 
 	return bHasExtraBoneInfluences
-		? GetTypedSkinnedVertexPosition<true, false>(Chunk, Model.VertexBufferGPUSkin, VertIndex, bSoftVertex)
-			//soft
-		: GetTypedSkinnedVertexPosition<false, false>(Chunk, Model.VertexBufferGPUSkin, VertIndex, bSoftVertex);
+		? GetTypedSkinnedVertexPosition<true, false>(Section, Model.VertexBufferGPUSkin, VertIndex)
+		: GetTypedSkinnedVertexPosition<false, false>(Section, Model.VertexBufferGPUSkin, VertIndex);
 
 
 }
@@ -1756,30 +1723,18 @@ void USkinnedMeshComponent::ComputeSkinnedPositions(TArray<FVector> & OutPositio
 	}
 
 	//update positions
-	for (int32 ChunkIdx = 0; ChunkIdx < Model.Chunks.Num(); ++ChunkIdx)
+	for (int32 SectionIdx = 0; SectionIdx < Model.Sections.Num(); ++SectionIdx)
 	{
-		const FSkelMeshChunk& Chunk = Model.Chunks[ChunkIdx];
-		const uint32 NumChunkVerts = Chunk.GetNumVertices();
-		bool bHasExtraBoneInfluences = Chunk.HasExtraBoneInfluences();
-		{
-			//rigid
-			const uint32 RigidOffset = Chunk.GetRigidVertexBufferIndex();
-			const uint32 NumRigidVerts = Chunk.GetNumRigidVertices();
-			for (uint32 RigidIdx = 0; RigidIdx < NumRigidVerts; ++RigidIdx)
-			{
-				FVector SkinnedPosition = bHasExtraBoneInfluences ? GetTypedSkinnedVertexPosition<true, true>(Chunk, Model.VertexBufferGPUSkin, RigidIdx, false, RefToLocals) :
-																	GetTypedSkinnedVertexPosition<false, true>(Chunk, Model.VertexBufferGPUSkin, RigidIdx, false, RefToLocals);
-				OutPositions[RigidOffset + RigidIdx] = SkinnedPosition;
-			}
-		}
+		const FSkelMeshSection& Section = Model.Sections[SectionIdx];
+		bool bHasExtraBoneInfluences = Section.HasExtraBoneInfluences();
 		{
 			//soft
-			const uint32 SoftOffset = Chunk.GetSoftVertexBufferIndex();
-			const uint32 NumSoftVerts = Chunk.GetNumSoftVertices();
+			const uint32 SoftOffset = Section.GetVertexBufferIndex();
+			const uint32 NumSoftVerts = Section.GetNumVertices();
 			for (uint32 SoftIdx = 0; SoftIdx < NumSoftVerts; ++SoftIdx)
 			{
-				FVector SkinnedPosition = bHasExtraBoneInfluences ? GetTypedSkinnedVertexPosition<true, true>(Chunk, Model.VertexBufferGPUSkin, SoftIdx, true, RefToLocals) :
-																	GetTypedSkinnedVertexPosition<false,true>(Chunk, Model.VertexBufferGPUSkin, SoftIdx, true, RefToLocals);
+				FVector SkinnedPosition = bHasExtraBoneInfluences ? GetTypedSkinnedVertexPosition<true, true>(Section, Model.VertexBufferGPUSkin, SoftIdx, RefToLocals) :
+																	GetTypedSkinnedVertexPosition<false,true>(Section, Model.VertexBufferGPUSkin, SoftIdx, RefToLocals);
 				OutPositions[SoftOffset + SoftIdx] = SkinnedPosition;
 			}
 		}
@@ -1803,16 +1758,15 @@ FColor USkinnedMeshComponent::GetVertexColor(int32 VertexIndex) const
 	}
 
 	// Find the chunk and vertex within that chunk, and skinning type, for this vertex.
-	int32 ChunkIndex;
+	int32 SectionIndex;
 	int32 VertIndex;
-	bool bSoftVertex;
 	bool bHasExtraBoneInfluences;
-	Model.GetChunkAndSkinType(VertexIndex, ChunkIndex, VertIndex, bSoftVertex, bHasExtraBoneInfluences);
+	Model.GetSectionFromVertexIndex(VertexIndex, SectionIndex, VertIndex, bHasExtraBoneInfluences);
 
-	check(ChunkIndex < Model.Chunks.Num());
-	const FSkelMeshChunk& Chunk = Model.Chunks[ChunkIndex];
+	check(SectionIndex < Model.Sections.Num());
+	const FSkelMeshSection& Section = Model.Sections[SectionIndex];
 	
-	int32 VertexBase = bSoftVertex ? Chunk.GetSoftVertexBufferIndex() : Chunk.GetRigidVertexBufferIndex();
+	int32 VertexBase = Section.GetVertexBufferIndex();
 
 	return Model.ColorVertexBuffer.VertexColor(VertexBase + VertIndex);
 }
diff --git a/Engine/Source/Runtime/Engine/Private/Controller.cpp b/Engine/Source/Runtime/Engine/Private/Controller.cpp
--- a/Engine/Source/Runtime/Engine/Private/Controller.cpp
+++ b/Engine/Source/Runtime/Engine/Private/Controller.cpp
@@ -59,12 +59,6 @@ void AController::K2_DestroyActor()
 
 bool AController::IsLocalController() const
 {
-	if (FGenericPlatformProperties::IsServerOnly())
-	{
-		// Never local on dedicated server. IsServerOnly() is checked at compile time and optimized out if false.
-		return false;
-	}
-
 	const ENetMode NetMode = GetNetMode();
 
 	if (NetMode == NM_Standalone)
diff --git a/Engine/Source/Runtime/Engine/Private/DestructibleMesh.cpp b/Engine/Source/Runtime/Engine/Private/DestructibleMesh.cpp
--- a/Engine/Source/Runtime/Engine/Private/DestructibleMesh.cpp
+++ b/Engine/Source/Runtime/Engine/Private/DestructibleMesh.cpp
@@ -54,9 +54,9 @@ void UDestructibleMesh::PostLoad()
 			ThisLODModel.ActiveBoneIndices.Sort();
 		}
 
-		for (int32 ChunkIndex = 0; ChunkIndex < ThisLODModel.Chunks.Num(); ++ChunkIndex)
+		for (int32 SectionIndex = 0; SectionIndex < ThisLODModel.Sections.Num(); ++SectionIndex)
 		{
-			if (ThisLODModel.Chunks[ChunkIndex].BoneMap.Num() > MaxGPUSkinBones)
+			if (ThisLODModel.Sections[SectionIndex].BoneMap.Num() > MaxGPUSkinBones)
 			{
 #if WITH_EDITOR
 				// re create destructible asset if it exceeds
diff --git a/Engine/Source/Runtime/Engine/Private/Distributions.cpp b/Engine/Source/Runtime/Engine/Private/Distributions.cpp
--- a/Engine/Source/Runtime/Engine/Private/Distributions.cpp
+++ b/Engine/Source/Runtime/Engine/Private/Distributions.cpp
@@ -31,7 +31,7 @@ ENGINE_API uint32 GDistributionType = 1;
 #define LOOKUP_TABLE_MAX_SAMPLES (128)
 
 // UDistribution will bake out (if CanBeBaked returns true)
-#define DISTRIBUTIONS_BAKES_OUT 0
+#define DISTRIBUTIONS_BAKES_OUT 1
 
 // The maximum number of samples must be a power of two.
 static_assert((LOOKUP_TABLE_MAX_SAMPLES & (LOOKUP_TABLE_MAX_SAMPLES - 1)) == 0, "Lookup table max samples is not a power of two.");
@@ -1082,6 +1082,11 @@ bool UDistributionFloat::NeedsLoadForServer() const
 	return true;
 }
 
+bool UDistributionFloat::NeedsLoadForEditorGame() const
+{
+	return true;
+}
+
 #if	WITH_EDITOR
 void UDistributionFloat::PostEditChangeProperty(FPropertyChangedEvent& PropertyChangedEvent)
 {
@@ -1283,6 +1288,11 @@ bool UDistributionVector::NeedsLoadForServer() const
 	return true;
 }
 
+bool UDistributionVector::NeedsLoadForEditorGame() const
+{
+	return true;
+}
+
 #if	WITH_EDITOR
 void UDistributionVector::PostEditChangeProperty(FPropertyChangedEvent& PropertyChangedEvent)
 {
diff --git a/Engine/Source/Runtime/Engine/Private/GPUSkinCache.cpp b/Engine/Source/Runtime/Engine/Private/GPUSkinCache.cpp
--- a/Engine/Source/Runtime/Engine/Private/GPUSkinCache.cpp
+++ b/Engine/Source/Runtime/Engine/Private/GPUSkinCache.cpp
@@ -30,7 +30,7 @@ DEFINE_LOG_CATEGORY_STATIC(LogSkinCache, Log, All);
 
 int32 GEnableGPUSkinCacheShaders = 0;
 static FAutoConsoleVariableRef CVarEnableGPUSkinCacheShaders(
-	TEXT("r.SkinCacheShaders"),
+	TEXT("r.SkinCache.CompileShaders"),
 	GEnableGPUSkinCacheShaders,
 	TEXT("Whether or not to compile the GPU compute skinning cache shaders.\n")
 	TEXT("This will compile the shaders for skinning on a compute job and not skin on the vertex shader.\n")
@@ -52,11 +52,11 @@ static FAutoConsoleVariableRef CVarGSkinCacheSafety(
 // 0/1
 int32 GEnableGPUSkinCache = 0;
 static TAutoConsoleVariable<int32> CVarEnableGPUSkinCache(
-	TEXT("r.SkinCaching"),
+	TEXT("r.SkinCache.Mode"),
 	0,
 	TEXT("Whether or not to use the GPU compute skinning cache.\n")
 	TEXT("This will perform skinning on a compute job and not skin on the vertex shader.\n")
-	TEXT("Requires r.SkinCacheShaders=1.\n")
+	TEXT("Requires r.SkinCache.CompileShaders=1\n")
 	TEXT(" 0: off(default)\n")
 	TEXT(" 1: on"),
 	ECVF_RenderThreadSafe
@@ -65,8 +65,8 @@ static TAutoConsoleVariable<int32> CVarEnableGPUSkinCache(
 TAutoConsoleVariable<int32> CVarGPUSkinCacheRecomputeTangents(
 	TEXT("r.SkinCache.RecomputeTangents"),
 	2,
-	TEXT("If r.SkinCaching is enabled this option recomputes the vertex tangents on the GPU\n")
-	TEXT("Can be changed at runtime, Requires r.SkinCacheShaders=1 and r.SkinCache=1\n")
+	TEXT("If r.SkinCache.Mode is enabled this option recomputes the vertex tangents on the GPU\n")
+	TEXT("Can be changed at runtime, Requires r.SkinCache.CompileShaders=1 and r.SkinCache.Mode=1\n")
 	TEXT(" 0: off\n")
 	TEXT(" 1: on\n")
 	TEXT(" 2: on, SkinCache only objects that would require the RecomputTangents feature (default)"),
@@ -75,19 +75,19 @@ TAutoConsoleVariable<int32> CVarGPUSkinCacheRecomputeTangents(
 
 static int32 GMaxGPUSkinCacheElementsPerFrame = 1000;
 static FAutoConsoleVariableRef CVarMaxGPUSkinCacheElementsPerFrame(
-	TEXT("r.MaxGPUSkinCacheElementsPerFrame"),
+	TEXT("r.SkinCache.MaxGPUElementsPerFrame"),
 	GMaxGPUSkinCacheElementsPerFrame,
-	TEXT("The maximum compute processed skin cache elements per frame. Only used with r.SkinCaching=1\n")
+	TEXT("The maximum compute processed skin cache elements per frame. Only used with r.SkinCache.Mode=1\n")
 	TEXT(" (default is 1000)")
 	);
 
 static int32 GGPUSkinCacheBufferSize = 99 * 1024 * 1024;
 static FAutoConsoleVariableRef CVarMaxGPUSkinCacheBufferSize(
 	TEXT("r.SkinCache.BufferSize"),
-	GMaxGPUSkinCacheElementsPerFrame,
+	GGPUSkinCacheBufferSize,
 	TEXT("The maximum memory used for writing out the vertices in bytes\n")
 	TEXT("Split into GPUSKINCACHE_FRAMES chunks (eg 300 MB means 100 MB for 3 frames)\n")
-	TEXT("Default is 99MB.  Only used with r.SkinCaching=1\n"),
+	TEXT("Default is 99MB.  Only used with r.SkinCache.Mode=1\n"),
 	ECVF_ReadOnly | ECVF_RenderThreadSafe
 	);
 
@@ -145,7 +145,7 @@ public:
 		SkinCacheDebug.Bind(Initializer.ParameterMap, TEXT("SkinCacheDebug"));
 	}
 
-		void SetParameters(const FVertexBufferAndSRV& BoneBuffer, FUniformBufferRHIRef UniformBuffer,
+	void SetParameters(const FVertexBufferAndSRV& BoneBuffer, FUniformBufferRHIRef UniformBuffer,
 		const FVector& MeshOrigin, const FVector& MeshExtension,
 		const FGPUSkinCache::FDispatchData& DispatchData)
 	{
@@ -230,7 +230,7 @@ private:
 };
 
 /** Compute shader that skins a batch of vertices. */
-// @param SkinType 0:normal, 1:with morph target, 2:with APEX cloth (not yet implemented)
+// @param SkinType 0:normal, 1:with morph targets calculated outside the cache, 2:with morph target calculated insde the cache (not yet implemented), 3:with APEX cloth (not yet implemented)
 template <bool bUseExtraBoneInfluencesT, uint32 SkinType>
 class TGPUSkinCacheCS : public FBaseGPUSkinCacheCS
 {
@@ -247,7 +247,7 @@ public:
 		FGlobalShader::ModifyCompilationEnvironment(Platform,OutEnvironment);
 		const uint32 UseExtraBoneInfluences = bUseExtraBoneInfluencesT;
 		OutEnvironment.SetDefine(TEXT("GPUSKIN_USE_EXTRA_INFLUENCES"), UseExtraBoneInfluences);
-		OutEnvironment.SetDefine(TEXT("GPUSKIN_MORPH_BLEND"), SkinType == 1);
+		OutEnvironment.SetDefine(TEXT("GPUSKIN_MORPH_BLEND"), SkinType);
 		OutEnvironment.SetDefine(TEXT("GPUSKIN_RWBUFFER_OFFSET_POSITION"), FGPUSkinCache::RWPositionOffsetInFloats);
 		OutEnvironment.SetDefine(TEXT("GPUSKIN_RWBUFFER_OFFSET_TANGENT_X"), FGPUSkinCache::RWTangentXOffsetInFloats);
 		OutEnvironment.SetDefine(TEXT("GPUSKIN_RWBUFFER_OFFSET_TANGENT_Z"), FGPUSkinCache::RWTangentZOffsetInFloats);
@@ -310,7 +310,7 @@ void FGPUSkinCache::Initialize(FRHICommandListImmediate& RHICmdList)
 	static const auto CVar = IConsoleManager::Get().FindTConsoleVariableDataInt(TEXT("r.BasePassOutputsVelocity"));
 	checkf(!CVar || CVar->GetValueOnRenderThread() == 0, TEXT("GPU Skin caching is not allowed with outputting velocity on base pass (r.BasePassOutputsVelocity=1)"));
 
-	checkf(GEnableGPUSkinCacheShaders, TEXT("GPU Skin caching requires the shaders enabled (r.SkinCacheShaders=1)"));
+	checkf(GEnableGPUSkinCacheShaders, TEXT("GPU Skin caching requires the shaders enabled (r.SkinCache.CompileShaders=1)"));
 
 	int32 NumFloatsPerBuffer = GGPUSkinCacheBufferSize / GPUSKINCACHE_FRAMES / sizeof(float);
 	for (int32 Index = 0; Index < GPUSKINCACHE_FRAMES; ++Index)
@@ -363,7 +363,7 @@ void FGPUSkinCache::Cleanup()
 
 // Kick off compute shader skinning for incoming chunk, return key for fast lookup in draw passes
 int32 FGPUSkinCache::StartCacheMesh(FRHICommandListImmediate& RHICmdList, uint32 FrameNumber, int32 Key, FGPUBaseSkinVertexFactory* VertexFactory,
-	FGPUSkinPassthroughVertexFactory* TargetVertexFactory, const FSkelMeshChunk& BatchElement, FSkeletalMeshObjectGPUSkin* Skin,
+	FGPUSkinPassthroughVertexFactory* TargetVertexFactory, const FSkelMeshSection& BatchElement, FSkeletalMeshObjectGPUSkin* Skin,
 	const FMorphVertexBuffer* MorphVertexBuffer)
 {
 	if (CachedChunksThisFrameCount >= GMaxGPUSkinCacheElementsPerFrame && FrameNumber <= SkinCacheFrameNumber)
@@ -525,7 +525,7 @@ int32 FGPUSkinCache::StartCacheMesh(FRHICommandListImmediate& RHICmdList, uint32
 		FSkeletalMeshResource& SkeletalMeshResource = Skin->GetSkeletalMeshResource();
 		int32 LODIndex = Skin->GetLOD();
 		FStaticLODModel& LodModel = SkeletalMeshResource.LODModels[LODIndex];
-		DispatchData.ChunkIdx = LodModel.FindChunkIndex(BatchElement);
+		DispatchData.SectionIdx = LodModel.FindSectionIndex(BatchElement);
 	}
 
 	DispatchData.NumVertices = NumVertices;
@@ -571,20 +571,15 @@ int32 FGPUSkinCache::StartCacheMesh(FRHICommandListImmediate& RHICmdList, uint32
 
 	if(bRecomputeTangents)
 	{
-		FStaticLODModel* StaticLODModel = MorphVertexBuffer->GetStaticLODModel();
-	
-		// todo: this can be optimized, worse: with multiple chunks per section this is likely not to work
-		const FSkelMeshSection* Section = StaticLODModel->FindSectionForChunk(StaticLODModel->FindChunkIndex(BatchElement));
-		
-		check(Section);
-		if(Section->bRecomputeTangent)
+		if(BatchElement.bRecomputeTangent)
 		{
+			FStaticLODModel* StaticLODModel = MorphVertexBuffer->GetStaticLODModel();
 			FRawStaticIndexBuffer16or32Interface* IndexBuffer = StaticLODModel->MultiSizeIndexContainer.GetIndexBuffer();
 
 			DispatchData.IndexBuffer = IndexBuffer->GetSRV();
 			check(!GSkinCacheSafety || DispatchData.IndexBuffer);
-			DispatchData.NumTriangles = Section->NumTriangles;
-			DispatchData.IndexBufferOffsetValue = Section->BaseIndex;
+			DispatchData.NumTriangles = BatchElement.NumTriangles;
+			DispatchData.IndexBufferOffsetValue = BatchElement.BaseIndex;
 			DispatchUpdateSkinTangents(DispatchData);
 		}
 	}
@@ -745,8 +740,9 @@ void FGPUSkinCache::DispatchSkinCacheProcess(
 	check(DispatchData.FeatureLevel >= ERHIFeatureLevel::SM5);
 
 	SCOPED_DRAW_EVENTF(DispatchData.RHICmdList, SkinCacheDispatch,
-		TEXT("SkinCacheDispatch Chunk=%d InputStreamStart=%d SkinCacheStart=%d Vert=%d Morph=%d/%d StrideInFloats(In/Out):%d/%d"),
-		DispatchData.ChunkIdx, DispatchData.InputStreamStart, DispatchData.SkinCacheStart, DispatchData.NumVertices, DispatchData.MorphBuffer != 0, DispatchData.MorphBufferOffset,
+		TEXT("SkinCacheDispatch<%d,%d> Chunk=%d InputStreamStart=%d SkinCacheStart=%d Vert=%d Morph=%d/%d StrideInFloats(In/Out):%d/%d"),
+		(int32)DispatchData.bExtraBoneInfluences, DispatchData.SkinType,
+		DispatchData.SectionIdx, DispatchData.InputStreamStart, DispatchData.SkinCacheStart, DispatchData.NumVertices, DispatchData.MorphBuffer != 0, DispatchData.MorphBufferOffset,
 		DispatchData.InputStreamStride / sizeof(float), FGPUSkinCache::RWStrideInFloats);
 
 	TShaderMapRef<TGPUSkinCacheCS<true,0> > SkinCacheCS10(GetGlobalShaderMap(DispatchData.FeatureLevel));
@@ -758,8 +754,10 @@ void FGPUSkinCache::DispatchSkinCacheProcess(
 	
 	switch(DispatchData.SkinType)
 	{
-		case 0: Shader = DispatchData.bExtraBoneInfluences ? (FBaseGPUSkinCacheCS*)*SkinCacheCS10 : (FBaseGPUSkinCacheCS*)*SkinCacheCS00; break;
-		case 1: Shader = DispatchData.bExtraBoneInfluences ? (FBaseGPUSkinCacheCS*)*SkinCacheCS11 : (FBaseGPUSkinCacheCS*)*SkinCacheCS01; break;
+		case 0: Shader = DispatchData.bExtraBoneInfluences ? (FBaseGPUSkinCacheCS*)*SkinCacheCS10 : (FBaseGPUSkinCacheCS*)*SkinCacheCS00;
+			break;
+		case 1: Shader = DispatchData.bExtraBoneInfluences ? (FBaseGPUSkinCacheCS*)*SkinCacheCS11 : (FBaseGPUSkinCacheCS*)*SkinCacheCS01;
+			break;
 		default:
 			check(0);
 	}
diff --git a/Engine/Source/Runtime/Engine/Private/InstancedStaticMesh.cpp b/Engine/Source/Runtime/Engine/Private/InstancedStaticMesh.cpp
--- a/Engine/Source/Runtime/Engine/Private/InstancedStaticMesh.cpp
+++ b/Engine/Source/Runtime/Engine/Private/InstancedStaticMesh.cpp
@@ -7,6 +7,7 @@
 #include "EnginePrivate.h"
 #include "NavigationSystemHelpers.h"
 #include "AI/Navigation/NavCollision.h"
+#include "AI/NavigationOctree.h"
 #include "Components/InstancedStaticMeshComponent.h"
 #include "InstancedStaticMesh.h"
 #include "../../Renderer/Private/ScenePrivate.h"
@@ -747,6 +748,36 @@ UInstancedStaticMeshComponent::UInstancedStaticMeshComponent(const FObjectInitia
 	bDisallowMeshPaintPerInstance = true;
 }
 
+
+int32 GetNumShapes(UBodySetup* BodySetup)
+{
+	int32 NumShapes = 1;
+	if (BodySetup)
+	{
+		NumShapes = FMath::Max(BodySetup->AggGeom.GetElementCount(), NumShapes);	//if there's no simple shapes we still have a trimesh so 1 is the min
+	}
+
+	return NumShapes;
+}
+
+int32 GetAggregateIndex(int32 BodyIndex, int32 NumShapes)
+{
+	const int32 BodiesPerBucket = AggregateMaxSize / NumShapes;
+	return BodyIndex / BodiesPerBucket;
+}
+
+int32 GetNumAggregates(int32 NumBodies, int32 NumShapes)
+{
+	if(NumShapes > AggregateMaxSize)
+	{
+		UE_LOG(LogPhysics, Warning, TEXT("Bodies inside foliage can only support up to 128 shapes (per body)"));
+	}
+	
+	const int32 BodiesPerBucket = AggregateMaxSize / NumShapes;
+	return FMath::DivideAndRoundUp<int32>(NumBodies, BodiesPerBucket);
+}
+
+
 #if WITH_EDITOR
 /** Helper class used to preserve lighting/selection state across blueprint reinstancing */
 class FInstancedStaticMeshComponentInstanceData : public FSceneComponentInstanceData
@@ -979,7 +1010,7 @@ void UInstancedStaticMeshComponent::InitInstanceBody(int32 InstanceIdx, FBodyIns
 #if WITH_PHYSX
 	// Create physics body instance.
 	// Aggregates aren't used for static objects
-	auto* Aggregate = (Mobility == EComponentMobility::Movable) ? Aggregates[FMath::DivideAndRoundDown<int32>(InstanceIdx, AggregateMaxSize)] : nullptr;
+	auto* Aggregate = (Mobility == EComponentMobility::Movable) ? Aggregates[GetAggregateIndex(InstanceIdx, GetNumShapes(BodySetup))] : nullptr;
 	InstanceBodyInstance->bAutoWeld = false;	//We don't support this for instanced meshes.
 	InstanceBodyInstance->InitBody(BodySetup, InstanceTransform, this, GetWorld()->GetPhysicsScene(), Aggregate);
 #endif //WITH_PHYSX
@@ -1003,6 +1034,7 @@ void UInstancedStaticMeshComponent::CreateAllInstanceBodies()
 
 		TArray<FTransform> Transforms;
 	    Transforms.Reserve(NumBodies);
+		const int32 NumShapes = GetNumShapes(BodySetup);
     
 	    for (int32 i = 0; i < NumBodies; ++i)
 	    {
@@ -1026,7 +1058,7 @@ void UInstancedStaticMeshComponent::CreateAllInstanceBodies()
 
 				if (Mobility == EComponentMobility::Movable)
 				{
-					Instance->InitBody(BodySetup, InstanceTM, this, PhysScene);
+					Instance->InitBody(BodySetup, InstanceTM, this, PhysScene, Aggregates[GetAggregateIndex(i, NumShapes)] );
 				}
 				else
 				{
@@ -1093,9 +1125,9 @@ void UInstancedStaticMeshComponent::CreatePhysicsState()
 	check(Aggregates.Num() == 0);
 
 	const int32 NumBodies = PerInstanceSMData.Num();
-
+	const int32 NumShapes = GetNumShapes(GetBodySetup());
 	// Aggregates aren't used for static objects
-	const int32 NumAggregates = (Mobility == EComponentMobility::Movable) ? FMath::DivideAndRoundUp<int32>(NumBodies, AggregateMaxSize) : 0;
+	const int32 NumAggregates = (Mobility == EComponentMobility::Movable) ? GetNumAggregates(NumBodies, NumShapes) : 0;
 
 	// Get the scene type from the main BodyInstance
 	const uint32 SceneType = BodyInstance.UseAsyncScene(PhysScene) ? PST_Async : PST_Sync;
@@ -1664,7 +1696,8 @@ void UInstancedStaticMeshComponent::SetupNewInstanceData(FInstancedStaticMeshIns
 		// Aggregates aren't used for static objects
 		if (Mobility == EComponentMobility::Movable)
 		{
-			const int32 AggregateIndex = FMath::DivideAndRoundDown<int32>(InInstanceIndex, AggregateMaxSize);
+			const int32 NumShapes = GetNumShapes(GetBodySetup());
+			const int32 AggregateIndex = GetAggregateIndex(InInstanceIndex, NumShapes);
 			if (AggregateIndex >= Aggregates.Num())
 			{
 				// Get the scene type from the main BodyInstance
diff --git a/Engine/Source/Runtime/Engine/Private/Level.cpp b/Engine/Source/Runtime/Engine/Private/Level.cpp
--- a/Engine/Source/Runtime/Engine/Private/Level.cpp
+++ b/Engine/Source/Runtime/Engine/Private/Level.cpp
@@ -404,7 +404,7 @@ bool ULevel::IsNetActor(const AActor* Actor)
 
 	// If this is a server, use RemoteRole.
 	// If this is a client, use Role.
-	const ENetRole NetRole = (Actor->GetNetMode() < NM_Client) ? Actor->GetRemoteRole() : (ENetRole)Actor->Role;
+	const ENetRole NetRole = (!Actor->IsNetMode(NM_Client)) ? Actor->GetRemoteRole() : (ENetRole)Actor->Role;
 
 	// This test will return true on clients for actors with ROLE_Authority, which might be counterintuitive,
 	// but clients will need to consider these actors in some cases, such as if their bTearOff flag is true.
diff --git a/Engine/Source/Runtime/Engine/Private/LevelTick.cpp b/Engine/Source/Runtime/Engine/Private/LevelTick.cpp
--- a/Engine/Source/Runtime/Engine/Private/LevelTick.cpp
+++ b/Engine/Source/Runtime/Engine/Private/LevelTick.cpp
@@ -12,6 +12,7 @@
 #include "Net/UnrealNetwork.h"
 #include "Collision.h"
 #include "PhysicsPublic.h"
+#include "Tickable.h"
 #include "IHeadMountedDisplay.h"
 
 #include "ParticleDefinitions.h"
diff --git a/Engine/Source/Runtime/Engine/Private/MorphTools.cpp b/Engine/Source/Runtime/Engine/Private/MorphTools.cpp
--- a/Engine/Source/Runtime/Engine/Private/MorphTools.cpp
+++ b/Engine/Source/Runtime/Engine/Private/MorphTools.cpp
@@ -246,71 +246,7 @@ FMorphMeshRawSource::FMorphMeshRawSource( USkeletalMesh* SrcMesh, int32 LODIndex
 	// get the mesh data for the given lod
 	FStaticLODModel& LODModel = ImportedResource->LODModels[LODIndex];
 
-
-	// vertices are packed in this order iot stay consistent
-	// with the indexing used by the FStaticLODModel vertex buffer
-	//
-	//	Chunk0
-	//		Rigid0
-	//		Rigid1
-	//		Soft0
-	//		Soft1
-	//	Chunk1
-	//		Rigid0
-	//		Rigid1
-	//		Soft0
-	//		Soft1
-
-	// iterate over the chunks for the skeletal mesh
-	for( int32 ChunkIdx=0; ChunkIdx < LODModel.Chunks.Num(); ChunkIdx++ )
-	{
-		// each chunk has both rigid and smooth vertices
-		const FSkelMeshChunk& Chunk = LODModel.Chunks[ChunkIdx];
-		// rigid vertices should always be added first so that we are
-		// consistent with the way vertices are packed in the FStaticLODModel vertex buffer
-		for( int32 VertexIdx=0; VertexIdx < Chunk.RigidVertices.Num(); VertexIdx++ )
-		{
-			const FRigidSkinVertex& SourceVertex = Chunk.RigidVertices[VertexIdx];
-			FMorphMeshVertexRaw RawVertex = 
-			{
-				SourceVertex.Position,
-				SourceVertex.TangentX,
-				SourceVertex.TangentY,
-				SourceVertex.TangentZ
-			};
-			Vertices.Add( RawVertex );			
-		}
-		// smooth vertices are added next. The resulting Vertices[] array should
-		// match the FStaticLODModel vertex buffer when indexing vertices
-		for( int32 VertexIdx=0; VertexIdx < Chunk.SoftVertices.Num(); VertexIdx++ )
-		{
-			const FSoftSkinVertex& SourceVertex = Chunk.SoftVertices[VertexIdx];
-			FMorphMeshVertexRaw RawVertex = 
-			{
-				SourceVertex.Position,
-				SourceVertex.TangentX,
-				SourceVertex.TangentY,
-				SourceVertex.TangentZ
-			};
-			Vertices.Add( RawVertex );			
-		}		
-	}
-
-	// Copy the indices manually, since the LODModel's index buffer may have a different alignment.
-	Indices.Empty(LODModel.MultiSizeIndexContainer.GetIndexBuffer()->Num());
-	for(int32 Index = 0;Index < LODModel.MultiSizeIndexContainer.GetIndexBuffer()->Num();Index++)
-	{
-		Indices.Add(LODModel.MultiSizeIndexContainer.GetIndexBuffer()->Get(Index));
-	}
-
-	// copy the wedge point indices
-	if( LODModel.RawPointIndices.GetBulkDataSize() )
-	{
-		WedgePointIndices.Empty( LODModel.RawPointIndices.GetElementCount() );
-		WedgePointIndices.AddUninitialized( LODModel.RawPointIndices.GetElementCount() );
-		FMemory::Memcpy( WedgePointIndices.GetData(), LODModel.RawPointIndices.Lock(LOCK_READ_ONLY), LODModel.RawPointIndices.GetBulkDataSize() );
-		LODModel.RawPointIndices.Unlock();
-	}
+	Initialize(LODModel);
 }
 
 FMorphMeshRawSource::FMorphMeshRawSource(FStaticLODModel& LODModel)
@@ -320,44 +256,22 @@ FMorphMeshRawSource::FMorphMeshRawSource(FStaticLODModel& LODModel)
 
 void FMorphMeshRawSource::Initialize(FStaticLODModel& LODModel)
 {
-	// vertices are packed in this order iot stay consistent
-	// with the indexing used by the FStaticLODModel vertex buffer
+	// vertices are packed to stay consistent with the indexing used by the FStaticLODModel vertex buffer
 	//
 	//	Chunk0
-	//		Rigid0
-	//		Rigid1
 	//		Soft0
 	//		Soft1
 	//	Chunk1
-	//		Rigid0
-	//		Rigid1
 	//		Soft0
 	//		Soft1
 
 	// iterate over the chunks for the skeletal mesh
-	for (int32 ChunkIdx = 0; ChunkIdx < LODModel.Chunks.Num(); ChunkIdx++)
+	for (int32 SectionIdx = 0; SectionIdx < LODModel.Sections.Num(); SectionIdx++)
 	{
-		// each chunk has both rigid and smooth vertices
-		const FSkelMeshChunk& Chunk = LODModel.Chunks[ChunkIdx];
-		// rigid vertices should always be added first so that we are
-		// consistent with the way vertices are packed in the FStaticLODModel vertex buffer
-		for (int32 VertexIdx = 0; VertexIdx < Chunk.RigidVertices.Num(); VertexIdx++)
-		{
-			const FRigidSkinVertex& SourceVertex = Chunk.RigidVertices[VertexIdx];
-			FMorphMeshVertexRaw RawVertex =
-			{
-				SourceVertex.Position,
-				SourceVertex.TangentX,
-				SourceVertex.TangentY,
-				SourceVertex.TangentZ
-			};
-			Vertices.Add(RawVertex);
-		}
-		// smooth vertices are added next. The resulting Vertices[] array should
-		// match the FStaticLODModel vertex buffer when indexing vertices
-		for (int32 VertexIdx = 0; VertexIdx < Chunk.SoftVertices.Num(); VertexIdx++)
+		const FSkelMeshSection& Section = LODModel.Sections[SectionIdx];
+		for (int32 VertexIdx = 0; VertexIdx < Section.SoftVertices.Num(); VertexIdx++)
 		{
-			const FSoftSkinVertex& SourceVertex = Chunk.SoftVertices[VertexIdx];
+			const FSoftSkinVertex& SourceVertex = Section.SoftVertices[VertexIdx];
 			FMorphMeshVertexRaw RawVertex =
 			{
 				SourceVertex.Position,
diff --git a/Engine/Source/Runtime/Engine/Private/NetworkDriver.cpp b/Engine/Source/Runtime/Engine/Private/NetworkDriver.cpp
--- a/Engine/Source/Runtime/Engine/Private/NetworkDriver.cpp
+++ b/Engine/Source/Runtime/Engine/Private/NetworkDriver.cpp
@@ -976,7 +976,7 @@ ENetMode UNetDriver::GetNetMode() const
 {
 	// Special case for PIE - forcing dedicated server behavior
 #if WITH_EDITOR
-	if (IsServer() && World && World->WorldType == EWorldType::PIE)
+	if (World && World->WorldType == EWorldType::PIE && IsServer())
 	{
 		if ( GEngine->GetWorldContextFromWorldChecked(World).RunAsDedicated )
 		{
diff --git a/Engine/Source/Runtime/Engine/Private/Particles/Emitter.cpp b/Engine/Source/Runtime/Engine/Private/Particles/Emitter.cpp
--- a/Engine/Source/Runtime/Engine/Private/Particles/Emitter.cpp
+++ b/Engine/Source/Runtime/Engine/Private/Particles/Emitter.cpp
@@ -168,7 +168,7 @@ void AEmitter::PostInitializeComponents()
 	Super::PostInitializeComponents();
 
 	// Let them die quickly on a dedicated server
-	if (GetNetMode() == NM_DedicatedServer && (GetRemoteRole() == ROLE_None || bNetTemporary))
+	if (IsNetMode(NM_DedicatedServer) && (GetRemoteRole() == ROLE_None || bNetTemporary))
 	{
 		SetLifeSpan(0.2f);
 	}
diff --git a/Engine/Source/Runtime/Engine/Private/Particles/ParticleComponents.cpp b/Engine/Source/Runtime/Engine/Private/Particles/ParticleComponents.cpp
--- a/Engine/Source/Runtime/Engine/Private/Particles/ParticleComponents.cpp
+++ b/Engine/Source/Runtime/Engine/Private/Particles/ParticleComponents.cpp
@@ -1546,9 +1546,9 @@ void UParticleEmitter::Build()
 			if(HighLODLevel->TypeDataModule->RequiresBuild())
 		{
 			FParticleEmitterBuildInfo EmitterBuildInfo;
-//#if WITH_EDITOR
+#if WITH_EDITOR
 				HighLODLevel->CompileModules( EmitterBuildInfo );
-//#endif
+#endif
 				HighLODLevel->TypeDataModule->Build( EmitterBuildInfo );
 			}
 
@@ -4354,7 +4354,7 @@ void UParticleSystemComponent::TickComponent(float DeltaTime, enum ELevelTick Ti
 	DeltaTimeTick = DeltaTime;
 
 	// Bail out if we are running on a dedicated server and we don't want to update on those
-	if ((bUpdateOnDedicatedServer == false) && (GetNetMode() == NM_DedicatedServer))
+	if ((bUpdateOnDedicatedServer == false) && (IsNetMode(NM_DedicatedServer)))
 	{
 		if (bAutoDestroy)
 		{
@@ -7115,7 +7115,7 @@ void AEmitterCameraLensEffectBase::ActivateLensEffect()
 	// only play the camera effect on clients
 	UWorld const* const World = GetWorld();
 	check(World);
-	if( GetNetMode() != NM_DedicatedServer )
+	if( !IsNetMode(NM_DedicatedServer) )
 	{
 		UParticleSystem* PSToActuallySpawn;
 		if( World->GameState && World->GameState->ShouldShowGore() )
diff --git a/Engine/Source/Runtime/Engine/Private/Particles/ParticleModules.cpp b/Engine/Source/Runtime/Engine/Private/Particles/ParticleModules.cpp
--- a/Engine/Source/Runtime/Engine/Private/Particles/ParticleModules.cpp
+++ b/Engine/Source/Runtime/Engine/Private/Particles/ParticleModules.cpp
@@ -4398,7 +4398,7 @@ void UParticleModuleTypeDataGpu::BeginDestroy()
 
 void UParticleModuleTypeDataGpu::Build( FParticleEmitterBuildInfo& EmitterBuildInfo )
 {
-//#if WITH_EDITOR
+#if WITH_EDITOR
 	FVector4Distribution Curve;
 	FComposableFloatDistribution ZeroDistribution;
 	FComposableFloatDistribution OneDistribution;
@@ -4693,7 +4693,7 @@ void UParticleModuleTypeDataGpu::Build( FParticleEmitterBuildInfo& EmitterBuildI
 	// Collision flag.
 	EmitterInfo.bEnableCollision = EmitterBuildInfo.bEnableCollision;
 	EmitterInfo.CollisionMode = (EParticleCollisionMode::Type)EmitterBuildInfo.CollisionMode;
-//#endif
+#endif
 
 
 	// Create or update GPU resources.
diff --git a/Engine/Source/Runtime/Engine/Private/Particles/ParticleModules_Location.cpp b/Engine/Source/Runtime/Engine/Private/Particles/ParticleModules_Location.cpp
--- a/Engine/Source/Runtime/Engine/Private/Particles/ParticleModules_Location.cpp
+++ b/Engine/Source/Runtime/Engine/Private/Particles/ParticleModules_Location.cpp
@@ -2821,93 +2821,59 @@ bool UParticleModuleLocationSkelVertSurface::VertInfluencedByActiveBone(FParticl
 			(FModuleLocationVertSurfaceInstancePayload*)(Owner->GetModuleInstanceData(this));
 
 		// Find the chunk and vertex within that chunk, and skinning type, for this vertex.
-		int32 ChunkIndex;
+		int32 SectionIndex;
 		int32 VertIndex;
-		bool bSoftVertex;
 		bool bHasExtraBoneInfluences;
-		Model.GetChunkAndSkinType(InVertexIndex, ChunkIndex, VertIndex, bSoftVertex, bHasExtraBoneInfluences);
+		Model.GetSectionFromVertexIndex(InVertexIndex, SectionIndex, VertIndex, bHasExtraBoneInfluences);
 
-		check(ChunkIndex < Model.Chunks.Num());
+		check(SectionIndex < Model.Sections.Num());
+
+		FSkelMeshSection& Section = Model.Sections[SectionIndex];
 
 		if (ValidMaterialIndices.Num() > 0)
 		{
-			for (int32 SectIdx = 0; SectIdx < Model.Sections.Num(); SectIdx++)
+			// Does the material match one of the valid ones
+			bool bFound = false;
+			for (int32 ValidIdx = 0; ValidIdx < ValidMaterialIndices.Num(); ValidIdx++)
 			{
-				FSkelMeshSection& Section = Model.Sections[SectIdx];
-				if (Section.ChunkIndex == ChunkIndex)
+				if (ValidMaterialIndices[ValidIdx] == Section.MaterialIndex)
 				{
-					// Does the material match one of the valid ones
-					bool bFound = false;
-					for (int32 ValidIdx = 0; ValidIdx < ValidMaterialIndices.Num(); ValidIdx++)
-					{
-						if (ValidMaterialIndices[ValidIdx] == Section.MaterialIndex)
-						{
-							bFound = true;
-							break;
-						}
-					}
-
-					if (!bFound)
-					{
-						// Material wasn't in the valid list...
-						return false;
-					}
+					bFound = true;
+					break;
 				}
 			}
-		}
 
-		const FSkelMeshChunk& Chunk = Model.Chunks[ChunkIndex];
+			if (!bFound)
+			{
+				// Material wasn't in the valid list...
+				return false;
+			}
+		}
 
 		return Model.VertexBufferGPUSkin.HasExtraBoneInfluences()
-			? VertInfluencedByActiveBoneTyped<true>(bSoftVertex, Model, Chunk, VertIndex, InSkelMeshComponent, InstancePayload, OutBoneIndex)
-			: VertInfluencedByActiveBoneTyped<false>(bSoftVertex, Model, Chunk, VertIndex, InSkelMeshComponent, InstancePayload, OutBoneIndex);
+			? VertInfluencedByActiveBoneTyped<true>(Model, Section, VertIndex, InSkelMeshComponent, InstancePayload, OutBoneIndex)
+			: VertInfluencedByActiveBoneTyped<false>(Model, Section, VertIndex, InSkelMeshComponent, InstancePayload, OutBoneIndex);
 	}
 	return false;
 }
 
 template<bool bExtraBoneInfluencesT>
-bool UParticleModuleLocationSkelVertSurface::VertInfluencedByActiveBoneTyped(bool bSoftVertex, FStaticLODModel& Model, const FSkelMeshChunk& Chunk, int32 VertIndex, USkeletalMeshComponent* InSkelMeshComponent, FModuleLocationVertSurfaceInstancePayload* InstancePayload, int32* OutBoneIndex)
+bool UParticleModuleLocationSkelVertSurface::VertInfluencedByActiveBoneTyped(FStaticLODModel& Model, const FSkelMeshSection& Section, int32 VertIndex, USkeletalMeshComponent* InSkelMeshComponent, FModuleLocationVertSurfaceInstancePayload* InstancePayload, int32* OutBoneIndex)
 {
 	const TArray<int32>& MasterBoneMap = InSkelMeshComponent->GetMasterBoneMap();
 	// Do soft skinning for this vertex.
-	if(bSoftVertex)
-	{
-		const TGPUSkinVertexBase<bExtraBoneInfluencesT>* SrcSoftVertex = Model.VertexBufferGPUSkin.GetVertexPtr<bExtraBoneInfluencesT>(Chunk.GetSoftVertexBufferIndex()+VertIndex);
+	const TGPUSkinVertexBase<bExtraBoneInfluencesT>* SrcSoftVertex = Model.VertexBufferGPUSkin.GetVertexPtr<bExtraBoneInfluencesT>(Section.GetVertexBufferIndex()+VertIndex);
 
 #if !PLATFORM_LITTLE_ENDIAN
-		// uint8[] elements in LOD.VertexBufferGPUSkin have been swapped for VET_UBYTE4 vertex stream use
-		for(int32 InfluenceIndex = MAX_INFLUENCES-1;InfluenceIndex >=  MAX_INFLUENCES-Chunk.MaxBoneInfluences;InfluenceIndex--)
+	// uint8[] elements in LOD.VertexBufferGPUSkin have been swapped for VET_UBYTE4 vertex stream use
+	for(int32 InfluenceIndex = MAX_INFLUENCES-1;InfluenceIndex >=  MAX_INFLUENCES- Section.MaxBoneInfluences;InfluenceIndex--)
 #else
-		for(int32 InfluenceIndex = 0;InfluenceIndex < Chunk.MaxBoneInfluences;InfluenceIndex++)
+	for(int32 InfluenceIndex = 0;InfluenceIndex < Section.MaxBoneInfluences;InfluenceIndex++)
 #endif
-		{
-			int32 BoneIndex = Chunk.BoneMap[SrcSoftVertex->InfluenceBones[InfluenceIndex]];
-			if(InSkelMeshComponent->MasterPoseComponent.IsValid())
-			{		
-				check(MasterBoneMap.Num() == InSkelMeshComponent->SkeletalMesh->RefSkeleton.GetNum());
-				BoneIndex = MasterBoneMap[BoneIndex];
-			}
-
-			if(!InstancePayload->NumValidAssociatedBoneIndices || InstancePayload->ValidAssociatedBoneIndices.Contains(BoneIndex))
-			{
-				if(OutBoneIndex)
-				{
-					*OutBoneIndex = BoneIndex;
-				}
-
-				return true;
-			}
-		}
-	}
-	// Do rigid (one-influence) skinning for this vertex.
-	else
 	{
-		const int32 RigidInfluenceIndex = SkinningTools::GetRigidInfluenceIndex();
-		const TGPUSkinVertexBase<bExtraBoneInfluencesT>* SrcRigidVertex = Model.VertexBufferGPUSkin.GetVertexPtr<bExtraBoneInfluencesT>(Chunk.GetRigidVertexBufferIndex()+VertIndex);
-
-		int32 BoneIndex = Chunk.BoneMap[SrcRigidVertex->InfluenceBones[RigidInfluenceIndex]];
+		int32 BoneIndex = Section.BoneMap[SrcSoftVertex->InfluenceBones[InfluenceIndex]];
 		if(InSkelMeshComponent->MasterPoseComponent.IsValid())
-		{
+		{		
 			check(MasterBoneMap.Num() == InSkelMeshComponent->SkeletalMesh->RefSkeleton.GetNum());
 			BoneIndex = MasterBoneMap[BoneIndex];
 		}
@@ -2918,9 +2884,11 @@ bool UParticleModuleLocationSkelVertSurface::VertInfluencedByActiveBoneTyped(boo
 			{
 				*OutBoneIndex = BoneIndex;
 			}
+
 			return true;
 		}
 	}
 
+
 	return false;
 }
diff --git a/Engine/Source/Runtime/Engine/Private/PhysicsEngine/BodyInstance.cpp b/Engine/Source/Runtime/Engine/Private/PhysicsEngine/BodyInstance.cpp
--- a/Engine/Source/Runtime/Engine/Private/PhysicsEngine/BodyInstance.cpp
+++ b/Engine/Source/Runtime/Engine/Private/PhysicsEngine/BodyInstance.cpp
@@ -3,6 +3,7 @@
 #include "EnginePrivate.h"
 #include "PhysicsPublic.h"
 #include "Collision.h"
+#include "PhysicsEngine/ConstraintInstance.h"
 
 #include "MessageLog.h"
 
@@ -4235,7 +4236,7 @@ bool FBodyInstance::InternalSweepPhysX(struct FHitResult& OutHit, const FVector&
 }
 #endif //WITH_PHYSX
 
-float FBodyInstance::GetDistanceToBody(const FVector& Point, FVector& OutPointOnBody) const
+bool FBodyInstance::GetSquaredDistanceToBody(const FVector& Point, float& OutDistanceSquared, FVector& OutPointOnBody) const
 {
 	OutPointOnBody = Point;
 	float ReturnDistance = -1.f;
@@ -4297,11 +4298,21 @@ float FBodyInstance::GetDistanceToBody(const FVector& Point, FVector& OutPointOn
 		UE_LOG(LogPhysics, Warning, TEXT("GetDistanceToBody: Component (%s) has no simple collision and cannot be queried for closest point."), OwnerComponent.Get() ? *OwnerComponent->GetPathName() : TEXT("NONE"));
 	}
 
-	return (bFoundValidBody ? FMath::Sqrt(MinDistanceSqr) : -1.f);
+	if (bFoundValidBody)
+	{
+		OutDistanceSquared = MinDistanceSqr;
+	}
+	return bFoundValidBody;
 
 	//@TODO: BOX2D: Implement DistanceToBody
 }
 
+float FBodyInstance::GetDistanceToBody(const FVector& Point, FVector& OutPointOnBody) const
+{
+	float DistanceSqr = -1.f;
+	return (GetSquaredDistanceToBody(Point, DistanceSqr, OutPointOnBody) ? FMath::Sqrt(DistanceSqr) : -1.f);
+}
+
 template <typename AllocatorType>
 bool FBodyInstance::OverlapTestForBodiesImpl(const FVector& Pos, const FQuat& Rot, const TArray<FBodyInstance*, AllocatorType>& Bodies) const
 {
diff --git a/Engine/Source/Runtime/Engine/Private/PhysicsEngine/PhysXSupport.cpp b/Engine/Source/Runtime/Engine/Private/PhysicsEngine/PhysXSupport.cpp
--- a/Engine/Source/Runtime/Engine/Private/PhysicsEngine/PhysXSupport.cpp
+++ b/Engine/Source/Runtime/Engine/Private/PhysicsEngine/PhysXSupport.cpp
@@ -14,6 +14,7 @@
 #include "PhysXSupport.h"
 #include "Components/DestructibleComponent.h"
 #include "PhysicsEngine/PhysicsAsset.h"
+#include "PhysicsEngine/ConstraintInstance.h"
 
 PxFoundation*			GPhysXFoundation = NULL;
 PxProfileZoneManager*	GPhysXProfileZoneManager = NULL;
