Copying //UE4/Dev-Framework to Dev-Main (//UE4/Dev-Main)

diff --git a/Engine/Source/Runtime/Engine/Private/PhysicsVolume.cpp b/Engine/Source/Runtime/Engine/Private/PhysicsVolume.cpp
--- a/Engine/Source/Runtime/Engine/Private/PhysicsVolume.cpp
+++ b/Engine/Source/Runtime/Engine/Private/PhysicsVolume.cpp
@@ -92,9 +92,21 @@ bool APhysicsVolume::IsOverlapInVolume(const class USceneComponent& TestComponen
 	if (!bPhysicsOnContact)
 	{
 		FVector ClosestPoint(0.f);
+		// If there is no primitive component as root we consider you inside the volume. This is odd, but the behavior 
+		// has existed for a long time, so keeping it this way
 		UPrimitiveComponent* RootPrimitive = Cast<UPrimitiveComponent>(GetRootComponent());
-		const float DistToCollision = RootPrimitive ? RootPrimitive->GetDistanceToCollision(TestComponent.GetComponentLocation(), ClosestPoint) : 0.f;
-		bInsideVolume = (DistToCollision == 0.f);
+		if (RootPrimitive)
+		{
+			float DistToCollisionSqr = -1.f;
+			if (RootPrimitive->GetSquaredDistanceToCollision(TestComponent.GetComponentLocation(), DistToCollisionSqr, ClosestPoint))
+			{
+				bInsideVolume = (DistToCollisionSqr == 0.f);
+			}
+			else
+			{
+				bInsideVolume = false;
+			}
+		}
 	}
 
 	return bInsideVolume;
diff --git a/Engine/Source/Runtime/Engine/Private/PlayerCameraManager.cpp b/Engine/Source/Runtime/Engine/Private/PlayerCameraManager.cpp
--- a/Engine/Source/Runtime/Engine/Private/PlayerCameraManager.cpp
+++ b/Engine/Source/Runtime/Engine/Private/PlayerCameraManager.cpp
@@ -816,7 +816,7 @@ void APlayerCameraManager::UpdateCamera(float DeltaTime)
 	{
 		DoUpdateCamera(DeltaTime);
 
-		if (bShouldSendClientSideCameraUpdate && GetNetMode() == NM_Client)
+		if (bShouldSendClientSideCameraUpdate && IsNetMode(NM_Client))
 		{
 			SCOPE_CYCLE_COUNTER(STAT_ServerUpdateCamera);
 
diff --git a/Engine/Source/Runtime/Engine/Private/PlayerController.cpp b/Engine/Source/Runtime/Engine/Private/PlayerController.cpp
--- a/Engine/Source/Runtime/Engine/Private/PlayerController.cpp
+++ b/Engine/Source/Runtime/Engine/Private/PlayerController.cpp
@@ -149,12 +149,12 @@ bool APlayerController::DestroyNetworkActorHandled()
 bool APlayerController::IsLocalController() const
 {
 	// Never local on dedicated server, always local on clients. IsServerOnly() and IsClientOnly() are checked at compile time and optimized out appropriately.
-	if (FGenericPlatformProperties::IsServerOnly())
+	if (FPlatformProperties::IsServerOnly())
 	{
-		check(!bIsLocalPlayerController);
+		checkSlow(!bIsLocalPlayerController);
 		return false;
 	}
-	else if (FGenericPlatformProperties::IsClientOnly())
+	else if (FPlatformProperties::IsClientOnly())
 	{
 		bIsLocalPlayerController = true;
 		return true;
@@ -170,7 +170,7 @@ bool APlayerController::IsLocalController() const
 	if (NetMode == NM_DedicatedServer)
 	{
 		// This is still checked for the PIE case, which would not be caught in the IsServerOnly() check above.
-		check(!bIsLocalPlayerController);
+		checkSlow(!bIsLocalPlayerController);
 		return false;
 	}
 
@@ -1510,7 +1510,7 @@ void APlayerController::SendClientAdjustment()
 	// Server sends updates.
 	// Note: we do this for both the pawn and spectator in case an implementation has a networked spectator.
 	APawn* RemotePawn = GetPawnOrSpectator();
-	if (RemotePawn && (GetNetMode() < NM_Client) && (RemotePawn->GetRemoteRole() == ROLE_AutonomousProxy))
+	if (RemotePawn && (RemotePawn->GetRemoteRole() == ROLE_AutonomousProxy) && !IsNetMode(NM_Client))
 	{
 		INetworkPredictionInterface* NetworkPredictionInterface = Cast<INetworkPredictionInterface>(RemotePawn->GetMovementComponent());
 		if (NetworkPredictionInterface)
@@ -3877,7 +3877,7 @@ void APlayerController::TickActor( float DeltaSeconds, ELevelTick TickType, FAct
 
 	//root of tick hierarchy
 
-	if ((GetNetMode() < NM_Client) && (GetRemoteRole() == ROLE_AutonomousProxy) && !IsLocalPlayerController())
+	if ((GetRemoteRole() == ROLE_AutonomousProxy) && !IsNetMode(NM_Client) && !IsLocalPlayerController())
 	{
 		// force physics update for clients that aren't sending movement updates in a timely manner 
 		// this prevents cheats associated with artificially induced ping spikes
diff --git a/Engine/Source/Runtime/Engine/Private/PrimitiveComponentPhysics.cpp b/Engine/Source/Runtime/Engine/Private/PrimitiveComponentPhysics.cpp
--- a/Engine/Source/Runtime/Engine/Private/PrimitiveComponentPhysics.cpp
+++ b/Engine/Source/Runtime/Engine/Private/PrimitiveComponentPhysics.cpp
@@ -881,17 +881,17 @@ FBodyInstance* UPrimitiveComponent::GetBodyInstance(FName BoneName, bool bGetWel
 	return const_cast<FBodyInstance*>(&BodyInstance);
 }
 
-float UPrimitiveComponent::GetDistanceToCollision(const FVector& Point, FVector& ClosestPointOnCollision) const
+bool UPrimitiveComponent::GetSquaredDistanceToCollision(const FVector& Point, float& OutSquaredDistance, FVector& OutClosestPointOnCollision) const
 {
-	ClosestPointOnCollision=Point;
+	OutClosestPointOnCollision = Point;
 
 	FBodyInstance* BodyInst = GetBodyInstance();
-	if(BodyInst != NULL)
+	if (BodyInst != nullptr)
 	{
-		return BodyInst->GetDistanceToBody(Point, ClosestPointOnCollision);
+		return BodyInst->GetSquaredDistanceToBody(Point, OutSquaredDistance, OutClosestPointOnCollision);
 	}
 
-	return -1.f;
+	return false;
 }
 
 float UPrimitiveComponent::GetClosestPointOnCollision(const FVector& Point, FVector& OutPointOnBody, FName BoneName) const
diff --git a/Engine/Source/Runtime/Engine/Private/SkeletalMesh.cpp b/Engine/Source/Runtime/Engine/Private/SkeletalMesh.cpp
--- a/Engine/Source/Runtime/Engine/Private/SkeletalMesh.cpp
+++ b/Engine/Source/Runtime/Engine/Private/SkeletalMesh.cpp
@@ -51,6 +51,9 @@ DEFINE_LOG_CATEGORY(LogSkeletalMesh);
 
 DECLARE_CYCLE_STAT(TEXT("GetShadowShapes"), STAT_GetShadowShapes, STATGROUP_Anim);
 
+const FGuid FSkeletalMeshCustomVersion::GUID(0xD78A4A00, 0xE8584697, 0xBAA819B5, 0x487D46B4);
+FCustomVersionRegistration GRegisterSkeletalMeshCustomVersion(FSkeletalMeshCustomVersion::GUID, FSkeletalMeshCustomVersion::LatestVersion, TEXT("SkeletalMeshVer"));
+
 #if WITH_APEX_CLOTHING
 /*-----------------------------------------------------------------------------
 	utility functions for apex clothing 
@@ -868,31 +871,21 @@ FArchive& operator<<(FArchive& Ar,FSoftSkinVertex& V)
 	return Ar;
 }
 
-/*-----------------------------------------------------------------------------
-	FRigidSkinVertex
------------------------------------------------------------------------------*/
-
-/**
-* Serializer
-*
-* @param Ar - archive to serialize with
-* @param V - vertex to serialize
-* @return archive that was used
-*/
-FArchive& operator<<(FArchive& Ar,FRigidSkinVertex& V)
+bool FSoftSkinVertex::GetRigidWeightBone(uint8& OutBoneIndex) const
 {
-	Ar << V.Position;
-	Ar << V.TangentX << V.TangentY << V.TangentZ;
+	bool bIsRigid = false;
 
-	for( int32 UVIdx = 0; UVIdx < MAX_TEXCOORDS; ++UVIdx )
+	for (int32 WeightIdx = 0; WeightIdx < MAX_TOTAL_INFLUENCES; WeightIdx++)
 	{
-		Ar << V.UVs[UVIdx];
+		if (InfluenceWeights[WeightIdx] == 255)
+		{
+			bIsRigid = true;
+			OutBoneIndex = InfluenceBones[WeightIdx];
+			break;
+		}
 	}
 
-	Ar << V.Color;
-	Ar << V.Bone;
-
-	return Ar;
+	return bIsRigid;
 }
 
 /*-----------------------------------------------------------------------------
@@ -1096,42 +1089,6 @@ FMultiSizeIndexContainer& FMultiSizeIndexContainer::operator=(const FMultiSizeIn
 }
 #endif
 
-/**
-* Serialize this class
-* @param Ar - archive to serialize to
-* @param C - skel mesh chunk to serialize
-*/
-FArchive& operator<<(FArchive& Ar,FSkelMeshChunk& C)
-{
-	FStripDataFlags StripFlags( Ar );
-
-	if( !StripFlags.IsDataStrippedForServer() )
-	{
-		// This is so that BaseVertexIndex is never set to anything else that 0 (for safety)
-		Ar << C.BaseVertexIndex;
-	}
-	if( !StripFlags.IsEditorDataStripped() )
-	{		
-		Ar << C.RigidVertices;
-		Ar << C.SoftVertices;
-	}
-	Ar << C.BoneMap;
-	Ar << C.NumRigidVertices;
-	Ar << C.NumSoftVertices;
-	Ar << C.MaxBoneInfluences;
-
-
-	if( Ar.UE4Ver() >= VER_UE4_APEX_CLOTH )
-	{
-		Ar << C.ApexClothMappingData;
-		Ar << C.PhysicalMeshVertices;
-		Ar << C.PhysicalMeshNormals;
-		Ar << C.CorrespondClothAssetIndex;
-		Ar << C.ClothAssetSubmeshIndex;
-	}
-
-	return Ar;
-}
 
 /*-----------------------------------------------------------------------------
 	FSkelMeshSection
@@ -1162,6 +1119,59 @@ const FGuid FRecomputeTangentCustomVersion::GUID(0x5579F886, 0x933A4C1F, 0x83BA0
 // Register the custom version with core
 FCustomVersionRegistration GRegisterRecomputeTangentCustomVersion(FRecomputeTangentCustomVersion::GUID, FRecomputeTangentCustomVersion::LatestVersion, TEXT("RecomputeTangentCustomVer"));
 
+/** Legacy 'rigid' skin vertex */
+struct FLegacyRigidSkinVertex
+{
+	FVector			Position;
+	FPackedNormal	TangentX,	// Tangent, U-direction
+		TangentY,	// Binormal, V-direction
+		TangentZ;	// Normal
+	FVector2D		UVs[MAX_TEXCOORDS]; // UVs
+	FColor			Color;		// Vertex color.
+	uint8			Bone;
+
+	friend FArchive& operator<<(FArchive& Ar, FLegacyRigidSkinVertex& V)
+	{
+		Ar << V.Position;
+		Ar << V.TangentX << V.TangentY << V.TangentZ;
+
+		for (int32 UVIdx = 0; UVIdx < MAX_TEXCOORDS; ++UVIdx)
+		{
+			Ar << V.UVs[UVIdx];
+		}
+
+		Ar << V.Color;
+		Ar << V.Bone;
+
+		return Ar;
+	}
+
+	/** Util to convert from legacy */
+	void ConvertToSoftVert(FSoftSkinVertex& DestVertex)
+	{
+		DestVertex.Position = Position;
+		DestVertex.TangentX = TangentX;
+		DestVertex.TangentY = TangentY;
+		DestVertex.TangentZ = TangentZ;
+		// store the sign of the determinant in TangentZ.W
+		DestVertex.TangentZ.Vector.W = GetBasisDeterminantSignByte(TangentX, TangentY, TangentZ);
+
+		// copy all texture coordinate sets
+		FMemory::Memcpy(DestVertex.UVs, UVs, sizeof(FVector2D)*MAX_TEXCOORDS);
+
+		DestVertex.Color = Color;
+		DestVertex.InfluenceBones[0] = Bone;
+		DestVertex.InfluenceWeights[0] = 255;
+		for (int32 InfluenceIndex = 1; InfluenceIndex < MAX_TOTAL_INFLUENCES; InfluenceIndex++)
+		{
+			DestVertex.InfluenceBones[InfluenceIndex] = 0;
+			DestVertex.InfluenceWeights[InfluenceIndex] = 0;
+		}
+	}
+};
+
+
+
 // Serialization.
 FArchive& operator<<(FArchive& Ar,FSkelMeshSection& S)
 {		
@@ -1171,7 +1181,14 @@ FArchive& operator<<(FArchive& Ar,FSkelMeshSection& S)
 	FStripDataFlags StripFlags( Ar );
 
 	Ar << S.MaterialIndex;
-	Ar << S.ChunkIndex;
+
+	Ar.UsingCustomVersion(FSkeletalMeshCustomVersion::GUID);
+	if (Ar.CustomVer(FSkeletalMeshCustomVersion::GUID) < FSkeletalMeshCustomVersion::CombineSectionWithChunk)
+	{
+		uint16 DummyChunkIndex;
+		Ar << DummyChunkIndex;
+	}
+
 	if (!StripFlags.IsDataStrippedForServer())
 	{
 		Ar << S.BaseIndex;
@@ -1202,6 +1219,62 @@ FArchive& operator<<(FArchive& Ar,FSkelMeshSection& S)
 		Ar << S.bRecomputeTangent;
 	}
 
+	if (Ar.CustomVer(FSkeletalMeshCustomVersion::GUID) >= FSkeletalMeshCustomVersion::CombineSectionWithChunk)
+	{
+
+		if (!StripFlags.IsDataStrippedForServer())
+		{
+			// This is so that BaseVertexIndex is never set to anything else that 0 (for safety)
+			Ar << S.BaseVertexIndex;
+		}
+
+		if (!StripFlags.IsEditorDataStripped())
+		{
+			// For backwards compat, read rigid vert array into array
+			TArray<FLegacyRigidSkinVertex> LegacyRigidVertices;
+			if (Ar.IsLoading() && Ar.CustomVer(FSkeletalMeshCustomVersion::GUID) < FSkeletalMeshCustomVersion::CombineSoftAndRigidVerts)
+			{
+				Ar << LegacyRigidVertices;
+			}
+
+			Ar << S.SoftVertices;
+
+			// Once we have read in SoftVertices, convert and insert legacy rigid verts (if present) at start
+			const int32 NumRigidVerts = LegacyRigidVertices.Num();
+			if (NumRigidVerts > 0 && Ar.CustomVer(FSkeletalMeshCustomVersion::GUID) < FSkeletalMeshCustomVersion::CombineSoftAndRigidVerts)
+			{
+				S.SoftVertices.InsertUninitialized(0, NumRigidVerts);
+
+				for (int32 VertIdx = 0; VertIdx < NumRigidVerts; VertIdx++)
+				{
+					LegacyRigidVertices[VertIdx].ConvertToSoftVert(S.SoftVertices[VertIdx]);
+				}
+			}
+		}
+
+		Ar << S.BoneMap;
+
+		// Removed NumRigidVertices and NumSoftVertices, just use array size
+		if (Ar.CustomVer(FSkeletalMeshCustomVersion::GUID) < FSkeletalMeshCustomVersion::CombineSoftAndRigidVerts)
+		{
+			int32 DummyNumRigidVerts, DummyNumSoftVerts;
+			Ar << DummyNumRigidVerts;
+			Ar << DummyNumSoftVerts;
+
+			if (DummyNumRigidVerts + DummyNumSoftVerts != S.SoftVertices.Num())
+			{
+				UE_LOG(LogSkeletalMesh, Error, TEXT("Legacy NumSoftVerts + NumRigidVerts != SoftVertices.Num()"));
+			}
+		}
+
+		Ar << S.MaxBoneInfluences;
+		Ar << S.ApexClothMappingData;
+		Ar << S.PhysicalMeshVertices;
+		Ar << S.PhysicalMeshNormals;
+		Ar << S.CorrespondClothAssetIndex;
+		Ar << S.ClothAssetSubmeshIndex;
+	}
+
 	return Ar;
 }
 
@@ -1232,13 +1305,112 @@ void FMorphTargetVertexInfoBuffers::InitRHI()
 void FMorphTargetVertexInfoBuffers::ReleaseRHI()
 {
 	PerVertexInfoVB.SafeRelease();
+	PerVertexInfoSRV.SafeRelease();
 	FlattenedDeltasVB.SafeRelease();
+	FlattenedDeltasSRV.SafeRelease();
 }
 
 /*-----------------------------------------------------------------------------
 	FStaticLODModel
 -----------------------------------------------------------------------------*/
 
+/** Legacy Chunk struct, now merged with FSkelMeshSection */
+struct FLegacySkelMeshChunk
+{
+	uint32 BaseVertexIndex;
+	TArray<FSoftSkinVertex> SoftVertices;
+	TArray<FApexClothPhysToRenderVertData> ApexClothMappingData;
+	TArray<FVector> PhysicalMeshVertices;
+	TArray<FVector> PhysicalMeshNormals;
+	TArray<FBoneIndexType> BoneMap;
+	int32 MaxBoneInfluences;
+
+	int16 CorrespondClothAssetIndex;
+	int16 ClothAssetSubmeshIndex;
+
+	FLegacySkelMeshChunk()
+		: BaseVertexIndex(0)
+		, MaxBoneInfluences(4)
+		, CorrespondClothAssetIndex(INDEX_NONE)
+		, ClothAssetSubmeshIndex(INDEX_NONE)
+	{}
+
+	void CopyToSection(FSkelMeshSection& Section)
+	{
+		Section.BaseVertexIndex = BaseVertexIndex;
+		Section.SoftVertices = SoftVertices;
+		Section.ApexClothMappingData = ApexClothMappingData;
+		Section.PhysicalMeshVertices = PhysicalMeshVertices;
+		Section.PhysicalMeshNormals = PhysicalMeshNormals;
+		Section.BoneMap = BoneMap;
+		Section.CorrespondClothAssetIndex = CorrespondClothAssetIndex;
+		Section.ClothAssetSubmeshIndex = ClothAssetSubmeshIndex;
+	}
+
+
+	friend FArchive& operator<<(FArchive& Ar, FLegacySkelMeshChunk& C)
+	{
+		FStripDataFlags StripFlags(Ar);
+
+		if (!StripFlags.IsDataStrippedForServer())
+		{
+			// This is so that BaseVertexIndex is never set to anything else that 0 (for safety)
+			Ar << C.BaseVertexIndex;
+		}
+		if (!StripFlags.IsEditorDataStripped())
+		{
+			// For backwards compat, read rigid vert array into array
+			TArray<FLegacyRigidSkinVertex> LegacyRigidVertices;
+			if (Ar.IsLoading() && Ar.CustomVer(FSkeletalMeshCustomVersion::GUID) < FSkeletalMeshCustomVersion::CombineSoftAndRigidVerts)
+			{
+				Ar << LegacyRigidVertices;
+			}
+
+			Ar << C.SoftVertices;
+
+			// Once we have read in SoftVertices, convert and insert legacy rigid verts (if present) at start
+			const int32 NumRigidVerts = LegacyRigidVertices.Num();
+			if (NumRigidVerts > 0 && Ar.CustomVer(FSkeletalMeshCustomVersion::GUID) < FSkeletalMeshCustomVersion::CombineSoftAndRigidVerts)
+			{
+				C.SoftVertices.InsertUninitialized(0, NumRigidVerts);
+
+				for (int32 VertIdx = 0; VertIdx < NumRigidVerts; VertIdx++)
+				{
+					LegacyRigidVertices[VertIdx].ConvertToSoftVert(C.SoftVertices[VertIdx]);
+				}
+			}
+		}
+		Ar << C.BoneMap;
+
+		// Removed NumRigidVertices and NumSoftVertices, just use array size
+		if (Ar.CustomVer(FSkeletalMeshCustomVersion::GUID) < FSkeletalMeshCustomVersion::CombineSoftAndRigidVerts)
+		{
+			int32 DummyNumRigidVerts, DummyNumSoftVerts;
+			Ar << DummyNumRigidVerts;
+			Ar << DummyNumSoftVerts;
+
+			if (DummyNumRigidVerts + DummyNumSoftVerts != C.SoftVertices.Num())
+			{
+				UE_LOG(LogSkeletalMesh, Error, TEXT("Legacy NumSoftVerts + NumRigidVerts != SoftVertices.Num()"));
+			}
+		}
+
+		Ar << C.MaxBoneInfluences;
+
+
+		if (Ar.UE4Ver() >= VER_UE4_APEX_CLOTH)
+		{
+			Ar << C.ApexClothMappingData;
+			Ar << C.PhysicalMeshVertices;
+			Ar << C.PhysicalMeshNormals;
+			Ar << C.CorrespondClothAssetIndex;
+			Ar << C.ClothAssetSubmeshIndex;
+		}
+
+		return Ar;
+	}
+};
+
 void FStaticLODModel::Serialize( FArchive& Ar, UObject* Owner, int32 Idx )
 {
 	DECLARE_SCOPE_CYCLE_COUNTER( TEXT("FStaticLODModel::Serialize"), STAT_StaticLODModel_Serialize, STATGROUP_LoadTime );
@@ -1257,7 +1429,21 @@ void FStaticLODModel::Serialize( FArchive& Ar, UObject* Owner, int32 Idx )
 	Ar << Sections;
 	MultiSizeIndexContainer.Serialize(Ar, bKeepBuffersInCPUMemory);
 	Ar << ActiveBoneIndices;
-	Ar << Chunks;
+
+	// Array of Sections for backwards compat
+	Ar.UsingCustomVersion(FSkeletalMeshCustomVersion::GUID);
+	if (Ar.CustomVer(FSkeletalMeshCustomVersion::GUID) < FSkeletalMeshCustomVersion::CombineSectionWithChunk)
+	{
+		TArray<FLegacySkelMeshChunk> LegacyChunks;
+
+		Ar << LegacyChunks;
+
+		check(LegacyChunks.Num() == Sections.Num());
+		for (int32 ChunkIdx = 0; ChunkIdx < LegacyChunks.Num(); ChunkIdx++)
+		{
+			LegacyChunks[ChunkIdx].CopyToSection(Sections[ChunkIdx]);
+		}
+	}
 	
 	// no longer in use
 	{
@@ -1448,40 +1634,28 @@ int32 FStaticLODModel::GetTotalFaces() const
 	return TotalFaces;
 }
 
-void FStaticLODModel::GetChunkAndSkinType(int32 InVertIndex, int32& OutChunkIndex, int32& OutVertIndex, bool& bOutSoftVert, bool& bOutHasExtraBoneInfluences) const
+void FStaticLODModel::GetSectionFromVertexIndex(int32 InVertIndex, int32& OutSectionIndex, int32& OutVertIndex, bool& bOutHasExtraBoneInfluences) const
 {
-	OutChunkIndex = 0;
+	OutSectionIndex = 0;
 	OutVertIndex = 0;
-	bOutSoftVert = false;
 	bOutHasExtraBoneInfluences = false;
 
 	int32 VertCount = 0;
 
 	// Iterate over each chunk
-	for(int32 ChunkCount = 0; ChunkCount < Chunks.Num(); ChunkCount++)
+	for(int32 SectionCount = 0; SectionCount < Sections.Num(); SectionCount++)
 	{
-		const FSkelMeshChunk& Chunk = Chunks[ChunkCount];
-		OutChunkIndex = ChunkCount;
-
-		// Is it in Rigid vertex range?
-		if(InVertIndex < VertCount + Chunk.GetNumRigidVertices())
-		{
-			OutVertIndex = InVertIndex - VertCount;
-			bOutSoftVert = false;
-			bOutHasExtraBoneInfluences = VertexBufferGPUSkin.HasExtraBoneInfluences();
-			return;
-		}
-		VertCount += Chunk.GetNumRigidVertices();
+		const FSkelMeshSection& Section = Sections[SectionCount];
+		OutSectionIndex = SectionCount;
 
 		// Is it in Soft vertex range?
-		if(InVertIndex < VertCount + Chunk.GetNumSoftVertices())
+		if(InVertIndex < VertCount + Section.GetNumVertices())
 		{
 			OutVertIndex = InVertIndex - VertCount;
-			bOutSoftVert = true;
 			bOutHasExtraBoneInfluences = VertexBufferGPUSkin.HasExtraBoneInfluences();
 			return;
 		}
-		VertCount += Chunk.GetNumSoftVertices();
+		VertCount += Section.GetNumVertices();
 	}
 
 	// InVertIndex should always be in some chunk!
@@ -1492,54 +1666,29 @@ void FStaticLODModel::GetVertices(TArray<FSoftSkinVertex>& Vertices) const
 {
 	Vertices.Empty(NumVertices);
 	Vertices.AddUninitialized(NumVertices);
-		
+	
 	// Initialize the vertex data
 	// All chunks are combined into one (rigid first, soft next)
 	FSoftSkinVertex* DestVertex = (FSoftSkinVertex*)Vertices.GetData();
-	for(int32 ChunkIndex = 0;ChunkIndex < Chunks.Num();ChunkIndex++)
-	{
-		const FSkelMeshChunk& Chunk = Chunks[ChunkIndex];
-		//check(Chunk.NumRigidVertices == Chunk.RigidVertices.Num());
-		//check(Chunk.NumSoftVertices == Chunk.SoftVertices.Num());
-		for(int32 VertexIndex = 0;VertexIndex < Chunk.RigidVertices.Num();VertexIndex++)
-		{
-			const FRigidSkinVertex& SourceVertex = Chunk.RigidVertices[VertexIndex];
-			DestVertex->Position = SourceVertex.Position;
-			DestVertex->TangentX = SourceVertex.TangentX;
-			DestVertex->TangentY = SourceVertex.TangentY;
-			DestVertex->TangentZ = SourceVertex.TangentZ;
-			// store the sign of the determinant in TangentZ.W
-			DestVertex->TangentZ.Vector.W = GetBasisDeterminantSignByte( SourceVertex.TangentX, SourceVertex.TangentY, SourceVertex.TangentZ );
-
-			// copy all texture coordinate sets
-			FMemory::Memcpy( DestVertex->UVs, SourceVertex.UVs, sizeof(FVector2D)*MAX_TEXCOORDS );
-
-			DestVertex->Color = SourceVertex.Color;
-			DestVertex->InfluenceBones[0] = SourceVertex.Bone;
-			DestVertex->InfluenceWeights[0] = 255;
-			for(int32 InfluenceIndex = 1;InfluenceIndex < MAX_TOTAL_INFLUENCES;InfluenceIndex++)
-			{
-				DestVertex->InfluenceBones[InfluenceIndex] = 0;
-				DestVertex->InfluenceWeights[InfluenceIndex] = 0;
-			}
-			DestVertex++;
-		}
-		FMemory::Memcpy(DestVertex,Chunk.SoftVertices.GetData(),Chunk.SoftVertices.Num() * sizeof(FSoftSkinVertex));
-		DestVertex += Chunk.SoftVertices.Num();
+	for(int32 SectionIndex = 0; SectionIndex < Sections.Num(); SectionIndex++)
+	{
+		const FSkelMeshSection& Section = Sections[SectionIndex];
+		FMemory::Memcpy(DestVertex, Section.SoftVertices.GetData(), Section.SoftVertices.Num() * sizeof(FSoftSkinVertex));
+		DestVertex += Section.SoftVertices.Num();
 	}
 }
 
 void FStaticLODModel::GetApexClothMappingData(TArray<FApexClothPhysToRenderVertData>& MappingData) const
 {
-	for( int32 ChunkIndex = 0; ChunkIndex < Chunks.Num(); ChunkIndex++ )
+	for (int32 SectionIndex = 0; SectionIndex < Sections.Num(); SectionIndex++)
 	{
-		const FSkelMeshChunk& Chunk = Chunks[ChunkIndex];
+		const FSkelMeshSection& Section = Sections[SectionIndex];
 
-		if( Chunk.ApexClothMappingData.Num() == 0 )
+		if(Section.ApexClothMappingData.Num() == 0 )
 		{
 			int32 PrevNum = MappingData.Num();
 
-			MappingData.AddZeroed(Chunk.RigidVertices.Num() + Chunk.SoftVertices.Num());
+			MappingData.AddZeroed(Section.SoftVertices.Num());
 			for(int32 i=PrevNum; i<MappingData.Num(); i++)
 			{
 				MappingData[i].PositionBaryCoordsAndDist[0] = 1.0f;
@@ -1551,7 +1700,7 @@ void FStaticLODModel::GetApexClothMappingData(TArray<FApexClothPhysToRenderVertD
 		}
 		else
 		{
-			MappingData += Chunk.ApexClothMappingData;
+			MappingData += Section.ApexClothMappingData;
 		}
 	}
 }
@@ -1571,7 +1720,7 @@ void FStaticLODModel::BuildVertexBuffers(uint32 BuildFlags)
 	// Set the number of texture coordinate sets
 	VertexBufferGPUSkin.SetNumTexCoords( NumTexCoords );
 
-	VertexBufferGPUSkin.SetHasExtraBoneInfluences(DoChunksNeedExtraBoneInfluences());
+	VertexBufferGPUSkin.SetHasExtraBoneInfluences(DoSectionsNeedExtraBoneInfluences());
 
 	// init vertex buffer with the vertex array
 	VertexBufferGPUSkin.Init(Vertices);
@@ -1707,7 +1856,6 @@ SIZE_T FStaticLODModel::GetResourceSize() const
 	SIZE_T ResourceSize = 0;
 
 	ResourceSize += Sections.GetAllocatedSize();
-	ResourceSize += Chunks.GetAllocatedSize();
 	ResourceSize += ActiveBoneIndices.GetAllocatedSize();  
 	ResourceSize += RequiredBones.GetAllocatedSize();
 
@@ -1920,7 +2068,7 @@ USkeletalMesh
 /**
 * Calculate max # of bone influences used by this skel mesh chunk
 */
-void FSkelMeshChunk::CalcMaxBoneInfluences()
+void FSkelMeshSection::CalcMaxBoneInfluences()
 {
 	// if we only have rigid verts then there is only one bone
 	MaxBoneInfluences = 1;
@@ -2053,7 +2201,7 @@ bool FSkeletalMeshResource::HasExtraBoneInfluences() const
 	for (int32 LODIndex = 0; LODIndex < LODModels.Num(); ++LODIndex)
 	{
 		const FStaticLODModel& Model = LODModels[LODIndex];
-		if (Model.DoChunksNeedExtraBoneInfluences())
+		if (Model.DoSectionsNeedExtraBoneInfluences())
 		{
 			return true;
 		}
@@ -2062,24 +2210,24 @@ bool FSkeletalMeshResource::HasExtraBoneInfluences() const
 	return false;
 }
 
-int32 FSkeletalMeshResource::GetMaxBonesPerChunk() const
+int32 FSkeletalMeshResource::GetMaxBonesPerSection() const
 {
-	int32 MaxBonesPerChunk = 0;
+	int32 MaxBonesPerSection = 0;
 	for (int32 LODIndex = 0; LODIndex < LODModels.Num(); ++LODIndex)
 	{
 		const FStaticLODModel& Model = LODModels[LODIndex];
-		for (int32 ChunkIndex = 0; ChunkIndex < Model.Chunks.Num(); ++ChunkIndex)
+		for (int32 SectionIndex = 0; SectionIndex < Model.Sections.Num(); ++SectionIndex)
 		{
-			MaxBonesPerChunk = FMath::Max<int32>(MaxBonesPerChunk,Model.Chunks[ChunkIndex].BoneMap.Num());
+			MaxBonesPerSection = FMath::Max<int32>(MaxBonesPerSection, Model.Sections[SectionIndex].BoneMap.Num());
 		}
 	}
-	return MaxBonesPerChunk;
+	return MaxBonesPerSection;
 }
 
 bool FSkeletalMeshResource::RequiresCPUSkinning(ERHIFeatureLevel::Type FeatureLevel) const
 {
 	const int32 MaxGPUSkinBones = GetFeatureLevelMaxNumberOfBones(FeatureLevel);
-	const int32 MaxBonesPerChunk = GetMaxBonesPerChunk();
+	const int32 MaxBonesPerChunk = GetMaxBonesPerSection();
 	// Do CPU skinning if we need too many bones per chunk, or if we have too many influences per vertex on lower end
 	return (MaxBonesPerChunk > MaxGPUSkinBones) || (HasExtraBoneInfluences() && FeatureLevel < ERHIFeatureLevel::ES3_1);
 }
@@ -2970,13 +3118,6 @@ void USkeletalMesh::PostLoad()
 			ThisLODInfo.TriangleSortSettings.AddZeroed( ThisLODModel.Sections.Num()-ThisLODInfo.TriangleSortSettings.Num() );
 		}
 
-		// Verify that section and chunk indices match.
-		check(ThisLODModel.Sections.Num() == ThisLODModel.Chunks.Num());
-		for (int32 SectionIndex = 0; SectionIndex < ThisLODModel.Sections.Num(); ++SectionIndex)
-		{
-			check(ThisLODModel.Sections[SectionIndex].ChunkIndex == SectionIndex);
-		}
-
 #if WITH_EDITOR
 		if (ThisLODInfo.ReductionSettings.BonesToRemove_DEPRECATED.Num() > 0)
 		{
@@ -3062,21 +3203,18 @@ void USkeletalMesh::PostLoad()
 		{
 			FStaticLODModel& CurLODModel = SkelResource->LODModels[LODIndex];
 
-			for(FSkelMeshSection& CurSection : CurLODModel.Sections)
+			for (int32 SectionIdx = 0; SectionIdx< CurLODModel.Sections.Num(); SectionIdx++)
 			{
-				if(CurSection.CorrespondClothSectionIndex != INDEX_NONE)
+				FSkelMeshSection& CurSection = CurLODModel.Sections[SectionIdx];
+
+				if(CurSection.CorrespondClothSectionIndex != INDEX_NONE && CurSection.MaxBoneInfluences > MAX_INFLUENCES_PER_STREAM)
 				{
-					// We have cloth on this section, check to make sure influences are in range
-					FSkelMeshChunk& CorrespondingChunk = CurLODModel.Chunks[CurSection.ChunkIndex];
-					if(CorrespondingChunk.MaxBoneInfluences > MAX_INFLUENCES_PER_STREAM)
-					{
-						UE_LOG(LogSkeletalMesh, Warning, TEXT("Chunk %d for LOD %d in skeletal mesh %s has clothing associated but has %d influences. Clothing only supports a maximum of %d influences - reduce influences on chunk and reimport mesh."),
-							CurSection.ChunkIndex,
-							LODIndex,
-							*GetName(),
-							CorrespondingChunk.MaxBoneInfluences,
-							MAX_INFLUENCES_PER_STREAM);
-					}
+					UE_LOG(LogSkeletalMesh, Warning, TEXT("Section %d for LOD %d in skeletal mesh %s has clothing associated but has %d influences. Clothing only supports a maximum of %d influences - reduce influences on chunk and reimport mesh."),
+						SectionIdx,
+						LODIndex,
+						*GetName(),
+						CurSection.MaxBoneInfluences,
+						MAX_INFLUENCES_PER_STREAM);
 				}
 			}
 		}
@@ -3875,27 +4013,18 @@ bool USkeletalMesh::GetPhysicsTriMeshData(FTriMeshCollisionData* CollisionData,
 		// Copy all verts into collision vertex buffer.
 		CollisionData->Vertices.Empty();
 		CollisionData->Vertices.AddUninitialized(Model.NumVertices);
-		const uint32 NumChunks = Model.Chunks.Num();
+		const uint32 NumSections = Model.Sections.Num();
 
-		for (uint32 ChunkIdx = 0; ChunkIdx < NumChunks; ++ChunkIdx)
+		for (uint32 SectionIdx = 0; SectionIdx < NumSections; ++SectionIdx)
 		{
-			const FSkelMeshChunk& Chunk = Model.Chunks[ChunkIdx];
-			{
-				//rigid
-				const uint32 RigidOffset = Chunk.GetRigidVertexBufferIndex();
-				const uint32 NumRigidVerts = Chunk.GetNumRigidVertices();
-				for (uint32 RigidIdx = 0; RigidIdx < NumRigidVerts; ++RigidIdx)
-				{
-					CollisionData->Vertices[RigidIdx + RigidOffset] = Chunk.RigidVertices[RigidIdx].Position;
-				}
-			}
+			const FSkelMeshSection& Section = Model.Sections[SectionIdx];
 			{
 				//soft
-				const uint32 SoftOffset = Chunk.GetSoftVertexBufferIndex();
-				const uint32 NumSoftVerts = Chunk.GetNumSoftVertices();
+				const uint32 SoftOffset = Section.GetVertexBufferIndex();
+				const uint32 NumSoftVerts = Section.GetNumVertices();
 				for (uint32 SoftIdx = 0; SoftIdx < NumSoftVerts; ++SoftIdx)
 				{
-					CollisionData->Vertices[SoftIdx + SoftOffset] = Chunk.SoftVertices[SoftIdx].Position;
+					CollisionData->Vertices[SoftIdx + SoftOffset] = Section.SoftVertices[SoftIdx].Position;
 				}
 			}
 
@@ -4012,12 +4141,8 @@ bool USkeletalMesh::IsSectionUsingCloth(int32 InSectionIndex, bool bCheckCorresp
 				{
 					SectionToCheck = &LodModel.Sections[SectionToCheck->CorrespondClothSectionIndex];
 				}
-
-				int32 ChunkIdx = SectionToCheck->ChunkIndex;
-				check(LodModel.Chunks.IsValidIndex(ChunkIdx)); // Built on import, should be valid at all times
-				
-				FSkelMeshChunk& Chunk = LodModel.Chunks[ChunkIdx];
-				if(Chunk.HasApexClothData())
+			
+				if(SectionToCheck->HasApexClothData())
 				{
 					return true;
 				}
@@ -4731,9 +4856,7 @@ public:
 		, MeshObject(InMeshObject)
 		, LODSectionElements(InLODSectionElements)
 		, Sections(InLODModel.Sections)
-		, Chunks(InLODModel.Chunks)
 #if WITH_EDITORONLY_DATA
-		, ChunkIndexPreview(InMeshObject.ChunkIndexPreview)
 		, SectionIndexPreview(InMeshObject.SectionIndexPreview)
 #endif
 	{
@@ -4754,10 +4877,6 @@ public:
 	{
 		return ((SectionIndex < Sections.Num()) && LODSectionElements.SectionElements.IsValidIndex(GetSectionElementIndex()));
 	}
-	FORCEINLINE const FSkelMeshChunk& GetChunk() const
-	{
-		return Chunks[GetSection().ChunkIndex];
-	}
 	FORCEINLINE const FSkelMeshSection& GetSection() const
 	{
 		return Sections[SectionIndex];
@@ -4779,8 +4898,7 @@ public:
 	{
 #if WITH_EDITORONLY_DATA
 		return	(SectionIndex < Sections.Num()) && 
-				(((ChunkIndexPreview >= 0) && (ChunkIndexPreview != Sections[SectionIndex].ChunkIndex)) ||
-				((SectionIndexPreview >= 0) && (SectionIndexPreview != SectionIndex)));
+				((SectionIndexPreview >= 0) && (SectionIndexPreview != SectionIndex));
 #else
 		return false;
 #endif
@@ -4790,9 +4908,7 @@ private:
 	const FSkeletalMeshObject& MeshObject;
 	const FSkeletalMeshSceneProxy::FLODSectionElements& LODSectionElements;
 	const TArray<FSkelMeshSection>& Sections;
-	const TArray<FSkelMeshChunk>& Chunks;
 #if WITH_EDITORONLY_DATA
-	const int32 ChunkIndexPreview;
 	const int32 SectionIndexPreview;
 #endif
 };
@@ -4878,7 +4994,7 @@ void FSkeletalMeshSceneProxy::GetMeshElementsConditionallySelectable(const TArra
 		for (FSkeletalMeshSectionIter Iter(LODIndex, *MeshObject, LODModel, LODSection); Iter; ++Iter)
 		{
 			FSkelMeshSection Section = Iter.GetSection();
-			const FSkelMeshChunk& Chunk = Iter.GetChunk();
+			const int32 SectionIndex = Iter.GetSectionElementIndex();
 			const FSectionElementInfo& SectionElementInfo = Iter.GetSectionElementInfo();
 			const FTwoVectors& CustomLeftRightVectors = Iter.GetCustomLeftRightVectors();
 
@@ -4898,7 +5014,7 @@ void FSkeletalMeshSceneProxy::GetMeshElementsConditionallySelectable(const TArra
 				continue;
 			}
 
-			GetDynamicElementsSection(Views, ViewFamily, VisibilityMap, LODModel, LODIndex, Section, Chunk, SectionElementInfo, CustomLeftRightVectors, bInSelectable, Collector);
+			GetDynamicElementsSection(Views, ViewFamily, VisibilityMap, LODModel, LODIndex, SectionIndex, SectionElementInfo, CustomLeftRightVectors, bInSelectable, Collector);
 		}
 	}
 
@@ -4922,9 +5038,11 @@ void FSkeletalMeshSceneProxy::GetMeshElementsConditionallySelectable(const TArra
 }
 
 void FSkeletalMeshSceneProxy::GetDynamicElementsSection(const TArray<const FSceneView*>& Views, const FSceneViewFamily& ViewFamily, uint32 VisibilityMap, 
-	const FStaticLODModel& LODModel, const int32 LODIndex, const FSkelMeshSection& Section, const FSkelMeshChunk& Chunk, 
+	const FStaticLODModel& LODModel, const int32 LODIndex, const int32 SectionIndex, 
 	const FSectionElementInfo& SectionElementInfo, const FTwoVectors& CustomLeftRightVectors, bool bInSelectable, FMeshElementCollector& Collector ) const
 {
+	const FSkelMeshSection& Section = LODModel.Sections[SectionIndex];
+
 	// If hidden skip the draw
 	if (Section.bDisabled || MeshObject->IsMaterialHidden(LODIndex,SectionElementInfo.UseMaterialIndex))
 	{
@@ -4959,7 +5077,7 @@ void FSkeletalMeshSceneProxy::GetDynamicElementsSection(const TArray<const FScen
 			Mesh.LCI = NULL;
 			Mesh.bWireframe |= bForceWireframe;
 			Mesh.Type = PT_TriangleList;
-			Mesh.VertexFactory = MeshObject->GetSkinVertexFactory(View, LODIndex, Section.ChunkIndex);
+			Mesh.VertexFactory = MeshObject->GetSkinVertexFactory(View, LODIndex, SectionIndex);
 			
 			if(!Mesh.VertexFactory)
 			{
@@ -4973,9 +5091,9 @@ void FSkeletalMeshSceneProxy::GetDynamicElementsSection(const TArray<const FScen
 			BatchElement.IndexBuffer = LODModel.MultiSizeIndexContainer.GetIndexBuffer();
 			BatchElement.MaxVertexIndex = LODModel.NumVertices - 1;
 	
-			if(Section.ChunkIndex < FSkeletalMeshObject::MAX_GPUSKINCACHE_CHUNKS_PER_LOD)
+			if(SectionIndex < FSkeletalMeshObject::MAX_GPUSKINCACHE_CHUNKS_PER_LOD)
 			{
-				BatchElement.UserIndex = MeshObject->GPUSkinCacheKeys[Section.ChunkIndex];
+				BatchElement.UserIndex = MeshObject->GPUSkinCacheKeys[SectionIndex];
 			}
 			else
 			{
@@ -5044,7 +5162,7 @@ void FSkeletalMeshSceneProxy::GetDynamicElementsSection(const TArray<const FScen
 			BatchElement.NumPrimitives = Section.NumTriangles;
 			if( GIsEditor && MeshObject->ProgressiveDrawingFraction != 1.f )
 			{
-				if (Mesh.MaterialRenderProxy->GetMaterial(FeatureLevel)->GetBlendMode() == BLEND_Translucent)
+				if (Mesh.MaterialRenderProxy && Mesh.MaterialRenderProxy->GetMaterial(FeatureLevel)->GetBlendMode() == BLEND_Translucent)
 				{
 					BatchElement.NumPrimitives = FMath::RoundToInt(((float)Section.NumTriangles)*FMath::Clamp<float>(MeshObject->ProgressiveDrawingFraction,0.f,1.f));
 					if( BatchElement.NumPrimitives == 0 )
@@ -5053,7 +5171,7 @@ void FSkeletalMeshSceneProxy::GetDynamicElementsSection(const TArray<const FScen
 					}
 				}
 			}
-			BatchElement.MinVertexIndex = Chunk.BaseVertexIndex;
+			BatchElement.MinVertexIndex = Section.BaseVertexIndex;
 			Mesh.ReverseCulling = IsLocalToWorldDeterminantNegative();
 			Mesh.CastShadow = SectionElementInfo.bEnableShadowCasting;
 
@@ -5070,7 +5188,7 @@ void FSkeletalMeshSceneProxy::GetDynamicElementsSection(const TArray<const FScen
 				Collector.AddMesh(ViewIndex, Mesh);
 			}
 
-			const int32 NumVertices = Chunk.NumRigidVertices + Chunk.NumSoftVertices;
+			const int32 NumVertices = Section.GetNumVertices();
 			INC_DWORD_STAT_BY(STAT_GPUSkinVertices,(uint32)(bIsCPUSkinned ? 0 : NumVertices));
 			INC_DWORD_STAT_BY(STAT_SkelMeshTriangles,Mesh.GetNumPrimitives());
 			INC_DWORD_STAT(STAT_SkelMeshDrawCalls);
@@ -5424,7 +5542,7 @@ FPrimitiveSceneProxy* USkinnedMeshComponent::CreateSceneProxy()
 		MeshObject)
 	{
 		// Only create a scene proxy if the bone count being used is supported, or if we don't have a skeleton (this is the case with destructibles)
-		int32 MaxBonesPerChunk = SkelMeshResource->GetMaxBonesPerChunk();
+		int32 MaxBonesPerChunk = SkelMeshResource->GetMaxBonesPerSection();
 		if (MaxBonesPerChunk <= GetFeatureLevelMaxNumberOfBones(SceneFeatureLevel))
 		{
 			Result = ::new FSkeletalMeshSceneProxy(this,SkelMeshResource);
diff --git a/Engine/Source/Runtime/Engine/Private/SkeletalMeshComponentPhysics.cpp b/Engine/Source/Runtime/Engine/Private/SkeletalMeshComponentPhysics.cpp
--- a/Engine/Source/Runtime/Engine/Private/SkeletalMeshComponentPhysics.cpp
+++ b/Engine/Source/Runtime/Engine/Private/SkeletalMeshComponentPhysics.cpp
@@ -259,9 +259,7 @@ bool USkeletalMesh::HasClothSections(int32 LODIndex, int AssetIndex)
 
 	for(int32 SecIdx=0; SecIdx < NumSections; SecIdx++)
 	{
-		uint16 ChunkIdx = LODModel.Sections[SecIdx].ChunkIndex;
-
-		if(LODModel.Chunks[ChunkIdx].CorrespondClothAssetIndex == AssetIndex)
+		if(LODModel.Sections[SecIdx].CorrespondClothAssetIndex == AssetIndex)
 		{
 			return true;
 		}
@@ -283,7 +281,7 @@ void USkeletalMesh::GetClothSectionIndices(int32 LODIndex, int32 AssetIndex, TAr
 
 	for(int32 SecIdx=0; SecIdx < NumSections; SecIdx++)
 	{
-		if(LODModel.Chunks[LODModel.Sections[SecIdx].ChunkIndex].CorrespondClothAssetIndex == AssetIndex)
+		if(LODModel.Sections[SecIdx].CorrespondClothAssetIndex == AssetIndex)
 		{
 			//add cloth sections
 			OutSectionIndices.Add(SecIdx);
@@ -325,7 +323,7 @@ void USkeletalMesh::GetOriginSectionIndicesWithCloth(int32 LODIndex, int32 Asset
 
 	for(int32 SecIdx=0; SecIdx < NumSections; SecIdx++)
 	{
-		if(LODModel.Chunks[LODModel.Sections[SecIdx].ChunkIndex].CorrespondClothAssetIndex == AssetIndex)
+		if(LODModel.Sections[SecIdx].CorrespondClothAssetIndex == AssetIndex)
 		{
 			//add original sections
 			OutSectionIndices.Add(LODModel.Sections[SecIdx].CorrespondClothSectionIndex);
@@ -345,12 +343,12 @@ bool USkeletalMesh::IsMappedClothingLOD(int32 InLODIndex, int32 InAssetIndex)
 
 		// loop reversely for optimized search
 		for (int32 SecIdx = NumSections-1; SecIdx >= 0; SecIdx--)
-	{
-			int32 ClothAssetIndex = LODModel.Chunks[LODModel.Sections[SecIdx].ChunkIndex].CorrespondClothAssetIndex;
-			if (ClothAssetIndex == InAssetIndex)
 		{
-			return true;
-		}
+			int32 ClothAssetIndex = LODModel.Sections[SecIdx].CorrespondClothAssetIndex;
+			if (ClothAssetIndex == InAssetIndex)
+			{
+				return true;
+			}
 			else if (ClothAssetIndex == INDEX_NONE) // no more cloth sections
 			{
 				return false;
@@ -385,8 +383,7 @@ int32 USkeletalMesh::GetClothAssetIndex(int32 LODIndex, int32 SectionIndex)
 		return INDEX_NONE;
 	}
 
-	int16 ChunkIdx = LODModel.Sections[ClothSecIdx].ChunkIndex;
-	return LODModel.Chunks[ChunkIdx].CorrespondClothAssetIndex;
+	return LODModel.Sections[ClothSecIdx].CorrespondClothAssetIndex;
 }
 #endif//#if WITH_APEX_CLOTHING
 
@@ -1843,52 +1840,40 @@ FVector USkeletalMeshComponent::GetSkinnedVertexPosition(int32 VertexIndex) cons
 		FStaticLODModel& Model = MeshObject->GetSkeletalMeshResource().LODModels[0];
 
 		// Find the chunk and vertex within that chunk, and skinning type, for this vertex.
-		int32 ChunkIndex;
+		int32 SectionIndex;
 		int32 VertIndexInChunk;
-		bool bSoftVertex;
 		bool bHasExtraBoneInfluences;
-		Model.GetChunkAndSkinType(VertexIndex, ChunkIndex, VertIndexInChunk, bSoftVertex, bHasExtraBoneInfluences);
+		Model.GetSectionFromVertexIndex(VertexIndex, SectionIndex, VertIndexInChunk, bHasExtraBoneInfluences);
 
 		bool bClothVertex = false;
 		int32 ClothAssetIndex = -1;
 
+		// if this section corresponds to a cloth section, returns corresponding cloth section's info instead
+		const FSkelMeshSection& Section = Model.Sections[SectionIndex];
+
 		// if this chunk has cloth data
-		if (Model.Chunks[ChunkIndex].HasApexClothData())
+		if (Section.HasApexClothData())
 		{
 			bClothVertex = true;
-			ClothAssetIndex = Model.Chunks[ChunkIndex].CorrespondClothAssetIndex;
+			ClothAssetIndex = Section.CorrespondClothAssetIndex;
 		}
 		else
 		{
-			// if this chunk corresponds to a cloth section, returns corresponding cloth section's info instead
-			for (int32 SectionIdx = 0; SectionIdx < Model.Sections.Num(); SectionIdx++)
+			// if current section is disabled and the corresponding cloth section is visible
+			if (Section.bDisabled && Section.CorrespondClothSectionIndex >= 0)
 			{
-				const FSkelMeshSection& Section = Model.Sections[SectionIdx];
+				bClothVertex = true;
 
-				// find a section which has this chunk index
-				if (Section.ChunkIndex == ChunkIndex)
-				{
-					// if current section is disabled and the corresponding cloth section is visible
-					if (Section.bDisabled && Section.CorrespondClothSectionIndex >= 0)
-					{
-						bClothVertex = true;
-
-						const FSkelMeshSection& ClothSection = Model.Sections[Section.CorrespondClothSectionIndex];
-						const FSkelMeshChunk& ClothChunk = Model.Chunks[ClothSection.ChunkIndex];
-
-						ClothAssetIndex = ClothChunk.CorrespondClothAssetIndex;
+				const FSkelMeshSection& ClothSection = Model.Sections[Section.CorrespondClothSectionIndex];
+				ClothAssetIndex = ClothSection.CorrespondClothAssetIndex;
 
-						// the index can exceed the range because this vertex index is based on the corresponding original section
-						// the number of cloth chunk's vertices is not always same as the corresponding one 
-						// cloth chunk has only soft vertices
-						if (VertIndexInChunk >= ClothChunk.GetNumSoftVertices())
-						{
-							// if the index exceeds, re-assign a random vertex index for this chunk
-							VertIndexInChunk = FMath::TruncToInt(FMath::SRand() * (ClothChunk.GetNumSoftVertices() - 1));
-						}
-					}
-					// if found, quit this loop quickly
-					break;
+				// the index can exceed the range because this vertex index is based on the corresponding original section
+				// the number of cloth chunk's vertices is not always same as the corresponding one 
+				// cloth chunk has only soft vertices
+				if (VertIndexInChunk >= ClothSection.GetNumVertices())
+				{
+					// if the index exceeds, re-assign a random vertex index for this chunk
+					VertIndexInChunk = FMath::TruncToInt(FMath::SRand() * (ClothSection.GetNumVertices() - 1));
 				}
 			}
 		}
@@ -1922,10 +1907,9 @@ FVector USkeletalMeshComponent::GetSkinnedVertexPosition(int32 VertexIndex) cons
 
 extern float DebugLineLifetime;
 
-float USkeletalMeshComponent::GetDistanceToCollision(const FVector& Point, FVector& ClosestPointOnCollision) const
+bool USkeletalMeshComponent::GetSquaredDistanceToCollision(const FVector& Point, float& OutSquaredDistance, FVector& OutClosestPointOnCollision) const
 {
-	ClosestPointOnCollision = Point;
-	float ClosestPointDistance = -1.f;
+	OutClosestPointOnCollision = Point;
 	bool bHasResult = false;
 
 	for (int32 BodyIdx = 0; BodyIdx < Bodies.Num(); ++BodyIdx)
@@ -1934,22 +1918,22 @@ float USkeletalMeshComponent::GetDistanceToCollision(const FVector& Point, FVect
 		if (BodyInst && BodyInst->IsValidBodyInstance() && (BodyInst->GetCollisionEnabled() != ECollisionEnabled::NoCollision))
 		{
 			FVector ClosestPoint;
-			const float Distance = Bodies[BodyIdx]->GetDistanceToBody(Point, ClosestPoint);
+			float DistanceSqr = -1.f;
 
-			if (Distance < 0.f)
+			if (!Bodies[BodyIdx]->GetSquaredDistanceToBody(Point, DistanceSqr, ClosestPoint))
 			{
 				// Invalid result, impossible to be better than ClosestPointDistance
 				continue;
 			}
 
-			if (!bHasResult || (Distance < ClosestPointDistance))
+			if (!bHasResult || (DistanceSqr < OutSquaredDistance))
 			{
 				bHasResult = true;
-				ClosestPointDistance = Distance;
-				ClosestPointOnCollision = ClosestPoint;
+				OutSquaredDistance = DistanceSqr;
+				OutClosestPointOnCollision = ClosestPoint;
 
 				// If we're inside collision, we're not going to find anything better, so abort search we've got our best find.
-				if (Distance <= KINDA_SMALL_NUMBER)
+				if (DistanceSqr <= KINDA_SMALL_NUMBER)
 				{
 					break;
 				}
@@ -1957,7 +1941,7 @@ float USkeletalMeshComponent::GetDistanceToCollision(const FVector& Point, FVect
 		}
 	}
 
-	return ClosestPointDistance;
+	return bHasResult;
 }
 
 DECLARE_CYCLE_STAT(TEXT("GetClosestPointOnPhysicsAsset"), STAT_GetClosestPointOnPhysicsAsset, STATGROUP_Physics);
@@ -3446,11 +3430,9 @@ void USkeletalMeshComponent::PrepareClothMorphTargets()
 			FStaticLODModel& Model = MeshObject->GetSkeletalMeshResource().LODModels[0];
 
 			// Find the chunk and vertex within that chunk, and skinning type, for this vertex.
-			int32 ChunkIndex;
+			int32 SectionIndex;
 			int32 VertIndexInChunk;
-			bool bSoftVertex;
 			bool bHasExtraBoneInfluences;
-			int32 SectionIndex = INDEX_NONE;
 
 			int32 NumAssets = SkeletalMesh->ClothingAssets.Num();
 			TArray<TArray<PxVec3>> ClothOriginalPosArray;
@@ -3458,42 +3440,24 @@ void USkeletalMeshComponent::PrepareClothMorphTargets()
 			ClothOriginalPosArray.AddZeroed(NumAssets);
 			ClothPositionDeltaArray.AddZeroed(NumAssets);
 
-			int ClothChunkIndex = INDEX_NONE;
-
 			for (int32 VertIdx = 0; VertIdx < NumVerts; VertIdx++)
 			{
-				Model.GetChunkAndSkinType(Vertices[VertIdx].SourceIdx, ChunkIndex, VertIndexInChunk, bSoftVertex, bHasExtraBoneInfluences);
+				Model.GetSectionFromVertexIndex(Vertices[VertIdx].SourceIdx, SectionIndex, VertIndexInChunk, bHasExtraBoneInfluences);
 
-				FSkelMeshChunk& Chunk = Model.Chunks[ChunkIndex];
+				int32 ClothSectionIndex = INDEX_NONE;
 
-				for (int32 SecIdx = 0; SecIdx < Model.Sections.Num(); SecIdx++)
+				// if current section is disabled and the corresponding cloth section is visible
+				FSkelMeshSection& Section = Model.Sections[SectionIndex];
+				if (Section.bDisabled && Section.CorrespondClothSectionIndex >= 0)
 				{
-					FSkelMeshSection& Section = Model.Sections[SecIdx];
-					if (Section.ChunkIndex == ChunkIndex)
-					{
-						// if current section is disabled and the corresponding cloth section is visible
-						if (Section.bDisabled && Section.CorrespondClothSectionIndex >= 0)
-						{
-							SectionIndex = SecIdx;
-							ClothChunkIndex = Model.Sections[Section.CorrespondClothSectionIndex].ChunkIndex;
-						}
-						break;
-					}
+					ClothSectionIndex = Section.CorrespondClothSectionIndex;
 				}
 
-				if (SectionIndex != INDEX_NONE)
+				if (ClothSectionIndex != INDEX_NONE)
 				{
-					FVector Position;
-					if (bSoftVertex)
-					{
-						Position = Chunk.SoftVertices[VertIndexInChunk].Position;
-					}
-					else
-					{
-						Position = Chunk.RigidVertices[VertIndexInChunk].Position;
-					}
+					FVector Position = Section.SoftVertices[VertIndexInChunk].Position;
 
-					int32 AssetIndex = Model.Chunks[ClothChunkIndex].CorrespondClothAssetIndex;
+					int32 AssetIndex = Model.Sections[ClothSectionIndex].CorrespondClothAssetIndex;
 					// save to original positions
 					ClothOriginalPosArray[AssetIndex].Add(U2PVector(Position));
 					ClothPositionDeltaArray[AssetIndex].Add(Vertices[VertIdx].PositionDelta);
@@ -4120,8 +4084,6 @@ void USkeletalMeshComponent::DrawClothingTangents(FPrimitiveDrawInterface* PDI)
 			// simulation is enabled and there exist simulated vertices
 			if(!bDisableClothSimulation && NumSimulVertices > 0)
 			{
-				uint16 ChunkIndex = 0;
-
 				FStaticLODModel& LODModel = MeshObject->GetSkeletalMeshResource().LODModels[PredictedLODLevel];
 
 				TArray<uint32> SectionIndices;
@@ -4132,24 +4094,21 @@ void USkeletalMeshComponent::DrawClothingTangents(FPrimitiveDrawInterface* PDI)
 				for(int32 SecIdx=0; SecIdx <NumSections; SecIdx++)
 				{
 					uint16 SectionIndex = SectionIndices[SecIdx];
-
-					ChunkIndex = LODModel.Sections[SectionIndex].ChunkIndex;
-
-					FSkelMeshChunk& Chunk = LODModel.Chunks[ChunkIndex];
+					FSkelMeshSection& Section = LODModel.Sections[SectionIndex];
 
 					const physx::PxVec3* SimulVertices = ApexClothingActor->getSimulationPositions();
 					const physx::PxVec3* SimulNormals = ApexClothingActor->getSimulationNormals();
 
-					uint32 NumMppingData = Chunk.ApexClothMappingData.Num();
+					uint32 NumMppingData = Section.ApexClothMappingData.Num();
 
 					FVector Start, End;		
 
 					for(uint32 MappingIndex=0; MappingIndex < NumMppingData; MappingIndex++)
 					{
-						FVector4 BaryCoordPos = Chunk.ApexClothMappingData[MappingIndex].PositionBaryCoordsAndDist;
-						FVector4 BaryCoordNormal = Chunk.ApexClothMappingData[MappingIndex].NormalBaryCoordsAndDist;
-						FVector4 BaryCoordTangent = Chunk.ApexClothMappingData[MappingIndex].TangentBaryCoordsAndDist;
-						uint16*  SimulIndices = Chunk.ApexClothMappingData[MappingIndex].SimulMeshVertIndices;
+						FVector4 BaryCoordPos = Section.ApexClothMappingData[MappingIndex].PositionBaryCoordsAndDist;
+						FVector4 BaryCoordNormal = Section.ApexClothMappingData[MappingIndex].NormalBaryCoordsAndDist;
+						FVector4 BaryCoordTangent = Section.ApexClothMappingData[MappingIndex].TangentBaryCoordsAndDist;
+						uint16*  SimulIndices = Section.ApexClothMappingData[MappingIndex].SimulMeshVertIndices;
 
 						bool bFixed = (SimulIndices[3] == 0xFFFF);
 
@@ -4234,8 +4193,6 @@ void USkeletalMeshComponent::DrawClothingTangents(FPrimitiveDrawInterface* PDI)
 
 				NumSimulVertices = VisualInfo.ClothPhysicalMeshVertices.Num();
 
-				uint16 ChunkIndex = 0;
-
 				FStaticLODModel& LODModel = MeshObject->GetSkeletalMeshResource().LODModels[PredictedLODLevel];
 
 				TArray<uint32> SectionIndices;
@@ -4246,24 +4203,21 @@ void USkeletalMeshComponent::DrawClothingTangents(FPrimitiveDrawInterface* PDI)
 				for(int32 SecIdx=0; SecIdx <NumSections; SecIdx++)
 				{
 					uint16 SectionIndex = SectionIndices[SecIdx];
-
-					ChunkIndex = LODModel.Sections[SectionIndex].ChunkIndex;
-
-					FSkelMeshChunk& Chunk = LODModel.Chunks[ChunkIndex];
+					FSkelMeshSection& Section = LODModel.Sections[SectionIndex];
 
 					const TArray<FVector>& SimulVertices = VisualInfo.ClothPhysicalMeshVertices;
 					const TArray<FVector>& SimulNormals = VisualInfo.ClothPhysicalMeshNormals;
 
-					uint32 NumMppingData = Chunk.ApexClothMappingData.Num();
+					uint32 NumMppingData = Section.ApexClothMappingData.Num();
 
 					FVector Start, End;		
 
 					for(uint32 MappingIndex=0; MappingIndex < NumMppingData; MappingIndex++)
 					{
-						FVector4 BaryCoordPos = Chunk.ApexClothMappingData[MappingIndex].PositionBaryCoordsAndDist;
-						FVector4 BaryCoordNormal = Chunk.ApexClothMappingData[MappingIndex].NormalBaryCoordsAndDist;
-						FVector4 BaryCoordTangent = Chunk.ApexClothMappingData[MappingIndex].TangentBaryCoordsAndDist;
-						uint16*  SimulIndices = Chunk.ApexClothMappingData[MappingIndex].SimulMeshVertIndices;
+						FVector4 BaryCoordPos = Section.ApexClothMappingData[MappingIndex].PositionBaryCoordsAndDist;
+						FVector4 BaryCoordNormal = Section.ApexClothMappingData[MappingIndex].NormalBaryCoordsAndDist;
+						FVector4 BaryCoordTangent = Section.ApexClothMappingData[MappingIndex].TangentBaryCoordsAndDist;
+						uint16*  SimulIndices = Section.ApexClothMappingData[MappingIndex].SimulMeshVertIndices;
 
 						bool bFixed = (SimulIndices[3] == 0xFFFF);
 
@@ -4565,8 +4519,6 @@ void USkeletalMeshComponent::DrawClothingFixedVertices(FPrimitiveDrawInterface*
 
 			NumSimulVertices = VisualInfo.ClothPhysicalMeshVertices.Num();
 
-			uint16 ChunkIndex = 0;
-
 			FStaticLODModel& LODModel = MeshObject->GetSkeletalMeshResource().LODModels[PredictedLODLevel];
 
 			TArray<uint32> SectionIndices;
@@ -4577,24 +4529,21 @@ void USkeletalMeshComponent::DrawClothingFixedVertices(FPrimitiveDrawInterface*
 			for(int32 SecIdx=0; SecIdx <NumSections; SecIdx++)
 			{
 				uint16 SectionIndex = SectionIndices[SecIdx];
-
-				ChunkIndex = LODModel.Sections[SectionIndex].ChunkIndex;
-
-				FSkelMeshChunk& Chunk = LODModel.Chunks[ChunkIndex];
+				FSkelMeshSection& Section = LODModel.Sections[SectionIndex];
 
 				const TArray<FVector>& SimulVertices = VisualInfo.ClothPhysicalMeshVertices;
 				const TArray<FVector>& SimulNormals = VisualInfo.ClothPhysicalMeshNormals;
 
-				uint32 NumMppingData = Chunk.ApexClothMappingData.Num();
+				uint32 NumMppingData = Section.ApexClothMappingData.Num();
 
 				FVector Start, End;		
 
 				for(uint32 MappingIndex=0; MappingIndex < NumMppingData; MappingIndex++)
 				{
-					FVector4 BaryCoordPos = Chunk.ApexClothMappingData[MappingIndex].PositionBaryCoordsAndDist;
-					FVector4 BaryCoordNormal = Chunk.ApexClothMappingData[MappingIndex].NormalBaryCoordsAndDist;
-					FVector4 BaryCoordTangent = Chunk.ApexClothMappingData[MappingIndex].TangentBaryCoordsAndDist;
-					uint16*  SimulIndices = Chunk.ApexClothMappingData[MappingIndex].SimulMeshVertIndices;
+					FVector4 BaryCoordPos = Section.ApexClothMappingData[MappingIndex].PositionBaryCoordsAndDist;
+					FVector4 BaryCoordNormal = Section.ApexClothMappingData[MappingIndex].NormalBaryCoordsAndDist;
+					FVector4 BaryCoordTangent = Section.ApexClothMappingData[MappingIndex].TangentBaryCoordsAndDist;
+					uint16*  SimulIndices = Section.ApexClothMappingData[MappingIndex].SimulMeshVertIndices;
 
 					bool bFixed = (SimulIndices[3] == 0xFFFF);
 
@@ -4618,7 +4567,7 @@ void USkeletalMeshComponent::DrawClothingFixedVertices(FPrimitiveDrawInterface*
 						+ BaryCoordPos.Y*(b + BaryCoordPos.W*nb)
 						+ BaryCoordPos.Z*(c + BaryCoordPos.W*nc);
 
-					FSoftSkinVertex& SoftVert = Chunk.SoftVertices[MappingIndex];
+					FSoftSkinVertex& SoftVert = Section.SoftVertices[MappingIndex];
 
 					// skinning, APEX clothing supports up to 4 bone influences for now
 					const uint8* BoneIndices = SoftVert.InfluenceBones;
@@ -4626,10 +4575,10 @@ void USkeletalMeshComponent::DrawClothingFixedVertices(FPrimitiveDrawInterface*
 				
 					FMatrix SkinningMat(ForceInitToZero);
 
-					for(int32 BoneWeightIdx=0; BoneWeightIdx < Chunk.MaxBoneInfluences; BoneWeightIdx++)
+					for(int32 BoneWeightIdx=0; BoneWeightIdx < Section.MaxBoneInfluences; BoneWeightIdx++)
 					{
 						float Weight = BoneWeights[BoneWeightIdx]*Inv_255;
-						SkinningMat += RefToLocals[Chunk.BoneMap[BoneIndices[BoneWeightIdx]]]*Weight;
+						SkinningMat += RefToLocals[Section.BoneMap[BoneIndices[BoneWeightIdx]]]*Weight;
 					}
 
 					FVector SkinnedPosition = SkinningMat.TransformPosition(Position);
diff --git a/Engine/Source/Runtime/Engine/Private/SkeletalMeshMerge.cpp b/Engine/Source/Runtime/Engine/Private/SkeletalMeshMerge.cpp
--- a/Engine/Source/Runtime/Engine/Private/SkeletalMeshMerge.cpp
+++ b/Engine/Source/Runtime/Engine/Private/SkeletalMeshMerge.cpp
@@ -180,7 +180,7 @@ bool FSkeletalMeshMerge::FinalizeMesh()
 					uint32& NumUVSets = PerLODNumUVSets[LODIdx];
 					NumUVSets = FMath::Max(NumUVSets, SrcResource->LODModels[LODIdx].NumTexCoords);
 
-					PerLODExtraBoneInfluences[LODIdx] |= SrcResource->LODModels[LODIdx].DoChunksNeedExtraBoneInfluences();
+					PerLODExtraBoneInfluences[LODIdx] |= SrcResource->LODModels[LODIdx].DoSectionsNeedExtraBoneInfluences();
 				}
 			}
 		}
@@ -288,11 +288,10 @@ void FSkeletalMeshMerge::GenerateNewSectionArray( TArray<FNewSectionInfo>& NewSe
 				}
 
 				FSkelMeshSection& Section = SrcLODModel.Sections[SectionIdx];
-				FSkelMeshChunk& Chunk = SrcLODModel.Chunks[Section.ChunkIndex];
 
 				// Convert Chunk.BoneMap from src to dest bone indices
 				TArray<FBoneIndexType> DestChunkBoneMap;
-				BoneMapToNewRefSkel(Chunk.BoneMap, SrcMeshInfo[MeshIdx].SrcToDestRefSkeletonMap, DestChunkBoneMap);
+				BoneMapToNewRefSkel(Section.BoneMap, SrcMeshInfo[MeshIdx].SrcToDestRefSkeletonMap, DestChunkBoneMap);
 
 
 				// get the material for this section
@@ -328,8 +327,7 @@ void FSkeletalMeshMerge::GenerateNewSectionArray( TArray<FNewSectionInfo>& NewSe
 							// add the source section as a new merge entry
 							FMergeSectionInfo& MergeSectionInfo = *new(NewSectionInfo.MergeSections) FMergeSectionInfo(
 								SrcMesh,
-								&SrcLODModel.Sections[SectionIdx],
-								&SrcLODModel.Chunks[Section.ChunkIndex]
+								&SrcLODModel.Sections[SectionIdx]
 								);
 							// keep track of remapping for the existing chunk's bonemap 
 							// so that the bone matrix indices can be updated for the vertices
@@ -356,8 +354,7 @@ void FSkeletalMeshMerge::GenerateNewSectionArray( TArray<FNewSectionInfo>& NewSe
 					// add a new merge section entry
 					FMergeSectionInfo& MergeSectionInfo = *new(NewSectionInfo.MergeSections) FMergeSectionInfo(
 						SrcMesh,
-						&SrcLODModel.Sections[SectionIdx],
-						&SrcLODModel.Chunks[Section.ChunkIndex]);
+						&SrcLODModel.Sections[SectionIdx]);
 					// since merged bonemap == chunk.bonemap then remapping is just pass-through
 					MergeSectionInfo.BoneMapToMergedBoneMap.Empty( DestChunkBoneMap.Num() );
 					for( int32 i=0; i < DestChunkBoneMap.Num(); i++ )
@@ -411,21 +408,17 @@ void FSkeletalMeshMerge::GenerateLODModel( int32 LODIdx )
 			MergeLODModel.ActiveBoneIndices.AddUnique( NewSectionInfo.MergedBoneMap[Idx] );
 		}
 
-		// add the new chunk entry 
-		FSkelMeshChunk& Chunk = *new(MergeLODModel.Chunks) FSkelMeshChunk;
+		// add the new section entry
+		FSkelMeshSection& Section = *new(MergeLODModel.Sections) FSkelMeshSection;
 
 		// set the new bonemap from the merged sections
 		// these are the bones that will be used by this new section
-		Chunk.BoneMap = NewSectionInfo.MergedBoneMap;
+		Section.BoneMap = NewSectionInfo.MergedBoneMap;
 
-		// init vert totals
-		Chunk.NumRigidVertices = 0;
-		Chunk.NumSoftVertices = 0;
 		// keep track of the current base vertex for this section in the merged vertex buffer
-		Chunk.BaseVertexIndex = MergedVertexBuffer.Num();
+		Section.BaseVertexIndex = MergedVertexBuffer.Num();
+
 
-		// add the new section entry
-		FSkelMeshSection& Section = *new(MergeLODModel.Sections) FSkelMeshSection;
 		// The current implementation of skeletal mesh merging can't do more than two sections.
 		// Therefore, we cannot e.g. create new sections based on the source section shadow casting.
 		MergeLODInfo.TriangleSortSettings.AddZeroed();
@@ -456,8 +449,6 @@ void FSkeletalMeshMerge::GenerateLODModel( int32 LODIdx )
 		
 		// init tri totals
 		Section.NumTriangles = 0;
-		// one-to-one mapping of sections to chunks
-		Section.ChunkIndex = MergeLODModel.Chunks.Num() - 1;
 		// keep track of the current base index for this section in the merged index buffer
 		Section.BaseIndex = MergedIndexBuffer.Num();
 
@@ -490,18 +481,15 @@ void FSkeletalMeshMerge::GenerateLODModel( int32 LODIdx )
 			}
 
 			// keep track of the max number of influences used by the vertices of the chunk
-			Chunk.MaxBoneInfluences = FMath::Max<int32>( Chunk.MaxBoneInfluences, MergeSectionInfo.Chunk->MaxBoneInfluences );
-			// update rigid/smooth verts totals
-			Chunk.NumRigidVertices += MergeSectionInfo.Chunk->NumRigidVertices;
-			Chunk.NumSoftVertices += MergeSectionInfo.Chunk->NumSoftVertices;
+			Section.MaxBoneInfluences = FMath::Max<int32>(Section.MaxBoneInfluences, MergeSectionInfo.Section->MaxBoneInfluences );
 
 			// update total number of vertices 
-			int32 NumTotalVertices = MergeSectionInfo.Chunk->NumRigidVertices + MergeSectionInfo.Chunk->NumSoftVertices;
+			int32 NumTotalVertices = MergeSectionInfo.Section->GetNumVertices();
 			MergeLODModel.NumVertices += NumTotalVertices;					
 
 			// add the vertices from the original source mesh to the merged vertex buffer					
 			int32 MaxVertIdx = FMath::Min<int32>( 
-				MergeSectionInfo.Chunk->BaseVertexIndex + NumTotalVertices, 
+				MergeSectionInfo.Section->BaseVertexIndex + NumTotalVertices,
 				SrcLODModel.VertexBufferGPUSkin.GetNumVertices() 
 				);
 
@@ -510,7 +498,7 @@ void FSkeletalMeshMerge::GenerateLODModel( int32 LODIdx )
 			// keep track of the current base vertex index before adding any new vertices
 			// this will be needed to remap the index buffer values to the new range
 			int32 CurrentBaseVertexIndex = MergedVertexBuffer.Num();
-			for( int32 VertIdx=MergeSectionInfo.Chunk->BaseVertexIndex; VertIdx < MaxVertIdx; VertIdx++ )
+			for( int32 VertIdx=MergeSectionInfo.Section->BaseVertexIndex; VertIdx < MaxVertIdx; VertIdx++ )
 			{
 				// add the new vertex
 				VertexDataType& DestVert = MergedVertexBuffer[MergedVertexBuffer.AddUninitialized()];
@@ -574,8 +562,8 @@ void FSkeletalMeshMerge::GenerateLODModel( int32 LODIdx )
 				uint32 SrcIndex = SrcLODModel.MultiSizeIndexContainer.GetIndexBuffer()->Get(IndexIdx);
 
 				// add offset to each index to match the new entries in the merged vertex buffer
-				checkSlow(SrcIndex >= MergeSectionInfo.Chunk->BaseVertexIndex);	
-				uint32 DstIndex = SrcIndex - MergeSectionInfo.Chunk->BaseVertexIndex + CurrentBaseVertexIndex;
+				checkSlow(SrcIndex >= MergeSectionInfo.Section->BaseVertexIndex);
+				uint32 DstIndex = SrcIndex - MergeSectionInfo.Section->BaseVertexIndex + CurrentBaseVertexIndex;
 				checkSlow(DstIndex < (uint32)MergedVertexBuffer.Num());
 
 				// add the new index to the merged vertex buffer
diff --git a/Engine/Source/Runtime/Engine/Private/SkeletalRender.cpp b/Engine/Source/Runtime/Engine/Private/SkeletalRender.cpp
--- a/Engine/Source/Runtime/Engine/Private/SkeletalRender.cpp
+++ b/Engine/Source/Runtime/Engine/Private/SkeletalRender.cpp
@@ -29,7 +29,6 @@ FSkeletalMeshObject::FSkeletalMeshObject(USkinnedMeshComponent* InMeshComponent,
 ,	WorkingMaxDistanceFactor(0.f)
 ,   bHasBeenUpdatedAtLeastOnce(false)
 #if WITH_EDITORONLY_DATA
-,   ChunkIndexPreview(InMeshComponent->ChunkIndexPreview)
 ,   SectionIndexPreview(InMeshComponent->SectionIndexPreview)
 #endif	
 ,	SkeletalMeshResource(InSkeletalMeshResource)
@@ -46,7 +45,6 @@ FSkeletalMeshObject::FSkeletalMeshObject(USkinnedMeshComponent* InMeshComponent,
 #if WITH_EDITORONLY_DATA
 	if ( !GIsEditor )
 	{
-		ChunkIndexPreview = -1;
 		SectionIndexPreview = -1;
 	}
 #endif // #if WITH_EDITORONLY_DATA
@@ -129,10 +127,10 @@ void FSkeletalMeshObject::UpdateMinDesiredLODLevel(const FSceneView* View, const
  * List of chunks to be rendered based on instance weight usage. Full swap of weights will render with its own chunks.
  * @return Chunks to iterate over for rendering
  */
-const TArray<FSkelMeshChunk>& FSkeletalMeshObject::GetRenderChunks(int32 InLODIndex) const
+const TArray<FSkelMeshSection>& FSkeletalMeshObject::GetRenderSections(int32 InLODIndex) const
 {
 	const FStaticLODModel& LOD = SkeletalMeshResource->LODModels[InLODIndex];
-	return LOD.Chunks;
+	return LOD.Sections;
 }
 
 /**
diff --git a/Engine/Source/Runtime/Engine/Private/SkeletalRenderCPUSkin.cpp b/Engine/Source/Runtime/Engine/Private/SkeletalRenderCPUSkin.cpp
--- a/Engine/Source/Runtime/Engine/Private/SkeletalRenderCPUSkin.cpp
+++ b/Engine/Source/Runtime/Engine/Private/SkeletalRenderCPUSkin.cpp
@@ -611,100 +611,33 @@ const VectorRegister		VECTOR4_UNPACK_MINUS_1	= DECLARE_VECTOR_REGISTER(-1.f, -1.
 
 const VectorRegister		VECTOR_0001				= DECLARE_VECTOR_REGISTER(0.f, 0.f, 0.f, 1.f);
 
-template<bool bExtraBoneInfluences, int32 MaxChunkBoneInfluences, typename BaseVertexType, typename VertexType>
-static void SkinVertexChunk( FFinalSkinVertex*& DestVertex, TArray<FMorphTargetInfo>& MorphEvalInfos, const TArray<float>& MorphWeights, const FSkelMeshChunk& Chunk, const FStaticLODModel &LOD, int32 VertexBufferBaseIndex, uint32 NumValidMorphs, int32 &CurBaseVertIdx, int32 LODIndex, int32 RigidInfluenceIndex, const FMatrix* RESTRICT ReferenceToLocal )
+template<bool bExtraBoneInfluences, int32 MaxSectionBoneInfluences, typename BaseVertexType, typename VertexType>
+static void SkinVertexSection( FFinalSkinVertex*& DestVertex, TArray<FMorphTargetInfo>& MorphEvalInfos, const TArray<float>& MorphWeights, const FSkelMeshSection& Section, const FStaticLODModel &LOD, int32 VertexBufferBaseIndex, uint32 NumValidMorphs, int32 &CurBaseVertIdx, int32 LODIndex, int32 RigidInfluenceIndex, const FMatrix* RESTRICT ReferenceToLocal )
 {
 	// VertexCopy for morph. Need to allocate right struct
 	// To avoid re-allocation, create 2 statics, and assign right struct
 	VertexType  VertexCopy;
 
 	// Prefetch all bone indices
-	const FBoneIndexType* BoneMap = Chunk.BoneMap.GetData();
+	const FBoneIndexType* BoneMap = Section.BoneMap.GetData();
 	FPlatformMisc::Prefetch( BoneMap );
 	FPlatformMisc::Prefetch( BoneMap, PLATFORM_CACHE_LINE_SIZE );
 
 	VertexType* SrcRigidVertex = NULL;
-	const int32 NumRigidVertices = Chunk.GetNumRigidVertices();
-	if (NumRigidVertices > 0)
-	{
-		INC_DWORD_STAT_BY(STAT_CPUSkinVertices,NumRigidVertices);
-
-		// Prefetch first vertex
-		FPlatformMisc::Prefetch( LOD.VertexBufferGPUSkin.GetVertexPtr<bExtraBoneInfluences>(Chunk.GetRigidVertexBufferIndex()) );
-
-		for(int32 VertexIndex = VertexBufferBaseIndex;VertexIndex < NumRigidVertices;VertexIndex++,DestVertex++)
-		{
-			int32 VertexBufferIndex = Chunk.GetRigidVertexBufferIndex() + VertexIndex;
-			SrcRigidVertex = (VertexType*)LOD.VertexBufferGPUSkin.GetVertexPtr<(bExtraBoneInfluences)>(VertexBufferIndex);
-			FPlatformMisc::Prefetch( SrcRigidVertex, PLATFORM_CACHE_LINE_SIZE );	// Prefetch next vertices
-			VertexType* MorphedVertex = SrcRigidVertex;
-			if( NumValidMorphs ) 
-			{
-				MorphedVertex = &VertexCopy;
-				UpdateMorphedVertex<VertexType>( *MorphedVertex, *SrcRigidVertex, CurBaseVertIdx, LODIndex, MorphEvalInfos, MorphWeights);
-			}
-
-			VectorRegister SrcNormals[3];
-			VectorRegister DstNormals[3];
-			const FVector VertexPosition = LOD.VertexBufferGPUSkin.GetVertexPositionFast((const BaseVertexType*)MorphedVertex);
-			SrcNormals[0] = VectorLoadFloat3_W1( &VertexPosition );
-			SrcNormals[1] = Unpack3( &MorphedVertex->TangentX.Vector.Packed );
-			SrcNormals[2] = Unpack4( &MorphedVertex->TangentZ.Vector.Packed );
-			VectorResetFloatRegisters(); // Need to call this to be able to use regular floating point registers again after Unpack().
-
-			uint8 BoneIndex = MorphedVertex->InfluenceBones[RigidInfluenceIndex];
-
-			const FMatrix BoneMatrix = ReferenceToLocal[BoneMap[BoneIndex]];
-			VectorRegister M00	= VectorLoadAligned( &BoneMatrix.M[0][0] );
-			VectorRegister M10	= VectorLoadAligned( &BoneMatrix.M[1][0] );
-			VectorRegister M20	= VectorLoadAligned( &BoneMatrix.M[2][0] );
-			VectorRegister M30	= VectorLoadAligned( &BoneMatrix.M[3][0] );
-
-			VectorRegister N_xxxx = VectorReplicate( SrcNormals[0], 0 );
-			VectorRegister N_yyyy = VectorReplicate( SrcNormals[0], 1 );
-			VectorRegister N_zzzz = VectorReplicate( SrcNormals[0], 2 );
-			DstNormals[0] = VectorMultiplyAdd( N_xxxx, M00, VectorMultiplyAdd( N_yyyy, M10, VectorMultiplyAdd( N_zzzz, M20, M30 ) ) );
-
-			N_xxxx = VectorReplicate( SrcNormals[1], 0 );
-			N_yyyy = VectorReplicate( SrcNormals[1], 1 );
-			N_zzzz = VectorReplicate( SrcNormals[1], 2 );
-			DstNormals[1] = VectorMultiplyAdd( N_xxxx, M00, VectorMultiplyAdd( N_yyyy, M10, VectorMultiply( N_zzzz, M20 ) ) );
-
-			N_xxxx = VectorReplicate( SrcNormals[2], 0 );
-			N_yyyy = VectorReplicate( SrcNormals[2], 1 );
-			N_zzzz = VectorReplicate( SrcNormals[2], 2 );
-			DstNormals[2] = VectorMultiplyAdd( N_xxxx, M00, VectorMultiplyAdd( N_yyyy, M10, VectorMultiply( N_zzzz, M20 ) ) );
-
-			// carry over the W component (sign of basis determinant) 
-			DstNormals[2] = VectorMultiplyAdd( VECTOR_0001, SrcNormals[2], DstNormals[2] );
-
-			// Write to 16-byte aligned memory:
-			VectorStore( DstNormals[0], &DestVertex->Position );
-			Pack3( DstNormals[1], &DestVertex->TangentX.Vector.Packed );
-			Pack4( DstNormals[2], &DestVertex->TangentZ.Vector.Packed );
-			VectorResetFloatRegisters(); // Need to call this to be able to use regular floating point registers again after Pack().
-
-			// Copy UVs.
-			FVector2D UVs = LOD.VertexBufferGPUSkin.GetVertexUVFast<bExtraBoneInfluences>(VertexBufferIndex,0);
-			DestVertex->U = UVs.X;
-			DestVertex->V = UVs.Y;
-
-			CurBaseVertIdx++;
-		}
-	}
+	
 
 	VertexType* SrcSoftVertex = NULL;
-	const int32 NumSoftVertices = Chunk.GetNumSoftVertices();
+	const int32 NumSoftVertices = Section.GetNumVertices();
 	if (NumSoftVertices > 0)
 	{
 		INC_DWORD_STAT_BY(STAT_CPUSkinVertices,NumSoftVertices);
 
 		// Prefetch first vertex
-		FPlatformMisc::Prefetch( LOD.VertexBufferGPUSkin.GetVertexPtr<(bExtraBoneInfluences)>(Chunk.GetSoftVertexBufferIndex()) );
+		FPlatformMisc::Prefetch( LOD.VertexBufferGPUSkin.GetVertexPtr<(bExtraBoneInfluences)>(Section.GetVertexBufferIndex()) );
 
 		for(int32 VertexIndex = VertexBufferBaseIndex;VertexIndex < NumSoftVertices;VertexIndex++,DestVertex++)
 		{
-			const int32 VertexBufferIndex = Chunk.GetSoftVertexBufferIndex() + VertexIndex;
+			const int32 VertexBufferIndex = Section.GetVertexBufferIndex() + VertexIndex;
 			SrcSoftVertex = (VertexType*)LOD.VertexBufferGPUSkin.GetVertexPtr<(bExtraBoneInfluences)>(VertexBufferIndex);
 			FPlatformMisc::Prefetch( SrcSoftVertex, PLATFORM_CACHE_LINE_SIZE );	// Prefetch next vertices
 			VertexType* MorphedVertex = SrcSoftVertex;
@@ -725,7 +658,7 @@ static void SkinVertexChunk( FFinalSkinVertex*& DestVertex, TArray<FMorphTargetI
 			SrcNormals[2] = Unpack4( &MorphedVertex->TangentZ.Vector.Packed );
 			VectorRegister Weights = VectorMultiply( VectorLoadByte4(BoneWeights), VECTOR_INV_255 );
 			VectorRegister ExtraWeights;
-			if (MaxChunkBoneInfluences > 4)
+			if (MaxSectionBoneInfluences > 4)
 			{
 				ExtraWeights = VectorMultiply( VectorLoadByte4(&BoneWeights[MAX_INFLUENCES_PER_STREAM]), VECTOR_INV_255 );
 			}
@@ -738,7 +671,7 @@ static void SkinVertexChunk( FFinalSkinVertex*& DestVertex, TArray<FMorphTargetI
 			VectorRegister M20	= VectorMultiply( VectorLoadAligned( &BoneMatrix0.M[2][0] ), Weight0 );
 			VectorRegister M30	= VectorMultiply( VectorLoadAligned( &BoneMatrix0.M[3][0] ), Weight0 );
 
-			if ( MaxChunkBoneInfluences > 1 )
+			if (MaxSectionBoneInfluences > 1 )
 			{
 				const FMatrix BoneMatrix1 = ReferenceToLocal[BoneMap[BoneIndices[INFLUENCE_1]]];
 				VectorRegister Weight1 = VectorReplicate( Weights, INFLUENCE_1 );
@@ -747,7 +680,7 @@ static void SkinVertexChunk( FFinalSkinVertex*& DestVertex, TArray<FMorphTargetI
 				M20	= VectorMultiplyAdd( VectorLoadAligned( &BoneMatrix1.M[2][0] ), Weight1, M20 );
 				M30	= VectorMultiplyAdd( VectorLoadAligned( &BoneMatrix1.M[3][0] ), Weight1, M30 );
 
-				if ( MaxChunkBoneInfluences > 2 )
+				if (MaxSectionBoneInfluences > 2 )
 				{
 					const FMatrix BoneMatrix2 = ReferenceToLocal[BoneMap[BoneIndices[INFLUENCE_2]]];
 					VectorRegister Weight2 = VectorReplicate( Weights, INFLUENCE_2 );
@@ -756,7 +689,7 @@ static void SkinVertexChunk( FFinalSkinVertex*& DestVertex, TArray<FMorphTargetI
 					M20	= VectorMultiplyAdd( VectorLoadAligned( &BoneMatrix2.M[2][0] ), Weight2, M20 );
 					M30	= VectorMultiplyAdd( VectorLoadAligned( &BoneMatrix2.M[3][0] ), Weight2, M30 );
 
-					if ( MaxChunkBoneInfluences > 3 )
+					if (MaxSectionBoneInfluences > 3 )
 					{
 						const FMatrix BoneMatrix3 = ReferenceToLocal[BoneMap[BoneIndices[INFLUENCE_3]]];
 						VectorRegister Weight3 = VectorReplicate( Weights, INFLUENCE_3 );
@@ -766,7 +699,7 @@ static void SkinVertexChunk( FFinalSkinVertex*& DestVertex, TArray<FMorphTargetI
 						M30	= VectorMultiplyAdd( VectorLoadAligned( &BoneMatrix3.M[3][0] ), Weight3, M30 );
 					}
 
-					if (MaxChunkBoneInfluences > 4)
+					if (MaxSectionBoneInfluences > 4)
 					{
 						const FMatrix BoneMatrix4 = ReferenceToLocal[BoneMap[BoneIndices[INFLUENCE_4]]];
 						VectorRegister Weight4 = VectorReplicate( ExtraWeights, INFLUENCE_4 - INFLUENCE_4 );
@@ -775,7 +708,7 @@ static void SkinVertexChunk( FFinalSkinVertex*& DestVertex, TArray<FMorphTargetI
 						M20	= VectorMultiplyAdd( VectorLoadAligned( &BoneMatrix4.M[2][0] ), Weight4, M20 );
 						M30	= VectorMultiplyAdd( VectorLoadAligned( &BoneMatrix4.M[3][0] ), Weight4, M30 );
 
-						if (MaxChunkBoneInfluences > 5)
+						if (MaxSectionBoneInfluences > 5)
 						{
 							const FMatrix BoneMatrix5 = ReferenceToLocal[BoneMap[BoneIndices[INFLUENCE_5]]];
 							VectorRegister Weight5 = VectorReplicate( ExtraWeights, INFLUENCE_5 - INFLUENCE_4 );
@@ -784,7 +717,7 @@ static void SkinVertexChunk( FFinalSkinVertex*& DestVertex, TArray<FMorphTargetI
 							M20	= VectorMultiplyAdd( VectorLoadAligned( &BoneMatrix5.M[2][0] ), Weight5, M20 );
 							M30	= VectorMultiplyAdd( VectorLoadAligned( &BoneMatrix5.M[3][0] ), Weight5, M30 );
 
-							if (MaxChunkBoneInfluences > 6)
+							if (MaxSectionBoneInfluences > 6)
 							{
 								const FMatrix BoneMatrix6 = ReferenceToLocal[BoneMap[BoneIndices[INFLUENCE_6]]];
 								VectorRegister Weight6 = VectorReplicate( ExtraWeights, INFLUENCE_6 - INFLUENCE_4 );
@@ -793,7 +726,7 @@ static void SkinVertexChunk( FFinalSkinVertex*& DestVertex, TArray<FMorphTargetI
 								M20	= VectorMultiplyAdd( VectorLoadAligned( &BoneMatrix6.M[2][0] ), Weight6, M20 );
 								M30	= VectorMultiplyAdd( VectorLoadAligned( &BoneMatrix6.M[3][0] ), Weight6, M30 );
 
-								if (MaxChunkBoneInfluences > 7)
+								if (MaxSectionBoneInfluences > 7)
 								{
 									const FMatrix BoneMatrix7 = ReferenceToLocal[BoneMap[BoneIndices[INFLUENCE_7]]];
 									VectorRegister Weight7 = VectorReplicate( ExtraWeights, INFLUENCE_7 - INFLUENCE_4 );
@@ -836,7 +769,7 @@ static void SkinVertexChunk( FFinalSkinVertex*& DestVertex, TArray<FMorphTargetI
 			VectorResetFloatRegisters(); // Need to call this to be able to use regular floating point registers again after Pack().
 
 			// Copy UVs.
-			FVector2D UVs = LOD.VertexBufferGPUSkin.GetVertexUVFast<bExtraBoneInfluences>(Chunk.GetSoftVertexBufferIndex()+VertexIndex,0);
+			FVector2D UVs = LOD.VertexBufferGPUSkin.GetVertexUVFast<bExtraBoneInfluences>(Section.GetVertexBufferIndex() + VertexIndex, 0);
 			DestVertex->U = UVs.X;
 			DestVertex->V = UVs.Y;
 
@@ -846,18 +779,18 @@ static void SkinVertexChunk( FFinalSkinVertex*& DestVertex, TArray<FMorphTargetI
 }
 
 template< bool bExtraBoneInfluences, typename BaseVertexType, typename VertexType>
-static void SkinVertexChunk( FFinalSkinVertex*& DestVertex, TArray<FMorphTargetInfo>& MorphEvalInfos, const TArray<float>& MorphWeights, const FSkelMeshChunk& Chunk, const FStaticLODModel &LOD, int32 VertexBufferBaseIndex, uint32 NumValidMorphs, int32 &CurBaseVertIdx, int32 LODIndex, int32 RigidInfluenceIndex, const FMatrix* RESTRICT ReferenceToLocal)
+static void SkinVertexSection( FFinalSkinVertex*& DestVertex, TArray<FMorphTargetInfo>& MorphEvalInfos, const TArray<float>& MorphWeights, const FSkelMeshSection& Section, const FStaticLODModel &LOD, int32 VertexBufferBaseIndex, uint32 NumValidMorphs, int32 &CurBaseVertIdx, int32 LODIndex, int32 RigidInfluenceIndex, const FMatrix* RESTRICT ReferenceToLocal)
 {
-	switch (Chunk.MaxBoneInfluences)
+	switch (Section.MaxBoneInfluences)
 	{
-		case 1: SkinVertexChunk<bExtraBoneInfluences, 1, BaseVertexType, VertexType>(DestVertex, MorphEvalInfos, MorphWeights, Chunk, LOD, VertexBufferBaseIndex, NumValidMorphs, CurBaseVertIdx, LODIndex, RigidInfluenceIndex, ReferenceToLocal); break;
-		case 2: SkinVertexChunk<bExtraBoneInfluences, 2, BaseVertexType, VertexType>(DestVertex, MorphEvalInfos, MorphWeights, Chunk, LOD, VertexBufferBaseIndex, NumValidMorphs, CurBaseVertIdx, LODIndex, RigidInfluenceIndex, ReferenceToLocal); break;
-		case 3: SkinVertexChunk<bExtraBoneInfluences, 3, BaseVertexType, VertexType>(DestVertex, MorphEvalInfos, MorphWeights, Chunk, LOD, VertexBufferBaseIndex, NumValidMorphs, CurBaseVertIdx, LODIndex, RigidInfluenceIndex, ReferenceToLocal); break;
-		case 4: SkinVertexChunk<bExtraBoneInfluences, 4, BaseVertexType, VertexType>(DestVertex, MorphEvalInfos, MorphWeights, Chunk, LOD, VertexBufferBaseIndex, NumValidMorphs, CurBaseVertIdx, LODIndex, RigidInfluenceIndex, ReferenceToLocal); break;
-		case 5: SkinVertexChunk<bExtraBoneInfluences, 5, BaseVertexType, VertexType>(DestVertex, MorphEvalInfos, MorphWeights, Chunk, LOD, VertexBufferBaseIndex, NumValidMorphs, CurBaseVertIdx, LODIndex, RigidInfluenceIndex, ReferenceToLocal); break;
-		case 6: SkinVertexChunk<bExtraBoneInfluences, 6, BaseVertexType, VertexType>(DestVertex, MorphEvalInfos, MorphWeights, Chunk, LOD, VertexBufferBaseIndex, NumValidMorphs, CurBaseVertIdx, LODIndex, RigidInfluenceIndex, ReferenceToLocal); break;
-		case 7: SkinVertexChunk<bExtraBoneInfluences, 7, BaseVertexType, VertexType>(DestVertex, MorphEvalInfos, MorphWeights, Chunk, LOD, VertexBufferBaseIndex, NumValidMorphs, CurBaseVertIdx, LODIndex, RigidInfluenceIndex, ReferenceToLocal); break;
-		case 8: SkinVertexChunk<bExtraBoneInfluences, 8, BaseVertexType, VertexType>(DestVertex, MorphEvalInfos, MorphWeights, Chunk, LOD, VertexBufferBaseIndex, NumValidMorphs, CurBaseVertIdx, LODIndex, RigidInfluenceIndex, ReferenceToLocal); break;
+		case 1: SkinVertexSection<bExtraBoneInfluences, 1, BaseVertexType, VertexType>(DestVertex, MorphEvalInfos, MorphWeights, Section, LOD, VertexBufferBaseIndex, NumValidMorphs, CurBaseVertIdx, LODIndex, RigidInfluenceIndex, ReferenceToLocal); break;
+		case 2: SkinVertexSection<bExtraBoneInfluences, 2, BaseVertexType, VertexType>(DestVertex, MorphEvalInfos, MorphWeights, Section, LOD, VertexBufferBaseIndex, NumValidMorphs, CurBaseVertIdx, LODIndex, RigidInfluenceIndex, ReferenceToLocal); break;
+		case 3: SkinVertexSection<bExtraBoneInfluences, 3, BaseVertexType, VertexType>(DestVertex, MorphEvalInfos, MorphWeights, Section, LOD, VertexBufferBaseIndex, NumValidMorphs, CurBaseVertIdx, LODIndex, RigidInfluenceIndex, ReferenceToLocal); break;
+		case 4: SkinVertexSection<bExtraBoneInfluences, 4, BaseVertexType, VertexType>(DestVertex, MorphEvalInfos, MorphWeights, Section, LOD, VertexBufferBaseIndex, NumValidMorphs, CurBaseVertIdx, LODIndex, RigidInfluenceIndex, ReferenceToLocal); break;
+		case 5: SkinVertexSection<bExtraBoneInfluences, 5, BaseVertexType, VertexType>(DestVertex, MorphEvalInfos, MorphWeights, Section, LOD, VertexBufferBaseIndex, NumValidMorphs, CurBaseVertIdx, LODIndex, RigidInfluenceIndex, ReferenceToLocal); break;
+		case 6: SkinVertexSection<bExtraBoneInfluences, 6, BaseVertexType, VertexType>(DestVertex, MorphEvalInfos, MorphWeights, Section, LOD, VertexBufferBaseIndex, NumValidMorphs, CurBaseVertIdx, LODIndex, RigidInfluenceIndex, ReferenceToLocal); break;
+		case 7: SkinVertexSection<bExtraBoneInfluences, 7, BaseVertexType, VertexType>(DestVertex, MorphEvalInfos, MorphWeights, Section, LOD, VertexBufferBaseIndex, NumValidMorphs, CurBaseVertIdx, LODIndex, RigidInfluenceIndex, ReferenceToLocal); break;
+		case 8: SkinVertexSection<bExtraBoneInfluences, 8, BaseVertexType, VertexType>(DestVertex, MorphEvalInfos, MorphWeights, Section, LOD, VertexBufferBaseIndex, NumValidMorphs, CurBaseVertIdx, LODIndex, RigidInfluenceIndex, ReferenceToLocal); break;
 		default: check(0);
 	}
 }
@@ -889,15 +822,14 @@ static void SkinVertices(FFinalSkinVertex* DestVertex, FMatrix* ReferenceToLocal
 	for(int32 SectionIndex= 0;SectionIndex< LOD.Sections.Num();SectionIndex++)
 	{
 		FSkelMeshSection& Section = LOD.Sections[SectionIndex];
-		FSkelMeshChunk& Chunk = LOD.Chunks[Section.ChunkIndex];
 
-		if (LOD.DoChunksNeedExtraBoneInfluences())
+		if (LOD.DoSectionsNeedExtraBoneInfluences())
 		{
-			SkinVertexChunk<true, BaseVertexType, VertexType>(DestVertex, MorphEvalInfos, MorphTargetWeights, Chunk, LOD, VertexBufferBaseIndex, NumValidMorphs, CurBaseVertIdx, LODIndex, RigidInfluenceIndex, ReferenceToLocal);
+			SkinVertexSection<true, BaseVertexType, VertexType>(DestVertex, MorphEvalInfos, MorphTargetWeights, Section, LOD, VertexBufferBaseIndex, NumValidMorphs, CurBaseVertIdx, LODIndex, RigidInfluenceIndex, ReferenceToLocal);
 		}
 		else
 		{
-			SkinVertexChunk<false, BaseVertexType, VertexType>(DestVertex, MorphEvalInfos, MorphTargetWeights, Chunk, LOD, VertexBufferBaseIndex, NumValidMorphs, CurBaseVertIdx, LODIndex, RigidInfluenceIndex, ReferenceToLocal);
+			SkinVertexSection<false, BaseVertexType, VertexType>(DestVertex, MorphEvalInfos, MorphTargetWeights, Section, LOD, VertexBufferBaseIndex, NumValidMorphs, CurBaseVertIdx, LODIndex, RigidInfluenceIndex, ReferenceToLocal);
 		}
 	}
 
@@ -927,7 +859,7 @@ FVector4 GetTangetToColor(FPackedNormal Tangent)
  * @param BonesOfInterest - array of bones we want to display
  */
 template <bool bExtraBoneInfluencesT>
-static FORCEINLINE void CalculateChunkBoneWeights(FFinalSkinVertex*& DestVertex, FSkeletalMeshVertexBuffer& VertexBufferGPUSkin, FSkelMeshChunk& Chunk, const TArray<int32>& BonesOfInterest)
+static FORCEINLINE void CalculateSectionBoneWeights(FFinalSkinVertex*& DestVertex, FSkeletalMeshVertexBuffer& VertexBufferGPUSkin, FSkelMeshSection& Section, const TArray<int32>& BonesOfInterest)
 {
 	const float INV255 = 1.f/255.f;
 
@@ -936,30 +868,11 @@ static FORCEINLINE void CalculateChunkBoneWeights(FFinalSkinVertex*& DestVertex,
 	int32 VertexBufferBaseIndex = 0;
 
 	//array of bone mapping
-	FBoneIndexType* BoneMap = Chunk.BoneMap.GetData();
-
-	for(int32 VertexIndex = VertexBufferBaseIndex;VertexIndex < Chunk.GetNumRigidVertices();VertexIndex++,DestVertex++)
-	{
-		int32 VertexBufferIndex = Chunk.GetRigidVertexBufferIndex() + VertexIndex;
-		auto* SrcRigidVertex = VertexBufferGPUSkin.GetVertexPtr<bExtraBoneInfluencesT>(VertexBufferIndex);
-
-		uint8 BoneIndex = SrcRigidVertex->InfluenceBones[RigidInfluenceIndex];
-
-		if (BonesOfInterest.Contains(BoneMap[BoneIndex]))
-		{
-			DestVertex->U = 1.f; 
-			DestVertex->V = 1.f; 
-		}
-		else
-		{
-			DestVertex->U = 0.0f;
-			DestVertex->V = 0.0f;
-		}
-	}
+	FBoneIndexType* BoneMap = Section.BoneMap.GetData();
 
-	for(int32 VertexIndex = VertexBufferBaseIndex;VertexIndex < Chunk.GetNumSoftVertices();VertexIndex++,DestVertex++)
+	for(int32 VertexIndex = VertexBufferBaseIndex;VertexIndex < Section.GetNumVertices();VertexIndex++,DestVertex++)
 	{
-		const int32 VertexBufferIndex = Chunk.GetSoftVertexBufferIndex() + VertexIndex;
+		const int32 VertexBufferIndex = Section.GetVertexBufferIndex() + VertexIndex;
 		auto* SrcSoftVertex = VertexBufferGPUSkin.GetVertexPtr<bExtraBoneInfluencesT>(VertexBufferIndex);
 
 		//Zero out the UV coords
@@ -997,15 +910,14 @@ static void CalculateBoneWeights(FFinalSkinVertex* DestVertex, FStaticLODModel&
 	for(int32 SectionIndex= 0;SectionIndex< LOD.Sections.Num();SectionIndex++)
 	{
 		FSkelMeshSection& Section = LOD.Sections[SectionIndex];
-		FSkelMeshChunk& Chunk = LOD.Chunks[Section.ChunkIndex];
 
 		if (LOD.VertexBufferGPUSkin.HasExtraBoneInfluences())
 		{
-			CalculateChunkBoneWeights<true>(DestVertex, LOD.VertexBufferGPUSkin, Chunk, BonesOfInterest);
+			CalculateSectionBoneWeights<true>(DestVertex, LOD.VertexBufferGPUSkin, Section, BonesOfInterest);
 		}
 		else
 		{
-			CalculateChunkBoneWeights<false>(DestVertex, LOD.VertexBufferGPUSkin, Chunk, BonesOfInterest);
+			CalculateSectionBoneWeights<false>(DestVertex, LOD.VertexBufferGPUSkin, Section, BonesOfInterest);
 		}
 	}
 }
diff --git a/Engine/Source/Runtime/Engine/Private/SkeletalRenderGPUSkin.cpp b/Engine/Source/Runtime/Engine/Private/SkeletalRenderGPUSkin.cpp
--- a/Engine/Source/Runtime/Engine/Private/SkeletalRenderGPUSkin.cpp
+++ b/Engine/Source/Runtime/Engine/Private/SkeletalRenderGPUSkin.cpp
@@ -78,16 +78,17 @@ void FMorphVertexBuffer::InitDynamicRHI()
 	FRHIResourceCreateInfo CreateInfo;
 	void* BufferData = nullptr;
 
-	const bool bMayUseGPUSkinCache = GEnableGPUSkinCacheShaders != 0;
+	bool bSupportsComputeShaders = RHISupportsComputeShaders(GMaxRHIShaderPlatform);
+	bUsesComputeShader = GUseGPUMorphTargets != 0 && bSupportsComputeShaders;
 
-	EBufferUsageFlags Flags = bMayUseGPUSkinCache ? BUF_Dynamic : (EBufferUsageFlags)(BUF_Static | BUF_UnorderedAccess);
+	EBufferUsageFlags Flags = bUsesComputeShader ? (EBufferUsageFlags)(BUF_Static | BUF_UnorderedAccess) : BUF_Dynamic;
 
 	// BUF_ShaderResource is needed for Morph support of the SkinCache
 	Flags = (EBufferUsageFlags)(Flags | BUF_ShaderResource);
 
 	VertexBufferRHI = RHICreateAndLockVertexBuffer(Size, Flags, CreateInfo, BufferData);
-
-	if (bMayUseGPUSkinCache)
+	bool bUsesSkinCache = bSupportsComputeShaders && GEnableGPUSkinCacheShaders && GEnableGPUSkinCache;
+	if (bUsesSkinCache)
 	{
 		SRVValue = RHICreateShaderResourceView(VertexBufferRHI, 4, PF_R32_FLOAT);
 	}
@@ -98,7 +99,7 @@ void FMorphVertexBuffer::InitDynamicRHI()
 	// Unlock the buffer.
 	RHIUnlockVertexBuffer(VertexBufferRHI);
 	
-	if (!bMayUseGPUSkinCache && RHISupportsComputeShaders(GMaxRHIShaderPlatform))
+	if (bUsesComputeShader)
 	{
 		UAVValue = RHICreateUnorderedAccessView(VertexBufferRHI, PF_R32_UINT);
 	}
@@ -370,17 +371,19 @@ void FSkeletalMeshObjectGPUSkin::ProcessUpdatedDynamicData(FRHICommandListImmedi
 	bool bMorph = DynamicData->NumWeightedActiveMorphTargets > 0;
 
 	const FStaticLODModel& LODModel = SkeletalMeshResource->LODModels[DynamicData->LODIndex];
-	const TArray<FSkelMeshChunk>& Chunks = GetRenderChunks(DynamicData->LODIndex);
+	const TArray<FSkelMeshSection>& Sections = GetRenderSections(DynamicData->LODIndex);
 
 	// use correct vertex factories based on alternate weights usage
 	FVertexFactoryData& VertexFactoryData = LOD.GPUSkinVertexFactories;
 
-	bool DataPresent = false;
+	bool bDataPresent = false;
+
+	bool bGPUSkinCacheEnabled = GEnableGPUSkinCache && (FeatureLevel >= ERHIFeatureLevel::SM5);
 
 	if(bMorph) 
 	{
-		DataPresent = true;
-		checkSlow((VertexFactoryData.MorphVertexFactories.Num() == Chunks.Num()));
+		bDataPresent = true;
+		checkSlow((VertexFactoryData.MorphVertexFactories.Num() == Sections.Num()));
 		
 		// only update if the morph data changed and there are weighted morph targets
 		if(bMorphNeedsUpdate)
@@ -401,33 +404,32 @@ void FSkeletalMeshObjectGPUSkin::ProcessUpdatedDynamicData(FRHICommandListImmedi
 	else
 	{
 //		checkSlow(VertexFactoryData.MorphVertexFactories.Num() == 0);
-		DataPresent = VertexFactoryData.VertexFactories.Num() > 0;
+		bDataPresent = VertexFactoryData.VertexFactories.Num() > 0;
 	}
 
-	if (DataPresent)
+	if (bDataPresent)
 	{
-		bool bGPUSkinCacheEnabled = GEnableGPUSkinCache && (FeatureLevel >= ERHIFeatureLevel::SM5);
-		for (int32 ChunkIdx = 0; ChunkIdx < Chunks.Num(); ChunkIdx++)
+		for (int32 SectionIdx = 0; SectionIdx < Sections.Num(); SectionIdx++)
 		{
-			const FSkelMeshChunk& Chunk = Chunks[ChunkIdx];
+			const FSkelMeshSection& Section = Sections[SectionIdx];
 
-			bool bClothFactory = (FeatureLevel >= ERHIFeatureLevel::SM4) && (DynamicData->ClothSimulUpdateData.Num() > 0) && Chunk.HasApexClothData();
+			bool bClothFactory = (FeatureLevel >= ERHIFeatureLevel::SM4) && (DynamicData->ClothSimulUpdateData.Num() > 0) && Section.HasApexClothData();
 
 			FGPUBaseSkinVertexFactory* VertexFactory;
 			{
 				if(bClothFactory)
 				{
-					VertexFactory = VertexFactoryData.ClothVertexFactories[ChunkIdx]->GetVertexFactory();
+					VertexFactory = VertexFactoryData.ClothVertexFactories[SectionIdx]->GetVertexFactory();
 				}
 				else
 				{
 					if(DynamicData->NumWeightedActiveMorphTargets > 0)
 					{
-						VertexFactory = VertexFactoryData.MorphVertexFactories[ChunkIdx].Get();
+						VertexFactory = VertexFactoryData.MorphVertexFactories[SectionIdx].Get();
 					}
 					else
 					{
-						VertexFactory = VertexFactoryData.VertexFactories[ChunkIdx].Get();
+						VertexFactory = VertexFactoryData.VertexFactories[SectionIdx].Get();
 					}
 				}
 			}
@@ -437,7 +439,7 @@ void FSkeletalMeshObjectGPUSkin::ProcessUpdatedDynamicData(FRHICommandListImmedi
 			bool bUseSkinCache = bGPUSkinCacheEnabled;
 			if (bUseSkinCache)
 			{
-				if (ChunkIdx >= MAX_GPUSKINCACHE_CHUNKS_PER_LOD)
+				if (SectionIdx >= MAX_GPUSKINCACHE_CHUNKS_PER_LOD)
 				{
 					INC_DWORD_STAT(STAT_GPUSkinCache_SkippedForMaxChunksPerLOD);
 					bUseSkinCache = false;
@@ -447,7 +449,7 @@ void FSkeletalMeshObjectGPUSkin::ProcessUpdatedDynamicData(FRHICommandListImmedi
 					INC_DWORD_STAT(STAT_GPUSkinCache_SkippedForCloth);
 					bUseSkinCache = false;
 				}
-				else if (Chunk.MaxBoneInfluences == 0)
+				else if (Section.MaxBoneInfluences == 0)
 				{
 					INC_DWORD_STAT(STAT_GPUSkinCache_SkippedForZeroInfluences);
 					bUseSkinCache = false;
@@ -457,9 +459,7 @@ void FSkeletalMeshObjectGPUSkin::ProcessUpdatedDynamicData(FRHICommandListImmedi
 					extern TAutoConsoleVariable<int32> CVarGPUSkinCacheRecomputeTangents;
 					if(CVarGPUSkinCacheRecomputeTangents.GetValueOnRenderThread() == 2)
 					{
-						const FSkelMeshSection* Section = LODModel.FindSectionForChunk(ChunkIdx);
-
-						if(!Section->bRecomputeTangent)
+						if(!Section.bRecomputeTangent)
 						{
 							// In this mode the SkinCache should only be used for RecomputeTangent meshes
 							bUseSkinCache = false;
@@ -470,13 +470,17 @@ void FSkeletalMeshObjectGPUSkin::ProcessUpdatedDynamicData(FRHICommandListImmedi
 
 			// Create a uniform buffer from the bone transforms.
 			TArray<FMatrix>& ReferenceToLocalMatrices = DynamicData->ReferenceToLocal;
-			bool bNeedFence = ShaderData.UpdateBoneData(RHICmdList, ReferenceToLocalMatrices, Chunk.BoneMap, FrameNumberToPrepare, FeatureLevel, bUseSkinCache);
+			bool bNeedFence = ShaderData.UpdateBoneData(RHICmdList, ReferenceToLocalMatrices, Section.BoneMap, FrameNumberToPrepare, FeatureLevel, bUseSkinCache);
 
 			// Try to use the GPU skinning cache if possible
 			if (bUseSkinCache)
 			{
-				int32 Key = GGPUSkinCache.StartCacheMesh(RHICmdList, FrameNumberToPrepare, GPUSkinCacheKeys[ChunkIdx], VertexFactory,
-					VertexFactoryData.PassthroughVertexFactories[ChunkIdx].Get(), Chunk, this, bMorph ? &LOD.MorphVertexBuffer : 0);
+				if (GUseGPUMorphTargets)
+				{
+					// When rewriting the skin cache, first do all morph compute jobs, then all skin cache compute jobs, and then all recompute tangent jobs instead of sequentially
+				}
+				int32 Key = GGPUSkinCache.StartCacheMesh(RHICmdList, FrameNumberToPrepare, GPUSkinCacheKeys[SectionIdx], VertexFactory,
+					VertexFactoryData.PassthroughVertexFactories[SectionIdx].Get(), Section, this, bMorph ? &LOD.MorphVertexBuffer : 0);
 
 				// TODO: If failed we could remove the key (on the other hand it might a bu in cache for longer)
 
@@ -484,7 +488,7 @@ void FSkeletalMeshObjectGPUSkin::ProcessUpdatedDynamicData(FRHICommandListImmedi
 				if(Key >= 0)
 				{
 					ensure(Key <= 0xffff);
-					GPUSkinCacheKeys[ChunkIdx] = (int16)Key;
+					GPUSkinCacheKeys[SectionIdx] = (int16)Key;
 				}
 			}
 
@@ -492,9 +496,9 @@ void FSkeletalMeshObjectGPUSkin::ProcessUpdatedDynamicData(FRHICommandListImmedi
 			// Update uniform buffer for APEX cloth simulation mesh positions and normals
 			if( bClothFactory )
 			{				
-				FGPUBaseSkinAPEXClothVertexFactory::ClothShaderType& ClothShaderData = VertexFactoryData.ClothVertexFactories[ChunkIdx]->GetClothShaderData();
+				FGPUBaseSkinAPEXClothVertexFactory::ClothShaderType& ClothShaderData = VertexFactoryData.ClothVertexFactories[SectionIdx]->GetClothShaderData();
 				ClothShaderData.ClothBlendWeight = DynamicData->ClothBlendWeight;
-				int16 ActorIdx = Chunk.CorrespondClothAssetIndex;
+				int16 ActorIdx = Section.CorrespondClothAssetIndex;
 				if( FClothSimulData* SimData = DynamicData->ClothSimulUpdateData.Find(ActorIdx) )
 				{
 					bNeedFence = ClothShaderData.UpdateClothSimulData(RHICmdList, SimData->ClothSimulPositions, SimData->ClothSimulNormals, FrameNumberToPrepare, FeatureLevel) || bNeedFence;
@@ -576,25 +580,28 @@ void FSkeletalMeshObjectGPUSkin::FSkeletalMeshObjectLOD::UpdateMorphVertexBuffer
 		// LOD of the skel mesh is used to find number of vertices in buffer
 		FStaticLODModel& LodModel = SkelMeshResource->LODModels[LODIndex];
 
+		MorphVertexBuffer.RecreateResourcesIfRequired(GUseGPUMorphTargets != 0);
+
 		if (GUseGPUMorphTargets)
 		{
-			if (GEnableGPUSkinCacheShaders == 0)
-			{
-				RHICmdList.TransitionResource(EResourceTransitionAccess::EWritable, EResourceTransitionPipeline::EGfxToCompute, MorphVertexBuffer.GetUAV());
+			SCOPED_DRAW_EVENTF(RHICmdList, MorphUpdate,
+				TEXT("MorphUpdate LodVertices=%d InflVerts=%d"),
+				LodModel.NumVertices,
+				NumInfluencedVerticesByMorph);
+			RHICmdList.TransitionResource(EResourceTransitionAccess::EWritable, EResourceTransitionPipeline::EGfxToCompute, MorphVertexBuffer.GetUAV());
 
-				RHICmdList.ClearUAV(MorphVertexBuffer.GetUAV(), {0, 0, 0, 0});
+			RHICmdList.ClearUAV(MorphVertexBuffer.GetUAV(), {0, 0, 0, 0});
 
-				{
-					SCOPE_CYCLE_COUNTER(STAT_MorphVertexBuffer_ApplyDelta);
-					TShaderMapRef<FGPUMorphUpdateCS> GPUMorphUpdateCS(GetGlobalShaderMap(ERHIFeatureLevel::SM5));
-					GPUMorphUpdateCS->Dispatch(RHICmdList, LodModel.NumVertices, NumInfluencedVerticesByMorph, MorphWeightsSRV, MorphPerVertexInfoSRV, MorphFlattenedSRV, MorphVertexBuffer.GetUAV());
-				}
+			{
+				SCOPE_CYCLE_COUNTER(STAT_MorphVertexBuffer_ApplyDelta);
+				TShaderMapRef<FGPUMorphUpdateCS> GPUMorphUpdateCS(GetGlobalShaderMap(ERHIFeatureLevel::SM5));
+				GPUMorphUpdateCS->Dispatch(RHICmdList, LodModel.NumVertices, NumInfluencedVerticesByMorph, MorphWeightsSRV, MorphPerVertexInfoSRV, MorphFlattenedSRV, MorphVertexBuffer.GetUAV());
+			}
 
-				RHICmdList.TransitionResource(EResourceTransitionAccess::EReadable, EResourceTransitionPipeline::EComputeToGfx, MorphVertexBuffer.GetUAV());
+			RHICmdList.TransitionResource(EResourceTransitionAccess::EReadable, EResourceTransitionPipeline::EComputeToGfx, MorphVertexBuffer.GetUAV());
 
-				// set update flag
-				MorphVertexBuffer.bHasBeenUpdated = true;
-			}
+			// set update flag
+			MorphVertexBuffer.bHasBeenUpdated = true;
 		}
 		else
 		{
@@ -1036,15 +1043,15 @@ void FSkeletalMeshObjectGPUSkin::FSkeletalMeshObjectLOD::GetVertexBuffers(
  */
 void FSkeletalMeshObjectGPUSkin::FVertexFactoryData::InitVertexFactories(
 	const FVertexFactoryBuffers& VertexBuffers, 
-	const TArray<FSkelMeshChunk>& Chunks, 
+	const TArray<FSkelMeshSection>& Sections, 
 	ERHIFeatureLevel::Type InFeatureLevel)
 {
 	// first clear existing factories (resources assumed to have been released already)
 	// then [re]create the factories
 
-	VertexFactories.Empty(Chunks.Num());
+	VertexFactories.Empty(Sections.Num());
 	{
-		for(int32 FactoryIdx = 0; FactoryIdx < Chunks.Num(); ++FactoryIdx)
+		for(int32 FactoryIdx = 0; FactoryIdx < Sections.Num(); ++FactoryIdx)
 		{
 			if (VertexBuffers.VertexBufferGPUSkin->HasExtraBoneInfluences())
 			{
@@ -1079,13 +1086,13 @@ void FSkeletalMeshObjectGPUSkin::FVertexFactoryData::ReleaseVertexFactories()
 
 void FSkeletalMeshObjectGPUSkin::FVertexFactoryData::InitMorphVertexFactories(
 	const FVertexFactoryBuffers& VertexBuffers, 
-	const TArray<FSkelMeshChunk>& Chunks,
+	const TArray<FSkelMeshSection>& Sections,
 	bool bInUsePerBoneMotionBlur,
 	ERHIFeatureLevel::Type InFeatureLevel)
 {
 	// clear existing factories (resources assumed to have been released already)
-	MorphVertexFactories.Empty(Chunks.Num());
-	for( int32 FactoryIdx=0; FactoryIdx < Chunks.Num(); FactoryIdx++ )
+	MorphVertexFactories.Empty(Sections.Num());
+	for( int32 FactoryIdx=0; FactoryIdx < Sections.Num(); FactoryIdx++ )
 	{
 		if (VertexBuffers.VertexBufferGPUSkin->HasExtraBoneInfluences())
 		{
@@ -1112,14 +1119,14 @@ void FSkeletalMeshObjectGPUSkin::FVertexFactoryData::ReleaseMorphVertexFactories
 
 void FSkeletalMeshObjectGPUSkin::FVertexFactoryData::InitAPEXClothVertexFactories(
 	const FVertexFactoryBuffers& VertexBuffers, 
-	const TArray<FSkelMeshChunk>& Chunks,
+	const TArray<FSkelMeshSection>& Sections,
 	ERHIFeatureLevel::Type InFeatureLevel)
 {
 	// clear existing factories (resources assumed to have been released already)
-	ClothVertexFactories.Empty(Chunks.Num());
-	for( int32 FactoryIdx=0; FactoryIdx < Chunks.Num(); FactoryIdx++ )
+	ClothVertexFactories.Empty(Sections.Num());
+	for( int32 FactoryIdx=0; FactoryIdx < Sections.Num(); FactoryIdx++ )
 	{
-		if (Chunks[FactoryIdx].HasApexClothData() && InFeatureLevel >= ERHIFeatureLevel::SM4)
+		if (Sections[FactoryIdx].HasApexClothData() && InFeatureLevel >= ERHIFeatureLevel::SM4)
 		{
 			if (VertexBuffers.VertexBufferGPUSkin->HasExtraBoneInfluences())
 			{
@@ -1170,10 +1177,10 @@ void FSkeletalMeshObjectGPUSkin::FSkeletalMeshObjectLOD::InitResources(const FSk
 	GetVertexBuffers(VertexBuffers,LODModel,MeshLODInfo);
 
 	// init gpu skin factories
-	GPUSkinVertexFactories.InitVertexFactories(VertexBuffers,LODModel.Chunks, InFeatureLevel);
+	GPUSkinVertexFactories.InitVertexFactories(VertexBuffers,LODModel.Sections, InFeatureLevel);
 	if ( LODModel.HasApexClothData() )
 	{
-		GPUSkinVertexFactories.InitAPEXClothVertexFactories(VertexBuffers,LODModel.Chunks, InFeatureLevel);
+		GPUSkinVertexFactories.InitAPEXClothVertexFactories(VertexBuffers,LODModel.Sections, InFeatureLevel);
 	}
 }
 
@@ -1204,7 +1211,7 @@ void FSkeletalMeshObjectGPUSkin::FSkeletalMeshObjectLOD::InitMorphResources(cons
 	FVertexFactoryBuffers VertexBuffers;
 	GetVertexBuffers(VertexBuffers,LODModel,MeshLODInfo);
 	// init morph skin factories
-	GPUSkinVertexFactories.InitMorphVertexFactories(VertexBuffers, LODModel.Chunks, bInUsePerBoneMotionBlur, InFeatureLevel);
+	GPUSkinVertexFactories.InitMorphVertexFactories(VertexBuffers, LODModel.Sections, bInUsePerBoneMotionBlur, InFeatureLevel);
 }
 
 /** 
diff --git a/Engine/Source/Runtime/Engine/Private/SkeletalRenderGPUSkin.h b/Engine/Source/Runtime/Engine/Private/SkeletalRenderGPUSkin.h
--- a/Engine/Source/Runtime/Engine/Private/SkeletalRenderGPUSkin.h
+++ b/Engine/Source/Runtime/Engine/Private/SkeletalRenderGPUSkin.h
@@ -140,6 +140,7 @@ public:
 	{
 		check(SkelMeshResource);
 		check(SkelMeshResource->LODModels.IsValidIndex(LODIdx));
+		bUsesComputeShader = false;
 	}
 	/** 
 	* Initialize the dynamic RHI for this rendering resource 
@@ -151,6 +152,14 @@ public:
 	*/
 	virtual void ReleaseDynamicRHI();
 
+	inline void RecreateResourcesIfRequired(bool bInUsesComputeShader)
+	{
+		if (bUsesComputeShader != bInUsesComputeShader)
+		{
+			UpdateRHI();
+		}
+	}
+
 	/** 
 	* Morph target vertex name 
 	*/
@@ -199,6 +208,8 @@ protected:
 	// guaranteed only to be valid if the vertex buffer is valid
 	FUnorderedAccessViewRHIRef UAVValue;
 
+	bool bUsesComputeShader;
+
 private:
 	/** index to the SkelMeshResource.LODModels */
 	int32	LODIdx;
@@ -315,9 +326,9 @@ private:
 		 * Init default vertex factory resources for this LOD 
 		 *
 		 * @param VertexBuffers - available vertex buffers to reference in vertex factory streams
-		 * @param Chunks - relevant chunk information (either original or from swapped influence)
+		 * @param Sections - relevant section information (either original or from swapped influence)
 		 */
-		void InitVertexFactories(const FVertexFactoryBuffers& VertexBuffers, const TArray<FSkelMeshChunk>& Chunks, ERHIFeatureLevel::Type FeatureLevel);
+		void InitVertexFactories(const FVertexFactoryBuffers& VertexBuffers, const TArray<FSkelMeshSection>& Sections, ERHIFeatureLevel::Type FeatureLevel);
 		/** 
 		 * Release default vertex factory resources for this LOD 
 		 */
@@ -326,9 +337,9 @@ private:
 		 * Init morph vertex factory resources for this LOD 
 		 *
 		 * @param VertexBuffers - available vertex buffers to reference in vertex factory streams
-		 * @param Chunks - relevant chunk information (either original or from swapped influence)
+		 * @param Sections - relevant section information (either original or from swapped influence)
 		 */
-		void InitMorphVertexFactories(const FVertexFactoryBuffers& VertexBuffers, const TArray<FSkelMeshChunk>& Chunks, bool bInUsePerBoneMotionBlur, ERHIFeatureLevel::Type InFeatureLevel);
+		void InitMorphVertexFactories(const FVertexFactoryBuffers& VertexBuffers, const TArray<FSkelMeshSection>& Sections, bool bInUsePerBoneMotionBlur, ERHIFeatureLevel::Type InFeatureLevel);
 		/** 
 		 * Release morph vertex factory resources for this LOD 
 		 */
@@ -337,9 +348,9 @@ private:
 		 * Init APEX cloth vertex factory resources for this LOD 
 		 *
 		 * @param VertexBuffers - available vertex buffers to reference in vertex factory streams
-		 * @param Chunks - relevant chunk information (either original or from swapped influence)
+		 * @param Sections - relevant section information (either original or from swapped influence)
 		 */
-		void InitAPEXClothVertexFactories(const FVertexFactoryBuffers& VertexBuffers, const TArray<FSkelMeshChunk>& Chunks, ERHIFeatureLevel::Type InFeatureLevel);
+		void InitAPEXClothVertexFactories(const FVertexFactoryBuffers& VertexBuffers, const TArray<FSkelMeshSection>& Sections, ERHIFeatureLevel::Type InFeatureLevel);
 		/** 
 		 * Release morph vertex factory resources for this LOD 
 		 */
diff --git a/Engine/Source/Runtime/Engine/Private/UnrealEngine.cpp b/Engine/Source/Runtime/Engine/Private/UnrealEngine.cpp
--- a/Engine/Source/Runtime/Engine/Private/UnrealEngine.cpp
+++ b/Engine/Source/Runtime/Engine/Private/UnrealEngine.cpp
@@ -8684,7 +8684,7 @@ bool UEngine::HandleOpenCommand( const TCHAR* Cmd, FOutputDevice& Ar, UWorld *In
 		}
 	}
 
-	SetClientTravel( InWorld, Cmd, TRAVEL_Absolute );
+	SetClientTravel( InWorld, *TestURL.Map, TRAVEL_Absolute );
 	return true;
 }
 
