Copying //UE4/Dev-Framework to Dev-Main (//UE4/Dev-Main)

diff --git a/Engine/Source/Runtime/Engine/Private/Vehicles/WheeledVehicleMovementComponent.cpp b/Engine/Source/Runtime/Engine/Private/Vehicles/WheeledVehicleMovementComponent.cpp
--- a/Engine/Source/Runtime/Engine/Private/Vehicles/WheeledVehicleMovementComponent.cpp
+++ b/Engine/Source/Runtime/Engine/Private/Vehicles/WheeledVehicleMovementComponent.cpp
@@ -9,6 +9,7 @@
 #include "Vehicles/WheeledVehicleMovementComponent.h"
 #include "Vehicles/TireType.h"
 #include "DisplayDebugHelpers.h"
+#include "PhysicsEngine/ConstraintInstance.h"
 
 #if WITH_PHYSX
 #include "../PhysicsEngine/PhysXSupport.h"
@@ -1562,7 +1563,11 @@ void UWheeledVehicleMovementComponent::CalculateAvoidanceVelocity(float DeltaTim
 	
 	UAvoidanceManager* AvoidanceManager = GetWorld()->GetAvoidanceManager();
 	APawn* MyOwner = UpdatedComponent ? Cast<APawn>(UpdatedComponent->GetOwner()) : NULL;
-	
+
+	// since we don't assign the avoidance velocity but instead use it to adjust steering and throttle,
+	// always reset the avoidance velocity to the current velocity
+	AvoidanceVelocity = GetVelocityForRVOConsideration();
+
 	if (AvoidanceWeight >= 1.0f || AvoidanceManager == NULL || MyOwner == NULL)
 	{
 		return;
@@ -1577,10 +1582,6 @@ void UWheeledVehicleMovementComponent::CalculateAvoidanceVelocity(float DeltaTim
 	const bool bShowDebug = AvoidanceManager->IsDebugEnabled(AvoidanceUID);
 #endif
 
-	// since we don't assign the avoidance velocity but instead use it to adjust steering and throttle,
-	// always reset the avoidance velocity to the current velocity
-	AvoidanceVelocity = GetVelocityForRVOConsideration();
-
 	if (!AvoidanceVelocity.IsZero())
 	{
 		//See if we're doing a locked avoidance move already, and if so, skip the testing and just do the move.
diff --git a/Engine/Source/Runtime/Engine/Private/Volume.cpp b/Engine/Source/Runtime/Engine/Private/Volume.cpp
--- a/Engine/Source/Runtime/Engine/Private/Volume.cpp
+++ b/Engine/Source/Runtime/Engine/Private/Volume.cpp
@@ -56,22 +56,31 @@ void AVolume::PostEditChangeProperty(FPropertyChangedEvent& PropertyChangedEvent
 
 bool AVolume::EncompassesPoint(FVector Point, float SphereRadius/*=0.f*/, float* OutDistanceToPoint) const
 {
-	if(GetBrushComponent())
+	if (GetBrushComponent())
 	{
 #if WITH_PHYSX
 		FVector ClosestPoint;
-		float Distance = GetBrushComponent()->GetDistanceToCollision(Point, ClosestPoint);
+		float DistanceSqr;
+
+		if (GetBrushComponent()->GetSquaredDistanceToCollision(Point, DistanceSqr, ClosestPoint) == false)
+		{
+			if (OutDistanceToPoint)
+			{
+				*OutDistanceToPoint = -1.f;
+			}
+			return false;
+		}
 #else
 		FBoxSphereBounds Bounds = BrushComponent->CalcBounds(BrushComponent->ComponentToWorld);
-		float Distance = FMath::Sqrt(Bounds.GetBox().ComputeSquaredDistanceToPoint(Point));
+		const float DistanceSqr = Bounds.GetBox().ComputeSquaredDistanceToPoint(Point);
 #endif
 
-		if(OutDistanceToPoint)
+		if (OutDistanceToPoint)
 		{
-			*OutDistanceToPoint = Distance;
+			*OutDistanceToPoint = FMath::Sqrt(DistanceSqr);
 		}
 
-		return Distance >= 0.f && Distance <= SphereRadius;
+		return DistanceSqr >= 0.f && DistanceSqr <= FMath::Square(SphereRadius);
 	}
 	else
 	{
diff --git a/Engine/Source/Runtime/Engine/Private/VorbisAudioInfo.cpp b/Engine/Source/Runtime/Engine/Private/VorbisAudioInfo.cpp
--- a/Engine/Source/Runtime/Engine/Private/VorbisAudioInfo.cpp
+++ b/Engine/Source/Runtime/Engine/Private/VorbisAudioInfo.cpp
@@ -79,9 +79,10 @@ struct FVorbisFileWrapper
 ------------------------------------------------------------------------------------*/
 FVorbisAudioInfo::FVorbisAudioInfo( void ) 
 	: VFWrapper(new FVorbisFileWrapper())
-	,	SrcBufferData(NULL)
-	,	SrcBufferDataSize(0)
-	,	BufferOffset(0)
+	, SrcBufferData(NULL)
+	, SrcBufferDataSize(0)
+	, BufferOffset(0)
+	, bPerformingOperation(false)
 { 
 	// Make sure we have properly allocated a VFWrapper
 	check(VFWrapper != NULL);
@@ -89,6 +90,9 @@ FVorbisAudioInfo::FVorbisAudioInfo( void )
 
 FVorbisAudioInfo::~FVorbisAudioInfo( void ) 
 { 
+	// Make sure we're not deleting ourselves while performing an opreation
+	check(!bPerformingOperation);
+
 	FScopeLock ScopeLock(&VorbisCriticalSection);
 	check(VFWrapper != nullptr);
 	delete VFWrapper;
@@ -172,6 +176,8 @@ static long OggTell( void *datasource )
  */
 bool FVorbisAudioInfo::ReadCompressedInfo( const uint8* InSrcBufferData, uint32 InSrcBufferDataSize, FSoundQualityInfo* QualityInfo )
 {
+	bPerformingOperation = true;
+
 	SCOPE_CYCLE_COUNTER( STAT_VorbisPrepareDecompressionTime );
 
 	FScopeLock ScopeLock(&VorbisCriticalSection);
@@ -220,6 +226,8 @@ bool FVorbisAudioInfo::ReadCompressedInfo( const uint8* InSrcBufferData, uint32
 		}
 	}
 
+	bPerformingOperation = false;
+
 	return( true );
 }
 
@@ -229,6 +237,8 @@ bool FVorbisAudioInfo::ReadCompressedInfo( const uint8* InSrcBufferData, uint32
  */
 void FVorbisAudioInfo::ExpandFile( uint8* DstBuffer, FSoundQualityInfo* QualityInfo )
 {
+	bPerformingOperation = true;
+
 	uint32		TotalBytesRead, BytesToRead;
 
 	check( VFWrapper != NULL );
@@ -256,6 +266,8 @@ void FVorbisAudioInfo::ExpandFile( uint8* DstBuffer, FSoundQualityInfo* QualityI
 		TotalBytesRead += BytesRead;
 		Destination += BytesRead;
 	}
+
+	bPerformingOperation = false;
 }
 
 
@@ -271,6 +283,8 @@ void FVorbisAudioInfo::ExpandFile( uint8* DstBuffer, FSoundQualityInfo* QualityI
  */
 bool FVorbisAudioInfo::ReadCompressedData( uint8* InDestination, bool bLooping, uint32 BufferSize )
 {
+	bPerformingOperation = true;
+
 	bool		bLooped;
 	uint32		TotalBytesRead;
 
@@ -322,22 +336,32 @@ bool FVorbisAudioInfo::ReadCompressedData( uint8* InDestination, bool bLooping,
 		Destination += BytesRead;
 	}
 
+	bPerformingOperation = false;
 	return( bLooped );
 }
 
 void FVorbisAudioInfo::SeekToTime( const float SeekTime )
 {
+	bPerformingOperation = true;
+
 	FScopeLock ScopeLock(&VorbisCriticalSection);
 
 	const float TargetTime = FMath::Min(SeekTime, (float)ov_time_total(&VFWrapper->vf, -1));
 	ov_time_seek( &VFWrapper->vf, TargetTime );
+
+	bPerformingOperation = false;
 }
 
 void FVorbisAudioInfo::EnableHalfRate( bool HalfRate )
 {
+
+	bPerformingOperation = true;
+
 	FScopeLock ScopeLock(&VorbisCriticalSection);
 
 	ov_halfrate(&VFWrapper->vf, int32(HalfRate));
+
+	bPerformingOperation = false;
 }
 
 void LoadVorbisLibraries()
diff --git a/Engine/Source/Runtime/Engine/Private/World.cpp b/Engine/Source/Runtime/Engine/Private/World.cpp
--- a/Engine/Source/Runtime/Engine/Private/World.cpp
+++ b/Engine/Source/Runtime/Engine/Private/World.cpp
@@ -5610,16 +5610,18 @@ ULevel* UWorld::GetCurrentLevel() const
 	return CurrentLevel;
 }
 
-ENetMode UWorld::GetNetMode() const
+ENetMode UWorld::InternalGetNetMode() const
 {
+	const bool bIsClientOnly = IsRunningClientOnly();
+
 	if ( NetDriver != NULL )
 	{
-		return NetDriver->GetNetMode();
+		return bIsClientOnly ? NM_Client : NetDriver->GetNetMode();
 	}
 
 	if ( DemoNetDriver )
 	{
-		return DemoNetDriver->GetNetMode();
+		return bIsClientOnly ? NM_Client : DemoNetDriver->GetNetMode();
 	}
 
 // PIE: NetDriver is not initialized so use PlayInSettings
diff --git a/Engine/Source/Runtime/Engine/Private/WorldComposition.cpp b/Engine/Source/Runtime/Engine/Private/WorldComposition.cpp
--- a/Engine/Source/Runtime/Engine/Private/WorldComposition.cpp
+++ b/Engine/Source/Runtime/Engine/Private/WorldComposition.cpp
@@ -532,7 +532,7 @@ void UWorldComposition::GetDistanceVisibleLevels(
 
 		bool bIsVisible = false;
 
-		if (OwningWorld->GetNetMode() == NM_DedicatedServer) 
+		if (OwningWorld->IsNetMode(NM_DedicatedServer))
 		{
 			// Dedicated server always loads all distance dependent tiles
 			bIsVisible = true;
diff --git a/Engine/Source/Runtime/Engine/Public/AI/NavLinkRenderingProxy.h b/Engine/Source/Runtime/Engine/Public/AI/NavLinkRenderingProxy.h
--- a/Engine/Source/Runtime/Engine/Public/AI/NavLinkRenderingProxy.h
+++ b/Engine/Source/Runtime/Engine/Public/AI/NavLinkRenderingProxy.h
@@ -1,7 +1,7 @@
 // Copyright 1998-2016 Epic Games, Inc. All Rights Reserved.
 
 #pragma once
-
+#include "AI/Navigation/NavLinkDefinition.h"
 #include "PrimitiveSceneProxy.h"
 
 class ENGINE_API FNavLinkRenderingProxy : public FPrimitiveSceneProxy
diff --git a/Engine/Source/Runtime/Engine/Public/Animation/AnimCurveTypes.h b/Engine/Source/Runtime/Engine/Public/Animation/AnimCurveTypes.h
--- a/Engine/Source/Runtime/Engine/Public/Animation/AnimCurveTypes.h
+++ b/Engine/Source/Runtime/Engine/Public/Animation/AnimCurveTypes.h
@@ -24,6 +24,8 @@ enum EAnimCurveFlags
 	ACF_DriveTrack			= 0x00000020,
 	// disabled, right now it's used by track
 	ACF_Disabled			= 0x00000040,
+	// pose curve
+	ACF_DrivesPose			= 0x00000080,
 
 	// default flag when created
 	ACF_DefaultCurve		= ACF_TriggerEvent | ACF_Editable,
@@ -134,10 +136,10 @@ struct ENGINE_API FAnimCurveBase
 	// mean the skeleton doesn't have a mapped name for our UID (such as a user saving the an
 	// animation but not the skeleton).
 	UPROPERTY()
-	FName		LastObservedName;
+	FName		LastObservedName_DEPRECATED;
 
-	/* For smart naming - management purpose - i.e. rename/delete */
-	FSmartNameMapping::UID CurveUid;
+	UPROPERTY()
+	FSmartName	Name;
 
 private:
 	/** Curve Type Flags */
@@ -147,25 +149,14 @@ private:
 public:
 	FAnimCurveBase(){}
 
-	FAnimCurveBase(FName InName, int32 InCurveTypeFlags)
-		: LastObservedName(InName)
+	FAnimCurveBase(FSmartName InName, int32 InCurveTypeFlags)
+		: Name(InName)
 		, CurveTypeFlags(InCurveTypeFlags)
 	{	
 	}
 
-	FAnimCurveBase(USkeleton::AnimCurveUID Uid, int32 InCurveTypeFlags)
-		: CurveUid(Uid)
-		, CurveTypeFlags(InCurveTypeFlags)
-	{}
-
 	// To be able to use typedef'd types we need to serialize manually
-	virtual void Serialize(FArchive& Ar)
-	{
-		if(Ar.UE4Ver() >= VER_UE4_SKELETON_ADD_SMARTNAMES)
-		{
-			Ar << CurveUid;
-		}
-	}
+	void PostSerialize(FArchive& Ar);
 
 	/**
 	 * Set InFlag to bValue
@@ -204,20 +195,11 @@ struct FFloatCurve : public FAnimCurveBase
 
 	FFloatCurve(){}
 
-	FFloatCurve(FName InName, int32 InCurveTypeFlags)
+	FFloatCurve(FSmartName InName, int32 InCurveTypeFlags)
 		: FAnimCurveBase(InName, InCurveTypeFlags)
 	{
 	}
 
-	FFloatCurve(USkeleton::AnimCurveUID Uid, int32 InCurveTypeFlags)
-		: FAnimCurveBase(Uid, InCurveTypeFlags)
-	{}
-
-	virtual void Serialize(FArchive& Ar) override
-	{
-		FAnimCurveBase::Serialize(Ar);
-	}
-
 	// we don't want to have = operator. This only copies curves, but leaving naming and everything else intact. 
 	void CopyCurve(FFloatCurve& SourceCurve);
 	ENGINE_API float Evaluate(float CurrentTime) const;
@@ -244,20 +226,11 @@ struct FVectorCurve : public FAnimCurveBase
 
 	FVectorCurve(){}
 
-	FVectorCurve(FName InName, int32 InCurveTypeFlags)
+	FVectorCurve(FSmartName InName, int32 InCurveTypeFlags)
 		: FAnimCurveBase(InName, InCurveTypeFlags)
 	{
 	}
 
-	FVectorCurve(USkeleton::AnimCurveUID Uid, int32 InCurveTypeFlags)
-		: FAnimCurveBase(Uid, InCurveTypeFlags)
-	{}
-
-	virtual void Serialize(FArchive& Ar) override
-	{
-		FAnimCurveBase::Serialize(Ar);
-	}
-
 	// we don't want to have = operator. This only copies curves, but leaving naming and everything else intact. 
 	void CopyCurve(FVectorCurve& SourceCurve);
 	FVector Evaluate(float CurrentTime, float BlendWeight) const;
@@ -287,20 +260,11 @@ struct FTransformCurve: public FAnimCurveBase
 
 	FTransformCurve(){}
 
-	FTransformCurve(FName InName, int32 InCurveTypeFlags)
+	FTransformCurve(FSmartName InName, int32 InCurveTypeFlags)
 		: FAnimCurveBase(InName, InCurveTypeFlags)
 	{
 	}
 
-	FTransformCurve(USkeleton::AnimCurveUID Uid, int32 InCurveTypeFlags)
-		: FAnimCurveBase(Uid, InCurveTypeFlags)
-	{}
-
-	virtual void Serialize(FArchive& Ar) override
-	{
-		FAnimCurveBase::Serialize(Ar);
-	}
-
 	// we don't want to have = operator. This only copies curves, but leaving naming and everything else intact. 
 	void CopyCurve(FTransformCurve& SourceCurve);
 	FTransform Evaluate(float CurrentTime, float BlendWeight) const;
@@ -682,7 +646,7 @@ struct FRawCurveTracks
 	/**
 	* Add new float curve from the given UID if not existing and add the key with time/value
 	*/
-	ENGINE_API void AddFloatCurveKey(const USkeleton::AnimCurveUID Uid, int32 CurveFlags, float Time, float Value);
+	ENGINE_API void AddFloatCurveKey(const FSmartName& NewCurve, int32 CurveFlags, float Time, float Value);
 	ENGINE_API void RemoveRedundantKeys();
 
 #endif // WITH_EDITOR
@@ -694,12 +658,12 @@ struct FRawCurveTracks
 	 * Add new curve from the provided UID and return true if success
 	 * bVectorInterpCurve == true, then it will create FVectorCuve, otherwise, FFloatCurve
 	 */
-	ENGINE_API bool AddCurveData(USkeleton::AnimCurveUID Uid, int32 CurveFlags = ACF_DefaultCurve, ESupportedCurveType SupportedCurveType = FloatType);
+	ENGINE_API bool AddCurveData(const FSmartName& NewCurve, int32 CurveFlags = ACF_DefaultCurve, ESupportedCurveType SupportedCurveType = FloatType);
 
 	/**
 	 * Delete curve data 
 	 */
-	ENGINE_API bool DeleteCurveData(USkeleton::AnimCurveUID Uid, ESupportedCurveType SupportedCurveType = FloatType);
+	ENGINE_API bool DeleteCurveData(const FSmartName& CurveToDelete, ESupportedCurveType SupportedCurveType = FloatType);
 
 	/**
 	 * Delete all curve data 
@@ -710,17 +674,17 @@ struct FRawCurveTracks
 	 * Duplicate curve data
 	 * 
 	 */
-	ENGINE_API bool DuplicateCurveData(USkeleton::AnimCurveUID ToCopyUid, USkeleton::AnimCurveUID NewUid, ESupportedCurveType SupportedCurveType = FloatType);
+	ENGINE_API bool DuplicateCurveData(const FSmartName& CurveToCopy, const FSmartName& NewCurve, ESupportedCurveType SupportedCurveType = FloatType);
 
 	/**
-	 * Updates the LastObservedName field of the curves from the provided name container
+	 * Updates the DisplayName field of the curves from the provided name container
 	 */
-	ENGINE_API void UpdateLastObservedNames(const FSmartNameMapping* NameMapping, ESupportedCurveType SupportedCurveType = FloatType);
+	ENGINE_API void RefreshName(const FSmartNameMapping* NameMapping, ESupportedCurveType SupportedCurveType = FloatType);
 
 	/** 
 	 * Serialize
 	 */
-	void Serialize(FArchive& Ar);
+	void PostSerialize(FArchive& Ar);
 
 	/*
 	 * resize curve length. If longer, it doesn't do any. If shorter, remove previous keys and add new key to the end of the frame. 
@@ -744,7 +708,7 @@ struct FRawCurveTracks
 		{
 			FORCEINLINE bool operator()(const FFloatCurve& A, const FFloatCurve& B) const
 			{
-				return (A.CurveUid < B.CurveUid);
+				return (A.Name.UID < B.Name.UID);
 			}
 		};
 
@@ -766,21 +730,21 @@ private:
 	 * bVectorInterpCurve == true, then it will create FVectorCuve, otherwise, FFloatCurve
 	 */
 	template <typename DataType>
-	bool AddCurveDataImpl(TArray<DataType>& Curves, USkeleton::AnimCurveUID Uid, int32 CurveFlags);
+	bool AddCurveDataImpl(TArray<DataType>& Curves, const FSmartName& NewCurve, int32 CurveFlags);
 	/**
 	 * Delete curve data 
 	 */
 	template <typename DataType>
-	bool DeleteCurveDataImpl(TArray<DataType>& Curves, USkeleton::AnimCurveUID Uid);
+	bool DeleteCurveDataImpl(TArray<DataType>& Curves, const FSmartName& CurveToDelete);
 	/**
 	 * Duplicate curve data
 	 * 
 	 */
 	template <typename DataType>
-	bool DuplicateCurveDataImpl(TArray<DataType>& Curves, USkeleton::AnimCurveUID ToCopyUid, USkeleton::AnimCurveUID NewUid);
+	bool DuplicateCurveDataImpl(TArray<DataType>& Curves, const FSmartName& CurveToCopy, const FSmartName& NewCurve);
 
 	/**
-	 * Updates the LastObservedName field of the curves from the provided name container
+	 * Updates the DisplayName field of the curves from the provided name container
 	 */
 	template <typename DataType>
 	void UpdateLastObservedNamesImpl(TArray<DataType>& Curves, const FSmartNameMapping* NameMapping);
diff --git a/Engine/Source/Runtime/Engine/Public/Animation/AnimInstanceProxy.h b/Engine/Source/Runtime/Engine/Public/Animation/AnimInstanceProxy.h
--- a/Engine/Source/Runtime/Engine/Public/Animation/AnimInstanceProxy.h
+++ b/Engine/Source/Runtime/Engine/Public/Animation/AnimInstanceProxy.h
@@ -9,6 +9,7 @@
 #include "AnimInstanceProxy.generated.h"
 
 struct FAnimNode_Base;
+struct FAnimNode_SaveCachedPose;
 
 /** Proxy object passed around during animation tree update in lieu of a UAnimInstance */
 USTRUCT(meta = (DisplayName = "Native Variables"))
@@ -174,12 +175,18 @@ public:
 	/** Get the current skeleton we are using. Note that this will return nullptr outside of pre/post update */
 	USkeleton* GetSkeleton() 
 	{ 
+		// Skeleton is only available during update/eval. If you're calling this function outside of it, it will return null. 
+		// adding ensure here so that we can catch them earlier
+		ensureAlways(Skeleton);
 		return Skeleton; 
 	}
 
 	/** Get the current skeletal mesh component we are running on. Note that this will return nullptr outside of pre/post update */
 	USkeletalMeshComponent* GetSkelMeshComponent() 
 	{ 
+		// Skeleton is only available during update/eval. If you're calling this function outside of it, it will return null. 
+		// adding ensure here so that we can catch them earlier
+		ensureAlways(SkeletalMeshComponent);
 		return SkeletalMeshComponent; 
 	}
 
@@ -192,6 +199,8 @@ public:
 	/** Helper function: make a tick record for a blend space */
 	void MakeBlendSpaceTickRecord(FAnimTickRecord& TickRecord, UBlendSpaceBase* BlendSpace, const FVector& BlendInput, TArray<FBlendSampleData>& BlendSampleDataCache, FBlendFilter& BlendFilter, bool bLooping, float PlayRate, float FinalBlendWeight, float& CurrentTime, FMarkerTickRecord& MarkerTickRecord) const;
 
+	/** Helper function: make a tick record for a pose asset*/
+	void MakePoseAssetTickRecord(FAnimTickRecord& TickRecord, class UPoseAsset* PoseAsset, float FinalBlendWeight) const;
 	/**
 	 * Get Slot Node Weight : this returns new Slot Node Weight, Source Weight, Original TotalNodeWeight
 	 *							this 3 values can't be derived from each other
@@ -379,14 +388,6 @@ protected:
 	const TArray<FMontageEvaluationState>& GetMontageEvaluationData() const { return MontageEvaluationData; }
 
 	/** Check whether we have active morph target curves */
-	bool HasMorphTargetCurves() const { return MorphTargetCurves.Num() > 0; }
-
-	/** Access the active morph target curves */
-	TMap<FName, float>& GetMorphTargetCurves() { return MorphTargetCurves; }
-
-	/** Access the active material parameter curves */
-	TMap<FName, float>& GetMaterialParameterCurves() { return MaterialParameterCurves; }
-
 	/** Gets the most relevant asset player in a specified state */
 	FAnimNode_AssetPlayerBase* GetRelevantAssetPlayerFromState(int32 MachineIndex, int32 StateIndex);
 
@@ -518,6 +519,9 @@ private:
 	/** Anim graph */
 	FAnimNode_Base* RootNode;
 
+	/** List of saved pose nodes to process after the graph has been updated */
+	TArray<FAnimNode_SaveCachedPose*> SavedPoseQueue;
+
 	/** The list of animation assets which are going to be evaluated this frame and need to be ticked (ungrouped) */
 	TArray<FAnimTickRecord> UngroupedActivePlayerArrays[2];
 
@@ -568,12 +572,6 @@ private:
 	/** When RequiredBones mapping has changed, AnimNodes need to update their bones caches. */
 	bool bBoneCachesInvalidated;
 
-	/** Morph Target Curves that will be used for SkeletalMeshComponent **/
-	TMap<FName, float> MorphTargetCurves;
-
-	/** Material Curves that will be used for SkeletalMeshComponent **/
-	TMap<FName, float> MaterialParameterCurves;
-
 	/** Copy of UAnimInstance::MontageInstances data used for update & evaluation */
 	TArray<FMontageEvaluationState> MontageEvaluationData;
 
diff --git a/Engine/Source/Runtime/Engine/Public/Animation/AnimSingleNodeInstanceProxy.h b/Engine/Source/Runtime/Engine/Public/Animation/AnimSingleNodeInstanceProxy.h
--- a/Engine/Source/Runtime/Engine/Public/Animation/AnimSingleNodeInstanceProxy.h
+++ b/Engine/Source/Runtime/Engine/Public/Animation/AnimSingleNodeInstanceProxy.h
@@ -41,6 +41,7 @@ public:
 	virtual bool Evaluate(FPoseContext& Output) override;
 	virtual void UpdateAnimationNode(float DeltaSeconds) override;
 	virtual void PostUpdate(UAnimInstance* InAnimInstance) const override;
+	virtual void PreUpdate(UAnimInstance* InAnimInstance, float DeltaSeconds) override;
 	virtual void InitializeObjects(UAnimInstance* InAnimInstance) override;
 	virtual void ClearObjects() override;
 
@@ -113,6 +114,7 @@ public:
 		return bCanProcessAdditiveAnimations;
 	}
 #endif
+	void SetPreviewPoseOverride(const FName& PoseName, float Value);
 
 private:
 	void InternalBlendSpaceEvaluatePose(class UBlendSpaceBase* BlendSpace, TArray<FBlendSampleData>& BlendSampleDataCache, FPoseContext& OutContext);
@@ -124,6 +126,9 @@ protected:
 	bool bCanProcessAdditiveAnimations;
 #endif
 
+	/** Pose Weight value that can override curve data. In the future, we'd like to have UCurveSet that can play by default**/
+	TMap<FName, float> PreviewPoseOverride;
+
 	/** Current Asset being played. Note that this will be nullptr outside of pre/post update **/
 	UAnimationAsset* CurrentAsset;
 
diff --git a/Engine/Source/Runtime/Engine/Public/Animation/AnimTypes.h b/Engine/Source/Runtime/Engine/Public/Animation/AnimTypes.h
--- a/Engine/Source/Runtime/Engine/Public/Animation/AnimTypes.h
+++ b/Engine/Source/Runtime/Engine/Public/Animation/AnimTypes.h
@@ -21,6 +21,33 @@ namespace EAnimEventTriggerOffsets
 
 ENGINE_API float GetTriggerTimeOffsetForType(EAnimEventTriggerOffsets::Type OffsetType);
 
+/** Enum for controlling which reference frame a controller is applied in. */
+UENUM()
+enum EBoneControlSpace
+{
+	/** Set absolute position of bone in world space. */
+	BCS_WorldSpace UMETA(DisplayName = "World Space"),
+	/** Set position of bone in SkeletalMeshComponent's reference frame. */
+	BCS_ComponentSpace UMETA(DisplayName = "Component Space"),
+	/** Set position of bone relative to parent bone. */
+	BCS_ParentBoneSpace UMETA(DisplayName = "Parent Bone Space"),
+	/** Set position of bone in its own reference frame. */
+	BCS_BoneSpace UMETA(DisplayName = "Bone Space"),
+	BCS_MAX,
+};
+
+/** Enum for specifying the source of a bone's rotation. */
+UENUM()
+enum EBoneRotationSource
+{
+	/** Don't change rotation at all. */
+	BRS_KeepComponentSpaceRotation UMETA(DisplayName = "No Change (Preserve Existing Component Space Rotation)"),
+	/** Keep forward direction vector relative to the parent bone. */
+	BRS_KeepLocalSpaceRotation UMETA(DisplayName = "Maintain Local Rotation Relative to Parent"),
+	/** Copy rotation of target to bone. */
+	BRS_CopyFromTarget UMETA(DisplayName = "Copy Target Rotation"),
+};
+
 /** Ticking method for AnimNotifies in AnimMontages. */
 UENUM()
 namespace EMontageNotifyTickType
@@ -48,6 +75,95 @@ namespace ENotifyFilterType
 	};
 }
 
+USTRUCT(BlueprintType)
+struct FPerBoneBlendWeight
+{
+	GENERATED_USTRUCT_BODY()
+
+		/** Source index of the buffer. */
+		UPROPERTY()
+		int32 SourceIndex;
+
+	UPROPERTY()
+		float BlendWeight;
+
+	FPerBoneBlendWeight()
+		: SourceIndex(0)
+		, BlendWeight(0.0f)
+	{
+	}
+};
+
+USTRUCT(BlueprintType)
+struct FPerBoneBlendWeights
+{
+	GENERATED_USTRUCT_BODY()
+
+		UPROPERTY()
+		TArray<FPerBoneBlendWeight> BoneBlendWeights;
+
+
+	FPerBoneBlendWeights() {}
+
+};
+
+struct FGraphTraversalCounter
+{
+private:
+	int16 InternalCounter;
+
+public:
+	FGraphTraversalCounter()
+		: InternalCounter(INDEX_NONE)
+	{}
+
+	int16 Get() const
+	{
+		return InternalCounter;
+	}
+
+	void Increment()
+	{
+		InternalCounter++;
+
+		// Avoid wrapping over back to INDEX_NONE, as this means 'never been traversed'
+		if (InternalCounter == INDEX_NONE)
+		{
+			InternalCounter++;
+		}
+	}
+
+	void Reset()
+	{
+		InternalCounter = INDEX_NONE;
+	}
+
+	void SynchronizeWith(const FGraphTraversalCounter& InMasterCounter)
+	{
+		InternalCounter = InMasterCounter.Get();
+	}
+
+	bool IsSynchronizedWith(const FGraphTraversalCounter& InMasterCounter) const
+	{
+		return ((InternalCounter != INDEX_NONE) && (InternalCounter == InMasterCounter.Get()));
+	}
+
+	bool WasSynchronizedInTheLastFrame(const FGraphTraversalCounter& InMasterCounter) const
+	{
+		// Test if we're currently in sync with our master counter
+		if (IsSynchronizedWith(InMasterCounter))
+		{
+			return true;
+		}
+
+		// If not, test if the Master Counter is a frame ahead of us
+		FGraphTraversalCounter TestCounter(*this);
+		TestCounter.Increment();
+
+		return TestCounter.IsSynchronizedWith(InMasterCounter);
+	}
+};
+
 /**
  * Triggers an animation notify.  Each AnimNotifyEvent contains an AnimNotify object
  * which has its Notify method called and passed to the animation.
diff --git a/Engine/Source/Runtime/Engine/Public/AnimationRuntime.h b/Engine/Source/Runtime/Engine/Public/AnimationRuntime.h
--- a/Engine/Source/Runtime/Engine/Public/AnimationRuntime.h
+++ b/Engine/Source/Runtime/Engine/Public/AnimationRuntime.h
@@ -21,6 +21,52 @@ struct FBlendedCurve;
 class UBlendSpaceBase;
 typedef TArray<FTransform> FTransformArrayA2;
 
+/////////////////////////////////////////////////////////
+// Templated Transform Blend Functionality
+
+namespace ETransformBlendMode
+{
+	enum Type
+	{
+		Overwrite,
+		Accumulate
+	};
+}
+
+template<int32>
+void BlendTransform(const FTransform& Source, FTransform& Dest, const float BlendWeight)
+{
+	check(false); /// should never call this
+}
+
+template<>
+FORCEINLINE void BlendTransform<ETransformBlendMode::Overwrite>(const FTransform& Source, FTransform& Dest, const float BlendWeight)
+{
+	const ScalarRegister VBlendWeight(BlendWeight);
+	Dest = Source * VBlendWeight;
+}
+
+template<>
+FORCEINLINE void BlendTransform<ETransformBlendMode::Accumulate>(const FTransform& Source, FTransform& Dest, const float BlendWeight)
+{
+	const ScalarRegister VBlendWeight(BlendWeight);
+	Dest.AccumulateWithShortestRotation(Source, VBlendWeight);
+}
+
+FORCEINLINE void BlendCurves(const TFixedSizeArrayView<FBlendedCurve>& SourceCurves, const TFixedSizeArrayView<float>& SourceWeights, FBlendedCurve& OutCurve)
+{
+	if (SourceCurves.Num() > 0)
+	{
+		OutCurve.Override(SourceCurves[0], SourceWeights[0]);
+
+		for (int32 CurveIndex = 1; CurveIndex < SourceCurves.Num(); ++CurveIndex)
+		{
+			OutCurve.Accumulate(SourceCurves[CurveIndex], SourceWeights[CurveIndex]);
+		}
+	}
+}
+
+/////////////////////////////////////////////////////////
 /** Interface used to provide interpolation indices for per bone blends
   *
   */
@@ -214,6 +260,9 @@ public:
 	/** Convert TargetPose into an AdditivePose, by doing TargetPose = TargetPose - BasePose */
 	static void ConvertPoseToAdditive(FCompactPose& TargetPose, const FCompactPose& BasePose);
 
+	/** convert transform to additive */
+	static void ConvertTransformToAdditive(FTransform& TargetTrasnform, const FTransform& BaseTransform);
+
 	/** Convert LocalPose into MeshSpaceRotations. Rotations are NOT normalized. */
 	static void ConvertPoseToMeshRotation(FCompactPose& LocalPose);
 
@@ -357,5 +406,13 @@ private:
 			enum ECurveBlendOption::Type CurveBlendOption,
 			/*out*/ FCompactPose& OutPose, 
 			/*out*/ FBlendedCurve& OutCurve);
+public:
+	/** 
+	 * Calculate distance how close two strings are. 
+	 * By close, it calculates how many operations to transform First to Second 
+	 * The return value is [0-MaxLengthString(First, Second)]
+	 * 0 means it's identical, Max means it's completely different
+	 */
+	static int32 GetStringDistance(const FString& First, const FString& Second);
 };
 
diff --git a/Engine/Source/Runtime/Engine/Public/AnimationUtils.h b/Engine/Source/Runtime/Engine/Public/AnimationUtils.h
--- a/Engine/Source/Runtime/Engine/Public/AnimationUtils.h
+++ b/Engine/Source/Runtime/Engine/Public/AnimationUtils.h
@@ -267,6 +267,6 @@ public:
 	/**
 	 * @return		The default animation compression algorithm singleton, instantiating it if necessary.
 	 */
-	static UAnimCompress* GetDefaultAnimationCompressionAlgorithm();
+	ENGINE_API static UAnimCompress* GetDefaultAnimationCompressionAlgorithm();
 
 };
diff --git a/Engine/Source/Runtime/Engine/Public/BoneIndices.h b/Engine/Source/Runtime/Engine/Public/BoneIndices.h
--- a/Engine/Source/Runtime/Engine/Public/BoneIndices.h
+++ b/Engine/Source/Runtime/Engine/Public/BoneIndices.h
@@ -2,6 +2,9 @@
 
 #pragma once
 
+typedef uint16 FBoneIndexType;
+
+
 struct FBoneIndexBase
 {
 	FBoneIndexBase() : BoneIndex(INDEX_NONE) {}
diff --git a/Engine/Source/Runtime/Engine/Public/Engine.h b/Engine/Source/Runtime/Engine/Public/Engine.h
--- a/Engine/Source/Runtime/Engine/Public/Engine.h
+++ b/Engine/Source/Runtime/Engine/Public/Engine.h
@@ -123,7 +123,6 @@
 #include "Components/LineBatchComponent.h"
 #include "Components/MaterialBillboardComponent.h"
 #include "Components/DestructibleComponent.h"
-#include "Components/PoseableMeshComponent.h"
 #include "Lightmass/LightmassPrimitiveSettingsObject.h"
 #include "Components/StaticMeshComponent.h"
 #include "Components/InstancedStaticMeshComponent.h"
@@ -142,11 +141,6 @@
 #include "Components/SceneCaptureComponent2D.h"
 #include "Components/SceneCaptureComponentCube.h"
 #include "Components/TimelineComponent.h"
-#include "Vehicles/VehicleAnimInstance.h"
-#include "Animation/AnimNotifies/AnimNotify.h"
-#include "Animation/AnimNotifies/AnimNotifyState.h"
-#include "Animation/AnimNotifies/AnimNotifyState_TimedParticleEffect.h"
-#include "Animation/AnimNotifies/AnimNotifyState_Trail.h"
 #include "EditorFramework/AssetImportData.h"
 #include "Engine/AssetUserData.h"
 #include "Engine/BlendableInterface.h"
diff --git a/Engine/Source/Runtime/Engine/Public/GPUSkinCache.h b/Engine/Source/Runtime/Engine/Public/GPUSkinCache.h
--- a/Engine/Source/Runtime/Engine/Public/GPUSkinCache.h
+++ b/Engine/Source/Runtime/Engine/Public/GPUSkinCache.h
@@ -7,8 +7,8 @@
 
 // Requirements
 // * Compute shader support (with Atomics)
-// * Project settings needs to be enabled (r.SkinCacheShaders)
-// * feature need to be enabled (r.SkinCaching)
+// * Project settings needs to be enabled (r.SkinCache.CompileShaders)
+// * feature need to be enabled (r.SkinCache.Mode)
 
 // Features
 // * Skeletal mesh, 4 / 8 weights per vertex, 16/32 index buffer
@@ -18,12 +18,12 @@
 // * RecomputeTangents results in improved tangent space for WorldPosOffset animation and morph target animation
 // * fixed amount of memory (r.MaxGPUSkinCacheElementsPerFrame, r.SkinCache.BufferSize)
 // * Velocity Rendering for MotionBlur and TemporalAA (test Velocity in BasePass)
-// * r.SkinCaching and r.SkinCaching.RecomputeTangents can be toggled at runtime
+// * r.SkinCache.Mode and r.SkinCache.RecomputeTangents can be toggled at runtime
 
 // TODO:
 // * Test: Tessellation
 // * Quality/Optimization: increase TANGENT_RANGE for better quality or accumulate two components in one 32bit value
-// * Feature: RecomputeTangents should be a per section setting, the cvar (r.SkinCaching.RecomputeTangents) can stay but should be on by default
+// * Feature: RecomputeTangents should be a per section setting, the cvar (r.SkinCache.RecomputeTangents) can stay but should be on by default
 // * Bug: We iterate through all chunks and try to find the Section. This is costly and might fail (Seems the section has a StartChunkIndex but it's called ChunkIndex)
 // * Bug: UpdateMorphVertexBuffer needs to handle SkinCacheObjects that have been rejected by the SkinCache (e.g. because it was running out of memory)
 // * Refactor: Unify the 3 compute shaders to use the same C++ setup code for the variables
@@ -102,7 +102,7 @@ public:
 	// @param MorphVertexBuffer if no morph targets are used
 	// @returns -1 if failed, otherwise index into CachedElements[]
 	int32 StartCacheMesh(FRHICommandListImmediate& RHICmdList, uint32 FrameNumber, int32 Key, class FGPUBaseSkinVertexFactory* VertexFactory,
-		class FGPUSkinPassthroughVertexFactory* TargetVertexFactory, const struct FSkelMeshChunk& BatchElement, FSkeletalMeshObjectGPUSkin* Skin,
+		class FGPUSkinPassthroughVertexFactory* TargetVertexFactory, const struct FSkelMeshSection& BatchElement, FSkeletalMeshObjectGPUSkin* Skin,
 		const class FMorphVertexBuffer* MorphVertexBuffer);
 
 	// @param FrameNumber from GFrameNumber or better ViewFamily.FrameNumber
@@ -124,7 +124,7 @@ public:
 		// key section -------------
 
 		// TODO: using a pointer as key is dangerous
-		const FSkelMeshChunk* BatchElement;
+		const FSkelMeshSection* BatchElement;
 		// TODO: using a pointer as key is dangerous
 		const FSkeletalMeshObjectGPUSkin* Skin;
 
@@ -161,10 +161,10 @@ public:
 		struct FSearchInfo
 		{
 			const FSkeletalMeshObjectGPUSkin* Skin;
-			const FSkelMeshChunk* BatchElement;
+			const FSkelMeshSection* BatchElement;
 
 			FSearchInfo()  { BatchElement = nullptr; Skin = nullptr; }
-			FSearchInfo(const FSkeletalMeshObjectGPUSkin* InSkin, const FSkelMeshChunk* InBatchElement) { BatchElement = InBatchElement; Skin = InSkin; }
+			FSearchInfo(const FSkeletalMeshObjectGPUSkin* InSkin, const FSkelMeshSection* InBatchElement) { BatchElement = InBatchElement; Skin = InSkin; }
 		};
 
 		// to support FindByKey()
@@ -174,7 +174,7 @@ public:
 		}
 
 		// to support == (confusing usage, better: IsSameBatchElement)
-		bool operator == (const FSkelMeshChunk& OtherBatchElement) const 
+		bool operator == (const FSkelMeshSection& OtherBatchElement) const
 		{
 			return (BatchElement == &OtherBatchElement);
 		}
@@ -184,9 +184,9 @@ public:
 	struct FDispatchData 
 	{
 		FRHICommandListImmediate& RHICmdList;
-		const FSkelMeshChunk& Chunk;
+		const FSkelMeshSection& Section;
 		// for debugging / draw events, -1 if not set
-		uint32 ChunkIdx;
+		uint32 SectionIdx;
 		ERHIFeatureLevel::Type FeatureLevel;
 		
 		// must not be 0
@@ -219,10 +219,10 @@ public:
 		uint32 IndexBufferOffsetValue;
 		uint32 NumTriangles;
 		
-		FDispatchData(FRHICommandListImmediate& InRHICmdList, const FSkelMeshChunk& InChunk, ERHIFeatureLevel::Type InFeatureLevel, FSkeletalMeshObjectGPUSkin* InGPUSkin)
+		FDispatchData(FRHICommandListImmediate& InRHICmdList, const FSkelMeshSection& InSection, ERHIFeatureLevel::Type InFeatureLevel, FSkeletalMeshObjectGPUSkin* InGPUSkin)
 			: RHICmdList(InRHICmdList)
-			, Chunk(InChunk)
-			, ChunkIdx(-1)
+			, Section(InSection)
+			, SectionIdx(-1)
 			, FeatureLevel(InFeatureLevel)
 			, GPUSkin(InGPUSkin)
 			, SkinType(0)
diff --git a/Engine/Source/Runtime/Engine/Public/ReferenceSkeleton.h b/Engine/Source/Runtime/Engine/Public/ReferenceSkeleton.h
--- a/Engine/Source/Runtime/Engine/Public/ReferenceSkeleton.h
+++ b/Engine/Source/Runtime/Engine/Public/ReferenceSkeleton.h
@@ -2,6 +2,7 @@
 
 #pragma once
 
+#include "BoneIndices.h"
 
 // This contains Reference-skeleton related info
 // Bone transform is saved as FTransform array
diff --git a/Engine/Source/Runtime/Engine/Public/SkeletalMeshMerge.h b/Engine/Source/Runtime/Engine/Public/SkeletalMeshMerge.h
--- a/Engine/Source/Runtime/Engine/Public/SkeletalMeshMerge.h
+++ b/Engine/Source/Runtime/Engine/Public/SkeletalMeshMerge.h
@@ -170,15 +170,12 @@ private:
 		const USkeletalMesh* SkelMesh;
 		/** ptr to source section for merging */
 		const FSkelMeshSection* Section;
-		/** ptr to source chunk for this section */
-		const FSkelMeshChunk* Chunk;
 		/** mapping from the original BoneMap for this sections chunk to the new MergedBoneMap */
 		TArray<FBoneIndexType> BoneMapToMergedBoneMap;
 
-		FMergeSectionInfo( const USkeletalMesh* InSkelMesh,const FSkelMeshSection* InSection, const FSkelMeshChunk* InChunk )
+		FMergeSectionInfo( const USkeletalMesh* InSkelMesh,const FSkelMeshSection* InSection )
 			:	SkelMesh(InSkelMesh)
 			,	Section(InSection)
-			,	Chunk(InChunk)
 		{}
 	};
 
diff --git a/Engine/Source/Runtime/Engine/Public/SkeletalMeshTypes.h b/Engine/Source/Runtime/Engine/Public/SkeletalMeshTypes.h
--- a/Engine/Source/Runtime/Engine/Public/SkeletalMeshTypes.h
+++ b/Engine/Source/Runtime/Engine/Public/SkeletalMeshTypes.h
@@ -14,9 +14,6 @@
 #include "GPUSkinPublicDefs.h"
 #include "PrimitiveSceneProxy.h"
 #include "Materials/MaterialInterface.h"
-
-typedef uint16 FBoneIndexType;
-
 #include "ReferenceSkeleton.h"
 
 class FRawStaticIndexBuffer16or32Interface;
@@ -57,6 +54,29 @@ enum ECustomSortAlternateIndexMode
 };
 
 
+// Custom serialization version for RecomputeTangent
+struct FSkeletalMeshCustomVersion
+{
+	enum Type
+	{
+		// Before any version changes were made
+		BeforeCustomVersionWasAdded = 0,
+		// Remove Chunks array in FStaticLODModel and combine with Sections array
+		CombineSectionWithChunk = 1,
+		// Remove FRigidSkinVertex and combine with FSoftSkinVertex array
+		CombineSoftAndRigidVerts = 2,
+		// -----<new versions can be added above this line>-------------------------------------------------
+		VersionPlusOne,
+		LatestVersion = VersionPlusOne - 1
+	};
+
+	// The GUID for this custom version number
+	const static FGuid GUID;
+
+private:
+	FSkeletalMeshCustomVersion() {}
+};
+
 class USkeletalMesh;
 
 
@@ -191,6 +211,9 @@ struct FSoftSkinVertex
 	uint8			InfluenceBones[MAX_TOTAL_INFLUENCES];
 	uint8			InfluenceWeights[MAX_TOTAL_INFLUENCES];
 
+	/** If this vert is rigidly weighted to a bone, return true and the bone index. Otherwise return false. */
+	ENGINE_API bool GetRigidWeightBone(uint8& OutBoneIndex) const;
+
 	/**
 	* Serializer
 	*
@@ -201,29 +224,7 @@ struct FSoftSkinVertex
 	friend FArchive& operator<<(FArchive& Ar,FSoftSkinVertex& V);
 };
 
-//
-//	FRigidSkinVertex
-//
-
-struct FRigidSkinVertex
-{
-	FVector			Position;
-	FPackedNormal	TangentX,	// Tangent, U-direction
-					TangentY,	// Binormal, V-direction
-					TangentZ;	// Normal
-	FVector2D		UVs[MAX_TEXCOORDS]; // UVs
-	FColor			Color;		// Vertex color.
-	uint8			Bone;
 
-	/**
-	* Serializer
-	*
-	* @param Ar - archive to serialize with
-	* @param V - vertex to serialize
-	* @return archive that was used
-	*/
-	friend FArchive& operator<<(FArchive& Ar,FRigidSkinVertex& V);
-};
 
 /**
  * A structure for holding the APEX cloth physical-to-render mapping data
@@ -315,145 +316,6 @@ struct FApexClothBoneSphereData
 	FVector LocalPos;
 };
 
-/**
- * A set of the skeletal mesh vertices which use the same set of <MAX_GPUSKIN_BONES bones.
- * In practice, there is a 1:1 mapping between chunks and sections, but for meshes which
- * were imported before chunking was implemented, there will be a single chunk for all
- * sections.
- */
-struct FSkelMeshChunk
-{
-	/** The offset into the LOD's vertex buffer of this chunk's vertices. */
-	uint32 BaseVertexIndex;
-
-	/** The rigid vertices of this chunk. */
-	TArray<FRigidSkinVertex> RigidVertices;
-
-	/** The soft vertices of this chunk. */
-	TArray<FSoftSkinVertex> SoftVertices;
-
-	/** The extra vertex data for mapping to an APEX clothing simulation mesh. */
-	TArray<FApexClothPhysToRenderVertData> ApexClothMappingData;
-
-	/** The physical mesh vertices imported from the APEX file. */
-	TArray<FVector> PhysicalMeshVertices;
-
-	/** The physical mesh normals imported from the APEX file. */
-	TArray<FVector> PhysicalMeshNormals;
-
-	/** The bones which are used by the vertices of this chunk. Indices of bones in the USkeletalMesh::RefSkeleton array */
-	TArray<FBoneIndexType> BoneMap;
-
-	/** The number of rigid vertices in this chunk */
-	int32 NumRigidVertices;
-	/** The number of soft vertices in this chunk */
-	int32 NumSoftVertices;
-
-	/** max # of bones used to skin the vertices in this chunk */
-	int32 MaxBoneInfluences;
-
-	// INDEX_NONE if not set
-	int16 CorrespondClothAssetIndex;
-	// INDEX_NONE if not set
-	int16 ClothAssetSubmeshIndex;
-
-	FSkelMeshChunk()
-		: BaseVertexIndex(0)
-		, NumRigidVertices(0)
-		, NumSoftVertices(0)
-		, MaxBoneInfluences(4)
-		, CorrespondClothAssetIndex(INDEX_NONE)
-		, ClothAssetSubmeshIndex(INDEX_NONE)
-	{}
-
-	FSkelMeshChunk(const FSkelMeshChunk& Other)
-	{
-		BaseVertexIndex = Other.BaseVertexIndex;
-		RigidVertices = Other.RigidVertices;
-		SoftVertices = Other.SoftVertices;
-		ApexClothMappingData = Other.ApexClothMappingData;
-		PhysicalMeshVertices = Other.PhysicalMeshVertices;
-		PhysicalMeshNormals = Other.PhysicalMeshNormals;
-		BoneMap = Other.BoneMap;
-		NumRigidVertices = Other.NumRigidVertices;
-		NumSoftVertices = Other.NumSoftVertices;
-		MaxBoneInfluences = Other.MaxBoneInfluences;
-		CorrespondClothAssetIndex = Other.CorrespondClothAssetIndex;
-		ClothAssetSubmeshIndex = Other.ClothAssetSubmeshIndex;
-	}
-
-	/**
-	* @return total num rigid verts for this chunk
-	*/
-	FORCEINLINE int32 GetNumRigidVertices() const
-	{
-		return NumRigidVertices;
-	}
-
-	/**
-	* @return total num soft verts for this chunk
-	*/
-	FORCEINLINE int32 GetNumSoftVertices() const
-	{
-		return NumSoftVertices;
-	}
-
-	/**
-	* @return total number of soft and rigid verts for this chunk
-	*/
-	FORCEINLINE int32 GetNumVertices() const
-	{
-		return GetNumRigidVertices() + GetNumSoftVertices();
-	}
-
-	/**
-	* @return starting index for rigid verts for this chunk in the LOD vertex buffer
-	*/
-	FORCEINLINE int32 GetRigidVertexBufferIndex() const
-	{
-		return BaseVertexIndex;
-	}
-
-	/**
-	* @return starting index for soft verts for this chunk in the LOD vertex buffer
-	*/
-	FORCEINLINE int32 GetSoftVertexBufferIndex() const
-	{
-		return BaseVertexIndex + NumRigidVertices;
-	}
-
-	/**
-	* @return TRUE if we have cloth data for this chunk
-	*/
-	FORCEINLINE bool HasApexClothData() const
-	{
-		return (ApexClothMappingData.Num() > 0);
-	}
-
-	FORCEINLINE void SetClothSubmeshIndex(int16 AssetIndex, int16 AssetSubmeshIndex)
-	{
-		CorrespondClothAssetIndex = AssetIndex;
-		ClothAssetSubmeshIndex = AssetSubmeshIndex;
-	}
-
-	/**
-	* Calculate max # of bone influences used by this skel mesh chunk
-	*/
-	ENGINE_API void CalcMaxBoneInfluences();
-
-	FORCEINLINE bool HasExtraBoneInfluences() const
-	{
-		return MaxBoneInfluences > MAX_INFLUENCES_PER_STREAM;
-	}
-
-	/**
-	* Serialize this class
-	* @param Ar - archive to serialize to
-	* @param C - skel mesh chunk to serialize
-	*/
-	friend FArchive& operator<<(FArchive& Ar,FSkelMeshChunk& C);
-};
-
 
 /** Helper to convert the above enum to string */
 static const TCHAR* TriangleSortOptionToString(ETriangleSortOption Option)
@@ -476,16 +338,13 @@ static const TCHAR* TriangleSortOptionToString(ETriangleSortOption Option)
 }
 
 /**
- * A set of skeletal mesh triangles which use the same material and chunk.
+ * A set of skeletal mesh triangles which use the same material
  */
 struct FSkelMeshSection
 {
 	/** Material (texture) used for this section. */
 	uint16 MaterialIndex;
 
-	/** The chunk that vertices for this section are from. */
-	uint16 ChunkIndex;
-
 	/** The offset of this section's indices in the LOD's index buffer. */
 	uint32 BaseIndex;
 
@@ -514,9 +373,35 @@ struct FSkelMeshSection
 	/** no need anymore because each clothing LOD will be assigned to each mesh LOD  */
 	uint8 bEnableClothLOD_DEPRECATED;
 
+	/** The offset into the LOD's vertex buffer of this section's vertices. */
+	uint32 BaseVertexIndex;
+
+	/** The soft vertices of this section. */
+	TArray<FSoftSkinVertex> SoftVertices;
+
+	/** The extra vertex data for mapping to an APEX clothing simulation mesh. */
+	TArray<FApexClothPhysToRenderVertData> ApexClothMappingData;
+
+	/** The physical mesh vertices imported from the APEX file. */
+	TArray<FVector> PhysicalMeshVertices;
+
+	/** The physical mesh normals imported from the APEX file. */
+	TArray<FVector> PhysicalMeshNormals;
+
+	/** The bones which are used by the vertices of this section. Indices of bones in the USkeletalMesh::RefSkeleton array */
+	TArray<FBoneIndexType> BoneMap;
+
+	/** max # of bones used to skin the vertices in this section */
+	int32 MaxBoneInfluences;
+
+	// INDEX_NONE if not set
+	int16 CorrespondClothAssetIndex;
+	// INDEX_NONE if not set
+	int16 ClothAssetSubmeshIndex;
+
+
 	FSkelMeshSection()
 		: MaterialIndex(0)
-		, ChunkIndex(0)
 		, BaseIndex(0)
 		, NumTriangles(0)
 		, TriangleSorting(0)
@@ -524,12 +409,56 @@ struct FSkelMeshSection
 		, bRecomputeTangent(false)
 		, bDisabled(false)
 		, CorrespondClothSectionIndex(-1)
+		, BaseVertexIndex(0)
+		, MaxBoneInfluences(4)
+		, CorrespondClothAssetIndex(INDEX_NONE)
+		, ClothAssetSubmeshIndex(INDEX_NONE)
 	{}
 
+
+	/**
+	* @return total num rigid verts for this section
+	*/
+	FORCEINLINE int32 GetNumVertices() const
+	{
+		return SoftVertices.Num();
+	}
+
+	/**
+	* @return starting index for rigid verts for this section in the LOD vertex buffer
+	*/
+	FORCEINLINE int32 GetVertexBufferIndex() const
+	{
+		return BaseVertexIndex;
+	}
+
+	/**
+	* @return TRUE if we have cloth data for this section
+	*/
+	FORCEINLINE bool HasApexClothData() const
+	{
+		return (ApexClothMappingData.Num() > 0);
+	}
+
+	FORCEINLINE void SetClothSubmeshIndex(int16 AssetIndex, int16 AssetSubmeshIndex)
+	{
+		CorrespondClothAssetIndex = AssetIndex;
+		ClothAssetSubmeshIndex = AssetSubmeshIndex;
+	}
+
+	/**
+	* Calculate max # of bone influences used by this skel mesh section
+	*/
+	ENGINE_API void CalcMaxBoneInfluences();
+
+	FORCEINLINE bool HasExtraBoneInfluences() const
+	{
+		return MaxBoneInfluences > MAX_INFLUENCES_PER_STREAM;
+	}
+
 	// Serialization.
 	friend FArchive& operator<<(FArchive& Ar, FSkelMeshSection& S);
 };
-template <> struct TIsPODType<FSkelMeshSection> { enum { Value = true }; };
 
 /**
 * Base vertex data for GPU skinned skeletal meshes
@@ -1535,9 +1464,6 @@ public:
 	/** Sections. */
 	TArray<FSkelMeshSection> Sections;
 
-	/** The vertex chunks which make up this LOD. */
-	TArray<FSkelMeshChunk> Chunks;
-
 	/** 
 	* Bone hierarchy subset active for this chunk.
 	* This is a map between the bones index of this LOD (as used by the vertex structs) and the bone index in the reference skeleton of this SkeletalMesh.
@@ -1655,8 +1581,15 @@ public:
 	/** Utility function for returning total number of faces in this LOD. */
 	ENGINE_API int32 GetTotalFaces() const;
 
-	/** Utility for finding the chunk that a particular vertex is in. */
-	ENGINE_API void GetChunkAndSkinType(int32 InVertIndex, int32& OutChunkIndex, int32& OutVertIndex, bool& bOutSoftVert, bool& bOutHasExtraBoneInfluences) const;
+	DEPRECATED(4.13, "Please use GetSectionFromVertIndex.")
+	ENGINE_API void GetChunkAndSkinType(int32 InVertIndex, int32& OutChunkIndex, int32& OutVertIndex, bool& bOutSoftVert, bool& bOutHasExtraBoneInfluences) const
+	{
+		GetSectionFromVertexIndex(InVertIndex, OutChunkIndex, OutVertIndex, bOutHasExtraBoneInfluences);
+		bOutSoftVert = true;
+	}
+
+	/** Utility for finding the section that a particular vertex is in. */
+	ENGINE_API void GetSectionFromVertexIndex(int32 InVertIndex, int32& OutSectionIndex, int32& OutVertIndex, bool& bOutHasExtraBoneInfluences) const;
 
 	/**
 	 * Sort the triangles with the specified sorting method
@@ -1669,9 +1602,9 @@ public:
 	*/
 	bool HasApexClothData() const
 	{
-		for( int32 ChunkIdx=0;ChunkIdx<Chunks.Num();ChunkIdx++ )
+		for( int32 SectionIdx=0; SectionIdx<Sections.Num(); SectionIdx++ )
 		{
-			if( Chunks[ChunkIdx].HasApexClothData() )
+			if(Sections[SectionIdx].HasApexClothData() )
 			{
 				return true;
 			}
@@ -1679,16 +1612,16 @@ public:
 		return false;
 	}
 
-	int32 GetApexClothCunkIndex(TArray<int32>& ChunkIndices) const
+	int32 GetApexClothSectionIndex(TArray<int32>& SectionIndices) const
 	{
-		ChunkIndices.Empty();
+		SectionIndices.Empty();
 
 		uint32 Count = 0;
-		for( int32 ChunkIdx=0;ChunkIdx<Chunks.Num();ChunkIdx++ )
+		for( int32 SectionIdx =0; SectionIdx<Sections.Num(); SectionIdx++ )
 		{
-			if( Chunks[ChunkIdx].HasApexClothData() )
+			if(Sections[SectionIdx].HasApexClothData() )
 			{
-				ChunkIndices.Add(ChunkIdx);
+				SectionIndices.Add(SectionIdx);
 				Count++;
 			}
 		}
@@ -1697,7 +1630,7 @@ public:
 
 	bool HasApexClothData(int32 SectionIndex) const
 	{
-		return Chunks[Sections[SectionIndex].ChunkIndex].HasApexClothData();
+		return Sections[SectionIndex].HasApexClothData();
 	}
 
 	int32 NumNonClothingSections() const
@@ -1722,11 +1655,11 @@ public:
 		return VertexBufferGPUSkin.HasExtraBoneInfluences();
 	}
 
-	bool DoChunksNeedExtraBoneInfluences() const
+	bool DoSectionsNeedExtraBoneInfluences() const
 	{
-		for (int32 ChunkIdx = 0; ChunkIdx < Chunks.Num(); ++ChunkIdx)
+		for (int32 SectionIdx = 0; SectionIdx < Sections.Num(); ++SectionIdx)
 		{
-			if (Chunks[ChunkIdx].HasExtraBoneInfluences())
+			if (Sections[SectionIdx].HasExtraBoneInfluences())
 			{
 				return true;
 			}
@@ -1737,41 +1670,25 @@ public:
 
 	// O(1)
 	// @return -1 if not found
-	uint32 FindChunkIndex(const FSkelMeshChunk& Chunk) const
+	uint32 FindSectionIndex(const FSkelMeshSection& Section) const
 	{
-		const FSkelMeshChunk* Start = Chunks.GetData();
+		const FSkelMeshSection* Start = Sections.GetData();
 
 		if(Start == nullptr)
 		{
 			return -1;
 		}
 
-		uint32 Ret = &Chunk - Start;
+		uint32 Ret = &Section - Start;
 
-		if(Ret >= (uint32)Chunks.Num())
-			{
+		if(Ret >= (uint32)Sections.Num())
+		{
 			Ret = -1;
 		}
 
 		return Ret;
 	}
 
-	// @return 0 if not found
-	const FSkelMeshSection* FindSectionForChunk(int16 ChunkIndex) const
-	{
-		for(auto& ref : Sections)
-		{
-			// should we check ranges instead?
-			if(ref.ChunkIndex == ChunkIndex)
-			{
-				return &ref;
-			}
-		}
-		// should never happen
-		ensure(0);
-		return 0;
-	}
-
 	/**
 	 * Get Resource Size
 	 */
@@ -1806,9 +1723,9 @@ public:
 	void Serialize(FArchive& Ar, USkeletalMesh* Owner);
 
 	/**
-	 * Computes the maximum number of bones per chunk used to render this mesh.
+	 * Computes the maximum number of bones per section used to render this mesh.
 	 */
-	int32 GetMaxBonesPerChunk() const;
+	int32 GetMaxBonesPerSection() const;
 
 	/** Returns true if this resource must be skinned on the CPU for the given feature level. */
 	bool RequiresCPUSkinning(ERHIFeatureLevel::Type FeatureLevel) const;
@@ -1982,7 +1899,7 @@ protected:
 #endif
 
 	void GetDynamicElementsSection(const TArray<const FSceneView*>& Views, const FSceneViewFamily& ViewFamily, uint32 VisibilityMap, 
-		const FStaticLODModel& LODModel, const int32 LODIndex, const FSkelMeshSection& Section, const FSkelMeshChunk& Chunk, 
+		const FStaticLODModel& LODModel, const int32 LODIndex, const int32 SectionIndex, 
 		const FSectionElementInfo& SectionElementInfo, const FTwoVectors& CustomLeftRightVectors, bool bInSelectable, FMeshElementCollector& Collector ) const;
 
 	void GetMeshElementsConditionallySelectable(const TArray<const FSceneView*>& Views, const FSceneViewFamily& ViewFamily, bool bInSelectable, uint32 VisibilityMap, FMeshElementCollector& Collector) const;
diff --git a/Engine/Source/Runtime/Engine/Public/SkeletalRenderPublic.h b/Engine/Source/Runtime/Engine/Public/SkeletalRenderPublic.h
--- a/Engine/Source/Runtime/Engine/Public/SkeletalRenderPublic.h
+++ b/Engine/Source/Runtime/Engine/Public/SkeletalRenderPublic.h
@@ -146,10 +146,10 @@ public:
 	virtual SIZE_T GetResourceSize() = 0;
 
 	/**
-	 * List of chunks to be rendered based on instance weight usage. Full swap of weights will render with its own chunks.
-	 * @return Chunks to iterate over for rendering
+	 * List of sections to be rendered based on instance weight usage. Full swap of weights will render with its own sections.
+	 * @return Sections to iterate over for rendering
 	 */
-	const TArray<FSkelMeshChunk>& GetRenderChunks(int32 InLODIndex) const;
+	const TArray<FSkelMeshSection>& GetRenderSections(int32 InLODIndex) const;
 
 	/**
 	 * Update the hidden material section flags for an LOD entry
@@ -213,9 +213,6 @@ public:
 	bool bHasBeenUpdatedAtLeastOnce;
 
 #if WITH_EDITORONLY_DATA
-	/** Index of the chunk to preview... If set to -1, all chunks will be rendered */
-	int32 ChunkIndexPreview;
-	
 	/** Index of the section to preview... If set to -1, all section will be rendered */
 	int32 SectionIndexPreview;
 #endif
diff --git a/Engine/Source/Runtime/Engine/Public/VorbisAudioInfo.h b/Engine/Source/Runtime/Engine/Public/VorbisAudioInfo.h
--- a/Engine/Source/Runtime/Engine/Public/VorbisAudioInfo.h
+++ b/Engine/Source/Runtime/Engine/Public/VorbisAudioInfo.h
@@ -87,6 +87,8 @@ public:
 	uint32			SrcBufferDataSize;
 	uint32			BufferOffset;
 
+	FThreadSafeBool bPerformingOperation;
+
 	/** Critical section used to prevent multiple threads accessing same ogg-vorbis file handles at the same time */
 	FCriticalSection VorbisCriticalSection;
 };
diff --git a/Engine/Source/Runtime/MovieSceneTracks/Public/Sections/MovieSceneSkeletalAnimationSection.h b/Engine/Source/Runtime/MovieSceneTracks/Public/Sections/MovieSceneSkeletalAnimationSection.h
--- a/Engine/Source/Runtime/MovieSceneTracks/Public/Sections/MovieSceneSkeletalAnimationSection.h
+++ b/Engine/Source/Runtime/MovieSceneTracks/Public/Sections/MovieSceneSkeletalAnimationSection.h
@@ -3,10 +3,11 @@
 #pragma once
 
 #include "MovieSceneSection.h"
+#include "Animation/AnimSequence.h"
 #include "MovieSceneSkeletalAnimationSection.generated.h"
 
 /**
- * Audio section, for use in the master audio, or by attached audio objects
+ * Movie scene section that control skeletal animation
  */
 UCLASS( MinimalAPI )
 class UMovieSceneSkeletalAnimationSection : public UMovieSceneSection
diff --git a/Engine/Source/Runtime/Serialization/Private/StructDeserializer.cpp b/Engine/Source/Runtime/Serialization/Private/StructDeserializer.cpp
--- a/Engine/Source/Runtime/Serialization/Private/StructDeserializer.cpp
+++ b/Engine/Source/Runtime/Serialization/Private/StructDeserializer.cpp
@@ -210,6 +210,21 @@ bool FStructDeserializer::Deserialize( void* OutStruct, UStruct& TypeInfo, IStru
 						UE_LOG(LogSerialization, Verbose, TEXT("An item in map '%s' could not be read (%s)"), *PropertyName, *Backend.GetDebugString());
 					}
 				}
+				else if ((CurrentState.Property != nullptr) && (CurrentState.Property->GetClass() == USetProperty::StaticClass()))
+				{
+					// handle set element
+					USetProperty* SetProperty = Cast<USetProperty>(CurrentState.Property);
+					FScriptSetHelper SetHelper(SetProperty, SetProperty->ContainerPtrToValuePtr<void>(CurrentState.Data));
+					UProperty* Property = SetProperty->ElementProp;
+
+					const int32 ElementIndex = SetHelper.AddDefaultValue_Invalid_NeedsRehash();
+					uint8* ElementPtr = SetHelper.GetElementPtr(ElementIndex);
+
+					if (!Backend.ReadProperty(Property, CurrentState.Property, ElementPtr, CurrentState.ArrayIndex))
+					{
+						UE_LOG(LogSerialization, Verbose, TEXT("An item in Set '%s' could not be read (%s)"), *PropertyName, *Backend.GetDebugString());
+					}
+				}
 				else
 				{
 					// handle scalar property
@@ -248,13 +263,20 @@ bool FStructDeserializer::Deserialize( void* OutStruct, UStruct& TypeInfo, IStru
 			{
 				// rehash if value was a map
 				UMapProperty* MapProperty = Cast<UMapProperty>(CurrentState.Property);
-
 				if (MapProperty != nullptr)
 				{			
 					FScriptMapHelper MapHelper(MapProperty, CurrentState.Data);
 					MapHelper.Rehash();
 				}
 
+				// rehash if value was a set
+				USetProperty* SetProperty = Cast<USetProperty>(CurrentState.Property);
+				if (SetProperty != nullptr)
+				{			
+					FScriptSetHelper SetHelper(SetProperty, CurrentState.Data);
+					SetHelper.Rehash();
+				}
+
 				if (StateStack.Num() == 0)
 				{
 					return true;
@@ -305,16 +327,27 @@ bool FStructDeserializer::Deserialize( void* OutStruct, UStruct& TypeInfo, IStru
 
 					MapProperty->KeyProp->ImportText(*PropertyName, PairPtr + MapProperty->MapLayout.KeyOffset, PPF_None, nullptr);
 				}
+				else if ((CurrentState.Property != nullptr) && (CurrentState.Property->GetClass() == USetProperty::StaticClass()))
+				{
+					// handle map or struct element inside map
+					USetProperty* SetProperty = Cast<USetProperty>(CurrentState.Property);
+					FScriptSetHelper SetHelper(SetProperty, CurrentState.Data);
+					const int32 ElementIndex = SetHelper.AddDefaultValue_Invalid_NeedsRehash();
+					uint8* ElementPtr = SetHelper.GetElementPtr(ElementIndex);
+
+					NewState.Data = ElementPtr + SetHelper.SetLayout.ElementOffset;
+					NewState.Property = SetProperty->ElementProp;
+				}
 				else
 				{
 					NewState.Property = FindField<UProperty>(CurrentState.TypeInfo, *PropertyName);
 
 					if (NewState.Property == nullptr)
 					{
-						// error: map or struct property not found
+						// error: map, set, or struct property not found
 						if (Policies.MissingFields != EStructDeserializerErrorPolicies::Ignore)
 						{
-							UE_LOG(LogSerialization, Verbose, TEXT("Map or struct property '%s' not found"), *PropertyName);
+							UE_LOG(LogSerialization, Verbose, TEXT("Map, Set, or struct property '%s' not found"), *PropertyName);
 						}
 
 						if (Policies.MissingFields == EStructDeserializerErrorPolicies::Error)
@@ -331,6 +364,15 @@ bool FStructDeserializer::Deserialize( void* OutStruct, UStruct& TypeInfo, IStru
 						FScriptMapHelper MapHelper(MapProperty, NewState.Data);
 						MapHelper.EmptyValues();
 					}
+					else if (NewState.Property->GetClass() == USetProperty::StaticClass())
+					{
+						// handle set property
+						USetProperty* SetProperty = Cast<USetProperty>(NewState.Property);
+						NewState.Data = SetProperty->ContainerPtrToValuePtr<void>(CurrentState.Data, CurrentState.ArrayIndex);
+
+						FScriptSetHelper SetHelper(SetProperty, NewState.Data);
+						SetHelper.EmptyElements();
+					}
 					else
 					{
 						// handle struct property
diff --git a/Engine/Source/Runtime/Serialization/Private/StructSerializer.cpp b/Engine/Source/Runtime/Serialization/Private/StructSerializer.cpp
--- a/Engine/Source/Runtime/Serialization/Private/StructSerializer.cpp
+++ b/Engine/Source/Runtime/Serialization/Private/StructSerializer.cpp
@@ -228,6 +228,42 @@ void FStructSerializer::Serialize( const void* Struct, UStruct& TypeInfo, IStruc
 			}
 		}
 
+		// sets
+		else if (CurrentState.ValueType == USetProperty::StaticClass())
+		{
+			if (!CurrentState.HasBeenProcessed)
+			{
+				Backend.BeginArray(CurrentState);
+
+				CurrentState.HasBeenProcessed = true;
+				StateStack.Push(CurrentState);
+
+				USetProperty* SetProperty = CastChecked<USetProperty>(CurrentState.ValueProperty);
+				FScriptSetHelper SetHelper(SetProperty, CurrentState.ValueData);
+				UProperty* ValueProperty = SetProperty->ElementProp;
+
+				// push elements on stack
+				for (int Index = SetHelper.Num() - 1; Index >= 0; --Index)
+				{
+					FStructSerializerState NewState;
+					{
+						NewState.HasBeenProcessed = false;
+						NewState.KeyData = nullptr;
+						NewState.KeyProperty = nullptr;
+						NewState.ValueData = SetHelper.GetElementPtr(Index);
+						NewState.ValueProperty = ValueProperty;
+						NewState.ValueType = ValueProperty->GetClass();
+					}
+
+					StateStack.Push(NewState);
+				}
+			}
+			else
+			{
+				Backend.EndArray(CurrentState);
+			}
+		}
+
 		// static arrays
 		else if (CurrentState.ValueProperty->ArrayDim > 1)
 		{
diff --git a/Engine/Source/Runtime/Windows/XAudio2/Private/XAudio2Device.cpp b/Engine/Source/Runtime/Windows/XAudio2/Private/XAudio2Device.cpp
--- a/Engine/Source/Runtime/Windows/XAudio2/Private/XAudio2Device.cpp
+++ b/Engine/Source/Runtime/Windows/XAudio2/Private/XAudio2Device.cpp
@@ -242,6 +242,9 @@ void FXAudio2Device::TeardownHardware()
 {
 	if (DeviceProperties)
 	{
+		// Make sure we clean up any pending decoding tasks before deleting the device properties
+		DeviceProperties->ProcessPendingTasksToCleanup(true);
+
 		delete DeviceProperties;
 		DeviceProperties = nullptr;
 	}
diff --git a/Engine/Source/Runtime/Windows/XAudio2/Private/XAudio2Support.h b/Engine/Source/Runtime/Windows/XAudio2/Private/XAudio2Support.h
--- a/Engine/Source/Runtime/Windows/XAudio2/Private/XAudio2Support.h
+++ b/Engine/Source/Runtime/Windows/XAudio2/Private/XAudio2Support.h
@@ -892,32 +892,57 @@ struct FXAudioDeviceProperties
 		PendingAsyncTasksToCleanUp.Add(PendingTaskInfo);
 	}
 
-	void ProcessPendingTasksToCleanup()
+	void ProcessPendingTasksToCleanup(bool bForceWait = false)
 	{
 		for (int32 i = PendingAsyncTasksToCleanUp.Num() - 1; i >= 0; --i)
 		{
 			FPendingAsyncTaskInfo& TaskInfo = PendingAsyncTasksToCleanUp[i];
 
-			// Clean up the tasks if they're finished
-			if (TaskInfo.RealtimeAsyncTask && TaskInfo.RealtimeAsyncTask->IsDone())
+			if (bForceWait)
 			{
-				delete TaskInfo.RealtimeAsyncTask;
-				TaskInfo.RealtimeAsyncTask = nullptr;
-			}
+				if (TaskInfo.RealtimeAsyncTask)
+				{
+					TaskInfo.RealtimeAsyncTask->EnsureCompletion(true);
 
-			if (TaskInfo.RealtimeAsyncHeaderParseTask && TaskInfo.RealtimeAsyncHeaderParseTask->IsDone())
-			{
-				delete TaskInfo.RealtimeAsyncHeaderParseTask;
-				TaskInfo.RealtimeAsyncHeaderParseTask = nullptr;
-			}
+					delete TaskInfo.RealtimeAsyncTask;
+					TaskInfo.RealtimeAsyncTask = nullptr;
+				}
+
+				if (TaskInfo.RealtimeAsyncHeaderParseTask)
+				{
+					TaskInfo.RealtimeAsyncHeaderParseTask->EnsureCompletion(true);
+					delete TaskInfo.RealtimeAsyncHeaderParseTask;
+					TaskInfo.RealtimeAsyncHeaderParseTask = nullptr;
+				}
+
+				check(TaskInfo.RealtimeAsyncTask == nullptr);
+				check(TaskInfo.RealtimeAsyncHeaderParseTask == nullptr);
 
-			// If both tasks are finished, clean up the buffer and remove it in the pending list
-			if (!TaskInfo.RealtimeAsyncHeaderParseTask && !TaskInfo.RealtimeAsyncTask)
-			{
-				check(TaskInfo.Buffer);
-				delete TaskInfo.Buffer;
 				PendingAsyncTasksToCleanUp.RemoveAtSwap(i, 1, false);
 			}
+			else
+			{
+				// Clean up the tasks if they're finished
+				if (TaskInfo.RealtimeAsyncTask && TaskInfo.RealtimeAsyncTask->IsDone())
+				{
+					delete TaskInfo.RealtimeAsyncTask;
+					TaskInfo.RealtimeAsyncTask = nullptr;
+				}
+
+				if (TaskInfo.RealtimeAsyncHeaderParseTask && TaskInfo.RealtimeAsyncHeaderParseTask->IsDone())
+				{
+					delete TaskInfo.RealtimeAsyncHeaderParseTask;
+					TaskInfo.RealtimeAsyncHeaderParseTask = nullptr;
+				}
+
+				// If both tasks are finished, clean up the buffer and remove it in the pending list
+				if (!TaskInfo.RealtimeAsyncHeaderParseTask && !TaskInfo.RealtimeAsyncTask)
+				{
+					check(TaskInfo.Buffer);
+					delete TaskInfo.Buffer;
+					PendingAsyncTasksToCleanUp.RemoveAtSwap(i, 1, false);
+				}
+			}
 		}
 	}
 };
