Mac NonUnity build fixes

diff --git a/Engine/Source/Runtime/Engine/Classes/Components/SplineComponent.h b/Engine/Source/Runtime/Engine/Classes/Components/SplineComponent.h
--- a/Engine/Source/Runtime/Engine/Classes/Components/SplineComponent.h
+++ b/Engine/Source/Runtime/Engine/Classes/Components/SplineComponent.h
@@ -1,7 +1,8 @@
 // Copyright 1998-2016 Epic Games, Inc. All Rights Reserved.
 
-
 #pragma once
+
+#include "Components/PrimitiveComponent.h"
 #include "SplineComponent.generated.h"
 
 
diff --git a/Engine/Source/Runtime/Engine/Public/StaticMeshResources.h b/Engine/Source/Runtime/Engine/Public/StaticMeshResources.h
--- a/Engine/Source/Runtime/Engine/Public/StaticMeshResources.h
+++ b/Engine/Source/Runtime/Engine/Public/StaticMeshResources.h
@@ -13,6 +13,7 @@
 #include "PrimitiveSceneProxy.h"
 #include "SceneManagement.h"
 #include "PhysicsEngine/BodySetupEnums.h"
+#include "Materials/MaterialInterface.h"
 
 /**
  * The LOD settings to use for a group of static meshes.
diff --git a/Engine/Source/Runtime/Landscape/Private/LandscapeCollision.cpp b/Engine/Source/Runtime/Landscape/Private/LandscapeCollision.cpp
--- a/Engine/Source/Runtime/Landscape/Private/LandscapeCollision.cpp
+++ b/Engine/Source/Runtime/Landscape/Private/LandscapeCollision.cpp
@@ -27,6 +27,7 @@
 #include "InstancedFoliageActor.h"
 #include "EngineUtils.h"
 #include "CookStats.h"
+#include "Materials/MaterialInstance.h"
 
 #if ENABLE_COOK_STATS
 namespace LandscapeCollisionCookStats
