#ue4 - Revert change from CL 2996025 that broke travel params when opening maps from the console (ie "?listen").

diff --git a/Engine/Source/Runtime/Engine/Private/UnrealEngine.cpp b/Engine/Source/Runtime/Engine/Private/UnrealEngine.cpp
--- a/Engine/Source/Runtime/Engine/Private/UnrealEngine.cpp
+++ b/Engine/Source/Runtime/Engine/Private/UnrealEngine.cpp
@@ -8684,7 +8684,7 @@ bool UEngine::HandleOpenCommand( const TCHAR* Cmd, FOutputDevice& Ar, UWorld *In
 		}
 	}
 
-	SetClientTravel( InWorld, *TestURL.Map, TRAVEL_Absolute );
+	SetClientTravel( InWorld, Cmd, TRAVEL_Absolute );
 	return true;
 }
 
