Copying //UE4/Dev-Framework to //UE4/Dev-Main (Source: //UE4/Dev-Framework @ 3021292)

diff --git a/Engine/Source/Runtime/Engine/Classes/PhysicsEngine/ConstraintInstance.h b/Engine/Source/Runtime/Engine/Classes/PhysicsEngine/ConstraintInstance.h
--- a/Engine/Source/Runtime/Engine/Classes/PhysicsEngine/ConstraintInstance.h
+++ b/Engine/Source/Runtime/Engine/Classes/PhysicsEngine/ConstraintInstance.h
@@ -1,6 +1,7 @@
 // Copyright 1998-2016 Epic Games, Inc. All Rights Reserved.
 #pragma once
-
+#include "ConstraintTypes.h"
+#include "ConstraintDrives.h"
 #include "ConstraintInstance.generated.h"
 
 #if WITH_PHYSX
@@ -12,31 +13,84 @@ namespace physx
 }
 #endif // WITH_PHYSX
 
-// LINEAR DOF
-UENUM()
-enum ELinearConstraintMotion
+class FMaterialRenderProxy;
+class FPrimitiveDrawInterface;
+class FMaterialRenderProxy;
+
+/** Container for properties of a physics constraint that can be easily swapped at runtime. This is useful for switching different setups when going from ragdoll to standup for example */
+USTRUCT()
+struct ENGINE_API FConstraintProfileProperties
 {
-	/** No constraint against this axis. */ 
-	LCM_Free	UMETA(DisplayName="Free"),
-	/** Limited freedom along this axis. */ 
-	LCM_Limited UMETA(DisplayName="Limited"),
-	/** Fully constraint against this axis. */
-	LCM_Locked UMETA(DisplayName="Locked"),
-
-	LCM_MAX,
+	GENERATED_USTRUCT_BODY()
+
+	/** Linear tolerance value in world units. If the distance error exceeds this tolerence limit, the body will be projected. */
+	UPROPERTY(EditAnywhere, Category = Projection, meta = (editcondition = "bEnableProjection", ClampMin = "0.0"))
+	float ProjectionLinearTolerance;
+
+	/** Angular tolerance value in world units. If the distance error exceeds this tolerence limit, the body will be projected. */
+	UPROPERTY(EditAnywhere, Category = Projection, meta = (editcondition = "bEnableProjection", ClampMin = "0.0"))
+	float ProjectionAngularTolerance;
+
+	/** Force needed to break the distance constraint. */
+	UPROPERTY(EditAnywhere, AdvancedDisplay, Category = Linear, meta = (ClampMin = "0.0"))
+	float LinearBreakThreshold;
+
+	/** Torque needed to break the joint. */
+	UPROPERTY(EditAnywhere, AdvancedDisplay, Category = Angular, meta = (editcondition = "bAngularBreakable", ClampMin = "0.0"))
+	float AngularBreakThreshold;
+
+	UPROPERTY(EditAnywhere, Category = Linear)
+	FLinearConstraint LinearLimit;
+
+	UPROPERTY(EditAnywhere, Category = Angular)
+	FConeConstraint ConeLimit;
+
+	UPROPERTY(EditAnywhere, Category = Angular)
+	FTwistConstraint TwistLimit;
+
+	UPROPERTY(EditAnywhere, Category = Linear)
+	FLinearDriveConstraint LinearDrive;
+
+	UPROPERTY(EditAnywhere, Category = Angular)
+	FAngularDriveConstraint AngularDrive;
+
+	// Disable collision between bodies joined by this constraint.
+	UPROPERTY(EditAnywhere, Category = Constraint)
+	uint8 bDisableCollision : 1;
+
+	/**
+	* If distance error between bodies exceeds 0.1 units, or rotation error exceeds 10 degrees, body will be projected to fix this.
+	* For example a chain spinning too fast will have its elements appear detached due to velocity, this will project all bodies so they still appear attached to each other.
+	*/
+	UPROPERTY(EditAnywhere, Category = Projection)
+	uint8 bEnableProjection : 1;
+
+	/** Whether it is possible to break the joint with angular force. */
+	UPROPERTY(EditAnywhere, AdvancedDisplay, Category = Angular)
+	uint8 bAngularBreakable : 1;
+
+	/** Whether it is possible to break the joint with linear force. */
+	UPROPERTY(EditAnywhere, AdvancedDisplay, Category = Linear)
+	uint8 bLinearBreakable : 1;
+
+	FConstraintProfileProperties();
+
+#if WITH_PHYSX
+
+	/** Updates physx joint properties from unreal properties (limits, drives, flags, etc...) */
+	void UpdatePhysX_AssumesLocked(physx::PxD6Joint* Joint, float AverageMass, float UseScale) const;
+
+	/** Updates physx drive target */
+	void UpdatePhysXDriveTarget_AssumesLocked(physx::PxD6Joint* Joint) const;
+
+	/** Updates physx joint breakable properties (threshold, etc...)*/
+	void UpdatePhysXBreakable_AssumesLocked(physx::PxD6Joint* Joint) const;
+
+	/** Updates physx joint flag based on profile properties */
+	void UpdatePhysXConstraintFlags_AssumesLocked(physx::PxD6Joint* Joint) const;
+#endif
 };
 
-UENUM()
-namespace EAngularDriveMode
-{
-	enum Type
-	{
-		/** Follows the shortest arc between a pair of anuglar configurations (Ignored if any angular limits/locks are used). */
-		SLERP,
-		/** Path is decomposed into twist and swing. Doesn't follow shortest arc and may have gimbal lock. (Works with angular limits/locks.) */
-		TwistAndSwing
-	};
-}
 
 /** Container for a physics representation of an object. */
 USTRUCT()
@@ -44,18 +98,11 @@ struct ENGINE_API FConstraintInstance
 {
 	GENERATED_USTRUCT_BODY()
 
-	///////////////////////////// INSTANCE DATA
-
-	/**
-	 *	Indicates position of this constraint within the array in SkeletalMeshComponent
-	 */
+	/** Indicates position of this constraint within the array in SkeletalMeshComponent. */
 	int32 ConstraintIndex;
 
-	/**
-     *  The component that created this instance.
-     *  This is intentionally *not* a UPROPERTY for serialization reasons.
-     *  Further, the owner should never be destroyed before the constraint.
-     */
+	/** The component that created this instance. */
+	UPROPERTY()
 	USceneComponent* OwnerComponent;
 
 #if WITH_PHYSX
@@ -114,339 +161,202 @@ struct ENGINE_API FConstraintInstance
 	UPROPERTY()
 	FVector SecAxis2;
 
-	// Disable collision between bodies joined by this constraint.
-	UPROPERTY(EditAnywhere, Category=Constraint)
-	uint32 bDisableCollision:1;
-
-	/** 
-	 * If distance error between bodies exceeds 0.1 units, or rotation error exceeds 10 degrees, body will be projected to fix this.
-	 * For example a chain spinning too fast will have its elements appear detached due to velocity, this will project all bodies so they still appear attached to each other. 
-	 */
-	UPROPERTY(EditAnywhere, Category=Projection)
-	uint32 bEnableProjection:1;
-
-	/** Linear tolerance value in world units. If the distance error exceeds this tolarance limit, the body will be projected. */
-	UPROPERTY(EditAnywhere, Category=Projection, meta=(editcondition = "bEnableProjection", ClampMin = "0.0"))
-	float ProjectionLinearTolerance;
-
-	/** Angular tolerance value in world units. If the distance error exceeds this tolarance limit, the body will be projected. */
-	UPROPERTY(EditAnywhere, Category=Projection, meta=(editcondition = "bEnableProjection", ClampMin = "0.0"))
-	float ProjectionAngularTolerance;
-
-	/** Indicates whether linear motion along the x axis is allowed, blocked or limited. If limited, the LinearLimit property will be used
-		to determine if a motion is allowed. See ELinearConstraintMotion. */
-	UPROPERTY(EditAnywhere, Category=Linear)
-	TEnumAsByte<enum ELinearConstraintMotion> LinearXMotion;
-
-	/** Indicates whether linear motion along the y axis is allowed, blocked or limited. If limited, the LinearLimit property will be used
-		to determine if a motion is allowed. See ELinearConstraintMotion. */
-	UPROPERTY(EditAnywhere, Category=Linear)
-	TEnumAsByte<enum ELinearConstraintMotion> LinearYMotion;
-
-	/** Indicates whether linear motion along the z axis is allowed, blocked or limited. If limited, the LinearLimit property will be used
-		to determine if a motion is allowed. See ELinearConstraintMotion. */
-	UPROPERTY(EditAnywhere, Category=Linear)
-	TEnumAsByte<enum ELinearConstraintMotion> LinearZMotion;
-
-	/** The limiting extent in world units of the linear motion for limited motion axes. */
-	UPROPERTY(EditAnywhere, Category=Linear, meta=(ClampMin = "0.0"))
-	float LinearLimitSize;
-
-	/** Whether we want to use soft limits instead of hard limits. With enabled soft limit, a constraint is used 
-		instead of hard-capping the motion. */
-	UPROPERTY(EditAnywhere, AdvancedDisplay, Category=Linear)
-	uint32 bLinearLimitSoft:1;
-
-	/** Stiffness of the linear soft limit constraint. Only used, when bLinearLimitSoft is true. */
-	UPROPERTY(EditAnywhere, AdvancedDisplay, Category=Linear, meta=(editcondition = "bLinearLimitSoft", ClampMin = "0.0"))
-	float LinearLimitStiffness;
-	
-	/** Damping of the linear soft limit constraint. Only used, when bLinearLimitSoft is true. */
-	UPROPERTY(EditAnywhere, AdvancedDisplay, Category=Linear, meta=(editcondition = "bLinearLimitSoft", ClampMin = "0.0"))
-	float LinearLimitDamping;
-
-	/** Defines whether the joint is breakable or not. */
-	UPROPERTY(EditAnywhere, AdvancedDisplay, Category=Linear)
-	uint32 bLinearBreakable:1;
-
-	/** Force needed to break the joint. */
-	UPROPERTY(EditAnywhere, AdvancedDisplay, Category=Linear, meta=(editcondition = "bLinearBreakable", ClampMin = "0.0"))
-	float LinearBreakThreshold;
-
-	/** Indicates whether rotation about the Z axis is allowed, blocked, or limited. If limited, the 
-		AngularLimit property will be used to determine the range of motion. See EAngularConstraintMotion. */
-	UPROPERTY(EditAnywhere, Category=Angular)
-	TEnumAsByte<enum EAngularConstraintMotion> AngularSwing1Motion;
-		
-	/** Indicates whether rotation about the the X axis is allowed, blocked, or limited. If limited, the
-		AngularLimit property will be used to determine the range of motion. See EAngularConstraintMotion. */
-	UPROPERTY(EditAnywhere, Category=Angular)
-	TEnumAsByte<enum EAngularConstraintMotion> AngularTwistMotion;
-
-	/** Indicates whether rotation about the Y axis is allowed, blocked, or limited. If limited, the
-		AngularLimit property will be used to determine the range of motion. See EAngularConstraintMotion. */
-	UPROPERTY(EditAnywhere, Category = Angular)
-	TEnumAsByte<enum EAngularConstraintMotion> AngularSwing2Motion;
-
-
-	/** Whether we want to use soft limits for swing motions instead of hard limits. With enabled 
-		soft limit, a constraint is used instead of hard-capping the motion. */
-	UPROPERTY(EditAnywhere, AdvancedDisplay, Category=Angular)
-	uint32 bSwingLimitSoft:1;
-
-	/** Whether we want to use soft limits for twist motions instead of hard limits. With enabled 
-		soft limit, a constraint is used instead of hard-capping the motion. */
-	UPROPERTY(EditAnywhere, AdvancedDisplay, Category=Angular)
-	uint32 bTwistLimitSoft:1;
-
-	/** Used if swing motion along the y axis is limited. The limit angle is specified in degrees and should be
-		between 0 and 180. */
-	UPROPERTY(EditAnywhere, Category=Angular, meta=(ClampMin = "0.0", ClampMax = "180.0"))
-	float Swing1LimitAngle;
-	
-	/** Used if twist motion along the x axis is limited. The limit angle is specified in degrees and should be
-		between 0 and 180. */
-	UPROPERTY(EditAnywhere, Category=Angular, meta=(ClampMin = "0.0", ClampMax = "180.0"))
-	float TwistLimitAngle;
-
-	/** Used if swing motion along the z axis is limited. The limit angle is specified in degrees and should be
-	between 0 and 180. */
-	UPROPERTY(EditAnywhere, Category = Angular, meta = (ClampMin = "0.0", ClampMax = "180.0"))
-	float Swing2LimitAngle;
-
-	/** Stiffness of the swing limit constraint if soft limit is used for swing motions. */
-	UPROPERTY(EditAnywhere, AdvancedDisplay, Category=Angular, meta=(editcondition = "bSwingLimitSoft", ClampMin = "0.0"))
-	float SwingLimitStiffness;
-
-	/** Damping of the swing limit constraint if soft limit is used for swing motions. */
-	UPROPERTY(EditAnywhere, AdvancedDisplay, Category=Angular, meta=(editcondition = "bSwingLimitSoft", ClampMin = "0.0"))
-	float SwingLimitDamping;
-
-	/** Stiffness of the twist limit constraint if soft limit is used for twist motions. */
-	UPROPERTY(EditAnywhere, AdvancedDisplay, Category=Angular, meta=(editcondition = "bTwistLimitSoft", ClampMin = "0.0"))
-	float TwistLimitStiffness;
-
-	/** Damping of the twist limit constraint if soft limit is used for twist motions. */
-	UPROPERTY(EditAnywhere, AdvancedDisplay, Category=Angular, meta=(editcondition = "bTwistLimitSoft", ClampMin = "0.0"))
-	float TwistLimitDamping;
-
 	/** Specifies the angular offset between the two frames of reference. By default limit goes from (-Angle, +Angle)
-	  * This allows you to bias the limit for swing1 swing2 and twist.
-      * This is also used as the reference orientation for Angular Orientation Drive.
-      */
+	* This allows you to bias the limit for swing1 swing2 and twist. */
 	UPROPERTY(EditAnywhere, Category = Angular)
 	FRotator AngularRotationOffset;
-	
-	/** Whether it is possible to break the joint with angular force. */
-	UPROPERTY(EditAnywhere, AdvancedDisplay, Category=Angular)
-	uint32 bAngularBreakable:1;
-
-	/** Angular force needed to break the joint. */
-	UPROPERTY(EditAnywhere, AdvancedDisplay, Category=Angular, meta=(editcondition = "bAngularBreakable", ClampMin = "0.0"))
-	float AngularBreakThreshold;
 
 	/** If true, linear limits scale using the absolute min of the 3d scale of the owning component */
 	UPROPERTY(EditAnywhere, Category = Linear)
 	uint32 bScaleLinearLimits : 1;
 
-private:
-
-	/** Enables/Disables linear position drive along the x axis. */
-	UPROPERTY(EditAnywhere, Category=LinearMotor)
-	uint32 bLinearXPositionDrive:1;
-
-	/** Enables/Disables linear velocity drive along the x axis. */
-	UPROPERTY(EditAnywhere, Category=LinearMotor)
-	uint32 bLinearXVelocityDrive:1;
-
-	/** Enables/Disables linear position drive along the y axis. */
-	UPROPERTY(EditAnywhere, Category=LinearMotor)
-	uint32 bLinearYPositionDrive:1;
-
-	/** Enables/Disables linear velocity drive along the y axis. */
-	UPROPERTY(EditAnywhere, Category=LinearMotor)
-	uint32 bLinearYVelocityDrive:1;
-
-	/** Enables/Disables linear position drive along the z axis. */
-	UPROPERTY(EditAnywhere, Category=LinearMotor)
-	uint32 bLinearZPositionDrive:1;
+	float AverageMass;
 
-	/** Enables/Disables linear velocity drive along the z axis. */
-	UPROPERTY(EditAnywhere, Category=LinearMotor)
-	uint32 bLinearZVelocityDrive:1;
+	//Constraint Data (properties easily swapped at runtime based on different constraint profiles)
+	UPROPERTY(EditAnywhere, Category = Constraint)
+	FConstraintProfileProperties ProfileInstance;
 
 public:
+	/** Copies behavior properties from the given profile. Automatically updates the physx representation if it's been created */
+	void CopyProfilePropertiesFrom(FConstraintProfileProperties& FromProperties);
 
-	/** Enables/Disables linear position drive. */
-	UPROPERTY(EditAnywhere, BlueprintReadOnly, Category=LinearMotor)
-	uint32 bLinearPositionDrive:1;
-	
-	/** Enables/Disables linear velocity drive. */
-	UPROPERTY(EditAnywhere, BlueprintReadOnly, Category = LinearMotor)
-	uint32 bLinearVelocityDrive:1;
-
-	/** Target position the linear drive. Only the components that are enabled are used. */
-	UPROPERTY(EditAnywhere, BlueprintReadOnly, Category=LinearMotor, meta=(editcondition="bLinearPositionDrive"))
-	FVector LinearPositionTarget;
-
-	/** Target velocity the linear drive. */
-	UPROPERTY(EditAnywhere, BlueprintReadOnly, Category = LinearMotor, meta = (editcondition = "bLinearVelocityDrive"))
-	FVector LinearVelocityTarget;
-
-	/** Spring to apply to the for linear drive. */
-	UPROPERTY(EditAnywhere, BlueprintReadOnly, Category = LinearMotor, meta = (DisplayName = "Linear Position Strength"))
-	float LinearDriveSpring;
-
-	/** Damping to apply to the for linear drive. */
-	UPROPERTY(EditAnywhere, BlueprintReadOnly, Category = LinearMotor, meta = (DisplayName = "Linear Velocity Strength"))
-	float LinearDriveDamping;
-
-	/** Limit to the force the linear drive can apply. */
-	UPROPERTY(EditAnywhere, BlueprintReadOnly, Category=LinearMotor, meta = (DisplayName = "Max Linear Force"))
-	float LinearDriveForceLimit;
-
-	UPROPERTY()
-	uint32 bSwingPositionDrive_DEPRECATED:1;
+#if WITH_PHYSX
+	FPhysxUserData PhysxUserData;
+#endif
 
-	UPROPERTY()
-	uint32 bSwingVelocityDrive_DEPRECATED:1;
+private:
+	/** The component scale passed in during initialization*/
+	float LastKnownScale;
 
-	UPROPERTY()
-	uint32 bTwistPositionDrive_DEPRECATED:1;
+public:
 
-	UPROPERTY()
-	uint32 bTwistVelocityDrive_DEPRECATED:1;
+	/** Constructor **/
+	FConstraintInstance();
 
-	UPROPERTY()
-	uint32 bAngularSlerpDrive_DEPRECATED:1;
+	/** Gets the linear limit size */
+	float GetLinearLimit() const
+	{
+		return ProfileInstance.LinearLimit.Limit;
+	}
 
-	/** Enables the angular drive towards a target orientation. */
-	UPROPERTY(EditAnywhere, BlueprintReadOnly, Category=AngularMotor)
-	uint32 bAngularOrientationDrive:1;
+	/** Gets the motion type for the linear X-axis limit. */
+	ELinearConstraintMotion GetLinearXMotion() const
+	{
+		return ProfileInstance.LinearLimit.XMotion;
+	}
 
-private:
-	/** Enables the swing drive. Only relevant when in twist and swing mode */
-	UPROPERTY()
-	uint32 bEnableSwingDrive: 1;
+	/** Sets the Linear XMotion type */
+	void SetLinearXMotion(ELinearConstraintMotion ConstraintType)
+	{
+		ProfileInstance.LinearLimit.XMotion = ConstraintType;
+		UpdateLinearLimit();
+	}
 
-	/** Enables the swing drive. Only relevant when in twist and swing mode */
-	UPROPERTY()
-	uint32 bEnableTwistDrive : 1;
+	/** Sets the LinearX Motion Type and the limit distance (Note distance is the same for all 3 axes) */
+	void SetLinearXLimit(ELinearConstraintMotion ConstraintType, float InLinearLimitSize)
+	{
+		ProfileInstance.LinearLimit.XMotion = ConstraintType;
+		ProfileInstance.LinearLimit.Limit = InLinearLimitSize;
+		UpdateLinearLimit();
+	}
 
-public:
-	
-	/** Enables the angular drive towards a target velocity. */
-	UPROPERTY(EditAnywhere, BlueprintReadOnly, Category=AngularMotor)
-	uint32 bAngularVelocityDrive:1;
+	/** Gets the motion type for the linear Y-axis limit. */
+	ELinearConstraintMotion GetLinearYMotion() const
+	{
+		return ProfileInstance.LinearLimit.YMotion;
+	}
 
-	UPROPERTY()
-	FQuat AngularPositionTarget_DEPRECATED;
-	
-	/** The way rotation paths are estimated */
-	UPROPERTY(EditAnywhere, BlueprintReadOnly, Category = AngularMotor)
-	TEnumAsByte<enum EAngularDriveMode::Type> AngularDriveMode;
+	/** Sets the Linear YMotion type */
+	void SetLinearYMotion(ELinearConstraintMotion ConstraintType)
+	{
+		ProfileInstance.LinearLimit.YMotion = ConstraintType;
+		UpdateLinearLimit();
+	}
 
-	/** Target orientation for the angular drive. */
-	UPROPERTY(EditAnywhere, BlueprintReadOnly, Category=AngularMotor)
-	FRotator AngularOrientationTarget;
+	/** Sets the LinearY Motion Type and the limit distance (Note distance is the same for all 3 axes) */
+	void SetLinearYLimit(ELinearConstraintMotion ConstraintType, float InLinearLimitSize)
+	{
+		ProfileInstance.LinearLimit.YMotion = ConstraintType;
+		ProfileInstance.LinearLimit.Limit = InLinearLimitSize;
+		UpdateLinearLimit();
+	}
 
-	/** Target velocity for the angular drive. */
-	UPROPERTY(EditAnywhere, BlueprintReadOnly, Category=AngularMotor)
-	FVector AngularVelocityTarget;    // Revolutions per second
+	/** Gets the motion type for the linear Z-axis limit. */
+	ELinearConstraintMotion GetLinearZMotion() const
+	{
+		return ProfileInstance.LinearLimit.ZMotion;
+	}
 
-	/** Spring value to apply to the for angular drive. */
-	UPROPERTY(EditAnywhere, BlueprintReadOnly, Category = AngularMotor, meta = (DisplayName = "Angular Position Strength", editcondition = "bAngularOrientationDrive"))
-	float AngularDriveSpring;
+	/** Sets the Linear ZMotion type */
+	void SetLinearZMotion(ELinearConstraintMotion ConstraintType)
+	{
+		ProfileInstance.LinearLimit.ZMotion = ConstraintType;
+		UpdateLinearLimit();
+	}
 
-	/** Damping value to apply to the for angular drive. */
-	UPROPERTY(EditAnywhere, BlueprintReadOnly, Category = AngularMotor, meta = (DisplayName = "Angular Velocity Strength", editcondition = "bAngularVelocityDrive"))
-	float AngularDriveDamping;
+	/** Sets the LinearZ Motion Type and the limit distance (Note distance is the same for all 3 axes) */
+	void SetLinearZLimit(ELinearConstraintMotion ConstraintType, float InLinearLimitSize)
+	{
+		ProfileInstance.LinearLimit.ZMotion = ConstraintType;
+		ProfileInstance.LinearLimit.Limit = InLinearLimitSize;
+		UpdateLinearLimit();
+	}
 
-	/** Limit to the force the angular drive can apply. */
-	UPROPERTY(EditAnywhere, BlueprintReadOnly, Category = AngularMotor, meta = (DisplayName = "Max Angular Force"))
-	float AngularDriveForceLimit;
+	/** Gets the motion type for the swing1 of the cone constraint */
+	EAngularConstraintMotion GetAngularSwing1Motion() const
+	{
+		return ProfileInstance.ConeLimit.Swing1Motion;
+	}
 
-	float AverageMass;
+	/** Sets the cone limit's swing1 motion type */
+	void SetAngularSwing1Motion(EAngularConstraintMotion MotionType)
+	{
+		ProfileInstance.ConeLimit.Swing1Motion = MotionType;
+		UpdateAngularLimit();
+	}
 
-private:
-	/** The component scale passed in during initialization*/
-	float LastKnownScale;
+	// The current swing1 of the constraint
+	float GetCurrentSwing1() const;
 
-public:
+	/** Gets the cone limit swing1 angle in degrees */
+	float GetAngularSwing1Limit() const
+	{
+		return ProfileInstance.ConeLimit.Swing1LimitDegrees;
+	}
 
-	float GetLastKnownScale() const { return LastKnownScale; }
-	/** Sets the LinearX Motion Type
+	/** Sets the Angular Swing1 Motion Type
 	*	@param MotionType	New Motion Type
 	*/
-	void SetLinearXLimit(ELinearConstraintMotion ConstraintType, float InLinearLimitSize);
+	void SetAngularSwing1Limit(EAngularConstraintMotion MotionType, float InSwing1LimitAngle)
+	{
+		ProfileInstance.ConeLimit.Swing1Motion = MotionType;
+		ProfileInstance.ConeLimit.Swing1LimitDegrees = InSwing1LimitAngle;
+		UpdateAngularLimit();
+	}
+	
+	/** Gets the motion type for the swing2 of the cone constraint */
+	EAngularConstraintMotion GetAngularSwing2Motion() const
+	{
+		return ProfileInstance.ConeLimit.Swing2Motion;
+	}
 
-	/** Sets the LinearY Motion Type
-	*	@param MotionType	New Motion Type
-	*/
-	void SetLinearYLimit(ELinearConstraintMotion ConstraintType, float InLinearLimitSize);
+	/** Sets the cone limit's swing2 motion type */
+	void SetAngularSwing2Motion(EAngularConstraintMotion MotionType)
+	{
+		ProfileInstance.ConeLimit.Swing2Motion = MotionType;
+		UpdateAngularLimit();
+	}
 
-	/** Sets the LinearZ Motion Type
-	*	@param MotionType	New Motion Type
-	*/
-	void SetLinearZLimit(ELinearConstraintMotion ConstraintType, float InLinearLimitSize);
+	// The current swing2 of the constraint
+	float GetCurrentSwing2() const;
 
-	/** Sets the Angular Swing1 Motion Type
-	*	@param MotionType	New Motion Type
-	*/
-	void SetAngularSwing1Limit(EAngularConstraintMotion MotionType, float InSwing1LimitAngle);
+	/** Gets the cone limit swing2 angle in degrees */
+	float GetAngularSwing2Limit() const
+	{
+		return ProfileInstance.ConeLimit.Swing2LimitDegrees;
+	}
 
 	/** Sets the Angular Swing2 Motion Type
 	*	@param MotionType	New Motion Type
 	*/
-	void SetAngularSwing2Limit(EAngularConstraintMotion MotionType, float InSwing2LimitAngle);
-
-	/** Sets the Angular Twist Motion Type
-	*	@param MotionType	New Motion Type
-	*/
-	void SetAngularTwistLimit(EAngularConstraintMotion MotionType, float InTwistLimitAngle);
-
-#if WITH_PHYSX
-	FPhysxUserData PhysxUserData;
-#endif
-
-
-public:
-
-
-	/** Constructor **/
-	FConstraintInstance();
-
-	/** Create physics engine constraint. */
-	void InitConstraint(USceneComponent* Owner, FBodyInstance* Body1, FBodyInstance* Body2, float Scale);
-
-	/** Terminate physics engine constraint */
-	void TermConstraint();
-
-	bool IsTerminated() const;
+	void SetAngularSwing2Limit(EAngularConstraintMotion MotionType, float InSwing2LimitAngle)
+	{
+		ProfileInstance.ConeLimit.Swing2Motion = MotionType;
+		ProfileInstance.ConeLimit.Swing2LimitDegrees = InSwing2LimitAngle;
+		UpdateAngularLimit();
+	}
 
-	/** See if this constraint is valid. */
-	bool IsValidConstraintInstance() const;
+	/** Gets the motion type for the twist of the cone constraint */
+	EAngularConstraintMotion GetAngularTwistMotion() const
+	{
+		return ProfileInstance.TwistLimit.TwistMotion;
+	}
 
-	// Pass in reference frame in. If the constraint is currently active, this will set its active local pose. Otherwise the change will take affect in InitConstraint. 
-	void SetRefFrame(EConstraintFrame::Type Frame, const FTransform& RefFrame);
+	/** Sets the twist limit's motion type */
+	void SetAngularTwistMotion(EAngularConstraintMotion MotionType)
+	{
+		ProfileInstance.TwistLimit.TwistMotion = MotionType;
+		UpdateAngularLimit();
+	}
 
-	// Pass in reference position in (maintains reference orientation). If the constraint is currently active, this will set its active local pose. Otherwise the change will take affect in InitConstraint.
-	void SetRefPosition(EConstraintFrame::Type Frame, const FVector& RefPosition);
-	
-	// Pass in reference orientation in (maintains reference position). If the constraint is currently active, this will set its active local pose. Otherwise the change will take affect in InitConstraint.
-	void SetRefOrientation(EConstraintFrame::Type Frame, const FVector& PriAxis, const FVector& SecAxis);
-	
 	// The current twist of the constraint
 	float GetCurrentTwist() const;
 
-	// The current swing1 of the constraint
-	float GetCurrentSwing1() const;
-
-	// The current swing2 of the constraint
-	float GetCurrentSwing2() const;
+	/** Gets the twist limit angle in degrees */
+	float GetAngularTwistLimit() const
+	{
+		return ProfileInstance.TwistLimit.TwistLimitDegrees;
+	}
 
-	// Get component ref frame
-	FTransform GetRefFrame(EConstraintFrame::Type Frame) const;
+	/** Sets the Angular Twist Motion Type
+	*	@param MotionType	New Motion Type
+	*/
+	void SetAngularTwistLimit(EAngularConstraintMotion MotionType, float InTwistLimitAngle)
+	{
+		ProfileInstance.TwistLimit.TwistMotion = MotionType;
+		ProfileInstance.TwistLimit.TwistLimitDegrees = InTwistLimitAngle;
+		UpdateAngularLimit();
+	}
 
 	// @todo document
 	void CopyConstraintGeometryFrom(const FConstraintInstance* FromInstance);
@@ -454,66 +364,140 @@ public:
 	// @todo document
 	void CopyConstraintParamsFrom(const FConstraintInstance* FromInstance);
 
-	/** Get the position of this constraint in world space. */
-	FVector GetConstraintLocation();
-
 	// Retrieve the constraint force most recently applied to maintain this constraint. Returns 0 forces if the constraint is not initialized or broken.
 	void GetConstraintForce(FVector& OutLinearForce, FVector& OutAngularForce);
 
+	/** Set which linear position drives are enabled */
 	void SetLinearPositionDrive(bool bEnableXDrive, bool bEnableYDrive, bool bEnableZDrive);
-	void SetLinearVelocityDrive(bool bEnableXDrive, bool bEnableYDrive, bool bEnableZDrive);
-	void SetAngularPositionDrive(bool bInEnableSwingDrive, bool bInEnableTwistDrive);
-	void SetAngularVelocityDrive(bool bInEnableSwingDrive, bool bInEnableTwistDrive);
 
+	/** Whether the linear position drive is enabled */
+	bool IsLinearPositionDriveEnabled() const
+	{
+		return ProfileInstance.LinearDrive.IsPositionDriveEnabled();
+	}
+
+	/** Set the linear drive's target position position */
 	void SetLinearPositionTarget(const FVector& InPosTarget);
+
+	/** Set which linear velocity drives are enabled */
+	void SetLinearVelocityDrive(bool bEnableXDrive, bool bEnableYDrive, bool bEnableZDrive);
+
+	/** Whether the linear velocity drive is enabled */
+	bool IsLinearVelocityDriveEnabled() const
+	{
+		return ProfileInstance.LinearDrive.IsVelocityDriveEnabled();
+	}
+
+	/** Set the linear drive's target velocity*/
 	void SetLinearVelocityTarget(const FVector& InVelTarget);
-	void SetLinearDriveParams(float InSpring, float InDamping, float InForceLimit);
 
+	/** Set the linear drive's strength parameters */
+	void SetLinearDriveParams(float InPositionStrength, float InVelocityStrength, float InForceLimit);
+
+	/** Set which twist and swing orientation drives are enabled. Only applicable when Twist And Swing drive mode is used */
+	void SetAngularPositionDrive(bool bInEnableSwingDrive, bool bInEnableTwistDrive);
+
+	/** Whether the angular orientation drive is enabled */
+	bool IsAngularOrientationDriveEnabled() const
+	{
+		return ProfileInstance.AngularDrive.IsOrientationDriveEnabled();
+	}
+	
+	/** Set the angular drive's orientation target*/
 	void SetAngularOrientationTarget(const FQuat& InPosTarget);
+
+	/** Set which twist and swing angular velocity drives are enabled. Only applicable when Twist And Swing drive mode is used */
+	void SetAngularVelocityDrive(bool bInEnableSwingDrive, bool bInEnableTwistDrive);
+
+	/** Whether the angular velocity drive is enabled */
+	bool IsAngularVelocityDriveEnabled() const
+	{
+		return ProfileInstance.AngularDrive.IsVelocityDriveEnabled();
+	}
+	
+	/** Set the angular drive's angular velocity target*/
 	void SetAngularVelocityTarget(const FVector& InVelTarget);
-	void SetAngularDriveParams(float InSpring, float InDamping, float InForceLimit);
 
-	void SetDisableCollision(bool InDisableCollision);
+	/** Set the angular drive's strength parameters*/
+	void SetAngularDriveParams(float InSpring, float InDamping, float InForceLimit);
 
+	/** Refreshes the physics engine joint's linear limits. Only applicable if the joint has been created already.*/
 	void UpdateLinearLimit();
-	void UpdateAngularLimit();
 
+	/** Refreshes the physics engine joint's angular limits. Only applicable if the joint has been created already.*/
+	void UpdateAngularLimit();
 
-	/** Scale Angular Limit Constraints (as defined in RB_ConstraintSetup) */
-	void	SetAngularDOFLimitScale(float InSwing1LimitScale, float InSwing2LimitScale, float InTwistLimitScale);
+	/** Scale Angular Limit Constraints (as defined in RB_ConstraintSetup). This only affects the physics engine and does not update the unreal side so you can do things like a LERP of the scale values. */
+	void SetAngularDOFLimitScale(float InSwing1LimitScale, float InSwing2LimitScale, float InTwistLimitScale);
 
 	/** Allows you to dynamically change the size of the linear limit 'sphere'. */
-	void	SetLinearLimitSize(float NewLimitSize);
+	void SetLinearLimitSize(float NewLimitSize);
 
-	// @todo document
-	void DrawConstraint(class FPrimitiveDrawInterface* PDI, 
-		float Scale, float LimitDrawScale, bool bDrawLimits, bool bDrawSelected,
-		const FTransform& Con1Frame, const FTransform& Con2Frame, bool bDrawAsPoint) const;
+	/** Create physics engine constraint. */
+	void InitConstraint(USceneComponent* Owner, FBodyInstance* Body1, FBodyInstance* Body2, float Scale);
 
-	void ConfigureAsHinge(bool bOverwriteLimits = true);
+	/** Terminate physics engine constraint */
+	void TermConstraint();
 
-	void ConfigureAsPrismatic(bool bOverwriteLimits = true);
+	/** Whether the physics engine constraint has been terminated */
+	bool IsTerminated() const;
 
-	void ConfigureAsSkelJoint(bool bOverwriteLimits = true);
+	/** See if this constraint is valid. */
+	bool IsValidConstraintInstance() const;
 
-	void ConfigureAsBS(bool bOverwriteLimits = true);
+	// Get component ref frame
+	FTransform GetRefFrame(EConstraintFrame::Type Frame) const;
 
-	bool IsHinge() const;
+	// Pass in reference frame in. If the constraint is currently active, this will set its active local pose. Otherwise the change will take affect in InitConstraint. 
+	void SetRefFrame(EConstraintFrame::Type Frame, const FTransform& RefFrame);
 
-	bool IsPrismatic() const;
+	/** Get the position of this constraint in world space. */
+	FVector GetConstraintLocation();
 
-	bool IsSkelJoint() const;
+	// Pass in reference position in (maintains reference orientation). If the constraint is currently active, this will set its active local pose. Otherwise the change will take affect in InitConstraint.
+	void SetRefPosition(EConstraintFrame::Type Frame, const FVector& RefPosition);
 
-	bool IsBallAndSocket() const;
+	// Pass in reference orientation in (maintains reference position). If the constraint is currently active, this will set its active local pose. Otherwise the change will take affect in InitConstraint.
+	void SetRefOrientation(EConstraintFrame::Type Frame, const FVector& PriAxis, const FVector& SecAxis);
 
-	void PostSerialize(const FArchive& Ar);
+	/** Whether collision is currently disabled */
+	bool IsCollisionDisabled() const
+	{
+		return ProfileInstance.bDisableCollision;
+	}
 
+	/** Set whether jointed actors can collide with each other */
+	void SetDisableCollision(bool InDisableCollision);
+
+	// @todo document
+	void DrawConstraint(int32 ViewIndex, class FMeshElementCollector& Collector,
+		float Scale, float LimitDrawScale, bool bDrawLimits, bool bDrawSelected,
+		const FTransform& Con1Frame, const FTransform& Con2Frame, bool bDrawAsPoint) const
+	{
+		DrawConstraintImp(FPDIOrCollector(ViewIndex, Collector), Scale, LimitDrawScale, bDrawLimits, bDrawSelected, Con1Frame, Con2Frame, bDrawAsPoint);
+	}
+
+	void DrawConstraint(FPrimitiveDrawInterface* PDI,
+		float Scale, float LimitDrawScale, bool bDrawLimits, bool bDrawSelected,
+		const FTransform& Con1Frame, const FTransform& Con2Frame, bool bDrawAsPoint) const
+	{
+		DrawConstraintImp(FPDIOrCollector(PDI), Scale, LimitDrawScale, bDrawLimits, bDrawSelected, Con1Frame, Con2Frame, bDrawAsPoint);
+	}
+
+	bool Serialize(FArchive& Ar);
+	void PostSerialize(const FArchive& Ar);
+	
+	/** Notifies owner component that constraint has broken */
 	void OnConstraintBroken();
 
+	/** Turn on linear and angular projection */
 	void EnableProjection();
 
+	/** Turn off linear and angular projection */
 	void DisableProjection();
 
+	float GetLastKnownScale() const { return LastKnownScale; }
+
 	//Hacks to easily get zeroed memory for special case when we don't use GC
 	static void Free(FConstraintInstance * Ptr);
 	static FConstraintInstance * Alloc();
@@ -529,8 +513,164 @@ private:
 	bool ExecuteOnUnbrokenJointReadWrite(TFunctionRef<void(physx::PxD6Joint*)> Func) const;
 #endif
 
+	struct FPDIOrCollector
+	{
+		FPrimitiveDrawInterface* PDI;
+		FMeshElementCollector* Collector;
+		int32 ViewIndex;
+
+		FPDIOrCollector(FPrimitiveDrawInterface* InPDI)
+			: PDI(InPDI)
+			, Collector(nullptr)
+			, ViewIndex(INDEX_NONE)
+		{
+		}
+
+		FPDIOrCollector(int32 InViewIndex, FMeshElementCollector& InCollector)
+			: PDI(nullptr)
+			, Collector(&InCollector)
+			, ViewIndex(InViewIndex)
+		{
+		}
+
+		bool HasCollector() const
+		{
+			return Collector != nullptr;
+		}
+
+		FPrimitiveDrawInterface* GetPDI() const;
+
+		void DrawCylinder(const FVector& Start, const FVector& End, float Thickness, FMaterialRenderProxy* MaterialProxy, ESceneDepthPriorityGroup DepthPriority) const;
+	};
+
+	void DrawConstraintImp(const FPDIOrCollector& PDIOrCollector,
+		float Scale, float LimitDrawScale, bool bDrawLimits, bool bDrawSelected,
+		const FTransform& Con1Frame, const FTransform& Con2Frame, bool bDrawAsPoint) const;
+
 	void UpdateBreakable();
 	void UpdateDriveTarget();
+
+public:
+
+	///////////////////////////// DEPRECATED
+#if WITH_EDITORONLY_DATA
+	UPROPERTY()
+	uint32 bDisableCollision_DEPRECATED : 1;
+	UPROPERTY()
+	uint32 bEnableProjection_DEPRECATED : 1;
+	UPROPERTY()
+	float ProjectionLinearTolerance_DEPRECATED;
+	UPROPERTY()
+	float ProjectionAngularTolerance_DEPRECATED;
+	UPROPERTY()
+	TEnumAsByte<enum ELinearConstraintMotion> LinearXMotion_DEPRECATED;
+	UPROPERTY()
+	TEnumAsByte<enum ELinearConstraintMotion> LinearYMotion_DEPRECATED;
+	UPROPERTY()
+	TEnumAsByte<enum ELinearConstraintMotion> LinearZMotion_DEPRECATED;
+	UPROPERTY()
+	float LinearLimitSize_DEPRECATED;
+	UPROPERTY()
+	uint32 bLinearLimitSoft_DEPRECATED : 1;
+	UPROPERTY()
+	float LinearLimitStiffness_DEPRECATED;
+	UPROPERTY()
+	float LinearLimitDamping_DEPRECATED;
+	UPROPERTY()
+	uint32 bLinearBreakable_DEPRECATED : 1;
+	UPROPERTY()
+	float LinearBreakThreshold_DEPRECATED;
+	UPROPERTY()
+	TEnumAsByte<enum EAngularConstraintMotion> AngularSwing1Motion_DEPRECATED;
+	UPROPERTY()
+	TEnumAsByte<enum EAngularConstraintMotion> AngularTwistMotion_DEPRECATED;
+	UPROPERTY()
+	TEnumAsByte<enum EAngularConstraintMotion> AngularSwing2Motion_DEPRECATED;
+	UPROPERTY()
+	uint32 bSwingLimitSoft_DEPRECATED : 1;
+	UPROPERTY()
+	uint32 bTwistLimitSoft_DEPRECATED : 1;
+	UPROPERTY()
+	float Swing1LimitAngle_DEPRECATED;
+	UPROPERTY()
+	float TwistLimitAngle_DEPRECATED;
+	UPROPERTY()
+	float Swing2LimitAngle_DEPRECATED;
+	UPROPERTY()
+	float SwingLimitStiffness_DEPRECATED;
+	UPROPERTY()
+	float SwingLimitDamping_DEPRECATED;
+	UPROPERTY()
+	float TwistLimitStiffness_DEPRECATED;
+	UPROPERTY()
+	float TwistLimitDamping_DEPRECATED;
+	UPROPERTY()
+	uint32 bAngularBreakable_DEPRECATED : 1;
+	UPROPERTY()
+	float AngularBreakThreshold_DEPRECATED;
+private:
+	UPROPERTY()
+	uint32 bLinearXPositionDrive_DEPRECATED : 1;
+	UPROPERTY()
+	uint32 bLinearXVelocityDrive_DEPRECATED : 1;
+	UPROPERTY()
+	uint32 bLinearYPositionDrive_DEPRECATED : 1;
+	UPROPERTY()
+	uint32 bLinearYVelocityDrive_DEPRECATED : 1;
+	UPROPERTY()
+	uint32 bLinearZPositionDrive_DEPRECATED : 1;
+	UPROPERTY()
+	uint32 bLinearZVelocityDrive_DEPRECATED : 1;
+public:
+	UPROPERTY()
+	uint32 bLinearPositionDrive_DEPRECATED : 1;
+	UPROPERTY()
+	uint32 bLinearVelocityDrive_DEPRECATED : 1;
+	UPROPERTY()
+	FVector LinearPositionTarget_DEPRECATED;
+	UPROPERTY()
+	FVector LinearVelocityTarget_DEPRECATED;
+	UPROPERTY()
+	float LinearDriveSpring_DEPRECATED;
+	UPROPERTY()
+	float LinearDriveDamping_DEPRECATED;
+	UPROPERTY()
+	float LinearDriveForceLimit_DEPRECATED;
+	UPROPERTY()
+	uint32 bSwingPositionDrive_DEPRECATED : 1;
+	UPROPERTY()
+	uint32 bSwingVelocityDrive_DEPRECATED : 1;
+	UPROPERTY()
+	uint32 bTwistPositionDrive_DEPRECATED : 1;
+	UPROPERTY()
+	uint32 bTwistVelocityDrive_DEPRECATED : 1;
+	UPROPERTY()
+	uint32 bAngularSlerpDrive_DEPRECATED : 1;
+	UPROPERTY()
+	uint32 bAngularOrientationDrive_DEPRECATED : 1;
+private:
+	UPROPERTY()
+	uint32 bEnableSwingDrive_DEPRECATED : 1;
+	UPROPERTY()
+	uint32 bEnableTwistDrive_DEPRECATED : 1;
+public:
+	UPROPERTY()
+	uint32 bAngularVelocityDrive_DEPRECATED : 1;
+	UPROPERTY()
+	FQuat AngularPositionTarget_DEPRECATED;
+	UPROPERTY()
+	TEnumAsByte<EAngularDriveMode::Type> AngularDriveMode_DEPRECATED;
+	UPROPERTY()
+	FRotator AngularOrientationTarget_DEPRECATED;
+	UPROPERTY()
+	FVector AngularVelocityTarget_DEPRECATED;    // Revolutions per second
+	UPROPERTY()
+	float AngularDriveSpring_DEPRECATED;
+	UPROPERTY()
+	float AngularDriveDamping_DEPRECATED;
+	UPROPERTY()
+	float AngularDriveForceLimit_DEPRECATED;
+#endif //EDITOR_ONLY_DATA
 };
 
 template<>
@@ -538,6 +678,7 @@ struct TStructOpsTypeTraits<FConstraintInstance> : public TStructOpsTypeTraitsBa
 {
 	enum 
 	{
-		WithPostSerialize = true,
+		WithSerializer = true,
+		WithPostSerialize = true
 	};
 };
diff --git a/Engine/Source/Runtime/Engine/Classes/PhysicsEngine/PhysicsAsset.h b/Engine/Source/Runtime/Engine/Classes/PhysicsEngine/PhysicsAsset.h
--- a/Engine/Source/Runtime/Engine/Classes/PhysicsEngine/PhysicsAsset.h
+++ b/Engine/Source/Runtime/Engine/Classes/PhysicsEngine/PhysicsAsset.h
@@ -97,7 +97,7 @@ public:
 	void GetCollisionMesh(int32 ViewIndex, FMeshElementCollector& Collector, const USkeletalMesh* SkelMesh, const TArray<FTransform>& SpaceBases, const FTransform& LocalToWorld, const FVector& Scale3D);
 
 	// @todo document
-	void DrawConstraints(class FPrimitiveDrawInterface* PDI, const USkeletalMesh* SkelMesh, const TArray<FTransform>& SpaceBases, const FTransform& LocalToWorld, float Scale);
+	void DrawConstraints(int32 ViewIndex, FMeshElementCollector& Collector, const USkeletalMesh* SkelMesh, const TArray<FTransform>& SpaceBases, const FTransform& LocalToWorld, float Scale);
 
 
 	// @todo document
diff --git a/Engine/Source/Runtime/Engine/Classes/PhysicsEngine/PhysicsConstraintTemplate.h b/Engine/Source/Runtime/Engine/Classes/PhysicsEngine/PhysicsConstraintTemplate.h
--- a/Engine/Source/Runtime/Engine/Classes/PhysicsEngine/PhysicsConstraintTemplate.h
+++ b/Engine/Source/Runtime/Engine/Classes/PhysicsEngine/PhysicsConstraintTemplate.h
@@ -13,12 +13,90 @@
 #include "ConstraintInstance.h"
 #include "PhysicsConstraintTemplate.generated.h"
 
+/** Holds the constraint profile applicable for a given constraint*/
+UCLASS()
+class UPhysicsConstraintProfile : public UObject
+{
+	GENERATED_BODY()
+public:
+	UPROPERTY()
+	FConstraintProfileProperties ProfileProperties;
+};
+
+USTRUCT()
+struct FPhysicsConstraintProfileHandle
+{
+	GENERATED_BODY()
+	
+	UPROPERTY()
+	UPhysicsConstraintProfile* Profile;
+
+	UPROPERTY(EditAnywhere, Category=Constraint)
+	FName ProfileName;
+};
 
 UCLASS(hidecategories=Object, MinimalAPI)
 class UPhysicsConstraintTemplate : public UObject
 {
 	GENERATED_UCLASS_BODY()
 
+	UPROPERTY(EditAnywhere, Category=Joint, meta=(ShowOnlyInnerProperties))
+	FConstraintInstance DefaultInstance;
+
+	/** Easily lookup from name to constraint profile */
+	UPROPERTY(transient)
+	TMap<FName, UPhysicsConstraintProfile*> ConstraintProfileNameMap;
+
+	/** Handles to the constraint profiles applicable to this constraint*/
+	UPROPERTY(EditAnywhere, Category = Joint)
+	TArray<FPhysicsConstraintProfileHandle> ProfileHandles;
+
+	/** Which constraint profile to use by default when creating the constraint instance */
+	UPROPERTY(EditAnywhere, Category=Joint)
+	FName DefaultProfileName;
+
+	//~ Begin UObject Interface.
+	virtual void Serialize(FArchive& Ar) override;
+	//~ End UObject Interface.
+
+	UPhysicsConstraintProfile* GetDefaultProfile() const
+	{
+		return ConstraintProfileNameMap.FindRef(DefaultProfileName);
+	}
+
+	UPhysicsConstraintProfile* GetProfile(FName ProfileName) const
+	{
+		return ConstraintProfileNameMap.FindRef(ProfileName);
+	}
+
+#if WITH_EDITOR
+	virtual void PostEditChangeProperty(struct FPropertyChangedEvent& PropertyChangedEvent) override;
+#endif
+
+private:
+	/** When no profile is selected, use these settings. Allows us to save whatever setting user had before switching to new profile for example */
+	UPROPERTY()
+	FConstraintProfileProperties NoProfileInstance;
+
+	// Only needed for old content! Pre VER_UE4_ALL_PROPS_TO_CONSTRAINTINSTANCE
+	void CopySetupPropsToInstance(FConstraintInstance* Instance);
+
+#if WITH_EDITORONLY_DATA
+	void UpdateInstanceFromProfile();
+#endif
+
+#if WITH_EDITORONLY_DATA
+	void UpdateProfileBeingEdited();
+#endif
+
+	// Makes sure there are no duplicate names or null entries in the profiles array.
+	void SanitizeProfileArray();
+
+	// Updates the constraint profile map after array has been changed
+	void UpdateConstraintProfileMap();
+
+public:	//DEPRECATED
+#if WITH_EDITORONLY_DATA
 	UPROPERTY()
 	FName JointName_DEPRECATED;
 	UPROPERTY()
@@ -38,7 +116,7 @@ class UPhysicsConstraintTemplate : public UObject
 	UPROPERTY()
 	FVector SecAxis2_DEPRECATED;
 	UPROPERTY()
-	uint32 bEnableProjection_DEPRECATED:1;
+	uint32 bEnableProjection_DEPRECATED : 1;
 	UPROPERTY()
 	float ProjectionLinearTolerance_DEPRECATED;
 	UPROPERTY()
@@ -52,13 +130,13 @@ class UPhysicsConstraintTemplate : public UObject
 	UPROPERTY()
 	float LinearLimitSize_DEPRECATED;
 	UPROPERTY()
-	uint32 bLinearLimitSoft_DEPRECATED:1;
+	uint32 bLinearLimitSoft_DEPRECATED : 1;
 	UPROPERTY()
 	float LinearLimitStiffness_DEPRECATED;
 	UPROPERTY()
 	float LinearLimitDamping_DEPRECATED;
 	UPROPERTY()
-	uint32 bLinearBreakable_DEPRECATED:1;
+	uint32 bLinearBreakable_DEPRECATED : 1;
 	UPROPERTY()
 	float LinearBreakThreshold_DEPRECATED;
 	UPROPERTY()
@@ -68,9 +146,9 @@ class UPhysicsConstraintTemplate : public UObject
 	UPROPERTY()
 	TEnumAsByte<enum EAngularConstraintMotion> AngularTwistMotion_DEPRECATED;
 	UPROPERTY()
-	uint32 bSwingLimitSoft_DEPRECATED:1;
+	uint32 bSwingLimitSoft_DEPRECATED : 1;
 	UPROPERTY()
-	uint32 bTwistLimitSoft_DEPRECATED:1;
+	uint32 bTwistLimitSoft_DEPRECATED : 1;
 	UPROPERTY()
 	float Swing1LimitAngle_DEPRECATED;
 	UPROPERTY()
@@ -86,21 +164,10 @@ class UPhysicsConstraintTemplate : public UObject
 	UPROPERTY()
 	float TwistLimitDamping_DEPRECATED;
 	UPROPERTY()
-	uint32 bAngularBreakable_DEPRECATED:1;
+	uint32 bAngularBreakable_DEPRECATED : 1;
 	UPROPERTY()
 	float AngularBreakThreshold_DEPRECATED;
-
-	UPROPERTY(EditAnywhere, Category=Joint, meta=(ShowOnlyInnerProperties))
-	FConstraintInstance DefaultInstance;
-
-
-	//~ Begin UObject Interface.
-	virtual void Serialize(FArchive& Ar) override;
-	//~ End UObject Interface.
-
-private:
-	// Only needed for old content! Pre VER_UE4_ALL_PROPS_TO_CONSTRAINTINSTANCE
-	void CopySetupPropsToInstance(FConstraintInstance* Instance);
+#endif
 };
 
 
diff --git a/Engine/Source/Runtime/Engine/Private/ActiveSound.cpp b/Engine/Source/Runtime/Engine/Private/ActiveSound.cpp
--- a/Engine/Source/Runtime/Engine/Private/ActiveSound.cpp
+++ b/Engine/Source/Runtime/Engine/Private/ActiveSound.cpp
@@ -71,9 +71,9 @@ FActiveSound::FActiveSound()
 	, LastAudioVolume(nullptr)
 	, LastUpdateTime(0.f)
 	, SourceInteriorVolume(1.f)
-	, SourceInteriorLPF(1.f)
+	, SourceInteriorLPF(MAX_FILTER_FREQUENCY)
 	, CurrentInteriorVolume(1.f)
-	, CurrentInteriorLPF(1.f)
+	, CurrentInteriorLPF(MAX_FILTER_FREQUENCY)
 	, bIsAudible(true)
 	, ClosestListenerPtr(nullptr)
 {
diff --git a/Engine/Source/Runtime/Engine/Private/Actor.cpp b/Engine/Source/Runtime/Engine/Private/Actor.cpp
--- a/Engine/Source/Runtime/Engine/Private/Actor.cpp
+++ b/Engine/Source/Runtime/Engine/Private/Actor.cpp
@@ -1482,7 +1482,7 @@ void AActor::OnRep_AttachmentReplication()
 	}
 	else
 	{
-		DetachFromActor(FDetachmentTransformRules::KeepRelativeTransform);
+		DetachFromActor(FDetachmentTransformRules::KeepWorldTransform);
 
 		// Handle the case where an object was both detached and moved on the server in the same frame.
 		// Calling this extraneously does not hurt but will properly fire events if the movement state changed while attached.
diff --git a/Engine/Source/Runtime/Engine/Private/AdvancedPreviewScene.cpp b/Engine/Source/Runtime/Engine/Private/AdvancedPreviewScene.cpp
new file mode 100644
--- /dev/null
+++ b/Engine/Source/Runtime/Engine/Private/AdvancedPreviewScene.cpp
@@ -0,0 +1,229 @@
+// Copyright 1998-2016 Epic Games, Inc. All Rights Reserved.
+
+#if WITH_EDITOR
+
+#include "EnginePrivate.h"
+#include "AdvancedPreviewScene.h"
+#include "Components/PostProcessComponent.h"
+#include "Engine/Scene.h"
+#include "Engine/TextureCube.h"
+#include "Materials/MaterialInstanceConstant.h"
+
+#include "UnrealClient.h"
+#include "Editor/EditorPerProjectUserSettings.h"
+
+#include "Components/Directionallightcomponent.h"
+#include "Components/SphereReflectionCaptureComponent.h"
+#include "Components/SkyLightComponent.h"
+
+#include "AssetViewerSettings.h"
+
+FAdvancedPreviewScene::FAdvancedPreviewScene(ConstructionValues CVS, float InFloorOffset)
+	: FPreviewScene(CVS)
+{
+	DefaultSettings = UAssetViewerSettings::Get();
+	CurrentProfileIndex = DefaultSettings->Profiles.IsValidIndex(CurrentProfileIndex) ? GetDefault<UEditorPerProjectUserSettings>()->AssetViewerProfileIndex : 0;
+	ensureMsgf(DefaultSettings && DefaultSettings->Profiles.IsValidIndex(CurrentProfileIndex), TEXT("Invalid default settings pointer or current profile index"));
+	FPreviewSceneProfile& Profile = DefaultSettings->Profiles[CurrentProfileIndex];
+
+	const FTransform Transform(FRotator(0, 0, 0), FVector(0, 0, 0), FVector(1));
+	
+	// Add and set up sky light using the set cube map texture
+	SkyLightComponent = NewObject<USkyLightComponent>();
+	SkyLightComponent->Cubemap = Profile.EnvironmentCubeMap;
+	SkyLightComponent->SourceType = ESkyLightSourceType::SLS_SpecifiedCubemap;
+	SkyLightComponent->Mobility = EComponentMobility::Movable;
+	SkyLightComponent->bLowerHemisphereIsBlack = false;
+	SkyLightComponent->Intensity = Profile.SkyLightIntensity;
+	AddComponent(SkyLightComponent, Transform);
+	SkyLightComponent->UpdateSkyCaptureContents(PreviewWorld);
+	
+	// Large scale to prevent sphere from clipping
+	const FTransform SphereTransform(FRotator(0, 0, 0), FVector(0, 0, 0), FVector(2000));
+	SkyComponent = NewObject<UStaticMeshComponent>(GetTransientPackage());
+
+	// Set up sky sphere showing hte same cube map as used by the sky light
+	UStaticMesh* SkySphere = LoadObject<UStaticMesh>(NULL, TEXT("/Engine/BasicShapes/Sphere.Sphere"), NULL, LOAD_None, NULL);
+	check(SkySphere);
+	SkyComponent->SetStaticMesh(SkySphere);
+
+	UMaterial* SkyMaterial = LoadObject<UMaterial>(NULL, TEXT("/Engine/EditorMaterials/AssetViewer/M_SkyBox.M_SkyBox"), NULL, LOAD_None, NULL);
+	check(SkyMaterial);
+
+	InstancedSkyMaterial = NewObject<UMaterialInstanceConstant>(GetTransientPackage());
+	InstancedSkyMaterial->Parent = SkyMaterial;		
+
+	UTextureCube* DefaultTexture = LoadObject<UTextureCube>(NULL, TEXT("/Engine/MapTemplates/Sky/SunsetAmbientCubemap.SunsetAmbientCubemap"));
+	InstancedSkyMaterial->SetTextureParameterValueEditorOnly(FName("SkyBox"), ( Profile.EnvironmentCubeMap != nullptr ) ? Profile.EnvironmentCubeMap : DefaultTexture );
+	InstancedSkyMaterial->SetScalarParameterValueEditorOnly(FName("CubemapRotation"), Profile.LightingRigRotation / 360.0f);
+	InstancedSkyMaterial->SetScalarParameterValueEditorOnly(FName("Intensity"), Profile.SkyLightIntensity);
+	InstancedSkyMaterial->PostLoad();
+	SkyComponent->SetMaterial(0, InstancedSkyMaterial);
+	AddComponent(SkyComponent, SphereTransform);
+	
+	PostProcessComponent = NewObject<UPostProcessComponent>();
+	PostProcessComponent->Settings = Profile.PostProcessingSettings;
+	PostProcessComponent->bUnbound = true;
+	AddComponent(PostProcessComponent, Transform);
+
+	UStaticMesh* FloorMesh = LoadObject<UStaticMesh>(NULL, TEXT("/Engine/EditorMeshes/PhAT_FloorBox.PhAT_FloorBox"), NULL, LOAD_None, NULL);
+	check(FloorMesh);
+	FloorMeshComponent = NewObject<UStaticMeshComponent>(GetTransientPackage());
+	FloorMeshComponent->SetStaticMesh(FloorMesh);
+
+	FTransform FloorTransform(FRotator(0, 0, 0), FVector(0, 0, -(InFloorOffset)), FVector(4.0f, 4.0f, 1.0f ));
+	AddComponent(FloorMeshComponent, FloorTransform);	
+
+	// Set up directional light's cascaded shadow maps for higher shadow quality
+	DirectionalLight->DynamicShadowCascades = 4;
+	DirectionalLight->DynamicShadowDistanceMovableLight = 5000.0f;
+	DirectionalLight->ShadowBias = 0.7f;
+
+	bRotateLighting = Profile.bRotateLightingRig;
+	CurrentRotationSpeed = Profile.RotationSpeed;
+}
+
+void FAdvancedPreviewScene::UpdateScene(FPreviewSceneProfile& Profile, bool bUpdateSkyLight /*= true*/, bool bUpdateEnvironment  /*= true*/, bool bUpdatePostProcessing /*= true*/, bool bUpdateDirectionalLight /*= true*/)
+{
+	if (bUpdateSkyLight)
+	{
+		SkyLightComponent->Cubemap = Profile.EnvironmentCubeMap;
+		SkyLightComponent->SourceCubemapAngle = Profile.LightingRigRotation;
+		SkyLightComponent->SetIntensity(Profile.SkyLightIntensity);
+		SkyLightComponent->SetCaptureIsDirty();
+		SkyLightComponent->MarkRenderStateDirty();
+		SkyLightComponent->UpdateSkyCaptureContents(PreviewWorld);	
+
+		InstancedSkyMaterial->SetScalarParameterValueEditorOnly(FName("Intensity"), Profile.SkyLightIntensity);
+		InstancedSkyMaterial->PostEditChange();
+	}
+
+	if (bUpdateEnvironment)
+	{
+		InstancedSkyMaterial->SetTextureParameterValueEditorOnly(FName("SkyBox"), Profile.EnvironmentCubeMap);		
+		InstancedSkyMaterial->SetScalarParameterValueEditorOnly(FName("CubemapRotation"), Profile.LightingRigRotation / 360.0f);
+		InstancedSkyMaterial->PostEditChange();		
+	}
+
+	if (bUpdatePostProcessing)
+	{
+		PostProcessComponent->Settings = Profile.PostProcessingSettings;
+	}
+
+	if (bUpdateDirectionalLight)
+	{
+		DirectionalLight->SetIntensity(Profile.DirectionalLightIntensity);
+		DirectionalLight->SetLightColor(Profile.DirectionalLightColor);
+	}
+
+	SkyComponent->SetVisibility(Profile.bShowEnvironment, true);
+	SkyComponent->MarkRenderStateDirty();
+	FloorMeshComponent->SetVisibility(Profile.bShowFloor, true);
+
+	bRotateLighting = Profile.bRotateLightingRig;
+	CurrentRotationSpeed = Profile.RotationSpeed;
+}
+
+void FAdvancedPreviewScene::SetFloorOffset(const float InFloorOffset)
+{
+	FTransform FloorTransform(FRotator(0, 0, 0), FVector(0, 0, -(InFloorOffset)), FVector(4.0f, 4.0f, 1.0f));
+
+	FloorMeshComponent->SetRelativeTransform(FloorTransform);
+}
+
+void FAdvancedPreviewScene::SetProfileIndex(const int32 InProfileIndex)
+{
+	CurrentProfileIndex = InProfileIndex;
+	UpdateScene(DefaultSettings->Profiles[CurrentProfileIndex]);
+}
+
+void FAdvancedPreviewScene::Tick(float DeltaTime)
+{
+	checkf(DefaultSettings && DefaultSettings->Profiles.IsValidIndex(CurrentProfileIndex), TEXT("Invalid default settings pointer or current profile index"));
+
+	FPreviewSceneProfile& Profile = DefaultSettings->Profiles[CurrentProfileIndex];
+	if (Profile.bRotateLightingRig)
+	{
+		CurrentRotationSpeed = Profile.RotationSpeed;
+		Profile.LightingRigRotation = FMath::Fmod(FMath::Max(FMath::Min(Profile.LightingRigRotation + (CurrentRotationSpeed * DeltaTime), 360.0f), 0.0f), 360.0f);
+
+		FRotator LightDir = DirectionalLight->ComponentToWorld.GetUnitAxis(EAxis::X).Rotation();
+		LightDir.Yaw += DeltaTime * -CurrentRotationSpeed;
+		SetLightDirection(LightDir);
+	}
+
+	if ( PreviousRotation != Profile.LightingRigRotation )
+	{		
+		SkyLightComponent->SourceCubemapAngle = Profile.LightingRigRotation;
+		SkyLightComponent->SetCaptureIsDirty();
+		SkyLightComponent->MarkRenderStateDirty();
+		SkyLightComponent->UpdateSkyCaptureContents(PreviewWorld);
+
+		InstancedSkyMaterial->SetScalarParameterValueEditorOnly(FName("CubemapRotation"), Profile.LightingRigRotation / 360.0f);
+		InstancedSkyMaterial->PostEditChange();
+
+		PreviewWorld->UpdateAllReflectionCaptures();
+		PreviewWorld->UpdateAllSkyCaptures();
+	}
+
+	PreviousRotation = Profile.LightingRigRotation;
+}
+
+bool FAdvancedPreviewScene::IsTickable() const
+{
+	return true;
+}
+
+TStatId FAdvancedPreviewScene::GetStatId() const
+{
+	return TStatId();
+}
+
+const bool FAdvancedPreviewScene::HandleViewportInput(FViewport* InViewport, int32 ControllerId, FKey Key, float Delta, float DeltaTime, int32 NumSamples, bool bGamepad)
+{
+	bool bResult = false;
+	const bool bMouseButtonDown = InViewport->KeyState(EKeys::LeftMouseButton) || InViewport->KeyState(EKeys::MiddleMouseButton) || InViewport->KeyState(EKeys::RightMouseButton);
+	const bool bSkyMove = InViewport->KeyState(EKeys::K);
+
+	// Look at which axis is being dragged and by how much
+	const float DragX = (Key == EKeys::MouseX) ? Delta : 0.f;
+
+	// Move the sky around if K is down and the mouse has moved on the X-axis
+	if (bSkyMove && bMouseButtonDown)
+	{
+		static const float SkyRotationSpeed = 0.22f;
+		float SkyRotation = GetSkyRotation();
+		SkyRotation += -DragX * SkyRotationSpeed;
+		SetSkyRotation(SkyRotation);		
+		bResult = true;
+	}
+
+	return bResult;
+}
+
+const float FAdvancedPreviewScene::GetSkyRotation() const
+{
+	checkf(DefaultSettings && DefaultSettings->Profiles.IsValidIndex(CurrentProfileIndex), TEXT("Invalid default settings pointer or current profile index"));
+	return DefaultSettings->Profiles[CurrentProfileIndex].LightingRigRotation;
+}
+
+void FAdvancedPreviewScene::SetSkyRotation(const float SkyRotation)
+{
+	checkf(DefaultSettings && DefaultSettings->Profiles.IsValidIndex(CurrentProfileIndex), TEXT("Invalid default settings pointer or current profile index"));
+
+	float ClampedSkyRotation = FMath::Fmod(SkyRotation, 360.0f);
+	// Clamp and wrap around sky rotation
+	if (ClampedSkyRotation < 0.0f)
+	{
+		ClampedSkyRotation += 360.0f;
+	}
+	DefaultSettings->Profiles[CurrentProfileIndex].LightingRigRotation = ClampedSkyRotation;
+}
+
+const UStaticMeshComponent* FAdvancedPreviewScene::GetFloorMeshComponent() const
+{
+	checkf(FloorMeshComponent != nullptr, TEXT("Invalid floor mesh component pointer"));
+	return FloorMeshComponent;
+}
+
+#endif // WITH_EDITOR
\ No newline at end of file
diff --git a/Engine/Source/Runtime/Engine/Private/AssetViewerSettings.cpp b/Engine/Source/Runtime/Engine/Private/AssetViewerSettings.cpp
new file mode 100644
--- /dev/null
+++ b/Engine/Source/Runtime/Engine/Private/AssetViewerSettings.cpp
@@ -0,0 +1,62 @@
+// Copyright 1998-2016 Epic Games, Inc. All Rights Reserved.
+
+#include "EnginePrivate.h"
+#include "AssetViewerSettings.h"
+
+#if WITH_EDITOR
+#include "Editor/EditorPerProjectUserSettings.h"
+#endif // WITH_EDITOR
+
+UAssetViewerSettings::UAssetViewerSettings()
+{
+	// Make sure there always is one profile as default
+	Profiles.AddDefaulted(1);
+	Profiles[0].ProfileName = TEXT("Profile_0");
+
+#if WITH_EDITOR
+	NumProfiles = Profiles.Num();
+#endif // WITH_EDITOR
+}
+
+UAssetViewerSettings* UAssetViewerSettings::Get()
+{
+	// This is a singleton, use default object
+	UAssetViewerSettings* DefaultSettings = GetMutableDefault<UAssetViewerSettings>();	
+	return DefaultSettings;
+}
+
+#if WITH_EDITOR
+void UAssetViewerSettings::PostEditChangeProperty(struct FPropertyChangedEvent& PropertyChangedEvent)
+{
+	FName PropertyName = (PropertyChangedEvent.Property != nullptr) ? PropertyChangedEvent.Property->GetFName() : NAME_None;	
+	UObject* Outer = (PropertyChangedEvent.Property != nullptr) ? PropertyChangedEvent.Property->GetOuter() : nullptr;
+	if (Outer != nullptr && Outer->GetName() == "PostProcessSettings")
+	{
+		PropertyName = GET_MEMBER_NAME_CHECKED(FPreviewSceneProfile, PostProcessingSettings);
+	}
+
+	if (NumProfiles != Profiles.Num())
+	{
+		OnAssetViewerProfileAddRemovedEvent.Broadcast();
+		NumProfiles = Profiles.Num();
+	}
+	
+	OnAssetViewerSettingsChangedEvent.Broadcast(PropertyName);
+
+	SaveConfig();
+	UpdateDefaultConfigFile();
+}
+
+void UAssetViewerSettings::PostInitProperties()
+{
+	Super::PostInitProperties();
+	NumProfiles = Profiles.Num();
+
+	UEditorPerProjectUserSettings* ProjectSettings = GetMutableDefault<UEditorPerProjectUserSettings>();
+	if (!Profiles.IsValidIndex(ProjectSettings->AssetViewerProfileIndex))
+	{
+		ProjectSettings->AssetViewerProfileIndex = 0;
+	}
+}
+
+#endif // WITH_EDITOR
\ No newline at end of file
diff --git a/Engine/Source/Runtime/Engine/Private/Components/DestructibleComponent.cpp b/Engine/Source/Runtime/Engine/Private/Components/DestructibleComponent.cpp
--- a/Engine/Source/Runtime/Engine/Private/Components/DestructibleComponent.cpp
+++ b/Engine/Source/Runtime/Engine/Private/Components/DestructibleComponent.cpp
@@ -481,8 +481,9 @@ void UDestructibleComponent::AddImpulseAtLocation( FVector Impulse, FVector Posi
 	ExecuteOnPhysicsReadWrite([&]
 	{
 		const int32 ChunkIdx = BoneIdxToChunkIdx(GetBoneIndex(BoneName));
-		PxVec3 Location = U2PVector(Position);
-		ApexDestructibleActor->addForce(ChunkIdx, U2PVector(Impulse),  PxForceMode::eIMPULSE);
+		PxVec3 PxPosition = U2PVector(Position);
+
+		ApexDestructibleActor->addForce(ChunkIdx, U2PVector(Impulse),  PxForceMode::eIMPULSE, &PxPosition);
 	});
 #endif
 }
diff --git a/Engine/Source/Runtime/Engine/Private/Components/SceneComponent.cpp b/Engine/Source/Runtime/Engine/Private/Components/SceneComponent.cpp
--- a/Engine/Source/Runtime/Engine/Private/Components/SceneComponent.cpp
+++ b/Engine/Source/Runtime/Engine/Private/Components/SceneComponent.cpp
@@ -3140,7 +3140,8 @@ void USceneComponent::PostUpdateNavigationData()
 	if (UNavigationSystem::ShouldUpdateNavOctreeOnComponentChange() && IsRegistered())
 	{
 		UWorld* MyWorld = GetWorld();
-		if ((MyWorld == nullptr) || !MyWorld->IsGameWorld() || !MyWorld->IsNetMode(ENetMode::NM_Client))
+		if (MyWorld != nullptr && MyWorld->GetNavigationSystem() != nullptr 
+			&& (MyWorld->GetNavigationSystem()->ShouldAllowClientSideNavigation() || !MyWorld->IsNetMode(ENetMode::NM_Client)))
 		{
 			// use propagated component's transform update in editor OR server game with additional navsys check
 			UNavigationSystem::UpdateNavOctreeAfterMove(this);
diff --git a/Engine/Source/Runtime/Engine/Private/Components/SkinnedMeshComponent.cpp b/Engine/Source/Runtime/Engine/Private/Components/SkinnedMeshComponent.cpp
--- a/Engine/Source/Runtime/Engine/Private/Components/SkinnedMeshComponent.cpp
+++ b/Engine/Source/Runtime/Engine/Private/Components/SkinnedMeshComponent.cpp
@@ -1207,24 +1207,25 @@ FTransform USkinnedMeshComponent::GetSocketTransform(FName InSocketName, ERelati
 
 	if (InSocketName != NAME_None)
 	{
-		int32 BoneIndex = GetBoneIndex(InSocketName);
-		if( BoneIndex != INDEX_NONE )
+		USkeletalMeshSocket const* const Socket = GetSocketByName(InSocketName);
+		// apply the socket transform first if we find a matching socket
+		if (Socket)
 		{
-			OutSocketTransform = GetBoneTransform(BoneIndex);
+			FTransform SocketLocalTransform = Socket->GetSocketLocalTransform();
+
+			int32 BoneIndex = GetBoneIndex(Socket->BoneName);
+			if (BoneIndex != INDEX_NONE)
+			{
+				FTransform BoneTransform = GetBoneTransform(BoneIndex);
+				OutSocketTransform = SocketLocalTransform * BoneTransform;
+			}
 		}
 		else
 		{
-			USkeletalMeshSocket const* const Socket = GetSocketByName(InSocketName);
-			if (Socket)
+			int32 BoneIndex = GetBoneIndex(InSocketName);
+			if (BoneIndex != INDEX_NONE)
 			{
-				FTransform SocketLocalTransform = Socket->GetSocketLocalTransform();
-
-				BoneIndex = GetBoneIndex(Socket->BoneName);
-				if( BoneIndex != INDEX_NONE )
-				{
-					FTransform BoneTransform = GetBoneTransform(BoneIndex);
-					OutSocketTransform = SocketLocalTransform * BoneTransform; 
-				}
+				OutSocketTransform = GetBoneTransform(BoneIndex);
 			}
 		}
 	}
diff --git a/Engine/Source/Runtime/Engine/Private/DestructibleMesh.cpp b/Engine/Source/Runtime/Engine/Private/DestructibleMesh.cpp
--- a/Engine/Source/Runtime/Engine/Private/DestructibleMesh.cpp
+++ b/Engine/Source/Runtime/Engine/Private/DestructibleMesh.cpp
@@ -186,7 +186,7 @@ void UDestructibleMesh::Serialize(FArchive& Ar)
 		{
 			Name = ApexDestructibleAsset->getName();
 		}
-		if( Name == NULL )
+		if( Name == NULL  || Ar.IsCooking())
 		{
 			Name = "";
 		}
diff --git a/Engine/Source/Runtime/Engine/Private/PhysicsEngine/BodyInstance.cpp b/Engine/Source/Runtime/Engine/Private/PhysicsEngine/BodyInstance.cpp
--- a/Engine/Source/Runtime/Engine/Private/PhysicsEngine/BodyInstance.cpp
+++ b/Engine/Source/Runtime/Engine/Private/PhysicsEngine/BodyInstance.cpp
@@ -665,9 +665,9 @@ void FBodyInstance::CreateDOFLock()
 
 	DOFConstraint = FConstraintInstance::Alloc();
 	{
-		DOFConstraint->bSwingLimitSoft = false;
-		DOFConstraint->bTwistLimitSoft = false;
-		DOFConstraint->bLinearLimitSoft = false;
+		DOFConstraint->ProfileInstance.ConeLimit.bSoftConstraint = false;
+		DOFConstraint->ProfileInstance.TwistLimit.bSoftConstraint  = false;
+		DOFConstraint->ProfileInstance.LinearLimit.bSoftConstraint  = false;
 
 		const FTransform TM = GetUnrealWorldTransform(false);
 		FVector Normal = FVector(1, 0, 0);
@@ -676,26 +676,26 @@ void FBodyInstance::CreateDOFLock()
 
 		if(DOF != EDOFMode::SixDOF)
 		{
-			DOFConstraint->AngularSwing1Motion = (bLockRotation || DOFMode != EDOFMode::CustomPlane) ? EAngularConstraintMotion::ACM_Locked : EAngularConstraintMotion::ACM_Free;
-			DOFConstraint->AngularSwing2Motion = (bLockRotation || DOFMode != EDOFMode::CustomPlane) ? EAngularConstraintMotion::ACM_Locked : EAngularConstraintMotion::ACM_Free;
-			DOFConstraint->AngularTwistMotion = EAngularConstraintMotion::ACM_Free;
-
-			DOFConstraint->LinearXMotion = (bLockTranslation || DOFMode != EDOFMode::CustomPlane) ? ELinearConstraintMotion::LCM_Locked : ELinearConstraintMotion::LCM_Free;
-			DOFConstraint->LinearYMotion = ELinearConstraintMotion::LCM_Free;
-			DOFConstraint->LinearZMotion = ELinearConstraintMotion::LCM_Free;
+			DOFConstraint->SetAngularSwing1Motion((bLockRotation || DOFMode != EDOFMode::CustomPlane) ? EAngularConstraintMotion::ACM_Locked : EAngularConstraintMotion::ACM_Free);
+			DOFConstraint->SetAngularSwing2Motion((bLockRotation || DOFMode != EDOFMode::CustomPlane) ? EAngularConstraintMotion::ACM_Locked : EAngularConstraintMotion::ACM_Free);
+			DOFConstraint->SetAngularTwistMotion(EAngularConstraintMotion::ACM_Free);
+			
+			DOFConstraint->SetLinearXMotion((bLockTranslation || DOFMode != EDOFMode::CustomPlane) ? ELinearConstraintMotion::LCM_Locked : ELinearConstraintMotion::LCM_Free);
+			DOFConstraint->SetLinearYMotion(ELinearConstraintMotion::LCM_Free);
+			DOFConstraint->SetLinearZMotion(ELinearConstraintMotion::LCM_Free);
 
 			Normal = LockedAxis.GetSafeNormal();
 			FVector Garbage;
 			Normal.FindBestAxisVectors(Garbage, Sec);
 		}else
 		{
-			DOFConstraint->AngularTwistMotion = bLockXRotation ? EAngularConstraintMotion::ACM_Locked : EAngularConstraintMotion::ACM_Free;
-			DOFConstraint->AngularSwing2Motion = bLockYRotation ? EAngularConstraintMotion::ACM_Locked : EAngularConstraintMotion::ACM_Free;
-			DOFConstraint->AngularSwing1Motion = bLockZRotation ? EAngularConstraintMotion::ACM_Locked : EAngularConstraintMotion::ACM_Free;
+			DOFConstraint->SetAngularTwistMotion(bLockXRotation ? EAngularConstraintMotion::ACM_Locked : EAngularConstraintMotion::ACM_Free);
+			DOFConstraint->SetAngularSwing2Motion(bLockYRotation ? EAngularConstraintMotion::ACM_Locked : EAngularConstraintMotion::ACM_Free);
+			DOFConstraint->SetAngularSwing1Motion(bLockZRotation ? EAngularConstraintMotion::ACM_Locked : EAngularConstraintMotion::ACM_Free);
 
-			DOFConstraint->LinearXMotion = bLockXTranslation ? ELinearConstraintMotion::LCM_Locked : ELinearConstraintMotion::LCM_Free;
-			DOFConstraint->LinearYMotion = bLockYTranslation ? ELinearConstraintMotion::LCM_Locked : ELinearConstraintMotion::LCM_Free;
-			DOFConstraint->LinearZMotion = bLockZTranslation ? ELinearConstraintMotion::LCM_Locked : ELinearConstraintMotion::LCM_Free;
+			DOFConstraint->SetLinearXMotion(bLockXTranslation ? ELinearConstraintMotion::LCM_Locked : ELinearConstraintMotion::LCM_Free);
+			DOFConstraint->SetLinearYMotion(bLockYTranslation ? ELinearConstraintMotion::LCM_Locked : ELinearConstraintMotion::LCM_Free);
+			DOFConstraint->SetLinearZMotion(bLockZTranslation ? ELinearConstraintMotion::LCM_Locked : ELinearConstraintMotion::LCM_Free);
 		}
 
 		DOFConstraint->PriAxis1 = TM.InverseTransformVectorNoScale(Normal);
diff --git a/Engine/Source/Runtime/Engine/Private/PhysicsEngine/ConstraintDrives.cpp b/Engine/Source/Runtime/Engine/Private/PhysicsEngine/ConstraintDrives.cpp
new file mode 100644
--- /dev/null
+++ b/Engine/Source/Runtime/Engine/Private/PhysicsEngine/ConstraintDrives.cpp
@@ -0,0 +1,136 @@
+// Copyright 1998-2016 Epic Games, Inc. All Rights Reserved.
+
+#include "EnginePrivate.h"
+#include "PhysicsPublic.h"
+#include "PhysXSupport.h"
+#include "PhysicsEngine/ConstraintInstance.h"
+
+const bool bIsAccelerationDrive = true;
+
+FConstraintDrive::FConstraintDrive()
+	: Stiffness(50.f)
+	, Damping(1.f)
+	, MaxForce(0.f)
+	, bEnablePositionDrive(false)
+	, bEnableVelocityDrive(false)
+{
+}
+
+FLinearDriveConstraint::FLinearDriveConstraint()
+	: PositionTarget(ForceInit)
+	, VelocityTarget(ForceInit)
+{
+}
+
+FAngularDriveConstraint::FAngularDriveConstraint()
+	: OrientationTarget(ForceInit)
+	, AngularVelocityTarget(ForceInit)
+	, AngularDriveMode(EAngularDriveMode::SLERP)
+{
+}
+
+#if WITH_PHYSX
+void WakeupJointedActors_AssumesLocked(PxD6Joint* Joint)
+{
+	PxRigidActor* Actor0 = nullptr;
+	PxRigidActor* Actor1 = nullptr;
+	Joint->getActors(Actor0, Actor1);
+
+	auto WakeupActorHelper = [](PxRigidActor* Actor)
+	{
+		if (PxRigidDynamic* DynamicActor = (Actor ? Actor->isRigidDynamic() : nullptr))
+		{
+			if (DynamicActor->getScene())
+			{
+				DynamicActor->wakeUp();
+			}
+
+		}
+	};
+
+	WakeupActorHelper(Actor0);
+	WakeupActorHelper(Actor1);
+}
+
+void FLinearDriveConstraint::UpdatePhysXLinearDrive_AssumesLocked(PxD6Joint* Joint) const
+{
+	XDrive.UpdatePhysXDrive_AssumesLocked(Joint, PxD6Drive::eX, true);
+	YDrive.UpdatePhysXDrive_AssumesLocked(Joint, PxD6Drive::eY, true);
+	ZDrive.UpdatePhysXDrive_AssumesLocked(Joint, PxD6Drive::eZ, true);
+	WakeupJointedActors_AssumesLocked(Joint);
+}
+
+void FAngularDriveConstraint::UpdatePhysXAngularDrive_AssumesLocked(PxD6Joint* Joint) const
+{
+	const bool bUseSlerpDrive = AngularDriveMode == EAngularDriveMode::SLERP;
+	SlerpDrive.UpdatePhysXDrive_AssumesLocked(Joint, PxD6Drive::eSLERP, bUseSlerpDrive);
+	SwingDrive.UpdatePhysXDrive_AssumesLocked(Joint, PxD6Drive::eSWING, !bUseSlerpDrive);
+	TwistDrive.UpdatePhysXDrive_AssumesLocked(Joint, PxD6Drive::eTWIST, !bUseSlerpDrive);
+	WakeupJointedActors_AssumesLocked(Joint);
+}
+
+void FConstraintDrive::UpdatePhysXDrive_AssumesLocked(PxD6Joint* Joint, int32 DriveType, bool bDriveEnabled) const
+{
+	check(DriveType >= PxD6Drive::eX && DriveType <= PxD6Drive::eCOUNT);
+	PxD6Drive::Enum PDriveType = (PxD6Drive::Enum) DriveType;
+
+	if(bDriveEnabled)
+	{
+		const float UseStiffness = bEnablePositionDrive ? Stiffness : 0.f;
+		const float UseDamping = bEnableVelocityDrive ? Damping : 0.f;
+		const float UseMaxForce = MaxForce > 0.f ? MaxForce : PX_MAX_F32;
+		Joint->setDrive(PDriveType, PxD6JointDrive(UseStiffness, UseDamping, UseMaxForce, bIsAccelerationDrive));
+	}
+	else
+	{
+		Joint->setDrive(PDriveType, PxD6JointDrive());
+	}
+}
+
+void FLinearDriveConstraint::SetLinearPositionDrive(bool bEnableXDrive, bool bEnableYDrive, bool bEnableZDrive)
+{
+	XDrive.bEnablePositionDrive = bEnableXDrive;
+	YDrive.bEnablePositionDrive = bEnableYDrive;
+	ZDrive.bEnablePositionDrive = bEnableZDrive;
+}
+
+void FLinearDriveConstraint::SetLinearVelocityDrive(bool bEnableXDrive, bool bEnableYDrive, bool bEnableZDrive)
+{
+	XDrive.bEnableVelocityDrive = bEnableXDrive;
+	YDrive.bEnableVelocityDrive = bEnableYDrive;
+	ZDrive.bEnableVelocityDrive = bEnableZDrive;
+}
+
+void FAngularDriveConstraint::SetAngularPositionDrive(bool InEnableSwingDrive, bool InEnableTwistDrive)
+{
+	SwingDrive.bEnablePositionDrive &= InEnableSwingDrive;
+	TwistDrive.bEnablePositionDrive &= InEnableSwingDrive;
+}
+
+void FAngularDriveConstraint::SetAngularVelocityDrive(bool InEnableSwingDrive, bool InEnableTwistDrive)
+{
+	SwingDrive.bEnableVelocityDrive &= InEnableSwingDrive;
+	TwistDrive.bEnableVelocityDrive &= InEnableSwingDrive;
+}
+
+void FConstraintDrive::SetDriveParams(float InStiffness, float InDamping, float InForceLimit)
+{
+	Stiffness = InStiffness;
+	Damping = InDamping;
+	MaxForce = InForceLimit;
+}
+
+void FLinearDriveConstraint::SetDriveParams(float InStiffness, float InDamping, float InForceLimit)
+{
+	XDrive.SetDriveParams(InStiffness, InDamping, InForceLimit);
+	YDrive.SetDriveParams(InStiffness, InDamping, InForceLimit);
+	ZDrive.SetDriveParams(InStiffness, InDamping, InForceLimit);
+}
+
+void FAngularDriveConstraint::SetDriveParams(float InStiffness, float InDamping, float InForceLimit)
+{
+	SwingDrive.SetDriveParams(InStiffness, InDamping, InForceLimit);
+	TwistDrive.SetDriveParams(InStiffness, InDamping, InForceLimit);
+	SlerpDrive.SetDriveParams(InStiffness, InDamping, InForceLimit);
+}
+#endif
\ No newline at end of file
