Copying //UE4/Dev-Framework to //UE4/Dev-Main (Source: //UE4/Dev-Framework @ 3021292)

diff --git a/Engine/Source/Runtime/Engine/Private/PhysicsEngine/ConstraintInstance.cpp b/Engine/Source/Runtime/Engine/Private/PhysicsEngine/ConstraintInstance.cpp
--- a/Engine/Source/Runtime/Engine/Private/PhysicsEngine/ConstraintInstance.cpp
+++ b/Engine/Source/Runtime/Engine/Private/PhysicsEngine/ConstraintInstance.cpp
@@ -5,27 +5,25 @@
 #include "PhysXSupport.h"
 #include "PhysicsEngine/PhysicsAsset.h"
 #include "PhysicsEngine/PhysicsConstraintComponent.h"
+#include "UObject/DevObjectVersion.h"
 
 #define LOCTEXT_NAMESPACE "ConstraintInstance"
 
-const bool bIsAccelerationDrive = true;
-
-/** Handy macro for setting BIT of VAR based on the bool CONDITION */
-#define SET_DRIVE_PARAM(VAR, CONDITION, BIT)   (VAR) = (CONDITION) ? ((VAR) | (BIT)) : ((VAR) & ~(BIT))
-
-// Cvars
-static TAutoConsoleVariable<float> CVarConstraintDampingScale(
+TAutoConsoleVariable<float> CVarConstraintDampingScale(
 	TEXT("p.ConstraintDampingScale"),
 	100000.f,
 	TEXT("The multiplier of constraint damping in simulation. Default: 100000"),
 	ECVF_ReadOnly);
 
-static TAutoConsoleVariable<float> CVarConstraintStiffnessScale(
+TAutoConsoleVariable<float> CVarConstraintStiffnessScale(
 	TEXT("p.ConstraintStiffnessScale"),
 	100000.f,
 	TEXT("The multiplier of constraint stiffness in simulation. Default: 100000"),
 	ECVF_ReadOnly);
 
+/** Handy macro for setting BIT of VAR based on the bool CONDITION */
+#define SET_DRIVE_PARAM(VAR, CONDITION, BIT)   (VAR) = (CONDITION) ? ((VAR) | (BIT)) : ((VAR) & ~(BIT))
+
 float RevolutionsToRads(const float Revolutions)
 {
 	return Revolutions * 2.f * PI;
@@ -37,56 +35,6 @@ FVector RevolutionsToRads(const FVector Revolutions)
 }
 
 #if WITH_PHYSX
-/** Util for setting soft limit params */
-void SetSoftLimitParams_AssumesLocked(PxJointLimitParameters* PLimit, bool bSoft, float Spring, float Damping)
-{
-	if(bSoft)
-	{
-		PLimit->stiffness = Spring * CVarConstraintStiffnessScale.GetValueOnGameThread();
-		PLimit->damping = Damping * CVarConstraintDampingScale.GetValueOnGameThread();
-	}
-}
-
-/** Util for converting from UE motion enum to physx motion enum */
-PxD6Motion::Enum U2PLinearMotion(ELinearConstraintMotion InMotion)
-{
-	switch (InMotion)
-	{
-		case ELinearConstraintMotion::LCM_Free: return PxD6Motion::eFREE;
-		case ELinearConstraintMotion::LCM_Limited: return PxD6Motion::eLIMITED;
-		case ELinearConstraintMotion::LCM_Locked: return PxD6Motion::eLOCKED;
-		default: check(0);	//unsupported motion type
-	}
-	
-	return PxD6Motion::eFREE;
-}
-
-/** Util for converting from UE motion enum to physx motion enum */
-PxD6Motion::Enum U2PAngularMotion(EAngularConstraintMotion InMotion)
-{
-	switch (InMotion)
-	{
-		case EAngularConstraintMotion::ACM_Free: return PxD6Motion::eFREE;
-		case EAngularConstraintMotion::ACM_Limited: return PxD6Motion::eLIMITED;
-		case EAngularConstraintMotion::ACM_Locked: return PxD6Motion::eLOCKED;
-		default: check(0);	//unsupported motion type
-	}
-	
-	return PxD6Motion::eFREE;
-}
-
-/** Util for setting linear movement for an axis */
-template <PxD6Axis::Enum PAxis>
-void SetLinearMovement_AssumesLocked(PxD6Joint* PD6Joint, ELinearConstraintMotion Motion, bool bLockLimitSize)
-{
-	if(Motion == LCM_Locked || (Motion == LCM_Limited && bLockLimitSize))
-	{
-		PD6Joint->setMotion(PAxis, PxD6Motion::eLOCKED);
-	}else
-	{
-		PD6Joint->setMotion(PAxis, U2PLinearMotion(Motion));
-	}
-}
 
 physx::PxD6Joint* FConstraintInstance::GetUnbrokenJoint_AssumesLocked() const
 {
@@ -126,95 +74,66 @@ bool FConstraintInstance::ExecuteOnUnbrokenJointReadWrite(TFunctionRef<void(phys
 }
 #endif //WITH_PHYSX
 
+
+FConstraintProfileProperties::FConstraintProfileProperties()
+	: ProjectionLinearTolerance(5.f)
+	, ProjectionAngularTolerance(180.f)
+	, LinearBreakThreshold(300.f)
+	, AngularBreakThreshold(500.f)
+	, bDisableCollision(false)
+	, bEnableProjection(true)
+	, bAngularBreakable(false)
+	, bLinearBreakable(false)
+{
+}
+
 void FConstraintInstance::UpdateLinearLimit()
 {
-	const float UseScale = bScaleLinearLimits ? LastKnownScale : 1.f;
 #if WITH_PHYSX
 	ExecuteOnUnbrokenJointReadWrite([&] (PxD6Joint* Joint)
 	{
-		bool bLockLimitSize = (LinearLimitSize < RB_MinSizeToLockDOF);
-
-		SetLinearMovement_AssumesLocked<PxD6Axis::eX>(Joint, LinearXMotion, bLockLimitSize);
-		SetLinearMovement_AssumesLocked<PxD6Axis::eY>(Joint, LinearYMotion, bLockLimitSize);
-		SetLinearMovement_AssumesLocked<PxD6Axis::eZ>(Joint, LinearZMotion, bLockLimitSize);
-
-		// If any DOF is locked/limited, set up the joint limit
-		if (LinearXMotion != LCM_Free || LinearYMotion != LCM_Free || LinearZMotion != LCM_Free)
-		{
-			// If limit drops below RB_MinSizeToLockDOF, just pass RB_MinSizeToLockDOF to physics - that axis will be locked anyway, and PhysX dislikes 0 here
-			float LinearLimit = FMath::Max<float>(LinearLimitSize * UseScale, RB_MinSizeToLockDOF);
-			PxJointLinearLimit PLinearLimit(GPhysXSDK->getTolerancesScale(), LinearLimit, 0.05f * GPhysXSDK->getTolerancesScale().length);
-			SetSoftLimitParams_AssumesLocked(&PLinearLimit, bLinearLimitSoft, LinearLimitStiffness*AverageMass, LinearLimitDamping*AverageMass);
-			Joint->setLinearLimit(PLinearLimit);
-		}
+		ProfileInstance.LinearLimit.UpdatePhysXLinearLimit_AssumesLocked(Joint, AverageMass, bScaleLinearLimits ? LastKnownScale : 1.f);
 	});
 #endif
 }
 
-
-
 void FConstraintInstance::UpdateAngularLimit()
 {
 #if WITH_PHYSX
-	ExecuteOnUnbrokenJointReadWrite([&] (PxD6Joint* Joint)
+	ExecuteOnUnbrokenJointReadWrite([&](PxD6Joint* Joint)
 	{
-		/////////////// TWIST LIMIT
-		PxD6Motion::Enum TwistMotion = PxD6Motion::eFREE;
-		if (AngularTwistMotion == ACM_Limited)
-		{
-			TwistMotion = PxD6Motion::eLIMITED;
-			// If angle drops below RB_MinAngleToLockDOF, just pass RB_MinAngleToLockDOF to physics - that axis will be locked anyway, and PhysX dislikes 0 here
-			float TwistLimitRad = FMath::DegreesToRadians(TwistLimitAngle);
-			PxJointAngularLimitPair PTwistLimitPair(-TwistLimitRad, TwistLimitRad, FMath::DegreesToRadians(1.f));
-			SetSoftLimitParams_AssumesLocked(&PTwistLimitPair, bTwistLimitSoft, TwistLimitStiffness*AverageMass, TwistLimitDamping*AverageMass);
-			Joint->setTwistLimit(PTwistLimitPair);
-		}
-		else if (AngularTwistMotion == ACM_Locked)
-		{
-			TwistMotion = PxD6Motion::eLOCKED;
-		}
-		Joint->setMotion(PxD6Axis::eTWIST, TwistMotion);
-
-		/////////////// SWING1 LIMIT
-		const PxD6Motion::Enum Swing1Motion = U2PAngularMotion(AngularSwing1Motion);
-		const PxD6Motion::Enum Swing2Motion = U2PAngularMotion(AngularSwing2Motion);
-
-		if (AngularSwing1Motion == ACM_Limited || AngularSwing2Motion == ACM_Limited)
-		{
-			//Clamp the limit value to valid range which PhysX won't ignore, both value have to be clamped even there is only one degree limit in constraint
-			float Limit1Rad = FMath::DegreesToRadians(FMath::ClampAngle(Swing1LimitAngle, KINDA_SMALL_NUMBER, 179.9999f));
-			float Limit2Rad = FMath::DegreesToRadians(FMath::ClampAngle(Swing2LimitAngle, KINDA_SMALL_NUMBER, 179.9999f));
-			PxJointLimitCone PSwingLimitCone(Limit2Rad, Limit1Rad, FMath::DegreesToRadians(1.f));
-			SetSoftLimitParams_AssumesLocked(&PSwingLimitCone, bSwingLimitSoft, SwingLimitStiffness*AverageMass, SwingLimitDamping*AverageMass);
-			Joint->setSwingLimit(PSwingLimitCone);
-		}
-
-		Joint->setMotion(PxD6Axis::eSWING2, Swing1Motion);
-		Joint->setMotion(PxD6Axis::eSWING1, Swing2Motion);
+		ProfileInstance.ConeLimit.UpdatePhysXConeLimit_AssumesLocked(Joint, AverageMass);
+		ProfileInstance.TwistLimit.UpdatePhysXTwistLimit_AssumesLocked(Joint, AverageMass);
 	});
 #endif
-	
 }
 
 void FConstraintInstance::UpdateBreakable()
+{
+#if WITH_PHYSX
+	ExecuteOnUnbrokenJointReadWrite([&](PxD6Joint* Joint)
+	{
+		ProfileInstance.UpdatePhysXBreakable_AssumesLocked(Joint);
+	});
+#endif
+}
+
+#if WITH_PHYSX
+void FConstraintProfileProperties::UpdatePhysXBreakable_AssumesLocked(PxD6Joint* Joint) const
 {
 	const float LinearBreakForce = bLinearBreakable ? LinearBreakThreshold : PX_MAX_REAL;
 	const float AngularBreakForce = bAngularBreakable ? AngularBreakThreshold : PX_MAX_REAL;
 
-#if WITH_PHYSX
-	ConstraintData->setBreakForce(LinearBreakForce, AngularBreakForce);
-#endif
+	Joint->setBreakForce(LinearBreakForce, AngularBreakForce);
 }
+#endif
 
 void FConstraintInstance::UpdateDriveTarget()
 {
 #if WITH_PHYSX
 	ExecuteOnUnbrokenJointReadWrite([&] (PxD6Joint* Joint)
 	{
-		FQuat OrientationTargetQuat(AngularOrientationTarget);
-
-		Joint->setDrivePosition(PxTransform(U2PVector(LinearPositionTarget), U2PQuat(OrientationTargetQuat)));
-		Joint->setDriveVelocity(U2PVector(LinearVelocityTarget), U2PVector(RevolutionsToRads(AngularVelocityTarget)));
+		ProfileInstance.UpdatePhysXDriveTarget_AssumesLocked(Joint);
 	});
 #endif
 }
@@ -227,44 +146,69 @@ FConstraintInstance::FConstraintInstance()
 	, ConstraintData(NULL)
 #endif	//WITH_PHYSX
 	, SceneIndex(0)
-	, bEnableProjection(true)
-	, bSwingLimitSoft(true)
-	, bTwistLimitSoft(true)
-	, Swing1LimitAngle(45)
-	, TwistLimitAngle(45)
-	, Swing2LimitAngle(45)
-	, SwingLimitStiffness(50)
-	, SwingLimitDamping(5)
-	, TwistLimitStiffness(50)
-	, TwistLimitDamping(5)
 	, bScaleLinearLimits(true)
-	, bLinearPositionDrive(false)
-	, bLinearVelocityDrive(false)
-	, LinearPositionTarget(ForceInit)
-	, LinearVelocityTarget(ForceInit)
-	, LinearDriveSpring(50.0f)
-	, LinearDriveDamping(1.0f)
-	, LinearDriveForceLimit(0)
-	, bAngularOrientationDrive(false)
-	, bEnableSwingDrive(true)
-	, bEnableTwistDrive(true)
-	, bAngularVelocityDrive(false)
-	, AngularOrientationTarget(ForceInit)
-	, AngularVelocityTarget(ForceInit)
-	, AngularDriveSpring(50.0f)
-	, AngularDriveDamping(1.0f)
-	, AngularDriveForceLimit(0)
 	, AverageMass(0.f)
 #if WITH_PHYSX
 	, PhysxUserData(this)
 #endif
+	, LastKnownScale(1.f)
+#if WITH_EDITORONLY_DATA
+	, bDisableCollision_DEPRECATED(false)
+	, bEnableProjection_DEPRECATED(true)
+	, ProjectionLinearTolerance_DEPRECATED(5.f)
+	, ProjectionAngularTolerance_DEPRECATED(180.f)
+	, LinearXMotion_DEPRECATED(ACM_Locked)
+	, LinearYMotion_DEPRECATED(ACM_Locked)
+	, LinearZMotion_DEPRECATED(ACM_Locked)
+	, LinearLimitSize_DEPRECATED(0.f)
+	, bLinearLimitSoft_DEPRECATED(false)
+	, LinearLimitStiffness_DEPRECATED(0.f)
+	, LinearLimitDamping_DEPRECATED(0.f)
+	, bLinearBreakable_DEPRECATED(false)
+	, LinearBreakThreshold_DEPRECATED(300.f)
+	, AngularSwing1Motion_DEPRECATED(ACM_Free)
+	, AngularTwistMotion_DEPRECATED(ACM_Free)
+	, AngularSwing2Motion_DEPRECATED(ACM_Free)
+	, bSwingLimitSoft_DEPRECATED(true)
+	, bTwistLimitSoft_DEPRECATED(true)
+	, Swing1LimitAngle_DEPRECATED(45)
+	, TwistLimitAngle_DEPRECATED(45)
+	, Swing2LimitAngle_DEPRECATED(45)
+	, SwingLimitStiffness_DEPRECATED(50)
+	, SwingLimitDamping_DEPRECATED(5)
+	, TwistLimitStiffness_DEPRECATED(50)
+	, TwistLimitDamping_DEPRECATED(5)
+	, bAngularBreakable_DEPRECATED(false)
+	, AngularBreakThreshold_DEPRECATED(500.f)
+	, bLinearXPositionDrive_DEPRECATED(false)
+	, bLinearXVelocityDrive_DEPRECATED(false)
+	, bLinearYPositionDrive_DEPRECATED(false)
+	, bLinearYVelocityDrive_DEPRECATED(false)
+	, bLinearZPositionDrive_DEPRECATED(false)
+	, bLinearZVelocityDrive_DEPRECATED(false)
+	, bLinearPositionDrive_DEPRECATED(false)
+	, bLinearVelocityDrive_DEPRECATED(false)
+	, LinearPositionTarget_DEPRECATED(ForceInit)
+	, LinearVelocityTarget_DEPRECATED(ForceInit)
+	, LinearDriveSpring_DEPRECATED(50.0f)
+	, LinearDriveDamping_DEPRECATED(1.0f)
+	, LinearDriveForceLimit_DEPRECATED(0)
+	, bSwingPositionDrive_DEPRECATED(false)
+	, bSwingVelocityDrive_DEPRECATED(false)
+	, bTwistPositionDrive_DEPRECATED(false)
+	, bTwistVelocityDrive_DEPRECATED(false)
+	, bAngularOrientationDrive_DEPRECATED(false)
+	, bEnableSwingDrive_DEPRECATED(true)
+	, bEnableTwistDrive_DEPRECATED(true)
+	, bAngularVelocityDrive_DEPRECATED(false)
+	, AngularPositionTarget_DEPRECATED(ForceInit)
+	, AngularOrientationTarget_DEPRECATED(ForceInit)
+	, AngularVelocityTarget_DEPRECATED(ForceInit)
+	, AngularDriveSpring_DEPRECATED(50.0f)
+	, AngularDriveDamping_DEPRECATED(1.0f)
+	, AngularDriveForceLimit_DEPRECATED(0)
+#endif	//EDITOR_ONLY_DATA
 {
-	bDisableCollision = false;
-
-	LinearXMotion = LCM_Locked;
-	LinearYMotion = LCM_Locked;
-	LinearZMotion = LCM_Locked;
-
 	Pos1 = FVector(0.0f, 0.0f, 0.0f);
 	PriAxis1 = FVector(1.0f, 0.0f, 0.0f);
 	SecAxis1 = FVector(0.0f, 1.0f, 0.0f);
@@ -272,22 +216,16 @@ FConstraintInstance::FConstraintInstance()
 	Pos2 = FVector(0.0f, 0.0f, 0.0f);
 	PriAxis2 = FVector(1.0f, 0.0f, 0.0f);
 	SecAxis2 = FVector(0.0f, 1.0f, 0.0f);
-
-	LinearBreakThreshold = 300.0f;
-	AngularBreakThreshold = 500.0f;
-
-	ProjectionLinearTolerance = 5.f; // Linear projection when error > 5 cm
-	ProjectionAngularTolerance = 180.f; // Angular projection when error > 180 degrees
 }
 
 void FConstraintInstance::SetDisableCollision(bool InDisableCollision)
 {
-	bDisableCollision = InDisableCollision;
+	ProfileInstance.bDisableCollision = InDisableCollision;
 #if WITH_PHYSX
 	ExecuteOnUnbrokenJointReadWrite([&] (PxD6Joint* Joint)
 	{
 		PxConstraintFlags Flags = Joint->getConstraintFlags();
-		if (bDisableCollision)
+		if (InDisableCollision)
 		{
 			Flags &= ~PxConstraintFlag::eCOLLISION_ENABLED;
 		}
@@ -467,7 +405,7 @@ bool FConstraintInstance::CreatePxJoint_AssumesLocked(physx::PxRigidActor* PActo
 	return true;
 }
 
-void FConstraintInstance::UpdateConstraintFlags_AssumesLocked()
+void FConstraintProfileProperties::UpdatePhysXConstraintFlags_AssumesLocked(PxD6Joint* Joint) const
 {
 	PxConstraintFlags Flags = PxConstraintFlags();
 
@@ -484,11 +422,11 @@ void FConstraintInstance::UpdateConstraintFlags_AssumesLocked()
 	{
 		Flags |= PxConstraintFlag::ePROJECTION;
 
-		ConstraintData->setProjectionLinearTolerance(ProjectionLinearTolerance);
-		ConstraintData->setProjectionAngularTolerance(FMath::DegreesToRadians(ProjectionAngularTolerance));
+		Joint->setProjectionLinearTolerance(ProjectionLinearTolerance);
+		Joint->setProjectionAngularTolerance(FMath::DegreesToRadians(ProjectionAngularTolerance));
 	}
 
-	ConstraintData->setConstraintFlags(Flags);
+	Joint->setConstraintFlags(Flags);
 }
 
 
@@ -548,35 +486,57 @@ void FConstraintInstance::InitConstraint(USceneComponent* Owner, FBodyInstance*
 	PxRigidActor* PActor2 = nullptr;
 	PxScene* PScene = GetPScene_LockFree(Body1, Body2);
 	SCOPED_SCENE_WRITE_LOCK(PScene);
-
-	const bool bValidConstraintSetup = PScene && GetPActors_AssumesLocked(Body1, Body2, &PActor1, &PActor2) && CreatePxJoint_AssumesLocked(PActor1, PActor2, PScene);
-	if (!bValidConstraintSetup)
 	{
-		return;
-	}
+		const bool bValidConstraintSetup = PScene && GetPActors_AssumesLocked(Body1, Body2, &PActor1, &PActor2) && CreatePxJoint_AssumesLocked(PActor1, PActor2, PScene);
+		if (!bValidConstraintSetup)
+		{
+			return;
+		}
 
-	// update mass
-	UpdateAverageMass_AssumesLocked(PActor1, PActor2);
-	
-	//flags and projection settings
-	UpdateConstraintFlags_AssumesLocked();
+		// update mass
+		UpdateAverageMass_AssumesLocked(PActor1, PActor2);
+
+		ProfileInstance.UpdatePhysX_AssumesLocked(ConstraintData, AverageMass, bScaleLinearLimits ? LastKnownScale : 1.f);
+
+		EnsureSleepingActorsStaySleeping_AssumesLocked(PActor1, PActor2);
+	}
 	
+#endif // WITH_PHYSX
+
+}
+
+#if WITH_PHYSX
+void FConstraintProfileProperties::UpdatePhysX_AssumesLocked(PxD6Joint* Joint, float AverageMass, float UseScale) const
+{
+	// flags and projection settings
+	UpdatePhysXConstraintFlags_AssumesLocked(Joint);
+
 	//limits
-	UpdateAngularLimit();
-	UpdateLinearLimit();
+	LinearLimit.UpdatePhysXLinearLimit_AssumesLocked(Joint, AverageMass, UseScale);
+	ConeLimit.UpdatePhysXConeLimit_AssumesLocked(Joint, AverageMass);
+	TwistLimit.UpdatePhysXTwistLimit_AssumesLocked(Joint, AverageMass);
 
 	//breakable
-	UpdateBreakable();
-	
+	UpdatePhysXBreakable_AssumesLocked(Joint);
+
 	//motors
-	SetLinearDriveParams(LinearDriveSpring, LinearDriveDamping, LinearDriveForceLimit);
-	SetAngularDriveParams(AngularDriveSpring, AngularDriveDamping, AngularDriveForceLimit);
-	UpdateDriveTarget();
-	
-	EnsureSleepingActorsStaySleeping_AssumesLocked(PActor1, PActor2);
-#endif // WITH_PHYSX
+	LinearDrive.UpdatePhysXLinearDrive_AssumesLocked(Joint);
+	AngularDrive.UpdatePhysXAngularDrive_AssumesLocked(Joint);
+	UpdatePhysXDriveTarget_AssumesLocked(Joint);
+
+}
+#endif
+
+#if WITH_PHYSX
+void FConstraintProfileProperties::UpdatePhysXDriveTarget_AssumesLocked(PxD6Joint* Joint) const
+{
+	const FQuat OrientationTargetQuat(AngularDrive.OrientationTarget);
+
+	Joint->setDrivePosition(PxTransform(U2PVector(LinearDrive.PositionTarget), U2PQuat(OrientationTargetQuat)));
+	Joint->setDriveVelocity(U2PVector(LinearDrive.VelocityTarget), U2PVector(RevolutionsToRads(AngularDrive.AngularVelocityTarget)));
 
 }
+#endif
 
 void FConstraintInstance::TermConstraint()
 {
@@ -615,7 +575,16 @@ bool FConstraintInstance::IsValidConstraintInstance() const
 #endif // WITH_PHYSX
 }
 
-
+void FConstraintInstance::CopyProfilePropertiesFrom(FConstraintProfileProperties& FromProperties)
+{
+	ProfileInstance = FromProperties;
+#if WITH_PHYSX
+	ExecuteOnUnbrokenJointReadWrite([&](PxD6Joint* Joint)
+	{
+		ProfileInstance.UpdatePhysX_AssumesLocked(Joint, AverageMass, bScaleLinearLimits ? LastKnownScale : 1.f);
+	});
+#endif
+}
 
 void FConstraintInstance::CopyConstraintGeometryFrom(const FConstraintInstance* FromInstance)
 {
@@ -811,146 +780,137 @@ void FConstraintInstance::GetConstraintForce(FVector& OutLinearForce, FVector& O
 /** Function for turning linear position drive on and off. */
 void FConstraintInstance::SetLinearPositionDrive(bool bEnableXDrive, bool bEnableYDrive, bool bEnableZDrive)
 {
-	bLinearXPositionDrive = bEnableXDrive;
-	bLinearYPositionDrive = bEnableYDrive;
-	bLinearZPositionDrive = bEnableZDrive;
-	bLinearPositionDrive = bEnableXDrive || bEnableYDrive || bEnableZDrive;
-	SetLinearDriveParams(LinearDriveSpring, LinearDriveDamping, LinearDriveForceLimit);
+	ProfileInstance.LinearDrive.SetLinearPositionDrive(bEnableXDrive, bEnableYDrive, bEnableZDrive);
+#if WITH_PHYSX
+	ExecuteOnUnbrokenJointReadWrite([this](PxD6Joint* Joint)
+	{
+		ProfileInstance.LinearDrive.UpdatePhysXLinearDrive_AssumesLocked(Joint);
+	});
+#endif
 }
 
-
 /** Function for turning linear velocity drive on and off. */
 void FConstraintInstance::SetLinearVelocityDrive(bool bEnableXDrive, bool bEnableYDrive, bool bEnableZDrive)
 {
-	bLinearXVelocityDrive = bEnableXDrive;
-	bLinearYVelocityDrive = bEnableYDrive;
-	bLinearZVelocityDrive = bEnableZDrive;
-	bLinearVelocityDrive = bEnableXDrive || bEnableYDrive || bEnableZDrive;
-	SetLinearDriveParams(LinearDriveSpring, LinearDriveDamping, LinearDriveForceLimit);
+	ProfileInstance.LinearDrive.SetLinearVelocityDrive(bEnableXDrive, bEnableYDrive, bEnableZDrive);
+#if WITH_PHYSX
+	ExecuteOnUnbrokenJointReadWrite([this](PxD6Joint* Joint)
+	{
+		ProfileInstance.LinearDrive.UpdatePhysXLinearDrive_AssumesLocked(Joint);
+	});
+#endif
 }
 
 /** Function for turning angular position drive on and off. */
 void FConstraintInstance::SetAngularPositionDrive(bool InEnableSwingDrive, bool InEnableTwistDrive)
 {
-	bEnableSwingDrive = InEnableSwingDrive;
-	bEnableTwistDrive = InEnableTwistDrive;
-	bAngularOrientationDrive = bEnableSwingDrive || bEnableTwistDrive;
-	SetAngularDriveParams(AngularDriveSpring, AngularDriveDamping, AngularDriveForceLimit);
+	ProfileInstance.AngularDrive.SetAngularPositionDrive(InEnableSwingDrive, InEnableTwistDrive);
+#if WITH_PHYSX
+	ExecuteOnUnbrokenJointReadWrite([this](PxD6Joint* Joint)
+	{
+		ProfileInstance.AngularDrive.UpdatePhysXAngularDrive_AssumesLocked(Joint);
+	});
+#endif
 }
 
 /** Function for turning angular velocity drive on and off. */
 void FConstraintInstance::SetAngularVelocityDrive(bool InEnableSwingDrive, bool InEnableTwistDrive)
 {
-	bEnableSwingDrive = InEnableSwingDrive;
-	bEnableTwistDrive = InEnableTwistDrive;
-	bAngularVelocityDrive = bEnableSwingDrive || bEnableTwistDrive;
-	SetAngularDriveParams(AngularDriveSpring, AngularDriveDamping, AngularDriveForceLimit);
+	ProfileInstance.AngularDrive.SetAngularVelocityDrive(InEnableSwingDrive, InEnableTwistDrive);
+#if WITH_PHYSX
+	ExecuteOnUnbrokenJointReadWrite([this](PxD6Joint* Joint)
+	{
+		ProfileInstance.AngularDrive.UpdatePhysXAngularDrive_AssumesLocked(Joint);
+	});
+#endif
 }
 
 /** Function for setting linear position target. */
 void FConstraintInstance::SetLinearPositionTarget(const FVector& InPosTarget)
 {
 	// If settings are the same, don't do anything.
-	if( LinearPositionTarget == InPosTarget )
+	if( ProfileInstance.LinearDrive.PositionTarget == InPosTarget )
 	{
 		return;
 	}
 
+	ProfileInstance.LinearDrive.PositionTarget = InPosTarget;
+
 #if WITH_PHYSX
-	if (PxD6Joint* Joint = ConstraintData)
+	ExecuteOnUnbrokenJointReadWrite([InPosTarget](PxD6Joint* Joint)
 	{
 		PxVec3 Pos = U2PVector(InPosTarget);
 		Joint->setDrivePosition(PxTransform(Pos, Joint->getDrivePosition().q));
-	}
+	});
 #endif
-
-	LinearPositionTarget = InPosTarget;
 }
 
 /** Function for setting linear velocity target. */
 void FConstraintInstance::SetLinearVelocityTarget(const FVector& InVelTarget)
 {
 	// If settings are the same, don't do anything.
-	if( LinearVelocityTarget == InVelTarget )
+	if (ProfileInstance.LinearDrive.VelocityTarget == InVelTarget)
 	{
 		return;
 	}
 
+	ProfileInstance.LinearDrive.VelocityTarget = InVelTarget;
+
 #if WITH_PHYSX
-	if (PxD6Joint* Joint = ConstraintData)
+	ExecuteOnUnbrokenJointReadWrite([InVelTarget](PxD6Joint* Joint)
 	{
 		PxVec3 CurrentLinearVel, CurrentAngVel;
 		Joint->getDriveVelocity(CurrentLinearVel, CurrentAngVel);
 
 		PxVec3 NewLinearVel = U2PVector(InVelTarget);
 		Joint->setDriveVelocity(NewLinearVel, CurrentAngVel);
-	}
+	});
 #endif
-	LinearVelocityTarget = InVelTarget;
 }
 
 /** Function for setting linear motor parameters. */
 void FConstraintInstance::SetLinearDriveParams(float InSpring, float InDamping, float InForceLimit)
 {
+	ProfileInstance.LinearDrive.SetDriveParams(InSpring, InDamping, InForceLimit);
+
 #if WITH_PHYSX
-	if (bLinearPositionDrive || bLinearVelocityDrive)
+	ExecuteOnUnbrokenJointReadWrite([this] (PxD6Joint* Joint)
 	{
-		ExecuteOnUnbrokenJointReadWrite([&] (PxD6Joint* Joint)
-		{
-			// X-Axis linear drive
-			const float DriveSpringX = bLinearPositionDrive && bLinearXPositionDrive ? InSpring : 0.0f;
-			const float DriveDampingX = bLinearVelocityDrive && bLinearXVelocityDrive ? InDamping : 0.0f;
-			const float LinearForceLimit = InForceLimit > 0.f ? InForceLimit : PX_MAX_F32;
-			Joint->setDrive(PxD6Drive::eX, PxD6JointDrive(DriveSpringX, DriveDampingX, LinearForceLimit, bIsAccelerationDrive));
-
-			// Y-Axis linear drive
-			const float DriveSpringY = bLinearPositionDrive && bLinearYPositionDrive ? InSpring : 0.0f;
-			const float DriveDampingY = bLinearVelocityDrive && bLinearYVelocityDrive ? InDamping : 0.0f;
-			Joint->setDrive(PxD6Drive::eY, PxD6JointDrive(DriveSpringY, DriveDampingY, LinearForceLimit, bIsAccelerationDrive));
-
-			// Z-Axis linear drive
-			const float DriveSpringZ = bLinearPositionDrive && bLinearZPositionDrive ? InSpring : 0.0f;
-			const float DriveDampingZ = bLinearVelocityDrive && bLinearZVelocityDrive ? InDamping : 0.0f;
-			Joint->setDrive(PxD6Drive::eZ, PxD6JointDrive(DriveSpringZ, DriveDampingZ, LinearForceLimit, bIsAccelerationDrive));
-		});
-	}
+		ProfileInstance.LinearDrive.UpdatePhysXLinearDrive_AssumesLocked(Joint);
+	});
 #endif
-
-	LinearDriveSpring = InSpring;
-	LinearDriveDamping = InDamping;
-	LinearDriveForceLimit = InForceLimit;
 }
 
 /** Function for setting target angular position. */
-void FConstraintInstance::SetAngularOrientationTarget(const FQuat& InPosTarget)
+void FConstraintInstance::SetAngularOrientationTarget(const FQuat& InOrientationTarget)
 {
-	FRotator OrientationTargetRot(InPosTarget);
+	FRotator OrientationTargetRot(InOrientationTarget);
 
 	// If settings are the same, don't do anything.
-	if( AngularOrientationTarget == OrientationTargetRot )
+	if( ProfileInstance.AngularDrive.OrientationTarget == OrientationTargetRot )
 	{
 		return;
 	}
-	
+
+	ProfileInstance.AngularDrive.OrientationTarget = OrientationTargetRot;
+
 #if WITH_PHYSX
-	if (PxD6Joint* Joint = ConstraintData)
+	ExecuteOnUnbrokenJointReadWrite([InOrientationTarget](PxD6Joint* Joint)
 	{
-		PxQuat Quat = U2PQuat(InPosTarget);
+		PxQuat Quat = U2PQuat(InOrientationTarget);
 		Joint->setDrivePosition(PxTransform(Joint->getDrivePosition().p, Quat));
-	}
+	});
 #endif
-
-	AngularOrientationTarget = OrientationTargetRot;
 }
 
 float FConstraintInstance::GetCurrentSwing1() const
 {
 	float Swing1 = 0.f;
 #if WITH_PHYSX
-	if (PxD6Joint* Joint = ConstraintData)
+	ExecuteOnUnbrokenJointReadOnly([&Swing1](const PxD6Joint* Joint)
 	{
 		Swing1 = Joint->getSwingZAngle();
-	}
+	});
 #endif
 
 	return Swing1;
@@ -960,10 +920,10 @@ float FConstraintInstance::GetCurrentSwing2() const
 {
 	float Swing2 = 0.f;
 #if WITH_PHYSX
-	if (PxD6Joint* Joint = ConstraintData)
+	ExecuteOnUnbrokenJointReadOnly([&Swing2](const PxD6Joint* Joint)
 	{
 		Swing2 = Joint->getSwingYAngle();
-	}
+	});
 #endif
 
 	return Swing2;
@@ -973,10 +933,10 @@ float FConstraintInstance::GetCurrentTwist() const
 {
 	float Twist = 0.f;
 #if WITH_PHYSX
-	if (PxD6Joint* Joint = ConstraintData)
+	ExecuteOnUnbrokenJointReadOnly([&Twist](const PxD6Joint* Joint)
 	{
 		Twist = Joint->getTwist();
-	}
+	});
 #endif
 
 	return Twist;
@@ -987,101 +947,85 @@ float FConstraintInstance::GetCurrentTwist() const
 void FConstraintInstance::SetAngularVelocityTarget(const FVector& InVelTarget)
 {
 	// If settings are the same, don't do anything.
-	if( AngularVelocityTarget == InVelTarget )
+	if( ProfileInstance.AngularDrive.AngularVelocityTarget == InVelTarget )
 	{
 		return;
 	}
 
+	ProfileInstance.AngularDrive.AngularVelocityTarget = InVelTarget;
+
 #if WITH_PHYSX
-	if (PxD6Joint* Joint = ConstraintData)
+	ExecuteOnUnbrokenJointReadWrite([InVelTarget](PxD6Joint* Joint)
 	{
 		PxVec3 CurrentLinearVel, CurrentAngVel;
 		Joint->getDriveVelocity(CurrentLinearVel, CurrentAngVel);
 
 		PxVec3 AngVel = U2PVector(RevolutionsToRads(InVelTarget));
 		Joint->setDriveVelocity(CurrentLinearVel, AngVel);
-	}
+	});
 #endif
-
-	AngularVelocityTarget = InVelTarget;
 }
 
 /** Function for setting angular motor parameters. */
 void FConstraintInstance::SetAngularDriveParams(float InSpring, float InDamping, float InForceLimit)
 {
+	ProfileInstance.AngularDrive.SetDriveParams(InSpring, InDamping, InForceLimit);
+
 #if WITH_PHYSX
 	ExecuteOnUnbrokenJointReadWrite([&] (PxD6Joint* Joint)
 	{
-		const float AngularForceLimit = InForceLimit > 0.0f ? InForceLimit : PX_MAX_F32;
-		const float DriveSpring = bAngularOrientationDrive ? InSpring : 0.0f;
-		const float DriveDamping = bAngularVelocityDrive ? InDamping : 0.0f;
-
-		if (AngularDriveMode == EAngularDriveMode::SLERP)
-		{
-			Joint->setDrive(PxD6Drive::eTWIST, PxD6JointDrive());
-			Joint->setDrive(PxD6Drive::eSWING, PxD6JointDrive());
-			Joint->setDrive(PxD6Drive::eSLERP, (bEnableSwingDrive || bEnableTwistDrive) ? PxD6JointDrive(DriveSpring, DriveDamping, AngularForceLimit, bIsAccelerationDrive) : PxD6JointDrive());				
-		}
-		else
-		{
-			Joint->setDrive(PxD6Drive::eTWIST, bEnableTwistDrive ? PxD6JointDrive(DriveSpring, DriveDamping, AngularForceLimit, bIsAccelerationDrive) : PxD6JointDrive());
-			Joint->setDrive(PxD6Drive::eSWING, bEnableSwingDrive ? PxD6JointDrive(DriveSpring, DriveDamping, AngularForceLimit, bIsAccelerationDrive) : PxD6JointDrive());
-			Joint->setDrive(PxD6Drive::eSLERP, PxD6JointDrive());
-		}
+		ProfileInstance.AngularDrive.UpdatePhysXAngularDrive_AssumesLocked(Joint);
 	});
 #endif
-
-	AngularDriveSpring = InSpring;
-	AngularDriveDamping = InDamping;
-	AngularDriveForceLimit = InForceLimit;
 }
 
+
 /** Scale Angular Limit Constraints (as defined in RB_ConstraintSetup) */
 void FConstraintInstance::SetAngularDOFLimitScale(float InSwing1LimitScale, float InSwing2LimitScale, float InTwistLimitScale)
 {
 #if WITH_PHYSX
 	ExecuteOnUnbrokenJointReadWrite([&] (PxD6Joint* Joint)
 	{
-		if ( AngularSwing1Motion == ACM_Limited || AngularSwing2Motion == ACM_Limited )
+		if ( ProfileInstance.ConeLimit.Swing1Motion == ACM_Limited || ProfileInstance.ConeLimit.Swing2Motion == ACM_Limited )
 		{
 			// PhysX swing directions are different from Unreal's - so change here.
-			if (AngularSwing1Motion == ACM_Limited)
+			if (ProfileInstance.ConeLimit.Swing1Motion == ACM_Limited)
 			{
 				Joint->setMotion(PxD6Axis::eSWING2, PxD6Motion::eLIMITED);
 			}
 
-			if (AngularSwing2Motion == ACM_Limited)
+			if (ProfileInstance.ConeLimit.Swing2Motion == ACM_Limited)
 			{
 				Joint->setMotion(PxD6Axis::eSWING1, PxD6Motion::eLIMITED);
 			}
 		
-			//The limite values need to be clamped so it will be valid in PhysX
-			PxReal ZLimitAngle = FMath::ClampAngle(Swing1LimitAngle * InSwing1LimitScale, KINDA_SMALL_NUMBER, 179.9999f) * (PI/180.0f);
-			PxReal YLimitAngle = FMath::ClampAngle(Swing2LimitAngle * InSwing2LimitScale, KINDA_SMALL_NUMBER, 179.9999f) * (PI/180.0f);
-			PxReal LimitContractDistance =  1.f * (PI/180.0f);
+			//The limit values need to be clamped so it will be valid in PhysX
+			PxReal ZLimitAngle = FMath::ClampAngle(ProfileInstance.ConeLimit.Swing1LimitDegrees * InSwing1LimitScale, KINDA_SMALL_NUMBER, 179.9999f) * (PI/180.0f);
+			PxReal YLimitAngle = FMath::ClampAngle(ProfileInstance.ConeLimit.Swing2LimitDegrees * InSwing2LimitScale, KINDA_SMALL_NUMBER, 179.9999f) * (PI/180.0f);
+			PxReal LimitContactDistance =  FMath::DegreesToRadians(FMath::Max(1.f, ProfileInstance.ConeLimit.ContactDistance * FMath::Min(InSwing1LimitScale, InSwing2LimitScale)));
 
-			Joint->setSwingLimit(PxJointLimitCone(YLimitAngle, ZLimitAngle, LimitContractDistance));
+			Joint->setSwingLimit(PxJointLimitCone(YLimitAngle, ZLimitAngle, LimitContactDistance));
 		}
 
-		if ( AngularSwing1Motion == ACM_Locked )
+		if ( ProfileInstance.ConeLimit.Swing1Motion  == ACM_Locked )
 		{
 			Joint->setMotion(PxD6Axis::eSWING2, PxD6Motion::eLOCKED);
 		}
 
-		if ( AngularSwing2Motion == ACM_Locked )
+		if ( ProfileInstance.ConeLimit.Swing2Motion  == ACM_Locked )
 		{
 			Joint->setMotion(PxD6Axis::eSWING1, PxD6Motion::eLOCKED);
 		}
 
-		if ( AngularTwistMotion == ACM_Limited )
+		if ( ProfileInstance.TwistLimit.TwistMotion == ACM_Limited )
 		{
 			Joint->setMotion(PxD6Axis::eTWIST, PxD6Motion::eLIMITED);
-			const float TwistLimitRad	= TwistLimitAngle * InTwistLimitScale * (PI/180.0f);
-			PxReal LimitContractDistance =  1.f * (PI/180.0f);
+			const float TwistLimitRad	= ProfileInstance.TwistLimit.TwistLimitDegrees * InTwistLimitScale * (PI/180.0f);
+			PxReal LimitContactDistance =  FMath::DegreesToRadians(FMath::Max(1.f, ProfileInstance.ConeLimit.ContactDistance * InTwistLimitScale));
 
-			Joint->setTwistLimit(PxJointAngularLimitPair(-TwistLimitRad, TwistLimitRad, LimitContractDistance)); 
+			Joint->setTwistLimit(PxJointAngularLimitPair(-TwistLimitRad, TwistLimitRad, LimitContactDistance)); 
 		}
-		else if ( AngularTwistMotion == ACM_Locked )
+		else if ( ProfileInstance.TwistLimit.TwistMotion == ACM_Locked )
 		{
 			Joint->setMotion(PxD6Axis::eTWIST, PxD6Motion::eLOCKED);
 		}
@@ -1093,156 +1037,127 @@ void FConstraintInstance::SetAngularDOFLimitScale(float InSwing1LimitScale, floa
 /** Allows you to dynamically change the size of the linear limit 'sphere'. */
 void FConstraintInstance::SetLinearLimitSize(float NewLimitSize)
 {
+	//TODO: Is this supposed to be scaling the linear limit? The code just sets it directly.
 #if WITH_PHYSX
 	ExecuteOnUnbrokenJointReadWrite([&] (PxD6Joint* Joint)
 	{
-		PxReal LimitContractDistance =  1.f * (PI/180.0f);
-		Joint->setLinearLimit(PxJointLinearLimit(GPhysXSDK->getTolerancesScale(), NewLimitSize, LimitContractDistance * GPhysXSDK->getTolerancesScale().length)); // LOC_MOD33 need to scale the contactDistance if not using its default value
+		PxReal LimitContactDistance = 1.f * (PI / 180.0f);
+		Joint->setLinearLimit(PxJointLinearLimit(GPhysXSDK->getTolerancesScale(), NewLimitSize, LimitContactDistance * GPhysXSDK->getTolerancesScale().length)); // LOC_MOD33 need to scale the contactDistance if not using its default value
 	});
 #endif
 }
 
-
-void FConstraintInstance::ConfigureAsHinge(bool bOverwriteLimits)
+bool FConstraintInstance::Serialize(FArchive& Ar)
 {
-	LinearXMotion = LCM_Locked;
-	LinearYMotion = LCM_Locked;
-	LinearZMotion = LCM_Locked;
-	AngularSwing1Motion = ACM_Locked;
-	AngularSwing2Motion = ACM_Locked;
-	AngularTwistMotion = ACM_Free;
-
-	if (bOverwriteLimits)
-	{
-		Swing1LimitAngle = 0.0f;
-		Swing2LimitAngle = 0.0f;
-		TwistLimitAngle = 45.0f;
-	}
-}
-
-void FConstraintInstance::ConfigureAsPrismatic(bool bOverwriteLimits)
-{
-	LinearXMotion = LCM_Free;
-	LinearYMotion = LCM_Locked;
-	LinearZMotion = LCM_Locked;
-	AngularSwing1Motion = ACM_Locked;
-	AngularSwing2Motion = ACM_Locked;
-	AngularTwistMotion = ACM_Locked;
-
-	if (bOverwriteLimits)
-	{
-		Swing1LimitAngle = 0.0f;
-		Swing2LimitAngle = 0.0f;
-		TwistLimitAngle = 0.0f;
-	}
-}
-
-void FConstraintInstance::ConfigureAsSkelJoint(bool bOverwriteLimits)
-{
-	LinearXMotion = LCM_Locked;
-	LinearYMotion = LCM_Locked;
-	LinearZMotion = LCM_Locked;
-	AngularSwing1Motion = ACM_Limited;
-	AngularSwing2Motion = ACM_Limited;
-	AngularTwistMotion = ACM_Limited;
-
-	if (bOverwriteLimits)
-	{
-		Swing1LimitAngle = 45.0f;
-		Swing2LimitAngle = 45.0f;
-		TwistLimitAngle = 15.0f;
-	}
-}
-
-void FConstraintInstance::ConfigureAsBS(bool bOverwriteLimits)
-{
-	LinearXMotion = LCM_Locked;
-	LinearYMotion = LCM_Locked;
-	LinearZMotion = LCM_Locked;
-	AngularSwing1Motion = ACM_Free;
-	AngularSwing2Motion = ACM_Free;
-	AngularTwistMotion = ACM_Free;
-	
-	if (bOverwriteLimits)
-	{
-		Swing1LimitAngle = 0.0f;
-		Swing2LimitAngle = 0.0f;
-		TwistLimitAngle = 0.0f;
-	}
-}
-
-bool FConstraintInstance::IsHinge() const
-{
-	int32 AngularDOF = 0;
-	AngularDOF += AngularSwing1Motion != ACM_Locked ? 1 : 0;
-	AngularDOF += AngularSwing2Motion != ACM_Locked ? 1 : 0;
-	AngularDOF += AngularTwistMotion != ACM_Locked ? 1 : 0;
-
-	return	LinearXMotion		== LCM_Locked &&
-			LinearYMotion		== LCM_Locked &&
-			LinearZMotion		== LCM_Locked &&
-			AngularDOF == 1;
-}
-
-bool FConstraintInstance::IsPrismatic() const
-{
-	int32 LinearDOF = 0;
-	LinearDOF += LinearXMotion != LCM_Locked ? 1 : 0;
-	LinearDOF += LinearYMotion != LCM_Locked ? 1 : 0;
-	LinearDOF += LinearZMotion != LCM_Locked ? 1 : 0;
-
-	return	LinearDOF > 0 &&
-			AngularSwing1Motion == ACM_Locked &&
-			AngularSwing2Motion == ACM_Locked &&
-			AngularTwistMotion	== ACM_Locked;
-}
-
-bool FConstraintInstance::IsSkelJoint() const
-{
-	return	LinearXMotion		== LCM_Locked &&
-			LinearYMotion		== LCM_Locked &&
-			LinearZMotion		== LCM_Locked &&
-			AngularSwing1Motion == ACM_Limited &&
-			AngularSwing2Motion == ACM_Limited &&
-			AngularTwistMotion	== ACM_Limited;
-}
-
-bool FConstraintInstance::IsBallAndSocket() const
-{
-	int32 AngularDOF = 0;
-	AngularDOF += AngularSwing1Motion != ACM_Locked ? 1 : 0;
-	AngularDOF += AngularSwing2Motion != ACM_Locked ? 1 : 0;
-	AngularDOF += AngularTwistMotion != ACM_Locked ? 1 : 0;
-
-	return	LinearXMotion		== LCM_Locked &&
-			LinearYMotion		== LCM_Locked &&
-			LinearZMotion		== LCM_Locked &&
-			AngularDOF > 1;
+	Ar.UsingCustomVersion(FFrameworkObjectVersion::GUID);
+	return false;	//We only have this function to mark custom GUID. Still want serialize tagged properties
 }
 
 void FConstraintInstance::PostSerialize(const FArchive& Ar)
 {
+#if WITH_EDITORONLY_DATA
 	if (Ar.IsLoading() && Ar.UE4Ver() < VER_UE4_FIXUP_STIFFNESS_AND_DAMPING_SCALE)
 	{
-		LinearLimitStiffness	/= CVarConstraintStiffnessScale.GetValueOnGameThread();
-		SwingLimitStiffness		/= CVarConstraintStiffnessScale.GetValueOnGameThread();
-		TwistLimitStiffness		/= CVarConstraintStiffnessScale.GetValueOnGameThread();
-		LinearLimitDamping		/=  CVarConstraintDampingScale.GetValueOnGameThread();
-		SwingLimitDamping		/=  CVarConstraintDampingScale.GetValueOnGameThread();
-		TwistLimitDamping		/=  CVarConstraintDampingScale.GetValueOnGameThread();
+		LinearLimitStiffness_DEPRECATED		/= CVarConstraintStiffnessScale.GetValueOnGameThread();
+		SwingLimitStiffness_DEPRECATED		/= CVarConstraintStiffnessScale.GetValueOnGameThread();
+		TwistLimitStiffness_DEPRECATED		/= CVarConstraintStiffnessScale.GetValueOnGameThread();
+		LinearLimitDamping_DEPRECATED		/=  CVarConstraintDampingScale.GetValueOnGameThread();
+		SwingLimitDamping_DEPRECATED		/=  CVarConstraintDampingScale.GetValueOnGameThread();
+		TwistLimitDamping_DEPRECATED		/=  CVarConstraintDampingScale.GetValueOnGameThread();
 	}
 
 	if (Ar.IsLoading() && Ar.UE4Ver() < VER_UE4_FIXUP_MOTOR_UNITS)
 	{
-		AngularVelocityTarget *= 1.f / (2.f * PI);	//we want to use revolutions per second - old system was using radians directly
+		AngularVelocityTarget_DEPRECATED *= 1.f / (2.f * PI);	//we want to use revolutions per second - old system was using radians directly
 	}
 
 	if (Ar.IsLoading() && Ar.UE4Ver() < VER_UE4_CONSTRAINT_INSTANCE_MOTOR_FLAGS)
 	{
-		bLinearXVelocityDrive = LinearVelocityTarget.X != 0.f;
-		bLinearYVelocityDrive = LinearVelocityTarget.Y != 0.f;
-		bLinearZVelocityDrive = LinearVelocityTarget.Z != 0.f;
+		bLinearXVelocityDrive_DEPRECATED = LinearVelocityTarget_DEPRECATED.X != 0.f;
+		bLinearYVelocityDrive_DEPRECATED = LinearVelocityTarget_DEPRECATED.Y != 0.f;
+		bLinearZVelocityDrive_DEPRECATED = LinearVelocityTarget_DEPRECATED.Z != 0.f;
 	}
+	
+	if (Ar.IsLoading() && Ar.CustomVer(FFrameworkObjectVersion::GUID) < FFrameworkObjectVersion::ConstraintInstanceBehaviorParameters)
+	{
+		//Need to move all the deprecated properties into the new profile struct
+		ProfileInstance.bDisableCollision = bDisableCollision_DEPRECATED;
+		ProfileInstance.bEnableProjection = bEnableProjection_DEPRECATED;
+		ProfileInstance.ProjectionLinearTolerance = ProjectionLinearTolerance_DEPRECATED;
+		ProfileInstance.ProjectionAngularTolerance = ProjectionAngularTolerance_DEPRECATED;
+		ProfileInstance.LinearLimit.XMotion = LinearXMotion_DEPRECATED;
+		ProfileInstance.LinearLimit.YMotion = LinearYMotion_DEPRECATED;
+		ProfileInstance.LinearLimit.ZMotion = LinearZMotion_DEPRECATED;
+		ProfileInstance.LinearLimit.Limit = LinearLimitSize_DEPRECATED;
+		ProfileInstance.LinearLimit.bSoftConstraint = bLinearLimitSoft_DEPRECATED;
+		ProfileInstance.LinearLimit.Stiffness = LinearLimitStiffness_DEPRECATED;
+		ProfileInstance.LinearLimit.Damping = LinearLimitDamping_DEPRECATED;
+		ProfileInstance.bLinearBreakable = bLinearBreakable_DEPRECATED;
+		ProfileInstance.LinearBreakThreshold = LinearBreakThreshold_DEPRECATED;
+		ProfileInstance.ConeLimit.Swing1Motion = AngularSwing1Motion_DEPRECATED;
+		ProfileInstance.TwistLimit.TwistMotion = AngularTwistMotion_DEPRECATED;
+		ProfileInstance.ConeLimit.Swing2Motion = AngularSwing2Motion_DEPRECATED;
+		ProfileInstance.ConeLimit.bSoftConstraint = bSwingLimitSoft_DEPRECATED;
+		ProfileInstance.TwistLimit.bSoftConstraint = bTwistLimitSoft_DEPRECATED;
+		ProfileInstance.ConeLimit.Swing1LimitDegrees = Swing1LimitAngle_DEPRECATED;
+		ProfileInstance.TwistLimit.TwistLimitDegrees = TwistLimitAngle_DEPRECATED;
+		ProfileInstance.ConeLimit.Swing2LimitDegrees = Swing2LimitAngle_DEPRECATED;
+		ProfileInstance.ConeLimit.Stiffness = SwingLimitStiffness_DEPRECATED;
+		ProfileInstance.ConeLimit.Damping = SwingLimitDamping_DEPRECATED;
+		ProfileInstance.TwistLimit.Stiffness = TwistLimitStiffness_DEPRECATED;
+		ProfileInstance.TwistLimit.Damping = TwistLimitDamping_DEPRECATED;
+		ProfileInstance.bAngularBreakable = bAngularBreakable_DEPRECATED;
+		ProfileInstance.AngularBreakThreshold = AngularBreakThreshold_DEPRECATED;
+
+		//we no longer have a single control for all linear axes. If it was off we ensure all individual drives are off. If it's on we just leave things alone.
+		//This loses a bit of info, but the ability to toggle drives on and off at runtime was very obfuscated so hopefuly this doesn't hurt too many people. They can still toggle individual drives on and off
+		ProfileInstance.LinearDrive.XDrive.bEnablePositionDrive = bLinearXPositionDrive_DEPRECATED && bLinearPositionDrive_DEPRECATED;
+		ProfileInstance.LinearDrive.XDrive.bEnableVelocityDrive = bLinearXVelocityDrive_DEPRECATED && bLinearVelocityDrive_DEPRECATED;
+		ProfileInstance.LinearDrive.YDrive.bEnablePositionDrive = bLinearYPositionDrive_DEPRECATED && bLinearPositionDrive_DEPRECATED;
+		ProfileInstance.LinearDrive.YDrive.bEnableVelocityDrive = bLinearYVelocityDrive_DEPRECATED && bLinearVelocityDrive_DEPRECATED;
+		ProfileInstance.LinearDrive.ZDrive.bEnablePositionDrive = bLinearZPositionDrive_DEPRECATED && bLinearPositionDrive_DEPRECATED;
+		ProfileInstance.LinearDrive.ZDrive.bEnableVelocityDrive = bLinearZVelocityDrive_DEPRECATED && bLinearVelocityDrive_DEPRECATED;
+		
+		ProfileInstance.LinearDrive.PositionTarget = LinearPositionTarget_DEPRECATED;
+		ProfileInstance.LinearDrive.VelocityTarget = LinearVelocityTarget_DEPRECATED;
+
+		//Linear drives now set settings per axis so duplicate old data
+		ProfileInstance.LinearDrive.XDrive.Stiffness = LinearDriveSpring_DEPRECATED;
+		ProfileInstance.LinearDrive.YDrive.Stiffness = LinearDriveSpring_DEPRECATED;
+		ProfileInstance.LinearDrive.ZDrive.Stiffness = LinearDriveSpring_DEPRECATED;
+		ProfileInstance.LinearDrive.XDrive.Damping = LinearDriveDamping_DEPRECATED;
+		ProfileInstance.LinearDrive.YDrive.Damping = LinearDriveDamping_DEPRECATED;
+		ProfileInstance.LinearDrive.ZDrive.Damping = LinearDriveDamping_DEPRECATED;
+		ProfileInstance.LinearDrive.XDrive.MaxForce = LinearDriveForceLimit_DEPRECATED;
+		ProfileInstance.LinearDrive.YDrive.MaxForce = LinearDriveForceLimit_DEPRECATED;
+		ProfileInstance.LinearDrive.ZDrive.MaxForce = LinearDriveForceLimit_DEPRECATED;
+
+		//We now expose twist swing and slerp drive directly. In the old system you had a single switch, but then there was also special switches for disabling twist and swing
+		//Technically someone COULD disable these, but they are not exposed in editor so it seems very unlikely. So if they are true and angular orientation is false we override it
+		ProfileInstance.AngularDrive.SwingDrive.bEnablePositionDrive = bEnableSwingDrive_DEPRECATED && bAngularOrientationDrive_DEPRECATED;
+		ProfileInstance.AngularDrive.SwingDrive.bEnableVelocityDrive = bEnableSwingDrive_DEPRECATED && bAngularVelocityDrive_DEPRECATED;
+		ProfileInstance.AngularDrive.TwistDrive.bEnablePositionDrive = bEnableTwistDrive_DEPRECATED && bAngularOrientationDrive_DEPRECATED;
+		ProfileInstance.AngularDrive.TwistDrive.bEnableVelocityDrive = bEnableTwistDrive_DEPRECATED && bAngularVelocityDrive_DEPRECATED;
+		ProfileInstance.AngularDrive.SlerpDrive.bEnablePositionDrive = bAngularOrientationDrive_DEPRECATED;
+		ProfileInstance.AngularDrive.SlerpDrive.bEnableVelocityDrive = bAngularVelocityDrive_DEPRECATED;
+
+		ProfileInstance.AngularDrive.AngularDriveMode = AngularDriveMode_DEPRECATED;
+		ProfileInstance.AngularDrive.OrientationTarget = AngularOrientationTarget_DEPRECATED;
+		ProfileInstance.AngularDrive.AngularVelocityTarget = AngularVelocityTarget_DEPRECATED;
+
+		//duplicate drive spring data into all 3 drives
+		ProfileInstance.AngularDrive.SwingDrive.Stiffness = AngularDriveSpring_DEPRECATED;
+		ProfileInstance.AngularDrive.TwistDrive.Stiffness = AngularDriveSpring_DEPRECATED;
+		ProfileInstance.AngularDrive.SlerpDrive.Stiffness = AngularDriveSpring_DEPRECATED;
+		ProfileInstance.AngularDrive.SwingDrive.Damping = AngularDriveDamping_DEPRECATED;
+		ProfileInstance.AngularDrive.TwistDrive.Damping = AngularDriveDamping_DEPRECATED;
+		ProfileInstance.AngularDrive.SlerpDrive.Damping = AngularDriveDamping_DEPRECATED;
+		ProfileInstance.AngularDrive.SwingDrive.MaxForce = AngularDriveForceLimit_DEPRECATED;
+		ProfileInstance.AngularDrive.TwistDrive.MaxForce = AngularDriveForceLimit_DEPRECATED;
+		ProfileInstance.AngularDrive.SlerpDrive.MaxForce = AngularDriveForceLimit_DEPRECATED;
+
+	}
+#endif
 }
 
 void FConstraintInstance::OnConstraintBroken()
@@ -1261,54 +1176,6 @@ void FConstraintInstance::OnConstraintBroken()
 
 }
 
-void FConstraintInstance::SetLinearXLimit(ELinearConstraintMotion XMotion, float InLinearLimitSize)
-{
-	LinearXMotion = XMotion;
-	LinearLimitSize = InLinearLimitSize;
-
-	UpdateLinearLimit();
-}
-
-void FConstraintInstance::SetLinearYLimit(ELinearConstraintMotion YMotion, float InLinearLimitSize)
-{
-	LinearYMotion = YMotion;
-	LinearLimitSize = InLinearLimitSize;
-	
-	UpdateLinearLimit();
-}
-
-void FConstraintInstance::SetLinearZLimit(ELinearConstraintMotion ZMotion, float InLinearLimitSize)
-{
-	LinearZMotion = ZMotion;
-	LinearLimitSize = InLinearLimitSize;
-
-	UpdateLinearLimit();
-}
-
-void FConstraintInstance::SetAngularSwing1Limit(EAngularConstraintMotion Swing1Motion, float InSwing1LimitAngle)
-{
-	AngularSwing1Motion = Swing1Motion;
-	Swing1LimitAngle = InSwing1LimitAngle;
-	
-	UpdateAngularLimit();
-}
-
-void FConstraintInstance::SetAngularSwing2Limit(EAngularConstraintMotion Swing2Motion, float InSwing2LimitAngle)
-{
-	AngularSwing2Motion = Swing2Motion;
-	Swing2LimitAngle = InSwing2LimitAngle;
-
-	UpdateAngularLimit();
-}
-
-void FConstraintInstance::SetAngularTwistLimit(EAngularConstraintMotion TwistMotion, float InTwistLimitAngle)
-{
-	AngularTwistMotion = TwistMotion;
-	TwistLimitAngle = InTwistLimitAngle;
-
-	UpdateAngularLimit();
-}
-
 //Hacks to easily get zeroed memory for special case when we don't use GC
 void FConstraintInstance::Free(FConstraintInstance * Ptr)
 {
@@ -1324,14 +1191,16 @@ FConstraintInstance * FConstraintInstance::Alloc()
 
 void FConstraintInstance::EnableProjection()
 {
+	ProfileInstance.bEnableProjection = true;
 	SCOPED_SCENE_WRITE_LOCK(ConstraintData->getScene());
-	ConstraintData->setProjectionLinearTolerance(ProjectionLinearTolerance);
-	ConstraintData->setProjectionAngularTolerance(ProjectionAngularTolerance);
+	ConstraintData->setProjectionLinearTolerance(ProfileInstance.ProjectionLinearTolerance);
+	ConstraintData->setProjectionAngularTolerance(ProfileInstance.ProjectionAngularTolerance);
 	ConstraintData->setConstraintFlag(PxConstraintFlag::ePROJECTION, true);
 }
 
 void FConstraintInstance::DisableProjection()
 {
+	ProfileInstance.bEnableProjection = false;
 	SCOPED_SCENE_WRITE_LOCK(ConstraintData->getScene());
 	ConstraintData->setConstraintFlag(PxConstraintFlag::ePROJECTION, false);
 }
diff --git a/Engine/Source/Runtime/Engine/Private/PhysicsEngine/ConstraintTypes.cpp b/Engine/Source/Runtime/Engine/Private/PhysicsEngine/ConstraintTypes.cpp
new file mode 100644
--- /dev/null
+++ b/Engine/Source/Runtime/Engine/Private/PhysicsEngine/ConstraintTypes.cpp
@@ -0,0 +1,168 @@
+// Copyright 1998-2016 Epic Games, Inc. All Rights Reserved.
+
+#include "EnginePrivate.h"
+#include "PhysicsPublic.h"
+#include "PhysXSupport.h"
+#include "PhysicsEngine/ConstraintInstance.h"
+
+extern TAutoConsoleVariable<float> CVarConstraintDampingScale;
+extern TAutoConsoleVariable<float> CVarConstraintStiffnessScale;
+
+#if WITH_PHYSX
+
+/** Util for converting from UE motion enum to physx motion enum */
+PxD6Motion::Enum U2PAngularMotion(EAngularConstraintMotion InMotion)
+{
+	switch (InMotion)
+	{
+	case EAngularConstraintMotion::ACM_Free: return PxD6Motion::eFREE;
+	case EAngularConstraintMotion::ACM_Limited: return PxD6Motion::eLIMITED;
+	case EAngularConstraintMotion::ACM_Locked: return PxD6Motion::eLOCKED;
+	default: check(0);	//unsupported motion type
+	}
+
+	return PxD6Motion::eFREE;
+}
+
+/** Util for converting from UE motion enum to physx motion enum */
+PxD6Motion::Enum U2PLinearMotion(ELinearConstraintMotion InMotion)
+{
+	switch (InMotion)
+	{
+	case ELinearConstraintMotion::LCM_Free: return PxD6Motion::eFREE;
+	case ELinearConstraintMotion::LCM_Limited: return PxD6Motion::eLIMITED;
+	case ELinearConstraintMotion::LCM_Locked: return PxD6Motion::eLOCKED;
+	default: check(0);	//unsupported motion type
+	}
+
+	return PxD6Motion::eFREE;
+}
+
+
+/** Util for setting soft limit params */
+void SetSoftLimitParams_AssumesLocked(PxJointLimitParameters* PLimit, bool bSoft, float Spring, float Damping)
+{
+	if(bSoft)
+	{
+		PLimit->stiffness = Spring * CVarConstraintStiffnessScale.GetValueOnGameThread();
+		PLimit->damping = Damping * CVarConstraintDampingScale.GetValueOnGameThread();
+	}
+}
+
+/** Util for setting linear movement for an axis */
+template <PxD6Axis::Enum PAxis>
+void SetLinearMovement_AssumesLocked(PxD6Joint* PD6Joint, ELinearConstraintMotion Motion, bool bLockLimitSize)
+{
+	if(Motion == LCM_Locked || (Motion == LCM_Limited && bLockLimitSize))
+	{
+		PD6Joint->setMotion(PAxis, PxD6Motion::eLOCKED);
+	}else
+	{
+		PD6Joint->setMotion(PAxis, U2PLinearMotion(Motion));
+	}
+}
+
+#endif //WITH_PHYSX
+
+FConstraintBaseParams::FConstraintBaseParams()
+	: Stiffness(50.f)
+	, Damping(5.f)
+	, Restitution(0.f)
+	, ContactDistance(1.f)
+	, bSoftConstraint(false)
+{
+
+}
+
+FLinearConstraint::FLinearConstraint()
+	: Limit(0.f)
+	, XMotion(LCM_Locked)
+	, YMotion(LCM_Locked)
+	, ZMotion(LCM_Locked)
+{
+	ContactDistance = 5.f;
+	Stiffness = 0.f;
+	Damping = 0.f;
+}
+
+FConeConstraint::FConeConstraint()
+	: Swing1LimitDegrees(45.f)
+	, Swing2LimitDegrees(45.f)
+	, Swing1Motion(ACM_Free)
+	, Swing2Motion(ACM_Free)
+{
+	bSoftConstraint = true;
+	ContactDistance = 1.f;
+}
+
+FTwistConstraint::FTwistConstraint()
+	: TwistLimitDegrees(45)
+	, TwistMotion(ACM_Free)
+{
+	bSoftConstraint = true;
+	ContactDistance = 1.f;
+}
+
+#if WITH_PHYSX
+void FLinearConstraint::UpdatePhysXLinearLimit_AssumesLocked(PxD6Joint* Joint, float AverageMass, float Scale) const
+{
+	const float UseLimit = Limit * Scale;
+	const bool bLockLimitSize = (UseLimit < RB_MinSizeToLockDOF);
+		
+	SetLinearMovement_AssumesLocked<PxD6Axis::eX>(Joint, XMotion, bLockLimitSize);
+	SetLinearMovement_AssumesLocked<PxD6Axis::eY>(Joint, YMotion, bLockLimitSize);
+	SetLinearMovement_AssumesLocked<PxD6Axis::eZ>(Joint, ZMotion, bLockLimitSize);
+
+	// If any DOF is locked/limited, set up the joint limit
+	if (XMotion != LCM_Free || YMotion != LCM_Free || ZMotion != LCM_Free)
+	{
+		PxJointLinearLimit PLinearLimit(GPhysXSDK->getTolerancesScale(), UseLimit, FMath::Clamp(ContactDistance, 5.f, UseLimit * 0.49f));
+		PLinearLimit.restitution = Restitution;
+		SetSoftLimitParams_AssumesLocked(&PLinearLimit, bSoftConstraint, Stiffness*AverageMass, Damping*AverageMass);
+		Joint->setLinearLimit(PLinearLimit);
+	}
+}
+#endif
+
+#if WITH_PHYSX
+void FConeConstraint::UpdatePhysXConeLimit_AssumesLocked(PxD6Joint* Joint, float AverageMass) const
+{
+	if (Swing1Motion == ACM_Limited || Swing2Motion == ACM_Limited)
+	{
+		//Clamp the limit value to valid range which PhysX won't ignore, both value have to be clamped even if there is only one degree limit in constraint
+		const float Limit1Rad = FMath::DegreesToRadians(FMath::ClampAngle(Swing1LimitDegrees, KINDA_SMALL_NUMBER, 179.9999f));
+		const float Limit2Rad = FMath::DegreesToRadians(FMath::ClampAngle(Swing2LimitDegrees, KINDA_SMALL_NUMBER, 179.9999f));
+
+		//Clamp the contact distance so that it's not too small (jittery joint) or too big (always active joint)
+		const float ContactRad = FMath::DegreesToRadians(FMath::Clamp(ContactDistance, 1.f, FMath::Min(Swing1LimitDegrees, Swing2LimitDegrees) * 0.49f));
+
+		PxJointLimitCone PSwingLimitCone(Limit2Rad, Limit1Rad, ContactRad);
+		PSwingLimitCone.restitution = Restitution;
+		SetSoftLimitParams_AssumesLocked(&PSwingLimitCone, bSoftConstraint, Stiffness * AverageMass, Damping * AverageMass);
+		Joint->setSwingLimit(PSwingLimitCone);
+	}
+
+	Joint->setMotion(PxD6Axis::eSWING2, U2PAngularMotion(Swing1Motion));
+	Joint->setMotion(PxD6Axis::eSWING1, U2PAngularMotion(Swing2Motion));
+}
+#endif
+
+#if WITH_PHYSX
+void FTwistConstraint::UpdatePhysXTwistLimit_AssumesLocked(PxD6Joint* Joint, float AverageMass) const
+{
+	if (TwistMotion == ACM_Limited)
+	{
+		const float TwistLimitRad = FMath::DegreesToRadians(TwistLimitDegrees);
+		
+		//Clamp the contact distance so that it's not too small (jittery joint) or too big (always active joint)
+		const float ContactRad = FMath::DegreesToRadians(FMath::Clamp(ContactDistance, 1.f, TwistLimitDegrees * .95f));
+
+		PxJointAngularLimitPair PTwistLimitPair(-TwistLimitRad, TwistLimitRad, ContactRad);
+		PTwistLimitPair.restitution = Restitution;
+		SetSoftLimitParams_AssumesLocked(&PTwistLimitPair, bSoftConstraint, Stiffness * AverageMass, Damping * AverageMass);
+		Joint->setTwistLimit(PTwistLimitPair);
+	}
+
+	Joint->setMotion(PxD6Axis::eTWIST, U2PAngularMotion(TwistMotion));
+}
+#endif
\ No newline at end of file
diff --git a/Engine/Source/Runtime/Engine/Private/PhysicsEngine/PhysAnim.cpp b/Engine/Source/Runtime/Engine/Private/PhysicsEngine/PhysAnim.cpp
--- a/Engine/Source/Runtime/Engine/Private/PhysicsEngine/PhysAnim.cpp
+++ b/Engine/Source/Runtime/Engine/Private/PhysicsEngine/PhysAnim.cpp
@@ -716,7 +716,7 @@ void USkeletalMeshComponent::UpdateRBJointMotors()
 			// If we found this bone, and a visible bone that is not the root, and its joint is motorised in some way..
 			if( (BoneIndex != INDEX_NONE) && (BoneIndex != 0) &&
 				(BoneVisibilityStates[BoneIndex] == BVS_Visible) &&
-				(CI->bAngularOrientationDrive) )
+				(CI->IsAngularOrientationDriveEnabled()) )
 			{
 				check(BoneIndex < LocalAtoms.Num());
 
diff --git a/Engine/Source/Runtime/Engine/Private/PhysicsEngine/PhysDrawing.cpp b/Engine/Source/Runtime/Engine/Private/PhysicsEngine/PhysDrawing.cpp
--- a/Engine/Source/Runtime/Engine/Private/PhysicsEngine/PhysDrawing.cpp
+++ b/Engine/Source/Runtime/Engine/Private/PhysicsEngine/PhysDrawing.cpp
@@ -733,17 +733,17 @@ void UPhysicsAsset::GetCollisionMesh(int32 ViewIndex, FMeshElementCollector& Col
 	}
 }
 
-void UPhysicsAsset::DrawConstraints(FPrimitiveDrawInterface* PDI, const USkeletalMesh* SkelMesh, const TArray<FTransform>& SpaceBases, const FTransform& LocalToWorld, float Scale)
+void UPhysicsAsset::DrawConstraints(int32 ViewIndex, FMeshElementCollector& Collector, const USkeletalMesh* SkelMesh, const TArray<FTransform>& SpaceBases, const FTransform& LocalToWorld, float Scale)
 {
-	for( int32 i=0; i<ConstraintSetup.Num(); i++ )
+	for (int32 i = 0; i < ConstraintSetup.Num(); i++)
 	{
 		FConstraintInstance& Instance = ConstraintSetup[i]->DefaultInstance;
 
 		// Get each constraint frame in world space.
 		FTransform Con1Frame = FTransform::Identity;
 		int32 Bone1Index = SkelMesh->RefSkeleton.FindBoneIndex(Instance.ConstraintBone1);
-		if(Bone1Index != INDEX_NONE)
-		{	
+		if (Bone1Index != INDEX_NONE)
+		{
 			FTransform Body1TM = GetSkelBoneTransform(Bone1Index, SpaceBases, LocalToWorld);
 			Body1TM.RemoveScaling();
 			Con1Frame = Instance.GetRefFrame(EConstraintFrame::Frame1) * Body1TM;
@@ -751,15 +751,15 @@ void UPhysicsAsset::DrawConstraints(FPrimitiveDrawInterface* PDI, const USkeleta
 
 		FTransform Con2Frame = FTransform::Identity;
 		int32 Bone2Index = SkelMesh->RefSkeleton.FindBoneIndex(Instance.ConstraintBone2);
-		if(Bone2Index != INDEX_NONE)
-		{	
+		if (Bone2Index != INDEX_NONE)
+		{
 			FTransform Body2TM = GetSkelBoneTransform(Bone2Index, SpaceBases, LocalToWorld);
 			Body2TM.RemoveScaling();
-			Con2Frame = Instance.GetRefFrame(EConstraintFrame::Frame1) * Body2TM;
+			Con2Frame = Instance.GetRefFrame(EConstraintFrame::Frame2) * Body2TM;
 		}
 
 
-		Instance.DrawConstraint(PDI, Scale, 1.f, true, true, Con1Frame, Con2Frame, false);
+		Instance.DrawConstraint(ViewIndex, Collector, Scale, 1.f, true, true, Con1Frame, Con2Frame, false);
 	}
 }
 static void DrawLinearLimit(FPrimitiveDrawInterface* PDI, const FVector& Origin, const FVector& Axis, const FVector& Orth, float LinearLimitRadius, bool bLinearLimited, float DrawScale)
@@ -823,21 +823,40 @@ float HelpBuildAngle(float LimitAngle, EAngularConstraintMotion LimitType)
 }
 
 
-void FConstraintInstance::DrawConstraint(	FPrimitiveDrawInterface* PDI, 
-											float Scale, float LimitDrawScale, bool bDrawLimits, bool bDrawSelected,
-											const FTransform& Con1Frame, const FTransform& Con2Frame, bool bDrawAsPoint ) const
+FPrimitiveDrawInterface* FConstraintInstance::FPDIOrCollector::GetPDI() const
+{
+	return PDI ? PDI : Collector->GetPDI(ViewIndex);
+}
+
+void FConstraintInstance::FPDIOrCollector::DrawCylinder(const FVector& Start, const FVector& End, float Thickness, FMaterialRenderProxy* MaterialProxy, ESceneDepthPriorityGroup DepthPriority) const
 {
-	// Do nothing unless we are in the interactive editor, otherwise limit materials are not loaded
-	if (!GIsEditor || IsRunningCommandlet())
+	if (HasCollector())
 	{
-		return;
+		GetCylinderMesh(Start, End, Thickness, 4, MaterialProxy, DepthPriority, ViewIndex, *Collector);
+	}
+	else
+	{
+		::DrawCylinder(PDI, Start, End, Thickness, 4, MaterialProxy, DepthPriority);
 	}
+}
+
+void FConstraintInstance::DrawConstraintImp(const FPDIOrCollector& PDIOrCollector, float Scale, float LimitDrawScale, bool bDrawLimits, bool bDrawSelected, const FTransform& Con1Frame, const FTransform& Con2Frame, bool bDrawAsPoint) const
+{
+	// Do nothing if we're shipping
+#if UE_BUILD_SHIPPING
+	return;
+#endif
+
+	FPrimitiveDrawInterface* PDI = PDIOrCollector.GetPDI();
 
 	check((GEngine->ConstraintLimitMaterialX != nullptr) && (GEngine->ConstraintLimitMaterialY != nullptr) && (GEngine->ConstraintLimitMaterialZ != nullptr));
 
 	static UMaterialInterface * LimitMaterialX = GEngine->ConstraintLimitMaterialX;
+	static UMaterialInterface * LimitMaterialXAxis = GEngine->ConstraintLimitMaterialXAxis;
 	static UMaterialInterface * LimitMaterialY = GEngine->ConstraintLimitMaterialY;
+	static UMaterialInterface * LimitMaterialYAxis = GEngine->ConstraintLimitMaterialYAxis;
 	static UMaterialInterface * LimitMaterialZ = GEngine->ConstraintLimitMaterialZ;
+	static UMaterialInterface * LimitMaterialZAxis = GEngine->ConstraintLimitMaterialZAxis;
 	
 	FVector Con1Pos = Con1Frame.GetTranslation();
 	FVector Con2Pos = Con2Frame.GetTranslation();
@@ -863,72 +882,123 @@ void FConstraintInstance::DrawConstraint(	FPrimitiveDrawInterface* PDI,
 		return;
 	}
 
-	FVector Position = Con1Frame.GetTranslation();
+	if (bDrawLimits)
+	{
 
-	PDI->DrawLine(Position, Position + Length * Con1Frame.GetScaledAxis(EAxis::X), JointRed, SDPG_World, Thickness);
-	PDI->DrawLine(Position, Position + Length * Con1Frame.GetScaledAxis(EAxis::Y), JointGreen, SDPG_World, Thickness);
-	PDI->DrawLine(Position, Position + Length * Con1Frame.GetScaledAxis(EAxis::Z), JointBlue, SDPG_World, Thickness);
+		//////////////////////////////////////////////////////////////////////////
+		// ANGULAR DRAWING
 
-	//////////////////////////////////////////////////////////////////////////
-	// LINEAR DRAWING
+		//Draw limits first as they are transparent and need to be under coordinate axes
+		const bool bLockSwing1 = GetAngularSwing1Motion() == ACM_Locked;
+		const bool bLockSwing2 = GetAngularSwing2Motion() == ACM_Locked;
+		const bool bLockAllSwing = bLockSwing1 && bLockSwing2;
 
-	bool bLinearXLocked = (LinearXMotion == LCM_Locked) || (LinearXMotion == LCM_Limited && LinearLimitSize < RB_MinSizeToLockDOF);
-	bool bLinearYLocked = (LinearYMotion == LCM_Locked) || (LinearYMotion == LCM_Limited && LinearLimitSize < RB_MinSizeToLockDOF);
-	bool bLinearZLocked = (LinearZMotion == LCM_Locked) || (LinearZMotion == LCM_Limited && LinearLimitSize < RB_MinSizeToLockDOF);
+		// If swing is limited (but not locked) - draw the limit cone.
+		if (!bLockAllSwing)
+		{
+			if (ProfileInstance.ConeLimit.Swing1Motion == ACM_Free && ProfileInstance.ConeLimit.Swing2Motion == ACM_Free)
+			{
+				if (PDIOrCollector.HasCollector())
+				{
+					GetSphereMesh(Con1Pos, FVector(Length*0.9f), DrawConeLimitSides, DrawConeLimitSides, LimitMaterialX->GetRenderProxy(false), SDPG_World, false, PDIOrCollector.ViewIndex, *PDIOrCollector.Collector);
+				}
+				else
+				{
+					DrawSphere(PDI, Con1Pos, FVector(Length * 0.9f), DrawConeLimitSides, DrawConeLimitSides, LimitMaterialX->GetRenderProxy(false), SDPG_World);
+				}
+			}
+			else
+			{
+				FTransform ConeLimitTM = Con2Frame;
+				ConeLimitTM.SetTranslation(Con1Frame.GetTranslation());
 
-	if(!bLinearXLocked)
-	{
-		bool bLinearXLimited = ( LinearXMotion == LCM_Limited && LinearLimitSize >= RB_MinSizeToLockDOF );
-		DrawLinearLimit(PDI, Con2Frame.GetTranslation(), Con2Frame.GetScaledAxis( EAxis::X ), Con2Frame.GetScaledAxis( EAxis::Z ), LinearLimitSize, bLinearXLimited, LimitDrawScale);
-	}
+				const float Swing1Ang = HelpBuildAngle(GetAngularSwing1Limit(), GetAngularSwing1Motion());
+				const float Swing2Ang = HelpBuildAngle(GetAngularSwing2Limit(), GetAngularSwing2Motion());
+				FMatrix ConeToWorld = FScaleMatrix(FVector(Length * 0.9f)) * ConeLimitTM.ToMatrixWithScale();
 
-	if(!bLinearYLocked)
-	{
-		bool bLinearYLimited = ( LinearYMotion == LCM_Limited && LinearLimitSize >= RB_MinSizeToLockDOF );
-		DrawLinearLimit(PDI, Con2Frame.GetTranslation(), Con2Frame.GetScaledAxis( EAxis::Y ), Con2Frame.GetScaledAxis( EAxis::Z ), LinearLimitSize, bLinearYLimited, LimitDrawScale);
-	}
+				if (PDIOrCollector.HasCollector())
+				{
+					GetConeMesh(ConeToWorld, FMath::RadiansToDegrees(Swing1Ang), FMath::RadiansToDegrees(Swing2Ang), DrawConeLimitSides, LimitMaterialX->GetRenderProxy(false), SDPG_World, PDIOrCollector.ViewIndex, *PDIOrCollector.Collector);
+				}
+				else
+				{
+					DrawCone(PDI, ConeToWorld, Swing1Ang, Swing2Ang, DrawConeLimitSides, false, JointLimitColor, LimitMaterialX->GetRenderProxy(false), SDPG_World);
+				}
+			}
+		}
 
-	if(!bLinearZLocked)
-	{
-		bool bLinearZLimited = ( LinearZMotion == LCM_Limited && LinearLimitSize >= RB_MinSizeToLockDOF );
-		DrawLinearLimit(PDI, Con2Frame.GetTranslation(), Con2Frame.GetScaledAxis( EAxis::Z ), Con2Frame.GetScaledAxis( EAxis::X ), LinearLimitSize, bLinearZLimited, LimitDrawScale);
+		//twist
+		if (GetAngularTwistMotion() != ACM_Locked)
+		{
+			FMatrix ConeToWorld = HelpBuildFan(Con1Frame, Con2Frame, EAxis::Y, EAxis::X, Length);
+			float Limit = HelpBuildAngle(GetAngularTwistLimit(), GetAngularTwistMotion());
+			if (PDIOrCollector.HasCollector())
+			{
+				GetConeMesh(ConeToWorld, FMath::RadiansToDegrees(Limit), 0, DrawConeLimitSides, LimitMaterialY->GetRenderProxy(false), SDPG_World, PDIOrCollector.ViewIndex, *PDIOrCollector.Collector);
+			}
+			else
+			{
+				DrawCone(PDI, ConeToWorld, Limit, 0, DrawConeLimitSides, false, JointLimitColor, LimitMaterialY->GetRenderProxy(false), SDPG_World);
+			}
+		}
 	}
 
 
-	if(!bDrawLimits)
-		return;
+	//////////////////////////////////////////////////////////////////////////
+	// COORDINATE AXES
+	FVector Position = Con1Frame.GetTranslation();
 
+	PDIOrCollector.DrawCylinder(Position, Position + Length * Con1Frame.GetScaledAxis(EAxis::X), Thickness, LimitMaterialXAxis->GetRenderProxy(false), SDPG_World);
+	PDIOrCollector.DrawCylinder(Position, Position + Length * Con1Frame.GetScaledAxis(EAxis::Y), Thickness, LimitMaterialYAxis->GetRenderProxy(false), SDPG_World);
+	PDIOrCollector.DrawCylinder(Position, Position + Length * Con1Frame.GetScaledAxis(EAxis::Z), Thickness, LimitMaterialZAxis->GetRenderProxy(false), SDPG_World);
 
-	//////////////////////////////////////////////////////////////////////////
-	// ANGULAR DRAWING
+	PDIOrCollector.DrawCylinder(Position, Position + Length * Con2Frame.GetScaledAxis(EAxis::X), Thickness, LimitMaterialXAxis->GetRenderProxy(false), SDPG_World);
+	PDIOrCollector.DrawCylinder(Position, Position + Length * Con2Frame.GetScaledAxis(EAxis::Y), Thickness, LimitMaterialYAxis->GetRenderProxy(false), SDPG_World);
+	PDIOrCollector.DrawCylinder(Position, Position + Length * Con2Frame.GetScaledAxis(EAxis::Z), Thickness, LimitMaterialZAxis->GetRenderProxy(false), SDPG_World);
 
-	bool bLockTwist = AngularTwistMotion == ACM_Locked; 
-	bool bLockSwing1 = AngularSwing1Motion == ACM_Locked; 
-	bool bLockSwing2 = AngularSwing2Motion == ACM_Locked; 
-	bool bLockAllSwing = bLockSwing1 && bLockSwing2;
 
-	bool bDrawnAxisLine = false;
-	FVector RefLineStart = Con1Frame.GetTranslation() + (Length * Con1Frame.GetScaledAxis(EAxis::X));
-	FVector RefLineEnd = Con1Frame.GetTranslation() + (Length * Con1Frame.GetScaledAxis(EAxis::X));
+	//Draw arrow on twist axist
+	{
+		FTransform ConeLimitTM = Con2Frame;
+		ConeLimitTM.SetTranslation(Con1Frame.GetTranslation() + Length*1.05*Con2Frame.GetScaledAxis(EAxis::X));
+
+		const float Swing1Ang = PI / 4;
+		const float Swing2Ang = PI / 4;
+		FMatrix ConeToWorld = FScaleMatrix(FVector(Length * -0.1f)) * ConeLimitTM.ToMatrixWithScale();
+
+		if (PDIOrCollector.HasCollector())
+		{
+			GetConeMesh(ConeToWorld, FMath::RadiansToDegrees(Swing1Ang), FMath::RadiansToDegrees(Swing2Ang), DrawConeLimitSides, LimitMaterialXAxis->GetRenderProxy(false), SDPG_World, PDIOrCollector.ViewIndex, *PDIOrCollector.Collector);
+		}
+		else
+		{
+			DrawCone(PDI, ConeToWorld, Swing1Ang, Swing2Ang, DrawConeLimitSides, false, JointLimitColor, LimitMaterialXAxis->GetRenderProxy(false), SDPG_World);
+		}
+	}
+
+	//////////////////////////////////////////////////////////////////////////
+	// LINEAR DRAWING
 
-	//swing1
+	//TODO: Move this all into a draw function on linear constraint
+	bool bLinearXLocked = (GetLinearXMotion() == LCM_Locked) || (GetLinearXMotion() == LCM_Limited && GetLinearLimit() < RB_MinSizeToLockDOF);
+	bool bLinearYLocked = (GetLinearYMotion() == LCM_Locked) || (GetLinearYMotion() == LCM_Limited && GetLinearLimit() < RB_MinSizeToLockDOF);
+	bool bLinearZLocked = (GetLinearZMotion() == LCM_Locked) || (GetLinearZMotion() == LCM_Limited && GetLinearLimit() < RB_MinSizeToLockDOF);
+
+	if(!bLinearXLocked)
 	{
-		FMatrix ConeToWorld = HelpBuildFan(Con1Frame, Con2Frame, EAxis::X, EAxis::Z, Length);
-		float Limit = HelpBuildAngle(Swing1LimitAngle, AngularSwing1Motion);
-		DrawCone(PDI, ConeToWorld, Limit, 0, DrawConeLimitSides, false, JointLimitColor, LimitMaterialX->GetRenderProxy(false), SDPG_World);
+		bool bLinearXLimited = ( GetLinearXMotion() == LCM_Limited && GetLinearLimit() >= RB_MinSizeToLockDOF );
+		DrawLinearLimit(PDI, Con2Frame.GetTranslation(), Con2Frame.GetScaledAxis( EAxis::X ), Con2Frame.GetScaledAxis( EAxis::Z ), GetLinearLimit(), bLinearXLimited, LimitDrawScale);
 	}
 
-	//swing2
+	if(!bLinearYLocked)
 	{
-		FMatrix ConeToWorld = HelpBuildFan(Con1Frame, Con2Frame, EAxis::Z, EAxis::Y, Length);
-		float Limit = HelpBuildAngle(Swing2LimitAngle, AngularSwing2Motion);
-		DrawCone(PDI, ConeToWorld, Limit, 0, DrawConeLimitSides, false, JointLimitColor, LimitMaterialZ->GetRenderProxy(false), SDPG_World);
+		bool bLinearYLimited = ( GetLinearYMotion() == LCM_Limited && GetLinearLimit() >= RB_MinSizeToLockDOF );
+		DrawLinearLimit(PDI, Con2Frame.GetTranslation(), Con2Frame.GetScaledAxis( EAxis::Y ), Con2Frame.GetScaledAxis( EAxis::Z ), GetLinearLimit(), bLinearYLimited, LimitDrawScale);
 	}
 
-	//twist
+	if(!bLinearZLocked)
 	{
-		FMatrix ConeToWorld = HelpBuildFan(Con1Frame, Con2Frame, EAxis::Y, EAxis::X, Length);
-		float Limit = HelpBuildAngle(TwistLimitAngle, AngularTwistMotion);
-		DrawCone(PDI, ConeToWorld, Limit, 0, DrawConeLimitSides, false, JointLimitColor, LimitMaterialY->GetRenderProxy(false), SDPG_World);
+		bool bLinearZLimited = ( GetLinearZMotion() == LCM_Limited && GetLinearLimit() >= RB_MinSizeToLockDOF );
+		DrawLinearLimit(PDI, Con2Frame.GetTranslation(), Con2Frame.GetScaledAxis( EAxis::Z ), Con2Frame.GetScaledAxis( EAxis::X ), GetLinearLimit(), bLinearZLimited, LimitDrawScale);
 	}
 }
diff --git a/Engine/Source/Runtime/Engine/Private/PhysicsEngine/PhysicsConstraintActor.cpp b/Engine/Source/Runtime/Engine/Private/PhysicsEngine/PhysicsConstraintActor.cpp
--- a/Engine/Source/Runtime/Engine/Private/PhysicsEngine/PhysicsConstraintActor.cpp
+++ b/Engine/Source/Runtime/Engine/Private/PhysicsEngine/PhysicsConstraintActor.cpp
@@ -4,6 +4,7 @@
 #include "PhysXSupport.h"
 #include "PhysicsEngine/PhysicsConstraintActor.h"
 #include "PhysicsEngine/PhysicsConstraintComponent.h"
+#include "PhysicsEngine/ConstraintUtils.h"
 
 APhysicsConstraintActor::APhysicsConstraintActor(const FObjectInitializer& ObjectInitializer)
 	: Super(ObjectInitializer)
@@ -35,7 +36,7 @@ void APhysicsConstraintActor::PostLoad()
 	{
 		ConstraintComp->ConstraintActor1 = ConstraintActor1_DEPRECATED;
 		ConstraintComp->ConstraintActor2 = ConstraintActor2_DEPRECATED;
-		ConstraintComp->ConstraintInstance.bDisableCollision = bDisableCollision_DEPRECATED;
+		ConstraintComp->ConstraintInstance.ProfileInstance.bDisableCollision = bDisableCollision_DEPRECATED;
 	}
 }
 
@@ -50,15 +51,16 @@ void APhysicsConstraintActor::LoadedFromAnotherClass( const FName& OldClassName
 
 	if (OldClassName == PhysicsHingeActor_NAME)
 	{
-		ConstraintComp->ConstraintInstance.ConfigureAsHinge(false);
+		ConstraintUtils::ConfigureAsHinge(ConstraintComp->ConstraintInstance, false);
 	}
 	else if (OldClassName == PhysicsPrismaticActor_NAME)
 	{
-		ConstraintComp->ConstraintInstance.ConfigureAsPrismatic(false);
+		ConstraintUtils::ConfigureAsPrismatic(ConstraintComp->ConstraintInstance, false);
 	}
 	else if (OldClassName == PhysicsBSJointActor_NAME)
 	{
-		ConstraintComp->ConstraintInstance.ConfigureAsBS(false);
+		ConstraintUtils::ConfigureAsBallAndSocket(ConstraintComp->ConstraintInstance, false);
+
 	}
 
 	ConstraintComp->UpdateSpriteTexture();
diff --git a/Engine/Source/Runtime/Engine/Private/PhysicsEngine/PhysicsConstraintComponent.cpp b/Engine/Source/Runtime/Engine/Private/PhysicsEngine/PhysicsConstraintComponent.cpp
--- a/Engine/Source/Runtime/Engine/Private/PhysicsEngine/PhysicsConstraintComponent.cpp
+++ b/Engine/Source/Runtime/Engine/Private/PhysicsEngine/PhysicsConstraintComponent.cpp
@@ -6,6 +6,7 @@
 #include "UObjectToken.h"
 #include "PhysicsEngine/PhysicsConstraintTemplate.h"
 #include "PhysicsEngine/PhysicsConstraintComponent.h"
+#include "PhysicsEngine/ConstraintUtils.h"
 #include "PhysicsEngine/PhysicsAsset.h"
 
 #define LOCTEXT_NAMESPACE "ConstraintComponent"
@@ -316,12 +317,15 @@ void UPhysicsConstraintComponent::PostLoad()
 		{
 			float AverageMass = TotalMass / NumDynamic;
 
-			ConstraintInstance.LinearLimitStiffness /= AverageMass;
-			ConstraintInstance.SwingLimitStiffness /= AverageMass;
-			ConstraintInstance.TwistLimitStiffness /= AverageMass;
-			ConstraintInstance.LinearLimitDamping /= AverageMass;
-			ConstraintInstance.SwingLimitDamping /= AverageMass;
-			ConstraintInstance.TwistLimitDamping /= AverageMass;
+#if WITH_EDITORONLY_DATA
+			ConstraintInstance.ProfileInstance.LinearLimit.Stiffness /= AverageMass;
+			ConstraintInstance.SwingLimitStiffness_DEPRECATED /= AverageMass;
+			ConstraintInstance.TwistLimitStiffness_DEPRECATED /= AverageMass;
+			ConstraintInstance.LinearLimitDamping_DEPRECATED /= AverageMass;
+			ConstraintInstance.SwingLimitDamping_DEPRECATED /= AverageMass;
+			ConstraintInstance.TwistLimitDamping_DEPRECATED /= AverageMass;
+			
+#endif
 		}
 
 	}
@@ -464,11 +468,11 @@ void UPhysicsConstraintComponent::UpdateSpriteTexture()
 {
 	if (SpriteComponent)
 	{
-		if (ConstraintInstance.IsHinge())
+		if (ConstraintUtils::IsHinge(ConstraintInstance))
 		{
 			SpriteComponent->SetSprite(LoadObject<UTexture2D>(NULL, TEXT("/Engine/EditorResources/S_KHinge.S_KHinge")));
 		}
-		else if (ConstraintInstance.IsPrismatic())
+		else if (ConstraintUtils::IsPrismatic(ConstraintInstance))
 		{
 			SpriteComponent->SetSprite(LoadObject<UTexture2D>(NULL, TEXT("/Engine/EditorResources/S_KPrismatic.S_KPrismatic")));
 		}
diff --git a/Engine/Source/Runtime/Engine/Private/PhysicsEngine/PhysicsConstraintTemplate.cpp b/Engine/Source/Runtime/Engine/Private/PhysicsEngine/PhysicsConstraintTemplate.cpp
--- a/Engine/Source/Runtime/Engine/Private/PhysicsEngine/PhysicsConstraintTemplate.cpp
+++ b/Engine/Source/Runtime/Engine/Private/PhysicsEngine/PhysicsConstraintTemplate.cpp
@@ -8,6 +8,7 @@
 UPhysicsConstraintTemplate::UPhysicsConstraintTemplate(const FObjectInitializer& ObjectInitializer)
 	: Super(ObjectInitializer)
 {
+#if WITH_EDITORONLY_DATA
 	LinearXMotion_DEPRECATED = LCM_Locked;
 	LinearYMotion_DEPRECATED = LCM_Locked;
 	LinearZMotion_DEPRECATED = LCM_Locked;
@@ -25,6 +26,7 @@ UPhysicsConstraintTemplate::UPhysicsConstraintTemplate(const FObjectInitializer&
 
 	ProjectionLinearTolerance_DEPRECATED = 0.5; // Linear projection when error > 5.0 unreal units
 	ProjectionAngularTolerance_DEPRECATED = 10.f; // Angular projection when error > 10 degrees
+#endif
 }
 
 void UPhysicsConstraintTemplate::Serialize(FArchive& Ar)
@@ -36,10 +38,107 @@ void UPhysicsConstraintTemplate::Serialize(FArchive& Ar)
 	{
 		CopySetupPropsToInstance(&DefaultInstance);
 	}
+
+	if(Ar.IsLoading() && !Ar.IsTransacting())
+	{
+		UpdateConstraintProfileMap();
+#if WITH_EDITORONLY_DATA
+
+		if(!ConstraintProfileNameMap.FindRef(DefaultProfileName))	//If no profile found make sure to copy what's in default
+		{
+			NoProfileInstance = DefaultInstance.ProfileInstance;
+		}
+
+		UpdateInstanceFromProfile();
+#endif
+	}
+}
+
+#if WITH_EDITORONLY_DATA
+void UPhysicsConstraintTemplate::UpdateInstanceFromProfile()
+{
+	if(UPhysicsConstraintProfile* ProfileInstance = ConstraintProfileNameMap.FindRef(DefaultProfileName))
+	{
+		DefaultInstance.ProfileInstance = ProfileInstance->ProfileProperties;
+	}
+	else
+	{
+		//no profile found so update the no profile cache
+		DefaultInstance.ProfileInstance = NoProfileInstance;
+	}
+}
+#endif
+
+#if WITH_EDITOR
+void UPhysicsConstraintTemplate::UpdateProfileBeingEdited()
+{
+	if (UPhysicsConstraintProfile* Profile = ConstraintProfileNameMap.FindRef(DefaultProfileName))
+	{
+		//Make sure any changes we make are passed back to the profile we are currently editing
+		Profile->ProfileProperties = DefaultInstance.ProfileInstance;
+	}
+	else
+	{
+		//No profile found so make sure to update the no profile cache
+		NoProfileInstance = DefaultInstance.ProfileInstance;
+	}
+}
+#endif
+
+void UPhysicsConstraintTemplate::SanitizeProfileArray()
+{
+	for(FPhysicsConstraintProfileHandle& ProfileHandle : ProfileHandles)
+	{
+		if(ProfileHandle.Profile == nullptr)
+		{
+			ProfileHandle.Profile = NewObject<UPhysicsConstraintProfile>(this);
+			checkSlow(ProfileHandle.Profile);	//checkSlow to make static analysis happy
+			ProfileHandle.Profile->ProfileProperties = DefaultInstance.ProfileInstance;	//new objects use whatever is the current default
+		}
+	}
+}
+
+void UPhysicsConstraintTemplate::UpdateConstraintProfileMap()
+{
+	SanitizeProfileArray();
+
+	ConstraintProfileNameMap.Empty(ProfileHandles.Num());
+	for(const FPhysicsConstraintProfileHandle& ProfileHandle : ProfileHandles)
+	{
+		if(ProfileHandle.Profile)
+		{
+			ConstraintProfileNameMap.Add(ProfileHandle.ProfileName, ProfileHandle.Profile);
+		}
+	}
+}
+
+#if WITH_EDITOR
+void UPhysicsConstraintTemplate::PostEditChangeProperty(FPropertyChangedEvent& PropertyChangedEvent)
+{
+	if(UProperty* Property = PropertyChangedEvent.Property)
+	{
+		const FName PropertyName = Property->GetFName();
+		FName ProfileHandlesName = GET_MEMBER_NAME_CHECKED(UPhysicsConstraintTemplate, ProfileHandles);
+		FName ProfileHandleInnerName = GET_MEMBER_NAME_CHECKED(FPhysicsConstraintProfileHandle, ProfileName);	//The profile name (inside the array element)
+		if(PropertyName == ProfileHandlesName || PropertyName == ProfileHandleInnerName)
+		{
+			UpdateConstraintProfileMap();
+		}
+		
+		FName DefaultProfileNamePropertyName = GET_MEMBER_NAME_CHECKED(UPhysicsConstraintTemplate, DefaultProfileName);
+		if (PropertyName == DefaultProfileNamePropertyName || PropertyName == ProfileHandlesName || PropertyName == ProfileHandleInnerName)
+		{
+			UpdateInstanceFromProfile();
+		}
+		
+		UpdateProfileBeingEdited();
+	}
 }
+#endif
 
 void UPhysicsConstraintTemplate::CopySetupPropsToInstance(FConstraintInstance* Instance)
 {
+#if WITH_EDITORONLY_DATA
 	Instance->JointName			= JointName_DEPRECATED;
 	Instance->ConstraintBone1	= ConstraintBone1_DEPRECATED;
 	Instance->ConstraintBone2	= ConstraintBone2_DEPRECATED;
@@ -51,30 +150,31 @@ void UPhysicsConstraintTemplate::CopySetupPropsToInstance(FConstraintInstance* I
 	Instance->PriAxis2			= PriAxis2_DEPRECATED;
 	Instance->SecAxis2			= SecAxis2_DEPRECATED;
 
-	Instance->bEnableProjection				= bEnableProjection_DEPRECATED;
-	Instance->ProjectionLinearTolerance		= ProjectionLinearTolerance_DEPRECATED;
-	Instance->ProjectionAngularTolerance	= ProjectionAngularTolerance_DEPRECATED;
-	Instance->LinearXMotion					= LinearXMotion_DEPRECATED;
-	Instance->LinearYMotion					= LinearYMotion_DEPRECATED;
-	Instance->LinearZMotion					= LinearZMotion_DEPRECATED;
-	Instance->LinearLimitSize				= LinearLimitSize_DEPRECATED;
-	Instance->bLinearLimitSoft				= bLinearLimitSoft_DEPRECATED;
-	Instance->LinearLimitStiffness			= LinearLimitStiffness_DEPRECATED;
-	Instance->LinearLimitDamping			= LinearLimitDamping_DEPRECATED;
-	Instance->bLinearBreakable				= bLinearBreakable_DEPRECATED;
-	Instance->LinearBreakThreshold			= LinearBreakThreshold_DEPRECATED;
-	Instance->AngularSwing1Motion			= AngularSwing1Motion_DEPRECATED;
-	Instance->AngularSwing2Motion			= AngularSwing2Motion_DEPRECATED;
-	Instance->AngularTwistMotion			= AngularTwistMotion_DEPRECATED;
-	Instance->bSwingLimitSoft				= bSwingLimitSoft_DEPRECATED;
-	Instance->bTwistLimitSoft				= bTwistLimitSoft_DEPRECATED;
-	Instance->Swing1LimitAngle				= Swing1LimitAngle_DEPRECATED;
-	Instance->Swing2LimitAngle				= Swing2LimitAngle_DEPRECATED;
-	Instance->TwistLimitAngle				= TwistLimitAngle_DEPRECATED;
-	Instance->SwingLimitStiffness			= SwingLimitStiffness_DEPRECATED;
-	Instance->SwingLimitDamping				= SwingLimitDamping_DEPRECATED;
-	Instance->TwistLimitStiffness			= TwistLimitStiffness_DEPRECATED;
-	Instance->TwistLimitDamping				= TwistLimitDamping_DEPRECATED;
-	Instance->bAngularBreakable				= bAngularBreakable_DEPRECATED;
-	Instance->AngularBreakThreshold			= AngularBreakThreshold_DEPRECATED;
+	Instance->ProfileInstance.bEnableProjection				= bEnableProjection_DEPRECATED;
+	Instance->ProfileInstance.ProjectionLinearTolerance		= ProjectionLinearTolerance_DEPRECATED;
+	Instance->ProfileInstance.ProjectionAngularTolerance		= ProjectionAngularTolerance_DEPRECATED;
+	Instance->ProfileInstance.LinearLimit.XMotion				= LinearXMotion_DEPRECATED;
+	Instance->ProfileInstance.LinearLimit.YMotion				= LinearYMotion_DEPRECATED;
+	Instance->ProfileInstance.LinearLimit.ZMotion				= LinearZMotion_DEPRECATED;
+	Instance->ProfileInstance.LinearLimit.Limit				= LinearLimitSize_DEPRECATED;
+	Instance->ProfileInstance.LinearLimit.bSoftConstraint		= bLinearLimitSoft_DEPRECATED;
+	Instance->ProfileInstance.LinearLimit.Stiffness			= LinearLimitStiffness_DEPRECATED;
+	Instance->ProfileInstance.LinearLimit.Damping				= LinearLimitDamping_DEPRECATED;
+	Instance->ProfileInstance.bLinearBreakable				= bLinearBreakable_DEPRECATED;
+	Instance->ProfileInstance.LinearBreakThreshold			= LinearBreakThreshold_DEPRECATED;
+	Instance->ProfileInstance.ConeLimit.Swing1Motion			= AngularSwing1Motion_DEPRECATED;
+	Instance->ProfileInstance.ConeLimit.Swing2Motion			= AngularSwing2Motion_DEPRECATED;
+	Instance->ProfileInstance.TwistLimit.TwistMotion			= AngularTwistMotion_DEPRECATED;
+	Instance->ProfileInstance.ConeLimit.bSoftConstraint		= bSwingLimitSoft_DEPRECATED;
+	Instance->ProfileInstance.TwistLimit.bSoftConstraint		= bTwistLimitSoft_DEPRECATED;
+	Instance->ProfileInstance.ConeLimit.Swing1LimitDegrees	= Swing1LimitAngle_DEPRECATED;
+	Instance->ProfileInstance.ConeLimit.Swing2LimitDegrees	= Swing2LimitAngle_DEPRECATED;
+	Instance->ProfileInstance.TwistLimit.TwistLimitDegrees	= TwistLimitAngle_DEPRECATED;
+	Instance->ProfileInstance.ConeLimit.Stiffness				= SwingLimitStiffness_DEPRECATED;
+	Instance->ProfileInstance.ConeLimit.Damping				= SwingLimitDamping_DEPRECATED;
+	Instance->ProfileInstance.TwistLimit.Stiffness			= TwistLimitStiffness_DEPRECATED;
+	Instance->ProfileInstance.TwistLimit.Damping				= TwistLimitDamping_DEPRECATED;
+	Instance->ProfileInstance.bAngularBreakable				= bAngularBreakable_DEPRECATED;
+	Instance->ProfileInstance.AngularBreakThreshold			= AngularBreakThreshold_DEPRECATED;
+#endif
 }
\ No newline at end of file
diff --git a/Engine/Source/Runtime/Engine/Private/PrimitiveDrawingUtils.cpp b/Engine/Source/Runtime/Engine/Private/PrimitiveDrawingUtils.cpp
--- a/Engine/Source/Runtime/Engine/Private/PrimitiveDrawingUtils.cpp
+++ b/Engine/Source/Runtime/Engine/Private/PrimitiveDrawingUtils.cpp
@@ -609,6 +609,21 @@ void BuildCylinderVerts(const FVector& Base, const FVector& XAxis, const FVector
 
 }
 
+void GetCylinderMesh(const FVector& Start, const FVector& End, float Radius, int32 Sides, const FMaterialRenderProxy* MaterialInstance, uint8 DepthPriority, int32 ViewIndex, FMeshElementCollector& Collector)
+{
+	FVector Dir = End - Start;
+	float Length = Dir.Size();
+
+	if (Length > SMALL_NUMBER)
+	{
+		FVector Z = Dir.GetUnsafeNormal();
+		FVector X, Y;
+		Z.GetUnsafeNormal().FindBestAxisVectors(X, Y);
+
+		GetCylinderMesh(FMatrix::Identity, Z * Length*0.5 + Start, X, Y, Z, Radius, Length * 0.5f, Sides, MaterialInstance, DepthPriority, ViewIndex, Collector);
+	}
+
+}
 
 void GetCylinderMesh(const FVector& Base, const FVector& XAxis, const FVector& YAxis, const FVector& ZAxis,
 				  float Radius, float HalfHeight, int32 Sides, const FMaterialRenderProxy* MaterialRenderProxy, uint8 DepthPriority, int32 ViewIndex, FMeshElementCollector& Collector)
@@ -675,6 +690,22 @@ void DrawCylinder(FPrimitiveDrawInterface* PDI, const FMatrix& CylToWorld, const
 	MeshBuilder.Draw(PDI, CylToWorld, MaterialRenderProxy, DepthPriority,0.f);
 }
 
+void DrawCylinder(class FPrimitiveDrawInterface* PDI, const FVector& Start, const FVector& End, float Radius, int32 Sides, const FMaterialRenderProxy* MaterialInstance, uint8 DepthPriority)
+{
+	FVector Dir = End - Start;
+	float Length = Dir.Size();
+
+	if (Length > SMALL_NUMBER)
+	{
+		FVector Z = Dir.GetUnsafeNormal();
+		FVector X, Y;
+		Z.GetUnsafeNormal().FindBestAxisVectors(X, Y);
+
+		DrawCylinder(PDI, FMatrix::Identity, Z * Length*0.5 + Start, X, Y, Z, Radius, Length * 0.5f, Sides, MaterialInstance, DepthPriority);
+	}
+
+}
+
 void DrawDisc(class FPrimitiveDrawInterface* PDI,const FVector& Base,const FVector& XAxis,const FVector& YAxis,FColor Color,float Radius,int32 NumSides,const FMaterialRenderProxy* MaterialRenderProxy, uint8 DepthPriority)
 {
 	check (NumSides >= 3);
diff --git a/Engine/Source/Runtime/Engine/Private/SkeletalMesh.cpp b/Engine/Source/Runtime/Engine/Private/SkeletalMesh.cpp
--- a/Engine/Source/Runtime/Engine/Private/SkeletalMesh.cpp
+++ b/Engine/Source/Runtime/Engine/Private/SkeletalMesh.cpp
@@ -3306,8 +3306,9 @@ void USkeletalMesh::GetAssetRegistryTags(TArray<FAssetRegistryTag>& OutTags) con
 		NumTriangles = LODModel.GetTotalFaces();
 	}
 
-	OutTags.Add( FAssetRegistryTag("Triangles", FString::FromInt(NumTriangles), FAssetRegistryTag::TT_Numerical) );
-	OutTags.Add( FAssetRegistryTag("Bones", FString::FromInt(RefSkeleton.GetNum()), FAssetRegistryTag::TT_Numerical) );
+	OutTags.Add(FAssetRegistryTag("Triangles", FString::FromInt(NumTriangles), FAssetRegistryTag::TT_Numerical));
+	OutTags.Add(FAssetRegistryTag("Bones", FString::FromInt(RefSkeleton.GetNum()), FAssetRegistryTag::TT_Numerical));
+	OutTags.Add(FAssetRegistryTag("MorphTargets", FString::FromInt(MorphTargets.Num()), FAssetRegistryTag::TT_Numerical));
 
 #if WITH_EDITORONLY_DATA
 	if (AssetImportData)
@@ -5160,6 +5161,8 @@ void FSkeletalMeshSceneProxy::GetDynamicElementsSection(const TArray<const FScen
 			}
 
 			BatchElement.NumPrimitives = Section.NumTriangles;
+
+#if WITH_EDITORONLY_DATA
 			if( GIsEditor && MeshObject->ProgressiveDrawingFraction != 1.f )
 			{
 				if (Mesh.MaterialRenderProxy && Mesh.MaterialRenderProxy->GetMaterial(FeatureLevel)->GetBlendMode() == BLEND_Translucent)
@@ -5171,6 +5174,8 @@ void FSkeletalMeshSceneProxy::GetDynamicElementsSection(const TArray<const FScen
 					}
 				}
 			}
+#endif // WITH_EDITORONLY_DATA
+
 			BatchElement.MinVertexIndex = Section.BaseVertexIndex;
 			Mesh.ReverseCulling = IsLocalToWorldDeterminantNegative();
 			Mesh.CastShadow = SectionElementInfo.bEnableShadowCasting;
@@ -5258,6 +5263,10 @@ FPrimitiveViewRelevance FSkeletalMeshSceneProxy::GetViewRelevance(const FSceneVi
 	Result.bRenderInMainPass = ShouldRenderInMainPass();
 	Result.bUsesLightingChannels = GetLightingChannelMask() != GetDefaultLightingChannelMask();
 	MaterialRelevance.SetPrimitiveViewRelevance(Result);
+
+#if !UE_BUILD_SHIPPING
+	Result.bSeparateTranslucencyRelevance |= View->Family->EngineShowFlags.Constraints;
+#endif
 	return Result;
 }
 
@@ -5302,15 +5311,15 @@ void FSkeletalMeshSceneProxy::DebugDrawPhysicsAsset(int32 ViewIndex, FMeshElemen
 		TArray<FTransform>* BoneSpaceBases = MeshObject->GetSpaceBases();
 		if(BoneSpaceBases)
 		{
+			//TODO: These data structures are not double buffered. This is not thread safe!
 			check(PhysicsAssetForDebug);
-			if( EngineShowFlags.Collision && IsCollisionEnabled() )
+			if (EngineShowFlags.Collision && IsCollisionEnabled())
 			{
 				PhysicsAssetForDebug->GetCollisionMesh(ViewIndex, Collector, SkeletalMeshForDebug, *BoneSpaceBases, LocalToWorldTransform, TotalScale);
 			}
-			if( EngineShowFlags.Constraints )
+			if (EngineShowFlags.Constraints)
 			{
-				// TODO: Currently doesn't work because the PDI returned here (FSimpleElementCollector) can't draw meshes
-				//PhysicsAssetForDebug->DrawConstraints(Collector.GetPDI(ViewIndex), SkeletalMeshForDebug, *BoneSpaceBases, LocalToWorldTransform, TotalScale.X);
+				PhysicsAssetForDebug->DrawConstraints(ViewIndex, Collector, SkeletalMeshForDebug, *BoneSpaceBases, LocalToWorldTransform, TotalScale.X);
 			}
 		}
 	}
@@ -5405,10 +5414,10 @@ USkinnedMeshComponent::USkinnedMeshComponent(const FObjectInitializer& ObjectIni
 	SetCollisionProfileName(UCollisionProfile::NoCollision_ProfileName);
 
 	StreamingDistanceMultiplier = 1.0f;
-	ProgressiveDrawingFraction = 1.0f;
 	bCanHighlightSelectedSections = false;
 	CanCharacterStepUpOn = ECB_Owner;
 #if WITH_EDITORONLY_DATA
+	ProgressiveDrawingFraction = 1.0f;
 	ChunkIndexPreview = -1;
 	SectionIndexPreview = -1;
 #endif // WITH_EDITORONLY_DATA
diff --git a/Engine/Source/Runtime/Engine/Private/SkeletalMeshComponentPhysics.cpp b/Engine/Source/Runtime/Engine/Private/SkeletalMeshComponentPhysics.cpp
--- a/Engine/Source/Runtime/Engine/Private/SkeletalMeshComponentPhysics.cpp
+++ b/Engine/Source/Runtime/Engine/Private/SkeletalMeshComponentPhysics.cpp
@@ -933,6 +933,11 @@ void USkeletalMeshComponent::InitArticulated(FPhysScene* PhysScene)
 		check( ConInst );
 		ConInst->ConstraintIndex = i; // Set the ConstraintIndex property in the ConstraintInstance.
 		ConInst->CopyConstraintParamsFrom(&ConstraintSetup->DefaultInstance);
+		if(UPhysicsConstraintProfile* DefaultProfile = ConstraintSetup->GetDefaultProfile())
+		{
+			ConInst->CopyProfilePropertiesFrom(DefaultProfile->ProfileProperties);
+		}
+		
 
 		// Get bodies we want to joint
 		FBodyInstance* Body1 = GetBodyInstance(ConInst->ConstraintBone1);
@@ -1063,26 +1068,12 @@ void USkeletalMeshComponent::TermBodiesBelow(FName ParentBoneName)
 
 float USkeletalMeshComponent::GetTotalMassBelowBone(FName InBoneName)
 {
-	UPhysicsAsset* const PhysicsAsset = GetPhysicsAsset();
-	if(!PhysicsAsset || !SkeletalMesh)
-	{
-		return 0.f;
-	}
-
-	// if physics state is invalid - i.e. collision is disabled - or it does not have a valid bodies, this will crash right away
-	if (!IsPhysicsStateCreated()  || !bHasValidBodies)
-	{
-		return 0.f;
-	}
-
-	TArray<int32> BodyIndices;
-	PhysicsAsset->GetBodyIndicesBelow(BodyIndices, InBoneName, SkeletalMesh);
-
 	float TotalMass = 0.f;
-	for(int32 i=0; i<BodyIndices.Num(); i++)
+
+	ForEachBodyBelow(InBoneName, /*bIncludeSelf=*/true, /*bSkipCustomPhysics=*/false, [&TotalMass](FBodyInstance* BI)
 	{
-		TotalMass += Bodies[BodyIndices[i]]->GetBodyMass();
-	}
+		TotalMass += BI->GetBodyMass();
+	});
 
 	return TotalMass;
 }
@@ -1118,37 +1109,23 @@ void USkeletalMeshComponent::SetAllBodiesNotifyRigidBodyCollision(bool bNewNotif
 }
 
 
-void USkeletalMeshComponent::SetAllBodiesBelowSimulatePhysics( const FName& InBoneName, bool bNewSimulate  )
+void USkeletalMeshComponent::SetAllBodiesBelowSimulatePhysics( const FName& InBoneName, bool bNewSimulate, bool bIncludeSelf )
 {
-	UPhysicsAsset* const PhysicsAsset = GetPhysicsAsset();
-	if( !PhysicsAsset || !SkeletalMesh )
+	int32 NumBodiesFound = ForEachBodyBelow(InBoneName, bIncludeSelf, /*bSkipCustomPhysicsType=*/ false, [bNewSimulate](FBodyInstance* BI)
 	{
-		return;
-	}
-
-	// if physics state is invalid - i.e. collision is disabled - or it does not have a valid bodies, this will crash right away
-	if (!IsPhysicsStateCreated()  || !bHasValidBodies)
-	{
-		FMessageLog("PIE").Warning(LOCTEXT("InvalidBodies", "Invalid Bodies : Make sure collision is enabled or root bone has body in PhysicsAsset."));
-		return;
-	}
-
-	TArray<int32> BodyIndices;
-	PhysicsAsset->GetBodyIndicesBelow(BodyIndices, InBoneName, SkeletalMesh);
+		BI->SetInstanceSimulatePhysics(bNewSimulate);
+	});
 
-	for(int32 i=0; i<BodyIndices.Num(); i++)
+	if (NumBodiesFound)
 	{
-		//UE_LOG(LogSkeletalMesh, Warning, TEXT( "ForceAllBodiesBelowUnfixed %s" ), *InAsset->BodySetup(BodyIndices(i))->BoneName.ToString() );
-		Bodies[BodyIndices[i]]->SetInstanceSimulatePhysics(bNewSimulate);
-	}
+		if (IsSimulatingPhysics())
+		{
+			SetRootBodyIndex(RootBodyData.BodyIndex);	//Update the root body data cache in case animation has moved root body relative to root joint
+		}
 
-	if (IsSimulatingPhysics())
-	{
-		SetRootBodyIndex(RootBodyData.BodyIndex);	//Update the root body data cache in case animation has moved root body relative to root joint
+		UpdatePostPhysicsTickRegisteredState();
+		UpdateClothTickRegisteredState();
 	}
-
-	UpdatePostPhysicsTickRegisteredState();
-	UpdateClothTickRegisteredState();
 }
 
 
@@ -1242,6 +1219,26 @@ void USkeletalMeshComponent::SetAllMotorsAngularVelocityDrive(bool bEnableSwingD
 	}
 }
 
+void USkeletalMeshComponent::SetConstraintProfile(FName JointName, FName ProfileName)
+{
+	UPhysicsAsset* const PhysicsAsset = GetPhysicsAsset();
+	if (!PhysicsAsset)
+	{
+		return;
+	}
+
+	for (int32 i = 0; i < Constraints.Num(); i++)
+	{
+		FConstraintInstance* ConstraintInstance = Constraints[i];
+		if(ConstraintInstance->JointName == JointName)
+		{
+			if(UPhysicsConstraintProfile* Profile = PhysicsAsset->ConstraintSetup[i]->GetProfile(ProfileName))
+			{
+				ConstraintInstance->CopyProfilePropertiesFrom(Profile->ProfileProperties);
+			}
+		}
+	}
+}
 
 void USkeletalMeshComponent::SetAllMotorsAngularDriveParams(float InSpring, float InDamping, float InForceLimit, bool bSkipCustomPhysicsType)
 {
@@ -1343,52 +1340,36 @@ void USkeletalMeshComponent::SetAllBodiesPhysicsBlendWeight(float PhysicsBlendWe
 }
 
 
-void USkeletalMeshComponent::SetAllBodiesBelowPhysicsBlendWeight( const FName& InBoneName, float PhysicsBlendWeight, bool bSkipCustomPhysicsType )
+void USkeletalMeshComponent::SetAllBodiesBelowPhysicsBlendWeight( const FName& InBoneName, float PhysicsBlendWeight, bool bSkipCustomPhysicsType, bool bIncludeSelf )
 {
-	UPhysicsAsset* const PhysicsAsset = GetPhysicsAsset();
-	if( !PhysicsAsset || !SkeletalMesh )
+	int32 NumBodiesFound = ForEachBodyBelow(InBoneName, bIncludeSelf, bSkipCustomPhysicsType, [PhysicsBlendWeight](FBodyInstance* BI)
 	{
-		return;
-	}
+		BI->PhysicsBlendWeight = PhysicsBlendWeight;
+	});
 
-	// if physics state is invalid - i.e. collision is disabled - or it does not have a valid bodies, this will crash right away
-	if (!IsPhysicsStateCreated()  || !bHasValidBodies)
+	if (NumBodiesFound)
 	{
-		FMessageLog("PIE").Warning(LOCTEXT("InvalidBodies", "Invalid Bodies : Make sure collision is enabled or root bone has body in PhysicsAsset."));
-		return;
-	}
-
-	TArray<int32> BodyIndices;
-	PhysicsAsset->GetBodyIndicesBelow(BodyIndices, InBoneName, SkeletalMesh);
+		bBlendPhysics = false;
 
-	for(int32 i=0; i<BodyIndices.Num(); i++)
-	{
-		Bodies[BodyIndices[i]]->PhysicsBlendWeight = PhysicsBlendWeight;
+		UpdatePostPhysicsTickRegisteredState();
+		UpdateClothTickRegisteredState();
 	}
 }
 
 
 void USkeletalMeshComponent::AccumulateAllBodiesBelowPhysicsBlendWeight( const FName& InBoneName, float PhysicsBlendWeight, bool bSkipCustomPhysicsType )
 {
-	UPhysicsAsset* const PhysicsAsset = GetPhysicsAsset();
-	if( !PhysicsAsset || !SkeletalMesh )
+	int32 NumBodiesFound = ForEachBodyBelow(InBoneName, /*bIncludeSelf=*/ true, /*bSkipCustomPhysicsType=*/ bSkipCustomPhysicsType, [PhysicsBlendWeight](FBodyInstance* BI)
 	{
-		return;
-	}
+		BI->PhysicsBlendWeight = FMath::Min(BI->PhysicsBlendWeight + PhysicsBlendWeight, 1.f);
+	});
 
-	// if physics state is invalid - i.e. collision is disabled - or it does not have a valid bodies, this will crash right away
-	if (!IsPhysicsStateCreated()  || !bHasValidBodies)
+	if (NumBodiesFound)
 	{
-		FMessageLog("PIE").Warning(LOCTEXT("InvalidBodies", "Invalid Bodies : Make sure collision is enabled or root bone has body in PhysicsAsset."));
-		return;
-	}
-
-	TArray<int32> BodyIndices;
-	PhysicsAsset->GetBodyIndicesBelow(BodyIndices, InBoneName, SkeletalMesh);
+		bBlendPhysics = false;
 
-	for(int32 i=0; i<BodyIndices.Num(); i++)
-	{
-		Bodies[BodyIndices[i]]->PhysicsBlendWeight = FMath::Min(Bodies[BodyIndices[i]]->PhysicsBlendWeight + PhysicsBlendWeight, 1.f);
+		UpdatePostPhysicsTickRegisteredState();
+		UpdateClothTickRegisteredState();
 	}
 }
 
@@ -1407,6 +1388,24 @@ FConstraintInstance* USkeletalMeshComponent::FindConstraintInstance(FName ConNam
 	return NULL;
 }
 
+void USkeletalMeshComponent::AddForceToAllBodiesBelow(FVector Force, FName BoneName, bool bAccelChange, bool bIncludeSelf)
+{
+	ForEachBodyBelow(BoneName, bIncludeSelf, /*bSkipCustomPhysics=*/false, [Force, bAccelChange](FBodyInstance* BI)
+	{
+		BI->AddForce(Force, /*bAllowSubstepping=*/ true, bAccelChange);
+	}
+	);
+}
+
+void USkeletalMeshComponent::AddImpulseToAllBodiesBelow(FVector Impulse, FName BoneName, bool bVelChange, bool bIncludeSelf)
+{
+	ForEachBodyBelow(BoneName, bIncludeSelf,/*bSkipCustomPhysics=*/false, [Impulse, bVelChange](FBodyInstance* BI)
+	{
+		BI->AddImpulse(Impulse, bVelChange);
+	}
+	);
+}
+
 #ifndef OLD_FORCE_UPDATE_BEHAVIOR
 #define OLD_FORCE_UPDATE_BEHAVIOR 0
 #endif
@@ -1544,22 +1543,22 @@ void USkeletalMeshComponent::UpdateMeshForBrokenConstraints()
 					FConstraintInstance* ChildConstraintInst = FindConstraintInstance(PhysicsAssetBodySetup->BoneName);
 					if( ChildConstraintInst )
 					{
-						if( ChildConstraintInst->bLinearPositionDrive )
+						if( ChildConstraintInst->IsLinearPositionDriveEnabled() )
 						{
 							DEBUGBROKENCONSTRAINTUPDATE(UE_LOG(LogSkeletalMesh, Log, TEXT("      Turning off LinearPositionDrive."));)
 							ChildConstraintInst->SetLinearPositionDrive(false, false, false);
 						}
-						if( ChildConstraintInst->bLinearVelocityDrive )
+						if( ChildConstraintInst->IsLinearVelocityDriveEnabled() )
 						{
 							DEBUGBROKENCONSTRAINTUPDATE(UE_LOG(LogSkeletalMesh, Log, TEXT("      Turning off LinearVelocityDrive."));)
 							ChildConstraintInst->SetLinearVelocityDrive(false, false, false);
 						}
-						if( ChildConstraintInst->bAngularOrientationDrive )
+						if( ChildConstraintInst->IsAngularOrientationDriveEnabled() )
 						{
 							DEBUGBROKENCONSTRAINTUPDATE(UE_LOG(LogSkeletalMesh, Log, TEXT("      Turning off AngularPositionDrive."));)
 							ChildConstraintInst->SetAngularPositionDrive(false, false);
 						}
-						if( ChildConstraintInst->bAngularVelocityDrive )
+						if( ChildConstraintInst->IsAngularVelocityDriveEnabled() )
 						{
 							DEBUGBROKENCONSTRAINTUPDATE(UE_LOG(LogSkeletalMesh, Log, TEXT("      Turning off AngularVelocityDrive."));)
 							ChildConstraintInst->SetAngularVelocityDrive(false, false);
@@ -1653,6 +1652,60 @@ void USkeletalMeshComponent::GetWeldedBodies(TArray<FBodyInstance*> & OutWeldedB
 	}
 }
 
+int32 USkeletalMeshComponent::ForEachBodyBelow(FName BoneName, bool bIncludeSelf, bool bSkipCustomType, TFunctionRef<void(FBodyInstance*)> Func)
+{
+	if (BoneName == NAME_None && bIncludeSelf && !bSkipCustomType)
+	{
+		for (FBodyInstance* BI : Bodies)	//we want all bodies so just iterate the regular array
+		{
+			Func(BI);
+		}
+
+		return Bodies.Num();
+	}
+	else
+	{
+		UPhysicsAsset* const PhysicsAsset = GetPhysicsAsset();
+		if (!PhysicsAsset || !SkeletalMesh)
+		{
+			return 0;
+		}
+
+		// if physics state is invalid - i.e. collision is disabled - or it does not have a valid bodies, this will crash right away
+		if (!IsPhysicsStateCreated() || !bHasValidBodies)
+		{
+			FMessageLog("PIE").Warning(LOCTEXT("InvalidBodies", "Invalid Bodies : Make sure collision is enabled or root bone has body in PhysicsAsset."));
+			return 0;
+		}
+
+		TArray<int32> BodyIndices;
+		BodyIndices.Reserve(Bodies.Num());
+		PhysicsAsset->GetBodyIndicesBelow(BodyIndices, BoneName, SkeletalMesh, bIncludeSelf);
+
+		int32 NumBodiesFound = 0;
+		for (int32 BodyIdx = 0; BodyIdx < BodyIndices.Num(); BodyIdx++)
+		{
+			FBodyInstance* BI = Bodies[BodyIndices[BodyIdx]];
+			if (bSkipCustomType)
+			{
+				if (UBodySetup* PhysAssetBodySetup = PhysicsAsset->BodySetup[BodyIdx])
+				{
+					if (PhysAssetBodySetup->PhysicsType != EPhysicsType::PhysType_Default)
+					{
+						continue;
+					}
+				}
+			}
+
+			++NumBodiesFound;
+			Func(BI);
+		}
+
+		return NumBodiesFound;
+	}
+
+	return 0;
+}
 
 void USkeletalMeshComponent::BreakConstraint(FVector Impulse, FVector HitLocation, FName InBoneName)
 {
diff --git a/Engine/Source/Runtime/Engine/Private/SkeletalMeshMerge.cpp b/Engine/Source/Runtime/Engine/Private/SkeletalMeshMerge.cpp
--- a/Engine/Source/Runtime/Engine/Private/SkeletalMeshMerge.cpp
+++ b/Engine/Source/Runtime/Engine/Private/SkeletalMeshMerge.cpp
@@ -495,6 +495,10 @@ void FSkeletalMeshMerge::GenerateLODModel( int32 LODIdx )
 
 			int32 MaxColorIdx = SrcLODModel.ColorVertexBuffer.GetNumVertices();
 
+
+			TArray<FSoftSkinVertex> VertexArray;
+			SrcLODModel.GetVertices(VertexArray);
+
 			// keep track of the current base vertex index before adding any new vertices
 			// this will be needed to remap the index buffer values to the new range
 			int32 CurrentBaseVertexIndex = MergedVertexBuffer.Num();
@@ -504,12 +508,11 @@ void FSkeletalMeshMerge::GenerateLODModel( int32 LODIdx )
 				VertexDataType& DestVert = MergedVertexBuffer[MergedVertexBuffer.AddUninitialized()];
 
 				// copy from source vertex
-				const TGPUSkinVertexBase<bExtraBoneInfluencesT>* SrcBaseVert = SrcLODModel.VertexBufferGPUSkin.GetVertexPtr<bExtraBoneInfluencesT>(VertIdx);
-				DestVert.Position = SrcLODModel.VertexBufferGPUSkin.GetVertexPositionFast<bExtraBoneInfluencesT>(SrcBaseVert);
-				DestVert.TangentX = SrcBaseVert->TangentX;
-				DestVert.TangentZ = SrcBaseVert->TangentZ;
-				FMemory::Memcpy(DestVert.InfluenceBones,SrcBaseVert->InfluenceBones,sizeof(SrcBaseVert->InfluenceBones));
-				FMemory::Memcpy(DestVert.InfluenceWeights,SrcBaseVert->InfluenceWeights,sizeof(SrcBaseVert->InfluenceWeights));
+				DestVert.Position = VertexArray[VertIdx].Position;
+				DestVert.TangentX = VertexArray[VertIdx].TangentX;
+				DestVert.TangentZ = VertexArray[VertIdx].TangentZ;
+				FMemory::Memcpy(DestVert.InfluenceBones, VertexArray[VertIdx].InfluenceBones, sizeof(DestVert.InfluenceBones));
+				FMemory::Memcpy(DestVert.InfluenceWeights, VertexArray[VertIdx].InfluenceWeights, sizeof(DestVert.InfluenceWeights));
 
 				// if the mesh uses vertex colors, copy the source color if possible or default to white
 				if( MergeMesh->bHasVertexColors )
