Fix compilation on a case-sensitive filesystem.

diff --git a/Engine/Source/Editor/Persona/Private/Shared/PersonaMode.cpp b/Engine/Source/Editor/Persona/Private/Shared/PersonaMode.cpp
--- a/Engine/Source/Editor/Persona/Private/Shared/PersonaMode.cpp
+++ b/Engine/Source/Editor/Persona/Private/Shared/PersonaMode.cpp
@@ -9,7 +9,7 @@
 #include "SSkeletonSlotNames.h"
 #include "SSkeletonSmartNameManager.h"
 #include "SSkeletonBlendProfiles.h"
-#include "SDocktab.h"
+#include "SDockTab.h"
 #include "SAdvancedPreviewDetailsTab.h"
 
 #define LOCTEXT_NAMESPACE "PersonaModes"
diff --git a/Engine/Source/Runtime/Engine/Private/AdvancedPreviewScene.cpp b/Engine/Source/Runtime/Engine/Private/AdvancedPreviewScene.cpp
--- a/Engine/Source/Runtime/Engine/Private/AdvancedPreviewScene.cpp
+++ b/Engine/Source/Runtime/Engine/Private/AdvancedPreviewScene.cpp
@@ -12,7 +12,7 @@
 #include "UnrealClient.h"
 #include "Editor/EditorPerProjectUserSettings.h"
 
-#include "Components/Directionallightcomponent.h"
+#include "Components/DirectionalLightComponent.h"
 #include "Components/SphereReflectionCaptureComponent.h"
 #include "Components/SkyLightComponent.h"
 
diff --git a/Engine/Source/Runtime/Engine/Public/AssetViewerSettings.h b/Engine/Source/Runtime/Engine/Public/AssetViewerSettings.h
--- a/Engine/Source/Runtime/Engine/Public/AssetViewerSettings.h
+++ b/Engine/Source/Runtime/Engine/Public/AssetViewerSettings.h
@@ -2,7 +2,7 @@
 
 #pragma once
 
-#include "Classes/Engine/Texturecube.h"
+#include "Classes/Engine/TextureCube.h"
 
 #include "AssetViewerSettings.generated.h"
 
