Copying //UE4/Dev-Blueprints to //UE4/Dev-Main (Source: //UE4/Dev-Blueprints @ 3025888)

diff --git a/Engine/Source/Runtime/CoreUObject/Private/UObject/Enum.cpp b/Engine/Source/Runtime/CoreUObject/Private/UObject/Enum.cpp
--- a/Engine/Source/Runtime/CoreUObject/Private/UObject/Enum.cpp
+++ b/Engine/Source/Runtime/CoreUObject/Private/UObject/Enum.cpp
@@ -457,7 +457,7 @@ int32 UEnum::FindEnumRedirects(const UEnum* Enum, FName EnumEntryName)
  *
  * @return	true unless the MAX enum already exists.
  */
-bool UEnum::SetEnums(TArray<TPair<FName, uint8>>& InNames, UEnum::ECppForm InCppForm)
+bool UEnum::SetEnums(TArray<TPair<FName, uint8>>& InNames, UEnum::ECppForm InCppForm, bool bAddMaxKeyIfMissing)
 {
 	if (Names.Num() > 0)
 	{
@@ -466,29 +466,31 @@ bool UEnum::SetEnums(TArray<TPair<FName, uint8>>& InNames, UEnum::ECppForm InCpp
 	Names   = InNames;
 	CppForm = InCppForm;
 
-	const FString EnumPrefix = GenerateEnumPrefix();
-	checkSlow(EnumPrefix.Len());
+	if (bAddMaxKeyIfMissing)
+	{
+		const FString EnumPrefix = GenerateEnumPrefix();
+		checkSlow(EnumPrefix.Len());
 
-	const FName MaxEnumItem = *GenerateFullEnumName(*(EnumPrefix + TEXT("_MAX")));
-	const int32 MaxEnumItemIndex = GetIndexByName(MaxEnumItem);
+		const FName MaxEnumItem = *GenerateFullEnumName(*(EnumPrefix + TEXT("_MAX")));
+		const int32 MaxEnumItemIndex = GetIndexByName(MaxEnumItem);
 
-	if (MaxEnumItemIndex == INDEX_NONE)
-	{
-		if (LookupEnumName(MaxEnumItem) != INDEX_NONE)
+		if (MaxEnumItemIndex == INDEX_NONE)
 		{
-			// the MAX identifier is already being used by another enum
-			return false;
-		}
+			if (LookupEnumName(MaxEnumItem) != INDEX_NONE)
+			{
+				// the MAX identifier is already being used by another enum
+				return false;
+			}
 
 #if HACK_HEADER_GENERATOR
-		if (GetMaxEnumValue() == 255)
-		{
-			FError::Throwf(TEXT("Enum %s doesn't contain user defined %s_MAX entry and its maximum entry value equals 255. Autogenerated %s_MAX enum entry will equal 0. Change maximum entry value to 254 or define %s_MAX entry in enum."), *GetName(), *EnumPrefix, *EnumPrefix, *EnumPrefix, *EnumPrefix);
-		}
+			if (GetMaxEnumValue() == 255)
+			{
+				FError::Throwf(TEXT("Enum %s doesn't contain user defined %s_MAX entry and its maximum entry value equals 255. Autogenerated %s_MAX enum entry will equal 0. Change maximum entry value to 254 or define %s_MAX entry in enum."), *GetName(), *EnumPrefix, *EnumPrefix, *EnumPrefix, *EnumPrefix);
+			}
 #endif
-		Names.Add(TPairInitializer<FName, uint8>(MaxEnumItem, GetMaxEnumValue() + 1));
+			Names.Add(TPairInitializer<FName, uint8>(MaxEnumItem, GetMaxEnumValue() + 1));
+		}
 	}
-
 	AddNamesToMasterList();
 
 	return true;
@@ -567,6 +569,10 @@ FText UEnum::GetToolTipText(int32 NameIndex) const
 	return LocalizedToolTip;
 }
 
+#endif
+
+#if WITH_EDITOR || HACK_HEADER_GENERATOR
+
 /**
  * Wrapper method for easily determining whether this enum has metadata associated with it.
  * 
diff --git a/Engine/Source/Runtime/CoreUObject/Private/UObject/ScriptCore.cpp b/Engine/Source/Runtime/CoreUObject/Private/UObject/ScriptCore.cpp
--- a/Engine/Source/Runtime/CoreUObject/Private/UObject/ScriptCore.cpp
+++ b/Engine/Source/Runtime/CoreUObject/Private/UObject/ScriptCore.cpp
@@ -130,7 +130,7 @@ void FBlueprintCoreDelegates::ThrowScriptException(const UObject* ActiveObject,
 	if (IsInGameThread())
 	{
 		// If nothing is bound, show warnings so something is left in the log.
-		if (OnScriptException.IsBound() == false)
+		if (bShouldLogWarning && (OnScriptException.IsBound() == false))
 		{
 			UE_LOG(LogScript, Warning, TEXT("%s"), *StackFrame.GetStackTrace());
 		}
@@ -1141,11 +1141,18 @@ void UObject::ProcessEvent( UFunction* Function, void* Parms )
 #endif
 
 #if !(UE_BUILD_SHIPPING || UE_BUILD_TEST)
+	bool bInstrumentScriptEvent = false;
 	if (GetClass()->HasInstrumentation())
 	{
-		const EScriptInstrumentation::Type EventType = Function->HasAnyFunctionFlags(FUNC_Event|FUNC_BlueprintEvent) ? EScriptInstrumentation::Event : EScriptInstrumentation::ResumeEvent;
-		EScriptInstrumentationEvent EventInstrumentationInfo(EventType, this, Function->GetFName());
-		FBlueprintCoreDelegates::InstrumentScriptEvent(EventInstrumentationInfo);
+		// Don't instrument native events that have not been implemented/overridden in script (BP). These events will not have had any profiler data generated for them at compile time.
+		if (!Function->HasAnyFunctionFlags(FUNC_Native) || Function->GetOuter()->GetClass()->HasAnyClassFlags(CLASS_CompiledFromBlueprint))
+		{
+			const EScriptInstrumentation::Type EventType = Function->HasAnyFunctionFlags(FUNC_Event | FUNC_BlueprintEvent) ? EScriptInstrumentation::Event : EScriptInstrumentation::ResumeEvent;
+			EScriptInstrumentationEvent EventInstrumentationInfo(EventType, this, Function->GetFName());
+			FBlueprintCoreDelegates::InstrumentScriptEvent(EventInstrumentationInfo);
+
+			bInstrumentScriptEvent = true;
+		}
 	}
 #endif
 
@@ -1276,7 +1283,7 @@ void UObject::ProcessEvent( UFunction* Function, void* Parms )
 	}
 
 #if !(UE_BUILD_SHIPPING || UE_BUILD_TEST)
-	if (GetClass()->HasInstrumentation())
+	if (bInstrumentScriptEvent)
 	{
 		EScriptInstrumentationEvent EventInstrumentationInfo(EScriptInstrumentation::Stop, this, NAME_None);
 		FBlueprintCoreDelegates::InstrumentScriptEvent(EventInstrumentationInfo);
diff --git a/Engine/Source/Runtime/CoreUObject/Private/UObject/UObjectBase.cpp b/Engine/Source/Runtime/CoreUObject/Private/UObject/UObjectBase.cpp
--- a/Engine/Source/Runtime/CoreUObject/Private/UObject/UObjectBase.cpp
+++ b/Engine/Source/Runtime/CoreUObject/Private/UObject/UObjectBase.cpp
@@ -631,21 +631,13 @@ void UClassCompiledInDefer(FFieldCompiledInInfo* ClassInfo, const TCHAR* Name, S
 	GetDeferredClassRegistration().Add(ClassInfo);
 }
 
-struct FClassConstructFunctions
+TMap<FName, FDynamicClassStaticData>& GetDynamicClassMap()
 {
-	/** Autogenerated Z_Construct* function pointer */
-	UClass::StaticClassFunctionType ZConstructFn;
-	/** StaticClass() function pointer */
-	UClass::StaticClassFunctionType StaticClassFn;
-};
-
-TMap<FName, FClassConstructFunctions>& GetDynamicClassMap()
-{
-	static TMap<FName, FClassConstructFunctions> DynamicClassMap;
+	static TMap<FName, FDynamicClassStaticData> DynamicClassMap;
 	return DynamicClassMap;
 }
 
-void UObjectCompiledInDefer(UClass *(*InRegister)(), UClass *(*InStaticClass)(), const TCHAR* Name, bool bDynamic, const TCHAR* DynamicPathName)
+void UObjectCompiledInDefer(UClass *(*InRegister)(), UClass *(*InStaticClass)(), const TCHAR* Name, bool bDynamic, const TCHAR* DynamicPathName, void (*InInitSearchableValues)(TMap<FName, FName>&))
 {
 	if (!bDynamic)
 	{
@@ -660,9 +652,13 @@ void UObjectCompiledInDefer(UClass *(*InRegister)(), UClass *(*InStaticClass)(),
 	}
 	else
 	{
-		FClassConstructFunctions ClassFunctions;
+		FDynamicClassStaticData ClassFunctions;
 		ClassFunctions.ZConstructFn = InRegister;
 		ClassFunctions.StaticClassFn = InStaticClass;
+		if (InInitSearchableValues)
+		{
+			InInitSearchableValues(ClassFunctions.SelectedSearchableValues);
+		}
 		GetDynamicClassMap().Add(FName(DynamicPathName), ClassFunctions);
 	}
 }
@@ -1169,7 +1165,7 @@ UEnum* FindExistingEnumIfHotReloadOrDynamic(UObject* Outer, const TCHAR* EnumNam
 
 UClass::StaticClassFunctionType GetDynamicClassConstructFn(FName ClassPathName)
 {
-	FClassConstructFunctions* ClassConstructFn = GetDynamicClassMap().Find(ClassPathName);
+	FDynamicClassStaticData* ClassConstructFn = GetDynamicClassMap().Find(ClassPathName);
 	UE_CLOG(!ClassConstructFn, LogUObjectBase, Fatal, TEXT("Unable to find construct function pointer for dynamic class %s. Make sure dynamic class exists."), *ClassPathName.ToString());
 	if (ClassConstructFn)
 	{
@@ -1182,7 +1178,7 @@ UClass::StaticClassFunctionType GetDynamicClassConstructFn(FName ClassPathName)
 UObject* ConstructDynamicType(FName TypePathName)
 {
 	UObject* Result = nullptr;
-	if (FClassConstructFunctions* ClassConstructFn = GetDynamicClassMap().Find(TypePathName))
+	if (FDynamicClassStaticData* ClassConstructFn = GetDynamicClassMap().Find(TypePathName))
 	{
 		UClass* DynamicClass = ClassConstructFn->StaticClassFn();
 		check(DynamicClass);
diff --git a/Engine/Source/Runtime/CoreUObject/Private/UObject/UObjectGlobals.cpp b/Engine/Source/Runtime/CoreUObject/Private/UObject/UObjectGlobals.cpp
--- a/Engine/Source/Runtime/CoreUObject/Private/UObject/UObjectGlobals.cpp
+++ b/Engine/Source/Runtime/CoreUObject/Private/UObject/UObjectGlobals.cpp
@@ -2548,6 +2548,11 @@ FObjectInitializer::~FObjectInitializer()
 				// make sure we haven't already deferred this once, if we have
 				// then something is destroying this one prematurely 
 				check(bIsDeferredInitializer == false);
+
+				for (const FSubobjectsToInit::FSubobjectInit& SubObjInfo : ComponentInits.SubobjectInits)
+				{
+					check(!SubObjInfo.Subobject->HasAnyFlags(RF_NeedLoad));
+				}
 #endif // USE_DEFERRED_DEPENDENCY_CHECK_VERIFICATION_TESTS
 
 				// makes a copy of this and saves it off, to be ran later
@@ -2779,7 +2784,11 @@ bool FObjectInitializer::IsInstancingAllowed() const
 
 bool FObjectInitializer::InitSubobjectProperties(bool bAllowInstancing) const
 {
+#if USE_CIRCULAR_DEPENDENCY_LOAD_DEFERRING
+	bool bNeedSubobjectInstancing = bAllowInstancing && bIsDeferredInitializer;
+#else 
 	bool bNeedSubobjectInstancing = false;
+#endif
 	// initialize any subobjects, now that the constructors have run
 	for (int32 Index = 0; Index < ComponentInits.SubobjectInits.Num(); Index++)
 	{
@@ -2821,7 +2830,12 @@ void FObjectInitializer::InstanceSubobjects(UClass* Class, bool bNeedInstancing,
 		{
 			UObject* Subobject = ComponentInits.SubobjectInits[Index].Subobject;
 			UObject* Template = ComponentInits.SubobjectInits[Index].Template;
-			if (!Subobject->HasAnyFlags(RF_NeedLoad))
+
+#if USE_CIRCULAR_DEPENDENCY_LOAD_DEFERRING
+			if ( !Subobject->HasAnyFlags(RF_NeedLoad) || bIsDeferredInitializer )
+#else 
+			if ( !Subobject->HasAnyFlags(RF_NeedLoad) )
+#endif
 			{
 				Subobject->GetClass()->InstanceSubobjectTemplates(Subobject, Template, Template->GetClass(), Subobject, UseInstancingGraph);
 			}
@@ -2896,6 +2910,9 @@ void FObjectInitializer::InitProperties(UObject* Obj, UClass* DefaultsClass, UOb
 	{
 		QUICK_SCOPE_CYCLE_COUNTER(STAT_InitProperties_Blueprint);
 
+		// As with native classes, we must iterate through all properties (slow path) if default data is pointing at something other than the CDO.
+		bCanUsePostConstructLink &= (DefaultData == Class->GetDefaultObject(false));
+
 		UObject* ClassDefaults = bCopyTransientsFromClassDefaults ? DefaultsClass->GetDefaultObject() : NULL;		
 		for (UProperty* P = bCanUsePostConstructLink ? Class->PostConstructLink : Class->PropertyLink; P; P = bCanUsePostConstructLink ? P->PostConstructLinkNext : P->PropertyLinkNext)
 		{
diff --git a/Engine/Source/Runtime/CoreUObject/Private/UObject/UTextProperty.cpp b/Engine/Source/Runtime/CoreUObject/Private/UObject/UTextProperty.cpp
--- a/Engine/Source/Runtime/CoreUObject/Private/UObject/UTextProperty.cpp
+++ b/Engine/Source/Runtime/CoreUObject/Private/UObject/UTextProperty.cpp
@@ -37,36 +37,40 @@ bool UTextProperty::ConvertFromType(const FPropertyTag& Tag, FArchive& Ar, uint8
 	return bOutAdvanceProperty;
 }
 
+bool UTextProperty::Identical_Implementation(const FText& ValueA, const FText& ValueB, uint32 PortFlags)
+{
+	if (ValueA.IsCultureInvariant() != ValueB.IsCultureInvariant() || ValueA.IsTransient() != ValueB.IsTransient())
+	{
+		//A culture variant text is never equal to a culture invariant text
+		//A transient text is never equal to a non-transient text
+		return false;
+	}
+
+	if (ValueA.IsCultureInvariant() == ValueB.IsCultureInvariant() || ValueA.IsTransient() == ValueB.IsTransient())
+	{
+		//Culture invariant text don't have a namespace/key so we compare the source string
+		//Transient text don't have a namespace/key or source so we compare the display string
+		return FTextInspector::GetDisplayString(ValueA) == FTextInspector::GetDisplayString(ValueB);
+	}
+
+	if (GIsEditor)
+	{
+		return FTextInspector::GetSourceString(ValueA)->Compare(*FTextInspector::GetSourceString(ValueB), ESearchCase::CaseSensitive) == 0;
+	}
+	else
+	{
+		return	FTextInspector::GetNamespace(ValueA) == FTextInspector::GetNamespace(ValueB) &&
+			FTextInspector::GetKey(ValueA) == FTextInspector::GetKey(ValueB);
+	}
+}
+
 bool UTextProperty::Identical( const void* A, const void* B, uint32 PortFlags ) const
 {
 	const TCppType ValueA = GetPropertyValue(A);
 	if ( B )
 	{
 		const TCppType ValueB = GetPropertyValue(B);
-
-		if ( ValueA.IsCultureInvariant() != ValueB.IsCultureInvariant() || ValueA.IsTransient() != ValueB.IsTransient() )
-		{
-			//A culture variant text is never equal to a culture invariant text
-			//A transient text is never equal to a non-transient text
-			return false;
-		}
-		
-		if ( ValueA.IsCultureInvariant() == ValueB.IsCultureInvariant() || ValueA.IsTransient() == ValueB.IsTransient() )
-		{
-			//Culture invariant text don't have a namespace/key so we compare the source string
-			//Transient text don't have a namespace/key or source so we compare the display string
-			return FTextInspector::GetDisplayString(ValueA) == FTextInspector::GetDisplayString(ValueB);
-		}
-
-		if (GIsEditor)
-		{
-			return FTextInspector::GetSourceString(ValueA)->Compare(*FTextInspector::GetSourceString(ValueB), ESearchCase::CaseSensitive) == 0;
-		}
-		else
-		{
-			return	FTextInspector::GetNamespace(ValueA)	==	FTextInspector::GetNamespace(ValueB) &&
-					FTextInspector::GetKey(ValueA)			==	FTextInspector::GetKey(ValueB);
-		}
+		return Identical_Implementation(ValueA, ValueB, PortFlags);
 	}
 
 	return FTextInspector::GetDisplayString(ValueA).IsEmpty();
diff --git a/Engine/Source/Runtime/CoreUObject/Public/UObject/Class.h b/Engine/Source/Runtime/CoreUObject/Public/UObject/Class.h
--- a/Engine/Source/Runtime/CoreUObject/Public/UObject/Class.h
+++ b/Engine/Source/Runtime/CoreUObject/Public/UObject/Class.h
@@ -1588,9 +1588,10 @@ public:
 	 *
 	 * @param InNames List of enum names.
 	 * @param InCppForm The form of enum.
+	 * @param bAddMaxKeyIfMissing Should a default Max item be added.
 	 * @return	true unless the MAX enum already exists and isn't the last enum.
 	 */
-	virtual bool SetEnums(TArray<TPair<FName, uint8>>& InNames, ECppForm InCppForm);
+	virtual bool SetEnums(TArray<TPair<FName, uint8>>& InNames, ECppForm InCppForm, bool bAddMaxKeyIfMissing = true);
 
 	/**
 	 * @return	The enum name at the specified Index.
@@ -1678,7 +1679,9 @@ public:
 	 * @return The tooltip for this object.
 	 */
 	FText GetToolTipText(int32 NameIndex=INDEX_NONE) const;
+#endif
 
+#if WITH_EDITOR || HACK_HEADER_GENERATOR
 	/**
 	 * Wrapper method for easily determining whether this enum has metadata associated with it.
 	 * 
diff --git a/Engine/Source/Runtime/CoreUObject/Public/UObject/UObjectBase.h b/Engine/Source/Runtime/CoreUObject/Public/UObject/UObjectBase.h
--- a/Engine/Source/Runtime/CoreUObject/Public/UObject/UObjectBase.h
+++ b/Engine/Source/Runtime/CoreUObject/Public/UObject/UObjectBase.h
@@ -309,17 +309,17 @@ struct TClassCompiledInDefer : public FFieldCompiledInInfo
 /**
  * Stashes the singleton function that builds a compiled in class. Later, this is executed.
  */
-COREUOBJECT_API void UObjectCompiledInDefer(class UClass *(*InRegister)(), class UClass *(*InStaticClass)(), const TCHAR* Name, bool bDynamic, const TCHAR* DynamicPathName);
+COREUOBJECT_API void UObjectCompiledInDefer(class UClass *(*InRegister)(), class UClass *(*InStaticClass)(), const TCHAR* Name, bool bDynamic, const TCHAR* DynamicPathName, void (*InInitSearchableValues)(TMap<FName, FName>&));
 
 struct FCompiledInDefer
 {
-	FCompiledInDefer(class UClass *(*InRegister)(), class UClass *(*InStaticClass)(), const TCHAR* Name, bool bDynamic, const TCHAR* DynamicPackageName = nullptr, const TCHAR* DynamicPathName = nullptr)
+	FCompiledInDefer(class UClass *(*InRegister)(), class UClass *(*InStaticClass)(), const TCHAR* Name, bool bDynamic, const TCHAR* DynamicPackageName = nullptr, const TCHAR* DynamicPathName = nullptr, void (*InInitSearchableValues)(TMap<FName, FName>&) = nullptr)
 	{
 		if (bDynamic)
 		{
 			GetConvertedDynamicPackageNameToTypeName().Add(FName(DynamicPackageName), FName(Name));
 		}
-		UObjectCompiledInDefer(InRegister, InStaticClass, Name, bDynamic, DynamicPathName);
+		UObjectCompiledInDefer(InRegister, InStaticClass, Name, bDynamic, DynamicPathName, InInitSearchableValues);
 	}
 };
 
diff --git a/Engine/Source/Runtime/CoreUObject/Public/UObject/UObjectGlobals.h b/Engine/Source/Runtime/CoreUObject/Public/UObject/UObjectGlobals.h
--- a/Engine/Source/Runtime/CoreUObject/Public/UObject/UObjectGlobals.h
+++ b/Engine/Source/Runtime/CoreUObject/Public/UObject/UObjectGlobals.h
@@ -1931,6 +1931,18 @@ COREUOBJECT_API UPackage* FindOrConstructDynamicTypePackage(const TCHAR* Package
 /** Get names of "virtual" packages, that contain Dynamic types  */
 COREUOBJECT_API TMap<FName, FName>& GetConvertedDynamicPackageNameToTypeName();
 
+struct COREUOBJECT_API FDynamicClassStaticData
+{
+	/** Autogenerated Z_Construct* function pointer */
+	UClass* (*ZConstructFn)();
+	/** StaticClass() function pointer */
+	UClass* (*StaticClassFn)();
+	/** Selected AssetRegistrySearchable values */
+	TMap<FName, FName> SelectedSearchableValues;
+};
+
+COREUOBJECT_API TMap<FName, FDynamicClassStaticData>& GetDynamicClassMap();
+
 #if WITH_EDITOR
 /** 
  * Returns if true if the object is editor-only:
diff --git a/Engine/Source/Runtime/CoreUObject/Public/UObject/UTextProperty.h b/Engine/Source/Runtime/CoreUObject/Public/UObject/UTextProperty.h
--- a/Engine/Source/Runtime/CoreUObject/Public/UObject/UTextProperty.h
+++ b/Engine/Source/Runtime/CoreUObject/Public/UObject/UTextProperty.h
@@ -39,4 +39,6 @@ public:
 
 	/** Generate the correct C++ code for the given text value */
 	static FString GenerateCppCodeForTextValue(const FText& InValue, const FString& Indent);
+
+	static bool Identical_Implementation(const FText& A, const FText& B, uint32 PortFlags);
 };
\ No newline at end of file
diff --git a/Engine/Source/Runtime/Engine/Classes/EdGraph/EdGraphNode.h b/Engine/Source/Runtime/Engine/Classes/EdGraph/EdGraphNode.h
--- a/Engine/Source/Runtime/Engine/Classes/EdGraph/EdGraphNode.h
+++ b/Engine/Source/Runtime/Engine/Classes/EdGraph/EdGraphNode.h
@@ -102,9 +102,13 @@ class ENGINE_API UEdGraphNode : public UObject
 {
 	GENERATED_UCLASS_BODY()
 
+	virtual ~UEdGraphNode();
+
+	TArray<UEdGraphPin*> Pins;
+
 	/** List of connector pins */
 	UPROPERTY()
-	TArray<class UEdGraphPin*> Pins;
+	TArray<class UEdGraphPin_Deprecated*> DeprecatedPins;
 
 	/** X position of node in the editor */
 	UPROPERTY()
@@ -219,6 +223,9 @@ public:
 	virtual void PreSave(const class ITargetPlatform* TargetPlatform) override;
 	virtual void PostLoad() override;
 	virtual void PostEditChangeProperty(FPropertyChangedEvent& PropertyChangedEvent) override;
+	virtual void PostEditUndo() override;
+	virtual void ExportCustomProperties(FOutputDevice& Out, uint32 Indent) override;
+	virtual void ImportCustomProperties(const TCHAR* SourceText, FFeedbackContext* Warn) override;
 	// End of UObject interface
 
 	/** widget representing this node if it exists; Note: This is not safe to use in general and will be removed in the future, as there is no guarantee that only one graph editor/panel is viewing a given graph */
@@ -230,11 +237,8 @@ public:
 	/** Create a new pin on this node using the supplied pin type, and return the new pin */
 	UEdGraphPin* CreatePin(EEdGraphPinDirection Dir, const FEdGraphPinType& InPinType, const FString& PinName, int32 Index = INDEX_NONE);
 
-	// Allocates a pin from the pool
-	static UEdGraphPin* AllocatePinFromPool(UEdGraphNode* OuterNode);
-
-	// Returns the specified pin to the pool
-	static void ReturnPinToPool(UEdGraphPin* OldPin);
+	/** Destroys the specified pin, does not modify its owning pin's Pins list */
+	static void DestroyPin(UEdGraphPin* Pin);
 
 	/** Find a pin on this node with the supplied name */
 	UEdGraphPin* FindPin(const FString& PinName) const;
diff --git a/Engine/Source/Runtime/Engine/Classes/EdGraph/EdGraphPin.h b/Engine/Source/Runtime/Engine/Classes/EdGraph/EdGraphPin.h
--- a/Engine/Source/Runtime/Engine/Classes/EdGraph/EdGraphPin.h
+++ b/Engine/Source/Runtime/Engine/Classes/EdGraph/EdGraphPin.h
@@ -165,104 +165,137 @@
 	BPST_VariantA UMETA(DisplayName="Directional Circles")
 };
 
-UCLASS(MinimalAPI)
-class UEdGraphPin : public UObject
+USTRUCT()
+struct ENGINE_API FEdGraphPinReference
 {
-	GENERATED_UCLASS_BODY()
+	GENERATED_USTRUCT_BODY()
 
-	/** Name of this pin */
+	/** Constructors */
+	FEdGraphPinReference() : OwningNode(nullptr) {}
+	FEdGraphPinReference(UEdGraphPin* InPin) : OwningNode(nullptr) { SetPin(InPin); }
+	FEdGraphPinReference(const UEdGraphPin* InPin) : OwningNode(nullptr) { SetPin(InPin); }
+
+	/** Sets the pin referred to by this struct */
+	void SetPin(const UEdGraphPin* NewPin);
+
+	/** Gets the pin referred to by this struct */
+	UEdGraphPin* Get() const;
+
+	friend uint32 GetTypeHash(const FEdGraphPinReference& EdGraphPinReference)
+	{
+		UEdGraphNode* ResolvedOwningNode = EdGraphPinReference.OwningNode.Get();
+		ensureMsgf(ResolvedOwningNode || !EdGraphPinReference.PinId.IsValid(), TEXT("Trying to reference an unowned pin: %s"), *EdGraphPinReference.PinId.ToString());
+		uint32 NodeHash = ResolvedOwningNode ? FCrc::StrCrc32(*ResolvedOwningNode->GetName()) : 0;
+		return FCrc::StrCrc32(*EdGraphPinReference.PinId.ToString(), 0 );
+	}
+
+	bool operator==(const FEdGraphPinReference& Other) const
+	{
+		return PinId == Other.PinId && OwningNode == Other.OwningNode;
+	}
+
+private:
+	/** The node that owns the pin referred to by this struct. Updated at Set and Save time. */
 	UPROPERTY()
+	TWeakObjectPtr<UEdGraphNode> OwningNode;
+
+	/** The pin's unique ID. Updated at Set and Save time. */
+	UPROPERTY()
+	FGuid PinId;
+};
+
+enum class EPinResolveType : uint8;
+
+class UEdGraphPin
+{
+	friend struct FEdGraphPinReference;
+private:
+	/** The node that owns this pin. */
+	UEdGraphNode* OwningNode;
+
+
+public:
+	/** The pin's unique ID. */
+	FGuid PinId;
+
+	/** Name of this pin. */
 	FString PinName;
 
 #if WITH_EDITORONLY_DATA
-	/** Used as the display name if set */
-	UPROPERTY()
+	/** Used as the display name if set. */
 	FText PinFriendlyName;
 #endif
 
 	/** The tool-tip describing this pin's purpose */
-	UPROPERTY()
 	FString PinToolTip;
 
 	/** Direction of flow of this pin (input or output) */
-	UPROPERTY()
 	TEnumAsByte<enum EEdGraphPinDirection> Direction;
 
 	/** The type of information carried on this pin */
-	UPROPERTY()
-	struct FEdGraphPinType PinType;
+	FEdGraphPinType PinType;
 
 	/** Default value for this pin (used if the pin has no connections), stored as a string */
-	UPROPERTY()
 	FString DefaultValue;
 
 	/** Initial default value (the autogenerated value, to identify if the user has modified the value), stored as a string */
-	UPROPERTY()
 	FString AutogeneratedDefaultValue;
 
 	/** If the default value for this pin should be an object, we store a pointer to it */
-	UPROPERTY()
 	class UObject* DefaultObject;
 
 	/** If the default value for this pin should be an FText, it is stored here. */
-	UPROPERTY()
 	FText DefaultTextValue;
 
 	/** Set of pins that we are linked to */
-	UPROPERTY()
-	TArray<class UEdGraphPin*> LinkedTo;
+	TArray<UEdGraphPin*> LinkedTo;
 
 	/** The pins created when a pin is split and hidden */ 
-	UPROPERTY()
-	TArray<class UEdGraphPin*> SubPins;
+	TArray<UEdGraphPin*> SubPins;
 
 	/** The pin that was split and generated this pin */
-	UPROPERTY()
 	UEdGraphPin* ParentPin;
 
 	/** Pin that this pin uses for passing through reference connection */
-	UPROPERTY()
 	UEdGraphPin* ReferencePassThroughConnection;
 
 #if WITH_EDITORONLY_DATA
+	/** Pin name could be changed, so whenever possible it's good to have a persistent GUID identifying Pin to reconstruct Node seamlessly */
+	FGuid PersistentGuid;
+
+////////////////////////////////////////////////////////////////////////////////////
+// ONLY PUT BITFIELD PROPERTIES AFTER THIS TO ENSURE GOOD MEMORY ALIGNMENT
+////////////////////////////////////////////////////////////////////////////////////
+
 	/** If true, this connector is currently hidden. */
-	UPROPERTY()
 	uint32 bHidden:1;
 
 	/** If true, this connector is unconnectable, and present only to allow the editing of the default text. */
-	UPROPERTY()
 	uint32 bNotConnectable:1;
 
-	/** If true, the default value of this connector is fixed and cannot be modified by the user (it's visible for reference only) */
-	UPROPERTY()
+	/** If true, the default value of this connector is fixed and cannot be modified by the user (it's visible for reference only). */
 	uint32 bDefaultValueIsReadOnly:1;
 
-	/** If true, the default value on this pin is ignored and should not be set */
-	UPROPERTY()
+	/** If true, the default value on this pin is ignored and should not be set. */
 	uint32 bDefaultValueIsIgnored:1;
 
-	/** If true, this pin is the focus of a diff */
-	UPROPERTY(transient)
+	/** If true, this pin is the focus of a diff. This is transient. */
 	uint32 bIsDiffing:1;
 
-	/** If true, the pin may be hidden by user */
-	UPROPERTY()
+	/** If true, the pin may be hidden by user. */
 	uint32 bAdvancedView:1;
 
-	/** If true, the pin is displayed as ref */
-	UPROPERTY(transient)
-	uint32 bDisplayAsMutableRef : 1;
-
-	/** Pin name could be changed, so whenever possible it's good to have a persistent GUID identifying Pin to reconstruct Node seamlessly */
-	UPROPERTY()
-	FGuid PersistentGuid;
+	/** If true, the pin is displayed as ref. This is transient. */
+	uint32 bDisplayAsMutableRef:1;
 #endif // WITH_EDITORONLY_DATA
 
+	/** True when InvalidateAndTrash was called. This pin is intended to be discarded and destroyed. */
+	uint32 bWasTrashed:1;
+
 public:
-	// UObject interface
-	virtual bool IsSafeForRootSet() const override { return false; }
-	virtual void PostLoad() override;
-	// End UObject interface
 
+	/** Creates a new pin. Should be called from the OwningNode so it can be immediately added to the Pins array. */
+	ENGINE_API static UEdGraphPin* CreatePin(UEdGraphNode* InOwningNode);
+
 	/** Create a link. Note, this does not check that schema allows it, and will not break any existing connections */
 	ENGINE_API void MakeLinkTo(UEdGraphPin* ToPin);
@@ -270,15 +303,6 @@
 	/** Break a link to the specified pin (if present) */
 	ENGINE_API void BreakLinkTo(UEdGraphPin* ToPin);
 
-	/**
-	 * Generate a string detailing the link this pin has to another pin.
-	 * 
-	 * @Param	InFunctionName	String with function name requesting the info
-	 * @Param	InInfoData		String detailing the info (EG. Is Not linked to)
-	 * @Param	InToPin			The relevant pin
-	 */
-	const FString GetLinkInfoString( const FString& InFunctionName, const FString& InInfoData, const UEdGraphPin* InToPin ) const;
-
 	/** Break all links from this pin */
 	ENGINE_API void BreakAllPinLinks();
 
@@ -293,15 +317,8 @@
 	ENGINE_API void AssignByRefPassThroughConnection(UEdGraphPin* InTargetPin);
 
 	/** Returns the node that owns this pin */
-	class UEdGraphNode* GetOwningNode() const
-	{
-		return CastChecked<UEdGraphNode>(GetOuter());
-	}
+	FORCEINLINE class UEdGraphNode* GetOwningNode() const { check(OwningNode); return OwningNode; }
+	FORCEINLINE class UEdGraphNode* GetOwningNodeUnchecked() const { return OwningNode; }
 
-	class UEdGraphNode* GetOwningNodeUnchecked() const
-	{
-		return Cast<UEdGraphNode>(GetOuter());
-	}
-
 	/** Shorthand way to access the schema of the graph that owns the node that owns this pin */
 	ENGINE_API const class UEdGraphSchema* GetSchema() const;
@@ -331,6 +348,15 @@
 	ENGINE_API FText GetDisplayName() const;
 #endif // WITH_EDITORONLY_DATA
 
+	/**
+	 * Generate a string detailing the link this pin has to another pin.
+	 * 
+	 * @Param	InFunctionName	String with function name requesting the info
+	 * @Param	InInfoData		String detailing the info (EG. Is Not linked to)
+	 * @Param	InToPin			The relevant pin
+	 */
+	const FString GetLinkInfoString( const FString& InFunctionName, const FString& InInfoData, const UEdGraphPin* InToPin ) const;
+
 	void ResetDefaultValue()
 	{
 		DefaultValue.Empty();
@@ -358,6 +384,176 @@
 		bDefaultValueIsIgnored = false;
 #endif // WITH_EDITORONLY_DATA
 	}
+
+	/** Provides a reference collector with all object references this pin has. Should only be called by the owning node. */
+	void AddStructReferencedObjects(class FReferenceCollector& Collector);
+
+	/** Serializes an array of pins as the owner. Only the OwningNode should call this function. */
+	static void SerializeAsOwningNode(FArchive& Ar, TArray<UEdGraphPin*>& ArrayRef);
+
+	/** Marks the owning node as modified. */
+	ENGINE_API bool Modify(bool bAlwaysMarkDirty = true);
+
+	/** Changes the owning node. This will remove the pin from the old owning node's pin list and add itself to the new node's pin list. */
+	ENGINE_API void SetOwningNode(UEdGraphNode* NewOwningNode);
+
+	/** Marks the pin as 'trashed'. Removed the pin from the Owning Node's Pins list */
+	ENGINE_API void InvalidateAndTrash();
+
+	/** Marks the pin as 'trashed'. *Does not* remove the pin from the Owning Node's Pins list */
+	ENGINE_API void MarkPendingKill();
+
+	/** Returns true if InvalidateAndTrash was ever called on this pin. */
+	FORCEINLINE bool WasTrashed() const { return bWasTrashed; }
+
+	/** Transition support for UEdGraphPins */
+	ENGINE_API static UEdGraphPin* CreatePinFromDeprecatedPin(class UEdGraphPin_Deprecated* DeprecatedPin);
+	ENGINE_API static UEdGraphPin* FindPinCreatedFromDeprecatedPin(class UEdGraphPin_Deprecated* DeprecatedPin);
+
+	/** ExportText/ImportText */
+	ENGINE_API bool ExportTextItem(FString& ValueStr, int32 PortFlags) const;
+	ENGINE_API bool ImportTextItem(const TCHAR*& Buffer, int32 PortFlags, class UObject* Parent, FOutputDevice* ErrorText);
+
+	ENGINE_API const FString& GetName() const { return PinName; }
+	ENGINE_API UEdGraphNode* GetOuter() const { return GetOwningNode(); }
+	ENGINE_API bool IsPendingKill() const {	return bWasTrashed; }
+
+	/** Verification that all pins have been destroyed after shutting down */
+	ENGINE_API static void ShutdownVerification();
+
+	ENGINE_API static void Purge();
+
+	/** Destructor */
+	~UEdGraphPin();
+
+	/** This needs to be called if you want to use pin data within PostEditUndo */
+	ENGINE_API static void ResolveAllPinReferences();
+private:
+	/** 
+		Unsafe helper function used to deal with odd design of transaction buffer, 
+		which immediately applies changes to some records before saving state of others 
+	  */
+	void QueueForDestruction();
+
+	/** Private Constructor. Create pins using CreatePin since all pin instances are managed by TSharedPtr. */
+	UEdGraphPin(UEdGraphNode* InOwningNode, const FGuid& PinGuid);
+
+	/** Backward compatibility code to populate this pin with data from the supplied deprecated UEdGraphPin. */
+	void InitFromDeprecatedPin(class UEdGraphPin_Deprecated* DeprecatedPin);
+
+	bool Serialize(FArchive& Ar);
+
+	// Helper functions
+	static void EnableAllConnectedNodes(class UEdGraphNode* InNode);
+	static void ResolveReferencesToPin(UEdGraphPin* Pin);
+	static void SerializePinArray(FArchive& Ar, TArray<UEdGraphPin*>& ArrayRef, UEdGraphPin* RequestingPin, EPinResolveType ResolveType);
+	static bool SerializePin(FArchive& Ar, UEdGraphPin*& PinRef, int32 ArrayIdx, UEdGraphPin* RequestingPin, EPinResolveType ResolveType, TArray<UEdGraphPin*>& OldPins);
+	static FString ExportText_PinReference(const UEdGraphPin* Pin);
+	static FString ExportText_PinArray(const TArray<UEdGraphPin*>& PinArray);
+	static bool ImportText_PinReference(const TCHAR*& Buffer, UEdGraphPin*& PinRef, int32 ArrayIdx, UEdGraphPin* RequestingPin, EPinResolveType ResolveType);
+	static bool ImportText_PinArray(const TCHAR*& Buffer, TArray<UEdGraphPin*>& ArrayRef, UEdGraphPin* RequestingPin, EPinResolveType ResolveType);
+};
+
+UCLASS(MinimalAPI)
+class UEdGraphPin_Deprecated : public UObject
+{
+	GENERATED_UCLASS_BODY()
+
+	/** Name of this pin */
+	UPROPERTY()
+	FString PinName;
+
+#if WITH_EDITORONLY_DATA
+	/** Used as the display name if set */
+	UPROPERTY()
+	FText PinFriendlyName;
+#endif
+
+	/** The tool-tip describing this pin's purpose */
+	UPROPERTY()
+	FString PinToolTip;
+
+	/** Direction of flow of this pin (input or output) */
+	UPROPERTY()
+	TEnumAsByte<enum EEdGraphPinDirection> Direction;
+
+	/** The type of information carried on this pin */
+	UPROPERTY()
+	struct FEdGraphPinType PinType;
+
+	/** Default value for this pin (used if the pin has no connections), stored as a string */
+	UPROPERTY()
+	FString DefaultValue;
+
+	/** Initial default value (the autogenerated value, to identify if the user has modified the value), stored as a string */
+	UPROPERTY()
+	FString AutogeneratedDefaultValue;
+
+	/** If the default value for this pin should be an object, we store a pointer to it */
+	UPROPERTY()
+	class UObject* DefaultObject;
+
+	/** If the default value for this pin should be an FText, it is stored here. */
+	UPROPERTY()
+	FText DefaultTextValue;
+
+	/** Set of pins that we are linked to */
+	UPROPERTY()
+	TArray<class UEdGraphPin_Deprecated*> LinkedTo;
+
+	/** The pins created when a pin is split and hidden */ 
+	UPROPERTY()
+	TArray<class UEdGraphPin_Deprecated*> SubPins;
+
+	/** The pin that was split and generated this pin */
+	UPROPERTY()
+	UEdGraphPin_Deprecated* ParentPin;
+
+	/** Pin that this pin uses for passing through reference connection */
+	UPROPERTY()
+	UEdGraphPin_Deprecated* ReferencePassThroughConnection;
+
+#if WITH_EDITORONLY_DATA
+	/** If true, this connector is currently hidden. */
+	UPROPERTY()
+	uint32 bHidden:1;
+
+	/** If true, this connector is unconnectable, and present only to allow the editing of the default text. */
+	UPROPERTY()
+	uint32 bNotConnectable:1;
+
+	/** If true, the default value of this connector is fixed and cannot be modified by the user (it's visible for reference only) */
+	UPROPERTY()
+	uint32 bDefaultValueIsReadOnly:1;
+
+	/** If true, the default value on this pin is ignored and should not be set */
+	UPROPERTY()
+	uint32 bDefaultValueIsIgnored:1;
+
+	/** If true, this pin is the focus of a diff */
+	UPROPERTY(transient)
+	uint32 bIsDiffing:1;
+
+	/** If true, the pin may be hidden by user */
+	UPROPERTY()
+	uint32 bAdvancedView:1;
+
+	/** If true, the pin is displayed as ref */
+	UPROPERTY(transient)
+	uint32 bDisplayAsMutableRef : 1;
+
+	/** Pin name could be changed, so whenever possible it's good to have a persistent GUID identifying Pin to reconstruct Node seamlessly */
+	UPROPERTY()
+	FGuid PersistentGuid;
+#endif // WITH_EDITORONLY_DATA
+
+public:
+	// UObject interface
+	virtual bool IsSafeForRootSet() const override { return false; }
+	// End UObject interface
+
+	/** Legacy fix up for a bug in older EdGraphPins */
+	void FixupDefaultValue();
 };
 
 
diff --git a/Engine/Source/Runtime/Engine/Classes/Engine/Blueprint.h b/Engine/Source/Runtime/Engine/Classes/Engine/Blueprint.h
--- a/Engine/Source/Runtime/Engine/Classes/Engine/Blueprint.h
+++ b/Engine/Source/Runtime/Engine/Classes/Engine/Blueprint.h
@@ -45,7 +45,7 @@ enum EBlueprintType
 	BPTYPE_Interface			UMETA(DisplayName="Blueprint Interface"),
 	/** Blueprint that handles level scripting. */
 	BPTYPE_LevelScript			UMETA(DisplayName="Level Blueprint"),
-	/** Blueprint that servers as a container for functions to be used in other blueprints. */
+	/** Blueprint that serves as a container for functions to be used in other blueprints. */
 	BPTYPE_FunctionLibrary		UMETA(DisplayName="Blueprint Function Library"),
 
 	BPTYPE_MAX,
@@ -451,7 +451,10 @@ class ENGINE_API UBlueprint : public UBlueprintCore
 	TArray<class UBreakpoint*> Breakpoints;
 
 	UPROPERTY()
-	TArray<class UEdGraphPin*> PinWatches;
+	TArray<FEdGraphPinReference> WatchedPins;
+
+	UPROPERTY()
+	TArray<class UEdGraphPin_Deprecated*> DeprecatedPinWatches;
 #endif // WITH_EDITORONLY_DATA
 
 public:
diff --git a/Engine/Source/Runtime/Engine/Classes/Engine/BlueprintGeneratedClass.h b/Engine/Source/Runtime/Engine/Classes/Engine/BlueprintGeneratedClass.h
--- a/Engine/Source/Runtime/Engine/Classes/Engine/BlueprintGeneratedClass.h
+++ b/Engine/Source/Runtime/Engine/Classes/Engine/BlueprintGeneratedClass.h
@@ -53,10 +53,10 @@ public:
 	TMultiMap< int32, TWeakObjectPtr<UEdGraphNode> > LineNumberToMacroInstanceNodeMap;
 
 	// Reverse map from code offset to source pin
-	TMap< int32, TWeakObjectPtr<UEdGraphPin> > LineNumberToSourcePinMap;
+	TMap< int32, FEdGraphPinReference > LineNumberToSourcePinMap;
 
 	// Reverse map from source pin to mapped code offset(s)
-	TMultiMap< TWeakObjectPtr<UEdGraphPin>, int32 > SourcePinToLineNumbersMap;
+	TMultiMap< FEdGraphPinReference, int32 > SourcePinToLineNumbersMap;
 
 	// Map from source node (impure) to pure node script code range
 	TMap< TWeakObjectPtr<UEdGraphNode>, FInt32Range > PureNodeScriptCodeRangeMap;
@@ -200,8 +200,11 @@ protected:
 	// Acceleration structure for execution wire highlighting at runtime
 	TMap<TWeakObjectPtr<UFunction>, FDebuggingInfoForSingleFunction> PerFunctionLineNumbers;
 
+	// Map from objects to class properties they created
+	TMap<TWeakObjectPtr<UObject>, UProperty*> DebugObjectToPropertyMap;
+
 	// Map from pins or nodes to class properties they created
-	TMap<TWeakObjectPtr<UObject>, UProperty*> DebugPinToPropertyMap;
+	TMap<FEdGraphPinReference, UProperty*> DebugPinToPropertyMap;
 
 public:
 
@@ -345,6 +348,11 @@ public:
 	// Looks thru the debugging data for any class variables associated with the node
 	UProperty* FindClassPropertyForPin(const UEdGraphPin* Pin) const
 	{
+		if (!Pin)
+		{
+			return nullptr;
+		}
+
 		UProperty* PropertyPtr = DebugPinToPropertyMap.FindRef(Pin);
 		if ((PropertyPtr == nullptr) && (Pin->LinkedTo.Num() > 0))
 		{
@@ -358,7 +366,7 @@ public:
 	// Looks thru the debugging data for any class variables associated with the node (e.g., temporary variables or timelines)
 	UProperty* FindClassPropertyForNode(const UEdGraphNode* Node) const
 	{
-		return DebugPinToPropertyMap.FindRef(Node);
+		return DebugObjectToPropertyMap.FindRef(Node);
 	}
 
 	// Adds a debug record for a source node and destination in the bytecode of a specified function
@@ -401,7 +409,15 @@ public:
 	// Registers an association between an object (pin or node typically) and an associated class member property
 	void RegisterClassPropertyAssociation(class UObject* TrueSourceObject, class UProperty* AssociatedProperty)
 	{
-		DebugPinToPropertyMap.Add(TrueSourceObject, AssociatedProperty);
+		DebugObjectToPropertyMap.Add(TrueSourceObject, AssociatedProperty);
+	}
+
+	void RegisterClassPropertyAssociation(const UEdGraphPin* TrueSourcePin, class UProperty* AssociatedProperty)
+	{
+		if (TrueSourcePin)
+		{
+			DebugPinToPropertyMap.Add(TrueSourcePin, AssociatedProperty);
+		}
 	}
 
 	// Registers an association between a UUID and a node
@@ -413,7 +429,20 @@ public:
 	// Returns the object that caused the specified property to be created (can return nullptr if the association is unknown)
 	UObject* FindObjectThatCreatedProperty(class UProperty* AssociatedProperty) const
 	{
-		if (const TWeakObjectPtr<UObject>* pValue = DebugPinToPropertyMap.FindKey(AssociatedProperty))
+		if (const TWeakObjectPtr<UObject>* pValue = DebugObjectToPropertyMap.FindKey(AssociatedProperty))
+		{
+			return pValue->Get();
+		}
+		else
+		{
+			return nullptr;
+		}
+	}
+
+	// Returns the pin that caused the specified property to be created (can return nullptr if the association is unknown or the association is from an object instead)
+	UEdGraphPin* FindPinThatCreatedProperty(class UProperty* AssociatedProperty) const
+	{
+		if (const FEdGraphPinReference* pValue = DebugPinToPropertyMap.FindKey(AssociatedProperty))
 		{
 			return pValue->Get();
 		}
@@ -425,7 +454,7 @@ public:
 
 	void GenerateReversePropertyMap(TMap<UProperty*, UObject*>& PropertySourceMap)
 	{
-		for (TMap<TWeakObjectPtr<UObject>, UProperty*>::TIterator MapIt(DebugPinToPropertyMap); MapIt; ++MapIt)
+		for (TMap<TWeakObjectPtr<UObject>, UProperty*>::TIterator MapIt(DebugObjectToPropertyMap); MapIt; ++MapIt)
 		{
 			if (UObject* SourceObj = MapIt.Key().Get())
 			{
diff --git a/Engine/Source/Runtime/Engine/Classes/Engine/UserDefinedEnum.h b/Engine/Source/Runtime/Engine/Classes/Engine/UserDefinedEnum.h
--- a/Engine/Source/Runtime/Engine/Classes/Engine/UserDefinedEnum.h
+++ b/Engine/Source/Runtime/Engine/Classes/Engine/UserDefinedEnum.h
@@ -59,7 +59,7 @@ public:
 	 */
 	virtual FText GetEnumText(int32 InIndex) const override;
 
-	virtual bool SetEnums(TArray<TPair<FName, uint8>>& InNames, ECppForm InCppForm) override;
+	virtual bool SetEnums(TArray<TPair<FName, uint8>>& InNames, ECppForm InCppForm, bool bAddMaxKeyIfMissing = true) override;
 
 #if WITH_EDITOR
 	//~ Begin UObject Interface
diff --git a/Engine/Source/Runtime/Engine/Classes/Kismet/KismetArrayLibrary.h b/Engine/Source/Runtime/Engine/Classes/Kismet/KismetArrayLibrary.h
--- a/Engine/Source/Runtime/Engine/Classes/Kismet/KismetArrayLibrary.h
+++ b/Engine/Source/Runtime/Engine/Classes/Kismet/KismetArrayLibrary.h
@@ -79,7 +79,7 @@ class ENGINE_API UKismetArrayLibrary : public UBlueprintFunctionLibrary
 	 *
 	 *@param	TargetArray		The array to clear
 	*/
-	UFUNCTION(BlueprintCallable, CustomThunk, meta=(DisplayName = "Clear", CompactNodeTitle = "CLEAR", ArrayParm = "TargetArray"), Category="Utilities|Array")
+	UFUNCTION(BlueprintCallable, CustomThunk, meta=(DisplayName = "Clear", CompactNodeTitle = "CLEAR", Keywords = "empty", ArrayParm = "TargetArray"), Category="Utilities|Array")
 	static void Array_Clear(const TArray<int32>& TargetArray);
 
 	/* 
diff --git a/Engine/Source/Runtime/Engine/Classes/Kismet/KismetMathLibrary.h b/Engine/Source/Runtime/Engine/Classes/Kismet/KismetMathLibrary.h
--- a/Engine/Source/Runtime/Engine/Classes/Kismet/KismetMathLibrary.h
+++ b/Engine/Source/Runtime/Engine/Classes/Kismet/KismetMathLibrary.h
@@ -100,7 +100,7 @@ class ENGINE_API UKismetMathLibrary : public UBlueprintFunctionLibrary
 	static bool RandomBoolWithWeightFromStream(float Weight, const FRandomStream& RandomStream);
 
 	/* Returns the logical complement of the Boolean value (NOT A) */
-	UFUNCTION(BlueprintPure, meta=(DisplayName = "NOT Boolean", CompactNodeTitle = "NOT", Keywords = "! not"), Category="Math|Boolean")
+	UFUNCTION(BlueprintPure, meta=(DisplayName = "NOT Boolean", CompactNodeTitle = "NOT", Keywords = "! not negate"), Category="Math|Boolean")
 	static bool Not_PreBool(bool A);
 
 	/* Returns true if the values are equal (A == B) */
@@ -359,7 +359,7 @@ class ENGINE_API UKismetMathLibrary : public UBlueprintFunctionLibrary
 	UFUNCTION(BlueprintPure, meta=(DisplayName = "float < float", CompactNodeTitle = "<", Keywords = "< less"), Category="Math|Float")
 	static bool Less_FloatFloat(float A, float B);
 
-	/*Returns true if A is Greater than B (A > B) */
+	/*Returns true if A is greater than B (A > B) */
 	UFUNCTION(BlueprintPure, meta=(DisplayName = "float > float", CompactNodeTitle = ">", Keywords = "> greater"), Category="Math|Float")
 	static bool Greater_FloatFloat(float A, float B);
 
@@ -367,7 +367,7 @@ class ENGINE_API UKismetMathLibrary : public UBlueprintFunctionLibrary
 	UFUNCTION(BlueprintPure, meta=(DisplayName = "float <= float", CompactNodeTitle = "<=", Keywords = "<= less"), Category="Math|Float")
 	static bool LessEqual_FloatFloat(float A, float B);
 
-	/* Returns true if A is Greater than or equal to B (A >= B) */
+	/* Returns true if A is greater than or equal to B (A >= B) */
 	UFUNCTION(BlueprintPure, meta=(DisplayName = "float >= float", CompactNodeTitle = ">=", Keywords = ">= greater"), Category="Math|Float")
 	static bool GreaterEqual_FloatFloat(float A, float B);
 
@@ -995,19 +995,19 @@ class ENGINE_API UKismetMathLibrary : public UBlueprintFunctionLibrary
 	UFUNCTION(BlueprintPure, meta=(DisplayName="NotEqual (DateTime)", CompactNodeTitle="!=", Keywords="!= not equal"), Category="Math|DateTime")
 	static bool NotEqual_DateTimeDateTime( FDateTime A, FDateTime B );
 
-	/* Returns true if A is Greater than B (A > B) */
+	/* Returns true if A is greater than B (A > B) */
 	UFUNCTION(BlueprintPure, meta=(DisplayName="DateTime > DateTime", CompactNodeTitle=">", Keywords="> greater"), Category="Math|DateTime")
 	static bool Greater_DateTimeDateTime( FDateTime A, FDateTime B );
 
-	/* Returns true if A is Greater than B (A >= B) */
+	/* Returns true if A is greater than or equal to B (A >= B) */
 	UFUNCTION(BlueprintPure, meta=(DisplayName="DateTime >= DateTime", CompactNodeTitle=">=", Keywords=">= greater"), Category="Math|DateTime")
 	static bool GreaterEqual_DateTimeDateTime( FDateTime A, FDateTime B );
 
-	/* Returns true if A is Greater than B (A < B) */
+	/* Returns true if A is less than B (A < B) */
 	UFUNCTION(BlueprintPure, meta=(DisplayName="DateTime < DateTime", CompactNodeTitle="<", Keywords="< less"), Category="Math|DateTime")
 	static bool Less_DateTimeDateTime( FDateTime A, FDateTime B );
 
-	/* Returns true if A is Greater than B (A <= B) */
+	/* Returns true if A is less than or equal to B (A <= B) */
 	UFUNCTION(BlueprintPure, meta=(DisplayName="DateTime <= DateTime", CompactNodeTitle="<=", Keywords="<= less"), Category="Math|DateTime")
 	static bool LessEqual_DateTimeDateTime( FDateTime A, FDateTime B );
 
@@ -1135,19 +1135,19 @@ class ENGINE_API UKismetMathLibrary : public UBlueprintFunctionLibrary
 	UFUNCTION(BlueprintPure, meta=(DisplayName="NotEqual (Timespan)", CompactNodeTitle="!=", Keywords="!= not equal"), Category="Math|Timespan")
 	static bool NotEqual_TimespanTimespan( FTimespan A, FTimespan B );
 
-	/* Returns true if A is Greater than B (A > B) */
+	/* Returns true if A is greater than B (A > B) */
 	UFUNCTION(BlueprintPure, meta=(DisplayName="Timespan > Timespan", CompactNodeTitle=">", Keywords="> greater"), Category="Math|Timespan")
 	static bool Greater_TimespanTimespan( FTimespan A, FTimespan B );
 
-	/* Returns true if A is Greater than B (A >= B) */
+	/* Returns true if A is greater than or equal to B (A >= B) */
 	UFUNCTION(BlueprintPure, meta=(DisplayName="Timespan >= Timespan", CompactNodeTitle=">=", Keywords=">= greater"), Category="Math|Timespan")
 	static bool GreaterEqual_TimespanTimespan( FTimespan A, FTimespan B );
 
-	/* Returns true if A is Greater than B (A < B) */
+	/* Returns true if A is less than B (A < B) */
 	UFUNCTION(BlueprintPure, meta=(DisplayName="Timespan < Timespan", CompactNodeTitle="<", Keywords="< less"), Category="Math|Timespan")
 	static bool Less_TimespanTimespan( FTimespan A, FTimespan B );
 
-	/* Returns true if A is Greater than B (A <= B) */
+	/* Returns true if A is less than or equal to B (A <= B) */
 	UFUNCTION(BlueprintPure, meta=(DisplayName="Timespan <= Timespan", CompactNodeTitle="<=", Keywords="<= less"), Category="Math|Timespan")
 	static bool LessEqual_TimespanTimespan( FTimespan A, FTimespan B );
 
@@ -1657,6 +1657,10 @@ class ENGINE_API UKismetMathLibrary : public UBlueprintFunctionLibrary
 	UFUNCTION(BlueprintPure, meta=(DisplayName = "vector2d * float", CompactNodeTitle = "*", Keywords = "* multiply"), Category="Math|Vector2D")
 	static FVector2D Multiply_Vector2DFloat(FVector2D A, float B);
 
+	/* Element-wise Vector multiplication (Result = {A.x*B.x, A.y*B.y}) */
+	UFUNCTION(BlueprintPure, meta = (DisplayName = "vector2d * vector2d", CompactNodeTitle = "*", Keywords = "* multiply", CommutativeAssociativeBinaryOperator = "true"), Category = "Math|Vector2D")
+	static FVector2D Multiply_Vector2DVector2D(FVector2D A, FVector2D B);
+
 	/** Returns Vector A divided by B */
 	UFUNCTION(BlueprintPure, meta=(DisplayName = "vector2d / float", CompactNodeTitle = "/", Keywords = "/ divide division"), Category="Math|Vector2D")
 	static FVector2D Divide_Vector2DFloat(FVector2D A, float B = 1.f);
@@ -1843,7 +1847,7 @@ class ENGINE_API UKismetMathLibrary : public UBlueprintFunctionLibrary
 	 * @outparam	OutRectSideA - Vector oriented and sized to represent one edge of the enclosing rectangle, orthogonal to OutRectSideB
 	 * @outparam	OutRectSideB - Vector oriented and sized to represent one edge of the enclosing rectangle, orthogonal to OutRectSideA
 	*/
-	UFUNCTION(BlueprintCallable, BlueprintAuthorityOnly, Category="Math|Geometry", meta=(WorldContext="WorldContextObject"))
+	UFUNCTION(BlueprintCallable, BlueprintAuthorityOnly, Category="Math|Geometry", meta=(WorldContext="WorldContextObject", CallableWithoutWorldContext))
 	static void MinimumAreaRectangle(UObject* WorldContextObject, const TArray<FVector>& InVerts, const FVector& SampleSurfaceNormal, FVector& OutRectCenter, FRotator& OutRectRotation, float& OutSideLengthX, float& OutSideLengthY, bool bDebugDraw = false);
 
 	/**
diff --git a/Engine/Source/Runtime/Engine/Classes/Kismet/KismetTextLibrary.h b/Engine/Source/Runtime/Engine/Classes/Kismet/KismetTextLibrary.h
--- a/Engine/Source/Runtime/Engine/Classes/Kismet/KismetTextLibrary.h
+++ b/Engine/Source/Runtime/Engine/Classes/Kismet/KismetTextLibrary.h
@@ -46,6 +46,30 @@ class ENGINE_API UKismetTextLibrary : public UBlueprintFunctionLibrary
 {
 	GENERATED_UCLASS_BODY()
 
+	/** Converts a vector value to a localizable text, in the form 'X= Y= Z=' */
+	UFUNCTION(BlueprintPure, meta = (DisplayName = "ToText (Vector)", CompactNodeTitle = "->", BlueprintAutocast), Category = "Utilities|Text")
+	static FText Conv_VectorToText(FVector InVec);
+
+	/** Converts a vector2d value to a localizable text, in the form 'X= Y=' */
+	UFUNCTION(BlueprintPure, meta = (DisplayName = "ToText (vector2d)", CompactNodeTitle = "->", BlueprintAutocast), Category = "Utilities|Text")
+	static FText Conv_Vector2dToText(FVector2D InVec);
+
+	/** Converts a rotator value to a localizable text, in the form 'P= Y= R=' */
+	UFUNCTION(BlueprintPure, meta = (DisplayName = "ToText (rotator)", CompactNodeTitle = "->", BlueprintAutocast), Category = "Utilities|Text")
+	static FText Conv_RotatorToText(FRotator InRot);
+
+	/** Converts a transform value to a localizable text, in the form 'Translation: X= Y= Z= Rotation: P= Y= R= Scale: X= Y= Z=' */
+	UFUNCTION(BlueprintPure, meta = (DisplayName = "ToText (transform)", CompactNodeTitle = "->", BlueprintAutocast), Category = "Utilities|Text")
+	static FText Conv_TransformToText(const FTransform& InTrans);
+
+	/** Converts a UObject value to a localizable text by calling the object's GetName method */
+	UFUNCTION(BlueprintPure, meta = (DisplayName = "ToText (object)", CompactNodeTitle = "->", BlueprintAutocast), Category = "Utilities|Text")
+	static FText Conv_ObjectToText(class UObject* InObj);
+
+	/** Converts a linear color value to a localizable text, in the form '(R=,G=,B=,A=)' */
+	UFUNCTION(BlueprintPure, meta = (DisplayName = "ToText (linear color)", CompactNodeTitle = "->", BlueprintAutocast), Category = "Utilities|Text")
+	static FText Conv_ColorToText(FLinearColor InColor);
+
 	/** Converts localizable text to the string */
 	UFUNCTION(BlueprintPure, meta = (DisplayName = "ToString (text)", CompactNodeTitle = "->", BlueprintAutocast), Category = "Utilities|String")
 	static FString Conv_TextToString(const FText& InText);
diff --git a/Engine/Source/Runtime/Engine/Private/Blueprint.cpp b/Engine/Source/Runtime/Engine/Private/Blueprint.cpp
--- a/Engine/Source/Runtime/Engine/Private/Blueprint.cpp
+++ b/Engine/Source/Runtime/Engine/Private/Blueprint.cpp
@@ -491,6 +491,50 @@ void UBlueprint::PostLoad()
 {
 	Super::PostLoad();
 
+	if (DeprecatedPinWatches.Num() > 0)
+	{
+		// patch DeprecatedPinWatches to WatchedPins:
+		for (UEdGraphPin_Deprecated* OldPin : DeprecatedPinWatches)
+		{
+			if (UEdGraphPin* NewPin = UEdGraphPin::FindPinCreatedFromDeprecatedPin(OldPin))
+			{
+				WatchedPins.Add(NewPin);
+			}
+		}
+
+		DeprecatedPinWatches.Empty();
+	}
+
+	// now that pin watches have been fixed up, we can remove the old pins:
+	{
+		TArray<UEdGraph*> Graphs;
+		GetAllGraphs(Graphs);
+		for (UEdGraph* Graph : Graphs)
+		{
+			if (Graph)
+			{
+				for (UEdGraphNode* Node : Graph->Nodes)
+				{
+					if (Node)
+					{
+						TArray<UEdGraphPin_Deprecated*>& DeprecatedPins = Node->DeprecatedPins;
+						if (DeprecatedPins.Num() > 0)
+						{
+							for (UEdGraphPin_Deprecated* LegacyPin : DeprecatedPins)
+							{
+								LegacyPin->Rename(nullptr, GetTransientPackage(), REN_ForceNoResetLoaders);
+								LegacyPin->SetFlags(RF_Transient);
+								LegacyPin->MarkPendingKill();
+							}
+
+							DeprecatedPins.Empty();
+						}
+					}
+				}
+			}
+		}
+	}
+
 	// Can't use TGuardValue here as bIsRegeneratingOnLoad is a bitfield
 	struct FScopedRegeneratingOnLoad
 	{
diff --git a/Engine/Source/Runtime/Engine/Private/Components/SceneComponent.cpp b/Engine/Source/Runtime/Engine/Private/Components/SceneComponent.cpp
--- a/Engine/Source/Runtime/Engine/Private/Components/SceneComponent.cpp
+++ b/Engine/Source/Runtime/Engine/Private/Components/SceneComponent.cpp
@@ -1909,7 +1909,7 @@ FSceneComponentInstanceData::FSceneComponentInstanceData(const USceneComponent*
 	for (int32 i = AttachedChildren.Num()-1; i >= 0; --i)
 	{
 		USceneComponent* SceneComponent = AttachedChildren[i];
-		if (SceneComponent && SceneComponent->GetOwner() == SourceOwner && !SceneComponent->IsCreatedByConstructionScript())
+		if (SceneComponent && SceneComponent->GetOwner() == SourceOwner && !SceneComponent->IsCreatedByConstructionScript() && !SceneComponent->HasAnyFlags(RF_DefaultSubObject))
 		{
 			AttachedInstanceComponents.Add(TPairInitializer<USceneComponent*,const FTransform&>(SceneComponent, FTransform(SceneComponent->RelativeRotation, SceneComponent->RelativeLocation, SceneComponent->RelativeScale3D)));
 		}
@@ -1968,7 +1968,7 @@ FActorComponentInstanceData* USceneComponent::GetComponentInstanceData() const
 
 	for (USceneComponent* Child : GetAttachChildren())
 	{
-		if (Child && !Child->IsCreatedByConstructionScript())
+		if (Child && !Child->IsCreatedByConstructionScript() && !Child->HasAnyFlags(RF_DefaultSubObject))
 		{
 			InstanceData = new FSceneComponentInstanceData(this);
 			break;
diff --git a/Engine/Source/Runtime/Engine/Private/EdGraph/EdGraphNode.cpp b/Engine/Source/Runtime/Engine/Private/EdGraph/EdGraphNode.cpp
--- a/Engine/Source/Runtime/Engine/Private/EdGraph/EdGraphNode.cpp
+++ b/Engine/Source/Runtime/Engine/Private/EdGraph/EdGraphNode.cpp
@@ -2,6 +2,7 @@
 
 #include "EnginePrivate.h"
 #include "EdGraph/EdGraph.h"
+#include "BlueprintsObjectVersion.h"
 #include "BlueprintUtilities.h"
 #if WITH_EDITOR
 #include "Editor/UnrealEd/Public/CookerSettings.h"
@@ -56,23 +57,24 @@ UEdGraphNode::UEdGraphNode(const FObjectInitializer& ObjectInitializer)
 #endif // WITH_EDITORONLY_DATA
 }
 
+UEdGraphNode::~UEdGraphNode()
+{
+	for (UEdGraphPin* Pin : Pins)
+	{
+		Pin->MarkPendingKill();
+	}
+}
+
 #if WITH_EDITOR
 
 UEdGraphPin* UEdGraphNode::CreatePin(EEdGraphPinDirection Dir, const FEdGraphPinType& InPinType, const FString& PinName, int32 Index /*= INDEX_NONE*/)
 {
-	UEdGraphPin* NewPin = NewObject<UEdGraphPin>(this);
+	UEdGraphPin* NewPin = UEdGraphPin::CreatePin(this);
 	NewPin->PinName = PinName;
 	NewPin->Direction = Dir;
 
 	NewPin->PinType = InPinType;
 
-	NewPin->SetFlags(RF_Transactional);
-
-	if (HasAnyFlags(RF_Transient))
-	{
-		NewPin->SetFlags(RF_Transient);
-	}
-
 	Modify(false);
 	if (Pins.IsValidIndex(Index))
 	{
@@ -118,8 +120,8 @@ bool UEdGraphNode::RemovePin(UEdGraphPin* Pin)
 	check( Pin );
 	
 	Modify();
-	UEdGraphPin* RootPin = (Pin->ParentPin != nullptr)? Pin->ParentPin : Pin;
-	RootPin->BreakAllPinLinks();
+	UEdGraphPin* RootPin = (Pin->ParentPin != nullptr) ? Pin->ParentPin : Pin;
+	RootPin->MarkPendingKill();
 
 	if (Pins.Remove( RootPin ))
 	{
@@ -127,11 +129,11 @@ bool UEdGraphNode::RemovePin(UEdGraphPin* Pin)
 		for (UEdGraphPin* ChildPin : RootPin->SubPins)
 		{
 			Pins.Remove(ChildPin);
-			ChildPin->Modify();
-			ChildPin->BreakAllPinLinks();
+			ChildPin->MarkPendingKill();
 		}
 		return true;
 	}
+
 	return false;
 }
 
@@ -251,29 +253,24 @@ FString UEdGraphNode::GetDeprecationMessage() const
 	return NSLOCTEXT("EdGraphCompiler", "NodeDeprecated_Warning", "@@ is deprecated; please replace or remove it.").ToString();
 }
 
-// Array of pooled pins
-TArray<UEdGraphPin*> UEdGraphNode::PooledPins;
-
 void UEdGraphNode::AddReferencedObjects(UObject* InThis, FReferenceCollector& Collector)
 {
-	UEdGraphNode* This = CastChecked<UEdGraphNode>(InThis);	
+	Super::AddReferencedObjects(InThis, Collector);
 
-	// Only register the pool once per GC pass
-	if (This->HasAnyFlags(RF_ClassDefaultObject))
+	UEdGraphNode* This = CastChecked<UEdGraphNode>(InThis);
+	for (UEdGraphPin* Pin : This->Pins)
 	{
-		if (This->GetClass() == UEdGraphNode::StaticClass())
+		if (Pin)
 		{
-			for (int32 Index = 0; Index < PooledPins.Num(); ++Index)
-			{
-				Collector.AddReferencedObject(PooledPins[Index], This);
-			}
+			Pin->AddStructReferencedObjects(Collector);
 		}
 	}
-	Super::AddReferencedObjects(This, Collector);
 }
 
 void UEdGraphNode::Serialize(FArchive& Ar)
 {
+	Ar.UsingCustomVersion(FBlueprintsObjectVersion::GUID);
+
 	Super::Serialize(Ar);
 
 	if (Ar.IsLoading())
@@ -284,6 +281,26 @@ void UEdGraphNode::Serialize(FArchive& Ar)
 		{
 			EnabledState = ENodeEnabledState::Disabled;
 		}
+
+		if (Ar.IsPersistent() && !Ar.HasAnyPortFlags(PPF_Duplicate | PPF_DuplicateForPIE))
+		{
+			if (Ar.CustomVer(FBlueprintsObjectVersion::GUID) < FBlueprintsObjectVersion::EdGraphPinOptimized)
+			{
+				for (UEdGraphPin_Deprecated* LegacyPin : DeprecatedPins)
+				{
+					Ar.Preload(LegacyPin);
+					if (UEdGraphPin::FindPinCreatedFromDeprecatedPin(LegacyPin) == nullptr)
+					{
+						UEdGraphPin::CreatePinFromDeprecatedPin(LegacyPin);
+					}
+				}
+			}
+		}
+	}
+
+	if (Ar.CustomVer(FBlueprintsObjectVersion::GUID) >= FBlueprintsObjectVersion::EdGraphPinOptimized)
+	{
+		UEdGraphPin::SerializeAsOwningNode(Ar, Pins);
 	}
 }
 
@@ -338,40 +355,57 @@ void UEdGraphNode::PostEditChangeProperty(FPropertyChangedEvent& PropertyChanged
 	}
 }
 
-void UEdGraphNode::CreateNewGuid()
+void UEdGraphNode::PostEditUndo()
 {
-	NodeGuid = FGuid::NewGuid();
+	UEdGraphPin::ResolveAllPinReferences();
+	
+	return UObject::PostEditUndo();
 }
 
-void UEdGraphNode::FindDiffs( class UEdGraphNode* OtherNode, struct FDiffResults& Results ) 
+void UEdGraphNode::ExportCustomProperties(FOutputDevice& Out, uint32 Indent)
 {
-}
+	Super::ExportCustomProperties(Out, Indent);
 
-UEdGraphPin* UEdGraphNode::AllocatePinFromPool(UEdGraphNode* OuterNode)
-{
-	if (PooledPins.Num() > 0)
+	for (const UEdGraphPin* Pin : Pins)
 	{
-		UEdGraphPin* Result = PooledPins.Pop();
-		Result->Rename(NULL, OuterNode);
-		return Result;
+		FString PinString;
+		Pin->ExportTextItem(PinString, PPF_Delimited);
+		Out.Logf(TEXT("%sCustomProperties Pin %s\r\n"), FCString::Spc(Indent), *PinString);
 	}
-	else
+}
+
+void UEdGraphNode::ImportCustomProperties(const TCHAR* SourceText, FFeedbackContext* Warn)
+{
+	Super::ImportCustomProperties(SourceText, Warn);
+
+	if (FParse::Command(&SourceText, TEXT("Pin")))
 	{
-		UEdGraphPin* Result = NewObject<UEdGraphPin>(OuterNode);
-		return Result;
+		UEdGraphPin* NewPin = UEdGraphPin::CreatePin(this);
+		const bool bParseSuccess = NewPin->ImportTextItem(SourceText, PPF_Delimited, nullptr, GWarn);
+		if (bParseSuccess)
+		{
+			Pins.Add(NewPin);
+		}
+		else
+		{
+			// Still adding a nullptr to preserve indices
+			Pins.Add(nullptr);
+		}
 	}
 }
 
-void UEdGraphNode::ReturnPinToPool(UEdGraphPin* OldPin)
+void UEdGraphNode::CreateNewGuid()
 {
-	check(OldPin);
-
-	check(!OldPin->HasAnyFlags(RF_NeedLoad));
-	OldPin->ClearFlags(RF_NeedPostLoadSubobjects|RF_NeedPostLoad);
+	NodeGuid = FGuid::NewGuid();
+}
 
-	OldPin->ResetToDefaults();
+void UEdGraphNode::FindDiffs( class UEdGraphNode* OtherNode, struct FDiffResults& Results ) 
+{
+}
 
-	PooledPins.Add(OldPin);
+void UEdGraphNode::DestroyPin(UEdGraphPin* Pin)
+{
+	Pin->MarkPendingKill();
 }
 
 bool UEdGraphNode::CanDuplicateNode() const
