Copying //UE4/Dev-Blueprints to //UE4/Dev-Main (Source: //UE4/Dev-Blueprints @ 3025888)

diff --git a/Engine/Source/Runtime/Engine/Private/EdGraph/EdGraphPin.cpp b/Engine/Source/Runtime/Engine/Private/EdGraph/EdGraphPin.cpp
--- a/Engine/Source/Runtime/Engine/Private/EdGraph/EdGraphPin.cpp
+++ b/Engine/Source/Runtime/Engine/Private/EdGraph/EdGraphPin.cpp
@@ -4,15 +4,58 @@
 #include "EdGraph/EdGraph.h"
 #include "EdGraph/EdGraphSchema.h"
 #include "BlueprintUtilities.h"
+#include "Tickable.h"
 #if WITH_EDITOR
 #include "Editor/UnrealEd/Public/Kismet2/BlueprintEditorUtils.h"
 #include "SlateBasics.h"
 #include "ScopedTransaction.h"
 #include "Editor/UnrealEd/Public/Kismet2/Kismet2NameValidators.h"
+#include "TickableEditorObject.h"
 #endif
 
 #define LOCTEXT_NAMESPACE "EdGraph"
 
+TArray<UEdGraphPin*> PinsToDelete;
+
+class FPinDeletionQueue : 
+#if WITH_EDITOR
+	public FTickableEditorObject
+#else
+	public FTickableGameObject
+#endif
+{
+public:
+	virtual void Tick(float DeltaTime) override
+	{
+		for (UEdGraphPin* Pin : PinsToDelete)
+		{
+			delete Pin;
+		}
+		PinsToDelete.Empty();
+	}
+
+	virtual bool IsTickable() const override
+	{
+		return true;
+	}
+
+	/** return the stat id to use for this tickable **/
+	virtual TStatId GetStatId() const
+	{
+		RETURN_QUICK_DECLARE_CYCLE_STAT(FPinDeletionQueue, STATGROUP_Tickables);
+	}
+
+	virtual ~FPinDeletionQueue() {}
+};
+
+FPinDeletionQueue* PinDeletionQueue = new FPinDeletionQueue();;
+
+//#define TRACK_PINS
+
+#ifdef TRACK_PINS
+TArray<TPair<UEdGraphPin*, FString>> PinAllocationTracking;
+#endif //TRACK_PINS
+
 /////////////////////////////////////////////////////
 // FEdGraphPinType
 
@@ -76,50 +119,141 @@ bool FEdGraphPinType::Serialize(FArchive& Ar)
 }
 
 /////////////////////////////////////////////////////
-// GraphPinHelpers
-namespace GraphPinHelpers
+// FEdGraphPinReference
+void FEdGraphPinReference::SetPin(const UEdGraphPin* NewPin)
 {
-	void EnableAllConnectedNodes(UEdGraphNode* InNode)
+	if (NewPin)
 	{
-		// Only enable when it has not been explicitly user-disabled
-		if(InNode && InNode->EnabledState == ENodeEnabledState::Disabled && !InNode->bUserSetEnabledState)
-		{
-			// Enable the node and clear the comment
-			InNode->Modify();
-			InNode->EnableNode();
-			InNode->NodeComment.Empty();
+		OwningNode = NewPin->GetOwningNodeUnchecked();
+		ensure(OwningNode.Get() || NewPin->IsPendingKill());
+		PinId = NewPin->PinId;
+	}
+	else
+	{
+		OwningNode = nullptr;
+		PinId.Invalidate();
+	}
+}
 
-			// Go through all pin connections and enable the nodes. Enabled nodes will prevent further iteration
-			for (UEdGraphPin* Pin : InNode->Pins)
-			{
-				for (UEdGraphPin* OtherPin : Pin->LinkedTo)
-				{
-					EnableAllConnectedNodes(OtherPin->GetOwningNode());
-				}
-			}
+UEdGraphPin* FEdGraphPinReference::Get() const
+{
+	UEdGraphNode* ResolvedOwningNode = OwningNode.Get();
+	if (ResolvedOwningNode && PinId.IsValid())
+	{
+		const FGuid& TargetId = PinId;
+		UEdGraphPin** ExistingPin = ResolvedOwningNode->Pins.FindByPredicate([TargetId](const UEdGraphPin* TargetPin) { return TargetPin && TargetPin->PinId == TargetId; });
+		if (ExistingPin)
+		{
+			return *ExistingPin;
 		}
-	};
+	}
+
+	return nullptr;
 }
 
 /////////////////////////////////////////////////////
-// UEdGraphPin
+// Resolve enums and structs
+enum class EPinResolveType : uint8
+{
+	OwningNode,
+	LinkedTo,
+	SubPins,
+	ParentPin,
+	ReferencePassThroughConnection
+};
 
-UEdGraphPin::UEdGraphPin(const FObjectInitializer& ObjectInitializer)
-	: Super(ObjectInitializer)
+struct FUnresolvedPinData
+{
+	UEdGraphPin* ReferencingPin;
+	int32 ArrayIdx;
+	EPinResolveType ResolveType;
+
+	FUnresolvedPinData() : ReferencingPin(nullptr), ArrayIdx(INDEX_NONE), ResolveType(EPinResolveType::OwningNode) {}
+	FUnresolvedPinData(UEdGraphPin* InReferencingPin, EPinResolveType InResolveType, int32 InArrayIdx)
+		: ReferencingPin(InReferencingPin), ArrayIdx(InArrayIdx), ResolveType(InResolveType)
+	{}
+};
+
+struct FPinResolveId
+{
+	FGuid PinGuid;
+	TWeakObjectPtr<UEdGraphNode> OwningNode;
+
+	FPinResolveId() {}
+
+	FPinResolveId(const FGuid& InPinGuid, const TWeakObjectPtr<UEdGraphNode>& InOwningNode)
+		: PinGuid(InPinGuid), OwningNode(InOwningNode)
+	{}
+
+	bool operator==(const FPinResolveId& Other) const
+	{
+		return PinGuid == Other.PinGuid && OwningNode == Other.OwningNode;
+	}
+
+	friend inline uint32 GetTypeHash(const FPinResolveId& Id)
+	{
+		return HashCombine(GetTypeHash(Id.PinGuid), GetTypeHash(Id.OwningNode));
+	}
+};
+
+/////////////////////////////////////////////////////
+// PinHelpers
+namespace PinHelpers
 {
+	static TMap<FPinResolveId, TArray<FUnresolvedPinData>> UnresolvedPins;
+	static TMap<TWeakObjectPtr<UEdGraphPin_Deprecated>, FGuid> DeprecatedPinToNewPinGUIDMap;
+	static TMap<TWeakObjectPtr<UEdGraphPin_Deprecated>, TArray<FUnresolvedPinData>> UnresolvedDeprecatedPins;
+
+	static uint64 NumPinsInMemory = 0;
+
+	static const TCHAR ExportTextPropDelimiter = ',';
+
+	static FString PinIdName = GET_MEMBER_NAME_STRING_CHECKED(UEdGraphPin, PinId);
+	static FString PinNameName = GET_MEMBER_NAME_STRING_CHECKED(UEdGraphPin, PinName);
 #if WITH_EDITORONLY_DATA
-	bHidden = false;
-	bNotConnectable = false;
-	bAdvancedView = false;
-#endif // WITH_EDITORONLY_DATA
+	static FString PinFriendlyNameName = GET_MEMBER_NAME_STRING_CHECKED(UEdGraphPin, PinFriendlyName);
+#endif
+	static FString PinToolTipName = GET_MEMBER_NAME_STRING_CHECKED(UEdGraphPin, PinToolTip);
+	static FString DirectionName = GET_MEMBER_NAME_STRING_CHECKED(UEdGraphPin, Direction);
+	static FString PinTypeName = GET_MEMBER_NAME_STRING_CHECKED(UEdGraphPin, PinType);
+	static FString DefaultValueName = GET_MEMBER_NAME_STRING_CHECKED(UEdGraphPin, DefaultValue);
+	static FString AutogeneratedDefaultValueName = GET_MEMBER_NAME_STRING_CHECKED(UEdGraphPin, AutogeneratedDefaultValue);
+	static FString DefaultObjectName = GET_MEMBER_NAME_STRING_CHECKED(UEdGraphPin, DefaultObject);
+	static FString DefaultTextValueName = GET_MEMBER_NAME_STRING_CHECKED(UEdGraphPin, DefaultTextValue);
+	static FString LinkedToName = GET_MEMBER_NAME_STRING_CHECKED(UEdGraphPin, LinkedTo);
+	static FString SubPinsName = GET_MEMBER_NAME_STRING_CHECKED(UEdGraphPin, SubPins);
+	static FString ParentPinName = GET_MEMBER_NAME_STRING_CHECKED(UEdGraphPin, ParentPin);
+	static FString ReferencePassThroughConnectionName = GET_MEMBER_NAME_STRING_CHECKED(UEdGraphPin, ReferencePassThroughConnection);
+
+#if WITH_EDITORONLY_DATA
+	static FString PersistentGuidName = GET_MEMBER_NAME_STRING_CHECKED(UEdGraphPin, PersistentGuid);
+	static FString bHiddenName = GET_MEMBER_NAME_STRING_CHECKED(UEdGraphPin, bHidden);
+	static FString bNotConnectableName = GET_MEMBER_NAME_STRING_CHECKED(UEdGraphPin, bNotConnectable);
+	static FString bDefaultValueIsReadOnlyName = GET_MEMBER_NAME_STRING_CHECKED(UEdGraphPin, bDefaultValueIsReadOnly);
+	static FString bDefaultValueIsIgnoredName = GET_MEMBER_NAME_STRING_CHECKED(UEdGraphPin, bDefaultValueIsIgnored);
+	// Don't need bIsDiffing's name because it is transient
+	static FString bAdvancedViewName = GET_MEMBER_NAME_STRING_CHECKED(UEdGraphPin, bAdvancedView);
+	// Don't need bDisplayAsMutableRef's name because it is transient
+	// Don't need bWasTrashed's name because it is transient
+#endif
+}
+
+/////////////////////////////////////////////////////
+// UEdGraphPin
+UEdGraphPin* UEdGraphPin::CreatePin(UEdGraphNode* InOwningNode)
+{
+	check(InOwningNode);
+	UEdGraphPin* NewPin = new UEdGraphPin(InOwningNode, FGuid::NewGuid());
+	return NewPin;
 }
 
 void UEdGraphPin::MakeLinkTo(UEdGraphPin* ToPin)
 {
 	Modify();
 
-	if (ToPin != NULL)
+	if (ToPin)
 	{
+		check(!bWasTrashed);
 		ToPin->Modify();
 
 		// Make sure we don't already link to it
@@ -128,15 +262,15 @@ void UEdGraphPin::MakeLinkTo(UEdGraphPin* ToPin)
 			UEdGraphNode* MyNode = GetOwningNode();
 
 			// Check that the other pin does not link to us
-			ensureMsgf(!ToPin->LinkedTo.Contains(this), *GetLinkInfoString( LOCTEXT("MakeLinkTo", "MakeLinkTo").ToString(), LOCTEXT("IsLinked", "is linked with pin").ToString(), ToPin));			    
+			ensureMsgf(!ToPin->LinkedTo.Contains(this), *GetLinkInfoString( LOCTEXT("MakeLinkTo", "MakeLinkTo").ToString(), LOCTEXT("IsLinked", "is linked with pin").ToString(), ToPin));
 			ensureMsgf(MyNode->GetOuter() == ToPin->GetOwningNode()->GetOuter(), *GetLinkInfoString( LOCTEXT("MakeLinkTo", "MakeLinkTo").ToString(), LOCTEXT("OuterMismatch", "has a different outer than pin").ToString(), ToPin)); // Ensure both pins belong to the same graph
 
 			// Add to both lists
 			LinkedTo.Add(ToPin);
 			ToPin->LinkedTo.Add(this);
 
-			GraphPinHelpers::EnableAllConnectedNodes(MyNode);
-			GraphPinHelpers::EnableAllConnectedNodes(ToPin->GetOwningNode());
+			UEdGraphPin::EnableAllConnectedNodes(MyNode);
+			UEdGraphPin::EnableAllConnectedNodes(ToPin->GetOwningNode());
 		}
 	}
 }
@@ -145,23 +279,22 @@ void UEdGraphPin::BreakLinkTo(UEdGraphPin* ToPin)
 {
 	Modify();
 
-	if (ToPin != NULL)
+	if (ToPin)
 	{
 		ToPin->Modify();
 
 		// If we do indeed link to the passed in pin...
 		if (LinkedTo.Contains(ToPin))
 		{
-			LinkedTo.Remove(ToPin);
-
 			// Check that the other pin links to us
-			ensureMsgf(ToPin->LinkedTo.Contains(this), *GetLinkInfoString( LOCTEXT("BreakLinkTo", "BreakLinkTo").ToString(), LOCTEXT("NotLinked", "not reciprocally linked with pin").ToString(), ToPin) );
+			ensureAlwaysMsgf(ToPin->LinkedTo.Contains(this), *GetLinkInfoString(LOCTEXT("BreakLinkTo", "BreakLinkTo").ToString(), LOCTEXT("NotLinked", "not reciprocally linked with pin").ToString(), ToPin));
 			ToPin->LinkedTo.Remove(this);
+			LinkedTo.Remove(ToPin);
 		}
 		else
 		{
 			// Check that the other pin does not link to us
-			ensureMsgf(!ToPin->LinkedTo.Contains(this), *GetLinkInfoString( LOCTEXT("MakeLinkTo", "MakeLinkTo").ToString(), LOCTEXT("IsLinked", "is linked with pin").ToString(), ToPin));
+			ensureAlwaysMsgf(!ToPin->LinkedTo.Contains(this), *GetLinkInfoString(LOCTEXT("MakeLinkTo", "MakeLinkTo").ToString(), LOCTEXT("IsLinked", "is linked with pin").ToString(), ToPin));
 		}
 	}
 }
@@ -182,8 +315,10 @@ void UEdGraphPin::CopyPersistentDataFromOldPin(const UEdGraphPin& SourcePin)
 	// The name matches already, doesn't get copied here
 	// The PinType, Direction, and bNotConnectable are properties generated from the schema
 
+	PinId = SourcePin.PinId;
+
 	// Only move the default value if it was modified; inherit the new default value otherwise
-	if (SourcePin.DefaultValue != SourcePin.AutogeneratedDefaultValue || SourcePin.DefaultObject != NULL || SourcePin.DefaultTextValue.ToString() != SourcePin.AutogeneratedDefaultValue)
+	if (SourcePin.DefaultValue != SourcePin.AutogeneratedDefaultValue || SourcePin.DefaultObject != nullptr || SourcePin.DefaultTextValue.ToString() != SourcePin.AutogeneratedDefaultValue)
 	{
 		DefaultObject = SourcePin.DefaultObject;
 		DefaultValue = SourcePin.DefaultValue;
@@ -194,7 +329,7 @@ void UEdGraphPin::CopyPersistentDataFromOldPin(const UEdGraphPin& SourcePin)
 	for (int32 LinkIndex = 0; LinkIndex < SourcePin.LinkedTo.Num(); ++LinkIndex)
 	{
 		UEdGraphPin* OtherPin = SourcePin.LinkedTo[LinkIndex];
-		check(NULL != OtherPin);
+		check(OtherPin);
 
 		Modify();
 		OtherPin->Modify();
@@ -227,8 +362,8 @@ void UEdGraphPin::CopyPersistentDataFromOldPin(const UEdGraphPin& SourcePin)
 #if WITH_EDITORONLY_DATA
 	// Copy advanced visibility property, if it can be changed by user.
 	// Otherwise we don't want to copy this, or we'd be ignoring new metadata that tries to hide old pins.
-	UEdGraphNode* OuterNode = Cast<UEdGraphNode>(GetOuter());
-	if (OuterNode != nullptr && OuterNode->CanUserEditPinAdvancedViewFlag())
+	UEdGraphNode* LocalOwningNode = GetOwningNodeUnchecked();
+	if (LocalOwningNode != nullptr && LocalOwningNode->CanUserEditPinAdvancedViewFlag())
 	{
 		bAdvancedView = SourcePin.bAdvancedView;
 	}
@@ -241,25 +376,25 @@ void UEdGraphPin::AssignByRefPassThroughConnection(UEdGraphPin* InTargetPin)
 	{
 		if (GetOwningNode() != InTargetPin->GetOwningNode())
 		{
-			UE_LOG(LogBlueprint, Warning, TEXT("Pin '%s' is owned by node '%s' and pin '%s' is owned by '%s', they must be owned by the same node!"), *GetName(), *GetOwningNode()->GetName(), *InTargetPin->GetName(), *InTargetPin->GetOwningNode()->GetName());
+			UE_LOG(LogBlueprint, Warning, TEXT("Pin '%s' is owned by node '%s' and pin '%s' is owned by '%s', they must be owned by the same node!"), *PinName, *GetOwningNode()->GetName(), *InTargetPin->PinName, *InTargetPin->GetOwningNode()->GetName());
 		}
 		else if (Direction == InTargetPin->Direction)
 		{
-			UE_LOG(LogBlueprint, Warning, TEXT("Both pin '%s' and pin '%s' on node '%s' go in the same direction, one must be an input and the other an output!"), *GetName(), *InTargetPin->GetName(), *GetOwningNode()->GetName());
+			UE_LOG(LogBlueprint, Warning, TEXT("Both pin '%s' and pin '%s' on node '%s' go in the same direction, one must be an input and the other an output!"), *PinName, *InTargetPin->PinName, *GetOwningNode()->GetName());
 		}
 		else if (!PinType.bIsReference && !InTargetPin->PinType.bIsReference)
 		{
 			UEdGraphPin* InputPin = (Direction == EGPD_Input) ? this : InTargetPin;
 			UEdGraphPin* OutputPin = (Direction == EGPD_Input) ? InTargetPin : this;
-			UE_LOG(LogBlueprint, Warning, TEXT("Input pin '%s' is attempting to by-ref pass-through to output pin '%s' on node '%s', however neither pin is by-ref!"), *InputPin->GetName(), *OutputPin->GetName(), *InputPin->GetOwningNode()->GetName());
+			UE_LOG(LogBlueprint, Warning, TEXT("Input pin '%s' is attempting to by-ref pass-through to output pin '%s' on node '%s', however neither pin is by-ref!"), *InputPin->PinName, *OutputPin->PinName, *InputPin->GetOwningNode()->GetName());
 		}
 		else if (!PinType.bIsReference)
 		{
-			UE_LOG(LogBlueprint, Warning, TEXT("Pin '%s' on node '%s' is not by-ref but it is attempting to pass-through to '%s'"), *GetName(), *GetOwningNode()->GetName(), *InTargetPin->GetName());
+			UE_LOG(LogBlueprint, Warning, TEXT("Pin '%s' on node '%s' is not by-ref but it is attempting to pass-through to '%s'"), *PinName, *GetOwningNode()->GetName(), *InTargetPin->PinName);
 		}
 		else if (!InTargetPin->PinType.bIsReference)
 		{
-			UE_LOG(LogBlueprint, Warning, TEXT("Pin '%s' on node '%s' is not by-ref but it is attempting to pass-through to '%s'"), *InTargetPin->GetName(), *InTargetPin->GetOwningNode()->GetName(), *GetName());
+			UE_LOG(LogBlueprint, Warning, TEXT("Pin '%s' on node '%s' is not by-ref but it is attempting to pass-through to '%s'"), *InTargetPin->PinName, *InTargetPin->GetOwningNode()->GetName(), *PinName);
 		}
 		else
 		{
@@ -272,7 +407,7 @@ void UEdGraphPin::AssignByRefPassThroughConnection(UEdGraphPin* InTargetPin)
 const class UEdGraphSchema* UEdGraphPin::GetSchema() const
 {
 #if WITH_EDITOR
-	auto OwnerNode = GetOwningNodeUnchecked();
+	auto OwnerNode = GetOwningNode();
 	return OwnerNode ? OwnerNode->GetSchema() : NULL;
 #else
 	return NULL;
@@ -281,7 +416,7 @@ const class UEdGraphSchema* UEdGraphPin::GetSchema() const
 
 FString UEdGraphPin::GetDefaultAsString() const
 {
-	if(DefaultObject != NULL)
+	if(DefaultObject)
 	{
 		return DefaultObject->GetPathName();
 	}
@@ -323,12 +458,12 @@ const FString UEdGraphPin::GetLinkInfoString( const FString& InFunctionName, con
 {
 #if WITH_EDITOR
 	const FString FromPinName = PinName;
-	const UEdGraphNode* FromPinNode = Cast<UEdGraphNode>(GetOuter());
-	const FString FromPinNodeName = (FromPinNode != NULL) ? FromPinNode->GetNodeTitle(ENodeTitleType::ListView).ToString() : TEXT("Unknown");
+	const UEdGraphNode* FromPinNode = GetOwningNode();
+	const FString FromPinNodeName = (FromPinNode != nullptr) ? FromPinNode->GetNodeTitle(ENodeTitleType::ListView).ToString() : TEXT("Unknown");
 	
 	const FString ToPinName = InToPin->PinName;
-	const UEdGraphNode* ToPinNode = Cast<UEdGraphNode>(InToPin->GetOuter());
-	const FString ToPinNodeName = (ToPinNode != NULL) ? ToPinNode->GetNodeTitle(ENodeTitleType::ListView).ToString() : TEXT("Unknown");
+	const UEdGraphNode* ToPinNode = InToPin->GetOwningNode();
+	const FString ToPinNodeName = (ToPinNode != nullptr) ? ToPinNode->GetNodeTitle(ENodeTitleType::ListView).ToString() : TEXT("Unknown");
 	const FString LinkInfo = FString::Printf( TEXT("UEdGraphPin::%s Pin '%s' on node '%s' %s '%s' on node '%s'"), *InFunctionName, *ToPinName, *ToPinNodeName, *InInfoData, *FromPinName, *FromPinNodeName);
 	return LinkInfo;
 #else
@@ -336,10 +471,1145 @@ const FString UEdGraphPin::GetLinkInfoString( const FString& InFunctionName, con
 #endif
 }
 
-void UEdGraphPin::PostLoad()
+void UEdGraphPin::AddStructReferencedObjects(class FReferenceCollector& Collector)
+{
+	Collector.AddReferencedObject(PinType.PinSubCategoryMemberReference.MemberParent);
+	Collector.AddReferencedObject(DefaultObject);
+}
+
+void UEdGraphPin::SerializeAsOwningNode(FArchive& Ar, TArray<UEdGraphPin*>& ArrayRef)
+{
+	UEdGraphPin::SerializePinArray(Ar, ArrayRef, nullptr, EPinResolveType::OwningNode);
+}
+
+bool UEdGraphPin::Modify(bool bAlwaysMarkDirty)
+{
+	UEdGraphNode* LocalOwningNode = GetOwningNodeUnchecked();
+	if (LocalOwningNode)
+	{
+		return LocalOwningNode->Modify(bAlwaysMarkDirty);
+	}
+
+	return false;
+}
+
+void UEdGraphPin::SetOwningNode(UEdGraphNode* NewOwningNode)
+{
+	UEdGraphNode* OldOwningNode = GetOwningNodeUnchecked();
+	if (OldOwningNode)
+	{
+		int32 ExistingPinIdx = OldOwningNode->Pins.IndexOfByKey(this);
+		if (ExistingPinIdx != INDEX_NONE)
+		{
+			OldOwningNode->Pins.RemoveAt(ExistingPinIdx);
+		}
+	}
+
+	OwningNode = NewOwningNode;
+
+	if (NewOwningNode)
+	{
+		int32 ExistingPinIdx = NewOwningNode->Pins.IndexOfByPredicate([this](UEdGraphPin* Pin) { return Pin && Pin->PinId == this->PinId; });
+		if (ExistingPinIdx == INDEX_NONE)
+		{
+			NewOwningNode->Pins.Add(this);
+		}
+		else
+		{
+			// If a pin with this ID is already in the node, it better be this pin
+			check(NewOwningNode->Pins[ExistingPinIdx] == this);
+		}
+	}
+}
+
+int32 NumSubPins = 0;
+
+void UEdGraphPin::InvalidateAndTrash()
+{
+	check(!bWasTrashed);
+	PinsToDelete.Add(this);
+	BreakAllPinLinks();
+	SetOwningNode(nullptr);
+	NumSubPins += SubPins.Num();
+	SubPins.Empty();
+	ParentPin = nullptr;
+	ReferencePassThroughConnection = nullptr;
+	bWasTrashed = true;
+}
+
+UEdGraphPin* UEdGraphPin::CreatePinFromDeprecatedPin(UEdGraphPin_Deprecated* DeprecatedPin)
+{
+	if (!DeprecatedPin)
+	{
+		return nullptr;
+	}
+
+	UEdGraphNode* DeprecatedPinOwningNode = CastChecked<UEdGraphNode>(DeprecatedPin->GetOuter());
+
+	UEdGraphPin* NewPin = CreatePin(DeprecatedPinOwningNode);
+	DeprecatedPinOwningNode->Pins.Add(NewPin);
+	NewPin->InitFromDeprecatedPin(DeprecatedPin);
+
+	return NewPin;
+}
+
+UEdGraphPin* UEdGraphPin::FindPinCreatedFromDeprecatedPin(UEdGraphPin_Deprecated* DeprecatedPin)
+{
+	if (!DeprecatedPin)
+	{
+		return nullptr;
+	}
+
+	UEdGraphNode* DeprecatedPinOwningNode = CastChecked<UEdGraphNode>(DeprecatedPin->GetOuter());
+	
+	UEdGraphPin* ReturnPin = nullptr;
+	FGuid PinGuidThatWasCreatedFromDeprecatedPin = PinHelpers::DeprecatedPinToNewPinGUIDMap.FindRef(DeprecatedPin);
+	if (PinGuidThatWasCreatedFromDeprecatedPin.IsValid())
+	{
+		UEdGraphPin** ExistingPin = DeprecatedPinOwningNode->Pins.FindByPredicate([&PinGuidThatWasCreatedFromDeprecatedPin](const UEdGraphPin* Pin) { return Pin && Pin->PinId == PinGuidThatWasCreatedFromDeprecatedPin; });
+		if (ExistingPin)
+		{
+			ReturnPin = *ExistingPin;
+		}
+	}
+
+	return ReturnPin;
+}
+
+bool UEdGraphPin::ExportTextItem(FString& ValueStr, int32 PortFlags) const
+{
+	const UStrProperty* StrPropCDO = GetDefault<UStrProperty>();
+	const UTextProperty* TextPropCDO = GetDefault<UTextProperty>();
+	const UBoolProperty* BoolPropCDO = GetDefault<UBoolProperty>();
+	static UEdGraphPin DefaultPin(nullptr, FGuid());
+
+	ValueStr += "(";
+
+	ValueStr += PinHelpers::PinIdName + TEXT("=");
+	PinId.ExportTextItem(ValueStr, FGuid(), nullptr, PortFlags, nullptr);
+	ValueStr += PinHelpers::ExportTextPropDelimiter;
+
+	if (PinName != DefaultPin.PinName)
+	{
+		ValueStr += PinHelpers::PinNameName + TEXT("=");
+		StrPropCDO->ExportTextItem(ValueStr, &PinName, nullptr, nullptr, PortFlags, nullptr);
+		ValueStr += PinHelpers::ExportTextPropDelimiter;
+	}
+
+#if WITH_EDITORONLY_DATA
+	if (!PinFriendlyName.EqualTo(DefaultPin.PinFriendlyName))
+	{
+		ValueStr += PinHelpers::PinFriendlyNameName + TEXT("=");
+		TextPropCDO->ExportTextItem(ValueStr, &PinFriendlyName, nullptr, nullptr, PortFlags, nullptr);
+		ValueStr += PinHelpers::ExportTextPropDelimiter;
+	}
+#endif // WITH_EDITORONLY_DATA
+
+	if (PinToolTip != DefaultPin.PinToolTip)
+	{
+		ValueStr += PinHelpers::PinToolTipName + TEXT("=");
+		StrPropCDO->ExportTextItem(ValueStr, &PinToolTip, nullptr, nullptr, PortFlags, nullptr);
+		ValueStr += PinHelpers::ExportTextPropDelimiter;
+	}
+
+	if (Direction != DefaultPin.Direction)
+	{
+		const FString DirectionString = UEnum::GetValueAsString(TEXT("/Script/Engine.EEdGraphPinDirection"), Direction);
+		ValueStr += PinHelpers::DirectionName + TEXT("=");
+		StrPropCDO->ExportTextItem(ValueStr, &DirectionString, nullptr, nullptr, PortFlags, nullptr);
+		ValueStr += PinHelpers::ExportTextPropDelimiter;
+	}
+
+	for (TFieldIterator<UProperty> FieldIt(FEdGraphPinType::StaticStruct()); FieldIt; ++FieldIt)
+	{
+		FString PropertyStr;
+		const uint8* PropertyAddr = FieldIt->ContainerPtrToValuePtr<uint8>(&PinType);
+		const uint8* DefaultAddr = FieldIt->ContainerPtrToValuePtr<uint8>(&DefaultPin.PinType);
+		FieldIt->ExportTextItem(PropertyStr, PropertyAddr, DefaultAddr, NULL, PortFlags, nullptr);
+
+		if (!PropertyStr.IsEmpty())
+		{
+			ValueStr += PinHelpers::PinTypeName + TEXT(".") + FieldIt->GetName() + "=" + PropertyStr;
+			ValueStr += PinHelpers::ExportTextPropDelimiter;
+		}
+	}
+
+	if (DefaultValue != DefaultPin.DefaultValue)
+	{
+		ValueStr += PinHelpers::DefaultValueName + TEXT("=");
+		StrPropCDO->ExportTextItem(ValueStr, &DefaultValue, nullptr, nullptr, PortFlags, nullptr);
+		ValueStr += PinHelpers::ExportTextPropDelimiter;
+	}
+
+	if (AutogeneratedDefaultValue != DefaultPin.AutogeneratedDefaultValue)
+	{
+		ValueStr += PinHelpers::AutogeneratedDefaultValueName + TEXT("=");
+		StrPropCDO->ExportTextItem(ValueStr, &AutogeneratedDefaultValue, nullptr, nullptr, PortFlags, nullptr);
+		ValueStr += PinHelpers::ExportTextPropDelimiter;
+	}
+
+	if (DefaultObject != DefaultPin.DefaultObject)
+	{
+		const FString DefaultObjectPath = GetPathNameSafe(DefaultObject);
+		ValueStr += PinHelpers::DefaultObjectName + TEXT("=");
+		StrPropCDO->ExportTextItem(ValueStr, &DefaultObjectPath, nullptr, nullptr, PortFlags, nullptr);
+		ValueStr += PinHelpers::ExportTextPropDelimiter;
+	}
+
+	if (!DefaultTextValue.EqualTo(DefaultPin.DefaultTextValue))
+	{
+		ValueStr += PinHelpers::DefaultTextValueName + TEXT("=");
+		TextPropCDO->ExportTextItem(ValueStr, &DefaultTextValue, nullptr, nullptr, PortFlags, nullptr);
+		ValueStr += PinHelpers::ExportTextPropDelimiter;
+	}
+
+	if (LinkedTo.Num() > 0)
+	{
+		ValueStr += PinHelpers::LinkedToName + TEXT("=") + UEdGraphPin::ExportText_PinArray(LinkedTo);
+		ValueStr += PinHelpers::ExportTextPropDelimiter;
+	}
+
+	if (SubPins.Num() > 0)
+	{
+		ValueStr += PinHelpers::SubPinsName + TEXT("=") + UEdGraphPin::ExportText_PinArray(SubPins);
+		ValueStr += PinHelpers::ExportTextPropDelimiter;
+	}
+
+	if (ParentPin)
+	{
+		ValueStr += PinHelpers::ParentPinName + TEXT("=") + UEdGraphPin::ExportText_PinReference(ParentPin);
+		ValueStr += PinHelpers::ExportTextPropDelimiter;
+	}
+
+	if (ReferencePassThroughConnection)
+	{
+		ValueStr += PinHelpers::ReferencePassThroughConnectionName + TEXT("=") + UEdGraphPin::ExportText_PinReference(ReferencePassThroughConnection);
+		ValueStr += PinHelpers::ExportTextPropDelimiter;
+	}
+
+#if WITH_EDITORONLY_DATA
+	{
+		ValueStr += PinHelpers::PersistentGuidName + TEXT("=");
+		PersistentGuid.ExportTextItem(ValueStr, FGuid(), nullptr, PortFlags, nullptr);
+		ValueStr += PinHelpers::ExportTextPropDelimiter;
+
+		ValueStr += PinHelpers::bHiddenName + TEXT("=");
+		bool LocalHidden = bHidden;
+		BoolPropCDO->ExportTextItem(ValueStr, &LocalHidden, nullptr, nullptr, PortFlags, nullptr);
+		ValueStr += PinHelpers::ExportTextPropDelimiter;
+
+		ValueStr += PinHelpers::bNotConnectableName + TEXT("=");
+		bool LocalNotConnectable = bNotConnectable;
+		BoolPropCDO->ExportTextItem(ValueStr, &LocalNotConnectable, nullptr, nullptr, PortFlags, nullptr);
+		ValueStr += PinHelpers::ExportTextPropDelimiter;
+
+		ValueStr += PinHelpers::bDefaultValueIsReadOnlyName + TEXT("=");
+		bool LocalDefaultValueIsReadOnly = bDefaultValueIsReadOnly;
+		BoolPropCDO->ExportTextItem(ValueStr, &LocalDefaultValueIsReadOnly, nullptr, nullptr, PortFlags, nullptr);
+		ValueStr += PinHelpers::ExportTextPropDelimiter;
+
+		ValueStr += PinHelpers::bDefaultValueIsIgnoredName + TEXT("=");
+		bool LocalDefaultValueIsIgnored = bDefaultValueIsIgnored;
+		BoolPropCDO->ExportTextItem(ValueStr, &LocalDefaultValueIsIgnored, nullptr, nullptr, PortFlags, nullptr);
+		ValueStr += PinHelpers::ExportTextPropDelimiter;
+
+		// Intentionally not exporting bIsDiffing as it is transient
+
+		ValueStr += PinHelpers::bAdvancedViewName + TEXT("=");
+		bool LocalAdvancedView = bAdvancedView;
+		BoolPropCDO->ExportTextItem(ValueStr, &LocalAdvancedView, nullptr, nullptr, PortFlags, nullptr);
+		ValueStr += PinHelpers::ExportTextPropDelimiter;
+
+		// Intentionally not exporting bDisplayAsMutableRef as it is transient
+		// Intentionally not exporting bWasTrashed as it is transient
+	}
+#endif //WITH_EDITORONLY_DATA
+	ValueStr += ")";
+
+	return true;
+}
+
+static void SkipWhitespace(const TCHAR*& Str)
+{
+	while (Str && FChar::IsWhitespace(*Str))
+	{
+		Str++;
+	}
+}
+
+bool UEdGraphPin::ImportTextItem(const TCHAR*& Buffer, int32 PortFlags, class UObject* Parent, FOutputDevice* ErrorText)
 {
-	Super::PostLoad();
+	PinId = FGuid();
+
+	const UStrProperty* StrPropCDO = GetDefault<UStrProperty>();
+	const UTextProperty* TextPropCDO = GetDefault<UTextProperty>();
+	const UBoolProperty* BoolPropCDO = GetDefault<UBoolProperty>();
+	UEnum* PinDirectionEnum = FindObjectChecked<UEnum>(nullptr, TEXT("/Script/Engine.EEdGraphPinDirection"));
+
+	SkipWhitespace(Buffer);
+	if (*Buffer != '(')
+	{
+		// Parse error
+		return false;
+	}
+
+	Buffer++;
+	SkipWhitespace(Buffer);
+
+	while (*Buffer != ')')
+	{
+		const TCHAR* StartBuffer = Buffer;
+		if (*Buffer == 0)
+		{
+			// Parse error
+			return false;
+		}
+
+		while (*Buffer != '=')
+		{
+			Buffer++;
+		}
+
+		if (*Buffer == 0)
+		{
+			// Parse error
+			return false;
+		}
+
+		int32 NumCharsInToken = Buffer - StartBuffer;
+		if (NumCharsInToken <= 0)
+		{
+			// Parse error
+			return false;
+		}
 
+		// Advance over the '='
+		Buffer++;
+
+		FString PropertyToken(NumCharsInToken, StartBuffer);
+		PropertyToken = PropertyToken.TrimTrailing();
+		bool bParseSuccess = false;
+		if (PropertyToken == PinHelpers::PinIdName)
+		{
+			bParseSuccess = PinId.ImportTextItem(Buffer, PortFlags, nullptr, ErrorText);
+		}
+		else if (PropertyToken == PinHelpers::PinNameName)
+		{
+			Buffer = StrPropCDO->ImportText(Buffer, &PinName, PortFlags, nullptr, ErrorText);
+			bParseSuccess = (Buffer != nullptr);
+		}
+#if WITH_EDITORONLY_DATA
+		else if (PropertyToken == PinHelpers::PinFriendlyNameName)
+		{
+			Buffer = TextPropCDO->ImportText(Buffer, &PinFriendlyName, PortFlags, nullptr, ErrorText);
+			bParseSuccess = (Buffer != nullptr);
+		}
+#endif //WITH_EDITORONLY_DATA
+		else if (PropertyToken == PinHelpers::PinToolTipName)
+		{
+			Buffer = StrPropCDO->ImportText(Buffer, &PinToolTip, PortFlags, nullptr, ErrorText);
+			bParseSuccess = (Buffer != nullptr);
+		}
+		else if (PropertyToken == PinHelpers::DirectionName)
+		{
+			FString DirectionString;
+			Buffer = StrPropCDO->ImportText(Buffer, &DirectionString, PortFlags, nullptr, ErrorText);
+			bParseSuccess = (Buffer != nullptr);
+			if (bParseSuccess)
+			{
+				int32 EnumIdx = PinDirectionEnum->GetIndexByName(FName(*DirectionString));
+				if (EnumIdx != INDEX_NONE)
+				{
+					Direction = (EEdGraphPinDirection)EnumIdx;
+				}
+				else
+				{
+					// Invalid enum entry
+					bParseSuccess = false;
+				}
+			}
+		}
+		else if (PropertyToken.StartsWith(PinHelpers::PinTypeName + TEXT(".")))
+		{
+			FString Left;
+			FString Right;
+			if (ensure(PropertyToken.Split(".", &Left, &Right, ESearchCase::CaseSensitive)))
+			{
+				UProperty* FoundProp = FEdGraphPinType::StaticStruct()->FindPropertyByName(FName(*Right));
+				if (FoundProp)
+				{
+					uint8* PropertyAddr = FoundProp->ContainerPtrToValuePtr<uint8>(&PinType);
+					Buffer = FoundProp->ImportText(Buffer, PropertyAddr, PortFlags, nullptr, ErrorText);
+					bParseSuccess = (Buffer != nullptr);
+				}
+				else
+				{
+					// Couldn't find the prop
+					bParseSuccess = false;
+				}
+			}
+			else
+			{
+				// Eh, split failed to find the '.' we saw in the token?
+				bParseSuccess = false;
+			}
+		}
+		else if (PropertyToken == PinHelpers::DefaultValueName)
+		{
+			Buffer = StrPropCDO->ImportText(Buffer, &DefaultValue, PortFlags, nullptr, ErrorText);
+			bParseSuccess = (Buffer != nullptr);
+		}
+		else if (PropertyToken == PinHelpers::AutogeneratedDefaultValueName)
+		{
+			Buffer = StrPropCDO->ImportText(Buffer, &AutogeneratedDefaultValue, PortFlags, nullptr, ErrorText);
+			bParseSuccess = (Buffer != nullptr);
+		}
+		else if (PropertyToken == PinHelpers::DefaultObjectName)
+		{
+			FString DefaultObjectString;
+			Buffer = StrPropCDO->ImportText(Buffer, &DefaultObjectString, PortFlags, nullptr, ErrorText);
+			bParseSuccess = (Buffer != nullptr);
+			if (bParseSuccess)
+			{
+				DefaultObject = FindObject<UObject>(nullptr, *DefaultObjectString);
+			}
+		}
+		else if (PropertyToken == PinHelpers::DefaultTextValueName)
+		{
+			Buffer = TextPropCDO->ImportText(Buffer, &DefaultTextValue, PortFlags, nullptr, ErrorText);
+			bParseSuccess = (Buffer != nullptr);
+		}
+		else if (PropertyToken == PinHelpers::LinkedToName)
+		{
+			bParseSuccess = UEdGraphPin::ImportText_PinArray(Buffer, LinkedTo, this, EPinResolveType::LinkedTo);
+		}
+		else if (PropertyToken == PinHelpers::SubPinsName)
+		{
+			bParseSuccess = UEdGraphPin::ImportText_PinArray(Buffer, SubPins, this, EPinResolveType::SubPins);
+		}
+		else if (PropertyToken == PinHelpers::ParentPinName)
+		{
+			bParseSuccess = UEdGraphPin::ImportText_PinReference(Buffer, ParentPin, INDEX_NONE, this, EPinResolveType::ParentPin);
+		}
+		else if (PropertyToken == PinHelpers::ReferencePassThroughConnectionName)
+		{
+			bParseSuccess = UEdGraphPin::ImportText_PinReference(Buffer, ReferencePassThroughConnection, INDEX_NONE, this, EPinResolveType::ReferencePassThroughConnection);
+		}
+#if WITH_EDITORONLY_DATA
+		else if (PropertyToken == PinHelpers::PersistentGuidName)
+		{
+			bParseSuccess = PersistentGuid.ImportTextItem(Buffer, PortFlags, nullptr, ErrorText);
+		}
+		else if (PropertyToken == PinHelpers::bHiddenName)
+		{
+			bool LocalHidden = bHidden;
+			Buffer = BoolPropCDO->ImportText(Buffer, &LocalHidden, PortFlags, nullptr, ErrorText);
+			bParseSuccess = (Buffer != nullptr);
+			bHidden = LocalHidden;
+		}
+		else if (PropertyToken == PinHelpers::bNotConnectableName)
+		{
+			bool LocalNotConnectable = bNotConnectable;
+			Buffer = BoolPropCDO->ImportText(Buffer, &LocalNotConnectable, PortFlags, nullptr, ErrorText);
+			bParseSuccess = (Buffer != nullptr);
+			bNotConnectable = LocalNotConnectable;
+		}
+		else if (PropertyToken == PinHelpers::bDefaultValueIsReadOnlyName)
+		{
+			bool LocalDefaultValueIsReadOnly = bDefaultValueIsReadOnly;
+			Buffer = BoolPropCDO->ImportText(Buffer, &LocalDefaultValueIsReadOnly, PortFlags, nullptr, ErrorText);
+			bParseSuccess = (Buffer != nullptr);
+			bDefaultValueIsReadOnly = LocalDefaultValueIsReadOnly;
+		}
+		else if (PropertyToken == PinHelpers::bDefaultValueIsIgnoredName)
+		{
+			bool LocalDefaultValueIsIgnored = bDefaultValueIsIgnored;
+			Buffer = BoolPropCDO->ImportText(Buffer, &LocalDefaultValueIsIgnored, PortFlags, nullptr, ErrorText);
+			bParseSuccess = (Buffer != nullptr);
+			bDefaultValueIsIgnored = LocalDefaultValueIsIgnored;
+		}
+		// Intentionally not importing bIsDiffing as it is transient
+		else if (PropertyToken == PinHelpers::bAdvancedViewName)
+		{
+			bool LocalAdvancedView = bAdvancedView;
+			Buffer = BoolPropCDO->ImportText(Buffer, &LocalAdvancedView, PortFlags, nullptr, ErrorText);
+			bParseSuccess = (Buffer != nullptr);
+			bAdvancedView = LocalAdvancedView;
+		}
+		// Intentionally not importing bDisplayAsMutableRef as it is transient
+		// Intentionally not importing bWasTrashed as it is transient
+#endif
+
+		if (!bParseSuccess)
+		{
+			// Parse error
+			return false;
+		}
+
+		SkipWhitespace(Buffer);
+
+		if (*Buffer == PinHelpers::ExportTextPropDelimiter)
+		{
+			Buffer++;
+		}
+
+		SkipWhitespace(Buffer);
+	}
+
+	// Advance over the last ')'
+	Buffer++;
+
+	if (!PinId.IsValid())
+	{
+		// Did not contain a valid PinId
+		return false;
+	}
+
+	// Someone might have been waiting for this pin to be created. Let them know that this pin now exists.
+	ResolveReferencesToPin(this);
+
+	return true;
+}
+
+void UEdGraphPin::MarkPendingKill()
+{
+	if (!bWasTrashed)
+	{
+		PinsToDelete.Add(this);
+#if WITH_EDITOR
+		if (!GIsTransacting)
+		{
+#endif //WITH_EDITOR
+			BreakAllPinLinks();
+			if (ParentPin)
+			{
+				ParentPin->SubPins.Remove(this);
+			}
+#if WITH_EDITOR
+		}
+		else
+		{
+			LinkedTo.Empty();
+		}
+#endif //WITH_EDITOR
+		OwningNode = nullptr;
+		NumSubPins = 0;
+		SubPins.Empty();
+		ParentPin = nullptr;
+		ReferencePassThroughConnection = nullptr;
+		bWasTrashed = true;
+	}
+}
+
+void UEdGraphPin::ShutdownVerification()
+{
+	Purge();
+	// There's a static UEdGraphPin in UEdGraphPin::ExportTextItem so if that code
+	// has run we'll have a single pin 'in memory' on shutdown
+	ensure(PinHelpers::NumPinsInMemory == 0 || PinHelpers::NumPinsInMemory == 1);
+}
+
+void UEdGraphPin::Purge()
+{
+	PinDeletionQueue->Tick(0.f);
+}
+
+void UEdGraphPin::QueueForDestruction()
+{
+	PinsToDelete.Add(this);
+}
+
+UEdGraphPin::UEdGraphPin(UEdGraphNode* InOwningNode, const FGuid& PinIdGuid)
+	: OwningNode(InOwningNode)
+	, PinId(PinIdGuid)
+	, PinName()
+#if WITH_EDITORONLY_DATA
+	, PinFriendlyName()
+#endif // WITH_EDITORONLY_DATA
+	, PinToolTip()
+	, Direction(EGPD_Input)
+	, PinType()
+	, DefaultValue()
+	, AutogeneratedDefaultValue()
+	, DefaultObject(nullptr)
+	, DefaultTextValue()
+	, LinkedTo()
+	, SubPins()
+	, ParentPin(nullptr)
+	, ReferencePassThroughConnection(nullptr)
+#if WITH_EDITORONLY_DATA
+	, PersistentGuid()
+	, bHidden(false)
+	, bNotConnectable(false)
+	, bDefaultValueIsReadOnly(false)
+	, bDefaultValueIsIgnored(false)
+	, bIsDiffing(false)
+	, bAdvancedView(false)
+	, bDisplayAsMutableRef(false)
+#endif
+	, bWasTrashed(false)
+{
+	PinHelpers::NumPinsInMemory++;
+#ifdef TRACK_PINS
+	PinAllocationTracking.Add(TPairInitializer<UEdGraphPin*, FString>(this, InOwningNode ? InOwningNode->GetName() : FString(TEXT("UNOWNED"))));
+#endif //TRACK_PINS
+}
+
+UEdGraphPin::~UEdGraphPin()
+{
+	check(PinHelpers::NumPinsInMemory > 0);
+	PinHelpers::NumPinsInMemory--;
+#ifdef TRACK_PINS
+	PinAllocationTracking.RemoveAll([this](const TPair<UEdGraphPin*, FString>& Entry) { return Entry.Key == this; });
+#endif //TRACK_PINS
+}
+
+void UEdGraphPin::ResolveAllPinReferences()
+{
+	for (auto& Entry : PinHelpers::UnresolvedPins)
+	{
+		FPinResolveId Key = Entry.Key;
+		UEdGraphNode* Node = Key.OwningNode.Get();
+		if (!Node)
+		{
+			continue;
+		}
+
+		const FGuid& LocalGuid = Key.PinGuid;
+		UEdGraphPin** TargetPin = Node->Pins.FindByPredicate([&LocalGuid](const UEdGraphPin* Pin) { return Pin && Pin->PinId == LocalGuid; });
+		if (TargetPin)
+		{
+			ResolveReferencesToPin(*TargetPin);
+		}
+	}
+	PinHelpers::UnresolvedPins.Empty();
+}
+
+void UEdGraphPin::InitFromDeprecatedPin(class UEdGraphPin_Deprecated* DeprecatedPin)
+{
+	check(DeprecatedPin);
+	PinHelpers::DeprecatedPinToNewPinGUIDMap.Add(DeprecatedPin, PinId);
+	DeprecatedPin->FixupDefaultValue();
+
+	OwningNode = CastChecked<UEdGraphNode>(DeprecatedPin->GetOuter());
+
+	PinName = DeprecatedPin->PinName;
+
+#if WITH_EDITORONLY_DATA
+	PinFriendlyName = DeprecatedPin->PinFriendlyName;
+#endif
+
+	PinToolTip = DeprecatedPin->PinToolTip;
+	Direction = DeprecatedPin->Direction;
+	PinType = DeprecatedPin->PinType;
+	DefaultValue = DeprecatedPin->DefaultValue;
+	AutogeneratedDefaultValue = DeprecatedPin->AutogeneratedDefaultValue;
+	DefaultObject = DeprecatedPin->DefaultObject;
+	DefaultTextValue = DeprecatedPin->DefaultTextValue;
+	
+	LinkedTo.Empty(DeprecatedPin->LinkedTo.Num());
+	LinkedTo.SetNum(DeprecatedPin->LinkedTo.Num());
+	for (int32 PinIdx = 0; PinIdx < DeprecatedPin->LinkedTo.Num(); ++PinIdx)
+	{
+		UEdGraphPin_Deprecated* OldLinkedToPin = DeprecatedPin->LinkedTo[PinIdx];
+		UEdGraphPin*& NewPin = LinkedTo[PinIdx];
+		if (OldLinkedToPin)
+		{
+			UEdGraphPin* ExistingPin = FindPinCreatedFromDeprecatedPin(OldLinkedToPin);
+			if (ExistingPin)
+			{
+				NewPin = ExistingPin;
+			}
+			else
+			{
+				TArray<FUnresolvedPinData>& UnresolvedPinData = PinHelpers::UnresolvedDeprecatedPins.FindOrAdd(OldLinkedToPin);
+				UnresolvedPinData.Add(FUnresolvedPinData(this, EPinResolveType::LinkedTo, PinIdx));
+			}
+		}
+	}
+	
+	SubPins.Empty(DeprecatedPin->SubPins.Num());
+	SubPins.SetNum(DeprecatedPin->SubPins.Num());
+	for (int32 PinIdx = 0; PinIdx < DeprecatedPin->SubPins.Num(); ++PinIdx)
+	{
+		UEdGraphPin_Deprecated* OldSubPinsPin = DeprecatedPin->SubPins[PinIdx];
+		UEdGraphPin*& NewPin = SubPins[PinIdx];
+		if (OldSubPinsPin)
+		{
+			UEdGraphPin* ExistingPin = FindPinCreatedFromDeprecatedPin(OldSubPinsPin);
+			if (ExistingPin)
+			{
+				NewPin = ExistingPin;
+			}
+			else
+			{
+				TArray<FUnresolvedPinData>& UnresolvedPinData = PinHelpers::UnresolvedDeprecatedPins.FindOrAdd(OldSubPinsPin);
+				UnresolvedPinData.Add(FUnresolvedPinData(this, EPinResolveType::SubPins, PinIdx));
+			}
+		}
+	}
+	
+	if (DeprecatedPin->ParentPin)
+	{
+		UEdGraphPin* ExistingPin = FindPinCreatedFromDeprecatedPin(DeprecatedPin->ParentPin);
+		if (ExistingPin)
+		{
+			ParentPin = ExistingPin;
+		}
+		else
+		{
+			TArray<FUnresolvedPinData>& UnresolvedPinData = PinHelpers::UnresolvedDeprecatedPins.FindOrAdd(DeprecatedPin->ParentPin);
+			UnresolvedPinData.Add(FUnresolvedPinData(this, EPinResolveType::ParentPin, INDEX_NONE));
+		}
+	}
+	
+	if (DeprecatedPin->ReferencePassThroughConnection)
+	{
+		UEdGraphPin* ExistingPin = FindPinCreatedFromDeprecatedPin(DeprecatedPin->ReferencePassThroughConnection);
+		if (ExistingPin)
+		{
+			ReferencePassThroughConnection = ExistingPin;
+		}
+		else
+		{
+			TArray<FUnresolvedPinData>& UnresolvedPinData = PinHelpers::UnresolvedDeprecatedPins.FindOrAdd(DeprecatedPin->ReferencePassThroughConnection);
+			UnresolvedPinData.Add(FUnresolvedPinData(this, EPinResolveType::ReferencePassThroughConnection, INDEX_NONE));
+		}
+	}
+
+#if WITH_EDITORONLY_DATA
+	bHidden = DeprecatedPin->bHidden;
+	bNotConnectable = DeprecatedPin->bNotConnectable;
+	bDefaultValueIsReadOnly = DeprecatedPin->bDefaultValueIsReadOnly;
+	bDefaultValueIsIgnored = DeprecatedPin->bDefaultValueIsIgnored;
+	bIsDiffing = DeprecatedPin->bIsDiffing;
+	bAdvancedView = DeprecatedPin->bAdvancedView;
+	bDisplayAsMutableRef = DeprecatedPin->bDisplayAsMutableRef;
+	PersistentGuid = DeprecatedPin->PersistentGuid;
+#endif // WITH_EDITORONLY_DATA
+
+	// Now resolve any references this pin which was created from a deprecated pin. This should only be from other pins that were created from deprecated pins.
+	TArray<FUnresolvedPinData>* UnresolvedPinData = PinHelpers::UnresolvedDeprecatedPins.Find(DeprecatedPin);
+	if (UnresolvedPinData)
+	{
+		for (const FUnresolvedPinData& PinData : *UnresolvedPinData)
+		{
+			UEdGraphPin* ReferencingPin = PinData.ReferencingPin;
+			if (ReferencingPin)
+			{
+				switch (PinData.ResolveType)
+				{
+				case EPinResolveType::LinkedTo:
+					ReferencingPin->LinkedTo[PinData.ArrayIdx] = this;
+					ensure(LinkedTo.Contains(ReferencingPin));
+					break;
+
+				case EPinResolveType::SubPins:
+					ReferencingPin->SubPins[PinData.ArrayIdx] = this;
+					break;
+
+				case EPinResolveType::ParentPin:
+					ReferencingPin->ParentPin = this;
+					break;
+
+				case EPinResolveType::ReferencePassThroughConnection:
+					ReferencingPin->ReferencePassThroughConnection = this;
+					break;
+
+				case EPinResolveType::OwningNode:
+				default:
+					checkf(0, TEXT("Unhandled ResolveType %d"), PinData.ResolveType);
+					break;
+				}
+			}
+		}
+
+		PinHelpers::UnresolvedDeprecatedPins.Remove(DeprecatedPin);
+		UnresolvedPinData = nullptr;
+	}
+}
+
+bool UEdGraphPin::Serialize(FArchive& Ar)
+{
+	// These properties are in every pin and are unlikely to be removed, so they are native serialized for speed.
+	Ar << OwningNode;
+	Ar << PinId;
+
+	Ar << PinName;
+
+#if WITH_EDITORONLY_DATA
+	if (!Ar.IsFilterEditorOnly())
+	{
+		Ar << PinFriendlyName;
+	}
+#endif
+
+	Ar << PinToolTip;
+	Ar << Direction;
+	PinType.Serialize(Ar);
+	Ar << DefaultValue;
+	Ar << AutogeneratedDefaultValue;
+	Ar << DefaultObject;
+	Ar << DefaultTextValue;
+
+	// There are several options to solve the problem wwe're fixing here. First,
+	// the problem: we have two UObjects (A and B) that own Pin arrays, and within
+	// those Pin arrays there are pins with references to pins owned by the other UObject
+	// in LinkedTo. We could:
+	//  * Use FEdGraphPinReference to 'softly' reference the other pins by outer+name
+	//  * Defer resolving of UEdGraphPin* in volatile contexts (e.g. when serialzing)
+	//  * Aggressively fix up references:
+	// Our current approach is to defer resolving of UEdGraphPin*
+	UEdGraphPin::SerializePinArray(Ar, LinkedTo, this, EPinResolveType::LinkedTo);
+
+	// SubPin array is owned by a single UObject, so no complexity of keeping LinkedTo synchronized:
+	UEdGraphPin::SerializePinArray(Ar, SubPins, this, EPinResolveType::SubPins);
+	TArray<UEdGraphPin*> NoExistingValues; // unused for non-owner paths..
+ 	SerializePin(Ar, ParentPin, INDEX_NONE, this, EPinResolveType::ParentPin, NoExistingValues);
+ 	SerializePin(Ar, ReferencePassThroughConnection, INDEX_NONE, this, EPinResolveType::ReferencePassThroughConnection, NoExistingValues);
+
+#if WITH_EDITORONLY_DATA
+	if (!Ar.IsFilterEditorOnly())
+	{
+		Ar << PersistentGuid;
+
+		// Do not reorder these! If you are adding a new one, add it to the end, and if you want to remove one, leave the others at the same offset in the uint32
+
+		uint32 BitField = 0;
+		BitField |= bHidden << 0;
+		BitField |= bNotConnectable << 1;
+		BitField |= bDefaultValueIsReadOnly << 2;
+		BitField |= bDefaultValueIsIgnored << 3;
+		// Intentionally not including bIsDiffing. It is transient.
+		BitField |= bAdvancedView << 4;
+		// Intentionally not including bDisplayAsMutableRef. It is transient.
+		// Intentionally not including bWasTrashed. It is transient.
+
+		Ar << BitField;
+
+		bHidden = !!(BitField & 0x01);
+		bNotConnectable = !!(BitField & 0x02);
+		bDefaultValueIsReadOnly = !!(BitField & 0x04);
+		bDefaultValueIsIgnored = !!(BitField & 0x08);
+		// Intentionally not including bIsDiffing. It is transient.
+		bAdvancedView = !!(BitField & 0x10);
+		// Intentionally not including bDisplayAsMutableRef. It is transient.
+		// Intentionally not including bWasTrashed. It is transient.
+	}
+#endif
+
+	// Strictly speaking we could resolve now when transacting, but
+	// we would still have to do the resolve step later to catch
+	// any references that are created after we serialize a given node:
+	if (Ar.IsLoading()
+#if WITH_EDITOR 
+		&& !Ar.IsTransacting()
+#endif // WITH_EDITOR
+		)
+	{
+		ResolveReferencesToPin(this);
+	}
+
+	return true;
+}
+
+void UEdGraphPin::EnableAllConnectedNodes(UEdGraphNode* InNode)
+{
+	// Only enable when it has not been explicitly user-disabled
+	if (InNode && InNode->EnabledState == ENodeEnabledState::Disabled && !InNode->bUserSetEnabledState)
+	{
+		// Enable the node and clear the comment
+		InNode->Modify();
+		InNode->EnableNode();
+		InNode->NodeComment.Empty();
+
+		// Go through all pin connections and enable the nodes. Enabled nodes will prevent further iteration
+		for (UEdGraphPin*  Pin : InNode->Pins)
+		{
+			for (UEdGraphPin* OtherPin : Pin->LinkedTo)
+			{
+				EnableAllConnectedNodes(OtherPin->GetOwningNode());
+			}
+		}
+	}
+}
+
+void UEdGraphPin::ResolveReferencesToPin(UEdGraphPin* Pin)
+{
+	FPinResolveId ResolveId(Pin->PinId, Pin->OwningNode);
+	TArray<FUnresolvedPinData>* UnresolvedPinData = PinHelpers::UnresolvedPins.Find(ResolveId);
+	if (UnresolvedPinData)
+	{
+		for (const FUnresolvedPinData& PinData : *UnresolvedPinData)
+		{
+			UEdGraphPin* ReferencingPin = PinData.ReferencingPin;
+			switch (PinData.ResolveType)
+			{
+			case EPinResolveType::LinkedTo:
+				ReferencingPin->LinkedTo[PinData.ArrayIdx] = Pin;
+#if WITH_EDITOR
+				// When in the middle of a transaction the LinkedTo lists will be in an 
+				// indeterminate state. After PostEditUndo runs we could validate LinkedTo
+				// coherence.
+				ensureAlways(GIsTransacting || Pin->LinkedTo.Contains(ReferencingPin));
+#else
+				ensureAlways(Pin->LinkedTo.Contains(ReferencingPin));
+#endif//WITH_EDITOR
+				break;
+
+			case EPinResolveType::SubPins:
+				ReferencingPin->SubPins[PinData.ArrayIdx] = Pin;
+				break;
+
+			case EPinResolveType::ParentPin:
+				ReferencingPin->ParentPin = Pin;
+				break;
+
+			case EPinResolveType::ReferencePassThroughConnection:
+				ReferencingPin->ReferencePassThroughConnection = Pin;
+				break;
+
+			case EPinResolveType::OwningNode:
+			default:
+				checkf(0, TEXT("Unhandled ResolveType %d"), PinData.ResolveType);
+				break;
+			}
+		}
+
+		PinHelpers::UnresolvedPins.Remove(ResolveId);
+	}
+}
+
+void UEdGraphPin::SerializePinArray(FArchive& Ar, TArray<UEdGraphPin*>& ArrayRef, UEdGraphPin* RequestingPin, EPinResolveType ResolveType)
+{
+	const int32 NumCurrentEntries = ArrayRef.Num();
+	int32 ArrayNum = NumCurrentEntries;
+	Ar << ArrayNum;
+	
+	TArray<UEdGraphPin*> OldPins = Ar.IsLoading() && ResolveType == EPinResolveType::OwningNode ? ArrayRef : TArray<UEdGraphPin*>();
+
+	if (Ar.IsLoading())
+	{
+		ArrayRef.SetNum(ArrayNum);
+	}
+
+	for (int32 PinIdx = 0; PinIdx < ArrayNum; ++PinIdx)
+	{
+		UEdGraphPin*& PinRef = ArrayRef[PinIdx];
+		SerializePin(Ar, PinRef, PinIdx, RequestingPin, ResolveType, OldPins);
+	}
+
+	// free unused pins:
+	for (UEdGraphPin* Pin : OldPins)
+	{
+		Pin->QueueForDestruction();
+	}
+}
+
+bool UEdGraphPin::SerializePin(FArchive& Ar, UEdGraphPin*& PinRef, int32 ArrayIdx, UEdGraphPin* RequestingPin, EPinResolveType ResolveType, TArray<UEdGraphPin*>& OldPins)
+{
+	bool bRetVal = true;
+
+	bool bNullPtr = PinRef == nullptr;
+	Ar << bNullPtr;
+	if (bNullPtr)
+	{
+		if (Ar.IsLoading())
+		{
+			PinRef = nullptr;
+		}
+	}
+	else
+	{
+		UEdGraphNode* LocalOwningNode = nullptr;
+		FGuid PinGuid;
+		if (!Ar.IsLoading())
+		{
+			check(!PinRef->bWasTrashed);
+			LocalOwningNode = PinRef->OwningNode;
+			PinGuid = PinRef->PinId;
+		}
+
+		Ar << LocalOwningNode;
+		Ar << PinGuid;
+		check(LocalOwningNode);
+		check(PinGuid.IsValid());
+
+		if (ResolveType != EPinResolveType::OwningNode)
+		{
+			if (Ar.IsLoading())
+			{
+				UEdGraphPin** ExistingPin = 
+#if WITH_EDITOR
+					// When transacting we're in a volatile state, and we cannot trust any pins we might find, so we defer all 
+					// pin resolving:
+					Ar.IsTransacting() ? nullptr : 
+#endif//WITH_EDITOR
+					LocalOwningNode->Pins.FindByPredicate([&PinGuid](const UEdGraphPin* Pin) { return Pin && Pin->PinId == PinGuid; });
+				if (ExistingPin)
+				{
+					PinRef = *ExistingPin;
+				}
+				else
+				{
+					check(RequestingPin);
+					TArray<FUnresolvedPinData>& UnresolvedPinData = PinHelpers::UnresolvedPins.FindOrAdd(FPinResolveId(PinGuid, LocalOwningNode));
+					UnresolvedPinData.Add(FUnresolvedPinData(RequestingPin, ResolveType, ArrayIdx));
+					bRetVal = false;
+				}
+			}
+		}
+		else
+		{
+			if (Ar.IsLoading())
+			{
+				UEdGraphPin** PinToReuse = OldPins.FindByPredicate([&PinGuid](const UEdGraphPin* Pin) { return Pin && Pin->PinId == PinGuid; });
+				if (PinToReuse)
+				{
+					PinRef = *PinToReuse;
+					OldPins.RemoveAtSwap(PinToReuse - OldPins.GetData());
+				}
+				else
+				{
+					PinRef = UEdGraphPin::CreatePin(LocalOwningNode);
+				}
+			}
+
+			PinRef->Serialize(Ar);
+		}
+	}
+
+	return bRetVal;
+}
+
+FString UEdGraphPin::ExportText_PinReference(const UEdGraphPin* Pin)
+{
+	if (Pin)
+	{
+		const FString OwningNodeString = Pin->GetOwningNodeUnchecked() ? Pin->OwningNode->GetName() : TEXT("null");
+		return OwningNodeString + " " + Pin->PinId.ToString();
+	}
+
+	return FString();
+}
+
+FString UEdGraphPin::ExportText_PinArray(const TArray<UEdGraphPin*>& PinArray)
+{
+	FString RetVal;
+
+	RetVal += "(";
+	if (PinArray.Num() > 0)
+	{
+		for (UEdGraphPin* Pin : PinArray)
+		{
+			RetVal += ExportText_PinReference(Pin);
+			RetVal += PinHelpers::ExportTextPropDelimiter;
+		}
+	}
+	RetVal += ")";
+
+	return RetVal;
+}
+
+bool UEdGraphPin::ImportText_PinReference(const TCHAR*& Buffer, UEdGraphPin*& PinRef, int32 ArrayIdx, UEdGraphPin* RequestingPin, EPinResolveType ResolveType)
+{
+	const TCHAR* BufferStart = Buffer;
+	while (*Buffer != PinHelpers::ExportTextPropDelimiter && *Buffer != ')')
+	{
+		if (*Buffer == 0)
+		{
+			// Parse error
+			return false;
+		}
+
+		Buffer++;
+	}
+
+	FString PinRefLine(Buffer - BufferStart, BufferStart);
+	FString OwningNodeString;
+	FString PinGuidString;
+	if (PinRefLine.Split(TEXT(" "), &OwningNodeString, &PinGuidString, ESearchCase::CaseSensitive, ESearchDir::FromEnd))
+	{
+		FGuid PinGuid;
+		if (FGuid::Parse(PinGuidString, PinGuid))
+		{
+			UEdGraphNode* LocalOwningNode = FindObject<UEdGraphNode>(RequestingPin->GetOwningNode()->GetOuter(), *OwningNodeString);
+			if (LocalOwningNode)
+			{
+				UEdGraphPin** ExistingPin = LocalOwningNode->Pins.FindByPredicate([&PinGuid](const UEdGraphPin* Pin) { return Pin && Pin->PinId == PinGuid; });
+				if (ExistingPin)
+				{
+					PinRef = *ExistingPin;
+				}
+				else
+				{
+					TArray<FUnresolvedPinData>& UnresolvedPinData = PinHelpers::UnresolvedPins.FindOrAdd(FPinResolveId(PinGuid, LocalOwningNode));
+					UnresolvedPinData.Add(FUnresolvedPinData(RequestingPin, ResolveType, ArrayIdx));
+				}
+			}
+		}
+	}
+
+	return true;
+}
+
+bool UEdGraphPin::ImportText_PinArray(const TCHAR*& Buffer, TArray<UEdGraphPin*>& ArrayRef, UEdGraphPin* RequestingPin, EPinResolveType ResolveType)
+{
+	bool bSuccess = false;
+
+	SkipWhitespace(Buffer);
+	if (*Buffer == '(')
+	{
+		Buffer++;
+		SkipWhitespace(Buffer);
+		bool bParseError = false;
+		while (*Buffer != ')')
+		{
+			if (*Buffer == 0)
+			{
+				return false;
+			}
+
+			int32 NewItemIdx = ArrayRef.Add(nullptr);
+			bool bParseSuccess = ImportText_PinReference(Buffer, ArrayRef[NewItemIdx], NewItemIdx, RequestingPin, ResolveType);
+			if (!bParseSuccess)
+			{
+				return false;
+			}
+
+			if (*Buffer == PinHelpers::ExportTextPropDelimiter)
+			{
+				// Advance over the ','
+				Buffer++;
+			}
+
+			SkipWhitespace(Buffer);
+		}
+
+		// Advance over the last ')'
+		Buffer++;
+
+		return true;
+	}
+
+	return false;
+}
+
+/////////////////////////////////////////////////////
+// UEdGraphPin
+
+UEdGraphPin_Deprecated::UEdGraphPin_Deprecated(const FObjectInitializer& ObjectInitializer)
+	: Super(ObjectInitializer)
+{
+#if WITH_EDITORONLY_DATA
+	bHidden = false;
+	bNotConnectable = false;
+	bAdvancedView = false;
+#endif // WITH_EDITORONLY_DATA
+}
+
+void UEdGraphPin_Deprecated::FixupDefaultValue()
+{
 	static FName GameplayTagName = TEXT("GameplayTag");
 	static FName GameplayTagContainerName = TEXT("GameplayTagContainer");
 	static FName GameplayTagsPathName = TEXT("/Script/GameplayTags");
diff --git a/Engine/Source/Runtime/Engine/Private/EdGraph/EdGraphSchema.cpp b/Engine/Source/Runtime/Engine/Private/EdGraph/EdGraphSchema.cpp
--- a/Engine/Source/Runtime/Engine/Private/EdGraph/EdGraphSchema.cpp
+++ b/Engine/Source/Runtime/Engine/Private/EdGraph/EdGraphSchema.cpp
@@ -566,11 +566,15 @@ void UEdGraphSchema::BreakSinglePinLink(UEdGraphPin* SourcePin, UEdGraphPin* Tar
 	SourcePin->BreakLinkTo(TargetPin);
 
 #if WITH_EDITOR
-	TargetPin->GetOwningNode()->PinConnectionListChanged(TargetPin);
-	SourcePin->GetOwningNode()->PinConnectionListChanged(SourcePin);
-
-	TargetPin->GetOwningNode()->NodeConnectionListChanged();
-	SourcePin->GetOwningNode()->NodeConnectionListChanged();
+	// get a reference to these now as the following calls can potentially clear the OwningNode (ex: split pins in MakeArray nodes)
+	UEdGraphNode* TargetNode = TargetPin->GetOwningNode();
+	UEdGraphNode* SourceNode = SourcePin->GetOwningNode();
+	
+	TargetNode->PinConnectionListChanged(TargetPin);
+	SourceNode->PinConnectionListChanged(SourcePin);
+	
+	TargetNode->NodeConnectionListChanged();
+	SourceNode->NodeConnectionListChanged();
 #endif	//#if WITH_EDITOR
 }
 
diff --git a/Engine/Source/Runtime/Engine/Private/KismetMathLibrary.cpp b/Engine/Source/Runtime/Engine/Private/KismetMathLibrary.cpp
--- a/Engine/Source/Runtime/Engine/Private/KismetMathLibrary.cpp
+++ b/Engine/Source/Runtime/Engine/Private/KismetMathLibrary.cpp
@@ -1306,6 +1306,11 @@ FVector2D UKismetMathLibrary::Multiply_Vector2DFloat(FVector2D A, float B)
 	return A * B;
 }
 
+FVector2D UKismetMathLibrary::Multiply_Vector2DVector2D(FVector2D A, FVector2D B)
+{
+	return A * B;
+}
+
 FVector2D UKismetMathLibrary::Divide_Vector2DFloat(FVector2D A, float B)
 {
 	if (B == 0.f)
@@ -2291,7 +2296,7 @@ void UKismetMathLibrary::MinimumAreaRectangle(class UObject* WorldContextObject,
 
 	if( bDebugDraw )
 	{
-		UWorld* World = GEngine->GetWorldFromContextObject(WorldContextObject);
+		UWorld* World = (WorldContextObject) ? GEngine->GetWorldFromContextObject(WorldContextObject) : nullptr;
 		if(World != nullptr)
 		{
 			DrawDebugSphere(World, OutRectCenter, 10.f, 12, FColor::Yellow, true);
@@ -2299,6 +2304,10 @@ void UKismetMathLibrary::MinimumAreaRectangle(class UObject* WorldContextObject,
 			DrawDebugLine(World, OutRectCenter - RectSideA * 0.5f + FVector(0,0,10.f), OutRectCenter + RectSideA * 0.5f + FVector(0,0,10.f), FColor::Green, true,-1, 0, 5.f);
 			DrawDebugLine(World, OutRectCenter - RectSideB * 0.5f + FVector(0,0,10.f), OutRectCenter + RectSideB * 0.5f + FVector(0,0,10.f), FColor::Blue, true,-1, 0, 5.f);
 		}
+		else
+		{
+			FFrame::KismetExecutionMessage(TEXT("WorldContext required for MinimumAreaRectangle to draw a debug visualization."), ELogVerbosity::Warning);
+		}
 	}
 }
 
diff --git a/Engine/Source/Runtime/Engine/Private/KismetTextLibrary.cpp b/Engine/Source/Runtime/Engine/Private/KismetTextLibrary.cpp
--- a/Engine/Source/Runtime/Engine/Private/KismetTextLibrary.cpp
+++ b/Engine/Source/Runtime/Engine/Private/KismetTextLibrary.cpp
@@ -17,6 +17,62 @@ UKismetTextLibrary::UKismetTextLibrary(const FObjectInitializer& ObjectInitializ
 /* UKismetTextLibrary static functions
  *****************************************************************************/
 
+FText UKismetTextLibrary::Conv_VectorToText(FVector InVec)
+{
+	FFormatNamedArguments Args;
+	Args.Add(TEXT("X"), InVec.X);
+	Args.Add(TEXT("Y"), InVec.Y);
+	Args.Add(TEXT("Z"), InVec.Z);
+
+	return FText::Format(NSLOCTEXT("Core", "Vector3", "X={X} Y={Y} Z={Z}"), Args);
+}
+
+
+FText UKismetTextLibrary::Conv_Vector2dToText(FVector2D InVec)
+{
+	FFormatNamedArguments Args;
+	Args.Add(TEXT("X"), InVec.X);
+	Args.Add(TEXT("Y"), InVec.Y);
+
+	return FText::Format(NSLOCTEXT("Core", "Vector2", "X={X} Y={Y}"), Args);
+}
+
+
+FText UKismetTextLibrary::Conv_RotatorToText(FRotator InRot)
+{
+	FFormatNamedArguments Args;
+	Args.Add(TEXT("P"), InRot.Pitch);
+	Args.Add(TEXT("Y"), InRot.Yaw);
+	Args.Add(TEXT("R"), InRot.Roll);
+
+	return FText::Format(NSLOCTEXT("Core", "Rotator", "P={P} Y={Y} R={R}"), Args);
+}
+
+
+FText UKismetTextLibrary::Conv_TransformToText(const FTransform& InTrans)
+{
+	return FText::FromString(InTrans.ToString());
+}
+
+
+FText UKismetTextLibrary::Conv_ObjectToText(class UObject* InObj)
+{
+	return FText::FromString((InObj != NULL) ? InObj->GetName() : FString(TEXT("None")));
+}
+
+
+FText UKismetTextLibrary::Conv_ColorToText(FLinearColor InColor)
+{
+	FFormatNamedArguments Args;
+	Args.Add(TEXT("R"), InColor.R);
+	Args.Add(TEXT("G"), InColor.G);
+	Args.Add(TEXT("B"), InColor.B);
+	Args.Add(TEXT("A"), InColor.A);
+
+	return FText::Format(NSLOCTEXT("Core", "LinearColor", "R={R} G={G} B={B} A={A}"), Args);
+}
+
+
 FString UKismetTextLibrary::Conv_TextToString(const FText& InText)
 {
 	return InText.ToString();
diff --git a/Engine/Source/Runtime/Engine/Private/UserDefinedEnum.cpp b/Engine/Source/Runtime/Engine/Private/UserDefinedEnum.cpp
--- a/Engine/Source/Runtime/Engine/Private/UserDefinedEnum.cpp
+++ b/Engine/Source/Runtime/Engine/Private/UserDefinedEnum.cpp
@@ -116,8 +116,9 @@ FText UUserDefinedEnum::GetEnumText(int32 InIndex) const
 	return Super::GetEnumText(InIndex);
 }
 
-bool UUserDefinedEnum::SetEnums(TArray<TPair<FName, uint8>>& InNames, ECppForm InCppForm)
+bool UUserDefinedEnum::SetEnums(TArray<TPair<FName, uint8>>& InNames, ECppForm InCppForm, bool bAddMaxKeyIfMissing)
 {
+	ensure(bAddMaxKeyIfMissing);
 	if (Names.Num() > 0)
 	{
 		RemoveNamesFromMasterList();
diff --git a/Engine/Source/Runtime/Engine/Public/GeneratedCodeHelpers.h b/Engine/Source/Runtime/Engine/Public/GeneratedCodeHelpers.h
--- a/Engine/Source/Runtime/Engine/Public/GeneratedCodeHelpers.h
+++ b/Engine/Source/Runtime/Engine/Public/GeneratedCodeHelpers.h
@@ -100,12 +100,6 @@ public:
 		return TargetArray.Add(NewItem);
 	}
 
-	template<typename T, typename U>
-	static int32 Array_AddUnique(TArray<T>& TargetArray, const U& NewItem)
-	{
-		return TargetArray.AddUnique(NewItem);
-	}
-
 	template<typename T>
 	static void Array_Shuffle(TArray<T>& TargetArray)
 	{
@@ -158,12 +152,103 @@ public:
 		return TargetArray.Find(ItemToFind);
 	}
 
+	template<typename T>
+	static int32 Array_Find_Struct(const TArray<T>& TargetArray, const T& ItemToFind)
+	{
+		auto ScriptStruct = T::StaticStruct();
+		return TargetArray.IndexOfByPredicate([&](const T& Element) -> bool
+		{
+			return ScriptStruct->CompareScriptStruct(&Element, &ItemToFind, 0);
+		});
+	}
+
+	static int32 Array_Find_FText(const TArray<FText>& TargetArray, const FText& ItemToFind)
+	{
+		return TargetArray.IndexOfByPredicate([&](const FText& Element) -> bool
+		{
+			return UTextProperty::Identical_Implementation(Element, ItemToFind, 0);
+		});
+	}
+
+	template<typename T, typename U>
+	static bool Array_Contains(const TArray<T>& TargetArray, const U& ItemToFind)
+	{
+		return TargetArray.Contains(ItemToFind);
+	}
+
+	template<typename T>
+	static bool Array_Contains_Struct(const TArray<T>& TargetArray, const T& ItemToFind)
+	{
+		auto ScriptStruct = T::StaticStruct();
+		return TargetArray.ContainsByPredicate([&](const T& Element) -> bool
+		{
+			return ScriptStruct->CompareScriptStruct(&Element, &ItemToFind, 0);
+		});
+	}
+
+	static bool Array_Contains_FText(const TArray<FText>& TargetArray, const FText& ItemToFind)
+	{
+		return TargetArray.ContainsByPredicate([&](const FText& Element) -> bool
+		{
+			return UTextProperty::Identical_Implementation(Element, ItemToFind, 0);
+		});
+	}
+
+	template<typename T, typename U>
+	static int32 Array_AddUnique(TArray<T>& TargetArray, const U& NewItem)
+	{
+		return TargetArray.AddUnique(NewItem);
+	}
+
+	template<typename T>
+	static int32 Array_AddUnique_Struct(TArray<T>& TargetArray, const T& NewItem)
+	{
+		int32 Index = Array_Find_Struct<T>(TargetArray, NewItem);
+		if (Index != INDEX_NONE)
+		{
+			return Index;
+		}
+		return TargetArray.Add(NewItem);
+	}
+
+	static int32 Array_AddUnique_FText(TArray<FText>& TargetArray, const FText& NewItem)
+	{
+		int32 Index = Array_Find_FText(TargetArray, NewItem);
+		if (Index != INDEX_NONE)
+		{
+			return Index;
+		}
+		return TargetArray.Add(NewItem);
+	}
+
 	template<typename T, typename U>
 	static bool Array_RemoveItem(TArray<T>& TargetArray, const U& Item)
 	{
 		return TargetArray.Remove(Item) != 0;
 	}
 
+	template<typename T>
+	static bool Array_RemoveItem_Struct(TArray<T>& TargetArray, const T& Item)
+	{
+		TargetArray.CheckAddress(&Item);
+
+		auto ScriptStruct = T::StaticStruct();
+		return TargetArray.RemoveAll([&](const T& Element) -> bool
+		{
+			return ScriptStruct->CompareScriptStruct(&Element, &Item, 0);
+		}) != 0;
+	}
+
+	static bool Array_RemoveItem_FText(TArray<FText>& TargetArray, const FText& Item)
+	{
+		TargetArray.CheckAddress(&Item);
+
+		return TargetArray.RemoveAll([&](const FText& Element) -> bool
+		{
+			return UTextProperty::Identical_Implementation(Element, Item, 0);
+		}) != 0;
+	}
+
 	template<typename T>
 	static void Array_Clear(TArray<T>& TargetArray)
 	{
@@ -228,12 +313,6 @@ public:
 		}
 	}
 
-	template<typename T, typename U>
-	static bool Array_Contains(const TArray<T>& TargetArray, const U& ItemToFind)
-	{
-		return TargetArray.Contains(ItemToFind);
-	}
-
 	template<typename T>
 	static void SetArrayPropertyByName(UObject* Object, FName PropertyName, TArray<T>& Value)
 	{
diff --git a/Engine/Source/Runtime/Engine/Public/InstancedReferenceSubobjectHelper.h b/Engine/Source/Runtime/Engine/Public/InstancedReferenceSubobjectHelper.h
--- a/Engine/Source/Runtime/Engine/Public/InstancedReferenceSubobjectHelper.h
+++ b/Engine/Source/Runtime/Engine/Public/InstancedReferenceSubobjectHelper.h
@@ -2,88 +2,173 @@
 
 #pragma once
 
-class FFindInstancedReferenceSubobjectHelper
+/** 
+ * Meant to represent a specific object property that is setup to reference a 
+ * instanced sub-object. Tracks the property hierarchy used to reach the 
+ * property, so that we can easily retrieve instanced sub-objects from a 
+ * container object.
+ */
+struct FInstancedPropertyPath
 {
+private:
+	struct FPropertyLink
+	{
+		FPropertyLink(const UProperty* Property, int32 ArrayIndexIn = INDEX_NONE)
+			: PropertyPtr(Property), ArrayIndex(ArrayIndexIn)
+		{}
+
+		const UProperty* PropertyPtr;
+		int32            ArrayIndex;
+	};
+
 public:
-	template<typename T>
-	static void Get(const UStruct* Struct, const uint8* ContainerAddress, T& OutObjects)
+	//--------------------------------------------------------------------------
+	FInstancedPropertyPath(UProperty* RootProperty)
 	{
-		check(ContainerAddress);
-		for (UProperty* Prop = (Struct ? Struct->RefLink : NULL); Prop; Prop = Prop->NextRef)
+		PropertyChain.Add(FPropertyLink(RootProperty));
+	}
+
+	//--------------------------------------------------------------------------
+	void Push(const UProperty* Property, int32 ArrayIndex = INDEX_NONE)
+	{
+		PropertyChain.Add(FPropertyLink(Property, ArrayIndex));		
+	}
+
+	//--------------------------------------------------------------------------
+	void Pop()
+	{
+ 		PropertyChain.RemoveAt(PropertyChain.Num() - 1);
+	}
+
+	//--------------------------------------------------------------------------
+	const UProperty* Head() const
+	{
+		return PropertyChain.Last().PropertyPtr;
+	}
+
+	//--------------------------------------------------------------------------
+	UObject* Resolve(const UObject* Container) const
+	{
+		UStruct* CurrentContainerType = Container->GetClass();
+
+		const TArray<FPropertyLink>& PropChainRef = PropertyChain;
+		auto GetProperty = [&CurrentContainerType, &PropChainRef](int32 ChainIndex)->UProperty*
 		{
-			auto ArrayProperty = Cast<const UArrayProperty>(Prop);
-			if (ArrayProperty && Prop->HasAnyPropertyFlags(CPF_PersistentInstance | CPF_ContainsInstancedReference))
-			{
-				auto InnerStructProperty = Cast<const UStructProperty>(ArrayProperty->Inner);
-				auto InnerObjectProperty = Cast<const UObjectProperty>(ArrayProperty->Inner);
-				if (InnerStructProperty && InnerStructProperty->Struct)
-				{
-					FScriptArrayHelper_InContainer ArrayHelper(ArrayProperty, ContainerAddress);
-					for (int32 ElementIndex = 0; ElementIndex < ArrayHelper.Num(); ++ElementIndex)
-					{
-						const uint8* ValueAddress = ArrayHelper.GetRawPtr(ElementIndex);
-						if (ValueAddress)
-						{
-							Get(InnerStructProperty->Struct, ValueAddress, OutObjects);
-						}
-					}
-				}
-				else if (InnerObjectProperty && InnerObjectProperty->HasAllPropertyFlags(CPF_PersistentInstance))
-				{
-					ensure(InnerObjectProperty->HasAllPropertyFlags(CPF_InstancedReference));
-					FScriptArrayHelper_InContainer ArrayHelper(ArrayProperty, ContainerAddress);
-					for (int32 ElementIndex = 0; ElementIndex < ArrayHelper.Num(); ++ElementIndex)
-					{
-						UObject* ObjectValue = InnerObjectProperty->GetObjectPropertyValue(ArrayHelper.GetRawPtr(ElementIndex));
-						if (ObjectValue)
-						{
-							OutObjects.Add(ObjectValue);
-						}
-					}
-				}
-			}
-			else if (Prop->HasAllPropertyFlags(CPF_PersistentInstance))
+			const UProperty* SrcProperty = PropChainRef[ChainIndex].PropertyPtr;
+			return FindField<UProperty>(CurrentContainerType, SrcProperty->GetFName());
+		};
+
+		UProperty* CurrentProp = GetProperty(0);
+		const uint8* ValuePtr = (CurrentProp) ? CurrentProp->ContainerPtrToValuePtr<uint8>(Container) : nullptr;
+
+		for (int32 ChainIndex = 1; CurrentProp && ChainIndex < PropertyChain.Num(); ++ChainIndex)
+		{
+			if (const UArrayProperty* ArrayProperty = Cast<UArrayProperty>(CurrentProp))
 			{
-				ensure(Prop->HasAllPropertyFlags(CPF_InstancedReference));
-				auto ObjectProperty = Cast<const UObjectProperty>(Prop);
-				if (ObjectProperty)
+				check(PropertyChain[ChainIndex].PropertyPtr == ArrayProperty->Inner);
+
+				int32 TargetIndex = PropertyChain[ChainIndex].ArrayIndex;
+				check(TargetIndex != INDEX_NONE);
+
+				FScriptArrayHelper ArrayHelper(ArrayProperty, ValuePtr);
+				if (TargetIndex >= ArrayHelper.Num())
 				{
-					for (int32 ArrayIdx = 0; ArrayIdx < Prop->ArrayDim; ++ArrayIdx)
-					{
-						auto ObjectValue = ObjectProperty->GetObjectPropertyValue_InContainer(ContainerAddress, ArrayIdx);
-						if (ObjectValue)
-						{
-							OutObjects.Add(ObjectValue);
-						}
-					}
+					CurrentProp = nullptr;
+					break;
 				}
+
+				CurrentProp = ArrayProperty->Inner;
+				ValuePtr    = ArrayHelper.GetRawPtr(TargetIndex);
 			}
-			else if (Prop->HasAnyPropertyFlags(CPF_ContainsInstancedReference))
+			else if (ensure(PropertyChain[ChainIndex].ArrayIndex <= 0))
 			{
-				auto StructProperty = Cast<const UStructProperty>(Prop);
-				if (StructProperty && StructProperty->Struct)
+				if (const UStructProperty* StructProperty = Cast<UStructProperty>(CurrentProp))
 				{
-					for (int32 ArrayIdx = 0; ArrayIdx < Prop->ArrayDim; ++ArrayIdx)
-					{
-						const uint8* ValueAddress = StructProperty->ContainerPtrToValuePtr<uint8>(ContainerAddress, ArrayIdx);
-						Get(StructProperty->Struct, ValueAddress, OutObjects);
-					}
+					CurrentContainerType = StructProperty->Struct;
 				}
+
+				CurrentProp = GetProperty(ChainIndex);
+				ValuePtr    = (CurrentProp) ? CurrentProp->ContainerPtrToValuePtr<uint8>(ValuePtr) : nullptr;
 			}
 		}
+
+		const UObjectProperty* TargetPropety = Cast<UObjectProperty>(CurrentProp);
+		if ( ensure(TargetPropety && TargetPropety->HasAnyPropertyFlags(CPF_InstancedReference)) )
+		{ 
+			return TargetPropety->GetObjectPropertyValue(ValuePtr);
+		}
+		return nullptr;
+	}
+
+private:
+	TArray<FPropertyLink> PropertyChain;
+};
+
+/** 
+ * Can be used as a raw sub-object pointer, but also contains a 
+ * FInstancedPropertyPath to identify the property that this sub-object is 
+ * referenced by. Paired together for ease of use (so API users don't have to manage a map).
+ */
+struct FInstancedSubObjRef
+{
+	FInstancedSubObjRef(UObject* SubObj, const FInstancedPropertyPath& PropertyPathIn)
+		: SubObjInstance(SubObj)
+		, PropertyPath(PropertyPathIn)
+	{}
+
+	//--------------------------------------------------------------------------
+	operator UObject*() const
+	{
+		return SubObjInstance;
 	}
 
+	//--------------------------------------------------------------------------
+	UObject* operator->() const
+	{
+		return SubObjInstance;
+	}
+
+	//--------------------------------------------------------------------------
+	friend uint32 GetTypeHash(const FInstancedSubObjRef& SubObjRef)
+	{
+		return GetTypeHash((UObject*)SubObjRef);
+	}
+
+	UObject* SubObjInstance;
+	FInstancedPropertyPath PropertyPath;
+};
+ 
+/** 
+ * Contains a set of utility functions useful for searching out and identifying
+ * instanced sub-objects contained within a specific outer object.
+ */
+class FFindInstancedReferenceSubobjectHelper
+{
+public:
+	//--------------------------------------------------------------------------
+	template<typename T>
+	static void GetInstancedSubObjects(const UObject* Container, T& OutObjects)
+	{
+		const UClass* ContainerClass = Container->GetClass();
+		for (UProperty* Prop = ContainerClass->RefLink; Prop; Prop = Prop->NextRef)
+		{
+			FInstancedPropertyPath RootPropertyPath(Prop);
+			GetInstancedSubObjects_Inner(RootPropertyPath, reinterpret_cast<const uint8*>(Container), OutObjects);
+		}
+	}
+
+	//--------------------------------------------------------------------------
 	static void Duplicate(UObject* OldObject, UObject* NewObject, TMap<UObject*, UObject*>& ReferenceReplacementMap, TArray<UObject*>& DuplicatedObjects)
 	{
 		if (OldObject->GetClass()->HasAnyClassFlags(CLASS_HasInstancedReference) &&
 			NewObject->GetClass()->HasAnyClassFlags(CLASS_HasInstancedReference))
 		{
-			TSet<UObject*> OldEditInlineObjects;
-			Get(OldObject->GetClass(), reinterpret_cast<uint8*>(OldObject), OldEditInlineObjects);
+			TArray<FInstancedSubObjRef> OldEditInlineObjects;
+			GetInstancedSubObjects(OldObject, OldEditInlineObjects);
 			if (OldEditInlineObjects.Num())
 			{
-				TSet<UObject*> NewEditInlineObjects;
-				Get(NewObject->GetClass(), reinterpret_cast<uint8*>(NewObject), NewEditInlineObjects);
+				TArray<FInstancedSubObjRef> NewEditInlineObjects;
+				GetInstancedSubObjects(NewObject, NewEditInlineObjects);
 				for (auto Obj : NewEditInlineObjects)
 				{
 					const bool bProperOuter = (Obj->GetOuter() == OldObject);
@@ -97,8 +182,9 @@ public:
 							UObject* NewEditInlineSubobject = StaticDuplicateObject(Obj, NewObject);
 							ReferenceReplacementMap.Add(Obj, NewEditInlineSubobject);
 
-							// Patch the linker table with the new instance.
-							FLinkerLoad::PRIVATE_PatchNewObjectIntoExport(Obj, NewEditInlineSubobject);
+							// NOTE: we cannot patch OldObject's linker table here, since we don't 
+							//       know the relation between the  two objects (one could be of a 
+							//       super class, and the other a child)
 
 							// We also need to make sure to fixup any properties here
 							DuplicatedObjects.Add(NewEditInlineSubobject);
@@ -108,4 +194,87 @@ public:
 			}
 		}
 	}
+
+private:
+	//--------------------------------------------------------------------------
+	template<typename T>
+	static void GetInstancedSubObjects_Inner(FInstancedPropertyPath& PropertyPath, const uint8* ContainerAddress, T& OutObjects)
+	{
+		check(ContainerAddress);
+		const UProperty* TargetProp = PropertyPath.Head();
+
+		const UArrayProperty* ArrayProperty = Cast<const UArrayProperty>(TargetProp);
+		if (ArrayProperty && ArrayProperty->HasAnyPropertyFlags(CPF_PersistentInstance | CPF_ContainsInstancedReference))
+		{
+			const UStructProperty* InnerStructProperty = Cast<const UStructProperty>(ArrayProperty->Inner);
+			const UObjectProperty* InnerObjectProperty = Cast<const UObjectProperty>(ArrayProperty->Inner);
+			if (InnerStructProperty && InnerStructProperty->Struct)
+			{
+				FScriptArrayHelper_InContainer ArrayHelper(ArrayProperty, ContainerAddress);
+				for (int32 ElementIndex = 0; ElementIndex < ArrayHelper.Num(); ++ElementIndex)
+				{
+					const uint8* ValueAddress = ArrayHelper.GetRawPtr(ElementIndex);
+					if (ValueAddress)
+					{
+						for (UProperty* StructProp = InnerStructProperty->Struct->RefLink; StructProp; StructProp = StructProp->NextRef)
+						{
+							PropertyPath.Push(InnerStructProperty, ElementIndex);
+							GetInstancedSubObjects_Inner(PropertyPath, ValueAddress, OutObjects);
+							PropertyPath.Pop();
+						}
+					}
+				}
+			}
+			else if (InnerObjectProperty && InnerObjectProperty->HasAllPropertyFlags(CPF_PersistentInstance))
+			{
+				ensure(InnerObjectProperty->HasAllPropertyFlags(CPF_InstancedReference));
+				FScriptArrayHelper_InContainer ArrayHelper(ArrayProperty, ContainerAddress);
+				for (int32 ElementIndex = 0; ElementIndex < ArrayHelper.Num(); ++ElementIndex)
+				{
+					UObject* ObjectValue = InnerObjectProperty->GetObjectPropertyValue(ArrayHelper.GetRawPtr(ElementIndex));
+					if (ObjectValue)
+					{
+						PropertyPath.Push(InnerObjectProperty, ElementIndex);
+						OutObjects.Add(FInstancedSubObjRef(ObjectValue, PropertyPath));
+						PropertyPath.Pop();
+					}
+					
+				}
+			}
+		}
+		else if (TargetProp->HasAllPropertyFlags(CPF_PersistentInstance))
+		{
+			ensure(TargetProp->HasAllPropertyFlags(CPF_InstancedReference));
+			const UObjectProperty* ObjectProperty = Cast<const UObjectProperty>(TargetProp);
+			if (ObjectProperty)
+			{
+				for (int32 ArrayIdx = 0; ArrayIdx < ObjectProperty->ArrayDim; ++ArrayIdx)
+				{
+					UObject* ObjectValue = ObjectProperty->GetObjectPropertyValue_InContainer(ContainerAddress, ArrayIdx);
+					if (ObjectValue)
+					{
+						// don't need to push to PropertyPath, since this property is already at its head
+						OutObjects.Add(FInstancedSubObjRef(ObjectValue, PropertyPath));	
+					}
+				}
+			}
+		}
+		else if (TargetProp->HasAnyPropertyFlags(CPF_ContainsInstancedReference))
+		{
+			const UStructProperty* StructProperty = Cast<const UStructProperty>(TargetProp);
+			if (StructProperty && StructProperty->Struct)
+			{
+				for (int32 ArrayIdx = 0; ArrayIdx < StructProperty->ArrayDim; ++ArrayIdx)
+				{
+					const uint8* ValueAddress = StructProperty->ContainerPtrToValuePtr<uint8>(ContainerAddress, ArrayIdx);
+					for (UProperty* StructProp = StructProperty->Struct->RefLink; StructProp; StructProp = StructProp->NextRef)
+					{
+						PropertyPath.Push(StructProp, ArrayIdx);
+						GetInstancedSubObjects_Inner(PropertyPath, ValueAddress, OutObjects);
+						PropertyPath.Pop();
+					}
+				}
+			}
+		}
+	}
 };
\ No newline at end of file
diff --git a/Engine/Source/Runtime/GameplayAbilities/Private/GameplayCueManager.cpp b/Engine/Source/Runtime/GameplayAbilities/Private/GameplayCueManager.cpp
--- a/Engine/Source/Runtime/GameplayAbilities/Private/GameplayCueManager.cpp
+++ b/Engine/Source/Runtime/GameplayAbilities/Private/GameplayCueManager.cpp
@@ -444,6 +444,51 @@ void UGameplayCueManager::InitObjectLibraries(TArray<FString> Paths, UObjectLibr
 	BuildCuesToAddToGlobalSet(ActorAssetDatas, GET_MEMBER_NAME_CHECKED(AGameplayCueNotify_Actor, GameplayCueName), CuesToAdd, AssetsToLoad, ShouldLoadDelegate);
 	BuildCuesToAddToGlobalSet(StaticAssetDatas, GET_MEMBER_NAME_CHECKED(UGameplayCueNotify_Static, GameplayCueName), CuesToAdd, AssetsToLoad, ShouldLoadDelegate);
 
+	{
+		// Iterate over all Dynamic Classes (nativized Blueprints). Search for ones with GameplayCueName tag.
+
+		IGameplayTagsModule& GameplayTagsModule = IGameplayTagsModule::Get();
+		const FName PropertyName = GET_MEMBER_NAME_CHECKED(AGameplayCueNotify_Actor, GameplayCueName);
+		check(PropertyName == GET_MEMBER_NAME_CHECKED(UGameplayCueNotify_Static, GameplayCueName));
+		TMap<FName, FDynamicClassStaticData>& DynamicClassMap = GetDynamicClassMap();
+		for (auto PairIter : DynamicClassMap)
+		{
+			const FName* FoundGameplayTag = PairIter.Value.SelectedSearchableValues.Find(PropertyName);
+			if (!FoundGameplayTag)
+			{
+				continue;
+			}
+
+			const FString ClassPath = PairIter.Key.ToString();
+			for (FString& Path : Paths)
+			{
+				const bool PathContainsClass = ClassPath.StartsWith(Path); // TODO: is it enough?
+				if (!PathContainsClass)
+				{
+					continue;
+				}
+
+				ABILITY_LOG(Log, TEXT("GameplayCueManager Found a Dynamic Class: %s / %s"), *FoundGameplayTag->ToString(), *ClassPath);
+
+				FGameplayTag  GameplayCueTag = GameplayTagsModule.GetGameplayTagsManager().RequestGameplayTag(*FoundGameplayTag, false);
+				if (GameplayCueTag.IsValid())
+				{
+					FStringAssetReference StringRef(ClassPath); // TODO: is there any translation needed?
+					ensure(StringRef.IsValid());
+
+					CuesToAdd.Add(FGameplayCueReferencePair(GameplayCueTag, StringRef));
+					AssetsToLoad.Add(StringRef);
+				}
+				else
+				{
+					ABILITY_LOG(Warning, TEXT("Found GameplayCue tag %s in Dynamic Class %s but there is no corresponding tag in the GameplayTagManager."), *FoundGameplayTag->ToString(), *ClassPath);
+				}
+
+				break;
+			}
+		}
+	}
+
 	// Add these cues to the global set
 	check(GlobalCueSet);
 	GlobalCueSet->AddCues(CuesToAdd);
