revert engine change from WEX merge as it is not a good change

diff --git a/Engine/Source/Runtime/Engine/Private/PlayerController.cpp b/Engine/Source/Runtime/Engine/Private/PlayerController.cpp
--- a/Engine/Source/Runtime/Engine/Private/PlayerController.cpp
+++ b/Engine/Source/Runtime/Engine/Private/PlayerController.cpp
@@ -91,9 +91,6 @@ APlayerController::APlayerController(const FObjectInitializer& ObjectInitializer
 	bIsLocalPlayerController = false;
 
 	ClickEventKeys.Add(EKeys::LeftMouseButton);
-	// @pjs - adding right mouse button for HSL change, this may just need to be part of a derived class
-	ClickEventKeys.Add(EKeys::RightMouseButton);
-	// @pjs
 
 	if (RootComponent)
 	{
@@ -3816,9 +3813,7 @@ void APlayerController::TickPlayerInput(const float DeltaSeconds, const bool bGa
 			UGameViewportClient* ViewportClient = LocalPlayer->ViewportClient;
 
 			// Only send mouse hit events if we're directly over the viewport.
-//@HSL_BEGIN_PF - CCL - This code does not work, not sure why yet.  Returns false unless you click, briefly
-			//if ( ViewportClient && ViewportClient->GetGameViewportWidget().IsValid() && ViewportClient->GetGameViewportWidget()->IsDirectlyHovered() )
-//@HSL_END
+			if ( ViewportClient && ViewportClient->GetGameViewportWidget().IsValid() && ViewportClient->GetGameViewportWidget()->IsDirectlyHovered() )
 			{
 				if ( LocalPlayer->ViewportClient->GetMousePosition(MousePosition) )
 				{
