fix for not being able to break gems in WEX

diff --git a/Engine/Source/Runtime/Engine/Classes/GameFramework/PlayerController.h b/Engine/Source/Runtime/Engine/Classes/GameFramework/PlayerController.h
--- a/Engine/Source/Runtime/Engine/Classes/GameFramework/PlayerController.h
+++ b/Engine/Source/Runtime/Engine/Classes/GameFramework/PlayerController.h
@@ -1341,6 +1341,7 @@ protected:
 	virtual void ProcessPlayerInput(const float DeltaTime, const bool bGamePaused);
 	virtual void BuildInputStack(TArray<UInputComponent*>& InputStack);
 	void ProcessForceFeedbackAndHaptics(const float DeltaTime, const bool bGamePaused);
+	virtual bool IsInViewportClient(UGameViewportClient* ViewportClient) const;
 
 	/** Allows the PlayerController to set up custom input bindings. */
 	virtual void SetupInputComponent();
diff --git a/Engine/Source/Runtime/Engine/Private/PlayerController.cpp b/Engine/Source/Runtime/Engine/Private/PlayerController.cpp
--- a/Engine/Source/Runtime/Engine/Private/PlayerController.cpp
+++ b/Engine/Source/Runtime/Engine/Private/PlayerController.cpp
@@ -3795,6 +3795,11 @@ ULocalPlayer* APlayerController::GetLocalPlayer() const
 	return Cast<ULocalPlayer>(Player);
 }
 
+bool APlayerController::IsInViewportClient(UGameViewportClient* ViewportClient) const
+{
+	return ViewportClient && ViewportClient->GetGameViewportWidget().IsValid() && ViewportClient->GetGameViewportWidget()->IsDirectlyHovered();
+}
+
 void APlayerController::TickPlayerInput(const float DeltaSeconds, const bool bGamePaused)
 {
 	SCOPE_CYCLE_COUNTER(STAT_PC_TickInput);
@@ -3813,7 +3818,7 @@ void APlayerController::TickPlayerInput(const float DeltaSeconds, const bool bGa
 			UGameViewportClient* ViewportClient = LocalPlayer->ViewportClient;
 
 			// Only send mouse hit events if we're directly over the viewport.
-			if ( ViewportClient && ViewportClient->GetGameViewportWidget().IsValid() && ViewportClient->GetGameViewportWidget()->IsDirectlyHovered() )
+			if ( IsInViewportClient(ViewportClient) )
 			{
 				if ( LocalPlayer->ViewportClient->GetMousePosition(MousePosition) )
 				{
