Fix for shootergame cook failure by copying 3027607.

Protecting against corrupt, cross graph links. This fixes an ensure when cooking shootergame. The ensure fires because two soundcue objects in the Sanctuary.umap reference eachother.

diff --git a/Engine/Source/Runtime/Engine/Private/EdGraph/EdGraphPin.cpp b/Engine/Source/Runtime/Engine/Private/EdGraph/EdGraphPin.cpp
--- a/Engine/Source/Runtime/Engine/Private/EdGraph/EdGraphPin.cpp
+++ b/Engine/Source/Runtime/Engine/Private/EdGraph/EdGraphPin.cpp
@@ -1113,7 +1113,8 @@ void UEdGraphPin::InitFromDeprecatedPin(class UEdGraphPin_Deprecated* Deprecated
 	{
 		UEdGraphPin_Deprecated* OldLinkedToPin = DeprecatedPin->LinkedTo[PinIdx];
 		UEdGraphPin*& NewPin = LinkedTo[PinIdx];
-		if (OldLinkedToPin)
+		// check for corrupt data with cross graph links:
+		if (OldLinkedToPin && OldLinkedToPin->GetTypedOuter<UEdGraph>() == DeprecatedPin->GetTypedOuter<UEdGraph>())
 		{
 			UEdGraphPin* ExistingPin = FindPinCreatedFromDeprecatedPin(OldLinkedToPin);
 			if (ExistingPin)
