UE4 - Disabled concurrent tick queuing in windows because the windows scheduler is awful.

diff --git a/Engine/Source/Runtime/Engine/Private/TickTaskManager.cpp b/Engine/Source/Runtime/Engine/Private/TickTaskManager.cpp
--- a/Engine/Source/Runtime/Engine/Private/TickTaskManager.cpp
+++ b/Engine/Source/Runtime/Engine/Private/TickTaskManager.cpp
@@ -1338,10 +1338,13 @@ public:
 		
 		int32 NumWorkerThread = 0;
 		bool bConcurrentQueue = false;
+#if !PLATFORM_WINDOWS
+		// the windows scheduler will hang for seconds trying to do this algorithm, threads starve even though other threads are calling sleep(0)
 		if (!FTickTaskSequencer::SingleThreadedMode())
 		{
 			bConcurrentQueue = !!CVarAllowConcurrentQueue.GetValueOnGameThread();
 		}
+#endif
 
 		if (!bConcurrentQueue)
 		{
