Fixed issue with hieriarchical instanced SM data not being cleared on rebuild.
Rebuilt Agora_Terrain to clear warning

diff --git a/Engine/Source/Runtime/Engine/Private/HierarchicalInstancedStaticMesh.cpp b/Engine/Source/Runtime/Engine/Private/HierarchicalInstancedStaticMesh.cpp
--- a/Engine/Source/Runtime/Engine/Private/HierarchicalInstancedStaticMesh.cpp
+++ b/Engine/Source/Runtime/Engine/Private/HierarchicalInstancedStaticMesh.cpp
@@ -2198,20 +2198,23 @@ void UHierarchicalInstancedStaticMeshComponent::BuildTree()
 		FlushAccumulatedNavigationUpdates();
 
 		PostBuildStats();
-
-		if (bIsAsyncBuilding)
-		{
-			// We did a sync build while async building. The sync build is newer so we will use that.
-			bDiscardAsyncBuildResults = true;
-		}
 	}
 	else
 	{
+		ClusterTreePtr = MakeShareable(new TArray<FClusterNode>);
+		NumBuiltInstances = 0;
+		NumBuiltRenderInstances = 0;
 		InstanceReorderTable.Empty();
 		SortedInstances.Empty();
 		RemovedInstances.Empty();
 	}
 
+	if (bIsAsyncBuilding)
+	{
+		// We did a sync build while async building. The sync build is newer so we will use that.
+		bDiscardAsyncBuildResults = true;
+	}
+
 	ReleasePerInstanceRenderData();
 }
 
@@ -2405,6 +2408,9 @@ void UHierarchicalInstancedStaticMeshComponent::BuildTreeAsync()
 	}
 	else
 	{
+		ClusterTreePtr = MakeShareable(new TArray<FClusterNode>);
+		NumBuiltInstances = 0;
+		NumBuiltRenderInstances = 0;
 		InstanceReorderTable.Empty();
 		SortedInstances.Empty();
 		RemovedInstances.Empty();
